Cameras: Fix bug with Matinee camera shake when the camera is pointing far up or down.

We should apply angle limits *before* computing the absolute shake output transform (i.e. compute angle limits while we still have the additive/local values)

diff --git a/Engine/Plugins/Cameras/GameplayCameras/Source/GameplayCameras/Private/MatineeCameraShake.cpp b/Engine/Plugins/Cameras/GameplayCameras/Source/GameplayCameras/Private/MatineeCameraShake.cpp
--- a/Engine/Plugins/Cameras/GameplayCameras/Source/GameplayCameras/Private/MatineeCameraShake.cpp
+++ b/Engine/Plugins/Cameras/GameplayCameras/Source/GameplayCameras/Private/MatineeCameraShake.cpp
@@ -392,10 +392,11 @@ void UMatineeCameraShake::DoUpdateShake(const FCameraShakeUpdateParams& Params,
 	}
 
 	// Apply the playspace and the scaling so we have an absolute result we can pass to the legacy blueprint API.
+	check(OutResult.Flags == ECameraShakeUpdateResultFlags::Default);
 	const float CurShakeScale = Params.ShakeScale * Params.DynamicScale;
 	ApplyScale(CurShakeScale, OutResult);
-	ApplyPlaySpace(Params, OutResult);
 	ApplyLimits(Params.POV, OutResult);
+	ApplyPlaySpace(Params, OutResult);
 	check(EnumHasAnyFlags(OutResult.Flags, ECameraShakeUpdateResultFlags::ApplyAsAbsolute));
 
 	// Call the legacy blueprint API. We need to convert back and forth.
