Cameras: Hide Matinee shake implementation details from the editor

This prevents users from:

- Picking a different root shake pattern for a Matinee shake
- Picking a Matinee shake pattern in a "normal" (non legacy) shake

diff --git a/Engine/Plugins/Cameras/GameplayCameras/Source/GameplayCameras/Public/MatineeCameraShake.h b/Engine/Plugins/Cameras/GameplayCameras/Source/GameplayCameras/Public/MatineeCameraShake.h
--- a/Engine/Plugins/Cameras/GameplayCameras/Source/GameplayCameras/Public/MatineeCameraShake.h
+++ b/Engine/Plugins/Cameras/GameplayCameras/Source/GameplayCameras/Public/MatineeCameraShake.h
@@ -123,7 +123,7 @@ struct FVOscillator
 /**
  * Legacy camera shake which can do either oscillation or run camera anims.
  */
-UCLASS(Blueprintable)
+UCLASS(Blueprintable, HideCategories=(CameraShakePattern))
 class GAMEPLAYCAMERAS_API UMatineeCameraShake : public UCameraShakeBase
 {
 	GENERATED_UCLASS_BODY()
@@ -308,7 +308,7 @@ private:
  * was left on the shake class itself... so this pattern delegates everything back
  * to the owner shake.
  */
-UCLASS()
+UCLASS(HideDropdown)
 class GAMEPLAYCAMERAS_API UMatineeCameraShakePattern : public UCameraShakePattern
 {
 	GENERATED_BODY()
diff --git a/Engine/Source/Runtime/Engine/Classes/Camera/CameraShakeBase.h b/Engine/Source/Runtime/Engine/Classes/Camera/CameraShakeBase.h
--- a/Engine/Source/Runtime/Engine/Classes/Camera/CameraShakeBase.h
+++ b/Engine/Source/Runtime/Engine/Classes/Camera/CameraShakeBase.h
@@ -533,7 +533,7 @@ protected:
 private:
 
 	/** The root pattern for this camera shake */
-	UPROPERTY(EditAnywhere, Instanced, Category=CameraShake)
+	UPROPERTY(EditAnywhere, Instanced, Category=CameraShakePattern)
 	TObjectPtr<UCameraShakePattern> RootShakePattern;
 
 	/** The camera manager owning this camera shake. Only valid when the shake is active. */
