Add GAMEPLAYCAMERAS_AP to FROscillator so it can be used outside the plugin

#rnx

diff --git a/Engine/Plugins/Cameras/GameplayCameras/Source/GameplayCameras/Public/MatineeCameraShake.h b/Engine/Plugins/Cameras/GameplayCameras/Source/GameplayCameras/Public/MatineeCameraShake.h
--- a/Engine/Plugins/Cameras/GameplayCameras/Source/GameplayCameras/Public/MatineeCameraShake.h
+++ b/Engine/Plugins/Cameras/GameplayCameras/Source/GameplayCameras/Public/MatineeCameraShake.h
@@ -82,7 +82,7 @@ struct GAMEPLAYCAMERAS_API FFOscillator
 
 /** Defines FRotator oscillation. */
 USTRUCT(BlueprintType)
-struct FROscillator
+struct GAMEPLAYCAMERAS_API FROscillator
 {
 	GENERATED_USTRUCT_BODY()
 
