Cameras: Add gameplay cameras world subsystem for easier way to start/stop animations.

diff --git a/Engine/Plugins/Cameras/GameplayCameras/Source/GameplayCameras/Private/GameplayCamerasSubsystem.cpp b/Engine/Plugins/Cameras/GameplayCameras/Source/GameplayCameras/Private/GameplayCamerasSubsystem.cpp
new file mode 100644
--- /dev/null
+++ b/Engine/Plugins/Cameras/GameplayCameras/Source/GameplayCameras/Private/GameplayCamerasSubsystem.cpp
@@ -0,0 +1,67 @@
+// Copyright Epic Games, Inc. All Rights Reserved.
+
+#include "GameplayCamerasSubsystem.h"
+#include "CameraAnimationCameraModifier.h"
+#include "Engine/World.h"
+
+#define LOCTEXT_NAMESPACE "GameplayCamerasSubsystem"
+
+UGameplayCamerasSubsystem* UGameplayCamerasSubsystem::GetGameplayCamerasSubsystem(const UWorld* InWorld)
+{
+	if (InWorld)
+	{
+		return InWorld->GetSubsystem<UGameplayCamerasSubsystem>();
+	}
+
+	return nullptr;
+}
+
+FCameraAnimationHandle UGameplayCamerasSubsystem::PlayCameraAnimation(APlayerController* PlayerController, UCameraAnimationSequence* Sequence, FCameraAnimationParams Params)
+{
+	UCameraAnimationCameraModifier* CameraModifier = UCameraAnimationCameraModifier::GetCameraAnimationCameraModifierFromPlayerController(PlayerController);
+	if (ensure(CameraModifier))
+	{
+		return CameraModifier->PlayCameraAnimation(Sequence, Params);
+	}
+	return FCameraAnimationHandle::Invalid;
+}
+
+bool UGameplayCamerasSubsystem::IsCameraAnimationActive(APlayerController* PlayerController, const FCameraAnimationHandle& Handle) const
+{
+	UCameraAnimationCameraModifier* CameraModifier = UCameraAnimationCameraModifier::GetCameraAnimationCameraModifierFromPlayerController(PlayerController);
+	if (ensure(CameraModifier))
+	{
+		return CameraModifier->IsCameraAnimationActive(Handle);
+	}
+	return false;
+}
+
+void UGameplayCamerasSubsystem::StopCameraAnimation(APlayerController* PlayerController, const FCameraAnimationHandle& Handle, bool bImmediate)
+{
+	UCameraAnimationCameraModifier* CameraModifier = UCameraAnimationCameraModifier::GetCameraAnimationCameraModifierFromPlayerController(PlayerController);
+	if (ensure(CameraModifier))
+	{
+		return CameraModifier->StopCameraAnimation(Handle, bImmediate);
+	}
+}
+
+void UGameplayCamerasSubsystem::StopAllCameraAnimationsOf(APlayerController* PlayerController, UCameraAnimationSequence* Sequence, bool bImmediate)
+{
+	UCameraAnimationCameraModifier* CameraModifier = UCameraAnimationCameraModifier::GetCameraAnimationCameraModifierFromPlayerController(PlayerController);
+	if (ensure(CameraModifier))
+	{
+		CameraModifier->StopAllCameraAnimationsOf(Sequence, bImmediate);
+	}
+}
+
+void UGameplayCamerasSubsystem::StopAllCameraAnimations(APlayerController* PlayerController, bool bImmediate)
+{
+	UCameraAnimationCameraModifier* CameraModifier = UCameraAnimationCameraModifier::GetCameraAnimationCameraModifierFromPlayerController(PlayerController);
+	if (ensure(CameraModifier))
+	{
+		CameraModifier->StopAllCameraAnimations(bImmediate);
+	}
+}
+
+#undef LOCTEXT_NAMESPACE
+
diff --git a/Engine/Plugins/Cameras/GameplayCameras/Source/GameplayCameras/Public/GameplayCamerasSubsystem.h b/Engine/Plugins/Cameras/GameplayCameras/Source/GameplayCameras/Public/GameplayCamerasSubsystem.h
new file mode 100644
--- /dev/null
+++ b/Engine/Plugins/Cameras/GameplayCameras/Source/GameplayCameras/Public/GameplayCamerasSubsystem.h
@@ -0,0 +1,70 @@
+// Copyright Epic Games, Inc. All Rights Reserved.
+
+#pragma once
+
+#include "CameraAnimationCameraModifier.h"
+#include "CoreMinimal.h"
+#include "Subsystems/WorldSubsystem.h"
+#include "UObject/ObjectMacros.h"
+#include "GameplayCamerasSubsystem.generated.h"
+
+/**
+ * World subsystem that holds global objects for handling camera animation sequences.
+ */
+UCLASS()
+class UGameplayCamerasSubsystem : public UWorldSubsystem
+{
+	GENERATED_BODY()
+
+public:
+
+	/** Get the camera animation sequence subsystem for the given world */
+	static UGameplayCamerasSubsystem* GetGameplayCamerasSubsystem(const UWorld* InWorld);
+
+public:
+
+	/**
+	 * Play a new camera animation sequence.
+	 * @param PlayerController The player controller on which to play the animation.
+	 * @param Sequence		The sequence to use for the new camera animation.
+	 * @param Params		The parameters for the new camera animation instance.
+	 */
+	UFUNCTION(BlueprintCallable, Category="Camera Animation")
+	FCameraAnimationHandle PlayCameraAnimation(APlayerController* PlayerController, UCameraAnimationSequence* Sequence, FCameraAnimationParams Params);
+
+	/**
+	 * Returns whether the given camera animation is playing.
+	 * @param PlayerController The player controller on which to play the animation.
+	 * @param Handle		A handle to a previously started camera animation.
+	 * @return				Whether the corresponding camera animation is playing or not.
+	 */
+	UFUNCTION(BlueprintPure, Category="Camera Animation")
+	bool IsCameraAnimationActive(APlayerController* PlayerController, const FCameraAnimationHandle& Handle) const;
+
+	/** 
+	 * Stops the given camera animation instance.
+	 * @param PlayerController The player controller on which to play the animation.
+	 * @param Hanlde		A handle to a previously started camera animation.
+	 * @param bImmediate	True to stop it right now and ignore blend out, false to let it blend out as indicated.
+	 */
+	UFUNCTION(BlueprintCallable, Category = "Camera Animation")
+	void StopCameraAnimation(APlayerController* PlayerController, const FCameraAnimationHandle& Handle, bool bImmediate = false);
+	
+	/**
+	 * Stop playing all instances of the given camera animation sequence.
+	 * @param PlayerController The player controller on which to play the animation.
+	 * @param Sequence		The sequence of which to stop all instances.
+	 * @param bImmediate	True to stop it right now and ignore blend out, false to let it blend out as indicated.
+	 */
+	UFUNCTION(BlueprintCallable, Category="Camera Animation")
+	void StopAllCameraAnimationsOf(APlayerController* PlayerController, UCameraAnimationSequence* Sequence, bool bImmediate = false);
+	
+	/**
+	 * Stop all camera animation instances.
+	 * @param PlayerController The player controller on which to play the animation.
+	 * @param bImmediate	True to stop it right now and ignore blend out, false to let it blend out as indicated.
+	 */
+	UFUNCTION(BlueprintCallable, Category="Camera Animation")
+	void StopAllCameraAnimations(APlayerController* PlayerController, bool bImmediate = false);
+};
+
