CIS fix for non-unity builds

diff --git a/Engine/Plugins/Cameras/GameplayCameras/Source/GameplayCameras/Public/CameraAnimationCameraModifier.h b/Engine/Plugins/Cameras/GameplayCameras/Source/GameplayCameras/Public/CameraAnimationCameraModifier.h
--- a/Engine/Plugins/Cameras/GameplayCameras/Source/GameplayCameras/Public/CameraAnimationCameraModifier.h
+++ b/Engine/Plugins/Cameras/GameplayCameras/Source/GameplayCameras/Public/CameraAnimationCameraModifier.h
@@ -9,6 +9,7 @@
 #include "Camera/CameraTypes.h"
 #include "Camera/CameraModifier.h"
 #include "Containers/SparseArray.h"
+#include "Kismet/BlueprintFunctionLibrary.h"
 #include "CameraAnimationCameraModifier.generated.h"
 
 class UCameraAnimationSequence;
