Cameras: fix wrong serial on first camera animation handle

Serial number zero should indicate an invalid handle -- we were using it for the first handle.

diff --git a/Engine/Plugins/Cameras/GameplayCameras/Source/GameplayCameras/Private/CameraAnimationCameraModifier.cpp b/Engine/Plugins/Cameras/GameplayCameras/Source/GameplayCameras/Private/CameraAnimationCameraModifier.cpp
--- a/Engine/Plugins/Cameras/GameplayCameras/Source/GameplayCameras/Private/CameraAnimationCameraModifier.cpp
+++ b/Engine/Plugins/Cameras/GameplayCameras/Source/GameplayCameras/Private/CameraAnimationCameraModifier.cpp
@@ -77,6 +77,7 @@ float UCameraAnimationCameraModifier::EvaluateEasing(ECameraAnimationEasingType
 
 UCameraAnimationCameraModifier::UCameraAnimationCameraModifier(const FObjectInitializer& ObjectInitializer)
 	: Super(ObjectInitializer)
+	, NextInstanceSerialNumber(1)
 {
 }
 
@@ -92,7 +93,7 @@ FCameraAnimationHandle UCameraAnimationCameraModifier::PlayCameraAnimation(UCame
 	int32 NewIndex = FindInactiveCameraAnimation();
 	check(NewIndex < MAX_uint16);
 
-	const uint16 InstanceSerial = InstanceSerialNumber++;
+	const uint16 InstanceSerial = NextInstanceSerialNumber++;
 	FCameraAnimationHandle InstanceHandle { (uint16)NewIndex, InstanceSerial };
 
 	FActiveCameraAnimationInfo& NewCameraAnimation = ActiveAnimations[NewIndex];
diff --git a/Engine/Plugins/Cameras/GameplayCameras/Source/GameplayCameras/Public/CameraAnimationCameraModifier.h b/Engine/Plugins/Cameras/GameplayCameras/Source/GameplayCameras/Public/CameraAnimationCameraModifier.h
--- a/Engine/Plugins/Cameras/GameplayCameras/Source/GameplayCameras/Public/CameraAnimationCameraModifier.h
+++ b/Engine/Plugins/Cameras/GameplayCameras/Source/GameplayCameras/Public/CameraAnimationCameraModifier.h
@@ -261,7 +261,7 @@ protected:
 
 	/** Next serial number to use for a camera animation instance */
 	UPROPERTY()
-	uint16 InstanceSerialNumber;
+	uint16 NextInstanceSerialNumber;
 };
 
 /**
