Rename UMatineeCameraShake to ULegacyCameraShake

diff --git a/Engine/Config/BaseEngine.ini b/Engine/Config/BaseEngine.ini
--- a/Engine/Config/BaseEngine.ini
+++ b/Engine/Config/BaseEngine.ini
@@ -868,9 +868,6 @@ DefaultWorldPartitionHLODLayer=/Engine/Maps/Templates/HLODs/HLODLayer_Instanced.
 +StructRedirects=(OldName="FOscillator",NewName="/Script/GameplayCameras.FOscillator")
 +StructRedirects=(OldName="ROscillator",NewName="/Script/GameplayCameras.ROscillator")
 +StructRedirects=(OldName="VOscillator",NewName="/Script/GameplayCameras.VOscillator")
-+ClassRedirects=(OldName="CameraShake",NewName="/Script/GameplayCameras.MatineeCameraShake")
-+ClassRedirects=(OldName="MatineeCameraShake",NewName="/Script/GameplayCameras.MatineeCameraShake")
-+ClassRedirects=(OldName="MatineeCameraShakePattern",NewName="/Script/GameplayCameras.MatineeCameraShakePattern")
 +EnumRedirects=(OldName="ECameraAnimPlaySpace",NewName="/Script/Engine.ECameraShakePlaySpace")
 +FunctionRedirects=(OldName="PlayerController.ClientPlayCameraShake",NewName="ClientStartCameraShake")
 +FunctionRedirects=(OldName="PlayerController.ClientPlayCameraShakeFromSource",NewName="ClientStartCameraShakeFromSource")
@@ -1090,6 +1087,14 @@ DefaultWorldPartitionHLODLayer=/Engine/Maps/Templates/HLODs/HLODLayer_Instanced.
 +PropertyRedirects=(OldName="AGameplayAbilityWorldReticle.MasterPC", NewName="PrimaryPC")
 +PropertyRedirects=(OldName="AGameplayAbilityTargetActor.MasterPC", NewName="PrimaryPC")
 
++ClassRedirects=(OldName="CameraShake",NewName="/Script/GameplayCameras.LegacyCameraShake")
++ClassRedirects=(OldName="MatineeCameraShake",NewName="/Script/GameplayCameras.LegacyCameraShake")
++ClassRedirects=(OldName="MatineeCameraShakePattern",NewName="/Script/GameplayCameras.LegacyCameraShakePattern")
++ClassRedirects=(OldName="MatineeCameraShakeFunctionLibrary", NewName="/Script/GameplayCameras.LegacyCameraShakeFunctionLibrary")
++FunctionRedirects=(OldName="LegacyCameraShake.StartMatineeCameraShake", NewName="LegacyCameraShake.StartLegacyCameraShake")
++FunctionRedirects=(OldName="LegacyCameraShake.StartMatineeCameraShakeFromSource", NewName="LegacyCameraShake.StartLegacyCameraShakeFromSource")
++FunctionRedirects=(OldName="LegacyCameraShakeFunctionLibrary.Conv_MatineeCameraShake", NewName="LegacyCameraShakeFunctionLibrary.Conv_LegacyCameraShake")
+
 [CoreUObject.Metadata]
 ; Note: UnrealHeaderTool should be rerun after making changes to MetadataRedirects to catch any keys specified in class headers
 +MetadataRedirects=(OldKey="K2Protected", NewKey="BlueprintProtected")
@@ -2748,7 +2753,6 @@ DefaultASTCQualityBySize=3
 +ClassesExcludedOnDedicatedServer=Font
 +ClassesExcludedOnDedicatedServer=InterpCurveEdSetup
 +ClassesExcludedOnDedicatedServer=MaterialExpression
-+ClassesExcludedOnDedicatedServer=MatineeActorCameraAnim
 +ClassesExcludedOnDedicatedServer=NiagaraEmitter
 +ClassesExcludedOnDedicatedServer=NiagaraScript
 +ClassesExcludedOnDedicatedServer=ParticleEmitter
@@ -2762,7 +2766,6 @@ DefaultASTCQualityBySize=3
 +ClassesExcludedOnDedicatedClient=MetaData
 +ClassesExcludedOnDedicatedClient=ObjectRedirector
 +ClassesExcludedOnDedicatedClient=InterpCurveEdSetup
-+ClassesExcludedOnDedicatedClient=MatineeActorCameraAnim
 
 +VersionedIntRValues=r.AllowStaticLighting
 +VersionedIntRValues=r.GBuffer
diff --git a/Engine/Plugins/Cameras/CameraShakePreviewer/Source/CameraShakePreviewer/Private/SCameraShakePreviewer.cpp b/Engine/Plugins/Cameras/CameraShakePreviewer/Source/CameraShakePreviewer/Private/SCameraShakePreviewer.cpp
--- a/Engine/Plugins/Cameras/CameraShakePreviewer/Source/CameraShakePreviewer/Private/SCameraShakePreviewer.cpp
+++ b/Engine/Plugins/Cameras/CameraShakePreviewer/Source/CameraShakePreviewer/Private/SCameraShakePreviewer.cpp
@@ -12,7 +12,7 @@
 #include "Styling/AppStyle.h"
 #include "LevelEditor.h"
 #include "LevelEditorViewport.h"
-#include "MatineeCameraShake.h"
+#include "LegacyCameraShake.h"
 #include "Modules/ModuleManager.h"
 #include "Slate/SceneViewport.h"
 #include "Widgets/Input/SButton.h"
diff --git a/Engine/Plugins/Cameras/GameplayCameras/Source/GameplayCameras/Private/LegacyCameraShake.cpp b/Engine/Plugins/Cameras/GameplayCameras/Source/GameplayCameras/Private/LegacyCameraShake.cpp
new file mode 100644
--- /dev/null
+++ b/Engine/Plugins/Cameras/GameplayCameras/Source/GameplayCameras/Private/LegacyCameraShake.cpp
@@ -0,0 +1,514 @@
+// Copyright Epic Games, Inc. All Rights Reserved.
+
+#include "LegacyCameraShake.h"
+#include "SequenceCameraShake.h"
+#include "Camera/CameraActor.h"
+#include "Camera/CameraComponent.h"
+#include "Camera/PlayerCameraManager.h"
+#include "Engine/Engine.h"
+#include "IXRTrackingSystem.h" // for IsHeadTrackingAllowed()
+
+DEFINE_LOG_CATEGORY_STATIC(LogLegacyCameraShake, Warning, All);
+
+//////////////////////////////////////////////////////////////////////////
+// FFOscillator
+
+// static
+float FFOscillator::UpdateOffset(FFOscillator const& Osc, float& CurrentOffset, float DeltaTime)
+{
+	// LWC_TODO: Perf pessimization
+	double AsDouble = CurrentOffset;
+	float Result = UpdateOffset(Osc, AsDouble, DeltaTime);
+	CurrentOffset = (float)AsDouble;
+	return Result;
+}
+
+// static
+float FFOscillator::UpdateOffset(FFOscillator const& Osc, double& CurrentOffset, float DeltaTime)
+{
+	if (Osc.Amplitude != 0.f)
+	{
+		CurrentOffset += DeltaTime * Osc.Frequency;
+
+		float WaveformSample;
+		switch (Osc.Waveform)
+		{
+		case EOscillatorWaveform::SineWave:
+		default:
+			WaveformSample = FMath::Sin(CurrentOffset);
+			break;
+
+		case EOscillatorWaveform::PerlinNoise:
+			WaveformSample = FMath::PerlinNoise1D(CurrentOffset);
+			break;
+		}
+
+		return Osc.Amplitude * WaveformSample;
+	}
+	return 0.f;
+}
+
+// static
+float FFOscillator::GetInitialOffset(FFOscillator const& Osc)
+{
+	return (Osc.InitialOffset == EOO_OffsetRandom)
+		? FMath::FRand() * (2.f * PI)
+		: 0.f;
+}
+
+// static
+float FFOscillator::GetOffsetAtTime(FFOscillator const& Osc, float InitialOffset, float Time)
+{
+	return InitialOffset + (Time * Osc.Frequency);
+}
+
+//////////////////////////////////////////////////////////////////////////
+// ULegacyCameraShake
+
+ULegacyCameraShake::ULegacyCameraShake(const FObjectInitializer& ObjectInitializer)
+	: Super(ObjectInitializer
+				.SetDefaultSubobjectClass<ULegacyCameraShakePattern>(TEXT("RootShakePattern")))
+{
+	AnimPlayRate = 1.0f;
+	AnimScale = 1.0f;
+	AnimBlendInTime = 0.2f;
+	AnimBlendOutTime = 0.2f;
+	OscillationBlendInTime = 0.1f;
+	OscillationBlendOutTime = 0.2f;
+}
+
+void ULegacyCameraShake::DoStopShake(bool bImmediately)
+{
+	APlayerCameraManager* CameraOwner = GetCameraManager();
+
+	if (bImmediately)
+	{
+		// stop oscillation
+		OscillatorTimeRemaining = 0.f;
+	}
+	else
+	{
+		// advance to the blend out time
+		if (OscillatorTimeRemaining > 0.0f)
+		{
+			OscillatorTimeRemaining = FMath::Min(OscillatorTimeRemaining, OscillationBlendOutTime);
+		}
+		else
+		{
+			OscillatorTimeRemaining = OscillationBlendOutTime;
+		}
+	}
+	
+	if (SequenceShakePattern)
+	{
+		FCameraShakeStopParams StopParams;
+		StopParams.bImmediately = bImmediately;
+		SequenceShakePattern->StopShakePattern(StopParams);
+	}
+
+	UE_LOG(LogLegacyCameraShake, Verbose, TEXT("ULegacyCameraShake::DoStopShake %s"), *GetNameSafe(this));
+
+	ReceiveStopShake(bImmediately);
+}
+
+void ULegacyCameraShake::DoStartShake(const FCameraShakeStartParams& Params)
+{
+	const float EffectiveOscillationDuration = (OscillationDuration > 0.f) ? OscillationDuration : TNumericLimits<float>::Max();
+
+	// init oscillations
+	if (OscillationDuration != 0.f)
+	{
+		if (OscillatorTimeRemaining > 0.f)
+		{
+			// this shake was already playing
+			OscillatorTimeRemaining = EffectiveOscillationDuration;
+
+			if (bBlendingOut)
+			{
+				bBlendingOut = false;
+				CurrentBlendOutTime = 0.f;
+
+				// stop any blendout and reverse it to a blendin
+				if (OscillationBlendInTime > 0.f)
+				{
+					bBlendingIn = true;
+					CurrentBlendInTime = OscillationBlendInTime * (1.f - CurrentBlendOutTime / OscillationBlendOutTime);
+				}
+				else
+				{
+					bBlendingIn = false;
+					CurrentBlendInTime = 0.f;
+				}
+			}
+		}
+		else
+		{
+			RotSinOffset.X = FFOscillator::GetInitialOffset(RotOscillation.Pitch);
+			RotSinOffset.Y = FFOscillator::GetInitialOffset(RotOscillation.Yaw);
+			RotSinOffset.Z = FFOscillator::GetInitialOffset(RotOscillation.Roll);
+
+			LocSinOffset.X = FFOscillator::GetInitialOffset(LocOscillation.X);
+			LocSinOffset.Y = FFOscillator::GetInitialOffset(LocOscillation.Y);
+			LocSinOffset.Z = FFOscillator::GetInitialOffset(LocOscillation.Z);
+
+			FOVSinOffset = FFOscillator::GetInitialOffset(FOVOscillation);
+
+			InitialLocSinOffset = LocSinOffset;
+			InitialRotSinOffset = RotSinOffset;
+			InitialFOVSinOffset = FOVSinOffset;
+
+			OscillatorTimeRemaining = EffectiveOscillationDuration;
+
+			if (OscillationBlendInTime > 0.f)
+			{
+				bBlendingIn = true;
+				CurrentBlendInTime = 0.f;
+			}
+		}
+	}
+
+	// init cameraanim shakes
+	APlayerCameraManager* CameraOwner = GetCameraManager();
+	if (AnimSequence != nullptr)
+	{
+		if (SequenceShakePattern == nullptr)
+		{
+			SequenceShakePattern = NewObject<USequenceCameraShakePattern>(this);
+		}
+
+		// Copy our anim parameters over to the sequence shake pattern.
+		SequenceShakePattern->Sequence = AnimSequence;
+		SequenceShakePattern->PlayRate = AnimPlayRate;
+		SequenceShakePattern->Scale = AnimScale;
+		SequenceShakePattern->BlendInTime = AnimBlendInTime;
+		SequenceShakePattern->BlendOutTime = AnimBlendOutTime;
+		SequenceShakePattern->RandomSegmentDuration = RandomAnimSegmentDuration;
+		SequenceShakePattern->bRandomSegment = bRandomAnimSegment;
+		
+		// Initialize our state tracker for the sequence shake pattern.
+		FCameraShakeInfo SequenceShakeInfo;
+		SequenceShakePattern->GetShakePatternInfo(SequenceShakeInfo);
+		SequenceShakeState.Initialize(SequenceShakeInfo);
+
+		// Start the sequence shake pattern.
+		SequenceShakePattern->StartShakePattern(Params);
+	}
+
+	UE_LOG(LogLegacyCameraShake, Verbose, TEXT("ULegacyCameraShake::DoStartShake %s Duration: %f"), *GetNameSafe(this), OscillationDuration);
+
+	ReceivePlayShake(ShakeScale);
+}
+
+void ULegacyCameraShake::DoUpdateShake(const FCameraShakeUpdateParams& Params, FCameraShakeUpdateResult& OutResult)
+{
+	const float DeltaTime = Params.DeltaTime;
+	const float BaseShakeScale = Params.GetTotalScale();
+
+	// update oscillation times... only decrease the time remaining if we're not infinite
+	if (OscillatorTimeRemaining > 0.f)
+	{
+		OscillatorTimeRemaining -= DeltaTime;
+		OscillatorTimeRemaining = FMath::Max(0.f, OscillatorTimeRemaining);
+	}
+	if (bBlendingIn)
+	{
+		CurrentBlendInTime += DeltaTime;
+	}
+	if (bBlendingOut)
+	{
+		CurrentBlendOutTime += DeltaTime;
+	}
+
+	// see if we've crossed any important time thresholds and deal appropriately
+	bool bOscillationFinished = false;
+
+	if (OscillatorTimeRemaining <= 0.f)
+	{
+		// finished!
+		bOscillationFinished = true;
+	}
+	else if (OscillatorTimeRemaining < OscillationBlendOutTime)
+	{
+		// start blending out
+		bBlendingOut = true;
+		CurrentBlendOutTime = OscillationBlendOutTime - OscillatorTimeRemaining;
+	}
+	else if (OscillationDuration < 0.f)
+	{
+		// infinite oscillation, keep the time remaining up
+		OscillatorTimeRemaining = TNumericLimits<float>::Max();
+	}
+
+	if (bBlendingIn)
+	{
+		if (CurrentBlendInTime > OscillationBlendInTime)
+		{
+			// done blending in!
+			bBlendingIn = false;
+		}
+	}
+	if (bBlendingOut)
+	{
+		if (CurrentBlendOutTime > OscillationBlendOutTime)
+		{
+			// done!!
+			CurrentBlendOutTime = OscillationBlendOutTime;
+			bOscillationFinished = true;
+		}
+	}
+
+	// Do not update oscillation further if finished
+	if (bOscillationFinished == false)
+	{
+		// calculate blend weight. calculating separately and taking the minimum handles overlapping blends nicely.
+		float const BlendInWeight = (bBlendingIn) ? (CurrentBlendInTime / OscillationBlendInTime) : 1.f;
+		float const BlendOutWeight = (bBlendingOut) ? (1.f - CurrentBlendOutTime / OscillationBlendOutTime) : 1.f;
+		float const CurrentBlendWeight = FMath::Min(BlendInWeight, BlendOutWeight);
+
+		// this is the oscillation scale, which includes oscillation fading
+		// we'll apply the general shake scale, along with the current frame's dynamic scale, a bit later.
+		float const OscillationScale = CurrentBlendWeight;
+
+		if (OscillationScale > 0.f)
+		{
+			// View location offset, compute sin wave value for each component
+			FVector	LocOffset = FVector(0);
+			LocOffset.X = FFOscillator::UpdateOffset(LocOscillation.X, LocSinOffset.X, DeltaTime);
+			LocOffset.Y = FFOscillator::UpdateOffset(LocOscillation.Y, LocSinOffset.Y, DeltaTime);
+			LocOffset.Z = FFOscillator::UpdateOffset(LocOscillation.Z, LocSinOffset.Z, DeltaTime);
+			LocOffset *= OscillationScale;
+
+			OutResult.Location = LocOffset;
+
+			// View rotation offset, compute sin wave value for each component
+			FRotator RotOffset;
+			RotOffset.Pitch = FFOscillator::UpdateOffset(RotOscillation.Pitch, RotSinOffset.X, DeltaTime) * OscillationScale;
+			RotOffset.Yaw = FFOscillator::UpdateOffset(RotOscillation.Yaw, RotSinOffset.Y, DeltaTime) * OscillationScale;
+			RotOffset.Roll = FFOscillator::UpdateOffset(RotOscillation.Roll, RotSinOffset.Z, DeltaTime) * OscillationScale;
+
+			// Don't allow shake to flip pitch past vertical, if not using a headset (where we can't limit the camera locked to your head).
+			APlayerCameraManager* CameraOwner = GetCameraManager();
+			UWorld * World = (CameraOwner ? CameraOwner->GetWorld() : nullptr);
+			if (!GEngine->XRSystem.IsValid() ||
+				!(World != nullptr ? GEngine->XRSystem->IsHeadTrackingAllowedForWorld(*World) : GEngine->XRSystem->IsHeadTrackingAllowed()))
+			{
+				// Find normalized result when combined, and remove any offset that would push it past the limit.
+				const float NormalizedInputPitch = FRotator::NormalizeAxis(Params.POV.Rotation.Pitch);
+				RotOffset.Pitch = FRotator::NormalizeAxis(RotOffset.Pitch);
+				RotOffset.Pitch = FMath::ClampAngle(NormalizedInputPitch + RotOffset.Pitch, -89.9f, 89.9f) - NormalizedInputPitch;
+			}
+
+			OutResult.Rotation = RotOffset;
+
+			// Compute FOV change
+			OutResult.FOV = OscillationScale * FFOscillator::UpdateOffset(FOVOscillation, FOVSinOffset, DeltaTime);
+		}
+	}
+
+	// Update the sequence animation if there's one.
+	if (SequenceShakePattern != nullptr)
+	{
+		const float ChildBlendWeight = SequenceShakeState.Update(Params.DeltaTime);
+		if (SequenceShakeState.IsActive())
+		{
+			FCameraShakeUpdateParams ChildParams(Params);
+			ChildParams.BlendingWeight = Params.BlendingWeight * ChildBlendWeight;
+
+			FCameraShakeUpdateResult ChildResult;
+
+			SequenceShakePattern->UpdateShakePattern(ChildParams, ChildResult);
+
+			// The sequence shake pattern returns a local, additive, unscaled result. So we should be able to
+			// just combine the two results directly.
+			check(ChildResult.Flags == ECameraShakeUpdateResultFlags::Default);
+			ApplyScale(ChildParams.BlendingWeight, ChildResult);
+			OutResult.Location += ChildResult.Location;
+			OutResult.Rotation += ChildResult.Rotation;
+			OutResult.FOV += ChildResult.FOV;
+			// We don't have anything else animating post-process settings so we can stomp them.
+			OutResult.PostProcessSettings = ChildResult.PostProcessSettings;
+			OutResult.PostProcessBlendWeight = ChildResult.PostProcessBlendWeight;
+		}
+	}
+
+	// Apply the playspace and the scaling so we have an absolute result we can pass to the legacy blueprint API.
+	check(OutResult.Flags == ECameraShakeUpdateResultFlags::Default);
+	const float CurShakeScale = Params.ShakeScale * Params.DynamicScale;
+	ApplyScale(CurShakeScale, OutResult);
+	ApplyLimits(Params.POV, OutResult);
+	ApplyPlaySpace(Params, OutResult);
+	check(EnumHasAnyFlags(OutResult.Flags, ECameraShakeUpdateResultFlags::ApplyAsAbsolute));
+
+	// Call the legacy blueprint API. We need to convert back and forth.
+	{
+		FMinimalViewInfo InOutPOV(Params.POV);
+		InOutPOV.Location = OutResult.Location;
+		InOutPOV.Rotation = OutResult.Rotation;
+		InOutPOV.FOV = OutResult.FOV;
+
+		BlueprintUpdateCameraShake(DeltaTime, Params.DynamicScale, InOutPOV, InOutPOV);
+
+		OutResult.Location = InOutPOV.Location;
+		OutResult.Rotation = InOutPOV.Rotation;
+		OutResult.FOV = InOutPOV.FOV;
+	}
+
+
+	UE_LOG(LogLegacyCameraShake, Verbose, TEXT("ULegacyCameraShake::DoUpdateShake %s Finished: %i Duration: %f Remaining: %f"), *GetNameSafe(this), bOscillationFinished, OscillationDuration, OscillatorTimeRemaining);
+}
+
+void ULegacyCameraShake::DoScrubShake(const FCameraShakeScrubParams& Params, FCameraShakeUpdateResult& OutResult)
+{
+	const float NewTime = Params.AbsoluteTime;
+
+	// reset to start and advance to desired point
+	LocSinOffset = InitialLocSinOffset;
+	RotSinOffset = InitialRotSinOffset;
+	FOVSinOffset = InitialFOVSinOffset;
+
+	const float EffectiveOscillationDuration = (OscillationDuration > 0.f) ? OscillationDuration : TNumericLimits<float>::Max();
+
+	OscillatorTimeRemaining = EffectiveOscillationDuration;
+
+	if (OscillationBlendInTime > 0.f)
+	{
+		bBlendingIn = true;
+		CurrentBlendInTime = 0.f;
+	}
+
+	if (OscillationBlendOutTime > 0.f)
+	{
+		bBlendingOut = false;
+		CurrentBlendOutTime = 0.f;
+	}
+
+	if (OscillationDuration > 0.f)
+	{
+		if ((OscillationBlendOutTime > 0.f) && (NewTime > (OscillationDuration - OscillationBlendOutTime)))
+		{
+			bBlendingOut = true;
+			CurrentBlendOutTime = OscillationBlendOutTime - (OscillationDuration - NewTime);
+		}
+	}
+
+	FCameraShakeUpdateParams UpdateParams = Params.ToUpdateParams();
+
+	DoUpdateShake(UpdateParams, OutResult);
+
+	check(EnumHasAnyFlags(OutResult.Flags, ECameraShakeUpdateResultFlags::ApplyAsAbsolute));
+}
+
+bool ULegacyCameraShake::DoGetIsFinished() const
+{
+	return ((OscillatorTimeRemaining <= 0.f) &&									// oscillator is finished
+		((SequenceShakePattern == nullptr) ||									// other anim is finished
+			SequenceShakeState.GetElapsedTime() >= SequenceShakeState.GetDuration()) &&
+		ReceiveIsFinished()														// BP thinks it's finished
+		);
+}
+
+void ULegacyCameraShake::DoTeardownShake()
+{
+	if (SequenceShakePattern)
+	{
+		SequenceShakePattern->TeardownShakePattern();
+	}
+}
+
+/// @cond DOXYGEN_WARNINGS
+
+bool ULegacyCameraShake::ReceiveIsFinished_Implementation() const
+{
+	return true;
+}
+
+/// @endcond
+
+bool ULegacyCameraShake::IsLooping() const
+{
+	return OscillationDuration < 0.0f;
+}
+
+void ULegacyCameraShake::SetCurrentTimeAndApplyShake(float NewTime, FMinimalViewInfo& POV)
+{
+	ScrubAndApplyCameraShake(NewTime, 1.f, POV);
+}
+
+ULegacyCameraShake* ULegacyCameraShake::StartLegacyCameraShake(APlayerCameraManager* PlayerCameraManager, TSubclassOf<ULegacyCameraShake> ShakeClass, float Scale, ECameraShakePlaySpace PlaySpace, FRotator UserPlaySpaceRot)
+{
+	if (PlayerCameraManager)
+	{
+		return Cast<ULegacyCameraShake>(PlayerCameraManager->StartCameraShake(ShakeClass, Scale, PlaySpace, UserPlaySpaceRot));
+	}
+
+	return nullptr;
+}
+
+ULegacyCameraShake* ULegacyCameraShake::StartLegacyCameraShakeFromSource(APlayerCameraManager* PlayerCameraManager, TSubclassOf<ULegacyCameraShake> ShakeClass, UCameraShakeSourceComponent* SourceComponent, float Scale, ECameraShakePlaySpace PlaySpace, FRotator UserPlaySpaceRot)
+{
+	if (PlayerCameraManager)
+	{
+		return Cast<ULegacyCameraShake>(PlayerCameraManager->StartCameraShakeFromSource(ShakeClass, SourceComponent, Scale, PlaySpace, UserPlaySpaceRot));
+	}
+
+	return nullptr;
+}
+
+void ULegacyCameraShakePattern::GetShakePatternInfoImpl(FCameraShakeInfo& OutInfo) const
+{
+	// We will manage our own duration, but let's give a hint about how long we are for editor purposes.
+	ULegacyCameraShake* Shake = GetShakeInstance<ULegacyCameraShake>();
+	
+	if (Shake->AnimSequence)
+	{
+		float PlaybackSeconds = 0.f;
+		if (UMovieScene* MovieScene = Shake->AnimSequence->GetMovieScene())
+		{
+			const FFrameNumber PlaybackFrames = MovieScene->GetPlaybackRange().Size<FFrameNumber>();
+			PlaybackSeconds = MovieScene->GetTickResolution().AsSeconds(PlaybackFrames);
+		}
+		const float Duration = FMath::Max(Shake->OscillationDuration, PlaybackSeconds);
+		OutInfo.Duration = FCameraShakeDuration::Custom(Duration);
+	}
+	else
+	{
+		OutInfo.Duration = FCameraShakeDuration::Custom(Shake->OscillationDuration);
+	}
+}
+
+void ULegacyCameraShakePattern::StopShakePatternImpl(const FCameraShakeStopParams& Params)
+{
+	ULegacyCameraShake* Shake = GetShakeInstance<ULegacyCameraShake>();
+	Shake->DoStopShake(Params.bImmediately);
+}
+
+void ULegacyCameraShakePattern::StartShakePatternImpl(const FCameraShakeStartParams& Params)
+{
+	ULegacyCameraShake* Shake = GetShakeInstance<ULegacyCameraShake>();
+	Shake->DoStartShake(Params);
+}
+
+void ULegacyCameraShakePattern::UpdateShakePatternImpl(const FCameraShakeUpdateParams& Params, FCameraShakeUpdateResult& OutResult)
+{
+	ULegacyCameraShake* Shake = GetShakeInstance<ULegacyCameraShake>();
+	Shake->DoUpdateShake(Params, OutResult);
+}
+
+void ULegacyCameraShakePattern::ScrubShakePatternImpl(const FCameraShakeScrubParams& Params, FCameraShakeUpdateResult& OutResult)
+{
+	ULegacyCameraShake* Shake = GetShakeInstance<ULegacyCameraShake>();
+	Shake->DoScrubShake(Params, OutResult);
+}
+
+bool ULegacyCameraShakePattern::IsFinishedImpl() const
+{
+	ULegacyCameraShake* Shake = GetShakeInstance<ULegacyCameraShake>();
+	return Shake->DoGetIsFinished();
+}
+
+void ULegacyCameraShakePattern::TeardownShakePatternImpl()
+{
+	ULegacyCameraShake* Shake = GetShakeInstance<ULegacyCameraShake>();
+	Shake->DoTeardownShake();
+}
+
diff --git a/Engine/Plugins/Cameras/GameplayCameras/Source/GameplayCameras/Private/MatineeCameraShake.cpp b/Engine/Plugins/Cameras/GameplayCameras/Source/GameplayCameras/Private/MatineeCameraShake.cpp
deleted file mode 100644
diff --git a/Engine/Plugins/Cameras/GameplayCameras/Source/GameplayCameras/Public/Camera/CameraShake.h b/Engine/Plugins/Cameras/GameplayCameras/Source/GameplayCameras/Public/Camera/CameraShake.h
--- a/Engine/Plugins/Cameras/GameplayCameras/Source/GameplayCameras/Public/Camera/CameraShake.h
+++ b/Engine/Plugins/Cameras/GameplayCameras/Source/GameplayCameras/Public/Camera/CameraShake.h
@@ -4,6 +4,6 @@
 
 #include "CoreTypes.h"
 
-COMPILE_WARNING("UCameraShake has been renamed to UMatineeCameraShake, change to #include \"MatineeCameraShake.h\" and rename references")
-#include "MatineeCameraShake.h"
+COMPILE_WARNING("UCameraShake has been renamed to ULegacyCameraShake, change to #include \"LegacyCameraShake.h\" and rename references")
+#include "LegacyCameraShake.h"
 
diff --git a/Engine/Plugins/Cameras/GameplayCameras/Source/GameplayCameras/Public/LegacyCameraShake.h b/Engine/Plugins/Cameras/GameplayCameras/Source/GameplayCameras/Public/LegacyCameraShake.h
new file mode 100644
--- /dev/null
+++ b/Engine/Plugins/Cameras/GameplayCameras/Source/GameplayCameras/Public/LegacyCameraShake.h
@@ -0,0 +1,337 @@
+// Copyright Epic Games, Inc. All Rights Reserved.
+
+#pragma once
+
+#include "CoreMinimal.h"
+#include "Camera/CameraShakeBase.h"
+#include "Evaluation/MovieSceneCameraShakeTemplate.h"
+#include "Kismet/BlueprintFunctionLibrary.h"
+#include "UObject/Object.h"
+#include "GameFramework/Actor.h"
+#include "LegacyCameraShake.generated.h"
+
+class AActor;
+
+/************************************************************
+ * Parameters for defining oscillating camera shakes
+ ************************************************************/
+
+ /** Types of waveforms that can be used for camera shake oscillators */
+UENUM(BlueprintType)
+enum class EOscillatorWaveform : uint8
+{
+	/** A sinusoidal wave */
+	SineWave,
+
+	/** Perlin noise */
+	PerlinNoise,
+};
+
+/** Shake start offset parameter */
+UENUM()
+enum EInitialOscillatorOffset
+{
+	/** Start with random offset (default). */
+	EOO_OffsetRandom UMETA(DisplayName = "Random"),
+	/** Start with zero offset. */
+	EOO_OffsetZero UMETA(DisplayName = "Zero"),
+	EOO_MAX,
+};
+
+/** Defines oscillation of a single number. */
+USTRUCT(BlueprintType)
+struct GAMEPLAYCAMERAS_API FFOscillator
+{
+	GENERATED_USTRUCT_BODY()
+
+	/** Amplitude of the sinusoidal oscillation. */
+	UPROPERTY(EditAnywhere, BlueprintReadWrite, Category = FOscillator)
+	float Amplitude;
+
+	/** Frequency of the sinusoidal oscillation. */
+	UPROPERTY(EditAnywhere, BlueprintReadWrite, Category = FOscillator)
+	float Frequency;
+
+	/** Defines how to begin (either at zero, or at a randomized value. */
+	UPROPERTY(EditAnywhere, Category = FOscillator)
+	TEnumAsByte<enum EInitialOscillatorOffset> InitialOffset;
+
+	/** Type of waveform to use for oscillation. */
+	UPROPERTY(EditAnywhere, BlueprintReadWrite, Category = FOscillator)
+	EOscillatorWaveform Waveform;
+
+	FFOscillator()
+		: Amplitude(0)
+		, Frequency(0)
+		, InitialOffset(0)
+		, Waveform(EOscillatorWaveform::SineWave)
+	{}
+
+	/** Advances the oscillation time and returns the current value. */
+	static float UpdateOffset(FFOscillator const& Osc, float& CurrentOffset, float DeltaTime);
+
+	/** Advances the oscillation time and returns the current value. */
+	static float UpdateOffset(FFOscillator const& Osc, double& CurrentOffset, float DeltaTime);
+
+	/** Returns the initial value of the oscillator. */
+	static float GetInitialOffset(FFOscillator const& Osc);
+
+	/** Returns the offset at the given time */
+	static float GetOffsetAtTime(FFOscillator const& Osc, float InitialOffset, float Time);
+};
+
+/** Defines FRotator oscillation. */
+USTRUCT(BlueprintType)
+struct GAMEPLAYCAMERAS_API FROscillator
+{
+	GENERATED_USTRUCT_BODY()
+
+	/** Pitch oscillation. */
+	UPROPERTY(EditAnywhere, BlueprintReadWrite, Category = ROscillator)
+	struct FFOscillator Pitch;
+
+	/** Yaw oscillation. */
+	UPROPERTY(EditAnywhere, BlueprintReadWrite, Category = ROscillator)
+	struct FFOscillator Yaw;
+
+	/** Roll oscillation. */
+	UPROPERTY(EditAnywhere, BlueprintReadWrite, Category = ROscillator)
+	struct FFOscillator Roll;
+
+};
+
+/** Defines FVector oscillation. */
+USTRUCT(BlueprintType)
+struct FVOscillator
+{
+	GENERATED_USTRUCT_BODY()
+
+	/** Oscillation in the X axis. */
+	UPROPERTY(EditAnywhere, BlueprintReadWrite, Category = VOscillator)
+	struct FFOscillator X;
+
+	/** Oscillation in the Y axis. */
+	UPROPERTY(EditAnywhere, BlueprintReadWrite, Category = VOscillator)
+	struct FFOscillator Y;
+
+	/** Oscillation in the Z axis. */
+	UPROPERTY(EditAnywhere, BlueprintReadWrite, Category = VOscillator)
+	struct FFOscillator Z;
+};
+
+
+/**
+ * Legacy camera shake which can do either oscillation or run camera anims.
+ */
+UCLASS(Blueprintable, HideCategories = (CameraShakePattern))
+class GAMEPLAYCAMERAS_API ULegacyCameraShake : public UCameraShakeBase
+{
+	GENERATED_UCLASS_BODY()
+
+public:
+	/** Duration in seconds of current screen shake. Less than 0 means indefinite, 0 means no oscillation. */
+	UPROPERTY(EditAnywhere, Category = Oscillation)
+	float OscillationDuration;
+
+	/** Duration of the blend-in, where the oscillation scales from 0 to 1. */
+	UPROPERTY(EditAnywhere, Category = Oscillation, meta = (ClampMin = "0.0"))
+	float OscillationBlendInTime;
+
+	/** Duration of the blend-out, where the oscillation scales from 1 to 0. */
+	UPROPERTY(EditAnywhere, Category = Oscillation, meta = (ClampMin = "0.0"))
+	float OscillationBlendOutTime;
+
+	/** Rotational oscillation */
+	UPROPERTY(EditAnywhere, BlueprintReadWrite, Category = Oscillation)
+	struct FROscillator RotOscillation;
+
+	/** Positional oscillation */
+	UPROPERTY(EditAnywhere, BlueprintReadWrite, Category = Oscillation)
+	struct FVOscillator LocOscillation;
+
+	/** FOV oscillation */
+	UPROPERTY(EditAnywhere, BlueprintReadWrite, Category = Oscillation)
+	struct FFOscillator FOVOscillation;
+
+	/************************************************************
+	 * Parameters for defining sequencer-driven camera shakes
+	 ************************************************************/
+
+	 /** Scalar defining how fast to play the anim. */
+	UPROPERTY(EditAnywhere, Category = AnimShake, meta = (ClampMin = "0.001"))
+	float AnimPlayRate;
+
+	/** Scalar defining how "intense" to play the anim. */
+	UPROPERTY(EditAnywhere, Category = AnimShake, meta = (ClampMin = "0.0"))
+	float AnimScale;
+
+	/** Linear blend-in time. */
+	UPROPERTY(EditAnywhere, Category = AnimShake, meta = (ClampMin = "0.0"))
+	float AnimBlendInTime;
+
+	/** Linear blend-out time. */
+	UPROPERTY(EditAnywhere, Category = AnimShake, meta = (ClampMin = "0.0"))
+	float AnimBlendOutTime;
+
+	/** When bRandomAnimSegment is true, this defines how long the anim should play. */
+	UPROPERTY(EditAnywhere, Category = AnimShake, meta = (ClampMin = "0.0", editcondition = "bRandomAnimSegment"))
+	float RandomAnimSegmentDuration;
+
+	/** Source camera animation sequence to play. Can be null. */
+	UPROPERTY(EditAnywhere, Category = AnimShake)
+	TObjectPtr<class UCameraAnimationSequence> AnimSequence;
+
+	/**
+	* If true, play a random snippet of the animation of length Duration.  Implies bLoop and bRandomStartTime = true for the AnimSequence.
+	* If false, play the full anim once, non-looped. Useful for getting variety out of a single looped AnimSequence asset.
+	*/
+	UPROPERTY(EditAnywhere, Category = AnimShake)
+	uint32 bRandomAnimSegment : 1;
+
+public:
+
+	/** Time remaining for oscillation shakes. Less than 0.f means shake infinitely. */
+	UPROPERTY(transient, BlueprintReadOnly, Category = CameraShake)
+	float OscillatorTimeRemaining;
+
+public:
+
+	// Blueprint API
+
+	/** Called when the shake starts playing */
+	UFUNCTION(BlueprintImplementableEvent, Category = CameraShake)
+	void ReceivePlayShake(float Scale);
+
+	/** Called every tick to let the shake modify the point of view */
+	UFUNCTION(BlueprintImplementableEvent, Category = CameraShake)
+	void BlueprintUpdateCameraShake(float DeltaTime, float Alpha, const FMinimalViewInfo& POV, FMinimalViewInfo& ModifiedPOV);
+
+	/** Called to allow a shake to decide when it's finished playing. */
+	UFUNCTION(BlueprintNativeEvent, Category = CameraShake)
+	bool ReceiveIsFinished() const;
+
+	/**
+	 * Called when the shake is explicitly stopped.
+	 * @param bImmediatly		If true, shake stops right away regardless of blend out settings. If false, shake may blend out according to its settings.
+	 */
+	UFUNCTION(BlueprintImplementableEvent, Category = CameraShake)
+	void ReceiveStopShake(bool bImmediately);
+
+public:
+
+	/**
+	 * Backwards compatible method used by core BP redirectors. This is needed because the return value is specifically a legacy camera shake,
+	 * which some BP logic often uses directly to set oscillator/anim properties.
+	 */
+	UFUNCTION(BlueprintCallable, Category = "Camera Shakes")
+	static ULegacyCameraShake* StartLegacyCameraShake(APlayerCameraManager* PlayerCameraManager, TSubclassOf<ULegacyCameraShake> ShakeClass, float Scale = 1.f, ECameraShakePlaySpace PlaySpace = ECameraShakePlaySpace::CameraLocal, FRotator UserPlaySpaceRot = FRotator::ZeroRotator);
+
+	/**
+	 * Backwards compatible method used by core BP redirectors. This is needed because the return value is specifically a legacy camera shake,
+	 * which some BP logic often uses directly to set oscillator/anim properties.
+	 */
+	UFUNCTION(BlueprintCallable, Category = "Camera Shakes")
+	static ULegacyCameraShake* StartLegacyCameraShakeFromSource(APlayerCameraManager* PlayerCameraManager, TSubclassOf<ULegacyCameraShake> ShakeClass, UCameraShakeSourceComponent* SourceComponent, float Scale = 1.f, ECameraShakePlaySpace PlaySpace = ECameraShakePlaySpace::CameraLocal, FRotator UserPlaySpaceRot = FRotator::ZeroRotator);
+
+public:
+
+	/** Returns true if this camera shake will loop forever */
+	bool IsLooping() const;
+
+	/** Sets current playback time and applies the shake (both oscillation and camera animation sequence) to the given POV. */
+	UE_DEPRECATED(4.27, "SetCurrentTimeAndApplyShake is deprecated, please use ScrubAndApplyCameraShake")
+	void SetCurrentTimeAndApplyShake(float NewTime, FMinimalViewInfo& POV);
+
+private:
+
+	void DoStartShake(const FCameraShakeStartParams& Params);
+	void DoUpdateShake(const FCameraShakeUpdateParams& Params, FCameraShakeUpdateResult& OutResult);
+	void DoScrubShake(const FCameraShakeScrubParams& Params, FCameraShakeUpdateResult& OutResult);
+	void DoStopShake(bool bImmediately);
+	bool DoGetIsFinished() const;
+	void DoTeardownShake();
+
+protected:
+
+	/** Current location sinusoidal offset. */
+	FVector LocSinOffset;
+
+	/** Current rotational sinusoidal offset. */
+	FVector RotSinOffset;
+
+	/** Current FOV sinusoidal offset. */
+	float FOVSinOffset;
+
+	/** Initial offset (could have been assigned at random). */
+	FVector InitialLocSinOffset;
+
+	/** Initial offset (could have been assigned at random). */
+	FVector InitialRotSinOffset;
+
+	/** Initial offset (could have been assigned at random). */
+	float InitialFOVSinOffset;
+
+	/** Sequence shake pattern for when using a sequence instead of a camera anim */
+	UPROPERTY(Instanced)
+		TObjectPtr<class USequenceCameraShakePattern> SequenceShakePattern;
+
+	/** State tracking for the sequence shake pattern */
+	FCameraShakeState SequenceShakeState;
+
+private:
+
+	float CurrentBlendInTime;
+	float CurrentBlendOutTime;
+	bool bBlendingIn : 1;
+	bool bBlendingOut : 1;
+
+	friend class ULegacyCameraShakePattern;
+};
+
+/**
+ * Shake pattern for the ULegacyCameraShake class.
+ *
+ * It doesn't do anything because, for backwards compatibility reasons, all the data
+ * was left on the shake class itself... so this pattern delegates everything back
+ * to the owner shake.
+ */
+UCLASS(HideDropdown)
+class GAMEPLAYCAMERAS_API ULegacyCameraShakePattern : public UCameraShakePattern
+{
+	GENERATED_BODY()
+
+private:
+
+	// UCameraShakePattern interface
+	virtual void GetShakePatternInfoImpl(FCameraShakeInfo& OutInfo) const override;
+	virtual void StartShakePatternImpl(const FCameraShakeStartParams& Params) override;
+	virtual void UpdateShakePatternImpl(const FCameraShakeUpdateParams& Params, FCameraShakeUpdateResult& OutResult) override;
+	virtual void ScrubShakePatternImpl(const FCameraShakeScrubParams& Params, FCameraShakeUpdateResult& OutResult) override;
+	virtual bool IsFinishedImpl() const override;
+	virtual void StopShakePatternImpl(const FCameraShakeStopParams& Params) override;
+	virtual void TeardownShakePatternImpl() override;
+};
+
+/** Backwards compatible name for the legacy camera shake, for C++ code. */
+UE_DEPRECATED(4.26, "Please use ULegacyCameraShake")
+typedef ULegacyCameraShake UCameraShake;
+
+UE_DEPRECATED(5.2, "Please use ULegacyCameraShake")
+typedef ULegacyCameraShake UMatineeCameraShake;
+
+/**
+ * Blueprint function library for autocasting from a base camera shake to a legacy camera shake.
+ * This prevents breaking Blueprints now that APlayerCameraManager::StartCameraShake returns the base class.
+ */
+UCLASS()
+class GAMEPLAYCAMERAS_API ULegacyCameraShakeFunctionLibrary : public UBlueprintFunctionLibrary
+{
+	GENERATED_BODY()
+
+public:
+	UFUNCTION(BlueprintPure, Category = "Camera Shakes", meta = (BlueprintAutocast))
+	static ULegacyCameraShake* Conv_LegacyCameraShake(UCameraShakeBase* CameraShake)
+	{
+		return CastChecked<ULegacyCameraShake>(CameraShake, ECastCheckedType::NullAllowed);
+	}
+};
diff --git a/Engine/Plugins/Cameras/GameplayCameras/Source/GameplayCameras/Public/MatineeCameraShake.h b/Engine/Plugins/Cameras/GameplayCameras/Source/GameplayCameras/Public/MatineeCameraShake.h
--- a/Engine/Plugins/Cameras/GameplayCameras/Source/GameplayCameras/Public/MatineeCameraShake.h
+++ b/Engine/Plugins/Cameras/GameplayCameras/Source/GameplayCameras/Public/MatineeCameraShake.h
@@ -2,334 +2,7 @@
 
 #pragma once
 
-#include "CoreMinimal.h"
-#include "Camera/CameraShakeBase.h"
-#include "Evaluation/MovieSceneCameraShakeTemplate.h"
-#include "Kismet/BlueprintFunctionLibrary.h"
-#include "UObject/Object.h"
-#include "GameFramework/Actor.h"
-#include "MatineeCameraShake.generated.h"
-
-class AActor;
-
-/************************************************************
- * Parameters for defining oscillating camera shakes
- ************************************************************/
-
-/** Types of waveforms that can be used for camera shake oscillators */
-UENUM(BlueprintType)
-enum class EOscillatorWaveform : uint8
-{
-	/** A sinusoidal wave */
-	SineWave,
-
-	/** Perlin noise */
-	PerlinNoise,
-};
-
-/** Shake start offset parameter */
-UENUM()
-enum EInitialOscillatorOffset
-{
-	/** Start with random offset (default). */
-	EOO_OffsetRandom UMETA(DisplayName = "Random"),
-	/** Start with zero offset. */
-	EOO_OffsetZero UMETA(DisplayName = "Zero"),
-	EOO_MAX,
-};
-
-/** Defines oscillation of a single number. */
-USTRUCT(BlueprintType)
-struct GAMEPLAYCAMERAS_API FFOscillator
-{
-	GENERATED_USTRUCT_BODY()
-
-	/** Amplitude of the sinusoidal oscillation. */
-	UPROPERTY(EditAnywhere, BlueprintReadWrite, Category=FOscillator)
-	float Amplitude;
-
-	/** Frequency of the sinusoidal oscillation. */
-	UPROPERTY(EditAnywhere, BlueprintReadWrite, Category = FOscillator)
-	float Frequency;
-
-	/** Defines how to begin (either at zero, or at a randomized value. */
-	UPROPERTY(EditAnywhere, Category=FOscillator)
-	TEnumAsByte<enum EInitialOscillatorOffset> InitialOffset;
-	
-	/** Type of waveform to use for oscillation. */
-	UPROPERTY(EditAnywhere, BlueprintReadWrite, Category=FOscillator)
-	EOscillatorWaveform Waveform;
-
-	FFOscillator()
-		: Amplitude(0)
-		, Frequency(0)
-		, InitialOffset(0)
-		, Waveform( EOscillatorWaveform::SineWave ) 
-	{}
-
-	/** Advances the oscillation time and returns the current value. */
-	static float UpdateOffset(FFOscillator const& Osc, float& CurrentOffset, float DeltaTime);
-
-	/** Advances the oscillation time and returns the current value. */
-	static float UpdateOffset(FFOscillator const& Osc, double& CurrentOffset, float DeltaTime);
-
-	/** Returns the initial value of the oscillator. */
-	static float GetInitialOffset(FFOscillator const& Osc);
-
-	/** Returns the offset at the given time */
-	static float GetOffsetAtTime(FFOscillator const& Osc, float InitialOffset, float Time);
-};
-
-/** Defines FRotator oscillation. */
-USTRUCT(BlueprintType)
-struct GAMEPLAYCAMERAS_API FROscillator
-{
-	GENERATED_USTRUCT_BODY()
-
-	/** Pitch oscillation. */
-	UPROPERTY(EditAnywhere, BlueprintReadWrite, Category=ROscillator)
-	struct FFOscillator Pitch;
-
-	/** Yaw oscillation. */
-	UPROPERTY(EditAnywhere, BlueprintReadWrite, Category = ROscillator)
-	struct FFOscillator Yaw;
-
-	/** Roll oscillation. */
-	UPROPERTY(EditAnywhere, BlueprintReadWrite, Category = ROscillator)
-	struct FFOscillator Roll;
-
-};
-
-/** Defines FVector oscillation. */
-USTRUCT(BlueprintType)
-struct FVOscillator
-{
-	GENERATED_USTRUCT_BODY()
-
-	/** Oscillation in the X axis. */
-	UPROPERTY(EditAnywhere, BlueprintReadWrite, Category = VOscillator)
-	struct FFOscillator X;
-
-	/** Oscillation in the Y axis. */
-	UPROPERTY(EditAnywhere, BlueprintReadWrite, Category = VOscillator)
-	struct FFOscillator Y;
-
-	/** Oscillation in the Z axis. */
-	UPROPERTY(EditAnywhere, BlueprintReadWrite, Category = VOscillator)
-	struct FFOscillator Z;
-};
-
-
-/**
- * Legacy camera shake which can do either oscillation or run camera anims.
- */
-UCLASS(Blueprintable, HideCategories=(CameraShakePattern))
-class GAMEPLAYCAMERAS_API UMatineeCameraShake : public UCameraShakeBase
-{
-	GENERATED_UCLASS_BODY()
-
-public:
-	/** Duration in seconds of current screen shake. Less than 0 means indefinite, 0 means no oscillation. */
-	UPROPERTY(EditAnywhere, Category=Oscillation)
-	float OscillationDuration;
-
-	/** Duration of the blend-in, where the oscillation scales from 0 to 1. */
-	UPROPERTY(EditAnywhere, Category=Oscillation, meta=(ClampMin = "0.0"))
-	float OscillationBlendInTime;
-
-	/** Duration of the blend-out, where the oscillation scales from 1 to 0. */
-	UPROPERTY(EditAnywhere, Category=Oscillation, meta = (ClampMin = "0.0"))
-	float OscillationBlendOutTime;
-
-	/** Rotational oscillation */
-	UPROPERTY(EditAnywhere, BlueprintReadWrite, Category=Oscillation)
-	struct FROscillator RotOscillation;
-
-	/** Positional oscillation */
-	UPROPERTY(EditAnywhere, BlueprintReadWrite, Category=Oscillation)
-	struct FVOscillator LocOscillation;
-
-	/** FOV oscillation */
-	UPROPERTY(EditAnywhere, BlueprintReadWrite, Category=Oscillation)
-	struct FFOscillator FOVOscillation;
-
-	/************************************************************
-	 * Parameters for defining sequencer-driven camera shakes
-	 ************************************************************/
-
-	/** Scalar defining how fast to play the anim. */
-	UPROPERTY(EditAnywhere, Category=AnimShake, meta=(ClampMin = "0.001"))
-	float AnimPlayRate;
-
-	/** Scalar defining how "intense" to play the anim. */
-	UPROPERTY(EditAnywhere, Category=AnimShake, meta=(ClampMin = "0.0"))
-	float AnimScale;
-
-	/** Linear blend-in time. */
-	UPROPERTY(EditAnywhere, Category=AnimShake, meta=(ClampMin = "0.0"))
-	float AnimBlendInTime;
-
-	/** Linear blend-out time. */
-	UPROPERTY(EditAnywhere, Category=AnimShake, meta=(ClampMin = "0.0"))
-	float AnimBlendOutTime;
-
-	/** When bRandomAnimSegment is true, this defines how long the anim should play. */
-	UPROPERTY(EditAnywhere, Category = AnimShake, meta = (ClampMin = "0.0", editcondition = "bRandomAnimSegment"))
-	float RandomAnimSegmentDuration;
-
-	/** Source camera animation sequence to play. Can be null. */
-	UPROPERTY(EditAnywhere, Category = AnimShake)
-	TObjectPtr<class UCameraAnimationSequence> AnimSequence;
-
-	/**
-	* If true, play a random snippet of the animation of length Duration.  Implies bLoop and bRandomStartTime = true for the AnimSequence.
-	* If false, play the full anim once, non-looped. Useful for getting variety out of a single looped AnimSequence asset.
-	*/
-	UPROPERTY(EditAnywhere, Category = AnimShake)
-	uint32 bRandomAnimSegment : 1;
-
-public:
-
-	/** Time remaining for oscillation shakes. Less than 0.f means shake infinitely. */
-	UPROPERTY(transient, BlueprintReadOnly, Category = CameraShake)
-	float OscillatorTimeRemaining;
-
-public:
-
-	// Blueprint API
-
-	/** Called when the shake starts playing */
-	UFUNCTION(BlueprintImplementableEvent, Category = CameraShake)
-	void ReceivePlayShake(float Scale);
-
-	/** Called every tick to let the shake modify the point of view */
-	UFUNCTION(BlueprintImplementableEvent, Category = CameraShake)
-	void BlueprintUpdateCameraShake(float DeltaTime, float Alpha, const FMinimalViewInfo& POV, FMinimalViewInfo& ModifiedPOV);
-
-	/** Called to allow a shake to decide when it's finished playing. */
-	UFUNCTION(BlueprintNativeEvent, Category = CameraShake)
-	bool ReceiveIsFinished() const;
-
-	/**
-	 * Called when the shake is explicitly stopped.
-	 * @param bImmediatly		If true, shake stops right away regardless of blend out settings. If false, shake may blend out according to its settings.
-	 */
-	UFUNCTION(BlueprintImplementableEvent, Category = CameraShake)
-	void ReceiveStopShake(bool bImmediately);
-
-public:
-
-	/**
-	 * Backwards compatible method used by core BP redirectors. This is needed because the return value is specifically a Matinee camera shake,
-	 * which some BP logic often uses directly to set oscillator/anim properties.
-	 */
-	UFUNCTION(BlueprintCallable, Category = "Camera Shakes")
-	static UMatineeCameraShake* StartMatineeCameraShake(APlayerCameraManager* PlayerCameraManager, TSubclassOf<UMatineeCameraShake> ShakeClass, float Scale = 1.f, ECameraShakePlaySpace PlaySpace = ECameraShakePlaySpace::CameraLocal, FRotator UserPlaySpaceRot = FRotator::ZeroRotator);
-
-	/**
-	 * Backwards compatible method used by core BP redirectors. This is needed because the return value is specifically a Matinee camera shake,
-	 * which some BP logic often uses directly to set oscillator/anim properties.
-	 */
-	UFUNCTION(BlueprintCallable, Category = "Camera Shakes")
-	static UMatineeCameraShake* StartMatineeCameraShakeFromSource(APlayerCameraManager* PlayerCameraManager, TSubclassOf<UMatineeCameraShake> ShakeClass, UCameraShakeSourceComponent* SourceComponent, float Scale = 1.f, ECameraShakePlaySpace PlaySpace = ECameraShakePlaySpace::CameraLocal, FRotator UserPlaySpaceRot = FRotator::ZeroRotator);
-
-public:
-
-	/** Returns true if this camera shake will loop forever */
-	bool IsLooping() const;
-
-	/** Sets current playback time and applies the shake (both oscillation and camera animation sequence) to the given POV. */
-	UE_DEPRECATED(4.27, "SetCurrentTimeAndApplyShake is deprecated, please use ScrubAndApplyCameraShake")
-	void SetCurrentTimeAndApplyShake(float NewTime, FMinimalViewInfo& POV);
-
-private:
-
-	void DoStartShake(const FCameraShakeStartParams& Params);
-	void DoUpdateShake(const FCameraShakeUpdateParams& Params, FCameraShakeUpdateResult& OutResult);
-	void DoScrubShake(const FCameraShakeScrubParams& Params, FCameraShakeUpdateResult& OutResult);
-	void DoStopShake(bool bImmediately);
-	bool DoGetIsFinished() const;
-	void DoTeardownShake();
-
-protected:
-
-	/** Current location sinusoidal offset. */
-	FVector LocSinOffset;
-
-	/** Current rotational sinusoidal offset. */
-	FVector RotSinOffset;
-
-	/** Current FOV sinusoidal offset. */
-	float FOVSinOffset;
-
-	/** Initial offset (could have been assigned at random). */
-	FVector InitialLocSinOffset;
-
-	/** Initial offset (could have been assigned at random). */
-	FVector InitialRotSinOffset;
-
-	/** Initial offset (could have been assigned at random). */
-	float InitialFOVSinOffset;
-
-	/** Sequence shake pattern for when using a sequence instead of a camera anim */
-	UPROPERTY(Instanced)
-	TObjectPtr<class USequenceCameraShakePattern> SequenceShakePattern;
-
-	/** State tracking for the sequence shake pattern */
-	FCameraShakeState SequenceShakeState;
-
-private:
-
-	float CurrentBlendInTime;
-	float CurrentBlendOutTime;
-	bool bBlendingIn : 1;
-	bool bBlendingOut : 1;
-
-	friend class UMatineeCameraShakePattern;
-};
-
-/**
- * Shake pattern for the UMatineeCameraShake class.
- *
- * It doesn't do anything because, for backwards compatibility reasons, all the data
- * was left on the shake class itself... so this pattern delegates everything back
- * to the owner shake.
- */
-UCLASS(HideDropdown)
-class GAMEPLAYCAMERAS_API UMatineeCameraShakePattern : public UCameraShakePattern
-{
-	GENERATED_BODY()
-
-private:
-
-	// UCameraShakePattern interface
-	virtual void GetShakePatternInfoImpl(FCameraShakeInfo& OutInfo) const override;
-	virtual void StartShakePatternImpl(const FCameraShakeStartParams& Params) override;
-	virtual void UpdateShakePatternImpl(const FCameraShakeUpdateParams& Params, FCameraShakeUpdateResult& OutResult) override;
-	virtual void ScrubShakePatternImpl(const FCameraShakeScrubParams& Params, FCameraShakeUpdateResult& OutResult) override;
-	virtual bool IsFinishedImpl() const override;
-	virtual void StopShakePatternImpl(const FCameraShakeStopParams& Params) override;
-	virtual void TeardownShakePatternImpl() override;
-};
-
-/** Backwards compatible name for the Matinee camera shake, for C++ code. */
-UE_DEPRECATED(4.26, "Please use UMatineeCameraShake")
-typedef UMatineeCameraShake UCameraShake;
-
-/**
- * Blueprint function library for autocasting from a base camera shake to a matinee camera shake.
- * This prevents breaking Blueprints now that APlayerCameraManager::StartCameraShake returns the base class.
- */
-UCLASS()
-class GAMEPLAYCAMERAS_API UMatineeCameraShakeFunctionLibrary : public UBlueprintFunctionLibrary
-{
-	GENERATED_BODY()
-
-public:
-	UFUNCTION(BlueprintPure, Category="Camera Shakes", meta=(BlueprintAutocast))
-	static UMatineeCameraShake* Conv_MatineeCameraShake(UCameraShakeBase* CameraShake)
-	{
-		return CastChecked<UMatineeCameraShake>(CameraShake, ECastCheckedType::NullAllowed);
-	}
-};
+#include "CoreTypes.h"
 
+COMPILE_WARNING("UMatineeCameraShake has been renamed to ULegacyCameraShake, change to #include \"LegacyCameraShake.h\" and rename references")
+#include "LegacyCameraShake.h"
