Fix deprecation version number

diff --git a/Engine/Plugins/Cameras/GameplayCameras/Source/GameplayCameras/Public/LegacyCameraShake.h b/Engine/Plugins/Cameras/GameplayCameras/Source/GameplayCameras/Public/LegacyCameraShake.h
--- a/Engine/Plugins/Cameras/GameplayCameras/Source/GameplayCameras/Public/LegacyCameraShake.h
+++ b/Engine/Plugins/Cameras/GameplayCameras/Source/GameplayCameras/Public/LegacyCameraShake.h
@@ -316,7 +316,7 @@ private:
 UE_DEPRECATED(4.26, "Please use ULegacyCameraShake")
 typedef ULegacyCameraShake UCameraShake;
 
-UE_DEPRECATED(5.2, "Please use ULegacyCameraShake")
+UE_DEPRECATED(5.1, "Please use ULegacyCameraShake")
 typedef ULegacyCameraShake UMatineeCameraShake;
 
 /**
