Cameras: don't invalidate ongoing camera animation indices when stopping one of them

Instead, just clear the info struct so that it has an invalid serial and player, similar to DeactivateCameraAnimation.

diff --git a/Engine/Plugins/Cameras/GameplayCameras/Source/GameplayCameras/Private/CameraAnimationCameraModifier.cpp b/Engine/Plugins/Cameras/GameplayCameras/Source/GameplayCameras/Private/CameraAnimationCameraModifier.cpp
--- a/Engine/Plugins/Cameras/GameplayCameras/Source/GameplayCameras/Private/CameraAnimationCameraModifier.cpp
+++ b/Engine/Plugins/Cameras/GameplayCameras/Source/GameplayCameras/Private/CameraAnimationCameraModifier.cpp
@@ -146,7 +146,7 @@ void UCameraAnimationCameraModifier::StopCameraAnimation(const FCameraAnimationH
 		if (bImmediate || CameraAnimation->Params.EaseOutDuration == 0.f)
 		{
 			CameraAnimation->Player->Stop();
-			ActiveAnimations.RemoveAt(CameraAnimation->Handle.InstanceID);
+			ActiveAnimations[CameraAnimation->Handle.InstanceID] = FActiveCameraAnimationInfo();
 		}
 		else if (!CameraAnimation->bIsEasingOut)
 		{
