Cameras: fix invalid invalid handle, due to a typo between int16 and uint16

It previously worked out because the error was in the right places... but it could have resulted in bugs when comparing a *default* handle with an *invalid* handle.

diff --git a/Engine/Plugins/Cameras/GameplayCameras/Source/GameplayCameras/Private/CameraAnimationCameraModifier.cpp b/Engine/Plugins/Cameras/GameplayCameras/Source/GameplayCameras/Private/CameraAnimationCameraModifier.cpp
--- a/Engine/Plugins/Cameras/GameplayCameras/Source/GameplayCameras/Private/CameraAnimationCameraModifier.cpp
+++ b/Engine/Plugins/Cameras/GameplayCameras/Source/GameplayCameras/Private/CameraAnimationCameraModifier.cpp
@@ -20,7 +20,7 @@
 
 DECLARE_CYCLE_STAT(TEXT("Camera Animation Eval"), CameraAnimationEval_Total, STATGROUP_CameraAnimation);
 
-FCameraAnimationHandle FCameraAnimationHandle::Invalid(MAX_int16, 0);
+FCameraAnimationHandle FCameraAnimationHandle::Invalid(MAX_uint16, 0);
 
 FActiveCameraAnimationInfo::FActiveCameraAnimationInfo()
 	: Sequence(nullptr)
diff --git a/Engine/Plugins/Cameras/GameplayCameras/Source/GameplayCameras/Public/CameraAnimationCameraModifier.h b/Engine/Plugins/Cameras/GameplayCameras/Source/GameplayCameras/Public/CameraAnimationCameraModifier.h
--- a/Engine/Plugins/Cameras/GameplayCameras/Source/GameplayCameras/Public/CameraAnimationCameraModifier.h
+++ b/Engine/Plugins/Cameras/GameplayCameras/Source/GameplayCameras/Public/CameraAnimationCameraModifier.h
@@ -99,7 +99,7 @@ struct FCameraAnimationHandle
 
 	bool IsValid() const
 	{
-		return InstanceID != MAX_int16;
+		return InstanceID != MAX_uint16;
 	}
 
 	friend bool operator==(FCameraAnimationHandle A, FCameraAnimationHandle B)
