Camera Shake: Check for valid ChildStates before stopping. This fixes a crash where Stop shake is called before evaluation

Return if no sequence has been defined in SequencerCameraShake
Get max duration for composite shakes

diff --git a/Engine/Plugins/Cameras/GameplayCameras/Source/GameplayCameras/Private/CompositeCameraShakePattern.cpp b/Engine/Plugins/Cameras/GameplayCameras/Source/GameplayCameras/Private/CompositeCameraShakePattern.cpp
--- a/Engine/Plugins/Cameras/GameplayCameras/Source/GameplayCameras/Private/CompositeCameraShakePattern.cpp
+++ b/Engine/Plugins/Cameras/GameplayCameras/Source/GameplayCameras/Private/CompositeCameraShakePattern.cpp
@@ -6,15 +6,26 @@
 
 void UCompositeCameraShakePattern::GetShakePatternInfoImpl(FCameraShakeInfo& OutInfo) const
 {
-	// We will manage our duration ourselves.
-	//
-	// While we could return the maximum duration of all our children, the problem lies with
-	// the Stop method, when bImmediately is false: it means we have to move our current time
-	// to the beginning of the longest blend-out among our children. We could also set the
-	// blend in/out times to the maximum blend time of all our children, but frankly we don't
-	// need the make the base class to compute all kinds of stuff we don't need anyway. So
-	// it's simpler to use a custom duration and get on with it.
 	OutInfo.Duration = FCameraShakeDuration::Custom();
+
+	for (UCameraShakePattern* Pattern : ChildPatterns)
+	{
+		if (Pattern != nullptr)
+		{
+			FCameraShakeInfo ChildInfo;
+			Pattern->GetShakePatternInfo(ChildInfo);
+
+			if (ChildInfo.Duration.IsInfinite())
+			{
+				OutInfo = ChildInfo;
+				break;
+			}
+			else
+			{
+				OutInfo.Duration = FMath::Max(ChildInfo.Duration.Get(), OutInfo.Duration.Get());
+			}
+		}
+	}
 }
 
 void UCompositeCameraShakePattern::StartShakePatternImpl(const FCameraShakeStartParams& Params)
@@ -175,12 +186,15 @@ bool UCompositeCameraShakePattern::IsFinishedImpl() const
 	for (uint32 Index = 0, Num = ChildPatterns.Num(); Index < Num; ++Index)
 	{
 		const UCameraShakePattern* Pattern = ChildPatterns[Index];
-		const FCameraShakeState& PatternState = ChildStates[Index];
-
-		const bool bIsChildFinished = IsChildPatternFinished(PatternState, Pattern);
-		if (!bIsChildFinished)
+		if (Index < (uint32)ChildStates.Num())
 		{
-			return false;
+			const FCameraShakeState& PatternState = ChildStates[Index];
+
+			const bool bIsChildFinished = IsChildPatternFinished(PatternState, Pattern);
+			if (!bIsChildFinished)
+			{
+				return false;
+			}
 		}
 	}
 	return true;
@@ -192,15 +206,18 @@ void UCompositeCameraShakePattern::StopShakePatternImpl(const FCameraShakeStopPa
 	for (uint32 Index = 0, Num = ChildPatterns.Num(); Index < Num; ++Index)
 	{
 		UCameraShakePattern* Pattern = ChildPatterns[Index];
-		FCameraShakeState& PatternState = ChildStates[Index];
-
-		if (PatternState.IsActive())
+		if (Index < (uint32)ChildStates.Num())
 		{
-			PatternState.Stop(Params.bImmediately);
+			FCameraShakeState& PatternState = ChildStates[Index];
 
-			if (Pattern != nullptr)
+			if (PatternState.IsActive())
 			{
-				Pattern->StopShakePattern(Params);
+				PatternState.Stop(Params.bImmediately);
+
+				if (Pattern != nullptr)
+				{
+					Pattern->StopShakePattern(Params);
+				}
 			}
 		}
 	}
diff --git a/Engine/Plugins/MovieScene/TemplateSequence/Source/TemplateSequence/Private/SequenceCameraShake.cpp b/Engine/Plugins/MovieScene/TemplateSequence/Source/TemplateSequence/Private/SequenceCameraShake.cpp
--- a/Engine/Plugins/MovieScene/TemplateSequence/Source/TemplateSequence/Private/SequenceCameraShake.cpp
+++ b/Engine/Plugins/MovieScene/TemplateSequence/Source/TemplateSequence/Private/SequenceCameraShake.cpp
@@ -127,6 +127,11 @@ void USequenceCameraShakePattern::TeardownShakePatternImpl()
 
 void USequenceCameraShakePattern::UpdateCamera(FFrameTime NewPosition, const FMinimalViewInfo& InPOV, FCameraShakeUpdateResult& OutResult)
 {
+	if (!ensure(Sequence))
+	{
+		return;
+	}
+
 	using namespace UE::MovieScene;
 
 	check(CameraStandIn);
