Camera Shake: Check for valid ChildStates

diff --git a/Engine/Plugins/Cameras/GameplayCameras/Source/GameplayCameras/Private/CompositeCameraShakePattern.cpp b/Engine/Plugins/Cameras/GameplayCameras/Source/GameplayCameras/Private/CompositeCameraShakePattern.cpp
--- a/Engine/Plugins/Cameras/GameplayCameras/Source/GameplayCameras/Private/CompositeCameraShakePattern.cpp
+++ b/Engine/Plugins/Cameras/GameplayCameras/Source/GameplayCameras/Private/CompositeCameraShakePattern.cpp
@@ -61,39 +61,42 @@ void UCompositeCameraShakePattern::UpdateShakePatternImpl(const FCameraShakeUpda
 
 	for (uint32 Index = 0, Num = ChildPatterns.Num(); Index < Num; ++Index)
 	{
-		FCameraShakeState& PatternState = ChildStates[Index];
-		UCameraShakePattern* Pattern = ChildPatterns[Index];
-		if (Pattern != nullptr)
+		if (Index < (uint32)ChildStates.Num())
 		{
-			// Update the child state.
-			float ChildBlendingWeight = PatternState.Update(Params.DeltaTime);
-			if (!PatternState.IsActive())
+			FCameraShakeState& PatternState = ChildStates[Index];
+			UCameraShakePattern* Pattern = ChildPatterns[Index];
+			if (Pattern != nullptr)
 			{
-				continue;
-			}
+				// Update the child state.
+				float ChildBlendingWeight = PatternState.Update(Params.DeltaTime);
+				if (!PatternState.IsActive())
+				{
+					continue;
+				}
 
-			// Let the child pattern run on the current result, with its own blending weight.
-			ChildParams.BlendingWeight = Params.BlendingWeight * ChildBlendingWeight;
+				// Let the child pattern run on the current result, with its own blending weight.
+				ChildParams.BlendingWeight = Params.BlendingWeight * ChildBlendingWeight;
 
-			FCameraShakeUpdateResult ChildResult;
+				FCameraShakeUpdateResult ChildResult;
 
-			Pattern->UpdateShakePattern(ChildParams, ChildResult);
+				Pattern->UpdateShakePattern(ChildParams, ChildResult);
 
-			if (IsChildPatternFinished(PatternState, Pattern))
-			{
-				// This pattern just ended now. Reset its state and move on to the next.
-				ChildStates[Index] = FCameraShakeState();
-				continue;
-			}
+				if (IsChildPatternFinished(PatternState, Pattern))
+				{
+					// This pattern just ended now. Reset its state and move on to the next.
+					ChildStates[Index] = FCameraShakeState();
+					continue;
+				}
 			
-			// Apply this result and pass it on to the next child pattern.
-			FCameraShakeApplyResultParams ApplyParams;
-			ApplyParams.Scale = ChildParams.GetTotalScale();
-			ApplyParams.PlaySpace = ShakeInstance->GetPlaySpace();
-			ApplyParams.UserPlaySpaceMatrix = ShakeInstance->GetUserPlaySpaceMatrix();
-			// This applies the current pattern's update to the parameters we'll pass
-			// to the next one.
-			UCameraShakeBase::ApplyResult(ApplyParams, ChildResult, ChildParams.POV);
+				// Apply this result and pass it on to the next child pattern.
+				FCameraShakeApplyResultParams ApplyParams;
+				ApplyParams.Scale = ChildParams.GetTotalScale();
+				ApplyParams.PlaySpace = ShakeInstance->GetPlaySpace();
+				ApplyParams.UserPlaySpaceMatrix = ShakeInstance->GetUserPlaySpaceMatrix();
+				// This applies the current pattern's update to the parameters we'll pass
+				// to the next one.
+				UCameraShakeBase::ApplyResult(ApplyParams, ChildResult, ChildParams.POV);
+			}
 		}
 	}
 
@@ -118,41 +121,44 @@ void UCompositeCameraShakePattern::ScrubShakePatternImpl(const FCameraShakeScrub
 
 	for (uint32 Index = 0, Num = ChildPatterns.Num(); Index < Num; ++Index)
 	{
-		FCameraShakeState& PatternState = ChildStates[Index];
-		if (!PatternState.IsActive())
+		if (Index < (uint32)ChildStates.Num())
 		{
-			// This pattern was finished for some time already.
-			continue;
-		}
+			FCameraShakeState& PatternState = ChildStates[Index];
+			if (!PatternState.IsActive())
+			{
+				// This pattern was finished for some time already.
+				continue;
+			}
 
-		UCameraShakePattern* Pattern = ChildPatterns[Index];
-		if (Pattern != nullptr)
-		{
-			// Scrub the child state.
-			float ChildBlendingWeight = PatternState.Scrub(Params.AbsoluteTime);
+			UCameraShakePattern* Pattern = ChildPatterns[Index];
+			if (Pattern != nullptr)
+			{
+				// Scrub the child state.
+				float ChildBlendingWeight = PatternState.Scrub(Params.AbsoluteTime);
 
-			// Let the child pattern run on the current result, with its own blending weight.
-			ChildParams.BlendingWeight = Params.BlendingWeight * ChildBlendingWeight;
+				// Let the child pattern run on the current result, with its own blending weight.
+				ChildParams.BlendingWeight = Params.BlendingWeight * ChildBlendingWeight;
 
-			FCameraShakeUpdateResult ChildResult;
+				FCameraShakeUpdateResult ChildResult;
 
-			Pattern->ScrubShakePattern(ChildParams, ChildResult);
+				Pattern->ScrubShakePattern(ChildParams, ChildResult);
 
-			if (IsChildPatternFinished(PatternState, Pattern))
-			{
-				// This pattern just ended now. Reset its state and move on to the next.
-				ChildStates[Index] = FCameraShakeState();
-				continue;
-			}
+				if (IsChildPatternFinished(PatternState, Pattern))
+				{
+					// This pattern just ended now. Reset its state and move on to the next.
+					ChildStates[Index] = FCameraShakeState();
+					continue;
+				}
 			
-			// Apply this result and pass it on to the next child pattern.
-			FCameraShakeApplyResultParams ApplyParams;
-			ApplyParams.Scale = ChildParams.GetTotalScale();
-			ApplyParams.PlaySpace = ShakeInstance->GetPlaySpace();
-			ApplyParams.UserPlaySpaceMatrix = ShakeInstance->GetUserPlaySpaceMatrix();
-			// This applies the current pattern's scrubbing to the parameters we'll pass
-			// to the next one.
-			UCameraShakeBase::ApplyResult(ApplyParams, ChildResult, ChildParams.POV);
+				// Apply this result and pass it on to the next child pattern.
+				FCameraShakeApplyResultParams ApplyParams;
+				ApplyParams.Scale = ChildParams.GetTotalScale();
+				ApplyParams.PlaySpace = ShakeInstance->GetPlaySpace();
+				ApplyParams.UserPlaySpaceMatrix = ShakeInstance->GetUserPlaySpaceMatrix();
+				// This applies the current pattern's scrubbing to the parameters we'll pass
+				// to the next one.
+				UCameraShakeBase::ApplyResult(ApplyParams, ChildResult, ChildParams.POV);
+			}
 		}
 	}
 
