Cameras: fix composite shake duration info

diff --git a/Engine/Plugins/Cameras/GameplayCameras/Source/GameplayCameras/Private/CompositeCameraShakePattern.cpp b/Engine/Plugins/Cameras/GameplayCameras/Source/GameplayCameras/Private/CompositeCameraShakePattern.cpp
--- a/Engine/Plugins/Cameras/GameplayCameras/Source/GameplayCameras/Private/CompositeCameraShakePattern.cpp
+++ b/Engine/Plugins/Cameras/GameplayCameras/Source/GameplayCameras/Private/CompositeCameraShakePattern.cpp
@@ -6,6 +6,9 @@
 
 void UCompositeCameraShakePattern::GetShakePatternInfoImpl(FCameraShakeInfo& OutInfo) const
 {
+	// Set the duration to custom (since we want to handle timing ourselves),
+	// but let's set a hint duration e.g. for when we are shown as a clip in
+	// a sequence. The hint duration is the max of all our children's durations.
 	OutInfo.Duration = FCameraShakeDuration::Custom();
 
 	for (UCameraShakePattern* Pattern : ChildPatterns)
@@ -17,12 +20,16 @@ void UCompositeCameraShakePattern::GetShakePatternInfoImpl(FCameraShakeInfo& Out
 
 			if (ChildInfo.Duration.IsInfinite())
 			{
-				OutInfo = ChildInfo;
+				// If one of our children is infinite, we are infinite.
+				OutInfo.Duration = ChildInfo.Duration;
 				break;
 			}
 			else
 			{
-				OutInfo.Duration = FMath::Max(ChildInfo.Duration.Get(), OutInfo.Duration.Get());
+				// If a child has a fixed duration, or custom duration with hint,
+				// let's include that in our own hint.
+				OutInfo.Duration = FCameraShakeDuration::Custom(
+					FMath::Max(ChildInfo.Duration.Get(), OutInfo.Duration.Get()));
 			}
 		}
 	}
