Cameras: treat post-process settings from camera shakes and camera animations as overrides, rather than base world values.

Before, a camera shake's or camera animation's PP settings would only be visible if the main camera wasn't already setting a value in them. Now, they will correctly override the main camera's settings.

diff --git a/Engine/Plugins/Cameras/GameplayCameras/Source/GameplayCameras/Private/CameraAnimationCameraModifier.cpp b/Engine/Plugins/Cameras/GameplayCameras/Source/GameplayCameras/Private/CameraAnimationCameraModifier.cpp
--- a/Engine/Plugins/Cameras/GameplayCameras/Source/GameplayCameras/Private/CameraAnimationCameraModifier.cpp
+++ b/Engine/Plugins/Cameras/GameplayCameras/Source/GameplayCameras/Private/CameraAnimationCameraModifier.cpp
@@ -14,6 +14,11 @@
 
 DECLARE_CYCLE_STAT(TEXT("Camera Animation Eval"), CameraAnimationEval_Total, STATGROUP_CameraAnimation);
 
+TAutoConsoleVariable<bool> GCameraAnimationLegacyPostProcessBlending(
+	TEXT("r.CameraAnimation.LegacyPostProcessBlending"),
+	false,
+	TEXT("Blend camera animation post process settings under the main camera instead of over it"));
+
 FCameraAnimationHandle FCameraAnimationHandle::Invalid(MAX_uint16, 0);
 
 FActiveCameraAnimationInfo::FActiveCameraAnimationInfo()
@@ -427,7 +432,8 @@ void UCameraAnimationCameraModifier::TickAnimation(FActiveCameraAnimationInfo& C
 	// Add the post-process settings.
 	if (CameraOwner != nullptr && CameraStandIn->PostProcessBlendWeight > 0.f)
 	{
-		CameraOwner->AddCachedPPBlend(CameraStandIn->PostProcessSettings, CameraStandIn->PostProcessBlendWeight);
+		EViewTargetBlendOrder CameraShakeBlendOrder = GCameraAnimationLegacyPostProcessBlending.GetValueOnGameThread() ? VTBlendOrder_Base : VTBlendOrder_Override;
+		CameraOwner->AddCachedPPBlend(CameraStandIn->PostProcessSettings, CameraStandIn->PostProcessBlendWeight, CameraShakeBlendOrder);
 	}
 }
 
diff --git a/Engine/Source/Runtime/Engine/Private/Camera/CameraShakeBase.cpp b/Engine/Source/Runtime/Engine/Private/Camera/CameraShakeBase.cpp
--- a/Engine/Source/Runtime/Engine/Private/Camera/CameraShakeBase.cpp
+++ b/Engine/Source/Runtime/Engine/Private/Camera/CameraShakeBase.cpp
@@ -13,6 +13,11 @@
 DECLARE_CYCLE_STAT(TEXT("CameraShakeStartShake"), STAT_StartShake, STATGROUP_Game);
 DECLARE_CYCLE_STAT(TEXT("CameraShakeUpdateShake"), STAT_UpdateShake, STATGROUP_Game);
 
+TAutoConsoleVariable<bool> GCameraShakeLegacyPostProcessBlending(
+	TEXT("r.CameraShake.LegacyPostProcessBlending"),
+	false,
+	TEXT("Blend camera shake post process settings under the main camera instead of over it"));
+
 UCameraShakeBase::UCameraShakeBase(const FObjectInitializer& ObjectInitializer)
 	: Super(ObjectInitializer)
 	, bSingleInstance(false)
@@ -404,7 +409,8 @@ void UCameraShakeBase::ApplyResult(const FCameraShakeApplyResultParams& ApplyPar
 	// It's weird but the post-process settings go directly on the camera manager, not on the view info.
 	if (ApplyParams.CameraManager.IsValid() && TempResult.PostProcessBlendWeight > 0.f)
 	{
-		ApplyParams.CameraManager->AddCachedPPBlend(TempResult.PostProcessSettings, TempResult.PostProcessBlendWeight);
+		EViewTargetBlendOrder CameraShakeBlendOrder = GCameraShakeLegacyPostProcessBlending.GetValueOnGameThread() ? VTBlendOrder_Base : VTBlendOrder_Override;
+		ApplyParams.CameraManager->AddCachedPPBlend(TempResult.PostProcessSettings, TempResult.PostProcessBlendWeight, CameraShakeBlendOrder);
 	}
 }
 
