Cameras: expose camera animation parameter struct to detail views

diff --git a/Engine/Plugins/Cameras/GameplayCameras/Source/GameplayCameras/Public/CameraAnimationCameraModifier.h b/Engine/Plugins/Cameras/GameplayCameras/Source/GameplayCameras/Public/CameraAnimationCameraModifier.h
--- a/Engine/Plugins/Cameras/GameplayCameras/Source/GameplayCameras/Public/CameraAnimationCameraModifier.h
+++ b/Engine/Plugins/Cameras/GameplayCameras/Source/GameplayCameras/Public/CameraAnimationCameraModifier.h
@@ -44,44 +44,44 @@ struct FCameraAnimationParams
 	GENERATED_BODY()
 
 	/** Time scale for playing the new camera animation */
-	UPROPERTY(BlueprintReadWrite, Category="Camera Animation")
+	UPROPERTY(EditAnywhere, BlueprintReadWrite, Category="Camera Animation")
 	float PlayRate = 1.f;
 	/** Global scale to use for the new camera animation */
-	UPROPERTY(BlueprintReadWrite, Category="Camera Animation")
+	UPROPERTY(EditAnywhere, BlueprintReadWrite, Category="Camera Animation")
 	float Scale = 1.f;
 
 	/** Ease-in function type */
-	UPROPERTY(BlueprintReadWrite, Category="Camera Animation")
+	UPROPERTY(EditAnywhere, BlueprintReadWrite, Category="Camera Animation")
 	ECameraAnimationEasingType EaseInType = ECameraAnimationEasingType::Linear;
 	/** Ease-in duration in seconds */
-	UPROPERTY(BlueprintReadWrite, Category="Camera Animation")
+	UPROPERTY(EditAnywhere, BlueprintReadWrite, Category="Camera Animation")
 	float EaseInDuration = 0.f;
 
 	/** Ease-out function type */
-	UPROPERTY(BlueprintReadWrite, Category="Camera Animation")
+	UPROPERTY(EditAnywhere, BlueprintReadWrite, Category="Camera Animation")
 	ECameraAnimationEasingType EaseOutType = ECameraAnimationEasingType::Linear;
 	/** Ease-out duration in seconds */
-	UPROPERTY(BlueprintReadWrite, Category="Camera Animation")
+	UPROPERTY(EditAnywhere, BlueprintReadWrite, Category="Camera Animation")
 	float EaseOutDuration = 0.f;
 
 	/** Whether the camera animation should loop */
-	UPROPERTY(BlueprintReadWrite, Category="Camera Animation")
+	UPROPERTY(EditAnywhere, BlueprintReadWrite, Category="Camera Animation")
 	bool bLoop = false;
 	/** Offset, in frames, into the animation to start at */
-	UPROPERTY(BlueprintReadWrite, Category="Camera Animation")
+	UPROPERTY(EditAnywhere, BlueprintReadWrite, Category="Camera Animation")
 	int32 StartOffset = 0;
 	/** Whether the camera animation should have a random start time */
-	UPROPERTY(BlueprintReadWrite, Category="Camera Animation")
+	UPROPERTY(EditAnywhere, BlueprintReadWrite, Category="Camera Animation")
 	bool bRandomStartTime = false;
 	/** Override the duration of the animation with a new duration (including blends) */
-	UPROPERTY(BlueprintReadWrite, Category="Camera Animation")
+	UPROPERTY(EditAnywhere, BlueprintReadWrite, Category="Camera Animation")
 	float DurationOverride = 0.f;
 
 	/** The transform space to use for the new camera shake */
-	UPROPERTY(BlueprintReadWrite, Category="Camera Animation")
+	UPROPERTY(EditAnywhere, BlueprintReadWrite, Category="Camera Animation")
 	ECameraAnimationPlaySpace PlaySpace = ECameraAnimationPlaySpace::CameraLocal;
 	/** User space to use when PlaySpace is UserDefined */
-	UPROPERTY(BlueprintReadWrite, Category="Camera Animation")
+	UPROPERTY(EditAnywhere, BlueprintReadWrite, Category="Camera Animation")
 	FRotator UserPlaySpaceRot = FRotator::ZeroRotator;
 };
 
