Cameras: fix lingering bugs with camera animation start offsets

diff --git a/Engine/Plugins/Cameras/GameplayCameras/Source/GameplayCameras/Private/CameraAnimationCameraModifier.cpp b/Engine/Plugins/Cameras/GameplayCameras/Source/GameplayCameras/Private/CameraAnimationCameraModifier.cpp
--- a/Engine/Plugins/Cameras/GameplayCameras/Source/GameplayCameras/Private/CameraAnimationCameraModifier.cpp
+++ b/Engine/Plugins/Cameras/GameplayCameras/Source/GameplayCameras/Private/CameraAnimationCameraModifier.cpp
@@ -219,8 +219,9 @@ void UCameraAnimationCameraModifier::DisplayDebug(class UCanvas* Canvas, const F
 			const FFrameRate DisplayRate = ActiveAnimation.Player->GetInputRate();
 			const FFrameTime DurationFrames = ActiveAnimation.Player->GetDuration();
 			const FFrameTime CurrentPosition = ActiveAnimation.Player->GetCurrentPosition();
+			const FFrameTime StartFrame = ActiveAnimation.Player->GetStartFrame();
 
-			const float CurrentTime = DisplayRate.AsSeconds(CurrentPosition);
+			const float ElapsedTime = DisplayRate.AsSeconds(CurrentPosition - StartFrame);
 			const float DurationSeconds = DisplayRate.AsSeconds(DurationFrames);
 
 			const FString LoopString = ActiveAnimation.Params.bLoop ? TEXT(" Looping") : TEXT("");
@@ -231,7 +232,7 @@ void UCameraAnimationCameraModifier::DisplayDebug(class UCanvas* Canvas, const F
 					FString::Printf(
 						TEXT("[%d] %s PlayRate: %f Duration: %f Elapsed: %f%s%s"),
 						Index, *GetNameSafe(ActiveAnimation.Sequence), 
-						ActiveAnimation.Params.PlayRate, DurationSeconds, CurrentTime, *EaseInString, *EaseOutString),
+						ActiveAnimation.Params.PlayRate, DurationSeconds, ElapsedTime, *EaseInString, *EaseOutString),
 					Indentation* YL, (LineNumber++)* YL);
 		}
 	}
@@ -341,13 +342,13 @@ void UCameraAnimationCameraModifier::TickAnimation(FActiveCameraAnimationInfo& C
 	// Start easing out if we're nearing the end.
 	if (!Player->GetIsLooping())
 	{
-		const float NewTime = DisplayRate.AsSeconds(NewPosition);
+		const float ElapsedTime = DisplayRate.AsSeconds(NewPosition - Player->GetStartFrame());
 		const float DurationTime = DisplayRate.AsSeconds(Player->GetDuration()) * Params.PlayRate;
 		const float BlendOutStartTime = DurationTime - Params.EaseOutDuration;
-		if (NewTime > BlendOutStartTime)
+		if (ElapsedTime > BlendOutStartTime)
 		{
 			CameraAnimation.bIsEasingOut = true;
-			CameraAnimation.EaseOutCurrentTime = NewTime - BlendOutStartTime;
+			CameraAnimation.EaseOutCurrentTime = ElapsedTime - BlendOutStartTime;
 		}
 	}
 
diff --git a/Engine/Plugins/MovieScene/TemplateSequence/Source/TemplateSequence/Public/CameraAnimationSequencePlayer.h b/Engine/Plugins/MovieScene/TemplateSequence/Source/TemplateSequence/Public/CameraAnimationSequencePlayer.h
--- a/Engine/Plugins/MovieScene/TemplateSequence/Source/TemplateSequence/Public/CameraAnimationSequencePlayer.h
+++ b/Engine/Plugins/MovieScene/TemplateSequence/Source/TemplateSequence/Public/CameraAnimationSequencePlayer.h
@@ -155,6 +155,8 @@ public:
 	FFrameRate GetInputRate() const { return PlayPosition.GetInputRate(); }
 	/** Get the sequence tick resolution */
 	FFrameRate GetOutputRate() const { return PlayPosition.GetOutputRate(); }
+	/** Get the start frame of the current sequence */
+	FFrameNumber GetStartFrame() const { return StartFrame; }
 	/** Get the duration of the current sequence in frames (display rate) */
 	FFrameTime GetDuration() const;
 	/** Gets the current play position in frames (display rate) */
