Cameras: make BP functions on the camera animation subsystem gracefully handle null pointers

diff --git a/Engine/Plugins/Cameras/GameplayCameras/Source/GameplayCameras/Private/CameraAnimationCameraModifier.cpp b/Engine/Plugins/Cameras/GameplayCameras/Source/GameplayCameras/Private/CameraAnimationCameraModifier.cpp
--- a/Engine/Plugins/Cameras/GameplayCameras/Source/GameplayCameras/Private/CameraAnimationCameraModifier.cpp
+++ b/Engine/Plugins/Cameras/GameplayCameras/Source/GameplayCameras/Private/CameraAnimationCameraModifier.cpp
@@ -9,6 +9,7 @@
 #include "GameFramework/PlayerController.h"
 #include "GameplayCamerasModule.h"
 #include "Kismet/GameplayStatics.h"
+#include "UObject/Object.h"
 
 #include UE_INLINE_GENERATED_CPP_BY_NAME(CameraAnimationCameraModifier)
 
@@ -184,10 +185,10 @@ UCameraAnimationCameraModifier* UCameraAnimationCameraModifier::GetCameraAnimati
 
 UCameraAnimationCameraModifier* UCameraAnimationCameraModifier::GetCameraAnimationCameraModifierFromPlayerController(const APlayerController* PlayerController)
 {
-	if (ensure(PlayerController))
+	if (PlayerController)
 	{
 		UCameraModifier* CameraModifier = PlayerController->PlayerCameraManager->FindCameraModifierByClass(UCameraAnimationCameraModifier::StaticClass());
-		return CastChecked<UCameraAnimationCameraModifier>(CameraModifier);
+		return CastChecked<UCameraAnimationCameraModifier>(CameraModifier, ECastCheckedType::NullAllowed);
 	}
 	return nullptr;
 }
diff --git a/Engine/Plugins/Cameras/GameplayCameras/Source/GameplayCameras/Private/GameplayCamerasSubsystem.cpp b/Engine/Plugins/Cameras/GameplayCameras/Source/GameplayCameras/Private/GameplayCamerasSubsystem.cpp
--- a/Engine/Plugins/Cameras/GameplayCameras/Source/GameplayCameras/Private/GameplayCamerasSubsystem.cpp
+++ b/Engine/Plugins/Cameras/GameplayCameras/Source/GameplayCameras/Private/GameplayCamerasSubsystem.cpp
@@ -21,17 +21,18 @@ UGameplayCamerasSubsystem* UGameplayCamerasSubsystem::GetGameplayCamerasSubsyste
 FCameraAnimationHandle UGameplayCamerasSubsystem::PlayCameraAnimation(APlayerController* PlayerController, UCameraAnimationSequence* Sequence, FCameraAnimationParams Params)
 {
 	UCameraAnimationCameraModifier* CameraModifier = UCameraAnimationCameraModifier::GetCameraAnimationCameraModifierFromPlayerController(PlayerController);
-	if (ensure(CameraModifier))
+	if (ensureMsgf(CameraModifier, TEXT("No camera modifier found on the player controller")))
 	{
 		return CameraModifier->PlayCameraAnimation(Sequence, Params);
 	}
+	FFrame::KismetExecutionMessage(TEXT("Can't play camera animation: no camera animation modifier found"), ELogVerbosity::Error);
 	return FCameraAnimationHandle::Invalid;
 }
 
 bool UGameplayCamerasSubsystem::IsCameraAnimationActive(APlayerController* PlayerController, const FCameraAnimationHandle& Handle) const
 {
 	UCameraAnimationCameraModifier* CameraModifier = UCameraAnimationCameraModifier::GetCameraAnimationCameraModifierFromPlayerController(PlayerController);
-	if (ensure(CameraModifier))
+	if (CameraModifier)
 	{
 		return CameraModifier->IsCameraAnimationActive(Handle);
 	}
@@ -41,28 +42,31 @@ bool UGameplayCamerasSubsystem::IsCameraAnimationActive(APlayerController* Playe
 void UGameplayCamerasSubsystem::StopCameraAnimation(APlayerController* PlayerController, const FCameraAnimationHandle& Handle, bool bImmediate)
 {
 	UCameraAnimationCameraModifier* CameraModifier = UCameraAnimationCameraModifier::GetCameraAnimationCameraModifierFromPlayerController(PlayerController);
-	if (ensure(CameraModifier))
+	if (ensureMsgf(CameraModifier, TEXT("No camera modifier found on the player controller")))
 	{
 		return CameraModifier->StopCameraAnimation(Handle, bImmediate);
 	}
+	FFrame::KismetExecutionMessage(TEXT("Can't stop camera animation: no camera animation modifier found"), ELogVerbosity::Error);
 }
 
 void UGameplayCamerasSubsystem::StopAllCameraAnimationsOf(APlayerController* PlayerController, UCameraAnimationSequence* Sequence, bool bImmediate)
 {
 	UCameraAnimationCameraModifier* CameraModifier = UCameraAnimationCameraModifier::GetCameraAnimationCameraModifierFromPlayerController(PlayerController);
-	if (ensure(CameraModifier))
+	if (ensureMsgf(CameraModifier, TEXT("No camera modifier found on the player controller")))
 	{
 		CameraModifier->StopAllCameraAnimationsOf(Sequence, bImmediate);
 	}
+	FFrame::KismetExecutionMessage(TEXT("Can't stop camera animations: no camera animation modifier found"), ELogVerbosity::Error);
 }
 
 void UGameplayCamerasSubsystem::StopAllCameraAnimations(APlayerController* PlayerController, bool bImmediate)
 {
 	UCameraAnimationCameraModifier* CameraModifier = UCameraAnimationCameraModifier::GetCameraAnimationCameraModifierFromPlayerController(PlayerController);
-	if (ensure(CameraModifier))
+	if (ensureMsgf(CameraModifier, TEXT("No camera modifier found on the player controller")))
 	{
 		CameraModifier->StopAllCameraAnimations(bImmediate);
 	}
+	FFrame::KismetExecutionMessage(TEXT("Can't stop all camera animation: no camera animation modifier found"), ELogVerbosity::Error);
 }
 
 #undef LOCTEXT_NAMESPACE
