Cameras: apply current scale and easing value to post-process settings weight

diff --git a/Engine/Plugins/Cameras/GameplayCameras/Source/GameplayCameras/Private/CameraAnimationCameraModifier.cpp b/Engine/Plugins/Cameras/GameplayCameras/Source/GameplayCameras/Private/CameraAnimationCameraModifier.cpp
--- a/Engine/Plugins/Cameras/GameplayCameras/Source/GameplayCameras/Private/CameraAnimationCameraModifier.cpp
+++ b/Engine/Plugins/Cameras/GameplayCameras/Source/GameplayCameras/Private/CameraAnimationCameraModifier.cpp
@@ -434,8 +434,9 @@ void UCameraAnimationCameraModifier::TickAnimation(FActiveCameraAnimationInfo& C
 	// Add the post-process settings.
 	if (CameraOwner != nullptr && CameraStandIn->PostProcessBlendWeight > 0.f)
 	{
+		const float TotalPostProcessBlendWeight = CameraStandIn->PostProcessBlendWeight * Scale;
 		EViewTargetBlendOrder CameraShakeBlendOrder = GCameraAnimationLegacyPostProcessBlending.GetValueOnGameThread() ? VTBlendOrder_Base : VTBlendOrder_Override;
-		CameraOwner->AddCachedPPBlend(CameraStandIn->PostProcessSettings, CameraStandIn->PostProcessBlendWeight, CameraShakeBlendOrder);
+		CameraOwner->AddCachedPPBlend(CameraStandIn->PostProcessSettings, TotalPostProcessBlendWeight, CameraShakeBlendOrder);
 	}
 }
 
