Fix for crashes in UGameplayCamerasFunctionLibrary::Conv_CameraAnimationCameraModifier when called from blueprint with an invalid PlayerCameraManager. Changed to validated Get in blueprints and changed the code to not crash if the passed in PlayerCameraManager is nullptr.

diff --git a/Engine/Plugins/Cameras/GameplayCameras/Source/GameplayCameras/Private/CameraAnimationCameraModifier.cpp b/Engine/Plugins/Cameras/GameplayCameras/Source/GameplayCameras/Private/CameraAnimationCameraModifier.cpp
--- a/Engine/Plugins/Cameras/GameplayCameras/Source/GameplayCameras/Private/CameraAnimationCameraModifier.cpp
+++ b/Engine/Plugins/Cameras/GameplayCameras/Source/GameplayCameras/Private/CameraAnimationCameraModifier.cpp
@@ -442,7 +442,7 @@ void UCameraAnimationCameraModifier::TickAnimation(FActiveCameraAnimationInfo& C
 
 UCameraAnimationCameraModifier* UGameplayCamerasFunctionLibrary::Conv_CameraAnimationCameraModifier(APlayerCameraManager* PlayerCameraManager)
 {
-	return Cast<UCameraAnimationCameraModifier>(PlayerCameraManager->FindCameraModifierByClass(UCameraAnimationCameraModifier::StaticClass()));
+	return PlayerCameraManager ? Cast<UCameraAnimationCameraModifier>(PlayerCameraManager->FindCameraModifierByClass(UCameraAnimationCameraModifier::StaticClass())) : nullptr;
 }
 
 ECameraShakePlaySpace UGameplayCamerasFunctionLibrary::Conv_CameraShakePlaySpace(ECameraAnimationPlaySpace CameraAnimationPlaySpace)
