Enabling legacy post processing blend mode

diff --git a/Engine/Plugins/Cameras/GameplayCameras/Source/GameplayCameras/Private/CameraAnimationCameraModifier.cpp b/Engine/Plugins/Cameras/GameplayCameras/Source/GameplayCameras/Private/CameraAnimationCameraModifier.cpp
--- a/Engine/Plugins/Cameras/GameplayCameras/Source/GameplayCameras/Private/CameraAnimationCameraModifier.cpp
+++ b/Engine/Plugins/Cameras/GameplayCameras/Source/GameplayCameras/Private/CameraAnimationCameraModifier.cpp
@@ -17,7 +17,7 @@ DECLARE_CYCLE_STAT(TEXT("Camera Animation Eval"), CameraAnimationEval_Total, STA
 
 TAutoConsoleVariable<bool> GCameraAnimationLegacyPostProcessBlending(
 	TEXT("r.CameraAnimation.LegacyPostProcessBlending"),
-	false,
+	true,
 	TEXT("Blend camera animation post process settings under the main camera instead of over it"));
 
 FCameraAnimationHandle FCameraAnimationHandle::Invalid(MAX_uint16, 0);
