Cameras: Properly build the duration of the composite shake

diff --git a/Engine/Plugins/Cameras/GameplayCameras/Source/GameplayCameras/Private/CompositeCameraShakePattern.cpp b/Engine/Plugins/Cameras/GameplayCameras/Source/GameplayCameras/Private/CompositeCameraShakePattern.cpp
--- a/Engine/Plugins/Cameras/GameplayCameras/Source/GameplayCameras/Private/CompositeCameraShakePattern.cpp
+++ b/Engine/Plugins/Cameras/GameplayCameras/Source/GameplayCameras/Private/CompositeCameraShakePattern.cpp
@@ -7,10 +7,11 @@
 
 void UCompositeCameraShakePattern::GetShakePatternInfoImpl(FCameraShakeInfo& OutInfo) const
 {
-	// Set the duration to custom (since we want to handle timing ourselves),
-	// but let's set a hint duration e.g. for when we are shown as a clip in
-	// a sequence. The hint duration is the max of all our children's durations.
-	OutInfo.Duration = FCameraShakeDuration::Custom();
+	// Set the duration to fixed at first. We'll change it to something different
+	// if we encounter children with different duration types.
+	// The duration value will be the maximum of all our children's durations.
+	float Duration = 0.f;
+	ECameraShakeDurationType DurationType = ECameraShakeDurationType::Fixed;
 
 	for (UCameraShakePattern* Pattern : ChildPatterns)
 	{
@@ -19,21 +20,34 @@ void UCompositeCameraShakePattern::GetShakePatternInfoImpl(FCameraShakeInfo& Out
 			FCameraShakeInfo ChildInfo;
 			Pattern->GetShakePatternInfo(ChildInfo);
 
-			if (ChildInfo.Duration.IsInfinite())
+			switch (ChildInfo.Duration.GetDurationType())
 			{
-				// If one of our children is infinite, we are infinite.
-				OutInfo.Duration = ChildInfo.Duration;
-				break;
+				case ECameraShakeDurationType::Infinite:
+					// If one of our children is infinite, we are infinite.
+					Duration = 0.f;
+					DurationType = ECameraShakeDurationType::Infinite;
+					break;
+				case ECameraShakeDurationType::Custom:
+					// Change our type to custom, but include the child's duration hint
+					// in our own hint.
+					Duration = FMath::Max(Duration, ChildInfo.Duration.Get());
+					DurationType = ECameraShakeDurationType::Custom;
+					break;
+				case ECameraShakeDurationType::Fixed:
+					// Grow our fixed duration if necessary.
+					Duration = FMath::Max(Duration, ChildInfo.Duration.Get());
+					break;
 			}
-			else
+
+			if (DurationType == ECameraShakeDurationType::Infinite)
 			{
-				// If a child has a fixed duration, or custom duration with hint,
-				// let's include that in our own hint.
-				OutInfo.Duration = FCameraShakeDuration::Custom(
-					FMath::Max(ChildInfo.Duration.Get(), OutInfo.Duration.Get()));
+				// Can't get any bigger than infinite.
+				break;
 			}
 		}
 	}
+
+	OutInfo.Duration = FCameraShakeDuration(Duration, DurationType);
 }
 
 void UCompositeCameraShakePattern::StartShakePatternImpl(const FCameraShakePatternStartParams& Params)
