Cameras: move shake-related headers into sub-directories, move legacy headers into separate directory

diff --git a/Engine/Plugins/Cameras/GameplayCameras/Source/GameplayCameras/GameplayCameras.Build.cs b/Engine/Plugins/Cameras/GameplayCameras/Source/GameplayCameras/GameplayCameras.Build.cs
--- a/Engine/Plugins/Cameras/GameplayCameras/Source/GameplayCameras/GameplayCameras.Build.cs
+++ b/Engine/Plugins/Cameras/GameplayCameras/Source/GameplayCameras/GameplayCameras.Build.cs
@@ -1,11 +1,14 @@
 // Copyright Epic Games, Inc. All Rights Reserved.
 
 using UnrealBuildTool;
+using System.IO;
 
 public class GameplayCameras : ModuleRules
 {
 	public GameplayCameras(ReadOnlyTargetRules Target) : base(Target)
 	{
+		PublicIncludePaths.Add(Path.Combine(ModuleDirectory, "Legacy"));
+
 		PublicDependencyModuleNames.AddRange(
 			new string[] {
 				"Core",
diff --git a/Engine/Plugins/Cameras/GameplayCameras/Source/GameplayCameras/Legacy/Camera/CameraShake.h b/Engine/Plugins/Cameras/GameplayCameras/Source/GameplayCameras/Legacy/Camera/CameraShake.h
new file mode 100644
--- /dev/null
+++ b/Engine/Plugins/Cameras/GameplayCameras/Source/GameplayCameras/Legacy/Camera/CameraShake.h
@@ -0,0 +1,11 @@
+// Copyright Epic Games, Inc. All Rights Reserved.
+
+#pragma once
+
+// HEADER_UNIT_SKIP - Deprecated
+
+#include "CoreTypes.h"
+
+COMPILE_WARNING("UCameraShake has been renamed to ULegacyCameraShake, change to #include \"LegacyCameraShake.h\" and rename references")
+#include "LegacyCameraShake.h"
+
diff --git a/Engine/Plugins/Cameras/GameplayCameras/Source/GameplayCameras/Legacy/CameraAnimationCameraModifier.h b/Engine/Plugins/Cameras/GameplayCameras/Source/GameplayCameras/Legacy/CameraAnimationCameraModifier.h
new file mode 100644
--- /dev/null
+++ b/Engine/Plugins/Cameras/GameplayCameras/Source/GameplayCameras/Legacy/CameraAnimationCameraModifier.h
@@ -0,0 +1,6 @@
+// Copyright Epic Games, Inc. All Rights Reserved.
+
+#pragma once
+
+#include "Animations/CameraAnimationCameraModifier.h"
+
diff --git a/Engine/Plugins/Cameras/GameplayCameras/Source/GameplayCameras/Legacy/LegacyCameraShake.h b/Engine/Plugins/Cameras/GameplayCameras/Source/GameplayCameras/Legacy/LegacyCameraShake.h
new file mode 100644
--- /dev/null
+++ b/Engine/Plugins/Cameras/GameplayCameras/Source/GameplayCameras/Legacy/LegacyCameraShake.h
@@ -0,0 +1,6 @@
+// Copyright Epic Games, Inc. All Rights Reserved.
+
+#pragma once
+
+#include "Shakes/LegacyCameraShake.h"
+
diff --git a/Engine/Plugins/Cameras/GameplayCameras/Source/GameplayCameras/Private/Animations/CameraAnimationCameraModifier.cpp b/Engine/Plugins/Cameras/GameplayCameras/Source/GameplayCameras/Private/Animations/CameraAnimationCameraModifier.cpp
new file mode 100644
--- /dev/null
+++ b/Engine/Plugins/Cameras/GameplayCameras/Source/GameplayCameras/Private/Animations/CameraAnimationCameraModifier.cpp
@@ -0,0 +1,470 @@
+// Copyright Epic Games, Inc. All Rights Reserved.
+
+#include "CameraAnimationCameraModifier.h"
+
+#include "Camera/CameraAnimationHelper.h"
+#include "CameraAnimationSequence.h"
+#include "CameraAnimationSequencePlayer.h"
+#include "Engine/Canvas.h"
+#include "Engine/Engine.h"
+#include "GameFramework/PlayerController.h"
+#include "GameplayCamerasModule.h"
+#include "Kismet/GameplayStatics.h"
+#include "UObject/Object.h"
+
+#include UE_INLINE_GENERATED_CPP_BY_NAME(CameraAnimationCameraModifier)
+
+DECLARE_CYCLE_STAT(TEXT("Camera Animation Eval"), CameraAnimationEval_Total, STATGROUP_CameraAnimation);
+
+TAutoConsoleVariable<bool> GCameraAnimationLegacyPostProcessBlending(
+	TEXT("r.CameraAnimation.LegacyPostProcessBlending"),
+	true,
+	TEXT("Blend camera animation post process settings under the main camera instead of over it"));
+
+FCameraAnimationHandle FCameraAnimationHandle::Invalid(MAX_uint16, 0);
+
+FActiveCameraAnimationInfo::FActiveCameraAnimationInfo()
+	: Sequence(nullptr)
+	, Handle(FCameraAnimationHandle::Invalid)
+	, Player(nullptr)
+	, CameraStandIn(nullptr)
+	, EaseInCurrentTime(0)
+	, EaseOutCurrentTime(0)
+	, bIsEasingIn(false)
+	, bIsEasingOut(false)
+{
+}
+
+namespace CameraAnimationCameraModifierImpl
+{
+	float Sinusoidal(float InTime)
+	{
+		return FMath::Sin(.5f*PI*InTime);
+	}
+	float Power(float InTime, float Power)
+	{
+		return FMath::Pow(InTime, Power);
+	}
+	float Exponential(float InTime)
+	{
+		return FMath::Pow(2, 10*(InTime - 1.f));
+	}
+	float Circular(float InTime)
+	{
+		return 1.f-FMath::Sqrt(1-InTime*InTime);
+	}
+}
+
+float UCameraAnimationCameraModifier::EvaluateEasing(ECameraAnimationEasingType EasingType, float Interp)
+{
+	using namespace CameraAnimationCameraModifierImpl;
+	
+	switch (EasingType)
+	{
+	case ECameraAnimationEasingType::Sinusoidal: 	return Sinusoidal(Interp);
+	case ECameraAnimationEasingType::Quadratic: 	return Power(Interp,2);
+	case ECameraAnimationEasingType::Cubic:			return Power(Interp,3);
+	case ECameraAnimationEasingType::Quartic:	 	return Power(Interp,4);
+	case ECameraAnimationEasingType::Quintic:	 	return Power(Interp,5);
+	case ECameraAnimationEasingType::Exponential: 	return Exponential(Interp);
+	case ECameraAnimationEasingType::Circular:	 	return Circular(Interp);
+	case ECameraAnimationEasingType::Linear:
+	default:
+		return Interp;
+	}
+}
+
+UCameraAnimationCameraModifier::UCameraAnimationCameraModifier(const FObjectInitializer& ObjectInitializer)
+	: Super(ObjectInitializer)
+	, NextInstanceSerialNumber(1)
+{
+}
+
+FCameraAnimationHandle UCameraAnimationCameraModifier::PlayCameraAnimation(UCameraAnimationSequence* Sequence, FCameraAnimationParams Params)
+{
+	if (!ensure(Sequence))
+	{
+		return FCameraAnimationHandle::Invalid;
+	}
+
+	TRACE_CPUPROFILER_EVENT_SCOPE_TEXT(*FString::Printf(TEXT("PlayCameraAnimation %s"), *Sequence->GetFName().ToString()));
+
+	int32 NewIndex = FindInactiveCameraAnimation();
+	check(NewIndex < MAX_uint16);
+
+	const uint16 InstanceSerial = NextInstanceSerialNumber++;
+	FCameraAnimationHandle InstanceHandle { (uint16)NewIndex, InstanceSerial };
+
+	FActiveCameraAnimationInfo& NewCameraAnimation = ActiveAnimations[NewIndex];
+	NewCameraAnimation.Sequence = Sequence;
+	NewCameraAnimation.Params = Params;
+	NewCameraAnimation.Handle = InstanceHandle;
+
+	const FName PlayerName = MakeUniqueObjectName(this, UCameraAnimationSequencePlayer::StaticClass(), TEXT("CameraAnimationPlayer"));
+	NewCameraAnimation.Player = NewObject<UCameraAnimationSequencePlayer>(this, PlayerName);
+	const FName CameraStandInName = MakeUniqueObjectName(this, UCameraAnimationSequenceCameraStandIn::StaticClass(), TEXT("CameraStandIn"));
+	NewCameraAnimation.CameraStandIn = NewObject<UCameraAnimationSequenceCameraStandIn>(this, CameraStandInName);
+
+	// Start easing in immediately if there's any defined.
+	NewCameraAnimation.bIsEasingIn = (Params.EaseInDuration > 0.f);
+
+	// Initialize our stand-in object.
+	NewCameraAnimation.CameraStandIn->Initialize(Sequence);
+	
+	// Make the player always use our stand-in object whenever a sequence wants to spawn or possess an object.
+	NewCameraAnimation.Player->SetBoundObjectOverride(NewCameraAnimation.CameraStandIn);
+
+	// Initialize it and start playing.
+	NewCameraAnimation.Player->Initialize(Sequence, Params.StartOffset);
+	NewCameraAnimation.Player->Play(Params.bLoop, Params.bRandomStartTime);
+
+	return InstanceHandle;
+}
+
+bool UCameraAnimationCameraModifier::IsCameraAnimationActive(const FCameraAnimationHandle& Handle) const
+{
+	const FActiveCameraAnimationInfo* CameraAnimation = GetActiveCameraAnimation(Handle);
+	if (CameraAnimation)
+	{
+		return CameraAnimation->Player->GetPlaybackStatus() == EMovieScenePlayerStatus::Playing;
+	}
+	return false;
+}
+
+void UCameraAnimationCameraModifier::StopCameraAnimation(const FCameraAnimationHandle& Handle, bool bImmediate)
+{
+	FActiveCameraAnimationInfo* CameraAnimation = GetActiveCameraAnimation(Handle);
+	if (CameraAnimation)
+	{
+		check(CameraAnimation->Sequence && CameraAnimation->Player);
+
+		TRACE_CPUPROFILER_EVENT_SCOPE_TEXT(*FString::Printf(TEXT("StopCameraAnimation %s"), *CameraAnimation->Sequence->GetFName().ToString()));
+
+		if (bImmediate || CameraAnimation->Params.EaseOutDuration == 0.f)
+		{
+			CameraAnimation->Player->Stop();
+			ActiveAnimations[CameraAnimation->Handle.InstanceID] = FActiveCameraAnimationInfo();
+		}
+		else if (!CameraAnimation->bIsEasingOut)
+		{
+			CameraAnimation->bIsEasingOut = true;
+			CameraAnimation->EaseOutCurrentTime = 0.f;
+		}
+	}
+}
+
+void UCameraAnimationCameraModifier::StopAllCameraAnimationsOf(UCameraAnimationSequence* Sequence, bool bImmediate)
+{
+	for (FActiveCameraAnimationInfo& CameraAnimation : ActiveAnimations)
+	{
+		if (CameraAnimation.Sequence == Sequence)
+		{
+			StopCameraAnimation(CameraAnimation.Handle, bImmediate);
+		}
+	}
+}
+
+void UCameraAnimationCameraModifier::StopAllCameraAnimations(bool bImmediate)
+{
+	for (FActiveCameraAnimationInfo& CameraAnimation : ActiveAnimations)
+	{
+		StopCameraAnimation(CameraAnimation.Handle, bImmediate);
+	}
+}
+
+UCameraAnimationCameraModifier* UCameraAnimationCameraModifier::GetCameraAnimationCameraModifier(const UObject* WorldContextObject, int32 PlayerIndex)
+{
+	APlayerController* PlayerController = UGameplayStatics::GetPlayerController(WorldContextObject, PlayerIndex);
+	return GetCameraAnimationCameraModifierFromPlayerController(PlayerController);
+}
+
+UCameraAnimationCameraModifier* UCameraAnimationCameraModifier::GetCameraAnimationCameraModifierFromID(const UObject* WorldContextObject, int32 ControllerID)
+{
+	APlayerController* PlayerController = UGameplayStatics::GetPlayerControllerFromID(WorldContextObject, ControllerID);
+	return GetCameraAnimationCameraModifierFromPlayerController(PlayerController);
+}
+
+UCameraAnimationCameraModifier* UCameraAnimationCameraModifier::GetCameraAnimationCameraModifierFromPlayerController(const APlayerController* PlayerController)
+{
+	if (PlayerController)
+	{
+		UCameraModifier* CameraModifier = PlayerController->PlayerCameraManager->FindCameraModifierByClass(UCameraAnimationCameraModifier::StaticClass());
+		return CastChecked<UCameraAnimationCameraModifier>(CameraModifier, ECastCheckedType::NullAllowed);
+	}
+	return nullptr;
+}
+
+bool UCameraAnimationCameraModifier::ModifyCamera(float DeltaTime, struct FMinimalViewInfo& InOutPOV)
+{
+	Super::ModifyCamera(DeltaTime, InOutPOV);
+	TickAllAnimations(DeltaTime, InOutPOV);
+	return false;
+}
+
+void UCameraAnimationCameraModifier::DisplayDebug(class UCanvas* Canvas, const FDebugDisplayInfo& DebugDisplay, float& YL, float& YPos)
+{
+	Canvas->SetDrawColor(FColor::Yellow);
+	UFont* DrawFont = GEngine->GetSmallFont();
+
+	int Indentation = 1;
+	int LineNumber = FMath::CeilToInt(YPos / YL);
+
+	Canvas->DrawText(DrawFont, FString::Printf(TEXT("Modifier_CameraAnimationSequence %s, Alpha:%f"), *GetNameSafe(this), Alpha), Indentation * YL, (LineNumber++) * YL);
+
+	Indentation = 2;
+	for (int Index = 0; Index < ActiveAnimations.Num(); Index++)
+	{
+		FActiveCameraAnimationInfo& ActiveAnimation = ActiveAnimations[Index];
+
+		if (ActiveAnimation.IsValid())
+		{
+			const FFrameRate DisplayRate = ActiveAnimation.Player->GetInputRate();
+			const FFrameTime DurationFrames = ActiveAnimation.Player->GetDuration();
+			const FFrameTime CurrentPosition = ActiveAnimation.Player->GetCurrentPosition();
+			const FFrameTime StartFrame = ActiveAnimation.Player->GetStartFrame();
+
+			const float ElapsedTime = DisplayRate.AsSeconds(CurrentPosition - StartFrame);
+			const float DurationSeconds = DisplayRate.AsSeconds(DurationFrames);
+
+			const FString LoopString = ActiveAnimation.Params.bLoop ? TEXT(" Looping") : TEXT("");
+			const FString EaseInString = ActiveAnimation.bIsEasingIn ? FString::Printf(TEXT(" Easing In: %f / %f"), ActiveAnimation.EaseInCurrentTime, ActiveAnimation.Params.EaseInDuration) : TEXT("");
+			const FString EaseOutString = ActiveAnimation.bIsEasingOut ? FString::Printf(TEXT(" Easing Out: %f / %f"), ActiveAnimation.EaseOutCurrentTime, ActiveAnimation.Params.EaseOutDuration) : TEXT("");
+			Canvas->DrawText(
+					DrawFont,
+					FString::Printf(
+						TEXT("[%d] %s PlayRate: %f Duration: %f Elapsed: %f%s%s"),
+						Index, *GetNameSafe(ActiveAnimation.Sequence), 
+						ActiveAnimation.Params.PlayRate, DurationSeconds, ElapsedTime, *EaseInString, *EaseOutString),
+					Indentation* YL, (LineNumber++)* YL);
+		}
+	}
+
+	YPos = LineNumber * YL;
+
+	Super::DisplayDebug(Canvas, DebugDisplay, YL, YPos);
+}
+
+int32 UCameraAnimationCameraModifier::FindInactiveCameraAnimation()
+{
+	for (int32 Index = 0; Index < ActiveAnimations.Num(); ++Index)
+	{
+		const FActiveCameraAnimationInfo& CameraAnimation(ActiveAnimations[Index]);
+		if (!CameraAnimation.IsValid())
+		{
+			return Index;
+		}
+	}
+
+	return ActiveAnimations.Emplace();
+}
+
+const FActiveCameraAnimationInfo* UCameraAnimationCameraModifier::GetActiveCameraAnimation(const FCameraAnimationHandle& Handle) const
+{
+	if (Handle.IsValid() && ActiveAnimations.IsValidIndex(Handle.InstanceID))
+	{
+		const FActiveCameraAnimationInfo& CameraAnimation = ActiveAnimations[Handle.InstanceID];
+		if (CameraAnimation.Handle.InstanceSerial == Handle.InstanceSerial)
+		{
+			return &CameraAnimation;
+		}
+	}
+	return nullptr;
+}
+
+FActiveCameraAnimationInfo* UCameraAnimationCameraModifier::GetActiveCameraAnimation(const FCameraAnimationHandle& Handle)
+{
+	if (Handle.IsValid() && ActiveAnimations.IsValidIndex(Handle.InstanceID))
+	{
+		FActiveCameraAnimationInfo& CameraAnimation = ActiveAnimations[Handle.InstanceID];
+		if (CameraAnimation.Handle.InstanceSerial == Handle.InstanceSerial)
+		{
+			return &CameraAnimation;
+		}
+	}
+	return nullptr;
+}
+
+void UCameraAnimationCameraModifier::DeactivateCameraAnimation(int32 Index)
+{
+	check(ActiveAnimations.IsValidIndex(Index));
+	FActiveCameraAnimationInfo& CameraAnimation = ActiveAnimations[Index];
+
+	check(CameraAnimation.Player);
+	if (!ensure(CameraAnimation.Player->GetPlaybackStatus() == EMovieScenePlayerStatus::Stopped))
+	{
+		CameraAnimation.Player->Stop();
+	}
+
+	CameraAnimation = FActiveCameraAnimationInfo();
+}
+
+void UCameraAnimationCameraModifier::TickAllAnimations(float DeltaTime, FMinimalViewInfo& InOutPOV)
+{
+	for (FActiveCameraAnimationInfo& ActiveAnimation : ActiveAnimations)
+	{
+		if (ActiveAnimation.IsValid())
+		{
+			TickAnimation(ActiveAnimation, DeltaTime, InOutPOV);
+
+			if (ActiveAnimation.Player->GetPlaybackStatus() == EMovieScenePlayerStatus::Stopped)
+			{
+				DeactivateCameraAnimation(ActiveAnimation.Handle.InstanceID);
+			}
+		}
+	}
+}
+
+void UCameraAnimationCameraModifier::TickAnimation(FActiveCameraAnimationInfo& CameraAnimation, float DeltaTime, FMinimalViewInfo& InOutPOV)
+{
+	SCOPE_CYCLE_COUNTER(CameraAnimationEval_Total);
+
+	check(CameraAnimation.Player);
+	check(CameraAnimation.CameraStandIn);
+
+	const FCameraAnimationParams Params = CameraAnimation.Params;
+	UCameraAnimationSequencePlayer* Player = CameraAnimation.Player;
+	UCameraAnimationSequenceCameraStandIn* CameraStandIn = CameraAnimation.CameraStandIn;
+
+	const FFrameRate DisplayRate = Player->GetInputRate();
+	const FFrameTime CurrentPosition = Player->GetCurrentPosition();
+
+	const float ScaledDeltaTime = DeltaTime * Params.PlayRate;
+	const FFrameTime NewPosition = CurrentPosition + ScaledDeltaTime * DisplayRate;
+
+	// Advance any easing times.
+	if (CameraAnimation.bIsEasingIn)
+	{
+		CameraAnimation.EaseInCurrentTime += DeltaTime;
+	}
+	if (CameraAnimation.bIsEasingOut)
+	{
+		CameraAnimation.EaseOutCurrentTime += DeltaTime;
+	}
+
+	// Start easing out if we're nearing the end.
+	if (!Player->GetIsLooping())
+	{
+		const float ElapsedTime = DisplayRate.AsSeconds(NewPosition - Player->GetStartFrame());
+		const float DurationTime = DisplayRate.AsSeconds(Player->GetDuration()) * Params.PlayRate;
+		const float BlendOutStartTime = DurationTime - Params.EaseOutDuration;
+		if (ElapsedTime > BlendOutStartTime)
+		{
+			CameraAnimation.bIsEasingOut = true;
+			CameraAnimation.EaseOutCurrentTime = ElapsedTime - BlendOutStartTime;
+		}
+	}
+
+	// Check if we're done easing in or out.
+	bool bIsDoneEasingOut = false;
+	if (CameraAnimation.bIsEasingIn)
+	{
+		if (CameraAnimation.EaseInCurrentTime > Params.EaseInDuration || Params.EaseInDuration == 0.f)
+		{
+			CameraAnimation.bIsEasingIn = false;
+		}
+	}
+	if (CameraAnimation.bIsEasingOut)
+	{
+		if (CameraAnimation.EaseOutCurrentTime > Params.EaseOutDuration)
+		{
+			bIsDoneEasingOut = true;
+		}
+	}
+
+	// Figure out the final easing weight.
+	const float EasingInT = FMath::Clamp((CameraAnimation.EaseInCurrentTime / Params.EaseInDuration), 0.f, 1.f);
+	const float EasingInWeight = CameraAnimation.bIsEasingIn ?
+		EvaluateEasing(Params.EaseInType, EasingInT) : 1.f;
+
+	const float EasingOutT = FMath::Clamp((1.f - CameraAnimation.EaseOutCurrentTime / Params.EaseOutDuration), 0.f, 1.f);
+	const float EasingOutWeight = CameraAnimation.bIsEasingOut ?
+		EvaluateEasing(Params.EaseOutType, EasingOutT) : 1.f;
+
+	const float TotalEasingWeight = FMath::Min(EasingInWeight, EasingOutWeight);
+
+	// We might be done playing. Normally the player will stop on its own, but there are other situation in which
+	// the responsibility falls to this code:
+	// - If the animation is looping and waiting for an explicit Stop() call on us.
+	// - If there was a Stop() call with bImmediate=false to let an animation blend out.
+	if (bIsDoneEasingOut || TotalEasingWeight <= 0.f)
+	{
+		Player->Stop();
+		return;
+	}
+
+	UMovieSceneEntitySystemLinker* Linker = Player->GetEvaluationTemplate().GetEntitySystemLinker();
+	CameraStandIn->Reset(InOutPOV, Linker);
+
+	// Get the "unanimated" properties that need to be treated additively.
+	const float OriginalFieldOfView = CameraStandIn->FieldOfView;
+
+	// Update the sequence.
+	Player->Update(NewPosition);
+
+	// Recalculate properties that might be invalidated by other properties having been animated.
+	CameraStandIn->RecalcDerivedData();
+
+	// Grab the final animated (animated) values, figure out the delta, apply scale, and feed that into the result.
+	// Transform is always treated as a local, additive value. The data better be good.
+	const float Scale = Params.Scale * TotalEasingWeight;
+	const FTransform AnimatedTransform = CameraStandIn->GetTransform();
+	FVector AnimatedLocation = AnimatedTransform.GetLocation() * Scale;
+	FRotator AnimatedRotation = AnimatedTransform.GetRotation().Rotator() * Scale;
+
+	const bool bIsCameraLocal = (Params.PlaySpace == ECameraAnimationPlaySpace::CameraLocal);
+	const FMatrix UserPlaySpaceMatrix = (Params.PlaySpace == ECameraAnimationPlaySpace::UserDefined) ? 
+		FRotationMatrix(Params.UserPlaySpaceRot) : FRotationMatrix::Identity;
+	const FCameraAnimationHelperOffset CameraOffset { AnimatedLocation, AnimatedRotation };
+	if (bIsCameraLocal)
+	{
+		FCameraAnimationHelper::ApplyOffset(InOutPOV, CameraOffset, AnimatedLocation, AnimatedRotation);
+	}
+	else
+	{
+		FCameraAnimationHelper::ApplyOffset(UserPlaySpaceMatrix, InOutPOV, CameraOffset, AnimatedLocation, AnimatedRotation);
+	}
+	InOutPOV.Location = AnimatedLocation;
+	InOutPOV.Rotation = AnimatedRotation;
+
+	// FieldOfView follows the current camera's value every frame, so we can compute how much the animation is
+	// changing it.
+	const float AnimatedFieldOfView = CameraStandIn->FieldOfView;
+	const float DeltaFieldOfView = AnimatedFieldOfView - OriginalFieldOfView;
+	InOutPOV.FOV = OriginalFieldOfView + DeltaFieldOfView * Scale;
+
+	// Add the post-process settings.
+	if (CameraOwner != nullptr && CameraStandIn->PostProcessBlendWeight > 0.f)
+	{
+		const float TotalPostProcessBlendWeight = CameraStandIn->PostProcessBlendWeight * Scale;
+		EViewTargetBlendOrder CameraShakeBlendOrder = GCameraAnimationLegacyPostProcessBlending.GetValueOnGameThread() ? VTBlendOrder_Base : VTBlendOrder_Override;
+		CameraOwner->AddCachedPPBlend(CameraStandIn->PostProcessSettings, TotalPostProcessBlendWeight, CameraShakeBlendOrder);
+	}
+}
+
+UCameraAnimationCameraModifier* UGameplayCamerasFunctionLibrary::Conv_CameraAnimationCameraModifier(APlayerCameraManager* PlayerCameraManager)
+{
+	return PlayerCameraManager ? Cast<UCameraAnimationCameraModifier>(PlayerCameraManager->FindCameraModifierByClass(UCameraAnimationCameraModifier::StaticClass())) : nullptr;
+}
+
+ECameraShakePlaySpace UGameplayCamerasFunctionLibrary::Conv_CameraShakePlaySpace(ECameraAnimationPlaySpace CameraAnimationPlaySpace)
+{
+	switch (CameraAnimationPlaySpace)
+	{
+	case ECameraAnimationPlaySpace::CameraLocal: return ECameraShakePlaySpace::CameraLocal;
+	case ECameraAnimationPlaySpace::World: return ECameraShakePlaySpace::World;
+	case ECameraAnimationPlaySpace::UserDefined: return ECameraShakePlaySpace::UserDefined;
+	default: checkf(false, TEXT("Unsupported ECameraAnimationPlaySpace value")); return (ECameraShakePlaySpace)CameraAnimationPlaySpace;
+	}
+}
+
+ECameraAnimationPlaySpace UGameplayCamerasFunctionLibrary::Conv_CameraAnimationPlaySpace(ECameraShakePlaySpace CameraShakePlaySpace)
+{
+	switch (CameraShakePlaySpace)
+	{
+	case ECameraShakePlaySpace::CameraLocal: return ECameraAnimationPlaySpace::CameraLocal;
+	case ECameraShakePlaySpace::World: return ECameraAnimationPlaySpace::World;
+	case ECameraShakePlaySpace::UserDefined: return ECameraAnimationPlaySpace::UserDefined;
+	default: checkf(false, TEXT("Unsupported ECameraShakePlaySpace value")); return (ECameraAnimationPlaySpace)CameraShakePlaySpace;
+	}
+}
+
diff --git a/Engine/Plugins/Cameras/GameplayCameras/Source/GameplayCameras/Private/CameraAnimationCameraModifier.cpp b/Engine/Plugins/Cameras/GameplayCameras/Source/GameplayCameras/Private/CameraAnimationCameraModifier.cpp
deleted file mode 100644
diff --git a/Engine/Plugins/Cameras/GameplayCameras/Source/GameplayCameras/Private/CompositeCameraShakePattern.cpp b/Engine/Plugins/Cameras/GameplayCameras/Source/GameplayCameras/Private/CompositeCameraShakePattern.cpp
deleted file mode 100644
diff --git a/Engine/Plugins/Cameras/GameplayCameras/Source/GameplayCameras/Private/DefaultCameraShakeBase.cpp b/Engine/Plugins/Cameras/GameplayCameras/Source/GameplayCameras/Private/DefaultCameraShakeBase.cpp
deleted file mode 100644
diff --git a/Engine/Plugins/Cameras/GameplayCameras/Source/GameplayCameras/Private/LegacyCameraShake.cpp b/Engine/Plugins/Cameras/GameplayCameras/Source/GameplayCameras/Private/LegacyCameraShake.cpp
deleted file mode 100644
diff --git a/Engine/Plugins/Cameras/GameplayCameras/Source/GameplayCameras/Private/PerlinNoiseCameraShakePattern.cpp b/Engine/Plugins/Cameras/GameplayCameras/Source/GameplayCameras/Private/PerlinNoiseCameraShakePattern.cpp
deleted file mode 100644
diff --git a/Engine/Plugins/Cameras/GameplayCameras/Source/GameplayCameras/Private/Shakes/CompositeCameraShakePattern.cpp b/Engine/Plugins/Cameras/GameplayCameras/Source/GameplayCameras/Private/Shakes/CompositeCameraShakePattern.cpp
new file mode 100644
--- /dev/null
+++ b/Engine/Plugins/Cameras/GameplayCameras/Source/GameplayCameras/Private/Shakes/CompositeCameraShakePattern.cpp
@@ -0,0 +1,179 @@
+// Copyright Epic Games, Inc. All Rights Reserved.
+
+#include "Shakes/CompositeCameraShakePattern.h"
+#include "Camera/PlayerCameraManager.h"
+
+#include UE_INLINE_GENERATED_CPP_BY_NAME(CompositeCameraShakePattern)
+
+void UCompositeCameraShakePattern::GetShakePatternInfoImpl(FCameraShakeInfo& OutInfo) const
+{
+	// Set the duration to fixed at first. We'll change it to something different
+	// if we encounter children with different duration types.
+	// The duration value will be the maximum of all our children's durations.
+	float Duration = 0.f;
+	ECameraShakeDurationType DurationType = ECameraShakeDurationType::Fixed;
+
+	for (UCameraShakePattern* Pattern : ChildPatterns)
+	{
+		if (Pattern != nullptr)
+		{
+			FCameraShakeInfo ChildInfo;
+			Pattern->GetShakePatternInfo(ChildInfo);
+
+			switch (ChildInfo.Duration.GetDurationType())
+			{
+				case ECameraShakeDurationType::Infinite:
+					// If one of our children is infinite, we are infinite.
+					Duration = 0.f;
+					DurationType = ECameraShakeDurationType::Infinite;
+					break;
+				case ECameraShakeDurationType::Custom:
+					// Change our type to custom, but include the child's duration hint
+					// in our own hint.
+					Duration = FMath::Max(Duration, ChildInfo.Duration.Get());
+					DurationType = ECameraShakeDurationType::Custom;
+					break;
+				case ECameraShakeDurationType::Fixed:
+					// Grow our fixed duration if necessary.
+					Duration = FMath::Max(Duration, ChildInfo.Duration.Get());
+					break;
+			}
+
+			if (DurationType == ECameraShakeDurationType::Infinite)
+			{
+				// Can't get any bigger than infinite.
+				break;
+			}
+		}
+	}
+
+	OutInfo.Duration = FCameraShakeDuration(Duration, DurationType);
+}
+
+void UCompositeCameraShakePattern::StartShakePatternImpl(const FCameraShakePatternStartParams& Params)
+{
+	for (UCameraShakePattern* Pattern : ChildPatterns)
+	{
+		if (Pattern != nullptr)
+		{
+			// Start the child pattern.
+			Pattern->StartShakePattern(Params);
+		}
+	}
+}
+
+void UCompositeCameraShakePattern::UpdateShakePatternImpl(const FCameraShakePatternUpdateParams& Params, FCameraShakePatternUpdateResult& OutResult)
+{
+	UCameraShakeBase* ShakeInstance = GetShakeInstance();
+	checkf(ShakeInstance, TEXT("Running a shake pattern without an outer shake instance"));
+
+	FCameraShakePatternUpdateParams ChildParams(Params);
+
+	for (UCameraShakePattern* Pattern : ChildPatterns)
+	{
+		if (Pattern != nullptr && !Pattern->IsFinished())
+		{
+			// Let the child pattern run on the current result, with its own blending weight.
+			FCameraShakePatternUpdateResult ChildResult;
+			Pattern->UpdateShakePattern(ChildParams, ChildResult);
+
+			if (!Pattern->IsFinished())
+			{
+				// Apply this result and pass it on to the next child pattern.
+				FCameraShakeApplyResultParams ApplyParams;
+				ApplyParams.Scale = Params.GetTotalScale();
+				ApplyParams.PlaySpace = ShakeInstance->GetPlaySpace();
+				ApplyParams.UserPlaySpaceMatrix = ShakeInstance->GetUserPlaySpaceMatrix();
+				ApplyParams.CameraManager = ShakeInstance->GetCameraManager();
+				// This applies the current pattern's update to the parameters we'll pass
+				// to the next one.
+				UCameraShakeBase::ApplyResult(ApplyParams, ChildResult, ChildParams.POV);
+			}
+		}
+	}
+
+	// All our children patterns have applied their logic to our temp view, so we can put this in
+	// our own result, marking it as absolute.
+	OutResult.Location = ChildParams.POV.Location;
+	OutResult.Rotation = ChildParams.POV.Rotation;
+	OutResult.FOV = ChildParams.POV.FOV;
+	OutResult.Flags = ECameraShakePatternUpdateResultFlags::ApplyAsAbsolute;
+}
+
+void UCompositeCameraShakePattern::ScrubShakePatternImpl(const FCameraShakePatternScrubParams& Params, FCameraShakePatternUpdateResult& OutResult)
+{
+	// This method is similar to the UpdateShakePatternImpl method above, but calls the scrub methods
+	// instead of the update methods.
+
+	UCameraShakeBase* ShakeInstance = GetShakeInstance();
+	checkf(ShakeInstance, TEXT("Running a shake pattern without an outer shake instance"));
+
+	FCameraShakePatternScrubParams ChildParams(Params);
+
+	for (UCameraShakePattern* Pattern : ChildPatterns)
+	{
+		if (Pattern != nullptr) // Don't check for IsFinished here, we might scrub anywhere.
+		{
+			// Let the child pattern run on the current result, with its own blending weight.
+			FCameraShakePatternUpdateResult ChildResult;
+			Pattern->ScrubShakePattern(Params, ChildResult);
+
+			if (!Pattern->IsFinished())
+			{
+				// Apply this result and pass it on to the next child pattern.
+				FCameraShakeApplyResultParams ApplyParams;
+				ApplyParams.Scale = Params.GetTotalScale();
+				ApplyParams.PlaySpace = ShakeInstance->GetPlaySpace();
+				ApplyParams.UserPlaySpaceMatrix = ShakeInstance->GetUserPlaySpaceMatrix();
+				// This applies the current pattern's scrubbing to the parameters we'll pass
+				// to the next one.
+				UCameraShakeBase::ApplyResult(ApplyParams, ChildResult, ChildParams.POV);
+			}
+		}
+	}
+
+	// All our children patterns have applied their logic to our temp view, so we can put this in
+	// our own result, marking it as absolute.
+	OutResult.Location = ChildParams.POV.Location;
+	OutResult.Rotation = ChildParams.POV.Rotation;
+	OutResult.FOV = ChildParams.POV.FOV;
+	OutResult.Flags = ECameraShakePatternUpdateResultFlags::ApplyAsAbsolute;
+}
+
+bool UCompositeCameraShakePattern::IsFinishedImpl() const
+{
+	// We're not finished if any of our child patterns is not finished.
+	for (UCameraShakePattern* Pattern : ChildPatterns)
+	{
+		if (Pattern != nullptr && !Pattern->IsFinished())
+		{
+			return false;
+		}
+	}
+	return true;
+}
+
+void UCompositeCameraShakePattern::StopShakePatternImpl(const FCameraShakePatternStopParams& Params)
+{
+	// Stop all our children.
+	for (UCameraShakePattern* Pattern : ChildPatterns)
+	{
+		if (Pattern != nullptr)
+		{
+			Pattern->StopShakePattern(Params);
+		}
+	}
+}
+
+void UCompositeCameraShakePattern::TeardownShakePatternImpl()
+{
+	// Teardown are all our children.
+	for (UCameraShakePattern* Pattern : ChildPatterns)
+	{
+		if (Pattern != nullptr)
+		{
+			Pattern->TeardownShakePattern();
+		}
+	}
+}
+
diff --git a/Engine/Plugins/Cameras/GameplayCameras/Source/GameplayCameras/Private/Shakes/DefaultCameraShakeBase.cpp b/Engine/Plugins/Cameras/GameplayCameras/Source/GameplayCameras/Private/Shakes/DefaultCameraShakeBase.cpp
new file mode 100644
--- /dev/null
+++ b/Engine/Plugins/Cameras/GameplayCameras/Source/GameplayCameras/Private/Shakes/DefaultCameraShakeBase.cpp
@@ -0,0 +1,14 @@
+// Copyright Epic Games, Inc. All Rights Reserved.
+
+#include "Shakes/DefaultCameraShakeBase.h"
+#include "Shakes/PerlinNoiseCameraShakePattern.h"
+
+#include UE_INLINE_GENERATED_CPP_BY_NAME(DefaultCameraShakeBase)
+
+UDefaultCameraShakeBase::UDefaultCameraShakeBase(const FObjectInitializer& ObjInit)
+	: Super(ObjInit
+				.SetDefaultSubobjectClass<UPerlinNoiseCameraShakePattern>(TEXT("RootShakePattern")))
+{
+}
+
+
diff --git a/Engine/Plugins/Cameras/GameplayCameras/Source/GameplayCameras/Private/Shakes/LegacyCameraShake.cpp b/Engine/Plugins/Cameras/GameplayCameras/Source/GameplayCameras/Private/Shakes/LegacyCameraShake.cpp
new file mode 100644
--- /dev/null
+++ b/Engine/Plugins/Cameras/GameplayCameras/Source/GameplayCameras/Private/Shakes/LegacyCameraShake.cpp
@@ -0,0 +1,494 @@
+// Copyright Epic Games, Inc. All Rights Reserved.
+
+#include "Shakes/LegacyCameraShake.h"
+#include "CameraAnimationSequence.h"
+#include "SequenceCameraShake.h"
+#include "Camera/PlayerCameraManager.h"
+#include "Engine/Engine.h"
+#include "IXRTrackingSystem.h" // for IsHeadTrackingAllowed()
+#include "MovieScene.h"
+
+DEFINE_LOG_CATEGORY_STATIC(LogLegacyCameraShake, Warning, All);
+
+//////////////////////////////////////////////////////////////////////////
+// FFOscillator
+
+// static
+float FFOscillator::UpdateOffset(FFOscillator const& Osc, float& CurrentOffset, float DeltaTime)
+{
+	// LWC_TODO: Perf pessimization
+	double AsDouble = CurrentOffset;
+	float Result = UpdateOffset(Osc, AsDouble, DeltaTime);
+	CurrentOffset = (float)AsDouble;
+	return Result;
+}
+
+// static
+float FFOscillator::UpdateOffset(FFOscillator const& Osc, double& CurrentOffset, float DeltaTime)
+{
+	if (Osc.Amplitude != 0.f)
+	{
+		CurrentOffset += DeltaTime * Osc.Frequency;
+
+		float WaveformSample;
+		switch (Osc.Waveform)
+		{
+		case EOscillatorWaveform::SineWave:
+		default:
+			WaveformSample = FMath::Sin(CurrentOffset);
+			break;
+
+		case EOscillatorWaveform::PerlinNoise:
+			WaveformSample = FMath::PerlinNoise1D(CurrentOffset);
+			break;
+		}
+
+		return Osc.Amplitude * WaveformSample;
+	}
+	return 0.f;
+}
+
+// static
+float FFOscillator::GetInitialOffset(FFOscillator const& Osc)
+{
+	return (Osc.InitialOffset == EOO_OffsetRandom)
+		? FMath::FRand() * (2.f * PI)
+		: 0.f;
+}
+
+// static
+float FFOscillator::GetOffsetAtTime(FFOscillator const& Osc, float InitialOffset, float Time)
+{
+	return InitialOffset + (Time * Osc.Frequency);
+}
+
+//////////////////////////////////////////////////////////////////////////
+// ULegacyCameraShake
+
+ULegacyCameraShake::ULegacyCameraShake(const FObjectInitializer& ObjectInitializer)
+	: Super(ObjectInitializer
+				.SetDefaultSubobjectClass<ULegacyCameraShakePattern>(TEXT("RootShakePattern")))
+{
+	AnimPlayRate = 1.0f;
+	AnimScale = 1.0f;
+	AnimBlendInTime = 0.2f;
+	AnimBlendOutTime = 0.2f;
+	OscillationBlendInTime = 0.1f;
+	OscillationBlendOutTime = 0.2f;
+}
+
+void ULegacyCameraShake::DoStopShake(bool bImmediately)
+{
+	if (bImmediately)
+	{
+		// stop oscillation
+		OscillatorTimeRemaining = 0.f;
+	}
+	else
+	{
+		// advance to the blend out time
+		if (OscillatorTimeRemaining > 0.0f)
+		{
+			OscillatorTimeRemaining = FMath::Min(OscillatorTimeRemaining, OscillationBlendOutTime);
+		}
+		else
+		{
+			OscillatorTimeRemaining = OscillationBlendOutTime;
+		}
+	}
+	
+	if (SequenceShakePattern)
+	{
+		FCameraShakePatternStopParams StopParams;
+		StopParams.bImmediately = bImmediately;
+		SequenceShakePattern->StopShakePattern(StopParams);
+	}
+
+	UE_LOG(LogLegacyCameraShake, Verbose, TEXT("ULegacyCameraShake::DoStopShake %s"), *GetNameSafe(this));
+
+	ReceiveStopShake(bImmediately);
+}
+
+void ULegacyCameraShake::DoStartShake(const FCameraShakePatternStartParams& Params)
+{
+	ActualOscillationDuration = Params.bOverrideDuration ? Params.DurationOverride : OscillationDuration;
+
+	const float EffectiveOscillationDuration = (ActualOscillationDuration > 0.f) ? ActualOscillationDuration : TNumericLimits<float>::Max();
+
+	// init oscillations
+	if (ActualOscillationDuration != 0.f)
+	{
+		if (OscillatorTimeRemaining > 0.f)
+		{
+			// this shake was already playing
+			OscillatorTimeRemaining = EffectiveOscillationDuration;
+
+			if (bBlendingOut)
+			{
+				bBlendingOut = false;
+				CurrentBlendOutTime = 0.f;
+
+				// stop any blendout and reverse it to a blendin
+				if (OscillationBlendInTime > 0.f)
+				{
+					bBlendingIn = true;
+					CurrentBlendInTime = OscillationBlendInTime * (1.f - CurrentBlendOutTime / OscillationBlendOutTime);
+				}
+				else
+				{
+					bBlendingIn = false;
+					CurrentBlendInTime = 0.f;
+				}
+			}
+		}
+		else
+		{
+			RotSinOffset.X = FFOscillator::GetInitialOffset(RotOscillation.Pitch);
+			RotSinOffset.Y = FFOscillator::GetInitialOffset(RotOscillation.Yaw);
+			RotSinOffset.Z = FFOscillator::GetInitialOffset(RotOscillation.Roll);
+
+			LocSinOffset.X = FFOscillator::GetInitialOffset(LocOscillation.X);
+			LocSinOffset.Y = FFOscillator::GetInitialOffset(LocOscillation.Y);
+			LocSinOffset.Z = FFOscillator::GetInitialOffset(LocOscillation.Z);
+
+			FOVSinOffset = FFOscillator::GetInitialOffset(FOVOscillation);
+
+			InitialLocSinOffset = LocSinOffset;
+			InitialRotSinOffset = RotSinOffset;
+			InitialFOVSinOffset = FOVSinOffset;
+
+			OscillatorTimeRemaining = EffectiveOscillationDuration;
+
+			if (OscillationBlendInTime > 0.f)
+			{
+				bBlendingIn = true;
+				CurrentBlendInTime = 0.f;
+			}
+		}
+	}
+
+	// init cameraanim shakes
+	if (AnimSequence != nullptr)
+	{
+		if (SequenceShakePattern == nullptr)
+		{
+			SequenceShakePattern = NewObject<USequenceCameraShakePattern>(this);
+		}
+
+		// Copy our anim parameters over to the sequence shake pattern.
+		SequenceShakePattern->Sequence = AnimSequence;
+		SequenceShakePattern->PlayRate = AnimPlayRate;
+		SequenceShakePattern->Scale = AnimScale;
+		SequenceShakePattern->BlendInTime = AnimBlendInTime;
+		SequenceShakePattern->BlendOutTime = AnimBlendOutTime;
+		SequenceShakePattern->RandomSegmentDuration = RandomAnimSegmentDuration;
+		SequenceShakePattern->bRandomSegment = bRandomAnimSegment;
+
+		// Start the sequence shake pattern.
+		SequenceShakePattern->StartShakePattern(Params);
+	}
+
+	UE_LOG(LogLegacyCameraShake, Verbose, TEXT("ULegacyCameraShake::DoStartShake %s Duration: %f"), *GetNameSafe(this), ActualOscillationDuration);
+
+	ReceivePlayShake(ShakeScale);
+}
+
+void ULegacyCameraShake::DoUpdateShake(const FCameraShakePatternUpdateParams& Params, FCameraShakePatternUpdateResult& OutResult)
+{
+	const float DeltaTime = Params.DeltaTime;
+
+	// update oscillation times... only decrease the time remaining if we're not infinite
+	if (OscillatorTimeRemaining > 0.f)
+	{
+		OscillatorTimeRemaining -= DeltaTime;
+		OscillatorTimeRemaining = FMath::Max(0.f, OscillatorTimeRemaining);
+	}
+	if (bBlendingIn)
+	{
+		CurrentBlendInTime += DeltaTime;
+	}
+	if (bBlendingOut)
+	{
+		CurrentBlendOutTime += DeltaTime;
+	}
+
+	// see if we've crossed any important time thresholds and deal appropriately
+	bool bOscillationFinished = false;
+
+	if (OscillatorTimeRemaining <= 0.f)
+	{
+		// finished!
+		bOscillationFinished = true;
+	}
+	else if (OscillatorTimeRemaining < OscillationBlendOutTime)
+	{
+		// start blending out
+		bBlendingOut = true;
+		CurrentBlendOutTime = OscillationBlendOutTime - OscillatorTimeRemaining;
+	}
+	else if (ActualOscillationDuration < 0.f)
+	{
+		// infinite oscillation, keep the time remaining up
+		OscillatorTimeRemaining = TNumericLimits<float>::Max();
+	}
+
+	if (bBlendingIn)
+	{
+		if (CurrentBlendInTime > OscillationBlendInTime)
+		{
+			// done blending in!
+			bBlendingIn = false;
+		}
+	}
+	if (bBlendingOut)
+	{
+		if (CurrentBlendOutTime > OscillationBlendOutTime)
+		{
+			// done!!
+			CurrentBlendOutTime = OscillationBlendOutTime;
+			bOscillationFinished = true;
+		}
+	}
+
+	// Do not update oscillation further if finished
+	if (bOscillationFinished == false)
+	{
+		// calculate blend weight. calculating separately and taking the minimum handles overlapping blends nicely.
+		float const BlendInWeight = (bBlendingIn) ? (CurrentBlendInTime / OscillationBlendInTime) : 1.f;
+		float const BlendOutWeight = (bBlendingOut) ? (1.f - CurrentBlendOutTime / OscillationBlendOutTime) : 1.f;
+		float const CurrentBlendWeight = FMath::Min(BlendInWeight, BlendOutWeight);
+
+		// this is the oscillation scale, which includes oscillation fading
+		// we'll apply the general shake scale, along with the current frame's dynamic scale, a bit later.
+		float const OscillationScale = CurrentBlendWeight;
+
+		if (OscillationScale > 0.f)
+		{
+			// View location offset, compute sin wave value for each component
+			FVector	LocOffset = FVector(0);
+			LocOffset.X = FFOscillator::UpdateOffset(LocOscillation.X, LocSinOffset.X, DeltaTime);
+			LocOffset.Y = FFOscillator::UpdateOffset(LocOscillation.Y, LocSinOffset.Y, DeltaTime);
+			LocOffset.Z = FFOscillator::UpdateOffset(LocOscillation.Z, LocSinOffset.Z, DeltaTime);
+			LocOffset *= OscillationScale;
+
+			OutResult.Location = LocOffset;
+
+			// View rotation offset, compute sin wave value for each component
+			FRotator RotOffset;
+			RotOffset.Pitch = FFOscillator::UpdateOffset(RotOscillation.Pitch, RotSinOffset.X, DeltaTime) * OscillationScale;
+			RotOffset.Yaw = FFOscillator::UpdateOffset(RotOscillation.Yaw, RotSinOffset.Y, DeltaTime) * OscillationScale;
+			RotOffset.Roll = FFOscillator::UpdateOffset(RotOscillation.Roll, RotSinOffset.Z, DeltaTime) * OscillationScale;
+
+			// Don't allow shake to flip pitch past vertical, if not using a headset (where we can't limit the camera locked to your head).
+			APlayerCameraManager* CameraOwner = GetCameraManager();
+			UWorld * World = (CameraOwner ? CameraOwner->GetWorld() : nullptr);
+			if (!GEngine->XRSystem.IsValid() ||
+				!(World != nullptr ? GEngine->XRSystem->IsHeadTrackingAllowedForWorld(*World) : GEngine->XRSystem->IsHeadTrackingAllowed()))
+			{
+				// Find normalized result when combined, and remove any offset that would push it past the limit.
+				const float NormalizedInputPitch = FRotator::NormalizeAxis(Params.POV.Rotation.Pitch);
+				RotOffset.Pitch = FRotator::NormalizeAxis(RotOffset.Pitch);
+				RotOffset.Pitch = FMath::ClampAngle(NormalizedInputPitch + RotOffset.Pitch, -89.9f, 89.9f) - NormalizedInputPitch;
+			}
+
+			OutResult.Rotation = RotOffset;
+
+			// Compute FOV change
+			OutResult.FOV = OscillationScale * FFOscillator::UpdateOffset(FOVOscillation, FOVSinOffset, DeltaTime);
+		}
+	}
+
+	// Update the sequence animation if there's one.
+	if (SequenceShakePattern != nullptr && !SequenceShakePattern->IsFinished())
+	{
+		FCameraShakePatternUpdateResult ChildResult;
+		SequenceShakePattern->UpdateShakePattern(Params, ChildResult);
+
+		// The sequence shake pattern returns a local, additive result. So we should be able to
+		// just combine the two results directly.
+		OutResult.Location += ChildResult.Location;
+		OutResult.Rotation += ChildResult.Rotation;
+		OutResult.FOV += ChildResult.FOV;
+		// We don't have anything else animating post-process settings so we can stomp them.
+		OutResult.PostProcessSettings = ChildResult.PostProcessSettings;
+		OutResult.PostProcessBlendWeight = ChildResult.PostProcessBlendWeight;
+	}
+
+	// Apply the scaling, limits, and playspace so we have an absolute result we can pass to the legacy blueprint API.
+	check(OutResult.Flags == ECameraShakePatternUpdateResultFlags::Default);
+	ApplyScale(Params.GetTotalScale(), OutResult);
+	ApplyLimits(Params.POV, OutResult);
+	ApplyPlaySpace(Params, OutResult);
+	check(EnumHasAnyFlags(OutResult.Flags, ECameraShakePatternUpdateResultFlags::ApplyAsAbsolute));
+
+	// Call the legacy blueprint API. We need to convert back and forth.
+	{
+		FMinimalViewInfo InOutPOV(Params.POV);
+		InOutPOV.Location = OutResult.Location;
+		InOutPOV.Rotation = OutResult.Rotation;
+		InOutPOV.FOV = OutResult.FOV;
+
+		BlueprintUpdateCameraShake(DeltaTime, Params.DynamicScale, InOutPOV, InOutPOV);
+
+		OutResult.Location = InOutPOV.Location;
+		OutResult.Rotation = InOutPOV.Rotation;
+		OutResult.FOV = InOutPOV.FOV;
+	}
+
+	UE_LOG(LogLegacyCameraShake, Verbose, TEXT("ULegacyCameraShake::DoUpdateShake %s Finished: %i Duration: %f Remaining: %f"), *GetNameSafe(this), bOscillationFinished, ActualOscillationDuration, OscillatorTimeRemaining);
+}
+
+void ULegacyCameraShake::DoScrubShake(const FCameraShakePatternScrubParams& Params, FCameraShakePatternUpdateResult& OutResult)
+{
+	const float NewTime = Params.AbsoluteTime;
+
+	// reset to start and advance to desired point
+	LocSinOffset = InitialLocSinOffset;
+	RotSinOffset = InitialRotSinOffset;
+	FOVSinOffset = InitialFOVSinOffset;
+
+	const float EffectiveOscillationDuration = (ActualOscillationDuration > 0.f) ? ActualOscillationDuration : TNumericLimits<float>::Max();
+
+	OscillatorTimeRemaining = EffectiveOscillationDuration;
+
+	if (OscillationBlendInTime > 0.f)
+	{
+		bBlendingIn = true;
+		CurrentBlendInTime = 0.f;
+	}
+
+	if (OscillationBlendOutTime > 0.f)
+	{
+		bBlendingOut = false;
+		CurrentBlendOutTime = 0.f;
+	}
+
+	if (ActualOscillationDuration > 0.f)
+	{
+		if ((OscillationBlendOutTime > 0.f) && (NewTime > (ActualOscillationDuration - OscillationBlendOutTime)))
+		{
+			bBlendingOut = true;
+			CurrentBlendOutTime = OscillationBlendOutTime - (ActualOscillationDuration - NewTime);
+		}
+	}
+
+	FCameraShakePatternUpdateParams UpdateParams = Params.ToUpdateParams();
+
+	DoUpdateShake(UpdateParams, OutResult);
+
+	check(EnumHasAnyFlags(OutResult.Flags, ECameraShakePatternUpdateResultFlags::ApplyAsAbsolute));
+}
+
+bool ULegacyCameraShake::DoGetIsFinished() const
+{
+	return ((OscillatorTimeRemaining <= 0.f) &&									// oscillator is finished
+		((SequenceShakePattern == nullptr) || SequenceShakePattern->IsFinished()) && // other anim is finished
+		ReceiveIsFinished()														// BP thinks it's finished
+		);
+}
+
+void ULegacyCameraShake::DoTeardownShake()
+{
+	if (SequenceShakePattern)
+	{
+		SequenceShakePattern->TeardownShakePattern();
+	}
+}
+
+/// @cond DOXYGEN_WARNINGS
+
+bool ULegacyCameraShake::ReceiveIsFinished_Implementation() const
+{
+	return true;
+}
+
+/// @endcond
+
+bool ULegacyCameraShake::IsLooping() const
+{
+	return ActualOscillationDuration < 0.0f;
+}
+
+void ULegacyCameraShake::SetCurrentTimeAndApplyShake(float NewTime, FMinimalViewInfo& POV)
+{
+	ScrubAndApplyCameraShake(NewTime, 1.f, POV);
+}
+
+ULegacyCameraShake* ULegacyCameraShake::StartLegacyCameraShake(APlayerCameraManager* PlayerCameraManager, TSubclassOf<ULegacyCameraShake> ShakeClass, float Scale, ECameraShakePlaySpace PlaySpace, FRotator UserPlaySpaceRot)
+{
+	if (PlayerCameraManager)
+	{
+		return Cast<ULegacyCameraShake>(PlayerCameraManager->StartCameraShake(ShakeClass, Scale, PlaySpace, UserPlaySpaceRot));
+	}
+
+	return nullptr;
+}
+
+ULegacyCameraShake* ULegacyCameraShake::StartLegacyCameraShakeFromSource(APlayerCameraManager* PlayerCameraManager, TSubclassOf<ULegacyCameraShake> ShakeClass, UCameraShakeSourceComponent* SourceComponent, float Scale, ECameraShakePlaySpace PlaySpace, FRotator UserPlaySpaceRot)
+{
+	if (PlayerCameraManager)
+	{
+		return Cast<ULegacyCameraShake>(PlayerCameraManager->StartCameraShakeFromSource(ShakeClass, SourceComponent, Scale, PlaySpace, UserPlaySpaceRot));
+	}
+
+	return nullptr;
+}
+
+void ULegacyCameraShakePattern::GetShakePatternInfoImpl(FCameraShakeInfo& OutInfo) const
+{
+	// We will manage our own duration, but let's give a hint about how long we are for editor purposes.
+	ULegacyCameraShake* Shake = GetShakeInstance<ULegacyCameraShake>();
+	
+	if (Shake->AnimSequence)
+	{
+		float PlaybackSeconds = 0.f;
+		if (UMovieScene* MovieScene = Shake->AnimSequence->GetMovieScene())
+		{
+			const FFrameNumber PlaybackFrames = MovieScene->GetPlaybackRange().Size<FFrameNumber>();
+			PlaybackSeconds = MovieScene->GetTickResolution().AsSeconds(PlaybackFrames);
+		}
+		const float Duration = FMath::Max(Shake->OscillationDuration, PlaybackSeconds);
+		OutInfo.Duration = FCameraShakeDuration::Custom(Duration);
+	}
+	else
+	{
+		OutInfo.Duration = FCameraShakeDuration::Custom(Shake->OscillationDuration);
+	}
+}
+
+void ULegacyCameraShakePattern::StopShakePatternImpl(const FCameraShakePatternStopParams& Params)
+{
+	ULegacyCameraShake* Shake = GetShakeInstance<ULegacyCameraShake>();
+	Shake->DoStopShake(Params.bImmediately);
+}
+
+void ULegacyCameraShakePattern::StartShakePatternImpl(const FCameraShakePatternStartParams& Params)
+{
+	ULegacyCameraShake* Shake = GetShakeInstance<ULegacyCameraShake>();
+	Shake->DoStartShake(Params);
+}
+
+void ULegacyCameraShakePattern::UpdateShakePatternImpl(const FCameraShakePatternUpdateParams& Params, FCameraShakePatternUpdateResult& OutResult)
+{
+	ULegacyCameraShake* Shake = GetShakeInstance<ULegacyCameraShake>();
+	Shake->DoUpdateShake(Params, OutResult);
+}
+
+void ULegacyCameraShakePattern::ScrubShakePatternImpl(const FCameraShakePatternScrubParams& Params, FCameraShakePatternUpdateResult& OutResult)
+{
+	ULegacyCameraShake* Shake = GetShakeInstance<ULegacyCameraShake>();
+	Shake->DoScrubShake(Params, OutResult);
+}
+
+bool ULegacyCameraShakePattern::IsFinishedImpl() const
+{
+	ULegacyCameraShake* Shake = GetShakeInstance<ULegacyCameraShake>();
+	return Shake->DoGetIsFinished();
+}
+
+void ULegacyCameraShakePattern::TeardownShakePatternImpl()
+{
+	ULegacyCameraShake* Shake = GetShakeInstance<ULegacyCameraShake>();
+	Shake->DoTeardownShake();
+}
+
diff --git a/Engine/Plugins/Cameras/GameplayCameras/Source/GameplayCameras/Private/Shakes/PerlinNoiseCameraShakePattern.cpp b/Engine/Plugins/Cameras/GameplayCameras/Source/GameplayCameras/Private/Shakes/PerlinNoiseCameraShakePattern.cpp
new file mode 100644
--- /dev/null
+++ b/Engine/Plugins/Cameras/GameplayCameras/Source/GameplayCameras/Private/Shakes/PerlinNoiseCameraShakePattern.cpp
@@ -0,0 +1,79 @@
+// Copyright Epic Games, Inc. All Rights Reserved.
+
+#include "Shakes/PerlinNoiseCameraShakePattern.h"
+
+#include UE_INLINE_GENERATED_CPP_BY_NAME(PerlinNoiseCameraShakePattern)
+
+float FPerlinNoiseShaker::Update(float DeltaTime, float AmplitudeMultiplier, float FrequencyMultiplier, float& InOutCurrentOffset) const
+{
+	const float TotalAmplitude = Amplitude * AmplitudeMultiplier;
+	if (TotalAmplitude != 0.f)
+	{
+		InOutCurrentOffset += DeltaTime * Frequency * FrequencyMultiplier;
+		return TotalAmplitude * FMath::PerlinNoise1D(InOutCurrentOffset);
+	}
+	return 0.f;
+}
+
+UPerlinNoiseCameraShakePattern::UPerlinNoiseCameraShakePattern(const FObjectInitializer& ObjInit)
+	: Super(ObjInit)
+{
+	// Default to only location shaking.
+	RotationAmplitudeMultiplier = 0.f;
+	FOV.Amplitude = 0.f;
+}
+
+void UPerlinNoiseCameraShakePattern::StartShakePatternImpl(const FCameraShakePatternStartParams& Params)
+{
+	Super::StartShakePatternImpl(Params);
+
+	if (!Params.bIsRestarting)
+	{
+		// All offsets are random. This is because the core perlin noise implementation
+		// uses permutation tables, so if two shakers have the same initial offset and the same
+		// frequency, they will have the same exact values.
+		InitialLocationOffset = FVector3f((float)FMath::RandHelper(255), (float)FMath::RandHelper(255), (float)FMath::RandHelper(255));
+		InitialRotationOffset = FVector3f((float)FMath::RandHelper(255), (float)FMath::RandHelper(255), (float)FMath::RandHelper(255));
+		InitialFOVOffset = (float)FMath::RandHelper(255);
+
+		CurrentLocationOffset = InitialLocationOffset;
+		CurrentRotationOffset = InitialRotationOffset;
+		CurrentFOVOffset = InitialFOVOffset;
+	}
+}
+
+void UPerlinNoiseCameraShakePattern::UpdateShakePatternImpl(const FCameraShakePatternUpdateParams& Params, FCameraShakePatternUpdateResult& OutResult)
+{
+	UpdatePerlinNoise(Params.DeltaTime, OutResult);
+
+	const float BlendWeight = State.Update(Params.DeltaTime);
+	OutResult.ApplyScale(BlendWeight);
+}
+
+void UPerlinNoiseCameraShakePattern::ScrubShakePatternImpl(const FCameraShakePatternScrubParams& Params, FCameraShakePatternUpdateResult& OutResult)
+{
+	// Scrubbing is like going back to our initial state and updating directly to the scrub time.
+	CurrentLocationOffset = InitialLocationOffset;
+	CurrentRotationOffset = InitialRotationOffset;
+	CurrentFOVOffset = InitialFOVOffset;
+
+	UpdatePerlinNoise(Params.AbsoluteTime, OutResult);
+
+	const float BlendWeight = State.Scrub(Params.AbsoluteTime);
+	OutResult.ApplyScale(BlendWeight);
+}
+
+void UPerlinNoiseCameraShakePattern::UpdatePerlinNoise(float DeltaTime, FCameraShakePatternUpdateResult& OutResult)
+{
+	OutResult.Location.X = X.Update(DeltaTime, LocationAmplitudeMultiplier, LocationFrequencyMultiplier, CurrentLocationOffset.X);
+	OutResult.Location.Y = Y.Update(DeltaTime, LocationAmplitudeMultiplier, LocationFrequencyMultiplier, CurrentLocationOffset.Y);
+	OutResult.Location.Z = Z.Update(DeltaTime, LocationAmplitudeMultiplier, LocationFrequencyMultiplier, CurrentLocationOffset.Z);
+
+	OutResult.Rotation.Pitch = Pitch.Update(DeltaTime, RotationAmplitudeMultiplier, RotationFrequencyMultiplier, CurrentRotationOffset.X);
+	OutResult.Rotation.Yaw   = Yaw.Update(  DeltaTime, RotationAmplitudeMultiplier, RotationFrequencyMultiplier, CurrentRotationOffset.Y);
+	OutResult.Rotation.Roll  = Roll.Update( DeltaTime, RotationAmplitudeMultiplier, RotationFrequencyMultiplier, CurrentRotationOffset.Z);
+
+	OutResult.FOV = FOV.Update(DeltaTime, 1.f, 1.f, CurrentFOVOffset);
+}
+
+
diff --git a/Engine/Plugins/Cameras/GameplayCameras/Source/GameplayCameras/Private/Shakes/SimpleCameraShakePattern.cpp b/Engine/Plugins/Cameras/GameplayCameras/Source/GameplayCameras/Private/Shakes/SimpleCameraShakePattern.cpp
new file mode 100644
--- /dev/null
+++ b/Engine/Plugins/Cameras/GameplayCameras/Source/GameplayCameras/Private/Shakes/SimpleCameraShakePattern.cpp
@@ -0,0 +1,41 @@
+// Copyright Epic Games, Inc. All Rights Reserved.
+
+#include "Shakes/SimpleCameraShakePattern.h"
+
+#include UE_INLINE_GENERATED_CPP_BY_NAME(SimpleCameraShakePattern)
+
+void USimpleCameraShakePattern::GetShakePatternInfoImpl(FCameraShakeInfo& OutInfo) const
+{
+	if (Duration > 0.f)
+	{
+		OutInfo.Duration = Duration;
+	}
+	else
+	{
+		OutInfo.Duration = FCameraShakeDuration::Infinite();
+	}
+
+	OutInfo.BlendIn = BlendInTime;
+	OutInfo.BlendOut = BlendOutTime;
+}
+
+void USimpleCameraShakePattern::StartShakePatternImpl(const FCameraShakePatternStartParams& Params)
+{
+	State.Start(this, Params);
+}
+
+bool USimpleCameraShakePattern::IsFinishedImpl() const
+{
+	return !State.IsPlaying();
+}
+
+void USimpleCameraShakePattern::StopShakePatternImpl(const FCameraShakePatternStopParams& Params)
+{
+	State.Stop(Params.bImmediately);
+}
+
+void USimpleCameraShakePattern::TeardownShakePatternImpl()
+{
+	State = FCameraShakeState();
+}
+
diff --git a/Engine/Plugins/Cameras/GameplayCameras/Source/GameplayCameras/Private/Shakes/WaveOscillatorCameraShakePattern.cpp b/Engine/Plugins/Cameras/GameplayCameras/Source/GameplayCameras/Private/Shakes/WaveOscillatorCameraShakePattern.cpp
new file mode 100644
--- /dev/null
+++ b/Engine/Plugins/Cameras/GameplayCameras/Source/GameplayCameras/Private/Shakes/WaveOscillatorCameraShakePattern.cpp
@@ -0,0 +1,92 @@
+// Copyright Epic Games, Inc. All Rights Reserved.
+
+#include "Shakes/WaveOscillatorCameraShakePattern.h"
+
+#include UE_INLINE_GENERATED_CPP_BY_NAME(WaveOscillatorCameraShakePattern)
+
+float FWaveOscillator::Initialize(float& OutInitialOffset) const
+{
+	OutInitialOffset = (InitialOffsetType == EInitialWaveOscillatorOffsetType::Random)
+		? FMath::FRand() * (2.f * PI)
+		: 0.f;
+	return Amplitude * FMath::Sin(OutInitialOffset);
+}
+
+float FWaveOscillator::Update(float DeltaTime, float AmplitudeMultiplier, float FrequencyMultiplier, float& InOutCurrentOffset) const
+{
+	const float TotalAmplitude = Amplitude * AmplitudeMultiplier;
+	if (TotalAmplitude != 0.f)
+	{
+		InOutCurrentOffset += DeltaTime * Frequency * FrequencyMultiplier * (2.f * PI);
+		return TotalAmplitude * FMath::Sin(InOutCurrentOffset);
+	}
+	return 0.f;
+}
+
+UWaveOscillatorCameraShakePattern::UWaveOscillatorCameraShakePattern(const FObjectInitializer& ObjInit)
+	: Super(ObjInit)
+{
+	// Default to only location shaking.
+	RotationAmplitudeMultiplier = 0.f;
+	FOV.Amplitude = 0.f;
+}
+
+void UWaveOscillatorCameraShakePattern::StartShakePatternImpl(const FCameraShakePatternStartParams& Params)
+{
+	Super::StartShakePatternImpl(Params);
+
+	if (!Params.bIsRestarting)
+	{
+		X.Initialize(InitialLocationOffset.X);
+		Y.Initialize(InitialLocationOffset.Y);
+		Z.Initialize(InitialLocationOffset.Z);
+
+		CurrentLocationOffset = InitialLocationOffset;
+
+		Pitch.Initialize(InitialRotationOffset.X);
+		Yaw.Initialize(  InitialRotationOffset.Y);
+		Roll.Initialize( InitialRotationOffset.Z);
+
+		CurrentRotationOffset = InitialRotationOffset;
+
+		FOV.Initialize(InitialFOVOffset);
+
+		CurrentFOVOffset = InitialFOVOffset;
+	}
+}
+
+void UWaveOscillatorCameraShakePattern::UpdateShakePatternImpl(const FCameraShakePatternUpdateParams& Params, FCameraShakePatternUpdateResult& OutResult)
+{
+	UpdateOscillators(Params.DeltaTime, OutResult);
+
+	const float BlendWeight = State.Update(Params.DeltaTime);
+	OutResult.ApplyScale(BlendWeight);
+}
+
+void UWaveOscillatorCameraShakePattern::ScrubShakePatternImpl(const FCameraShakePatternScrubParams& Params, FCameraShakePatternUpdateResult& OutResult)
+{
+	// Scrubbing is like going back to our initial state and updating directly to the scrub time.
+	CurrentLocationOffset = InitialLocationOffset;
+	CurrentRotationOffset = InitialRotationOffset;
+	CurrentFOVOffset = InitialFOVOffset;
+	
+	UpdateOscillators(Params.AbsoluteTime, OutResult);
+
+	const float BlendWeight = State.Scrub(Params.AbsoluteTime);
+	OutResult.ApplyScale(BlendWeight);
+}
+
+void UWaveOscillatorCameraShakePattern::UpdateOscillators(float DeltaTime, FCameraShakePatternUpdateResult& OutResult)
+{
+	OutResult.Location.X = X.Update(DeltaTime, LocationAmplitudeMultiplier, LocationFrequencyMultiplier, CurrentLocationOffset.X);
+	OutResult.Location.Y = Y.Update(DeltaTime, LocationAmplitudeMultiplier, LocationFrequencyMultiplier, CurrentLocationOffset.Y);
+	OutResult.Location.Z = Z.Update(DeltaTime, LocationAmplitudeMultiplier, LocationFrequencyMultiplier, CurrentLocationOffset.Z);
+
+	OutResult.Rotation.Pitch = Pitch.Update(DeltaTime, RotationAmplitudeMultiplier, RotationFrequencyMultiplier, CurrentRotationOffset.X);
+	OutResult.Rotation.Yaw   = Yaw.Update(  DeltaTime, RotationAmplitudeMultiplier, RotationFrequencyMultiplier, CurrentRotationOffset.Y);
+	OutResult.Rotation.Roll  = Roll.Update( DeltaTime, RotationAmplitudeMultiplier, RotationFrequencyMultiplier, CurrentRotationOffset.Z);
+
+	OutResult.FOV = FOV.Update(DeltaTime, 1.f, 1.f, CurrentFOVOffset);
+}
+
+
diff --git a/Engine/Plugins/Cameras/GameplayCameras/Source/GameplayCameras/Private/SimpleCameraShakePattern.cpp b/Engine/Plugins/Cameras/GameplayCameras/Source/GameplayCameras/Private/SimpleCameraShakePattern.cpp
deleted file mode 100644
diff --git a/Engine/Plugins/Cameras/GameplayCameras/Source/GameplayCameras/Private/Tests/CameraShakeTestObjects.h b/Engine/Plugins/Cameras/GameplayCameras/Source/GameplayCameras/Private/Tests/CameraShakeTestObjects.h
--- a/Engine/Plugins/Cameras/GameplayCameras/Source/GameplayCameras/Private/Tests/CameraShakeTestObjects.h
+++ b/Engine/Plugins/Cameras/GameplayCameras/Source/GameplayCameras/Private/Tests/CameraShakeTestObjects.h
@@ -2,7 +2,7 @@
 
 #pragma once
 
-#include "SimpleCameraShakePattern.h"
+#include "Shakes/SimpleCameraShakePattern.h"
 #include "CameraShakeTestObjects.generated.h"
 
 template<typename T>
diff --git a/Engine/Plugins/Cameras/GameplayCameras/Source/GameplayCameras/Private/Tests/CameraShakeTests.cpp b/Engine/Plugins/Cameras/GameplayCameras/Source/GameplayCameras/Private/Tests/CameraShakeTests.cpp
--- a/Engine/Plugins/Cameras/GameplayCameras/Source/GameplayCameras/Private/Tests/CameraShakeTests.cpp
+++ b/Engine/Plugins/Cameras/GameplayCameras/Source/GameplayCameras/Private/Tests/CameraShakeTests.cpp
@@ -1,9 +1,9 @@
 // Copyright Epic Games, Inc. All Rights Reserved.
 
 #include "CameraShakeTestObjects.h"
-#include "DefaultCameraShakeBase.h"
-#include "CompositeCameraShakePattern.h"
-#include "WaveOscillatorCameraShakePattern.h"
+#include "Shakes/DefaultCameraShakeBase.h"
+#include "Shakes/CompositeCameraShakePattern.h"
+#include "Shakes/WaveOscillatorCameraShakePattern.h"
 #include "Misc/AutomationTest.h"
 
 #define LOCTEXT_NAMESPACE "CameraShakeTests"
diff --git a/Engine/Plugins/Cameras/GameplayCameras/Source/GameplayCameras/Private/WaveOscillatorCameraShakePattern.cpp b/Engine/Plugins/Cameras/GameplayCameras/Source/GameplayCameras/Private/WaveOscillatorCameraShakePattern.cpp
deleted file mode 100644
diff --git a/Engine/Plugins/Cameras/GameplayCameras/Source/GameplayCameras/Public/Animations/CameraAnimationCameraModifier.h b/Engine/Plugins/Cameras/GameplayCameras/Source/GameplayCameras/Public/Animations/CameraAnimationCameraModifier.h
new file mode 100644
--- /dev/null
+++ b/Engine/Plugins/Cameras/GameplayCameras/Source/GameplayCameras/Public/Animations/CameraAnimationCameraModifier.h
@@ -0,0 +1,291 @@
+// Copyright Epic Games, Inc. All Rights Reserved.
+
+#pragma once
+
+#include "Camera/CameraModifier.h"
+#include "Kismet/BlueprintFunctionLibrary.h"
+#include "CameraAnimationCameraModifier.generated.h"
+
+enum class ECameraShakePlaySpace : uint8;
+
+class UCameraAnimationSequence;
+class UCameraAnimationSequenceCameraStandIn;
+class UCameraAnimationSequencePlayer;
+struct FMinimalViewInfo;
+
+UENUM()
+enum class ECameraAnimationPlaySpace : uint8
+{
+	/** This anim is applied in camera space. */
+	CameraLocal,
+	/** This anim is applied in world space. */
+	World,
+	/** This anim is applied in a user-specified space (defined by UserPlaySpaceMatrix). */
+	UserDefined,
+};
+
+UENUM()
+enum class ECameraAnimationEasingType : uint8
+{
+	Linear,
+	Sinusoidal,
+	Quadratic,
+	Cubic,
+	Quartic,
+	Quintic,
+	Exponential,
+	Circular,
+};
+
+/** Parameter struct for adding new camera animations to UCameraAnimationCameraModifier */
+USTRUCT(BlueprintType)
+struct FCameraAnimationParams
+{
+	GENERATED_BODY()
+
+	/** Time scale for playing the new camera animation */
+	UPROPERTY(EditAnywhere, BlueprintReadWrite, Category="Camera Animation")
+	float PlayRate = 1.f;
+	/** Global scale to use for the new camera animation */
+	UPROPERTY(EditAnywhere, BlueprintReadWrite, Category="Camera Animation")
+	float Scale = 1.f;
+
+	/** Ease-in function type */
+	UPROPERTY(EditAnywhere, BlueprintReadWrite, Category="Camera Animation")
+	ECameraAnimationEasingType EaseInType = ECameraAnimationEasingType::Linear;
+	/** Ease-in duration in seconds */
+	UPROPERTY(EditAnywhere, BlueprintReadWrite, Category="Camera Animation")
+	float EaseInDuration = 0.f;
+
+	/** Ease-out function type */
+	UPROPERTY(EditAnywhere, BlueprintReadWrite, Category="Camera Animation")
+	ECameraAnimationEasingType EaseOutType = ECameraAnimationEasingType::Linear;
+	/** Ease-out duration in seconds */
+	UPROPERTY(EditAnywhere, BlueprintReadWrite, Category="Camera Animation")
+	float EaseOutDuration = 0.f;
+
+	/** Whether the camera animation should loop */
+	UPROPERTY(EditAnywhere, BlueprintReadWrite, Category="Camera Animation")
+	bool bLoop = false;
+	/** Offset, in frames, into the animation to start at */
+	UPROPERTY(EditAnywhere, BlueprintReadWrite, Category="Camera Animation")
+	int32 StartOffset = 0;
+	/** Whether the camera animation should have a random start time */
+	UPROPERTY(EditAnywhere, BlueprintReadWrite, Category="Camera Animation")
+	bool bRandomStartTime = false;
+	/** Override the duration of the animation with a new duration (including blends) */
+	UPROPERTY(EditAnywhere, BlueprintReadWrite, Category="Camera Animation")
+	float DurationOverride = 0.f;
+
+	/** The transform space to use for the new camera shake */
+	UPROPERTY(EditAnywhere, BlueprintReadWrite, Category="Camera Animation")
+	ECameraAnimationPlaySpace PlaySpace = ECameraAnimationPlaySpace::CameraLocal;
+	/** User space to use when PlaySpace is UserDefined */
+	UPROPERTY(EditAnywhere, BlueprintReadWrite, Category="Camera Animation")
+	FRotator UserPlaySpaceRot = FRotator::ZeroRotator;
+};
+
+/** A handle to a camera animation running in UCameraAnimationCameraModifier */
+USTRUCT(BlueprintType)
+struct FCameraAnimationHandle
+{
+	GENERATED_BODY()
+
+	static FCameraAnimationHandle Invalid;
+
+	FCameraAnimationHandle() : InstanceID(MAX_uint16), InstanceSerial(0) {}
+	FCameraAnimationHandle(uint16 InInstanceID, uint16 InInstanceSerial) : InstanceID(InInstanceID), InstanceSerial(InInstanceSerial) {}
+
+	bool IsValid() const
+	{
+		return InstanceID != MAX_uint16;
+	}
+
+	friend bool operator==(FCameraAnimationHandle A, FCameraAnimationHandle B)
+	{
+		return A.InstanceID == B.InstanceID && A.InstanceSerial == B.InstanceSerial;
+	}
+	friend bool operator!=(FCameraAnimationHandle A, FCameraAnimationHandle B)
+	{
+		return !(A == B);
+	}
+	friend bool operator<(FCameraAnimationHandle A, FCameraAnimationHandle B)
+	{
+		return A.InstanceID < B.InstanceID;
+	}
+	friend uint32 GetTypeHash(FCameraAnimationHandle In)
+	{
+		return HashCombine(In.InstanceID, In.InstanceSerial);
+	}
+
+private:
+	uint16 InstanceID;
+	uint16 InstanceSerial;
+
+	friend class UCameraAnimationCameraModifier;
+};
+
+/**
+ * Information about an active camera animation inside UCameraAnimationCameraModifier.
+ */
+USTRUCT()
+struct GAMEPLAYCAMERAS_API FActiveCameraAnimationInfo
+{
+	GENERATED_BODY()
+
+	FActiveCameraAnimationInfo();
+
+	/** Whether this is a valid, ongoing camera animation */
+	bool IsValid() const { return Sequence != nullptr; }
+
+	/** The sequence to use for the animation. */
+	UPROPERTY()
+	TObjectPtr<UCameraAnimationSequence> Sequence;
+
+	/** The parameters for playing the animation. */
+	UPROPERTY()
+	FCameraAnimationParams Params;
+
+	/** A reference handle for use with UCameraAnimationCameraModifier. */
+	UPROPERTY()
+	FCameraAnimationHandle Handle;
+
+	/** The player for playing the animation. */
+	UPROPERTY(Transient)
+	TObjectPtr<UCameraAnimationSequencePlayer> Player;
+
+	/** Standin for the camera actor and components */
+	UPROPERTY(Transient)
+	TObjectPtr<UCameraAnimationSequenceCameraStandIn> CameraStandIn;
+
+	/** Current time into easing in */
+	UPROPERTY()
+	float EaseInCurrentTime;
+
+	/** Current time into easing out */
+	UPROPERTY()
+	float EaseOutCurrentTime;
+
+	/** Whether easing in is ongoing */
+	UPROPERTY()
+	bool bIsEasingIn;
+
+	/** Whether easing out is ongoing */
+	UPROPERTY()
+	bool bIsEasingOut;
+};
+
+/**
+ * A camera modifier that plays camera animation sequences.
+ */
+UCLASS(config=Camera)
+class GAMEPLAYCAMERAS_API UCameraAnimationCameraModifier : public UCameraModifier
+{
+	GENERATED_BODY()
+
+public:
+	UCameraAnimationCameraModifier(const FObjectInitializer& ObjectInitializer);
+
+	/**
+	 * Play a new camera animation sequence.
+	 * @param Sequence		The sequence to use for the new camera animation.
+	 * @param Params		The parameters for the new camera animation instance.
+	 */
+	UFUNCTION(BlueprintCallable, Category="Camera Animation")
+	FCameraAnimationHandle PlayCameraAnimation(UCameraAnimationSequence* Sequence, FCameraAnimationParams Params);
+
+	/**
+	 * Returns whether the given camera animation is playing.
+	 * @param Handle		A handle to a previously started camera animation.
+	 * @return				Whether the corresponding camera animation is playing or not.
+	 */
+	UFUNCTION(BlueprintCallable, Category="Camera Animation")
+	bool IsCameraAnimationActive(const FCameraAnimationHandle& Handle) const;
+	
+	/** 
+	 * Stops the given camera animation instance.
+	 * @param Hanlde		A handle to a previously started camera animation.
+	 * @param bImmediate	True to stop it right now and ignore blend out, false to let it blend out as indicated.
+	 */
+	UFUNCTION(BlueprintCallable, Category = "Camera Animation")
+	void StopCameraAnimation(const FCameraAnimationHandle& Handle, bool bImmediate = false);
+	
+	/**
+	 * Stop playing all instances of the given camera animation sequence.
+	 * @param Sequence		The sequence of which to stop all instances.
+	 * @param bImmediate	True to stop it right now and ignore blend out, false to let it blend out as indicated.
+	 */
+	UFUNCTION(BlueprintCallable, Category="Camera Animation")
+	void StopAllCameraAnimationsOf(UCameraAnimationSequence* Sequence, bool bImmediate = false);
+	
+	/**
+	 * Stop all camera animation instances.
+	 * @param bImmediate	True to stop it right now and ignore blend out, false to let it blend out as indicated.
+	 */
+	UFUNCTION(BlueprintCallable, Category="Camera Animation")
+	virtual void StopAllCameraAnimations(bool bImmediate = false);
+
+	// UCameraModifier interface
+	virtual bool ModifyCamera(float DeltaTime, FMinimalViewInfo& InOutPOV) override;
+	virtual void DisplayDebug(class UCanvas* Canvas, const FDebugDisplayInfo& DebugDisplay, float& YL, float& YPos) override;
+
+public:
+
+	UFUNCTION(BlueprintPure, Category="Camera Animation", meta=(WorldContext="WorldContextObject"))
+	static UCameraAnimationCameraModifier* GetCameraAnimationCameraModifier(const UObject* WorldContextObject, int32 PlayerIndex);
+
+	UFUNCTION(BlueprintPure, Category="Camera Animation", meta=(WorldContext="WorldContextObject"))
+	static UCameraAnimationCameraModifier* GetCameraAnimationCameraModifierFromID(const UObject* WorldContextObject, int32 ControllerID);
+
+	UFUNCTION(BlueprintPure, Category="Camera Animation")
+	static UCameraAnimationCameraModifier* GetCameraAnimationCameraModifierFromPlayerController(const APlayerController* PlayerController);
+
+protected:
+
+	static float EvaluateEasing(ECameraAnimationEasingType EasingType, float Interp);
+
+	int32 FindInactiveCameraAnimation();
+	const FActiveCameraAnimationInfo* GetActiveCameraAnimation(const FCameraAnimationHandle& Handle) const;
+	FActiveCameraAnimationInfo* GetActiveCameraAnimation(const FCameraAnimationHandle& Handle);
+	void DeactivateCameraAnimation(int32 Index);
+
+	void TickAllAnimations(float DeltaTime, FMinimalViewInfo& InOutPOV);
+	void TickAnimation(FActiveCameraAnimationInfo& CameraAnimation, float DeltaTime, FMinimalViewInfo& InOutPOV);
+
+protected:
+
+	/** List of active camera animation instances */
+	UPROPERTY()
+	TArray<FActiveCameraAnimationInfo> ActiveAnimations;
+
+	/** Next serial number to use for a camera animation instance */
+	UPROPERTY()
+	uint16 NextInstanceSerialNumber;
+};
+
+/**
+ * Blueprint function library for autocasting a player camera manager into the camera animation camera modifier.
+ * This prevents breaking Blueprints now that APlayerCameraManager::StartCameraShake returns the base class.
+ */
+UCLASS()
+class GAMEPLAYCAMERAS_API UGameplayCamerasFunctionLibrary : public UBlueprintFunctionLibrary
+{
+	GENERATED_BODY()
+
+public:
+	UFUNCTION(BlueprintPure, Category="Camera Animation", meta=(BlueprintAutocast))
+	static UCameraAnimationCameraModifier* Conv_CameraAnimationCameraModifier(APlayerCameraManager* PlayerCameraManager);
+
+	UFUNCTION(BlueprintPure, Category = "Camera Animation", meta = (BlueprintAutocast))
+	static ECameraShakePlaySpace Conv_CameraShakePlaySpace(ECameraAnimationPlaySpace CameraAnimationPlaySpace);
+
+	UFUNCTION(BlueprintPure, Category = "Camera Animation", meta = (BlueprintAutocast))
+	static ECameraAnimationPlaySpace Conv_CameraAnimationPlaySpace(ECameraShakePlaySpace CameraShakePlaySpace);
+};
+
+
+#if UE_ENABLE_INCLUDE_ORDER_DEPRECATED_IN_5_2
+#include "Camera/CameraTypes.h"
+#include "CoreMinimal.h"
+#include "Templates/SubclassOf.h"
+#endif
diff --git a/Engine/Plugins/Cameras/GameplayCameras/Source/GameplayCameras/Public/Camera/CameraShake.h b/Engine/Plugins/Cameras/GameplayCameras/Source/GameplayCameras/Public/Camera/CameraShake.h
deleted file mode 100644
diff --git a/Engine/Plugins/Cameras/GameplayCameras/Source/GameplayCameras/Public/CameraAnimationCameraModifier.h b/Engine/Plugins/Cameras/GameplayCameras/Source/GameplayCameras/Public/CameraAnimationCameraModifier.h
deleted file mode 100644
diff --git a/Engine/Plugins/Cameras/GameplayCameras/Source/GameplayCameras/Public/CompositeCameraShakePattern.h b/Engine/Plugins/Cameras/GameplayCameras/Source/GameplayCameras/Public/CompositeCameraShakePattern.h
deleted file mode 100644
diff --git a/Engine/Plugins/Cameras/GameplayCameras/Source/GameplayCameras/Public/DefaultCameraShakeBase.h b/Engine/Plugins/Cameras/GameplayCameras/Source/GameplayCameras/Public/DefaultCameraShakeBase.h
deleted file mode 100644
diff --git a/Engine/Plugins/Cameras/GameplayCameras/Source/GameplayCameras/Public/LegacyCameraShake.h b/Engine/Plugins/Cameras/GameplayCameras/Source/GameplayCameras/Public/LegacyCameraShake.h
deleted file mode 100644
diff --git a/Engine/Plugins/Cameras/GameplayCameras/Source/GameplayCameras/Public/MatineeCameraShake.h b/Engine/Plugins/Cameras/GameplayCameras/Source/GameplayCameras/Public/MatineeCameraShake.h
deleted file mode 100644
diff --git a/Engine/Plugins/Cameras/GameplayCameras/Source/GameplayCameras/Public/PerlinNoiseCameraShakePattern.h b/Engine/Plugins/Cameras/GameplayCameras/Source/GameplayCameras/Public/PerlinNoiseCameraShakePattern.h
deleted file mode 100644
diff --git a/Engine/Plugins/Cameras/GameplayCameras/Source/GameplayCameras/Public/Shakes/CompositeCameraShakePattern.h b/Engine/Plugins/Cameras/GameplayCameras/Source/GameplayCameras/Public/Shakes/CompositeCameraShakePattern.h
new file mode 100644
--- /dev/null
+++ b/Engine/Plugins/Cameras/GameplayCameras/Source/GameplayCameras/Public/Shakes/CompositeCameraShakePattern.h
@@ -0,0 +1,45 @@
+// Copyright Epic Games, Inc. All Rights Reserved.
+
+#pragma once
+
+#include "Camera/CameraShakeBase.h"
+#include "CompositeCameraShakePattern.generated.h"
+
+/**
+ * A base class for a simple camera shake.
+ */
+UCLASS(meta=(AutoExpandCategories="CameraShake"))
+class GAMEPLAYCAMERAS_API UCompositeCameraShakePattern : public UCameraShakePattern
+{
+public:
+
+	GENERATED_BODY()
+
+	UCompositeCameraShakePattern(const FObjectInitializer& ObjInit) : Super(ObjInit) {}
+
+	template<typename PatternClass>
+	PatternClass* AddChildPattern()
+	{
+		PatternClass* NewChildPattern = NewObject<PatternClass>();
+		ChildPatterns.Add(NewChildPattern);
+		return NewChildPattern;
+	}
+
+public:
+
+	/** The list of child shake patterns */
+	UPROPERTY(EditAnywhere, Instanced, Category=CameraShake)
+	TArray<TObjectPtr<UCameraShakePattern>> ChildPatterns;
+
+private:
+
+	// UCameraShakePattern interface
+	virtual void GetShakePatternInfoImpl(FCameraShakeInfo& OutInfo) const override;
+	virtual void StartShakePatternImpl(const FCameraShakePatternStartParams& Params) override;
+	virtual void UpdateShakePatternImpl(const FCameraShakePatternUpdateParams& Params, FCameraShakePatternUpdateResult& OutResult) override;
+	virtual void ScrubShakePatternImpl(const FCameraShakePatternScrubParams& Params, FCameraShakePatternUpdateResult& OutResult) override;
+	virtual bool IsFinishedImpl() const override;
+	virtual void StopShakePatternImpl(const FCameraShakePatternStopParams& Params) override;
+	virtual void TeardownShakePatternImpl() override;
+};
+
diff --git a/Engine/Plugins/Cameras/GameplayCameras/Source/GameplayCameras/Public/Shakes/DefaultCameraShakeBase.h b/Engine/Plugins/Cameras/GameplayCameras/Source/GameplayCameras/Public/Shakes/DefaultCameraShakeBase.h
new file mode 100644
--- /dev/null
+++ b/Engine/Plugins/Cameras/GameplayCameras/Source/GameplayCameras/Public/Shakes/DefaultCameraShakeBase.h
@@ -0,0 +1,20 @@
+// Copyright Epic Games, Inc. All Rights Reserved.
+
+#pragma once
+
+#include "Camera/CameraShakeBase.h"
+#include "DefaultCameraShakeBase.generated.h"
+
+/**
+ * Like UCameraShakeBase but with a perlin noise shake pattern by default, for convenience.
+ */
+UCLASS()
+class GAMEPLAYCAMERAS_API UDefaultCameraShakeBase : public UCameraShakeBase
+{
+	GENERATED_BODY()
+
+public:
+
+	UDefaultCameraShakeBase(const FObjectInitializer& ObjInit);
+};
+
diff --git a/Engine/Plugins/Cameras/GameplayCameras/Source/GameplayCameras/Public/Shakes/LegacyCameraShake.h b/Engine/Plugins/Cameras/GameplayCameras/Source/GameplayCameras/Public/Shakes/LegacyCameraShake.h
new file mode 100644
--- /dev/null
+++ b/Engine/Plugins/Cameras/GameplayCameras/Source/GameplayCameras/Public/Shakes/LegacyCameraShake.h
@@ -0,0 +1,333 @@
+// Copyright Epic Games, Inc. All Rights Reserved.
+
+#pragma once
+
+#include "Camera/CameraShakeBase.h"
+#include "Kismet/BlueprintFunctionLibrary.h"
+
+#include "LegacyCameraShake.generated.h"
+
+class AActor;
+class UCameraShakeSourceComponent;
+
+/************************************************************
+ * Parameters for defining oscillating camera shakes
+ ************************************************************/
+
+ /** Types of waveforms that can be used for camera shake oscillators */
+UENUM(BlueprintType)
+enum class EOscillatorWaveform : uint8
+{
+	/** A sinusoidal wave */
+	SineWave,
+
+	/** Perlin noise */
+	PerlinNoise,
+};
+
+/** Shake start offset parameter */
+UENUM()
+enum EInitialOscillatorOffset : int
+{
+	/** Start with random offset (default). */
+	EOO_OffsetRandom UMETA(DisplayName = "Random"),
+	/** Start with zero offset. */
+	EOO_OffsetZero UMETA(DisplayName = "Zero"),
+	EOO_MAX,
+};
+
+/** Defines oscillation of a single number. */
+USTRUCT(BlueprintType)
+struct GAMEPLAYCAMERAS_API FFOscillator
+{
+	GENERATED_USTRUCT_BODY()
+
+	/** Amplitude of the sinusoidal oscillation. */
+	UPROPERTY(EditAnywhere, BlueprintReadWrite, Category = FOscillator)
+	float Amplitude;
+
+	/** Frequency of the sinusoidal oscillation. */
+	UPROPERTY(EditAnywhere, BlueprintReadWrite, Category = FOscillator)
+	float Frequency;
+
+	/** Defines how to begin (either at zero, or at a randomized value. */
+	UPROPERTY(EditAnywhere, Category = FOscillator)
+	TEnumAsByte<enum EInitialOscillatorOffset> InitialOffset;
+
+	/** Type of waveform to use for oscillation. */
+	UPROPERTY(EditAnywhere, BlueprintReadWrite, Category = FOscillator)
+	EOscillatorWaveform Waveform;
+
+	FFOscillator()
+		: Amplitude(0)
+		, Frequency(0)
+		, InitialOffset(0)
+		, Waveform(EOscillatorWaveform::SineWave)
+	{}
+
+	/** Advances the oscillation time and returns the current value. */
+	static float UpdateOffset(FFOscillator const& Osc, float& CurrentOffset, float DeltaTime);
+
+	/** Advances the oscillation time and returns the current value. */
+	static float UpdateOffset(FFOscillator const& Osc, double& CurrentOffset, float DeltaTime);
+
+	/** Returns the initial value of the oscillator. */
+	static float GetInitialOffset(FFOscillator const& Osc);
+
+	/** Returns the offset at the given time */
+	static float GetOffsetAtTime(FFOscillator const& Osc, float InitialOffset, float Time);
+};
+
+/** Defines FRotator oscillation. */
+USTRUCT(BlueprintType)
+struct GAMEPLAYCAMERAS_API FROscillator
+{
+	GENERATED_USTRUCT_BODY()
+
+	/** Pitch oscillation. */
+	UPROPERTY(EditAnywhere, BlueprintReadWrite, Category = ROscillator)
+	struct FFOscillator Pitch;
+
+	/** Yaw oscillation. */
+	UPROPERTY(EditAnywhere, BlueprintReadWrite, Category = ROscillator)
+	struct FFOscillator Yaw;
+
+	/** Roll oscillation. */
+	UPROPERTY(EditAnywhere, BlueprintReadWrite, Category = ROscillator)
+	struct FFOscillator Roll;
+
+};
+
+/** Defines FVector oscillation. */
+USTRUCT(BlueprintType)
+struct FVOscillator
+{
+	GENERATED_USTRUCT_BODY()
+
+	/** Oscillation in the X axis. */
+	UPROPERTY(EditAnywhere, BlueprintReadWrite, Category = VOscillator)
+	struct FFOscillator X;
+
+	/** Oscillation in the Y axis. */
+	UPROPERTY(EditAnywhere, BlueprintReadWrite, Category = VOscillator)
+	struct FFOscillator Y;
+
+	/** Oscillation in the Z axis. */
+	UPROPERTY(EditAnywhere, BlueprintReadWrite, Category = VOscillator)
+	struct FFOscillator Z;
+};
+
+
+/**
+ * Legacy camera shake which can do either oscillation or run camera anims.
+ */
+UCLASS(Blueprintable, HideCategories = (CameraShakePattern))
+class GAMEPLAYCAMERAS_API ULegacyCameraShake : public UCameraShakeBase
+{
+	GENERATED_UCLASS_BODY()
+
+public:
+	/** Duration in seconds of current screen shake. Less than 0 means indefinite, 0 means no oscillation. */
+	UPROPERTY(EditAnywhere, Category = Oscillation)
+	float OscillationDuration;
+
+	/** Duration of the blend-in, where the oscillation scales from 0 to 1. */
+	UPROPERTY(EditAnywhere, Category = Oscillation, meta = (ClampMin = "0.0"))
+	float OscillationBlendInTime;
+
+	/** Duration of the blend-out, where the oscillation scales from 1 to 0. */
+	UPROPERTY(EditAnywhere, Category = Oscillation, meta = (ClampMin = "0.0"))
+	float OscillationBlendOutTime;
+
+	/** Rotational oscillation */
+	UPROPERTY(EditAnywhere, BlueprintReadWrite, Category = Oscillation)
+	struct FROscillator RotOscillation;
+
+	/** Positional oscillation */
+	UPROPERTY(EditAnywhere, BlueprintReadWrite, Category = Oscillation)
+	struct FVOscillator LocOscillation;
+
+	/** FOV oscillation */
+	UPROPERTY(EditAnywhere, BlueprintReadWrite, Category = Oscillation)
+	struct FFOscillator FOVOscillation;
+
+	/************************************************************
+	 * Parameters for defining sequencer-driven camera shakes
+	 ************************************************************/
+
+	 /** Scalar defining how fast to play the anim. */
+	UPROPERTY(EditAnywhere, Category = AnimShake, meta = (ClampMin = "0.001"))
+	float AnimPlayRate;
+
+	/** Scalar defining how "intense" to play the anim. */
+	UPROPERTY(EditAnywhere, Category = AnimShake, meta = (ClampMin = "0.0"))
+	float AnimScale;
+
+	/** Linear blend-in time. */
+	UPROPERTY(EditAnywhere, Category = AnimShake, meta = (ClampMin = "0.0"))
+	float AnimBlendInTime;
+
+	/** Linear blend-out time. */
+	UPROPERTY(EditAnywhere, Category = AnimShake, meta = (ClampMin = "0.0"))
+	float AnimBlendOutTime;
+
+	/** When bRandomAnimSegment is true, this defines how long the anim should play. */
+	UPROPERTY(EditAnywhere, Category = AnimShake, meta = (ClampMin = "0.0", editcondition = "bRandomAnimSegment"))
+	float RandomAnimSegmentDuration;
+
+	/** Source camera animation sequence to play. Can be null. */
+	UPROPERTY(EditAnywhere, Category = AnimShake)
+	TObjectPtr<class UCameraAnimationSequence> AnimSequence;
+
+	/**
+	* If true, play a random snippet of the animation of length Duration.  Implies bLoop and bRandomStartTime = true for the AnimSequence.
+	* If false, play the full anim once, non-looped. Useful for getting variety out of a single looped AnimSequence asset.
+	*/
+	UPROPERTY(EditAnywhere, Category = AnimShake)
+	uint32 bRandomAnimSegment : 1;
+
+public:
+
+	/** Time remaining for oscillation shakes. Less than 0.f means shake infinitely. */
+	UPROPERTY(transient, BlueprintReadOnly, Category = CameraShake)
+	float OscillatorTimeRemaining;
+
+public:
+
+	// Blueprint API
+
+	/** Called when the shake starts playing */
+	UFUNCTION(BlueprintImplementableEvent, Category = CameraShake)
+	void ReceivePlayShake(float Scale);
+
+	/** Called every tick to let the shake modify the point of view */
+	UFUNCTION(BlueprintImplementableEvent, Category = CameraShake)
+	void BlueprintUpdateCameraShake(float DeltaTime, float Alpha, const FMinimalViewInfo& POV, FMinimalViewInfo& ModifiedPOV);
+
+	/** Called to allow a shake to decide when it's finished playing. */
+	UFUNCTION(BlueprintNativeEvent, Category = CameraShake)
+	bool ReceiveIsFinished() const;
+
+	/**
+	 * Called when the shake is explicitly stopped.
+	 * @param bImmediatly		If true, shake stops right away regardless of blend out settings. If false, shake may blend out according to its settings.
+	 */
+	UFUNCTION(BlueprintImplementableEvent, Category = CameraShake)
+	void ReceiveStopShake(bool bImmediately);
+
+public:
+
+	/**
+	 * Backwards compatible method used by core BP redirectors. This is needed because the return value is specifically a legacy camera shake,
+	 * which some BP logic often uses directly to set oscillator/anim properties.
+	 */
+	UFUNCTION(BlueprintCallable, Category = "Camera Shakes")
+	static ULegacyCameraShake* StartLegacyCameraShake(APlayerCameraManager* PlayerCameraManager, TSubclassOf<ULegacyCameraShake> ShakeClass, float Scale = 1.f, ECameraShakePlaySpace PlaySpace = ECameraShakePlaySpace::CameraLocal, FRotator UserPlaySpaceRot = FRotator::ZeroRotator);
+
+	/**
+	 * Backwards compatible method used by core BP redirectors. This is needed because the return value is specifically a legacy camera shake,
+	 * which some BP logic often uses directly to set oscillator/anim properties.
+	 */
+	UFUNCTION(BlueprintCallable, Category = "Camera Shakes")
+	static ULegacyCameraShake* StartLegacyCameraShakeFromSource(APlayerCameraManager* PlayerCameraManager, TSubclassOf<ULegacyCameraShake> ShakeClass, UCameraShakeSourceComponent* SourceComponent, float Scale = 1.f, ECameraShakePlaySpace PlaySpace = ECameraShakePlaySpace::CameraLocal, FRotator UserPlaySpaceRot = FRotator::ZeroRotator);
+
+public:
+
+	/** Returns true if this camera shake will loop forever */
+	bool IsLooping() const;
+
+	/** Sets current playback time and applies the shake (both oscillation and camera animation sequence) to the given POV. */
+	UE_DEPRECATED(4.27, "SetCurrentTimeAndApplyShake is deprecated, please use ScrubAndApplyCameraShake")
+	void SetCurrentTimeAndApplyShake(float NewTime, FMinimalViewInfo& POV);
+
+private:
+
+	void DoStartShake(const FCameraShakePatternStartParams& Params);
+	void DoUpdateShake(const FCameraShakePatternUpdateParams& Params, FCameraShakePatternUpdateResult& OutResult);
+	void DoScrubShake(const FCameraShakePatternScrubParams& Params, FCameraShakePatternUpdateResult& OutResult);
+	void DoStopShake(bool bImmediately);
+	bool DoGetIsFinished() const;
+	void DoTeardownShake();
+
+protected:
+
+	/** Actual oscillation duration. */
+	float ActualOscillationDuration;
+
+	/** Current location sinusoidal offset. */
+	FVector LocSinOffset;
+
+	/** Current rotational sinusoidal offset. */
+	FVector RotSinOffset;
+
+	/** Current FOV sinusoidal offset. */
+	float FOVSinOffset;
+
+	/** Initial offset (could have been assigned at random). */
+	FVector InitialLocSinOffset;
+
+	/** Initial offset (could have been assigned at random). */
+	FVector InitialRotSinOffset;
+
+	/** Initial offset (could have been assigned at random). */
+	float InitialFOVSinOffset;
+
+	/** Sequence shake pattern for when using a sequence instead of a camera anim */
+	UPROPERTY(Instanced)
+	TObjectPtr<class USequenceCameraShakePattern> SequenceShakePattern;
+
+private:
+
+	float CurrentBlendInTime;
+	float CurrentBlendOutTime;
+	bool bBlendingIn : 1;
+	bool bBlendingOut : 1;
+
+	friend class ULegacyCameraShakePattern;
+};
+
+/**
+ * Shake pattern for the ULegacyCameraShake class.
+ *
+ * It doesn't do anything because, for backwards compatibility reasons, all the data
+ * was left on the shake class itself... so this pattern delegates everything back
+ * to the owner shake.
+ */
+UCLASS(HideDropdown)
+class GAMEPLAYCAMERAS_API ULegacyCameraShakePattern : public UCameraShakePattern
+{
+	GENERATED_BODY()
+
+private:
+
+	// UCameraShakePattern interface
+	virtual void GetShakePatternInfoImpl(FCameraShakeInfo& OutInfo) const override;
+	virtual void StartShakePatternImpl(const FCameraShakePatternStartParams& Params) override;
+	virtual void UpdateShakePatternImpl(const FCameraShakePatternUpdateParams& Params, FCameraShakePatternUpdateResult& OutResult) override;
+	virtual void ScrubShakePatternImpl(const FCameraShakePatternScrubParams& Params, FCameraShakePatternUpdateResult& OutResult) override;
+	virtual bool IsFinishedImpl() const override;
+	virtual void StopShakePatternImpl(const FCameraShakePatternStopParams& Params) override;
+	virtual void TeardownShakePatternImpl() override;
+};
+
+/**
+ * Blueprint function library for autocasting from a base camera shake to a legacy camera shake.
+ * This prevents breaking Blueprints now that APlayerCameraManager::StartCameraShake returns the base class.
+ */
+UCLASS()
+class GAMEPLAYCAMERAS_API ULegacyCameraShakeFunctionLibrary : public UBlueprintFunctionLibrary
+{
+	GENERATED_BODY()
+
+public:
+	UFUNCTION(BlueprintPure, Category = "Camera Shakes", meta = (BlueprintAutocast))
+	static ULegacyCameraShake* Conv_LegacyCameraShake(UCameraShakeBase* CameraShake)
+	{
+		return CastChecked<ULegacyCameraShake>(CameraShake, ECastCheckedType::NullAllowed);
+	}
+};
+
+#if UE_ENABLE_INCLUDE_ORDER_DEPRECATED_IN_5_2
+#include "CoreMinimal.h"
+#include "GameFramework/Actor.h"
+#endif
diff --git a/Engine/Plugins/Cameras/GameplayCameras/Source/GameplayCameras/Public/Shakes/MatineeCameraShake.h b/Engine/Plugins/Cameras/GameplayCameras/Source/GameplayCameras/Public/Shakes/MatineeCameraShake.h
new file mode 100644
--- /dev/null
+++ b/Engine/Plugins/Cameras/GameplayCameras/Source/GameplayCameras/Public/Shakes/MatineeCameraShake.h
@@ -0,0 +1,10 @@
+// Copyright Epic Games, Inc. All Rights Reserved.
+
+#pragma once
+
+// HEADER_UNIT_SKIP - Deprecated
+
+#include "CoreTypes.h"
+
+COMPILE_WARNING("UMatineeCameraShake has been renamed to ULegacyCameraShake, change to #include \"LegacyCameraShake.h\" and rename references")
+#include "LegacyCameraShake.h"
diff --git a/Engine/Plugins/Cameras/GameplayCameras/Source/GameplayCameras/Public/Shakes/PerlinNoiseCameraShakePattern.h b/Engine/Plugins/Cameras/GameplayCameras/Source/GameplayCameras/Public/Shakes/PerlinNoiseCameraShakePattern.h
new file mode 100644
--- /dev/null
+++ b/Engine/Plugins/Cameras/GameplayCameras/Source/GameplayCameras/Public/Shakes/PerlinNoiseCameraShakePattern.h
@@ -0,0 +1,116 @@
+// Copyright Epic Games, Inc. All Rights Reserved.
+
+#pragma once
+
+#include "SimpleCameraShakePattern.h"
+#include "PerlinNoiseCameraShakePattern.generated.h"
+
+/** A perlin noise shaker for a single number. */
+USTRUCT(BlueprintType)
+struct GAMEPLAYCAMERAS_API FPerlinNoiseShaker
+{
+	GENERATED_BODY()
+
+	/** Amplitude of the perlin noise. */
+	UPROPERTY(EditAnywhere, BlueprintReadWrite, Category=PerlinNoise)
+	float Amplitude;
+
+	/** Frequency of the sinusoidal oscillation. */
+	UPROPERTY(EditAnywhere, BlueprintReadWrite, Category=PerlinNoise)
+	float Frequency;
+
+	/** Creates a new perlin noise shaker. */
+	FPerlinNoiseShaker()
+		: Amplitude(1.f)
+		, Frequency(1.f)
+	{}
+
+	/** Advances the shake time and returns the current value */
+	float Update(float DeltaTime, float AmplitudeMultiplier, float FrequencyMultiplier, float& InOutCurrentOffset) const;
+};
+
+/**
+ * A camera shake that uses Perlin noise to shake the camera.
+ */
+UCLASS(meta=(AutoExpandCategories="Location,Rotation,FOV,Timing"))
+class GAMEPLAYCAMERAS_API UPerlinNoiseCameraShakePattern : public USimpleCameraShakePattern
+{
+public:
+
+	GENERATED_BODY()
+
+	UPerlinNoiseCameraShakePattern(const FObjectInitializer& ObjInit);
+
+public:
+
+	/** Amplitude multiplier for all location shake */
+	UPROPERTY(EditAnywhere, BlueprintReadWrite, Category=Location)
+	float LocationAmplitudeMultiplier = 1.f;
+
+	/** Frequency multiplier for all location shake */
+	UPROPERTY(EditAnywhere, BlueprintReadWrite, Category=Location)
+	float LocationFrequencyMultiplier = 1.f;
+
+	/** Shake in the X axis. */
+	UPROPERTY(EditAnywhere, BlueprintReadWrite, Category=Location)
+	FPerlinNoiseShaker X;
+
+	/** Shake in the Y axis. */
+	UPROPERTY(EditAnywhere, BlueprintReadWrite, Category=Location)
+	FPerlinNoiseShaker Y;
+
+	/** Shake in the Z axis. */
+	UPROPERTY(EditAnywhere, BlueprintReadWrite, Category=Location)
+	FPerlinNoiseShaker Z;
+
+	/** Amplitude multiplier for all rotation shake */
+	UPROPERTY(EditAnywhere, BlueprintReadWrite, Category=Rotation)
+	float RotationAmplitudeMultiplier = 1.f;
+
+	/** Frequency multiplier for all rotation shake */
+	UPROPERTY(EditAnywhere, BlueprintReadWrite, Category=Rotation)
+	float RotationFrequencyMultiplier = 1.f;
+
+	/** Pitch shake. */
+	UPROPERTY(EditAnywhere, BlueprintReadWrite, Category=Rotation)
+	FPerlinNoiseShaker Pitch;
+
+	/** Yaw shake. */
+	UPROPERTY(EditAnywhere, BlueprintReadWrite, Category=Rotation)
+	FPerlinNoiseShaker Yaw;
+
+	/** Roll shake. */
+	UPROPERTY(EditAnywhere, BlueprintReadWrite, Category=Rotation)
+	FPerlinNoiseShaker Roll;
+
+	/** FOV shake. */
+	UPROPERTY(EditAnywhere, BlueprintReadWrite, Category=FOV)
+	FPerlinNoiseShaker FOV;
+
+private:
+
+	// UCameraShakePattern interface
+	virtual void StartShakePatternImpl(const FCameraShakePatternStartParams& Params) override;
+	virtual void UpdateShakePatternImpl(const FCameraShakePatternUpdateParams& Params, FCameraShakePatternUpdateResult& OutResult) override;
+	virtual void ScrubShakePatternImpl(const FCameraShakePatternScrubParams& Params, FCameraShakePatternUpdateResult& OutResult) override;
+
+	void UpdatePerlinNoise(float DeltaTime, FCameraShakePatternUpdateResult& OutResult);
+
+private:
+
+	/** Initial perlin noise offset for location oscillation. */
+	FVector3f InitialLocationOffset;
+	/** Current perlin noise offset for location oscillation. */
+	FVector3f CurrentLocationOffset;
+
+	/** Initial perlin noise offset for rotation oscillation. */
+	FVector3f InitialRotationOffset;
+	/** Current perlin noise offset for rotation oscillation. */
+	FVector3f CurrentRotationOffset;
+
+	/** Initial perlin noise offset for FOV oscillation */
+	float InitialFOVOffset;
+	/** Current perlin noise offset for FOV oscillation */
+	float CurrentFOVOffset;
+};
+
diff --git a/Engine/Plugins/Cameras/GameplayCameras/Source/GameplayCameras/Public/Shakes/SimpleCameraShakePattern.h b/Engine/Plugins/Cameras/GameplayCameras/Source/GameplayCameras/Public/Shakes/SimpleCameraShakePattern.h
new file mode 100644
--- /dev/null
+++ b/Engine/Plugins/Cameras/GameplayCameras/Source/GameplayCameras/Public/Shakes/SimpleCameraShakePattern.h
@@ -0,0 +1,48 @@
+// Copyright Epic Games, Inc. All Rights Reserved.
+
+#pragma once
+
+#include "Camera/CameraShakeBase.h"
+#include "SimpleCameraShakePattern.generated.h"
+
+/**
+ * A base class for a simple camera shake.
+ */
+UCLASS(Abstract)
+class GAMEPLAYCAMERAS_API USimpleCameraShakePattern : public UCameraShakePattern
+{
+public:
+
+	GENERATED_BODY()
+
+	USimpleCameraShakePattern(const FObjectInitializer& ObjInit) : Super(ObjInit) {}
+
+public:
+
+	/** Duration in seconds of this shake. Zero or less means infinite. */
+	UPROPERTY(EditAnywhere, Category=Timing)
+	float Duration = 1.f;
+
+	/** Blend-in time for this shake. Zero or less means no blend-in. */
+	UPROPERTY(EditAnywhere, Category=Timing)
+	float BlendInTime = 0.2f;
+
+	/** Blend-out time for this shake. Zero or less means no blend-out. */
+	UPROPERTY(EditAnywhere, Category=Timing)
+	float BlendOutTime = 0.2f;
+
+protected:
+
+	// UCameraShakePattern interface
+	virtual void GetShakePatternInfoImpl(FCameraShakeInfo& OutInfo) const override;
+	virtual void StartShakePatternImpl(const FCameraShakePatternStartParams& Params) override;
+	virtual bool IsFinishedImpl() const override;
+	virtual void StopShakePatternImpl(const FCameraShakePatternStopParams& Params) override;
+	virtual void TeardownShakePatternImpl()  override;
+
+protected:
+
+	/** The ongoing state for this shake */
+	FCameraShakeState State;
+};
+
diff --git a/Engine/Plugins/Cameras/GameplayCameras/Source/GameplayCameras/Public/Shakes/WaveOscillatorCameraShakePattern.h b/Engine/Plugins/Cameras/GameplayCameras/Source/GameplayCameras/Public/Shakes/WaveOscillatorCameraShakePattern.h
new file mode 100644
--- /dev/null
+++ b/Engine/Plugins/Cameras/GameplayCameras/Source/GameplayCameras/Public/Shakes/WaveOscillatorCameraShakePattern.h
@@ -0,0 +1,134 @@
+// Copyright Epic Games, Inc. All Rights Reserved.
+
+#pragma once
+
+#include "SimpleCameraShakePattern.h"
+#include "WaveOscillatorCameraShakePattern.generated.h"
+
+/** Shake start offset parameter. */
+UENUM()
+enum class EInitialWaveOscillatorOffsetType : uint8
+{
+	/** Start with random offset (default). */
+	Random,
+	/** Start with zero offset. */
+	Zero
+};
+
+/** A wave oscillator for a single number. */
+USTRUCT(BlueprintType)
+struct GAMEPLAYCAMERAS_API FWaveOscillator
+{
+	GENERATED_BODY()
+
+	/** Amplitude of the sinusoidal oscillation. */
+	UPROPERTY(EditAnywhere, BlueprintReadWrite, Category=Oscillation)
+	float Amplitude;
+
+	/** Frequency of the sinusoidal oscillation. */
+	UPROPERTY(EditAnywhere, BlueprintReadWrite, Category=Oscillation)
+	float Frequency;
+
+	/** Defines how to begin the oscillation. */
+	UPROPERTY(EditAnywhere, BlueprintReadWrite, Category=Oscillation)
+	EInitialWaveOscillatorOffsetType InitialOffsetType;
+	
+	/** Creates a new wave oscillator. */
+	FWaveOscillator()
+		: Amplitude(1.f)
+		, Frequency(1.f)
+		, InitialOffsetType(EInitialWaveOscillatorOffsetType::Random)
+	{}
+
+	/** Sets the initial offset and returns the initial value of the oscillator */
+	float Initialize(float& OutInitialOffset) const;
+
+	/** Advances the oscillation time and returns the current value */
+	float Update(float DeltaTime, float AmplitudeMultiplier, float FrequencyMultiplier, float& InOutCurrentOffset) const;
+};
+
+/**
+ * A camera shake that uses oscillations to move the camera.
+ */
+UCLASS(meta=(AutoExpandCategories="Location,Rotation,FOV,Timing"))
+class GAMEPLAYCAMERAS_API UWaveOscillatorCameraShakePattern : public USimpleCameraShakePattern
+{
+public:
+
+	GENERATED_BODY()
+
+	UWaveOscillatorCameraShakePattern(const FObjectInitializer& ObjInit);
+
+public:
+
+	/** Amplitude multiplier for all location oscillation */
+	UPROPERTY(EditAnywhere, BlueprintReadWrite, Category=Location)
+	float LocationAmplitudeMultiplier = 1.f;
+
+	/** Frequency multiplier for all location oscillation */
+	UPROPERTY(EditAnywhere, BlueprintReadWrite, Category=Location)
+	float LocationFrequencyMultiplier = 1.f;
+
+	/** Oscillation in the X axis. */
+	UPROPERTY(EditAnywhere, BlueprintReadWrite, Category=Location)
+	FWaveOscillator X;
+
+	/** Oscillation in the Y axis. */
+	UPROPERTY(EditAnywhere, BlueprintReadWrite, Category=Location)
+	FWaveOscillator Y;
+
+	/** Oscillation in the Z axis. */
+	UPROPERTY(EditAnywhere, BlueprintReadWrite, Category=Location)
+	FWaveOscillator Z;
+
+	/** Amplitude multiplier for all rotation oscillation */
+	UPROPERTY(EditAnywhere, BlueprintReadWrite, Category=Rotation)
+	float RotationAmplitudeMultiplier = 1.f;
+
+	/** Frequency multiplier for all rotation oscillation */
+	UPROPERTY(EditAnywhere, BlueprintReadWrite, Category=Rotation)
+	float RotationFrequencyMultiplier = 1.f;
+
+	/** Pitch oscillation. */
+	UPROPERTY(EditAnywhere, BlueprintReadWrite, Category=Rotation)
+	FWaveOscillator Pitch;
+
+	/** Yaw oscillation. */
+	UPROPERTY(EditAnywhere, BlueprintReadWrite, Category=Rotation)
+	FWaveOscillator Yaw;
+
+	/** Roll oscillation. */
+	UPROPERTY(EditAnywhere, BlueprintReadWrite, Category=Rotation)
+	FWaveOscillator Roll;
+
+	/** FOV oscillation. */
+	UPROPERTY(EditAnywhere, BlueprintReadWrite, Category=FOV)
+	FWaveOscillator FOV;
+
+private:
+
+	// UCameraShakePattern interface
+	virtual void StartShakePatternImpl(const FCameraShakePatternStartParams& Params) override;
+	virtual void UpdateShakePatternImpl(const FCameraShakePatternUpdateParams& Params, FCameraShakePatternUpdateResult& OutResult) override;
+	virtual void ScrubShakePatternImpl(const FCameraShakePatternScrubParams& Params, FCameraShakePatternUpdateResult& OutResult) override;
+
+	void UpdateOscillators(float DeltaTime, FCameraShakePatternUpdateResult& OutResult);
+
+private:
+
+	/** Initial sinusoidal offset for location oscillation. */
+	FVector3f InitialLocationOffset;
+	/** Current sinusoidal offset for location oscillation. */
+	FVector3f CurrentLocationOffset;
+
+	/** Initial sinusoidal offset for rotation oscillation. */
+	FVector3f InitialRotationOffset;
+	/** Current sinusoidal offset for rotation oscillation. */
+	FVector3f CurrentRotationOffset;
+
+	/** Initial sinusoidal offset for FOV oscillation */
+	float InitialFOVOffset;
+	/** Current sinusoidal offset for FOV oscillation */
+	float CurrentFOVOffset;
+};
+
diff --git a/Engine/Plugins/Cameras/GameplayCameras/Source/GameplayCameras/Public/SimpleCameraShakePattern.h b/Engine/Plugins/Cameras/GameplayCameras/Source/GameplayCameras/Public/SimpleCameraShakePattern.h
deleted file mode 100644
diff --git a/Engine/Plugins/Cameras/GameplayCameras/Source/GameplayCameras/Public/WaveOscillatorCameraShakePattern.h b/Engine/Plugins/Cameras/GameplayCameras/Source/GameplayCameras/Public/WaveOscillatorCameraShakePattern.h
deleted file mode 100644
