Cameras: first commit of the camera system

This change contains the basic infrastructure for the UE5 camera system: the main evaluator, the director, the blend stacks, the instantiation of camera node trees, etc. Other concepts such as the variable table, evaluation services (including target preservation), IK operations, and more aren't there yet.

diff --git a/Engine/Plugins/Cameras/GameplayCameras/Source/GameplayCameras/Private/Animations/CameraAnimationCameraModifier.cpp b/Engine/Plugins/Cameras/GameplayCameras/Source/GameplayCameras/Private/Animations/CameraAnimationCameraModifier.cpp
--- a/Engine/Plugins/Cameras/GameplayCameras/Source/GameplayCameras/Private/Animations/CameraAnimationCameraModifier.cpp
+++ b/Engine/Plugins/Cameras/GameplayCameras/Source/GameplayCameras/Private/Animations/CameraAnimationCameraModifier.cpp
@@ -8,7 +8,7 @@
 #include "Engine/Canvas.h"
 #include "Engine/Engine.h"
 #include "GameFramework/PlayerController.h"
-#include "GameplayCamerasModule.h"
+#include "IGameplayCamerasModule.h"
 #include "Kismet/GameplayStatics.h"
 #include "UObject/Object.h"
 
diff --git a/Engine/Plugins/Cameras/GameplayCameras/Source/GameplayCameras/Private/Core/BlendCameraNode.cpp b/Engine/Plugins/Cameras/GameplayCameras/Source/GameplayCameras/Private/Core/BlendCameraNode.cpp
new file mode 100644
--- /dev/null
+++ b/Engine/Plugins/Cameras/GameplayCameras/Source/GameplayCameras/Private/Core/BlendCameraNode.cpp
@@ -0,0 +1,11 @@
+// Copyright Epic Games, Inc. All Rights Reserved.
+
+#include "Core/BlendCameraNode.h"
+
+#include UE_INLINE_GENERATED_CPP_BY_NAME(BlendCameraNode)
+
+void UBlendCameraNode::BlendResults(const FCameraNodeBlendParams& Params, FCameraNodeBlendResult& OutResult)
+{
+	OnBlendResults(Params, OutResult);
+}
+
diff --git a/Engine/Plugins/Cameras/GameplayCameras/Source/GameplayCameras/Private/Core/BlendStackCameraNode.cpp b/Engine/Plugins/Cameras/GameplayCameras/Source/GameplayCameras/Private/Core/BlendStackCameraNode.cpp
new file mode 100644
--- /dev/null
+++ b/Engine/Plugins/Cameras/GameplayCameras/Source/GameplayCameras/Private/Core/BlendStackCameraNode.cpp
@@ -0,0 +1,301 @@
+// Copyright Epic Games, Inc. All Rights Reserved.
+
+#include "Core/BlendStackCameraNode.h"
+
+#include "Core/BlendCameraNode.h"
+#include "Core/BlendStackRootCameraNode.h"
+#include "Core/CameraAsset.h"
+#include "Core/CameraEvaluationContext.h"
+#include "Core/CameraMode.h"
+#include "Core/CameraRuntimeInstantiator.h"
+#include "Nodes/Blends/PopBlendCameraNode.h"
+
+#include UE_INLINE_GENERATED_CPP_BY_NAME(BlendStackCameraNode)
+
+void UBlendStackCameraNode::Push(const FBlendStackCameraPushParams& Params)
+{
+	if (!Entries.IsEmpty())
+	{
+		// Don't push anything if what is being requested is already the
+		// active camera mode.
+		const FCameraModeEntry& TopEntry(Entries.Top());
+		if (!TopEntry.bIsFrozen
+				&& TopEntry.OriginalCameraMode == Params.CameraMode
+				&& TopEntry.EvaluationContext == Params.EvaluationContext)
+		{
+			return;
+		}
+	}
+
+	// Create the new root node. Instantiated objects will go inside it.
+	UBlendStackRootCameraNode* EntryRootNode = NewObject<UBlendStackRootCameraNode>(this, NAME_None);
+
+	FCameraRuntimeInstantiationParams InstParams;
+	InstParams.InstantiationOuter = EntryRootNode;
+
+	// Instantiate the camera mode's node tree.
+	UCameraNode* ModeRootNode = FCameraRuntimeInstantiator::InstantiateCameraNodeTree(
+			Params.CameraMode->RootNode, InstParams);
+	EntryRootNode->RootNode = ModeRootNode;
+
+	// Find a transition and instantiate its blend. If not transition is found,
+	// make a camera cut transition.
+	if (const FCameraModeTransition* Transition = FindTransition(Params))
+	{
+		UBlendCameraNode* ModeBlend = CastChecked<UBlendCameraNode>(
+				FCameraRuntimeInstantiator::InstantiateCameraNodeTree(Transition->Blend, InstParams));
+		EntryRootNode->Blend = ModeBlend;
+	}
+	else
+	{
+		EntryRootNode->Blend = NewObject<UPopBlendCameraNode>(EntryRootNode, NAME_None);
+	}
+
+	// Make a new entry and add it to the stack.
+	FCameraModeEntry NewEntry;
+	NewEntry.EvaluationContext = Params.EvaluationContext;
+	NewEntry.OriginalCameraMode = Params.CameraMode;
+	NewEntry.RootNode = EntryRootNode;
+	NewEntry.bIsFirstFrame = true;
+
+	Entries.Add(MoveTemp(NewEntry));
+}
+
+void UBlendStackCameraNode::AddReferencedObjects(UObject* InThis, FReferenceCollector& Collector)
+{
+	UBlendStackCameraNode* TypedThis = CastChecked<UBlendStackCameraNode>(InThis);
+	for (FCameraModeEntry& Entry : TypedThis->Entries)
+	{
+		Collector.AddReferencedObject(Entry.OriginalCameraMode);
+		Collector.AddReferencedObject(Entry.RootNode);
+	}
+}
+
+FCameraNodeChildrenView UBlendStackCameraNode::OnGetChildren()
+{
+	FCameraNodeChildrenView View;
+	for (FCameraModeEntry& Entry : Entries)
+	{
+		View.Add(Entry.RootNode);
+	}
+	return View;
+}
+
+void UBlendStackCameraNode::OnRun(const FCameraNodeRunParams& Params, FCameraNodeRunResult& OutResult)
+{
+	// Start by evaluating all the root nodes in the stack.
+	for (FCameraModeEntry& Entry : Entries)
+	{
+		FCameraNodeRunParams CurParams;
+		CurParams.DeltaTime = Params.DeltaTime;
+		CurParams.bIsFirstFrame = Entry.bIsFirstFrame;
+
+		FCameraNodeRunResult& CurResult(Entry.Result);
+
+		if (!Entry.bIsFrozen)
+		{
+			// If the context in which this camera mode runs doesn't have a valid result,
+			// skip it.
+			const UCameraEvaluationContext* CurContext = Entry.EvaluationContext.Get();
+			if (UNLIKELY(CurContext == nullptr))
+			{
+				CurResult.bIsValid = false;
+				continue;
+			}
+			const FCameraNodeRunResult& ContextResult(CurContext->GetInitialResult());
+			if (UNLIKELY(!ContextResult.bIsValid))
+			{
+				CurResult.bIsValid = false;
+				continue;
+			}
+
+			// Start with the input given to us.
+			CurResult.CameraPose = OutResult.CameraPose;
+			CurResult.CameraPose.ClearAllChangedFlags();
+
+			// Override it with whatever the evaluation context has set on its result.
+			CurResult.CameraPose.OverrideChanged(ContextResult.CameraPose);
+			CurResult.bIsCameraCut = OutResult.bIsCameraCut || ContextResult.bIsCameraCut;
+			CurResult.bIsValid = true;
+
+			// Run the camera mode!
+			Entry.RootNode->Run(CurParams, CurResult);
+		}
+		else
+		{
+			// Only evaluate the blend via the root node.
+			Entry.RootNode->Run(CurParams, CurResult);
+		}
+	}
+
+	// Now blend all the results, keeping track of blends that have reached 100% so
+	// that we can remove any camera modes below (since the would have been completely
+	// blended out by that).
+	int32 EntryIndex = 0;
+	int32 PopEntriesBelow = INDEX_NONE;
+	for (FCameraModeEntry& Entry : Entries)
+	{
+		FCameraNodeRunResult& CurResult(Entry.Result);
+		if (UNLIKELY(!CurResult.bIsValid))
+		{
+			continue;
+		}
+
+		const FCameraPoseFlags ChangedFlags(CurResult.CameraPose.GetChangedFlags());
+
+		FCameraNodeRunParams CurParams;
+		CurParams.DeltaTime = Params.DeltaTime;
+		CurParams.bIsFirstFrame = Entry.bIsFirstFrame;
+		FCameraNodeBlendParams BlendParams(CurParams, CurResult);
+
+		FCameraNodeBlendResult BlendResult(OutResult);
+
+		UBlendCameraNode* EntryBlend = Entry.RootNode->Blend;
+		if (EntryBlend)
+		{
+			EntryBlend->BlendResults(BlendParams, BlendResult);
+
+			if (BlendResult.bIsBlendFull && BlendResult.bIsBlendFinished)
+			{
+				PopEntriesBelow = EntryIndex;
+			}
+		}
+		else
+		{
+			OutResult.CameraPose.OverrideChanged(CurResult.CameraPose);
+
+			PopEntriesBelow = EntryIndex;
+		}
+
+		CurResult.CameraPose.SetChangedFlags(ChangedFlags);
+		
+		++EntryIndex;
+	}
+
+	// Pop out camera modes that have been blended out.
+	if (bAutoPop && PopEntriesBelow != INDEX_NONE)
+	{
+		Entries.RemoveAt(0, PopEntriesBelow, false);
+	}
+
+	// Reset first frame flags.
+	for (FCameraModeEntry& Entry : Entries)
+	{
+		Entry.bIsFirstFrame = false;
+	}
+}
+
+const FCameraModeTransition* UBlendStackCameraNode::FindTransition(const FBlendStackCameraPushParams& Params) const
+{
+	const UCameraEvaluationContext* ToContext = Params.EvaluationContext.Get();
+	const UCameraAsset* ToCameraAsset = ToContext ? ToContext->GetCameraAsset() : nullptr;
+	const UCameraMode* ToCameraMode = Params.CameraMode;
+
+	// Find a transition that works for blending towards ToCameraMode.
+	// If the stack isn't empty, we need to find a transition that works between the previous and 
+	// next camera modes. If the stack is empty, we blend the new camera mode in from nothing if
+	// appropriate.
+	if (!Entries.IsEmpty())
+	{
+		const FCameraModeTransition* TransitionToUse = nullptr;
+
+		// Start by looking at exit transitions on the last active (top) camera mode.
+		const FCameraModeEntry& TopEntry = Entries.Top();
+
+		const UCameraEvaluationContext* FromContext = TopEntry.EvaluationContext.Get();
+		const UCameraAsset* FromCameraAsset = FromContext ? FromContext->GetCameraAsset() : nullptr;
+		const UCameraMode* FromCameraMode = TopEntry.OriginalCameraMode;
+
+		if (!TopEntry.bIsFrozen)
+		{
+			// Look for exit transitions on the last active camera mode itself.
+			TransitionToUse = FindTransition(
+					FromCameraMode->ExitTransitions,
+					FromCameraMode, FromCameraAsset, false,
+					ToCameraMode, ToCameraAsset);
+			if (TransitionToUse)
+			{
+				return TransitionToUse;
+			}
+			
+			// Look for exit transitions on its parent camera asset.
+			if (FromCameraAsset)
+			{
+				TransitionToUse = FindTransition(
+						FromCameraAsset->ExitTransitions,
+						FromCameraMode, FromCameraAsset, false,
+						ToCameraMode, ToCameraAsset);
+				if (TransitionToUse)
+				{
+					return TransitionToUse;
+				}
+			}
+		}
+
+		// Now look at enter transitions on the new camera mode.
+		TransitionToUse = FindTransition(
+				ToCameraMode->EnterTransitions,
+				FromCameraMode, FromCameraAsset, TopEntry.bIsFrozen,
+				ToCameraMode, ToCameraAsset);
+		if (TransitionToUse)
+		{
+			return TransitionToUse;
+		}
+
+		// Look at enter transitions on its parent camera asset.
+		if (ToCameraAsset)
+		{
+			TransitionToUse = FindTransition(
+					ToCameraAsset->EnterTransitions,
+					FromCameraMode, FromCameraAsset, TopEntry.bIsFrozen,
+					ToCameraMode, ToCameraAsset);
+			if (TransitionToUse)
+			{
+				return TransitionToUse;
+			}
+		}
+	}
+	else if (bBlendFirstCameraMode)
+	{
+		return FindTransition(
+				ToCameraMode->EnterTransitions,
+				nullptr, nullptr, false,
+				ToCameraMode, ToCameraAsset);
+	}
+
+	return nullptr;
+}
+
+const FCameraModeTransition* UBlendStackCameraNode::FindTransition(
+			TArrayView<const FCameraModeTransition> Transitions, 
+			const UCameraMode* FromCameraMode, const UCameraAsset* FromCameraAsset, bool bFromFrozen,
+			const UCameraMode* ToCameraMode, const UCameraAsset* ToCameraAsset) const
+{
+	FCameraModeTransitionConditionMatchParams MatchParams;
+	MatchParams.FromCameraMode = FromCameraMode;
+	MatchParams.FromCameraAsset = FromCameraAsset;
+	MatchParams.ToCameraMode = ToCameraMode;
+	MatchParams.ToCameraAsset = ToCameraAsset;
+
+	// The transition should be used if all its conditions pass.
+	for (const FCameraModeTransition& Transition : Transitions)
+	{
+		bool bConditionsPass = true;
+		for (const UCameraModeTransitionCondition* Condition : Transition.Conditions)
+		{
+			if (!Condition->TransitionMatches(MatchParams))
+			{
+				bConditionsPass = false;
+				break;
+			}
+		}
+
+		if (bConditionsPass)
+		{
+			return &Transition;
+		}
+	}
+
+	return nullptr;
+}
+
diff --git a/Engine/Plugins/Cameras/GameplayCameras/Source/GameplayCameras/Private/Core/BlendStackRootCameraNode.cpp b/Engine/Plugins/Cameras/GameplayCameras/Source/GameplayCameras/Private/Core/BlendStackRootCameraNode.cpp
new file mode 100644
--- /dev/null
+++ b/Engine/Plugins/Cameras/GameplayCameras/Source/GameplayCameras/Private/Core/BlendStackRootCameraNode.cpp
@@ -0,0 +1,26 @@
+// Copyright Epic Games, Inc. All Rights Reserved.
+
+#include "Core/BlendStackRootCameraNode.h"
+
+#include "Core/BlendCameraNode.h"
+#include "Core/CameraMode.h"
+
+#include UE_INLINE_GENERATED_CPP_BY_NAME(BlendStackRootCameraNode)
+
+FCameraNodeChildrenView UBlendStackRootCameraNode::OnGetChildren()
+{
+	return FCameraNodeChildrenView({ Blend, RootNode });
+}
+
+void UBlendStackRootCameraNode::OnRun(const FCameraNodeRunParams& Params, FCameraNodeRunResult& OutResult)
+{
+	if (Blend)
+	{
+		Blend->Run(Params, OutResult);
+	}
+	if (RootNode)
+	{
+		RootNode->Run(Params, OutResult);
+	}
+}
+
diff --git a/Engine/Plugins/Cameras/GameplayCameras/Source/GameplayCameras/Private/Core/BlendStackRootCameraNode.h b/Engine/Plugins/Cameras/GameplayCameras/Source/GameplayCameras/Private/Core/BlendStackRootCameraNode.h
new file mode 100644
--- /dev/null
+++ b/Engine/Plugins/Cameras/GameplayCameras/Source/GameplayCameras/Private/Core/BlendStackRootCameraNode.h
@@ -0,0 +1,37 @@
+// Copyright Epic Games, Inc. All Rights Reserved.
+
+#pragma once
+
+#include "Core/CameraNode.h"
+
+#include "BlendStackRootCameraNode.generated.h"
+
+class UBlendCameraNode;
+class UCameraMode;
+
+/**
+ * Root camera node for running a camera mode in a blend stack.
+ * This camera node wraps both the camera mode's root node, and the
+ * blend node used to blend it.
+ */
+UCLASS(MinimalAPI)
+class UBlendStackRootCameraNode : public UCameraNode
+{
+	GENERATED_BODY()
+
+protected:
+
+	virtual FCameraNodeChildrenView OnGetChildren() override;
+	virtual void OnRun(const FCameraNodeRunParams& Params, FCameraNodeRunResult& OutResult) override;
+
+public:
+
+	/** The blend to use on the camera mode. */
+	UPROPERTY()
+	TObjectPtr<UBlendCameraNode> Blend;
+
+	/** The root of the instantied camera node tree. */
+	UPROPERTY()
+	TObjectPtr<UCameraNode> RootNode;
+};
+
diff --git a/Engine/Plugins/Cameras/GameplayCameras/Source/GameplayCameras/Private/Core/CameraAsset.cpp b/Engine/Plugins/Cameras/GameplayCameras/Source/GameplayCameras/Private/Core/CameraAsset.cpp
new file mode 100644
--- /dev/null
+++ b/Engine/Plugins/Cameras/GameplayCameras/Source/GameplayCameras/Private/Core/CameraAsset.cpp
@@ -0,0 +1,6 @@
+// Copyright Epic Games, Inc. All Rights Reserved.
+
+#include "Core/CameraAsset.h"
+
+#include UE_INLINE_GENERATED_CPP_BY_NAME(CameraAsset)
+
diff --git a/Engine/Plugins/Cameras/GameplayCameras/Source/GameplayCameras/Private/Core/CameraDirector.cpp b/Engine/Plugins/Cameras/GameplayCameras/Source/GameplayCameras/Private/Core/CameraDirector.cpp
new file mode 100644
--- /dev/null
+++ b/Engine/Plugins/Cameras/GameplayCameras/Source/GameplayCameras/Private/Core/CameraDirector.cpp
@@ -0,0 +1,11 @@
+// Copyright Epic Games, Inc. All Rights Reserved.
+
+#include "Core/CameraDirector.h"
+
+#include UE_INLINE_GENERATED_CPP_BY_NAME(CameraDirector)
+
+void UCameraDirector::Run(const FCameraDirectorRunParams& Params, FCameraDirectorRunResult& OutResult)
+{
+	OnRun(Params, OutResult);
+}
+
diff --git a/Engine/Plugins/Cameras/GameplayCameras/Source/GameplayCameras/Private/Core/CameraEvaluationContext.cpp b/Engine/Plugins/Cameras/GameplayCameras/Source/GameplayCameras/Private/Core/CameraEvaluationContext.cpp
new file mode 100644
--- /dev/null
+++ b/Engine/Plugins/Cameras/GameplayCameras/Source/GameplayCameras/Private/Core/CameraEvaluationContext.cpp
@@ -0,0 +1,11 @@
+// Copyright Epic Games, Inc. All Rights Reserved.
+
+#include "Core/CameraEvaluationContext.h"
+
+#include UE_INLINE_GENERATED_CPP_BY_NAME(CameraEvaluationContext)
+
+UCameraEvaluationContext::UCameraEvaluationContext(const FObjectInitializer& ObjectInit)
+	: Super(ObjectInit)
+{
+}
+
diff --git a/Engine/Plugins/Cameras/GameplayCameras/Source/GameplayCameras/Private/Core/CameraEvaluationContextStack.cpp b/Engine/Plugins/Cameras/GameplayCameras/Source/GameplayCameras/Private/Core/CameraEvaluationContextStack.cpp
new file mode 100644
--- /dev/null
+++ b/Engine/Plugins/Cameras/GameplayCameras/Source/GameplayCameras/Private/Core/CameraEvaluationContextStack.cpp
@@ -0,0 +1,95 @@
+// Copyright Epic Games, Inc. All Rights Reserved.
+
+#include "Core/CameraEvaluationContextStack.h"
+
+#include "Core/CameraAsset.h"
+#include "Core/CameraDirector.h"
+#include "Core/CameraEvaluationContext.h"
+#include "Core/CameraRuntimeInstantiator.h"
+#include "Core/CameraSystemEvaluator.h"
+#include "UObject/Package.h"
+
+FCameraEvaluationContextInfo FCameraEvaluationContextStack::GetActiveContext() const
+{
+	for (const FContextEntry& Entry : ReverseIterate(Entries))
+	{
+		if (UCameraEvaluationContext* Context = Entry.WeakContext.Get())
+		{
+			return FCameraEvaluationContextInfo{ Context, Entry.CameraDirector };
+		}
+	}
+	return FCameraEvaluationContextInfo();
+}
+
+bool FCameraEvaluationContextStack::HasContext(UCameraEvaluationContext* Context) const
+{
+	for (const FContextEntry& Entry : ReverseIterate(Entries))
+	{
+		if (Context == Entry.WeakContext.Get())
+		{
+			return true;
+		}
+	}
+	return false;
+}
+
+void FCameraEvaluationContextStack::PushContext(UCameraEvaluationContext* Context)
+{
+	// If we're pushing an existing context, move it to the top.
+	const int32 ExistingIndex = Entries.IndexOfByPredicate(
+			[Context](const FContextEntry& Entry) { return Entry.WeakContext == Context; });
+	if (ExistingIndex != INDEX_NONE)
+	{
+		if (ExistingIndex < Entries.Num() - 1)
+		{
+			const FContextEntry EntryCopy(Entries[ExistingIndex]);
+			Entries.RemoveAt(ExistingIndex);
+			Entries.Add(EntryCopy);
+		}
+		return;
+	}
+
+	// Instantiate the camera director.
+	FCameraRuntimeInstantiationParams InstParams;
+	InstParams.InstantiationOuter = InstantiationOuter.Get();
+	if (!InstParams.InstantiationOuter)
+	{
+		InstParams.InstantiationOuter = GetTransientPackage();
+	}
+
+	const UCameraDirector* OriginalCameraDirector = Context->GetCameraAsset()->CameraDirector;
+	UCameraDirector* NewCameraDirector = FCameraRuntimeInstantiator::InstantiateCameraDirector(
+			OriginalCameraDirector, InstParams);
+	
+	// Add an entry in the stack.
+	FContextEntry NewEntry;
+	NewEntry.WeakContext = Context;
+	NewEntry.CameraDirector = NewCameraDirector;
+	Entries.Push(NewEntry);
+}
+
+bool FCameraEvaluationContextStack::RemoveContext(UCameraEvaluationContext* Context)
+{
+	const int32 NumRemoved = Entries.RemoveAll(
+			[Context](FContextEntry& Entry) { return Entry.WeakContext == Context; });
+	return (NumRemoved > 0);
+}
+
+void FCameraEvaluationContextStack::PopContext()
+{
+	Entries.Pop();
+}
+
+void FCameraEvaluationContextStack::SetInstantiationOuter(UObject* InInstantiationOuter)
+{
+	InstantiationOuter = InInstantiationOuter;
+}
+
+void FCameraEvaluationContextStack::AddReferencedObjects(FReferenceCollector& Collector)
+{
+	for (FContextEntry& Entry : Entries)
+	{
+		Collector.AddReferencedObject(Entry.CameraDirector);
+	}
+}
+
diff --git a/Engine/Plugins/Cameras/GameplayCameras/Source/GameplayCameras/Private/Core/CameraMode.cpp b/Engine/Plugins/Cameras/GameplayCameras/Source/GameplayCameras/Private/Core/CameraMode.cpp
new file mode 100644
--- /dev/null
+++ b/Engine/Plugins/Cameras/GameplayCameras/Source/GameplayCameras/Private/Core/CameraMode.cpp
@@ -0,0 +1,6 @@
+// Copyright Epic Games, Inc. All Rights Reserved.
+
+#include "Core/CameraMode.h"
+
+#include UE_INLINE_GENERATED_CPP_BY_NAME(CameraMode)
+
diff --git a/Engine/Plugins/Cameras/GameplayCameras/Source/GameplayCameras/Private/Core/CameraModeTransition.cpp b/Engine/Plugins/Cameras/GameplayCameras/Source/GameplayCameras/Private/Core/CameraModeTransition.cpp
new file mode 100644
--- /dev/null
+++ b/Engine/Plugins/Cameras/GameplayCameras/Source/GameplayCameras/Private/Core/CameraModeTransition.cpp
@@ -0,0 +1,11 @@
+// Copyright Epic Games, Inc. All Rights Reserved.
+
+#include "Core/CameraModeTransition.h"
+
+#include UE_INLINE_GENERATED_CPP_BY_NAME(CameraModeTransition)
+
+bool UCameraModeTransitionCondition::TransitionMatches(const FCameraModeTransitionConditionMatchParams& Params) const
+{
+	return OnTransitionMatches(Params);
+}
+
diff --git a/Engine/Plugins/Cameras/GameplayCameras/Source/GameplayCameras/Private/Core/CameraNode.cpp b/Engine/Plugins/Cameras/GameplayCameras/Source/GameplayCameras/Private/Core/CameraNode.cpp
new file mode 100644
--- /dev/null
+++ b/Engine/Plugins/Cameras/GameplayCameras/Source/GameplayCameras/Private/Core/CameraNode.cpp
@@ -0,0 +1,41 @@
+// Copyright Epic Games, Inc. All Rights Reserved.
+
+#include "Core/CameraNode.h"
+
+#include "Core/CameraRuntimeInstantiator.h"
+
+#include UE_INLINE_GENERATED_CPP_BY_NAME(CameraNode)
+
+void FCameraNodeRunResult::Reset()
+{
+	CameraPose.Reset();
+	bIsCameraCut = false;
+	bIsValid = false;
+}
+
+FCameraNodeChildrenView UCameraNode::GetChildren()
+{
+	return OnGetChildren();
+}
+
+void UCameraNode::Run(const FCameraNodeRunParams& Params, FCameraNodeRunResult& OutResult)
+{
+	if (bIsEnabled)
+	{
+		OnRun(Params, OutResult);
+	}
+}
+
+#if WITH_EDITOR
+
+void UCameraNode::PostEditChangeProperty(struct FPropertyChangedEvent& PropertyChangedEvent)
+{
+	Super::PostEditChangeProperty(PropertyChangedEvent);
+	
+	if (!bIsInstantiated)
+	{
+		FCameraRuntimeInstantiator::ForwardPropertyChange(this, PropertyChangedEvent);
+	}
+}
+
+#endif
diff --git a/Engine/Plugins/Cameras/GameplayCameras/Source/GameplayCameras/Private/Core/CameraPose.cpp b/Engine/Plugins/Cameras/GameplayCameras/Source/GameplayCameras/Private/Core/CameraPose.cpp
new file mode 100644
--- /dev/null
+++ b/Engine/Plugins/Cameras/GameplayCameras/Source/GameplayCameras/Private/Core/CameraPose.cpp
@@ -0,0 +1,252 @@
+// Copyright Epic Games, Inc. All Rights Reserved.
+
+#include "Core/CameraPose.h"
+
+#include "Math/Ray.h"
+#include "Runtime/Engine/Classes/PhysicsField/PhysicsFieldComponent.h"
+
+#include UE_INLINE_GENERATED_CPP_BY_NAME(CameraPose)
+
+FCameraPoseFlags& FCameraPoseFlags::SetAllFlags(bool bInValue)
+{
+#define UE_CAMERA_POSE_FOR_PROPERTY(PropType, PropName)\
+	PropName = bInValue;
+
+UE_CAMERA_POSE_FOR_ALL_PROPERTIES()
+
+#undef UE_CAMERA_POSE_FOR_PROPERTY
+
+	return *this;
+}
+
+FCameraPoseFlags& FCameraPoseFlags::ExclusiveCombine(const FCameraPoseFlags& OtherFlags)
+{
+#define UE_CAMERA_POSE_FOR_PROPERTY(PropType, PropName)\
+	if (OtherFlags.PropName)\
+	{\
+		ensureMsgf(!PropName, TEXT("Exclusive combination failed: " #PropName " set on both flags!"));\
+		PropName = true;\
+	}
+
+UE_CAMERA_POSE_FOR_ALL_PROPERTIES()
+
+#undef UE_CAMERA_POSE_FOR_PROPERTY
+
+	return *this;
+}
+
+FCameraPoseFlags& FCameraPoseFlags::AND(const FCameraPoseFlags& OtherFlags)
+{
+#define UE_CAMERA_POSE_FOR_PROPERTY(PropType, PropName)\
+	PropName = PropName && OtherFlags.PropName;
+
+UE_CAMERA_POSE_FOR_ALL_PROPERTIES()
+
+#undef UE_CAMERA_POSE_FOR_PROPERTY
+
+	return *this;
+}
+
+FCameraPoseFlags& FCameraPoseFlags::OR(const FCameraPoseFlags& OtherFlags)
+{
+#define UE_CAMERA_POSE_FOR_PROPERTY(PropType, PropName)\
+	PropName = PropName || OtherFlags.PropName;
+
+UE_CAMERA_POSE_FOR_ALL_PROPERTIES()
+
+#undef UE_CAMERA_POSE_FOR_PROPERTY
+
+	return *this;
+}
+
+FCameraPose::FCameraPose()
+{
+}
+
+void FCameraPose::Reset(bool bSetAllChangedFlags)
+{
+	*this = FCameraPose();
+
+	if (bSetAllChangedFlags)
+	{
+		SetAllChangedFlags();
+	}
+}
+
+void FCameraPose::SetAllChangedFlags()
+{
+	ChangedFlags.SetAllFlags(true);
+}
+
+void FCameraPose::ClearAllChangedFlags()
+{
+	ChangedFlags.SetAllFlags(false);
+}
+
+FTransform3d FCameraPose::GetTransform() const
+{
+	FTransform3d Transform;
+	Transform.SetLocation(Location);
+	Transform.SetRotation(Rotation.Quaternion());
+	return Transform;
+}
+
+void FCameraPose::SetTransform(FTransform3d Transform)
+{
+	SetLocation(Transform.GetLocation());
+	SetRotation(Transform.GetRotation().Rotator());
+}
+
+float FCameraPose::GetEffectiveFieldOfView() const
+{
+	checkf(
+			(FocalLength > 0.f && FieldOfView <= 0.f) || (FocalLength <= 0.f && FieldOfView > 0.f),
+			TEXT("FocalLength or FieldOfView must have a valid, positive value."));
+
+	if (FocalLength > 0.f)
+	{
+		// Compute FOV with similar code to UCineCameraComponent...
+		float CropedSensorWidth = SensorWidth * SqueezeFactor;
+		if (AspectRatio > 0.0f)
+		{
+			float DesqueezeAspectRatio = SensorWidth * SqueezeFactor / SensorHeight;
+			if (AspectRatio < DesqueezeAspectRatio)
+			{
+				CropedSensorWidth *= AspectRatio / DesqueezeAspectRatio;
+			}
+		}
+
+		return FMath::RadiansToDegrees(2.f * FMath::Atan(CropedSensorWidth / (2.f * FocalLength)));
+	}
+	else
+	{
+		// Let's use the FOV directly, like in the good old times.
+		return FieldOfView;
+	}
+}
+
+FRay3d FCameraPose::GetAimRay() const
+{
+	const bool bDirectionIsNormalized = false;
+	const FVector3d TargetDir{ TargetDistance, 0, 0 };
+	return FRay3d(Location, Rotation.RotateVector(TargetDir), bDirectionIsNormalized);
+}
+
+void FCameraPose::OverrideChanged(const FCameraPose& OtherPose)
+{
+	const FCameraPoseFlags& OtherPoseChangedFlags = OtherPose.GetChangedFlags();
+
+#define UE_CAMERA_POSE_FOR_PROPERTY(PropType, PropName)\
+	if (OtherPoseChangedFlags.PropName)\
+	{\
+		Set##PropName(OtherPose.Get##PropName());\
+	}
+
+UE_CAMERA_POSE_FOR_ALL_PROPERTIES()
+
+#undef UE_CAMERA_POSE_FOR_PROPERTY
+}
+
+void FCameraPose::LerpChanged(const FCameraPose& ToPose, float Factor)
+{
+	FCameraPoseFlags DummyFlags;
+	DummyFlags.SetAllFlags(true);
+	InternalLerpChanged(ToPose, Factor, DummyFlags, false, DummyFlags);
+}
+
+void FCameraPose::LerpChanged(const FCameraPose& ToPose, float Factor, const FCameraPoseFlags& InMask, bool bInvertMask, FCameraPoseFlags& OutMask)
+{
+	InternalLerpChanged(ToPose, Factor, InMask, bInvertMask, OutMask);
+}
+
+void FCameraPose::InternalLerpChanged(const FCameraPose& ToPose, float Factor, const FCameraPoseFlags& InMask, bool bInvertMask, FCameraPoseFlags& OutMask)
+{
+	if (UNLIKELY(Factor == 0.f))
+	{
+		return;
+	}
+
+	const bool bIsOverride = (Factor == 1.f);
+	const FCameraPoseFlags& ToPoseChangedFlags = ToPose.GetChangedFlags();
+
+	if (bIsOverride)
+	{
+		// The interpolation factor is 1 so we just override the properties.
+	
+#define UE_CAMERA_POSE_FOR_PROPERTY(PropType, PropName)\
+		if ((!bInvertMask && InMask.PropName) || (bInvertMask && !InMask.PropName))\
+		{\
+			if (ToPoseChangedFlags.PropName)\
+			{\
+				Set##PropName(ToPose.Get##PropName());\
+				OutMask.PropName = true;\
+			}\
+		}
+
+		UE_CAMERA_POSE_FOR_ALL_PROPERTIES()
+
+#undef UE_CAMERA_POSE_FOR_PROPERTY
+
+	}
+	else
+	{
+		// Interpolate all the properties.
+		//
+		// Start with those we can simply feed to a LERP formula. Some properties don't
+		// necessarily make sense to interpolate (like, what does it mean to interpolate the
+		// sensor size of a camera?) but, well, we use whatever we have been given at this
+		// point.
+
+#define UE_CAMERA_POSE_FOR_PROPERTY(PropType, PropName)\
+		if ((!bInvertMask && InMask.PropName) || (bInvertMask && !InMask.PropName))\
+		{\
+			if (ToPoseChangedFlags.PropName)\
+			{\
+				ensureMsgf(ChangedFlags.PropName, TEXT("Interpolating " #PropName " from default value!"));\
+				Set##PropName(FMath::Lerp(Get##PropName(), ToPose.Get##PropName(), Factor));\
+			}\
+			OutMask.PropName = true;\
+		}
+
+		UE_CAMERA_POSE_FOR_TRANSFORM_PROPERTIES()
+		UE_CAMERA_POSE_FOR_INTERPOLABLE_PROPERTIES()
+
+#undef UE_CAMERA_POSE_FOR_PROPERTY
+
+		// Next, handle the special case of FOV, where we might have to blend between a pose
+		// specifying FOV directly and a pose using focal length.
+		if (
+				(!bInvertMask && (InMask.FieldOfView || InMask.FocalLength)) ||
+				(bInvertMask && (!InMask.FieldOfView || !InMask.FocalLength)))
+		{
+			ensureMsgf(
+					ChangedFlags.FieldOfView || ChangedFlags.FocalLength, 
+					TEXT("Interpolating FieldOfView or FocalLength from default value!"));
+			const float FromFieldOfView = GetEffectiveFieldOfView();
+			const float ToFieldOfView = ToPose.GetEffectiveFieldOfView();
+			SetFieldOfView(FMath::Lerp(FromFieldOfView, ToFieldOfView, Factor));
+			SetFocalLength(-1);
+			OutMask.FieldOfView = true;
+			OutMask.FocalLength = true;
+		}
+
+		// Last, do booleans, which just flip their value once we reach 50% interpolation.
+
+#define UE_CAMERA_POSE_FOR_PROPERTY(PropType, PropName)\
+		if ((!bInvertMask && InMask.PropName) || (bInvertMask && !InMask.PropName))\
+		{\
+			if (ToPoseChangedFlags.PropName && Factor >= 0.5f)\
+			{\
+				ensureMsgf(ChangedFlags.PropName, TEXT("Interpolating " #PropName " from default value!"));\
+				Set##PropName(ToPose.Get##PropName());\
+			}\
+			OutMask.PropName = true;\
+		}
+
+		UE_CAMERA_POSE_FOR_BOOL_PROPERTIES()
+
+#undef UE_CAMERA_POSE_FOR_PROPERTY
+
+	}
+}
+
diff --git a/Engine/Plugins/Cameras/GameplayCameras/Source/GameplayCameras/Private/Core/CameraRuntimeInstantiator.cpp b/Engine/Plugins/Cameras/GameplayCameras/Source/GameplayCameras/Private/Core/CameraRuntimeInstantiator.cpp
new file mode 100644
--- /dev/null
+++ b/Engine/Plugins/Cameras/GameplayCameras/Source/GameplayCameras/Private/Core/CameraRuntimeInstantiator.cpp
@@ -0,0 +1,79 @@
+// Copyright Epic Games, Inc. All Rights Reserved.
+
+#include "Core/CameraRuntimeInstantiator.h"
+
+#include "Core/CameraDirector.h"
+#include "Core/CameraNode.h"
+#include "CoreGlobals.h"
+#include "IGameplayCamerasModule.h"
+#include "IGameplayCamerasLiveEditManager.h"
+#include "Modules/ModuleManager.h"
+#include "UObject/UObjectGlobals.h"
+
+IGameplayCamerasModule* FCameraRuntimeInstantiator::GameplayCamerasModule(nullptr);
+
+IGameplayCamerasModule& FCameraRuntimeInstantiator::GetGameplayCamerasModule()
+{
+	if (!GameplayCamerasModule)
+	{
+		GameplayCamerasModule = &FModuleManager::LoadModuleChecked<IGameplayCamerasModule>("GameplayCameras");
+	}
+	return *GameplayCamerasModule;
+}
+
+UCameraNode* FCameraRuntimeInstantiator::InstantiateCameraNodeTree(const UCameraNode* InRootNode, const FCameraRuntimeInstantiationParams& Params)
+{
+	UCameraNode* NewRootNode = CastChecked<UCameraNode>(InstantiateObject(InRootNode, Params));
+	return NewRootNode;
+}
+
+UCameraDirector* FCameraRuntimeInstantiator::InstantiateCameraDirector(const UCameraDirector* InCameraDirector, const FCameraRuntimeInstantiationParams& Params)
+{
+	UCameraDirector* NewCameraDirector = CastChecked<UCameraDirector>(InstantiateObject(InCameraDirector, Params));
+	return NewCameraDirector;
+}
+
+UObject* FCameraRuntimeInstantiator::InstantiateObject(const UObject* InSource, const FCameraRuntimeInstantiationParams& Params)
+{
+	FObjectDuplicationParameters DuplicationParams = InitStaticDuplicateObjectParams(
+			InSource,
+			Params.InstantiationOuter,
+			Params.InstantiationName);
+#if WITH_EDITOR
+	TMap<UObject*, UObject*> CreatedObjects;
+	DuplicationParams.CreatedObjects = &CreatedObjects;
+#endif
+
+	UObject* OutRootObject = StaticDuplicateObjectEx(DuplicationParams);
+
+#if WITH_EDITOR
+	if (!IsRunningCommandlet())
+	{
+		for (const TPair<UObject*, UObject*>& Pair : CreatedObjects)
+		{
+			if (UCameraNode* CreatedNode = Cast<UCameraNode>(Pair.Value))
+			{
+				CreatedNode->SetIsInstantiated();
+			}
+		}
+
+		if (TSharedPtr<IGameplayCamerasLiveEditManager> LiveEditManager = GetGameplayCamerasModule().GetLiveEditManager())
+		{
+			LiveEditManager->RegisterInstantiatedObjects(CreatedObjects);
+		}
+	}
+#endif
+
+	return OutRootObject;
+}
+
+void FCameraRuntimeInstantiator::ForwardPropertyChange(const UObject* Object, const FPropertyChangedEvent& PropertyChangedEvent)
+{
+#if WITH_EDITOR
+	if (TSharedPtr<IGameplayCamerasLiveEditManager> LiveEditManager = GetGameplayCamerasModule().GetLiveEditManager())
+	{
+		LiveEditManager->ForwardPropertyChange(Object, PropertyChangedEvent);
+	}
+#endif
+}
+
diff --git a/Engine/Plugins/Cameras/GameplayCameras/Source/GameplayCameras/Private/Core/CameraSystemEvaluator.cpp b/Engine/Plugins/Cameras/GameplayCameras/Source/GameplayCameras/Private/Core/CameraSystemEvaluator.cpp
new file mode 100644
--- /dev/null
+++ b/Engine/Plugins/Cameras/GameplayCameras/Source/GameplayCameras/Private/Core/CameraSystemEvaluator.cpp
@@ -0,0 +1,98 @@
+// Copyright Epic Games, Inc. All Rights Reserved.
+
+#include "Core/CameraSystemEvaluator.h"
+
+#include "Camera/CameraTypes.h"
+#include "Core/CameraAsset.h"
+#include "Core/CameraEvaluationContext.h"
+#include "Core/CameraMode.h"
+#include "Core/DefaultRootCameraNode.h"
+#include "IGameplayCamerasModule.h"
+
+#include UE_INLINE_GENERATED_CPP_BY_NAME(CameraSystemEvaluator)
+
+DECLARE_CYCLE_STAT(TEXT("Camera System Eval"), CameraSystemEval_Total, STATGROUP_CameraSystem);
+
+UCameraSystemEvaluator::UCameraSystemEvaluator(const FObjectInitializer& ObjectInit)
+	: Super(ObjectInit)
+{
+	RootNode = ObjectInit.CreateDefaultSubobject<UDefaultRootCameraNode>(this, TEXT("RootNode"));
+
+	ContextStack.SetInstantiationOuter(this);
+}
+
+void UCameraSystemEvaluator::AddReferencedObjects(UObject* InThis, FReferenceCollector& Collector)
+{
+	UCameraSystemEvaluator* TypedThis = CastChecked<UCameraSystemEvaluator>(InThis);
+	TypedThis->ContextStack.AddReferencedObjects(Collector);
+}
+
+void UCameraSystemEvaluator::PushEvaluationContext(UCameraEvaluationContext* EvaluationContext)
+{
+	ContextStack.PushContext(EvaluationContext);
+}
+
+void UCameraSystemEvaluator::RemoveEvaluationContext(UCameraEvaluationContext* EvaluationContext)
+{
+	ContextStack.RemoveContext(EvaluationContext);
+}
+
+void UCameraSystemEvaluator::PopEvaluationContext()
+{
+	ContextStack.PopContext();
+}
+
+void UCameraSystemEvaluator::Update(const FCameraSystemEvaluationUpdateParams& Params)
+{
+	SCOPE_CYCLE_COUNTER(CameraSystemEval_Total);
+
+	// Get the active evaluation context.
+	FCameraEvaluationContextInfo ActiveContextInfo = ContextStack.GetActiveContext();
+	if (UNLIKELY(!ActiveContextInfo.IsValid()))
+	{
+		Result.bIsValid = false;
+		return;
+	}
+
+	// Run the camera director, and activate any camera mode(s) it returns to us.
+	UCameraDirector* ActiveDirector = ActiveContextInfo.CameraDirector;
+	if (ActiveDirector)
+	{
+		FCameraDirectorRunParams DirectorParams;
+		DirectorParams.DeltaTime = Params.DeltaTime;
+		DirectorParams.OwnerContext = ActiveContextInfo.EvaluationContext;
+
+		FCameraDirectorRunResult DirectorResult;
+
+		ActiveDirector->Run(DirectorParams, DirectorResult);
+
+		if (DirectorResult.ActiveCameraModes.Num() == 1)
+		{
+			FActivateCameraModeParams CameraModeParams;
+			CameraModeParams.EvaluationContext = ActiveContextInfo.EvaluationContext;
+			CameraModeParams.CameraMode = DirectorResult.ActiveCameraModes[0];
+			RootNode->ActivateCameraMode(CameraModeParams);
+		}
+	}
+
+	// Run the root camera node.
+	FCameraNodeRunParams NodeParams;
+	NodeParams.DeltaTime = Params.DeltaTime;
+
+	RootNodeResult.Reset();
+
+	RootNode->Run(NodeParams, RootNodeResult);
+
+	Result.CameraPose = RootNodeResult.CameraPose;
+	Result.bIsCameraCut = RootNodeResult.bIsCameraCut;
+	Result.bIsValid = true;
+}
+
+void UCameraSystemEvaluator::GetEvaluatedCameraView(FMinimalViewInfo& DesiredView)
+{
+	const FCameraPose& CameraPose = Result.CameraPose;
+	DesiredView.Location = CameraPose.GetLocation();
+	DesiredView.Rotation = CameraPose.GetRotation();
+	DesiredView.FOV = CameraPose.GetEffectiveFieldOfView();
+}
+
diff --git a/Engine/Plugins/Cameras/GameplayCameras/Source/GameplayCameras/Private/Core/DefaultRootCameraNode.cpp b/Engine/Plugins/Cameras/GameplayCameras/Source/GameplayCameras/Private/Core/DefaultRootCameraNode.cpp
new file mode 100644
--- /dev/null
+++ b/Engine/Plugins/Cameras/GameplayCameras/Source/GameplayCameras/Private/Core/DefaultRootCameraNode.cpp
@@ -0,0 +1,80 @@
+// Copyright Epic Games, Inc. All Rights Reserved.
+
+#include "Core/DefaultRootCameraNode.h"
+
+#include "Core/BlendStackCameraNode.h"
+#include "Core/CameraMode.h"
+
+#include UE_INLINE_GENERATED_CPP_BY_NAME(DefaultRootCameraNode)
+
+namespace UE::Cameras::Private
+{
+
+TObjectPtr<UBlendStackCameraNode> CreateBlendStack(
+		UObject* This, const FObjectInitializer& ObjectInit,
+		const FName& Name, bool bAutoPop = true, bool bBlendFirstCameraMode = false)
+{
+	TObjectPtr<UBlendStackCameraNode> NewBlendStack = ObjectInit.CreateDefaultSubobject<UBlendStackCameraNode>(
+			This, Name);
+	NewBlendStack->bAutoPop = bAutoPop;
+	NewBlendStack->bBlendFirstCameraMode = bBlendFirstCameraMode;
+	return NewBlendStack;
+}
+
+}  // namespace UE::Cameras::Private
+
+UDefaultRootCameraNode::UDefaultRootCameraNode(const FObjectInitializer& ObjectInit)
+	: Super(ObjectInit)
+{
+	using namespace UE::Cameras::Private;
+
+	BaseLayer = CreateBlendStack(this, ObjectInit, TEXT("BaseLayer"), false, true);
+	MainLayer = CreateBlendStack(this, ObjectInit, TEXT("MainLayer"));
+	GlobalLayer = CreateBlendStack(this, ObjectInit, TEXT("GlobalLayer"), false, true);
+	VisualLayer = CreateBlendStack(this, ObjectInit, TEXT("VisualLayer"), false, true);
+}
+
+FCameraNodeChildrenView UDefaultRootCameraNode::OnGetChildren()
+{
+	return FCameraNodeChildrenView({ BaseLayer, MainLayer, GlobalLayer, VisualLayer });
+}
+
+void UDefaultRootCameraNode::OnRun(const FCameraNodeRunParams& Params, FCameraNodeRunResult& OutResult)
+{
+	BaseLayer->Run(Params, OutResult);
+
+	MainLayer->Run(Params, OutResult);
+
+	GlobalLayer->Run(Params, OutResult);
+
+	VisualLayer->Run(Params, OutResult);
+}
+
+void UDefaultRootCameraNode::OnActivateCameraMode(const FActivateCameraModeParams& Params)
+{
+	UBlendStackCameraNode* TargetStack = nullptr;
+	switch (Params.Layer)
+	{
+		case ECameraModeLayer::Base:
+			TargetStack = CastChecked<UBlendStackCameraNode>(BaseLayer);
+			break;
+		case ECameraModeLayer::Main:
+			TargetStack = CastChecked<UBlendStackCameraNode>(MainLayer);
+			break;
+		case ECameraModeLayer::Global:
+			TargetStack = CastChecked<UBlendStackCameraNode>(GlobalLayer);
+			break;
+		case ECameraModeLayer::Visual:
+			TargetStack = CastChecked<UBlendStackCameraNode>(VisualLayer);
+			break;
+	}
+
+	if (ensure(TargetStack))
+	{
+		FBlendStackCameraPushParams PushParams;
+		PushParams.EvaluationContext = Params.EvaluationContext;
+		PushParams.CameraMode = Params.CameraMode;
+		TargetStack->Push(PushParams);
+	}
+}
+
diff --git a/Engine/Plugins/Cameras/GameplayCameras/Source/GameplayCameras/Private/Core/RootCameraNode.cpp b/Engine/Plugins/Cameras/GameplayCameras/Source/GameplayCameras/Private/Core/RootCameraNode.cpp
new file mode 100644
--- /dev/null
+++ b/Engine/Plugins/Cameras/GameplayCameras/Source/GameplayCameras/Private/Core/RootCameraNode.cpp
@@ -0,0 +1,13 @@
+// Copyright Epic Games, Inc. All Rights Reserved.
+
+#include "Core/RootCameraNode.h"
+
+#include "Core/CameraMode.h"
+
+#include UE_INLINE_GENERATED_CPP_BY_NAME(RootCameraNode)
+
+void URootCameraNode::ActivateCameraMode(const FActivateCameraModeParams& Params)
+{
+	OnActivateCameraMode(Params);
+}
+
diff --git a/Engine/Plugins/Cameras/GameplayCameras/Source/GameplayCameras/Private/Directors/SingleCameraDirector.cpp b/Engine/Plugins/Cameras/GameplayCameras/Source/GameplayCameras/Private/Directors/SingleCameraDirector.cpp
new file mode 100644
--- /dev/null
+++ b/Engine/Plugins/Cameras/GameplayCameras/Source/GameplayCameras/Private/Directors/SingleCameraDirector.cpp
@@ -0,0 +1,19 @@
+// Copyright Epic Games, Inc. All Rights Reserved.
+
+#include "Directors/SingleCameraDirector.h"
+
+#include UE_INLINE_GENERATED_CPP_BY_NAME(SingleCameraDirector)
+
+USingleCameraDirector::USingleCameraDirector(const FObjectInitializer& ObjectInit)
+	: Super(ObjectInit)
+{
+}
+
+void USingleCameraDirector::OnRun(const FCameraDirectorRunParams& Params, FCameraDirectorRunResult& OutResult)
+{
+	if (CameraMode)
+	{
+		OutResult.ActiveCameraModes.Add(CameraMode);
+	}
+}
+
diff --git a/Engine/Plugins/Cameras/GameplayCameras/Source/GameplayCameras/Private/GameFramework/GameplayCameraActor.cpp b/Engine/Plugins/Cameras/GameplayCameras/Source/GameplayCameras/Private/GameFramework/GameplayCameraActor.cpp
new file mode 100644
--- /dev/null
+++ b/Engine/Plugins/Cameras/GameplayCameras/Source/GameplayCameras/Private/GameFramework/GameplayCameraActor.cpp
@@ -0,0 +1,31 @@
+// Copyright Epic Games, Inc. All Rights Reserved.
+
+#include "GameFramework/GameplayCameraActor.h"
+
+#include "Engine/EngineTypes.h"
+#include "Engine/World.h"
+#include "GameFramework/GameplayCameraComponent.h"
+#include "GameFramework/PlayerController.h"
+#include "Kismet/GameplayStatics.h"
+
+#include UE_INLINE_GENERATED_CPP_BY_NAME(GameplayCameraActor)
+
+#define LOCTEXT_NAMESPACE "GameplayCameraActor"
+
+AGameplayCameraActor::AGameplayCameraActor(const FObjectInitializer& ObjectInit)
+	: Super(ObjectInit)
+{
+	SceneComponent = CreateDefaultSubobject<USceneComponent>(TEXT("SceneComponent"));
+	RootComponent = SceneComponent;
+
+	CameraComponent = CreateDefaultSubobject<UGameplayCameraComponent>(TEXT("CameraComponent"));
+	CameraComponent->SetupAttachment(SceneComponent);
+}
+
+USceneComponent* AGameplayCameraActor::GetDefaultAttachComponent() const
+{
+	return CameraComponent;
+}
+
+#undef LOCTEXT_NAMESPACE
+
diff --git a/Engine/Plugins/Cameras/GameplayCameras/Source/GameplayCameras/Private/GameFramework/GameplayCameraComponent.cpp b/Engine/Plugins/Cameras/GameplayCameras/Source/GameplayCameras/Private/GameFramework/GameplayCameraComponent.cpp
new file mode 100644
--- /dev/null
+++ b/Engine/Plugins/Cameras/GameplayCameras/Source/GameplayCameras/Private/GameFramework/GameplayCameraComponent.cpp
@@ -0,0 +1,157 @@
+// Copyright Epic Games, Inc. All Rights Reserved.
+
+#include "GameFramework/GameplayCameraComponent.h"
+
+#include "Components/StaticMeshComponent.h"
+#include "Core/CameraSystemEvaluator.h"
+#include "Engine/StaticMesh.h"
+#include "Engine/World.h"
+#include "GameFramework/Controller.h"
+#include "GameFramework/GameplayCameraSystemActor.h"
+#include "GameFramework/GameplayCameraSystemComponent.h"
+#include "Logging/MessageLog.h"
+#include "UObject/ConstructorHelpers.h"
+#include "UObject/UnrealType.h"
+
+#include UE_INLINE_GENERATED_CPP_BY_NAME(GameplayCameraComponent)
+
+#define LOCTEXT_NAMESPACE "GameplayCameraComponent"
+
+UGameplayCameraComponent::UGameplayCameraComponent(const FObjectInitializer& ObjectInit)
+	: Super(ObjectInit)
+{
+	PrimaryComponentTick.bCanEverTick = true;
+
+#if WITH_EDITORONLY_DATA
+	if (GIsEditor && !IsRunningCommandlet())
+	{
+		static ConstructorHelpers::FObjectFinder<UStaticMesh> EditorCameraMesh(
+				TEXT("/Engine/EditorMeshes/Camera/SM_CineCam.SM_CineCam"));
+		PreviewMesh = EditorCameraMesh.Object;
+	}
+#endif  // WITH_EDITORONLY_DATA
+}
+
+void UGameplayCameraComponent::ActivateCamera(int32 PlayerIndex)
+{
+	UWorld* World = GetWorld();
+	if (!ensure(World))
+	{
+		return;
+	}
+
+	PlayerIndex = FMath::Max(0, PlayerIndex);
+	if (ensure(PlayerIndex < World->GetNumPlayerControllers()))
+	{
+		FConstPlayerControllerIterator It = World->GetPlayerControllerIterator();
+		It += PlayerIndex;
+		if (APlayerController* PlayerController = It->Get())
+		{
+			ActivateCamera(PlayerController);
+		}
+	}
+}
+
+void UGameplayCameraComponent::ActivateCamera(APlayerController* PlayerController)
+{
+	if (!ensure(PlayerController && PlayerController->PlayerCameraManager))
+	{
+		return;
+	}
+	
+	AGameplayCameraSystemActor* CameraSystem = Cast<AGameplayCameraSystemActor>(PlayerController->PlayerCameraManager->GetViewTarget());
+	if (!ensure(CameraSystem))
+	{
+		return;
+	}
+
+	if (EvaluationContext == nullptr)
+	{
+		EvaluationContext = NewObject<UGameplayCameraComponentEvaluationContext>(this, TEXT("EvaluationContext"));
+		EvaluationContext->Initialize(this);
+	}
+
+	UCameraSystemEvaluator* Evaluator = CameraSystem->GetCameraSystemComponent()->GetCameraSystemEvaluator();
+	Evaluator->PushEvaluationContext(EvaluationContext);
+
+	Activate();
+}
+
+void UGameplayCameraComponent::OnRegister()
+{
+#if WITH_EDITORONLY_DATA
+	if (PreviewMesh && !PreviewMeshComponent)
+	{
+		PreviewMeshComponent = NewObject<UStaticMeshComponent>(this, NAME_None, RF_Transactional | RF_TextExportTransient);
+		PreviewMeshComponent->SetupAttachment(this);
+		PreviewMeshComponent->SetIsVisualizationComponent(true);
+		PreviewMeshComponent->SetStaticMesh(PreviewMesh);
+		PreviewMeshComponent->SetCollisionProfileName(UCollisionProfile::NoCollision_ProfileName);
+		PreviewMeshComponent->bHiddenInGame = true;
+		PreviewMeshComponent->CastShadow = false;
+		PreviewMeshComponent->CreationMethod = CreationMethod;
+		PreviewMeshComponent->RegisterComponentWithWorld(GetWorld());
+	}
+
+	UpdatePreviewMeshTransform();
+#endif	// WITH_EDITORONLY_DATA
+
+	Super::OnRegister();
+}
+
+void UGameplayCameraComponent::TickComponent(float DeltaTime, enum ELevelTick TickType, FActorComponentTickFunction *ThisTickFunction)
+{
+	Super::TickComponent(DeltaTime, TickType, ThisTickFunction);
+
+	if (EvaluationContext)
+	{
+		EvaluationContext->Update(this);
+	}
+}
+
+void UGameplayCameraComponent::OnComponentDestroyed(bool bDestroyingHierarchy)
+{
+	Super::OnComponentDestroyed(bDestroyingHierarchy);
+
+#if WITH_EDITORONLY_DATA
+	if (PreviewMeshComponent)
+	{
+		PreviewMeshComponent->DestroyComponent();
+	}
+#endif  // WITH_EDITORONLY_DATA
+}
+
+#if WITH_EDITORONLY_DATA
+
+void UGameplayCameraComponent::UpdatePreviewMeshTransform()
+{
+	if (PreviewMeshComponent)
+	{
+		// CineCam mesh is wrong, adjust like UCineCameraComponent
+		PreviewMeshComponent->SetRelativeRotation(FRotator(0.f, 90.f, 0.f));
+		PreviewMeshComponent->SetRelativeLocation(FVector(-46.f, 0, -24.f));
+		PreviewMeshComponent->SetRelativeScale3D(FVector::OneVector);
+	}
+}
+
+#endif
+
+UGameplayCameraComponentEvaluationContext::UGameplayCameraComponentEvaluationContext(const FObjectInitializer& ObjectInit)
+	: Super(ObjectInit)
+{
+}
+
+void UGameplayCameraComponentEvaluationContext::Initialize(UGameplayCameraComponent* Owner)
+{
+	CameraAsset = Owner->Camera;
+}
+
+void UGameplayCameraComponentEvaluationContext::Update(UGameplayCameraComponent* Owner)
+{
+	const FTransform& OwnerTransform = Owner->GetComponentTransform();
+	InitialResult.CameraPose.SetTransform(OwnerTransform);
+	InitialResult.bIsValid = true;
+}
+
+#undef LOCTEXT_NAMESPACE
+
diff --git a/Engine/Plugins/Cameras/GameplayCameras/Source/GameplayCameras/Private/GameFramework/GameplayCameraSystemActor.cpp b/Engine/Plugins/Cameras/GameplayCameras/Source/GameplayCameras/Private/GameFramework/GameplayCameraSystemActor.cpp
new file mode 100644
--- /dev/null
+++ b/Engine/Plugins/Cameras/GameplayCameras/Source/GameplayCameras/Private/GameFramework/GameplayCameraSystemActor.cpp
@@ -0,0 +1,62 @@
+// Copyright Epic Games, Inc. All Rights Reserved.
+
+#include "GameFramework/GameplayCameraSystemActor.h"
+
+#include "Engine/EngineTypes.h"
+#include "Engine/World.h"
+#include "GameFramework/GameplayCameraSystemComponent.h"
+#include "GameFramework/PlayerController.h"
+#include "Kismet/GameplayStatics.h"
+
+#include UE_INLINE_GENERATED_CPP_BY_NAME(GameplayCameraSystemActor)
+
+#define LOCTEXT_NAMESPACE "GameplayCameraSystemActor"
+
+AGameplayCameraSystemActor::AGameplayCameraSystemActor(const FObjectInitializer& ObjectInit)
+	: Super(ObjectInit)
+{
+	CameraSystemComponent = CreateDefaultSubobject<UGameplayCameraSystemComponent>(TEXT("CameraSystemComponent"));
+	RootComponent = CameraSystemComponent;
+}
+
+void AGameplayCameraSystemActor::BeginPlay()
+{
+	Super::BeginPlay();
+
+	if (AutoActivateForPlayer != EAutoReceiveInput::Disabled && GetNetMode() != NM_Client)
+	{
+		const int32 PlayerIndex = AutoActivateForPlayer.GetIntValue() - 1;
+		ActivateForPlayer(PlayerIndex);
+	}
+}
+
+void AGameplayCameraSystemActor::BecomeViewTarget(APlayerController* PC)
+{
+	Super::BecomeViewTarget(PC);
+
+	CameraSystemComponent->OnBecomeViewTarget();
+}
+
+void AGameplayCameraSystemActor::CalcCamera(float DeltaTime, struct FMinimalViewInfo& OutResult)
+{
+	CameraSystemComponent->GetCameraView(DeltaTime, OutResult);
+}
+
+void AGameplayCameraSystemActor::EndViewTarget(APlayerController* PC)
+{
+	CameraSystemComponent->OnEndViewTarget();
+
+	Super::EndViewTarget(PC);
+}
+
+void AGameplayCameraSystemActor::ActivateForPlayer(int32 PlayerIndex)
+{
+	APlayerController* PC = UGameplayStatics::GetPlayerController(this, PlayerIndex);
+	if (PC)
+	{
+		PC->SetViewTarget(this);
+	}
+}
+
+#undef LOCTEXT_NAMESPACE
+
diff --git a/Engine/Plugins/Cameras/GameplayCameras/Source/GameplayCameras/Private/GameFramework/GameplayCameraSystemComponent.cpp b/Engine/Plugins/Cameras/GameplayCameras/Source/GameplayCameras/Private/GameFramework/GameplayCameraSystemComponent.cpp
new file mode 100644
--- /dev/null
+++ b/Engine/Plugins/Cameras/GameplayCameras/Source/GameplayCameras/Private/GameFramework/GameplayCameraSystemComponent.cpp
@@ -0,0 +1,104 @@
+// Copyright Epic Games, Inc. All Rights Reserved.
+
+#include "GameFramework/GameplayCameraSystemComponent.h"
+
+#include "Components/StaticMeshComponent.h"
+#include "Engine/StaticMesh.h"
+#include "Engine/World.h"
+#include "GameFramework/Controller.h"
+#include "Logging/MessageLog.h"
+#include "UObject/ConstructorHelpers.h"
+#include "UObject/UnrealType.h"
+
+#include UE_INLINE_GENERATED_CPP_BY_NAME(GameplayCameraSystemComponent)
+
+#define LOCTEXT_NAMESPACE "GameplayCameraSystemComponent"
+
+UGameplayCameraSystemComponent::UGameplayCameraSystemComponent(const FObjectInitializer& ObjectInit)
+	: Super(ObjectInit)
+{
+#if WITH_EDITORONLY_DATA
+	bTickInEditor = true;
+	PrimaryComponentTick.bCanEverTick = true;
+
+	if (GIsEditor && !IsRunningCommandlet())
+	{
+		static ConstructorHelpers::FObjectFinder<UStaticMesh> EditorCameraMesh(
+				TEXT("/Engine/EditorMeshes/Camera/SM_CineCam.SM_CineCam"));
+		PreviewMesh = EditorCameraMesh.Object;
+	}
+#endif  // WITH_EDITORONLY_DATA
+
+	Evaluator = ObjectInit.CreateDefaultSubobject<UCameraSystemEvaluator>(this, "CameraSystemEvaluator");
+}
+
+void UGameplayCameraSystemComponent::GetCameraView(float DeltaTime, FMinimalViewInfo& DesiredView)
+{
+	FCameraSystemEvaluationUpdateParams UpdateParams;
+	UpdateParams.DeltaTime = DeltaTime;
+	Evaluator->Update(UpdateParams);
+
+	Evaluator->GetEvaluatedCameraView(DesiredView);
+}
+
+void UGameplayCameraSystemComponent::OnRegister()
+{
+#if WITH_EDITORONLY_DATA
+	if (PreviewMesh && !PreviewMeshComponent)
+	{
+		PreviewMeshComponent = NewObject<UStaticMeshComponent>(this, NAME_None, RF_Transactional | RF_TextExportTransient);
+		PreviewMeshComponent->SetupAttachment(this);
+		PreviewMeshComponent->SetIsVisualizationComponent(true);
+		PreviewMeshComponent->SetStaticMesh(PreviewMesh);
+		PreviewMeshComponent->SetCollisionProfileName(UCollisionProfile::NoCollision_ProfileName);
+		PreviewMeshComponent->bHiddenInGame = true;
+		PreviewMeshComponent->CastShadow = false;
+		PreviewMeshComponent->CreationMethod = CreationMethod;
+		PreviewMeshComponent->RegisterComponentWithWorld(GetWorld());
+	}
+#endif	// WITH_EDITORONLY_DATA
+
+	Super::OnRegister();
+}
+
+void UGameplayCameraSystemComponent::TickComponent(float DeltaTime, ELevelTick TickType, FActorComponentTickFunction* ThisTickFunction)
+{
+	Super::TickComponent(DeltaTime, TickType, ThisTickFunction);
+}
+
+void UGameplayCameraSystemComponent::OnComponentDestroyed(bool bDestroyingHierarchy)
+{
+	Super::OnComponentDestroyed(bDestroyingHierarchy);
+
+#if WITH_EDITORONLY_DATA
+	if (PreviewMeshComponent)
+	{
+		PreviewMeshComponent->DestroyComponent();
+	}
+#endif  // WITH_EDITORONLY_DATA
+}
+
+#if WITH_EDITOR
+
+bool UGameplayCameraSystemComponent::GetEditorPreviewInfo(float DeltaTime, FMinimalViewInfo& ViewOut)
+{
+	const bool bIsCameraSystemActive = IsActive();
+	if (bIsCameraSystemActive)
+	{
+		GetCameraView(DeltaTime, ViewOut);
+	}
+	return bIsCameraSystemActive;
+}
+
+#endif  // WITH_EDITOR
+
+void UGameplayCameraSystemComponent::OnBecomeViewTarget()
+{
+}
+
+void UGameplayCameraSystemComponent::OnEndViewTarget()
+{
+}
+
+#undef LOCTEXT_NAMESPACE
+
diff --git a/Engine/Plugins/Cameras/GameplayCameras/Source/GameplayCameras/Private/GameplayCamerasModule.cpp b/Engine/Plugins/Cameras/GameplayCameras/Source/GameplayCameras/Private/GameplayCamerasModule.cpp
--- a/Engine/Plugins/Cameras/GameplayCameras/Source/GameplayCameras/Private/GameplayCamerasModule.cpp
+++ b/Engine/Plugins/Cameras/GameplayCameras/Source/GameplayCameras/Private/GameplayCamerasModule.cpp
@@ -1,6 +1,7 @@
 // Copyright Epic Games, Inc. All Rights Reserved.
 
-#include "GameplayCamerasModule.h"
+#include "IGameplayCamerasModule.h"
+
 #include "Camera/CameraModularFeature.h"
 #include "CameraAnimationCameraModifier.h"
 #include "CameraAnimationSequencePlayer.h"
@@ -9,7 +10,7 @@
 
 IGameplayCamerasModule& IGameplayCamerasModule::Get()
 {
-	return FModuleManager::LoadModuleChecked<IGameplayCamerasModule>("Camera");
+	return FModuleManager::LoadModuleChecked<IGameplayCamerasModule>("GameplayCameras");
 }
 
 class FGameplayCamerasModule : public IGameplayCamerasModule
@@ -35,6 +36,19 @@ public:
 		}
 	}
 
+	// IGameplayCamerasModule interface
+#if WITH_EDITOR
+	virtual TSharedPtr<IGameplayCamerasLiveEditManager> GetLiveEditManager() const override
+	{
+		return LiveEditManager;
+	}
+
+	virtual void SetLiveEditManager(TSharedPtr<IGameplayCamerasLiveEditManager> InLiveEditManager) override
+	{
+		LiveEditManager = InLiveEditManager;
+	}
+#endif
+
 private:
 	class FCameraModularFeature : public ICameraModularFeature
 	{
@@ -46,6 +60,10 @@ private:
 	};
 
 	TSharedPtr<FCameraModularFeature> CameraModularFeature;
+
+#if WITH_EDITOR
+	TSharedPtr<IGameplayCamerasLiveEditManager> LiveEditManager;
+#endif
 };
 
 IMPLEMENT_MODULE(FGameplayCamerasModule, GameplayCameras);
diff --git a/Engine/Plugins/Cameras/GameplayCameras/Source/GameplayCameras/Private/Nodes/Blends/LinearBlendCameraNode.cpp b/Engine/Plugins/Cameras/GameplayCameras/Source/GameplayCameras/Private/Nodes/Blends/LinearBlendCameraNode.cpp
new file mode 100644
--- /dev/null
+++ b/Engine/Plugins/Cameras/GameplayCameras/Source/GameplayCameras/Private/Nodes/Blends/LinearBlendCameraNode.cpp
@@ -0,0 +1,11 @@
+// Copyright Epic Games, Inc. All Rights Reserved.
+
+#include "Nodes/Blends/LinearBlendCameraNode.h"
+
+#include UE_INLINE_GENERATED_CPP_BY_NAME(LinearBlendCameraNode)
+
+void ULinearBlendCameraNode::OnComputeBlendFactor(const FCameraNodeRunParams& Params, FSimpleBlendCameraNodeRunResult& OutResult)
+{
+	OutResult.BlendFactor = FMath::Lerp(0.f, 1.f, GetTimeFactor());
+}
+
diff --git a/Engine/Plugins/Cameras/GameplayCameras/Source/GameplayCameras/Private/Nodes/Blends/PopBlendCameraNode.cpp b/Engine/Plugins/Cameras/GameplayCameras/Source/GameplayCameras/Private/Nodes/Blends/PopBlendCameraNode.cpp
new file mode 100644
--- /dev/null
+++ b/Engine/Plugins/Cameras/GameplayCameras/Source/GameplayCameras/Private/Nodes/Blends/PopBlendCameraNode.cpp
@@ -0,0 +1,25 @@
+// Copyright Epic Games, Inc. All Rights Reserved.
+
+#include "Nodes/Blends/PopBlendCameraNode.h"
+
+#include UE_INLINE_GENERATED_CPP_BY_NAME(PopBlendCameraNode)
+
+void UPopBlendCameraNode::OnRun(const FCameraNodeRunParams& Params, FCameraNodeRunResult& OutResult)
+{
+}
+
+void UPopBlendCameraNode::OnBlendResults(const FCameraNodeBlendParams& Params, FCameraNodeBlendResult& OutResult)
+{
+	const FCameraNodeRunResult& ChildResult(Params.ChildResult);
+	FCameraNodeRunResult& BlendedResult(OutResult.BlendedResult);
+	
+	BlendedResult.CameraPose.OverrideChanged(ChildResult.CameraPose);
+	if (ChildResult.bIsCameraCut || Params.ChildParams.bIsFirstFrame)
+	{
+		BlendedResult.bIsCameraCut = true;
+	}
+
+	OutResult.bIsBlendFull = true;
+	OutResult.bIsBlendFinished = true;
+}
+
diff --git a/Engine/Plugins/Cameras/GameplayCameras/Source/GameplayCameras/Private/Nodes/Blends/SimpleBlendCameraNode.cpp b/Engine/Plugins/Cameras/GameplayCameras/Source/GameplayCameras/Private/Nodes/Blends/SimpleBlendCameraNode.cpp
new file mode 100644
--- /dev/null
+++ b/Engine/Plugins/Cameras/GameplayCameras/Source/GameplayCameras/Private/Nodes/Blends/SimpleBlendCameraNode.cpp
@@ -0,0 +1,48 @@
+// Copyright Epic Games, Inc. All Rights Reserved.
+
+#include "Nodes/Blends/SimpleBlendCameraNode.h"
+
+#include UE_INLINE_GENERATED_CPP_BY_NAME(SimpleBlendCameraNode)
+
+void USimpleBlendCameraNode::OnRun(const FCameraNodeRunParams& Params, FCameraNodeRunResult& OutResult)
+{
+	FSimpleBlendCameraNodeRunResult FactorResult;
+	OnComputeBlendFactor(Params, FactorResult);
+	BlendFactor = FactorResult.BlendFactor;
+}
+
+void USimpleBlendCameraNode::OnBlendResults(const FCameraNodeBlendParams& Params, FCameraNodeBlendResult& OutResult)
+{
+	const FCameraNodeRunResult& ChildResult(Params.ChildResult);
+	FCameraNodeRunResult& BlendedResult(OutResult.BlendedResult);
+
+	BlendedResult.CameraPose.LerpChanged(
+			ChildResult.CameraPose, 
+			BlendFactor,
+			ChildResult.CameraPose.GetChangedFlags(), // or is it BlendedResult flags?
+			false,
+			BlendedResult.CameraPose.GetChangedFlags());
+
+	// If we have even a fraction of a camera cut, we need to make the
+	// whole result into a camera cut.
+	if (BlendFactor > 0.f && ChildResult.bIsCameraCut)
+	{
+		BlendedResult.bIsCameraCut = true;
+	}
+
+	OutResult.bIsBlendFull = BlendFactor >= 1.f;
+	OutResult.bIsBlendFinished = bIsBlendFinished;
+}
+
+void USimpleFixedTimeBlendCameraNode::OnRun(const FCameraNodeRunParams& Params, FCameraNodeRunResult& OutResult)
+{
+	CurrentTime += Params.DeltaTime;
+	if (CurrentTime >= BlendTime)
+	{
+		CurrentTime = BlendTime;
+		SetBlendFinished();
+	}
+
+	Super::OnRun(Params, OutResult);
+}
+
diff --git a/Engine/Plugins/Cameras/GameplayCameras/Source/GameplayCameras/Private/Nodes/Blends/SmoothBlendCameraNode.cpp b/Engine/Plugins/Cameras/GameplayCameras/Source/GameplayCameras/Private/Nodes/Blends/SmoothBlendCameraNode.cpp
new file mode 100644
--- /dev/null
+++ b/Engine/Plugins/Cameras/GameplayCameras/Source/GameplayCameras/Private/Nodes/Blends/SmoothBlendCameraNode.cpp
@@ -0,0 +1,27 @@
+// Copyright Epic Games, Inc. All Rights Reserved.
+
+#include "Nodes/Blends/SmoothBlendCameraNode.h"
+
+#include "Math/Interpolation.h"
+
+#include UE_INLINE_GENERATED_CPP_BY_NAME(SmoothBlendCameraNode)
+
+void USmoothBlendCameraNode::OnComputeBlendFactor(const FCameraNodeRunParams& Params, FSimpleBlendCameraNodeRunResult& OutResult)
+{
+	using namespace UE::Cameras;
+
+	const float t = GetTimeFactor();
+	switch (BlendType)
+	{
+		case ESmoothCameraBlendType::SmoothStep:
+			OutResult.BlendFactor = SmoothStep(t);
+			break;
+		case ESmoothCameraBlendType::SmootherStep:
+			OutResult.BlendFactor = SmootherStep(t);
+			break;
+		default:
+			OutResult.BlendFactor = 1.f;
+			break;
+	}
+}
+
diff --git a/Engine/Plugins/Cameras/GameplayCameras/Source/GameplayCameras/Private/Nodes/Common/ArrayCameraNode.cpp b/Engine/Plugins/Cameras/GameplayCameras/Source/GameplayCameras/Private/Nodes/Common/ArrayCameraNode.cpp
new file mode 100644
--- /dev/null
+++ b/Engine/Plugins/Cameras/GameplayCameras/Source/GameplayCameras/Private/Nodes/Common/ArrayCameraNode.cpp
@@ -0,0 +1,22 @@
+// Copyright Epic Games, Inc. All Rights Reserved.
+
+#include "Nodes/Common/ArrayCameraNode.h"
+
+#include UE_INLINE_GENERATED_CPP_BY_NAME(ArrayCameraNode)
+
+FCameraNodeChildrenView UArrayCameraNode::OnGetChildren()
+{
+	return FCameraNodeChildrenView(Children);
+}
+
+void UArrayCameraNode::OnRun(const FCameraNodeRunParams& Params, FCameraNodeRunResult& OutResult)
+{
+	for (TObjectPtr<UCameraNode> Child : Children)
+	{
+		if (Child)
+		{
+			Child->Run(Params, OutResult);
+		}
+	}
+}
+
diff --git a/Engine/Plugins/Cameras/GameplayCameras/Source/GameplayCameras/Private/Nodes/Common/OffsetCameraNode.cpp b/Engine/Plugins/Cameras/GameplayCameras/Source/GameplayCameras/Private/Nodes/Common/OffsetCameraNode.cpp
new file mode 100644
--- /dev/null
+++ b/Engine/Plugins/Cameras/GameplayCameras/Source/GameplayCameras/Private/Nodes/Common/OffsetCameraNode.cpp
@@ -0,0 +1,15 @@
+// Copyright Epic Games, Inc. All Rights Reserved.
+
+#include "Nodes/Common/OffsetCameraNode.h"
+
+#include UE_INLINE_GENERATED_CPP_BY_NAME(OffsetCameraNode)
+
+void UOffsetCameraNode::OnRun(const FCameraNodeRunParams& Params, FCameraNodeRunResult& OutResult)
+{
+	const FRotator3d Rotation = OutResult.CameraPose.GetRotation();
+	const FVector3d LocalOffset = Rotation.RotateVector(Offset);
+
+	const FVector3d Location = OutResult.CameraPose.GetLocation();
+	OutResult.CameraPose.SetLocation(Location + LocalOffset);
+}
+
diff --git a/Engine/Plugins/Cameras/GameplayCameras/Source/GameplayCameras/Public/Core/BlendCameraNode.h b/Engine/Plugins/Cameras/GameplayCameras/Source/GameplayCameras/Public/Core/BlendCameraNode.h
new file mode 100644
--- /dev/null
+++ b/Engine/Plugins/Cameras/GameplayCameras/Source/GameplayCameras/Public/Core/BlendCameraNode.h
@@ -0,0 +1,64 @@
+// Copyright Epic Games, Inc. All Rights Reserved.
+
+#pragma once
+
+#include "Core/CameraNode.h"
+
+#include "BlendCameraNode.generated.h"
+
+/**
+ * Parameter struct for blending camera node tree results.
+ */
+struct FCameraNodeBlendParams
+{
+	FCameraNodeBlendParams(
+			const FCameraNodeRunParams& InChildParams,
+			const FCameraNodeRunResult& InChildResult)
+		: ChildParams(InChildParams)
+		, ChildResult(InChildResult)
+	{}
+
+	/** The parameters that the blend received during the evaluation. */
+	const FCameraNodeRunParams& ChildParams;
+	/** The result that the blend should apply over another result. */
+	const FCameraNodeRunResult& ChildResult;
+};
+
+/**
+ * Result struct for blending camera node tree results.
+ */
+struct FCameraNodeBlendResult
+{
+	FCameraNodeBlendResult(FCameraNodeRunResult& InBlendedResult)
+		: BlendedResult(InBlendedResult)
+	{}
+
+	/** The result upon which another result should be blended. */
+	FCameraNodeRunResult& BlendedResult;
+
+	/** Whether the blend has reached 100%. */
+	bool bIsBlendFull = false;
+
+	/** Whether the blend is finished. */
+	bool bIsBlendFinished = false;
+};
+
+/**
+ * Base class for blend camera nodes.
+ */
+UCLASS(MinimalAPI)
+class UBlendCameraNode : public UCameraNode
+{
+	GENERATED_BODY()
+
+public:
+
+	/** Blend the result of a camera node tree over another result. */
+	void BlendResults(const FCameraNodeBlendParams& Params, FCameraNodeBlendResult& OutResult);
+
+protected:
+
+	/** Blend the result of a camera node tree over another result. */
+	virtual void OnBlendResults(const FCameraNodeBlendParams& Params, FCameraNodeBlendResult& OutResult) {}
+};
+
diff --git a/Engine/Plugins/Cameras/GameplayCameras/Source/GameplayCameras/Public/Core/BlendStackCameraNode.h b/Engine/Plugins/Cameras/GameplayCameras/Source/GameplayCameras/Public/Core/BlendStackCameraNode.h
new file mode 100644
--- /dev/null
+++ b/Engine/Plugins/Cameras/GameplayCameras/Source/GameplayCameras/Public/Core/BlendStackCameraNode.h
@@ -0,0 +1,85 @@
+// Copyright Epic Games, Inc. All Rights Reserved.
+
+#pragma once
+
+#include "Core/CameraNode.h"
+
+#include "BlendStackCameraNode.generated.h"
+
+class UBlendStackRootCameraNode;
+class UCameraAsset;
+class UCameraEvaluationContext;
+class UCameraMode;
+struct FCameraModeTransition;
+
+/**
+ * Parameter structure for pushing a camera mode onto a blend stack.
+ */
+struct FBlendStackCameraPushParams
+{
+	/** The evaluation context within which a camera mode's node tree should run. */
+	TWeakObjectPtr<const UCameraEvaluationContext> EvaluationContext;
+
+	/** The source camera mode asset to instantiate and push on the blend stack. */
+	TObjectPtr<const UCameraMode> CameraMode;
+};
+
+/**
+ * A blend stack implemented as a camera node.
+ */
+UCLASS(MinimalAPI)
+class UBlendStackCameraNode : public UCameraNode
+{
+	GENERATED_BODY()
+
+public:
+
+	/** Push a new camera mode onto the blend stack. */
+	void Push(const FBlendStackCameraPushParams& Params);
+
+public:
+
+	static void AddReferencedObjects(UObject* InThis, FReferenceCollector& Collector);
+
+protected:
+
+	// UCameraNode interface
+	virtual FCameraNodeChildrenView OnGetChildren() override;
+	virtual void OnRun(const FCameraNodeRunParams& Params, FCameraNodeRunResult& OutResult) override;
+
+	const FCameraModeTransition* FindTransition(const FBlendStackCameraPushParams& Params) const;
+	const FCameraModeTransition* FindTransition(
+			TArrayView<const FCameraModeTransition> Transitions, 
+			const UCameraMode* FromCameraMode, const UCameraAsset* FromCameraAsset, bool bFromFrozen,
+			const UCameraMode* ToCameraMode, const UCameraAsset* ToCameraAsset) const;
+
+public:
+
+	/** 
+	 * Whether to automatically pop camera modes out of the stack when another mode
+	 * has reached 100% blend above them.
+	 */
+	UPROPERTY()
+	bool bAutoPop = true;
+
+	/**
+	 * Whether to blend-in the first camera mode when the stack is previously empty.
+	 */
+	UPROPERTY()
+	bool bBlendFirstCameraMode = false;
+
+protected:
+
+	struct FCameraModeEntry
+	{
+		TWeakObjectPtr<const UCameraEvaluationContext> EvaluationContext;
+		TObjectPtr<const UCameraMode> OriginalCameraMode;
+		TObjectPtr<UBlendStackRootCameraNode> RootNode;
+		FCameraNodeRunResult Result;
+		bool bIsFirstFrame = false;
+		bool bIsFrozen = false;
+	};
+
+	TArray<FCameraModeEntry> Entries;
+};
+
diff --git a/Engine/Plugins/Cameras/GameplayCameras/Source/GameplayCameras/Public/Core/CameraAsset.h b/Engine/Plugins/Cameras/GameplayCameras/Source/GameplayCameras/Public/Core/CameraAsset.h
new file mode 100644
--- /dev/null
+++ b/Engine/Plugins/Cameras/GameplayCameras/Source/GameplayCameras/Public/Core/CameraAsset.h
@@ -0,0 +1,36 @@
+// Copyright Epic Games, Inc. All Rights Reserved.
+
+#pragma once
+
+#include "Core/CameraModeTransition.h"
+#include "CoreTypes.h"
+#include "UObject/ObjectPtr.h"
+
+#include "CameraAsset.generated.h"
+
+class UCameraDirector;
+
+/**
+ * A complete camera asset.
+ */
+UCLASS(MinimalAPI)
+class UCameraAsset : public UObject
+{
+	GENERATED_BODY()
+
+public:
+
+	/** The camera director to use in this camera. */
+	UPROPERTY(Instanced, EditAnywhere, Category=Director)
+	TObjectPtr<UCameraDirector> CameraDirector;
+
+	/** A list of default enter transitions for all the camera modes in this asset. */
+	UPROPERTY(EditAnywhere, Category=Blending)
+	TArray<FCameraModeTransition> EnterTransitions;
+
+	/** A list of default exit transitions for all the camera modes in this asset. */
+	UPROPERTY(EditAnywhere, Category=Blending)
+	TArray<FCameraModeTransition> ExitTransitions;
+};
+
+
diff --git a/Engine/Plugins/Cameras/GameplayCameras/Source/GameplayCameras/Public/Core/CameraDirector.h b/Engine/Plugins/Cameras/GameplayCameras/Source/GameplayCameras/Public/Core/CameraDirector.h
new file mode 100644
--- /dev/null
+++ b/Engine/Plugins/Cameras/GameplayCameras/Source/GameplayCameras/Public/Core/CameraDirector.h
@@ -0,0 +1,52 @@
+// Copyright Epic Games, Inc. All Rights Reserved.
+
+#pragma once
+
+#include "CoreTypes.h"
+#include "UObject/Object.h"
+
+#include "CameraDirector.generated.h"
+
+class UCameraEvaluationContext;
+class UCameraMode;
+
+/**
+ * Parameter structure for running a camera director.
+ */
+struct FCameraDirectorRunParams
+{
+	/** Time interval for the update. */
+	float DeltaTime = 0.f;
+
+	/** The context in which this director runs. */
+	const UCameraEvaluationContext* OwnerContext = nullptr;
+};
+
+/**
+ * Result struct for running a camera director.
+ */
+struct FCameraDirectorRunResult
+{
+	/** The camera mode(s) that the director says should be active this frame. */
+	TArray<TObjectPtr<const UCameraMode>, TInlineAllocator<2>> ActiveCameraModes;
+};
+
+/**
+ * Base class for a camera director.
+ */
+UCLASS(Abstract, DefaultToInstanced, MinimalAPI)
+class UCameraDirector : public UObject
+{
+	GENERATED_BODY()
+
+public:
+	
+	/** Runs the camera director to determine what camera mode(s) should be active this frame. */
+	void Run(const FCameraDirectorRunParams& Params, FCameraDirectorRunResult& OutResult);
+
+protected:
+
+	/** Runs the camera director to determine what camera mode(s) should be active this frame. */
+	virtual void OnRun(const FCameraDirectorRunParams& Params, FCameraDirectorRunResult& OutResult) {}
+};
+
diff --git a/Engine/Plugins/Cameras/GameplayCameras/Source/GameplayCameras/Public/Core/CameraEvaluationContext.h b/Engine/Plugins/Cameras/GameplayCameras/Source/GameplayCameras/Public/Core/CameraEvaluationContext.h
new file mode 100644
--- /dev/null
+++ b/Engine/Plugins/Cameras/GameplayCameras/Source/GameplayCameras/Public/Core/CameraEvaluationContext.h
@@ -0,0 +1,40 @@
+// Copyright Epic Games, Inc. All Rights Reserved.
+
+#pragma once
+
+#include "CoreTypes.h"
+#include "UObject/ObjectPtr.h"
+#include "Core/CameraNode.h"
+
+#include "CameraEvaluationContext.generated.h"
+
+class UCameraAsset;
+
+/**
+ * Base class for providing a context to running camera modes.
+ */
+UCLASS(MinimalAPI)
+class UCameraEvaluationContext : public UObject
+{
+	GENERATED_BODY()
+
+public:
+
+	UCameraEvaluationContext(const FObjectInitializer& ObjectInit);
+
+	/** Gets the camera asset that is hosted in this context. */
+	UCameraAsset* GetCameraAsset() const { return CameraAsset; }
+
+	/** Gets the initial evaluation result for all camera modes in this context. */
+	const FCameraNodeRunResult& GetInitialResult() const { return InitialResult; }
+
+protected:
+
+	/** The camera asset hosted in this context. */
+	UPROPERTY()
+	TObjectPtr<UCameraAsset> CameraAsset;
+
+	/** The initial result for all camera modes in this context. */
+	FCameraNodeRunResult InitialResult;
+};
+
diff --git a/Engine/Plugins/Cameras/GameplayCameras/Source/GameplayCameras/Public/Core/CameraEvaluationContextStack.h b/Engine/Plugins/Cameras/GameplayCameras/Source/GameplayCameras/Public/Core/CameraEvaluationContextStack.h
new file mode 100644
--- /dev/null
+++ b/Engine/Plugins/Cameras/GameplayCameras/Source/GameplayCameras/Public/Core/CameraEvaluationContextStack.h
@@ -0,0 +1,66 @@
+// Copyright Epic Games, Inc. All Rights Reserved.
+
+#pragma once
+
+#include "Containers/Array.h"
+#include "UObject/WeakObjectPtr.h"
+
+class UCameraAsset;
+class UCameraDirector;
+class UCameraEvaluationContext;
+class UCameraSystemEvaluator;
+
+/** Information about a running camera evaluation context. */
+struct FCameraEvaluationContextInfo
+{
+	/** The evaluation context. */
+	UCameraEvaluationContext* EvaluationContext = nullptr;
+
+	/** The instantiated camera director running in this context. */
+	UCameraDirector* CameraDirector = nullptr;
+
+	/** Returns whether this structure has a valid context and director. */
+	bool IsValid() const { return EvaluationContext && CameraDirector; }
+};
+
+/**
+ * A simple stack of evaluation contexts. The top one is the active one.
+ */
+struct FCameraEvaluationContextStack
+{
+public:
+
+	/** Gets the active (top) context. */
+	FCameraEvaluationContextInfo GetActiveContext() const;
+
+	/** Returns whether the given context exists in the stack. */
+	bool HasContext(UCameraEvaluationContext* Context) const;
+
+	/** Push a new context on the stack and instantiate its director. */
+	void PushContext(UCameraEvaluationContext* Context);
+
+	/** Remove an existing context from the stack. */
+	bool RemoveContext(UCameraEvaluationContext* Context);
+
+	/** Pop the active (top) context. */
+	void PopContext();
+
+public:
+
+	// Internal API
+	void SetInstantiationOuter(UObject* InInstantiationOuter);
+	void AddReferencedObjects(FReferenceCollector& Collector);
+
+private:
+
+	struct FContextEntry
+	{
+		TWeakObjectPtr<UCameraEvaluationContext> WeakContext;
+		TObjectPtr<UCameraDirector> CameraDirector;
+	};
+
+	TArray<FContextEntry> Entries;
+
+	TWeakObjectPtr<> InstantiationOuter;
+};
+
diff --git a/Engine/Plugins/Cameras/GameplayCameras/Source/GameplayCameras/Public/Core/CameraMode.h b/Engine/Plugins/Cameras/GameplayCameras/Source/GameplayCameras/Public/Core/CameraMode.h
new file mode 100644
--- /dev/null
+++ b/Engine/Plugins/Cameras/GameplayCameras/Source/GameplayCameras/Public/Core/CameraMode.h
@@ -0,0 +1,36 @@
+// Copyright Epic Games, Inc. All Rights Reserved.
+
+#pragma once
+
+#include "Core/CameraModeTransition.h"
+#include "CoreTypes.h"
+#include "UObject/ObjectPtr.h"
+
+#include "CameraMode.generated.h"
+
+class UCameraNode;
+
+/**
+ * A camera mode asset, which runs a hierarchy of camera nodes to drive 
+ * the behavior of a camera.
+ */
+UCLASS(MinimalAPI)
+class UCameraMode : public UObject
+{
+	GENERATED_BODY()
+
+public:
+
+	/** Root camera node. */
+	UPROPERTY(EditAnywhere, Instanced, Category=Common)
+	TObjectPtr<UCameraNode> RootNode;
+
+	/** List of enter transitions for this camera mode. */
+	UPROPERTY(EditAnywhere, Category=Blending)
+	TArray<FCameraModeTransition> EnterTransitions;
+
+	/** List of exist transitions for this camera mode. */
+	UPROPERTY(EditAnywhere, Category=Blending)
+	TArray<FCameraModeTransition> ExitTransitions;
+};
+
diff --git a/Engine/Plugins/Cameras/GameplayCameras/Source/GameplayCameras/Public/Core/CameraModeTransition.h b/Engine/Plugins/Cameras/GameplayCameras/Source/GameplayCameras/Public/Core/CameraModeTransition.h
new file mode 100644
--- /dev/null
+++ b/Engine/Plugins/Cameras/GameplayCameras/Source/GameplayCameras/Public/Core/CameraModeTransition.h
@@ -0,0 +1,66 @@
+// Copyright Epic Games, Inc. All Rights Reserved.
+
+#pragma once
+
+#include "Containers/Array.h"
+#include "CoreTypes.h"
+#include "UObject/ObjectPtr.h"
+
+#include "CameraModeTransition.generated.h"
+
+class UBlendCameraNode;
+class UCameraAsset;
+class UCameraMode;
+
+/**
+ * Parameter structure for camera transitions.
+ */
+struct FCameraModeTransitionConditionMatchParams
+{
+	/** The previous camera mode. */
+	const UCameraMode* FromCameraMode = nullptr;
+	/** The previous camera asset. */
+	const UCameraAsset* FromCameraAsset = nullptr;
+
+	/** The next camera mode. */
+	const UCameraMode* ToCameraMode = nullptr;
+	/** The next camera asset. */
+	const UCameraAsset* ToCameraAsset = nullptr;
+};
+
+/**
+ * Base class for a camera transition condition.
+ */
+UCLASS(Abstract, DefaultToInstanced, MinimalAPI)
+class UCameraModeTransitionCondition : public UObject
+{
+	GENERATED_BODY()
+
+public:
+
+	/** Evaluates whether this transition should be used. */
+	bool TransitionMatches(const FCameraModeTransitionConditionMatchParams& Params) const;
+
+protected:
+
+	/** Evaluates whether this transition should be used. */
+	virtual bool OnTransitionMatches(const FCameraModeTransitionConditionMatchParams& Params) const { return false; }
+};
+
+/**
+ * A camera transition.
+ */
+USTRUCT()
+struct FCameraModeTransition
+{
+	GENERATED_BODY()
+
+	/** The list of conditions that must pass for this transition to be used. */
+	UPROPERTY(EditAnywhere, Instanced, Category=Common)
+	TArray<TObjectPtr<UCameraModeTransitionCondition>> Conditions;
+
+	/** The blend to use to blend a given camera mode in or out. */
+	UPROPERTY(EditAnywhere, Instanced, Category=Common)
+	TObjectPtr<UBlendCameraNode> Blend;
+};
+
diff --git a/Engine/Plugins/Cameras/GameplayCameras/Source/GameplayCameras/Public/Core/CameraNode.h b/Engine/Plugins/Cameras/GameplayCameras/Source/GameplayCameras/Public/Core/CameraNode.h
new file mode 100644
--- /dev/null
+++ b/Engine/Plugins/Cameras/GameplayCameras/Source/GameplayCameras/Public/Core/CameraNode.h
@@ -0,0 +1,88 @@
+// Copyright Epic Games, Inc. All Rights Reserved.
+
+#pragma once
+
+#include "Core/CameraPose.h"
+#include "Core/CameraNodeChildrenView.h"
+#include "CoreTypes.h"
+#include "UObject/Object.h"
+
+#include "CameraNode.generated.h"
+
+/**
+ * Parameter structure for running a camera node.
+ */
+struct FCameraNodeRunParams
+{
+	/** The time interval for the evaluation. */
+	float DeltaTime = 0.f;
+	/** Whether this is the first evaluation of this camera node hierarchy. */
+	bool bIsFirstFrame = false;
+};
+
+/**
+ * Input/output result structure for running a camera node.
+ */
+struct FCameraNodeRunResult
+{
+	/** The camera pose. */
+	FCameraPose CameraPose;
+	/** Whether the current frame is a camera cut. */
+	bool bIsCameraCut = false;
+	/** Whether this result is valid. */
+	bool bIsValid = false;
+
+	/** Reset this result to its default (non-valid) state. */
+	void Reset();
+};
+
+/**
+ * The base class for a camera node.
+ */
+UCLASS(Abstract, DefaultToInstanced, EditInlineNew, MinimalAPI)
+class UCameraNode : public UObject
+{
+	GENERATED_BODY()
+
+public:
+
+	/** Get the list of children under this node. */
+	FCameraNodeChildrenView GetChildren();
+
+	/** Run this node. */
+	void Run(const FCameraNodeRunParams& Params, FCameraNodeRunResult& OutResult);
+
+public:
+
+	// Internal API
+
+	/** Whether this node was instantiated from a source node. */
+	bool GetIsInstantiated() const { return bIsInstantiated; }
+	/** Flag this node as being instantiated from a source node. */
+	void SetIsInstantiated(bool bInValue = true) { bIsInstantiated = bInValue; }
+
+protected:
+
+	/** Get the list of children under this node. */
+	virtual FCameraNodeChildrenView OnGetChildren() { return FCameraNodeChildrenView(); }
+
+	/** Run this node. */
+	virtual void OnRun(const FCameraNodeRunParams& Params, FCameraNodeRunResult& OutResult) {}
+
+#if WITH_EDITOR
+
+	virtual void PostEditChangeProperty(struct FPropertyChangedEvent& PropertyChangedEvent) override;
+
+#endif
+
+public:
+
+	/** Specifies whether this node is enabled. */
+	UPROPERTY(EditAnywhere, Category=Common)
+	bool bIsEnabled = true;
+
+private:
+
+	bool bIsInstantiated = false;
+};
+
diff --git a/Engine/Plugins/Cameras/GameplayCameras/Source/GameplayCameras/Public/Core/CameraNodeChildrenView.h b/Engine/Plugins/Cameras/GameplayCameras/Source/GameplayCameras/Public/Core/CameraNodeChildrenView.h
new file mode 100644
--- /dev/null
+++ b/Engine/Plugins/Cameras/GameplayCameras/Source/GameplayCameras/Public/Core/CameraNodeChildrenView.h
@@ -0,0 +1,140 @@
+// Copyright Epic Games, Inc. All Rights Reserved.
+
+#pragma once
+
+#include "Containers/Array.h"
+#include "Containers/ArrayView.h"
+#include "Misc/TVariant.h"
+#include "UObject/ObjectPtr.h"
+
+#include <initializer_list>
+
+class UCameraNode;
+
+/**
+ * Structure that describes a list of camera node children.
+ *
+ * This structure can either provide a TArrayView<> on an existing container of camera node children,
+ * or store within itself a list of arbitrary camera node children.
+ */
+struct FCameraNodeChildrenView
+{
+public:
+
+	using FArrayView = TArrayView<TObjectPtr<UCameraNode>>;
+	using FArray = TArray<TObjectPtr<UCameraNode>, TInlineAllocator<4>>;
+
+	/** An empty view. */
+	FCameraNodeChildrenView()
+	{
+	}
+
+	/** Sets the view to the given TArrayView<>. */
+	FCameraNodeChildrenView(FArrayView&& InArrayView)
+	{
+		Storage.Set<FArrayView>(MoveTemp(InArrayView));
+	}
+
+	/** Sets the view to pointer storage and adds the given list of children pointers. */
+	FCameraNodeChildrenView(std::initializer_list<TObjectPtr<UCameraNode>> InChildren)
+	{
+		Storage.Set<FArray>(FArray(InChildren));
+	}
+
+	/**
+	 * Sets the view to pointer storage (if not already done) and adds the given
+	 * pointer to the list.
+	 */
+	void Add(TObjectPtr<UCameraNode> InChild)
+	{
+		if (Storage.GetIndex() != FStorage::IndexOfType<FArray>())
+		{
+			Storage.Set<FArray>(FArray());
+		}
+		FArray& Array = Storage.Get<FArray>();
+		Array.Add(InChild);
+	}
+
+	/** Whether this view has any children. */
+	bool IsEmpty() const
+	{
+		switch (Storage.GetIndex())
+		{
+			case FStorage::IndexOfType<FArrayView>():
+				return Storage.Get<FArrayView>().IsEmpty();
+			case FStorage::IndexOfType<FArray>():
+				return Storage.Get<FArray>().IsEmpty();
+			default:
+				return true;
+		}
+	}
+
+	/** Returns the number of children. */
+	int32 Num() const
+	{
+		switch (Storage.GetIndex())
+		{
+			case FStorage::IndexOfType<FArrayView>():
+				return Storage.Get<FArrayView>().Num();
+			case FStorage::IndexOfType<FArray>():
+				return Storage.Get<FArray>().Num();
+			default:
+				return 0;
+		}
+	}
+
+	/** Gets the i'th child. */
+	UCameraNode* operator[](int32 Index) const
+	{
+		switch (Storage.GetIndex())
+		{
+			case FStorage::IndexOfType<FArrayView>():
+				return Storage.Get<FArrayView>()[Index].Get();
+			case FStorage::IndexOfType<FArray>():
+				return Storage.Get<FArray>()[Index].Get();
+			default:
+				return nullptr;
+		}
+	}
+
+public:
+
+	// Range iteration
+
+	struct FIterator
+	{
+		const FCameraNodeChildrenView* Owner;
+		int32 Index;
+
+		FORCEINLINE FIterator& operator++()
+		{
+			++Index;
+			return *this;
+		}
+
+		FORCEINLINE UCameraNode* operator*()
+		{
+			return (*Owner)[Index];
+		}
+
+		FORCEINLINE bool operator== (const FIterator& Other) const
+		{
+			return Owner == Other.Owner
+				&& Index == Other.Index;
+		}
+
+		FORCEINLINE bool operator!= (const FIterator& Other) const
+		{
+			return !(*this == Other);
+		}
+	};
+
+	FORCEINLINE FIterator begin() const { return FIterator{ this, 0 }; }
+	FORCEINLINE FIterator end() const { return FIterator{ this, Num() }; }
+
+private:
+
+	using FStorage = TVariant<FArrayView, FArray>;
+	FStorage Storage;
+};
+
diff --git a/Engine/Plugins/Cameras/GameplayCameras/Source/GameplayCameras/Public/Core/CameraPose.h b/Engine/Plugins/Cameras/GameplayCameras/Source/GameplayCameras/Public/Core/CameraPose.h
new file mode 100644
--- /dev/null
+++ b/Engine/Plugins/Cameras/GameplayCameras/Source/GameplayCameras/Public/Core/CameraPose.h
@@ -0,0 +1,228 @@
+// Copyright Epic Games, Inc. All Rights Reserved.
+
+#pragma once
+
+#include "CoreTypes.h"
+#include "Math/MathFwd.h"
+#include "Math/Transform.h"
+
+#include "CameraPose.generated.h"
+
+#define UE_CAMERA_POSE_FOR_TRANSFORM_PROPERTIES()\
+	UE_CAMERA_POSE_FOR_PROPERTY(FVector, Location)\
+	UE_CAMERA_POSE_FOR_PROPERTY(FRotator3d, Rotation)
+
+#define UE_CAMERA_POSE_FOR_INTERPOLABLE_PROPERTIES()\
+	UE_CAMERA_POSE_FOR_PROPERTY(double, TargetDistance)\
+	UE_CAMERA_POSE_FOR_PROPERTY(float,  Aperture)\
+	UE_CAMERA_POSE_FOR_PROPERTY(float,  FocusDistance)\
+	UE_CAMERA_POSE_FOR_PROPERTY(float,  SensorWidth)\
+	UE_CAMERA_POSE_FOR_PROPERTY(float,  SensorHeight)\
+	UE_CAMERA_POSE_FOR_PROPERTY(float,  SqueezeFactor)\
+	UE_CAMERA_POSE_FOR_PROPERTY(float,  AspectRatio)\
+	UE_CAMERA_POSE_FOR_PROPERTY(float,  NearClippingPlane)\
+	UE_CAMERA_POSE_FOR_PROPERTY(float,  FarClippingPlane)
+
+#define UE_CAMERA_POSE_FOR_FOV_PROPERTIES()\
+	UE_CAMERA_POSE_FOR_PROPERTY(float, FieldOfView)\
+	UE_CAMERA_POSE_FOR_PROPERTY(float, FocalLength)
+
+#define UE_CAMERA_POSE_FOR_BOOL_PROPERTIES()\
+	UE_CAMERA_POSE_FOR_PROPERTY(bool, bConstrainAspectRatio)
+
+#define UE_CAMERA_POSE_FOR_ALL_PROPERTIES()\
+	UE_CAMERA_POSE_FOR_TRANSFORM_PROPERTIES()\
+	UE_CAMERA_POSE_FOR_INTERPOLABLE_PROPERTIES()\
+	UE_CAMERA_POSE_FOR_FOV_PROPERTIES()\
+	UE_CAMERA_POSE_FOR_BOOL_PROPERTIES()
+
+/**
+ * Boolean flags for each of the properties inside FCameraPose.
+ */
+struct FCameraPoseFlags
+{
+#define UE_CAMERA_POSE_FOR_PROPERTY(PropType, PropName)\
+	bool PropName = false;
+
+UE_CAMERA_POSE_FOR_ALL_PROPERTIES()
+
+#undef UE_CAMERA_POSE_FOR_PROPERTY
+
+public:
+
+	/** Sets all flags to the given value. */
+	FCameraPoseFlags& SetAllFlags(bool bInValue);
+	/** Sets the flags that are set in OtherFlags, but checks that no flag is set on both structures. */
+	FCameraPoseFlags& ExclusiveCombine(const FCameraPoseFlags& OtherFlags);
+
+	/** Combines the flags with an AND logical operation. */
+	FCameraPoseFlags& AND(const FCameraPoseFlags& OtherFlags);
+	/** Combines the flags with an OR logical operation. */
+	FCameraPoseFlags& OR(const FCameraPoseFlags& OtherFlags);
+};
+
+/**
+ * Structure describing the state of a camera.
+ *
+ * Fields are private and can only be accessed via the getters and setters.
+ * The ChangedFlags structure keeps track of which fields were changed via the setters.
+ */
+USTRUCT()
+struct FCameraPose
+{
+	GENERATED_BODY()
+
+public:
+
+	FCameraPose();
+
+	/** Resets this camera pose to its default values. */
+	void Reset(bool bSetAllChangedFlags = true);
+
+public:
+
+	// Getters and setters
+
+#define UE_CAMERA_POSE_FOR_PROPERTY(PropType, PropName)\
+	PropType Get##PropName() const\
+	{\
+		return PropName;\
+	}\
+	void Set##PropName(TCallTraits<PropType>::ParamType InValue)\
+	{\
+		ChangedFlags.PropName = true;\
+		PropName = InValue;\
+	}
+
+UE_CAMERA_POSE_FOR_ALL_PROPERTIES()
+
+#undef UE_CAMERA_POSE_FOR_PROPERTY
+
+public:
+
+	// Changed flags management
+
+	/** Get the changed flags. */
+	FCameraPoseFlags& GetChangedFlags() { return ChangedFlags; }
+	/** Get the changed flags. */
+	const FCameraPoseFlags& GetChangedFlags() const { return ChangedFlags; }
+
+	/** Set the changed flags. */
+	void SetChangedFlags(const FCameraPoseFlags& InChangedFlags) { ChangedFlags = InChangedFlags; }
+	/** Set all fields as changed. */
+	void SetAllChangedFlags();
+	/** Set all fields as clean. */
+	void ClearAllChangedFlags();
+
+public:
+
+	// Utility
+
+	/** Gets the transform of the camera. */
+	FTransform3d GetTransform() const;
+	/** Sets the transform of the camera. */
+	void SetTransform(FTransform3d Transform);
+
+	/**
+	 * Computes the field of view of the camera.
+	 * The effective field of view can be driven by the FieldOfView property, or
+	 * the FocalLength property in combination with the sensor size.
+	 */
+	float GetEffectiveFieldOfView() const;
+
+	/**
+	 * Gets the aiming ray of the camera.
+	 */
+	FRay3d GetAimRay() const;
+
+public:
+
+	// Interpolation
+	
+	/** Takes all changed properties from OtherPose and sets them on this camera pose. */
+	void OverrideChanged(const FCameraPose& OtherPose);
+	/** Interpolates all changed properties from ToPose using the given factor. */
+	void LerpChanged(const FCameraPose& ToPose, float Factor);
+	/** Interpolates changed properties from ToPose using the given factor. Only properties defined by InMask are taken into account. */
+	void LerpChanged(const FCameraPose& ToPose, float Factor, const FCameraPoseFlags& InMask, bool bInvertMask, FCameraPoseFlags& OutMask);
+
+private:
+
+	void InternalLerpChanged(const FCameraPose& ToPose, float Factor, const FCameraPoseFlags& InMask, bool bInvertMask, FCameraPoseFlags& OutMask);
+
+private:
+
+	/** The location of the camera in the world */
+	UPROPERTY()
+	FVector3d Location = {0, 0, 0};
+
+	/** The rotation of the camera in the world */
+	UPROPERTY()
+	FRotator3d Rotation = {0, 0, 0};
+
+	/** Distance to the target */
+	UPROPERTY()
+	double TargetDistance = 100.0;
+
+	/**
+	 * The horizontal field of view of the camera, in degrees 
+	 * If zero or less, focal length is used instead
+	 */
+	UPROPERTY()
+	float FieldOfView = -1.f;  // Default to using a focal length
+
+	/**
+	 * The aspect ratio of the camera
+	 * If zero or less, the sensor width and height are used instead
+	 */
+	UPROPERTY()
+	float AspectRatio = -1.f;
+
+	/**
+	 * The focal length of the camera's lens, in millimeters
+	 * If zero or less, field of view is used instead
+	 */
+	UPROPERTY()
+	float FocalLength = 35.f;
+
+	/** The aperture of the camera's lens, in f-stops */
+	UPROPERTY()
+	float Aperture = 2.8f;
+
+	/** The focus distance of the camera's lens, in world units */
+	UPROPERTY()
+	float FocusDistance = -1.f;
+
+	/** The width of the camera's sensor, in millimeters */
+	UPROPERTY()
+	float SensorWidth = 24.89f;
+
+	/** The height of the camera's sensor, in millimeters */
+	UPROPERTY()
+	float SensorHeight = 18.67f;
+
+	/** Squeeze factor for anamorphic lenses */
+	UPROPERTY()
+	float SqueezeFactor = 1.f;
+
+	/** The distance to the near clipping plane, in world units */
+	UPROPERTY()
+	float NearClippingPlane = 10.f;
+
+	/** The distance to the far clipping plane, in world units */
+	UPROPERTY()
+	float FarClippingPlane = -1.f;
+
+	/** Whether to constrain aspect ratio */
+	UPROPERTY()
+	bool bConstrainAspectRatio = false;
+
+private:
+	
+	/**
+	 * Flags keeping track of which properties were written to since
+	 * last time the flags were cleared.
+	 */
+	FCameraPoseFlags ChangedFlags;
+};
+
diff --git a/Engine/Plugins/Cameras/GameplayCameras/Source/GameplayCameras/Public/Core/CameraRuntimeInstantiator.h b/Engine/Plugins/Cameras/GameplayCameras/Source/GameplayCameras/Public/Core/CameraRuntimeInstantiator.h
new file mode 100644
--- /dev/null
+++ b/Engine/Plugins/Cameras/GameplayCameras/Source/GameplayCameras/Public/Core/CameraRuntimeInstantiator.h
@@ -0,0 +1,56 @@
+// Copyright Epic Games, Inc. All Rights Reserved.
+
+#pragma once
+
+#include "CoreTypes.h"
+#include "UObject/NameTypes.h"
+#include "UObject/UnrealType.h"
+
+class IGameplayCamerasModule;
+class UCameraDirector;
+class UCameraNode;
+class UObject;
+
+/**
+ * Parameter structure for instantiating an object.
+ */
+struct FCameraRuntimeInstantiationParams
+{
+	UObject* InstantiationOuter = nullptr;
+	FName InstantiationName = NAME_None;
+};
+
+/**
+ * Utility class for instanting source objects into the game.
+ */
+class FCameraRuntimeInstantiator
+{
+public:
+
+	/** 
+	 * Instantiates a tree of camera nodes by cloning it, and keeping track of the relationship between
+	 * the clones and their source assets.
+	 */
+	static UCameraNode* InstantiateCameraNodeTree(const UCameraNode* InRootNode, const FCameraRuntimeInstantiationParams& Params);
+
+	/** 
+	 * Instantiates a camera directory by cloning it, and keeping track of the relationship between
+	 * the clone and its source asset.
+	 */
+	static UCameraDirector* InstantiateCameraDirector(const UCameraDirector* InCameraDirector, const FCameraRuntimeInstantiationParams& Params);
+
+	/**
+	 * Take a property change applied on a source object and try to replicate it on any related instantiated objects.
+	 */
+	static void ForwardPropertyChange(const UObject* Object, const FPropertyChangedEvent& PropertyChangedEvent);
+
+private:
+
+	static UObject* InstantiateObject(const UObject* InSource, const FCameraRuntimeInstantiationParams& Params);
+	static IGameplayCamerasModule& GetGameplayCamerasModule();
+
+private:
+
+	static IGameplayCamerasModule* GameplayCamerasModule;
+};
+
diff --git a/Engine/Plugins/Cameras/GameplayCameras/Source/GameplayCameras/Public/Core/CameraSystemEvaluator.h b/Engine/Plugins/Cameras/GameplayCameras/Source/GameplayCameras/Public/Core/CameraSystemEvaluator.h
new file mode 100644
--- /dev/null
+++ b/Engine/Plugins/Cameras/GameplayCameras/Source/GameplayCameras/Public/Core/CameraSystemEvaluator.h
@@ -0,0 +1,92 @@
+// Copyright Epic Games, Inc. All Rights Reserved.
+
+#pragma once
+
+#include "Core/CameraDirector.h"
+#include "Core/CameraEvaluationContextStack.h"
+#include "Core/CameraPose.h"
+#include "Core/CameraNode.h"
+#include "CoreTypes.h"
+#include "UObject/GCObject.h"
+
+#include "CameraSystemEvaluator.generated.h"
+
+class UCameraDirector;
+class UCameraEvaluationContext;
+class UCameraMode;
+class URootCameraNode;
+struct FMinimalViewInfo;
+
+/**
+ * Parameter structure for updating the camera system.
+ */
+struct FCameraSystemEvaluationUpdateParams
+{
+	/** Time interface for the update. */
+	float DeltaTime = 0.f;
+};
+
+/**
+ * Result structure for updating the camera system.
+ */
+struct FCameraSystemEvaluationUpdateResult
+{
+	/** The result camera pose. */
+	FCameraPose CameraPose;
+
+	/** Whether this evaluation was a camera cut. */
+	bool bIsCameraCut = false;
+
+	/** Whether this result is valid. */
+	bool bIsValid = false;
+};
+
+/**
+ * The main camera system evaluator class.
+ */
+UCLASS(MinimalAPI)
+class UCameraSystemEvaluator : public UObject
+{
+	GENERATED_BODY()
+
+public:
+
+	UCameraSystemEvaluator(const FObjectInitializer& ObjectInit);
+
+public:
+
+	/** Push a new evaluation context on the stack. */
+	void PushEvaluationContext(UCameraEvaluationContext* EvaluationContext);
+	/** Remove an existing evaluation context from the stack. */
+	void RemoveEvaluationContext(UCameraEvaluationContext* EvaluationContext);
+	/** Pop the active (top) evaluation context from the stack. */
+	void PopEvaluationContext();
+
+public:
+
+	/** Run an update of the camera system. */
+	void Update(const FCameraSystemEvaluationUpdateParams& Params);
+
+	/** Get the last evaluated camera. */
+	void GetEvaluatedCameraView(FMinimalViewInfo& DesiredView);
+
+public:
+
+	static void AddReferencedObjects(UObject* InThis, FReferenceCollector& Collector);
+
+private:
+
+	/** The root camera node. */
+	UPROPERTY()
+	TObjectPtr<URootCameraNode> RootNode;
+
+	/** The stack of active evaluation context. */
+	FCameraEvaluationContextStack ContextStack;
+
+	/** The current result of the root camera node. */
+	FCameraNodeRunResult RootNodeResult;
+
+	/** The current overall result of the camera system. */
+	FCameraSystemEvaluationUpdateResult Result;
+};
+
diff --git a/Engine/Plugins/Cameras/GameplayCameras/Source/GameplayCameras/Public/Core/DefaultRootCameraNode.h b/Engine/Plugins/Cameras/GameplayCameras/Source/GameplayCameras/Public/Core/DefaultRootCameraNode.h
new file mode 100644
--- /dev/null
+++ b/Engine/Plugins/Cameras/GameplayCameras/Source/GameplayCameras/Public/Core/DefaultRootCameraNode.h
@@ -0,0 +1,44 @@
+// Copyright Epic Games, Inc. All Rights Reserved.
+
+#pragma once
+
+#include "Core/RootCameraNode.h"
+
+#include "DefaultRootCameraNode.generated.h"
+
+class UBlendStackCameraNode;
+
+/**
+ * The default implementation of a root camera node.
+ */
+UCLASS(MinimalAPI)
+class UDefaultRootCameraNode : public URootCameraNode
+{
+	GENERATED_BODY()
+
+public:
+
+	UDefaultRootCameraNode(const FObjectInitializer& ObjectInit);
+
+protected:
+
+	virtual FCameraNodeChildrenView OnGetChildren() override;
+	virtual void OnRun(const FCameraNodeRunParams& Params, FCameraNodeRunResult& OutResult) override;
+
+	virtual void OnActivateCameraMode(const FActivateCameraModeParams& Params) override;
+
+private:
+
+	UPROPERTY(Instanced)
+	TObjectPtr<UCameraNode> BaseLayer;
+
+	UPROPERTY(Instanced)
+	TObjectPtr<UCameraNode> MainLayer;
+
+	UPROPERTY(Instanced)
+	TObjectPtr<UCameraNode> GlobalLayer;
+
+	UPROPERTY(Instanced)
+	TObjectPtr<UCameraNode> VisualLayer;
+};
+
diff --git a/Engine/Plugins/Cameras/GameplayCameras/Source/GameplayCameras/Public/Core/RootCameraNode.h b/Engine/Plugins/Cameras/GameplayCameras/Source/GameplayCameras/Public/Core/RootCameraNode.h
new file mode 100644
--- /dev/null
+++ b/Engine/Plugins/Cameras/GameplayCameras/Source/GameplayCameras/Public/Core/RootCameraNode.h
@@ -0,0 +1,62 @@
+// Copyright Epic Games, Inc. All Rights Reserved.
+
+#pragma once
+
+#include "Core/CameraNode.h"
+
+#include "RootCameraNode.generated.h"
+
+class UCameraEvaluationContext;
+class UCameraMode;
+
+/**
+ * Defines evaluation layers for camera modes.
+ */
+UENUM()
+enum class ECameraModeLayer
+{
+	Base,
+	Main,
+	Global,
+	Visual,
+	User0,
+	User1,
+	User2
+};
+ENUM_CLASS_FLAGS(ECameraModeLayer)
+
+/**
+ * Parameter structure for activating a new camera mode.
+ */
+struct FActivateCameraModeParams
+{
+	/** The evaluation context in which the camera mode runs. */
+	TWeakObjectPtr<const UCameraEvaluationContext> EvaluationContext;
+
+	/** The source camera mode asset that will be instantiated. */
+	TObjectPtr<const UCameraMode> CameraMode;
+
+	/** The evaluation layer on which to instantiate the camera mode. */
+	ECameraModeLayer Layer = ECameraModeLayer::Main;
+};
+
+/**
+ * The base class for a camera node that can act as the root of the
+ * camera system evaluation.
+ */
+UCLASS(MinimalAPI, Abstract)
+class URootCameraNode : public UCameraNode
+{
+	GENERATED_BODY()
+
+public:
+
+	/** Activates a camera mode. */
+	void ActivateCameraMode(const FActivateCameraModeParams& Params);
+
+private:
+
+	/** Activates a camera mode. */
+	virtual void OnActivateCameraMode(const FActivateCameraModeParams& Params) {}
+};
+
diff --git a/Engine/Plugins/Cameras/GameplayCameras/Source/GameplayCameras/Public/Directors/SingleCameraDirector.h b/Engine/Plugins/Cameras/GameplayCameras/Source/GameplayCameras/Public/Directors/SingleCameraDirector.h
new file mode 100644
--- /dev/null
+++ b/Engine/Plugins/Cameras/GameplayCameras/Source/GameplayCameras/Public/Directors/SingleCameraDirector.h
@@ -0,0 +1,31 @@
+// Copyright Epic Games, Inc. All Rights Reserved.
+
+#pragma once
+
+#include "Core/CameraDirector.h"
+
+#include "SingleCameraDirector.generated.h"
+
+/**
+ * A simple camera director that only ever returns one single camera mode.
+ */
+UCLASS(EditInlineNew)
+class USingleCameraDirector : public UCameraDirector
+{
+	GENERATED_BODY()
+
+public:
+
+	USingleCameraDirector(const FObjectInitializer& ObjectInit);
+
+protected:
+
+	virtual void OnRun(const FCameraDirectorRunParams& Params, FCameraDirectorRunResult& OutResult) override;
+
+public:
+
+	/** The camera mode to run every frame. */
+	UPROPERTY(EditAnywhere, Category=Common)
+	TObjectPtr<UCameraMode> CameraMode;
+};
+
diff --git a/Engine/Plugins/Cameras/GameplayCameras/Source/GameplayCameras/Public/GameFramework/GameplayCameraActor.h b/Engine/Plugins/Cameras/GameplayCameras/Source/GameplayCameras/Public/GameFramework/GameplayCameraActor.h
new file mode 100644
--- /dev/null
+++ b/Engine/Plugins/Cameras/GameplayCameras/Source/GameplayCameras/Public/GameFramework/GameplayCameraActor.h
@@ -0,0 +1,45 @@
+// Copyright Epic Games, Inc. All Rights Reserved.
+
+#pragma once
+
+#include "CoreMinimal.h"
+#include "Engine/EngineTypes.h"
+#include "GameFramework/Actor.h"
+#include "UObject/ObjectMacros.h"
+
+#include "GameplayCameraActor.generated.h"
+
+class APlayerController;
+class UGameplayCameraComponent;
+
+/**
+ * An actor that can run a camera asset.
+ */
+UCLASS(BlueprintType, MinimalAPI, ClassGroup=Camera, HideCategories=(Input, Rendering))
+class AGameplayCameraActor : public AActor
+{
+	GENERATED_BODY()
+
+public:
+
+	AGameplayCameraActor(const FObjectInitializer& ObjectInit);
+
+	virtual USceneComponent* GetDefaultAttachComponent() const override;
+
+public:
+
+	UFUNCTION(BlueprintGetter, Category=Camera)
+	USceneComponent* GetSceneComponent() const { return SceneComponent; }
+
+	UFUNCTION(BlueprintGetter, Category=Camera)
+	UGameplayCameraComponent* GetCameraComponent() const { return CameraComponent; }
+
+private:
+
+	UPROPERTY(VisibleAnywhere, Category=Camera, BluePrintGetter="GetSceneComponent")
+	TObjectPtr<USceneComponent> SceneComponent;
+
+	UPROPERTY(VisibleAnywhere, Category=Camera, BlueprintGetter="GetCameraComponent", meta=(ExposeFunctionCategories="Camera"))
+	TObjectPtr<UGameplayCameraComponent> CameraComponent;
+};
+
diff --git a/Engine/Plugins/Cameras/GameplayCameras/Source/GameplayCameras/Public/GameFramework/GameplayCameraComponent.h b/Engine/Plugins/Cameras/GameplayCameras/Source/GameplayCameras/Public/GameFramework/GameplayCameraComponent.h
new file mode 100644
--- /dev/null
+++ b/Engine/Plugins/Cameras/GameplayCameras/Source/GameplayCameras/Public/GameFramework/GameplayCameraComponent.h
@@ -0,0 +1,81 @@
+// Copyright Epic Games, Inc. All Rights Reserved.
+
+#pragma once
+
+#include "Components/SceneComponent.h"
+#include "Core/CameraEvaluationContext.h"
+#include "CoreMinimal.h"
+#include "UObject/ObjectMacros.h"
+#include "UObject/ScriptInterface.h"
+
+#include "GameplayCameraComponent.generated.h"
+
+class UCameraAsset;
+class UGameplayCameraComponentEvaluationContext;
+
+/**
+ * A component that can run a camera asset inside its own camera evaluation context.
+ */
+UCLASS(BlueprintType, MinimalAPI, ClassGroup=Camera, HideCategories=(Mobility, Rendering, LOD))
+class UGameplayCameraComponent : public USceneComponent
+{
+	GENERATED_BODY()
+
+public:
+
+	UGameplayCameraComponent(const FObjectInitializer& ObjectInit);
+
+	UFUNCTION(BlueprintCallable, Category=Camera)
+	GAMEPLAYCAMERAS_API void ActivateCamera(int32 PlayerIndex = 0);
+
+public:
+
+	// UActorComponent interface
+	virtual void OnRegister() override;
+	virtual void TickComponent(float DeltaTime, enum ELevelTick TickType, FActorComponentTickFunction *ThisTickFunction) override;
+	virtual void OnComponentDestroyed(bool bDestroyingHierarchy) override;
+
+private:
+
+	void ActivateCamera(APlayerController* PlayerController);
+
+#if WITH_EDITORONLY_DATA
+
+	void UpdatePreviewMeshTransform();
+
+#endif
+
+public:
+
+	UPROPERTY(EditAnywhere, BlueprintReadWrite, Category=Camera)
+	TObjectPtr<UCameraAsset> Camera;
+
+protected:
+
+	UPROPERTY(Transient)
+	TObjectPtr<UGameplayCameraComponentEvaluationContext> EvaluationContext;
+	
+#if WITH_EDITORONLY_DATA
+
+	UPROPERTY()
+	TObjectPtr<UStaticMesh> PreviewMesh;
+
+	UPROPERTY()
+	TObjectPtr<UStaticMeshComponent> PreviewMeshComponent;
+
+#endif	// WITH_EDITORONLY_DATA
+};
+
+UCLASS(MinimalAPI)
+class UGameplayCameraComponentEvaluationContext : public UCameraEvaluationContext
+{
+	GENERATED_BODY()
+
+public:
+
+	UGameplayCameraComponentEvaluationContext(const FObjectInitializer& ObjectInit);
+
+	void Initialize(UGameplayCameraComponent* Owner);
+	void Update(UGameplayCameraComponent* Owner);
+};
+
diff --git a/Engine/Plugins/Cameras/GameplayCameras/Source/GameplayCameras/Public/GameFramework/GameplayCameraSystemActor.h b/Engine/Plugins/Cameras/GameplayCameras/Source/GameplayCameras/Public/GameFramework/GameplayCameraSystemActor.h
new file mode 100644
--- /dev/null
+++ b/Engine/Plugins/Cameras/GameplayCameras/Source/GameplayCameras/Public/GameFramework/GameplayCameraSystemActor.h
@@ -0,0 +1,49 @@
+// Copyright Epic Games, Inc. All Rights Reserved.
+
+#pragma once
+
+#include "CoreMinimal.h"
+#include "Engine/EngineTypes.h"
+#include "GameFramework/Actor.h"
+#include "UObject/ObjectMacros.h"
+
+#include "GameplayCameraSystemActor.generated.h"
+
+class APlayerController;
+class UGameplayCameraSystemComponent;
+
+/**
+ * An actor that hosts a camera system.
+ */
+UCLASS(BlueprintType, MinimalAPI, ClassGroup=Camera, HideCategories=(Input, Rendering))
+class AGameplayCameraSystemActor : public AActor
+{
+	GENERATED_BODY()
+
+public:
+
+	AGameplayCameraSystemActor(const FObjectInitializer& ObjectInit);
+
+	UFUNCTION(BlueprintCallable, Category=Camera)
+	void ActivateForPlayer(int32 PlayerIndex);
+
+	UFUNCTION(BlueprintGetter, Category=Camera)
+	UGameplayCameraSystemComponent* GetCameraSystemComponent() const { return CameraSystemComponent; }
+
+public:
+
+	// AActor interface
+	virtual void BeginPlay() override;
+	virtual void BecomeViewTarget(APlayerController* PC) override;
+	virtual void CalcCamera(float DeltaTime, struct FMinimalViewInfo& OutResult) override;
+	virtual void EndViewTarget(APlayerController* PC) override;
+
+private:
+
+	UPROPERTY(VisibleAnywhere, Category=Camera, BlueprintGetter="GetCameraSystemComponent", meta=(ExposeFunctionCategories="CameraSystem"))
+	TObjectPtr<UGameplayCameraSystemComponent> CameraSystemComponent;
+
+	UPROPERTY(EditAnywhere, Category=Camera)
+	TEnumAsByte<EAutoReceiveInput::Type> AutoActivateForPlayer;
+};
+
diff --git a/Engine/Plugins/Cameras/GameplayCameras/Source/GameplayCameras/Public/GameFramework/GameplayCameraSystemComponent.h b/Engine/Plugins/Cameras/GameplayCameras/Source/GameplayCameras/Public/GameFramework/GameplayCameraSystemComponent.h
new file mode 100644
--- /dev/null
+++ b/Engine/Plugins/Cameras/GameplayCameras/Source/GameplayCameras/Public/GameFramework/GameplayCameraSystemComponent.h
@@ -0,0 +1,67 @@
+// Copyright Epic Games, Inc. All Rights Reserved.
+
+#pragma once
+
+#include "CoreMinimal.h"
+#include "Core/CameraSystemEvaluator.h"
+#include "UObject/ObjectMacros.h"
+#include "UObject/ScriptInterface.h"
+#include "Components/SceneComponent.h"
+
+#include "GameplayCameraSystemComponent.generated.h"
+
+class UCameraSystemEvaluator;
+class UCameraMode;
+struct FMinimalViewInfo;
+
+/**
+ * A component that hosts a camera system.
+ */
+UCLASS(BlueprintType, MinimalAPI, ClassGroup=Camera, HideCategories=(Mobility, Rendering, LOD))
+class UGameplayCameraSystemComponent : public USceneComponent
+{
+	GENERATED_BODY()
+
+public:
+
+	UGameplayCameraSystemComponent(const FObjectInitializer& ObjectInit);
+
+	UFUNCTION(BlueprintGetter)
+	UCameraSystemEvaluator* GetCameraSystemEvaluator() { return Evaluator; }
+
+	GAMEPLAYCAMERAS_API void GetCameraView(float DeltaTime, FMinimalViewInfo& DesiredView);
+
+public:
+
+	// UActorComponent interface
+	virtual void OnRegister() override;
+	virtual void TickComponent(float DeltaTime, ELevelTick TickType, FActorComponentTickFunction* ThisTickFunction) override;
+	virtual void OnComponentDestroyed(bool bDestroyingHierarchy) override;
+
+	// USceneComponent interface
+#if WITH_EDITOR
+	virtual bool GetEditorPreviewInfo(float DeltaTime, FMinimalViewInfo& ViewOut) override;
+#endif  // WITH_EDITOR
+
+public:
+
+	// Internal API
+	void OnBecomeViewTarget();
+	void OnEndViewTarget();
+
+private:
+	
+	UPROPERTY(Instanced, Transient)
+	TObjectPtr<UCameraSystemEvaluator> Evaluator;
+
+#if WITH_EDITORONLY_DATA
+
+	UPROPERTY()
+	TObjectPtr<UStaticMesh> PreviewMesh;
+
+	UPROPERTY()
+	TObjectPtr<UStaticMeshComponent> PreviewMeshComponent;
+
+#endif	// WITH_EDITORONLY_DATA
+};
+
diff --git a/Engine/Plugins/Cameras/GameplayCameras/Source/GameplayCameras/Public/GameplayCamerasModule.h b/Engine/Plugins/Cameras/GameplayCameras/Source/GameplayCameras/Public/GameplayCamerasModule.h
deleted file mode 100644
diff --git a/Engine/Plugins/Cameras/GameplayCameras/Source/GameplayCameras/Public/IGameplayCamerasLiveEditManager.h b/Engine/Plugins/Cameras/GameplayCameras/Source/GameplayCameras/Public/IGameplayCamerasLiveEditManager.h
new file mode 100644
--- /dev/null
+++ b/Engine/Plugins/Cameras/GameplayCameras/Source/GameplayCameras/Public/IGameplayCamerasLiveEditManager.h
@@ -0,0 +1,33 @@
+// Copyright Epic Games, Inc. All Rights Reserved.
+
+#pragma once
+
+#include "CoreTypes.h"
+#include "Templates/SharedPointer.h"
+#include "UObject/UnrealType.h"
+
+class UObject;
+
+/**
+ * Interface for an object that can handle live-editing features of the camera assets.
+ */
+class IGameplayCamerasLiveEditManager : public TSharedFromThis<IGameplayCamerasLiveEditManager>
+{
+public:
+
+	virtual ~IGameplayCamerasLiveEditManager() {}
+
+	/**
+	 * Register a new set of instantiated objects.
+	 *
+	 * @param InstantiatedObjects  A mapping between a source object and an instantiated object.
+	 */
+	virtual void RegisterInstantiatedObjects(const TMap<UObject*, UObject*> InstantiatedObjects) = 0;
+
+	/**
+	 * Request that a property change on the given source object should be replicated on any
+	 * known related instantied objects.
+	 */
+	virtual void ForwardPropertyChange(const UObject* Object, const FPropertyChangedEvent& PropertyChangedEvent) = 0;
+};
+
diff --git a/Engine/Plugins/Cameras/GameplayCameras/Source/GameplayCameras/Public/IGameplayCamerasModule.h b/Engine/Plugins/Cameras/GameplayCameras/Source/GameplayCameras/Public/IGameplayCamerasModule.h
new file mode 100644
--- /dev/null
+++ b/Engine/Plugins/Cameras/GameplayCameras/Source/GameplayCameras/Public/IGameplayCamerasModule.h
@@ -0,0 +1,40 @@
+// Copyright Epic Games, Inc. All Rights Reserved.
+
+#pragma once
+
+#include "Modules/ModuleInterface.h"
+#include "Stats/Stats.h"
+
+#if WITH_EDITOR
+class IGameplayCamerasLiveEditManager;
+#endif
+
+DECLARE_STATS_GROUP(TEXT("Camera System Evaluation"), STATGROUP_CameraSystem, STATCAT_Advanced)
+DECLARE_STATS_GROUP(TEXT("Camera Animation Evaluation"), STATGROUP_CameraAnimation, STATCAT_Advanced)
+
+class IGameplayCamerasModule : public IModuleInterface
+{
+public:
+	/**
+	 * Singleton-like access to ICameraModule
+	 *
+	 * @return The ICameraModule instance, loading the module on demand if needed
+	 */
+	static IGameplayCamerasModule& Get();
+
+#if WITH_EDITOR
+	/**
+	 * Gets the live edit manager.
+	 */
+	virtual TSharedPtr<IGameplayCamerasLiveEditManager> GetLiveEditManager() const = 0;
+
+	/**
+	 * Sets the live edit manager.
+	 */
+	virtual void SetLiveEditManager(TSharedPtr<IGameplayCamerasLiveEditManager> InLiveEditManager) = 0;
+#endif
+};
+
+#if UE_ENABLE_INCLUDE_ORDER_DEPRECATED_IN_5_2
+#include "CoreMinimal.h"
+#endif
diff --git a/Engine/Plugins/Cameras/GameplayCameras/Source/GameplayCameras/Public/Math/Interpolation.h b/Engine/Plugins/Cameras/GameplayCameras/Source/GameplayCameras/Public/Math/Interpolation.h
new file mode 100644
--- /dev/null
+++ b/Engine/Plugins/Cameras/GameplayCameras/Source/GameplayCameras/Public/Math/Interpolation.h
@@ -0,0 +1,29 @@
+// Copyright Epic Games, Inc. All Rights Reserved.
+
+#pragma once
+
+namespace UE::Cameras
+{
+
+/**
+ * Smoothstep interpolation
+ * 
+ * @param Value  A value between 0 and 1
+ */
+inline float SmoothStep(float Value)
+{
+	return (Value * Value * (3.f - 2.f * Value));
+}
+
+/**
+ * Smootherstep interpolation
+ *
+ * @param Value  A value between 0 and 1
+ */
+inline float SmootherStep(float Value)
+{
+	return (Value * Value * Value * (Value * (Value * 6.f - 15.f) + 10.f));
+}
+
+}  // namespace UE::Cameras
+
diff --git a/Engine/Plugins/Cameras/GameplayCameras/Source/GameplayCameras/Public/Nodes/Blends/LinearBlendCameraNode.h b/Engine/Plugins/Cameras/GameplayCameras/Source/GameplayCameras/Public/Nodes/Blends/LinearBlendCameraNode.h
new file mode 100644
--- /dev/null
+++ b/Engine/Plugins/Cameras/GameplayCameras/Source/GameplayCameras/Public/Nodes/Blends/LinearBlendCameraNode.h
@@ -0,0 +1,19 @@
+// Copyright Epic Games, Inc. All Rights Reserved.
+
+#pragma once
+
+#include "Nodes/Blends/SimpleBlendCameraNode.h"
+
+#include "LinearBlendCameraNode.generated.h"
+
+/**
+ * Linear blend node.
+ */
+UCLASS(MinimalAPI)
+class ULinearBlendCameraNode : public USimpleFixedTimeBlendCameraNode
+{
+	GENERATED_BODY()
+
+	virtual void OnComputeBlendFactor(const FCameraNodeRunParams& Params, FSimpleBlendCameraNodeRunResult& OutResult) override;
+};
+
diff --git a/Engine/Plugins/Cameras/GameplayCameras/Source/GameplayCameras/Public/Nodes/Blends/PopBlendCameraNode.h b/Engine/Plugins/Cameras/GameplayCameras/Source/GameplayCameras/Public/Nodes/Blends/PopBlendCameraNode.h
new file mode 100644
--- /dev/null
+++ b/Engine/Plugins/Cameras/GameplayCameras/Source/GameplayCameras/Public/Nodes/Blends/PopBlendCameraNode.h
@@ -0,0 +1,22 @@
+// Copyright Epic Games, Inc. All Rights Reserved.
+
+#pragma once
+
+#include "Core/BlendCameraNode.h"
+
+#include "PopBlendCameraNode.generated.h"
+
+/**
+ * A blend node that creates a camera cut (i.e. it doesn't blend at all).
+ */
+UCLASS(MinimalAPI)
+class UPopBlendCameraNode : public UBlendCameraNode
+{
+	GENERATED_BODY()
+
+protected:
+
+	virtual void OnRun(const FCameraNodeRunParams& Params, FCameraNodeRunResult& OutResult) override;
+	virtual void OnBlendResults(const FCameraNodeBlendParams& Params, FCameraNodeBlendResult& OutResult) override;
+};
+
diff --git a/Engine/Plugins/Cameras/GameplayCameras/Source/GameplayCameras/Public/Nodes/Blends/SimpleBlendCameraNode.h b/Engine/Plugins/Cameras/GameplayCameras/Source/GameplayCameras/Public/Nodes/Blends/SimpleBlendCameraNode.h
new file mode 100644
--- /dev/null
+++ b/Engine/Plugins/Cameras/GameplayCameras/Source/GameplayCameras/Public/Nodes/Blends/SimpleBlendCameraNode.h
@@ -0,0 +1,69 @@
+// Copyright Epic Games, Inc. All Rights Reserved.
+
+#pragma once
+
+#include "Core/BlendCameraNode.h"
+
+#include "SimpleBlendCameraNode.generated.h"
+
+/**
+ * Result structure for defining a simple scalar-factor-based blend.
+ */
+struct FSimpleBlendCameraNodeRunResult
+{
+	float BlendFactor = 0.f;
+};
+
+/**
+ * Base class for a blend camera node that uses a simple scalar factor.
+ */
+UCLASS(MinimalAPI, Abstract)
+class USimpleBlendCameraNode : public UBlendCameraNode
+{
+	GENERATED_BODY()
+
+public:
+
+	/** Gets the last evaluated blend factor. */
+	float GetBlendFactor() const { return BlendFactor; }
+
+protected:
+
+	virtual void OnRun(const FCameraNodeRunParams& Params, FCameraNodeRunResult& OutResult) override;
+	virtual void OnBlendResults(const FCameraNodeBlendParams& Params, FCameraNodeBlendResult& OutResult) override;
+
+	virtual void OnComputeBlendFactor(const FCameraNodeRunParams& Params, FSimpleBlendCameraNodeRunResult& OutResult) {}
+
+	void SetBlendFinished() { bIsBlendFinished = true; }
+
+private:
+
+	float BlendFactor = 0.f;
+	bool bIsBlendFinished = false;
+};
+
+/**
+ * Base class for a blend camera node that uses a simple scalar factor over a fixed time.
+ */
+UCLASS(MinimalAPI, Abstract)
+class USimpleFixedTimeBlendCameraNode : public USimpleBlendCameraNode
+{
+	GENERATED_BODY()
+
+public:
+
+	/** Duration of the blend. */
+	UPROPERTY(EditAnywhere, Category=Blending)
+	float BlendTime = 1.f;
+
+protected:
+
+	virtual void OnRun(const FCameraNodeRunParams& Params, FCameraNodeRunResult& OutResult) override;
+
+	float GetTimeFactor() const { return CurrentTime / BlendTime; }
+
+private:
+
+	float CurrentTime = 0.f;
+};
+
diff --git a/Engine/Plugins/Cameras/GameplayCameras/Source/GameplayCameras/Public/Nodes/Blends/SmoothBlendCameraNode.h b/Engine/Plugins/Cameras/GameplayCameras/Source/GameplayCameras/Public/Nodes/Blends/SmoothBlendCameraNode.h
new file mode 100644
--- /dev/null
+++ b/Engine/Plugins/Cameras/GameplayCameras/Source/GameplayCameras/Public/Nodes/Blends/SmoothBlendCameraNode.h
@@ -0,0 +1,37 @@
+// Copyright Epic Games, Inc. All Rights Reserved.
+
+#pragma once
+
+#include "Nodes/Blends/SimpleBlendCameraNode.h"
+
+#include "SmoothBlendCameraNode.generated.h"
+
+/**
+ * The smoothstep type.
+ */
+UENUM()
+enum class ESmoothCameraBlendType
+{
+	SmoothStep,
+	SmootherStep
+};
+
+/**
+ * A blend camera mode that implements the smoothstep and smoothersteps algorithms.
+ */
+UCLASS(MinimalAPI)
+class USmoothBlendCameraNode : public USimpleFixedTimeBlendCameraNode
+{
+	GENERATED_BODY()
+
+protected:
+
+	virtual void OnComputeBlendFactor(const FCameraNodeRunParams& Params, FSimpleBlendCameraNodeRunResult& OutResult) override;
+	
+public:
+
+	/** The type of algorithm to use. */
+	UPROPERTY(EditAnywhere, Category=Blending)
+	ESmoothCameraBlendType BlendType;
+};
+
diff --git a/Engine/Plugins/Cameras/GameplayCameras/Source/GameplayCameras/Public/Nodes/Common/ArrayCameraNode.h b/Engine/Plugins/Cameras/GameplayCameras/Source/GameplayCameras/Public/Nodes/Common/ArrayCameraNode.h
new file mode 100644
--- /dev/null
+++ b/Engine/Plugins/Cameras/GameplayCameras/Source/GameplayCameras/Public/Nodes/Common/ArrayCameraNode.h
@@ -0,0 +1,29 @@
+// Copyright Epic Games, Inc. All Rights Reserved.
+
+#pragma once
+
+#include "Containers/Array.h"
+#include "Core/CameraNode.h"
+
+#include "ArrayCameraNode.generated.h"
+
+/**
+ * A camera node that runs a list of other camera nodes.
+ */
+UCLASS(MinimalAPI)
+class UArrayCameraNode : public UCameraNode
+{
+	GENERATED_BODY()
+
+protected:
+
+	virtual FCameraNodeChildrenView OnGetChildren() override;
+	virtual void OnRun(const FCameraNodeRunParams& Params, FCameraNodeRunResult& OutResult) override;
+
+public:
+
+	/** The camera nodes to run. */
+	UPROPERTY(EditAnywhere, Category=Common)
+	TArray<TObjectPtr<UCameraNode>> Children;
+};
+
diff --git a/Engine/Plugins/Cameras/GameplayCameras/Source/GameplayCameras/Public/Nodes/Common/OffsetCameraNode.h b/Engine/Plugins/Cameras/GameplayCameras/Source/GameplayCameras/Public/Nodes/Common/OffsetCameraNode.h
new file mode 100644
--- /dev/null
+++ b/Engine/Plugins/Cameras/GameplayCameras/Source/GameplayCameras/Public/Nodes/Common/OffsetCameraNode.h
@@ -0,0 +1,27 @@
+// Copyright Epic Games, Inc. All Rights Reserved.
+
+#pragma once
+
+#include "Core/CameraNode.h"
+
+#include "OffsetCameraNode.generated.h"
+
+/**
+ * A camera node that offsets the location of the camera.
+ */
+UCLASS(MinimalAPI)
+class UOffsetCameraNode : public UCameraNode
+{
+	GENERATED_BODY()
+
+protected:
+
+	virtual void OnRun(const FCameraNodeRunParams& Params, FCameraNodeRunResult& OutResult) override;
+
+public:
+
+	/** The offset to apply to the camera, in local space. */
+	UPROPERTY(EditAnywhere, Category=Common)
+	FVector3d Offset;
+};
+
