Cameras: first commit for the camera system editors.

Right now both the camera asset and camera mode editors are just a Details View.

diff --git a/Engine/Plugins/Cameras/GameplayCameras/GameplayCameras.uplugin b/Engine/Plugins/Cameras/GameplayCameras/GameplayCameras.uplugin
--- a/Engine/Plugins/Cameras/GameplayCameras/GameplayCameras.uplugin
+++ b/Engine/Plugins/Cameras/GameplayCameras/GameplayCameras.uplugin
@@ -11,7 +11,7 @@
 	"MarketplaceURL" : "",
 	"SupportURL" : "",
 	"EnabledByDefault" : true,
-	"CanContainContent" : false,
+	"CanContainContent" : true,
 	"IsBetaVersion" : true,
 	"Installed" : false,
 	"Plugins":
@@ -27,6 +27,11 @@
 			"Name" : "GameplayCameras",
 			"Type" : "Runtime",
 			"LoadingPhase" : "PreDefault"
+		},
+		{
+			"Name": "GameplayCamerasEditor",
+			"Type": "Editor",
+			"LoadingPhase": "PostEngineInit"
 		}
 	]
 }
diff --git a/Engine/Plugins/Cameras/GameplayCameras/Source/GameplayCamerasEditor/GameplayCamerasEditor.Build.cs b/Engine/Plugins/Cameras/GameplayCameras/Source/GameplayCamerasEditor/GameplayCamerasEditor.Build.cs
new file mode 100644
--- /dev/null
+++ b/Engine/Plugins/Cameras/GameplayCameras/Source/GameplayCamerasEditor/GameplayCamerasEditor.Build.cs
@@ -0,0 +1,51 @@
+// Copyright Epic Games, Inc. All Rights Reserved.
+
+using UnrealBuildTool;
+
+public class GameplayCamerasEditor : ModuleRules
+{
+	public GameplayCamerasEditor(ReadOnlyTargetRules Target) : base(Target)
+	{
+		PrivateIncludePathModuleNames.AddRange(
+			new string[] 
+			{
+				"AssetTools",
+				"Kismet",
+				"EditorWidgets",
+				"MessageLog",
+			}
+		);
+
+		PrivateDependencyModuleNames.AddRange(
+			new string[] {
+				"AssetDefinition",
+				"AssetRegistry",
+				"GameplayCameras",
+				"BlueprintGraph",
+				"Core",
+				"CoreUObject",
+				"EditorFramework",
+				"Engine",
+				"InputCore",
+				"Kismet",
+				"Slate",
+				"SlateCore",
+				"ToolMenus",
+				"UnrealEd",
+				"TimeManagement"
+			}
+		);
+
+		var DynamicModuleNames = new string[] {
+			"LevelEditor",
+			"PropertyEditor",
+		};
+
+		foreach (var Name in DynamicModuleNames)
+		{
+			PrivateIncludePathModuleNames.Add(Name);
+			DynamicallyLoadedModuleNames.Add(Name);
+		}
+	}
+}
+
diff --git a/Engine/Plugins/Cameras/GameplayCameras/Source/GameplayCamerasEditor/Private/AssetTools/AssetDefinition_CameraAsset.cpp b/Engine/Plugins/Cameras/GameplayCameras/Source/GameplayCamerasEditor/Private/AssetTools/AssetDefinition_CameraAsset.cpp
new file mode 100644
--- /dev/null
+++ b/Engine/Plugins/Cameras/GameplayCameras/Source/GameplayCamerasEditor/Private/AssetTools/AssetDefinition_CameraAsset.cpp
@@ -0,0 +1,55 @@
+// Copyright Epic Games, Inc. All Rights Reserved.
+
+#include "AssetTools/AssetDefinition_CameraAsset.h"
+
+#include "IGameplayCamerasEditorModule.h"
+#include "Styles/GameplayCamerasEditorStyle.h"
+#include "Toolkits/CameraAssetEditorToolkit.h"
+#include "Toolkits/IToolkit.h"
+
+#define LOCTEXT_NAMESPACE "AssetTypeActions"
+
+#include UE_INLINE_GENERATED_CPP_BY_NAME(AssetDefinition_CameraAsset)
+
+FText UAssetDefinition_CameraAsset::GetAssetDisplayName() const
+{
+	return NSLOCTEXT("AssetTypeActions", "AssetTypeActions_CameraAsset", "Camera Asset");
+}
+
+FLinearColor UAssetDefinition_CameraAsset::GetAssetColor() const
+{
+	return FLinearColor(FColor(200, 80, 80));
+}
+
+TSoftClassPtr<UObject> UAssetDefinition_CameraAsset::GetAssetClass() const
+{
+	return UCameraAsset::StaticClass();
+}
+
+TConstArrayView<FAssetCategoryPath> UAssetDefinition_CameraAsset::GetAssetCategories() const
+{
+	static const auto Categories = { EAssetCategoryPaths::Gameplay };
+	return Categories;
+}
+
+FAssetOpenSupport UAssetDefinition_CameraAsset::GetAssetOpenSupport(const FAssetOpenSupportArgs& OpenSupportArgs) const
+{
+	return FAssetOpenSupport(
+			OpenSupportArgs.OpenMethod,
+			OpenSupportArgs.OpenMethod == EAssetOpenMethod::Edit,
+			EToolkitMode::Standalone); 
+}
+
+EAssetCommandResult UAssetDefinition_CameraAsset::OpenAssets(const FAssetOpenArgs& OpenArgs) const
+{
+	IGameplayCamerasEditorModule& GameplayCamerasEditorModule = FModuleManager::LoadModuleChecked<IGameplayCamerasEditorModule>("GameplayCamerasEditor");
+	for (UCameraAsset* CameraAsset : OpenArgs.LoadObjects<UCameraAsset>())
+	{
+		GameplayCamerasEditorModule.CreateCameraAssetEditor(OpenArgs.GetToolkitMode(), OpenArgs.ToolkitHost, CameraAsset);
+	}	
+
+	return EAssetCommandResult::Handled;
+}
+
+#undef LOCTEXT_NAMESPACE
+
diff --git a/Engine/Plugins/Cameras/GameplayCameras/Source/GameplayCamerasEditor/Private/AssetTools/AssetDefinition_CameraAsset.h b/Engine/Plugins/Cameras/GameplayCameras/Source/GameplayCamerasEditor/Private/AssetTools/AssetDefinition_CameraAsset.h
new file mode 100644
--- /dev/null
+++ b/Engine/Plugins/Cameras/GameplayCameras/Source/GameplayCamerasEditor/Private/AssetTools/AssetDefinition_CameraAsset.h
@@ -0,0 +1,25 @@
+// Copyright Epic Games, Inc. All Rights Reserved.
+
+#pragma once
+
+#include "Core/CameraAsset.h"
+#include "AssetDefinitionDefault.h"
+
+#include "AssetDefinition_CameraAsset.generated.h"
+
+UCLASS()
+class UAssetDefinition_CameraAsset : public UAssetDefinitionDefault
+{
+	GENERATED_BODY()
+
+public:
+
+	// UAssetDefinition interface
+	virtual FText GetAssetDisplayName() const override;
+	virtual FLinearColor GetAssetColor() const override;
+	virtual TSoftClassPtr<UObject> GetAssetClass() const override;
+	virtual TConstArrayView<FAssetCategoryPath> GetAssetCategories() const override;
+	virtual FAssetOpenSupport GetAssetOpenSupport(const FAssetOpenSupportArgs& OpenSupportArgs) const override;
+	virtual EAssetCommandResult OpenAssets(const FAssetOpenArgs& OpenArgs) const override;
+};
+
diff --git a/Engine/Plugins/Cameras/GameplayCameras/Source/GameplayCamerasEditor/Private/AssetTools/AssetDefinition_CameraMode.cpp b/Engine/Plugins/Cameras/GameplayCameras/Source/GameplayCamerasEditor/Private/AssetTools/AssetDefinition_CameraMode.cpp
new file mode 100644
--- /dev/null
+++ b/Engine/Plugins/Cameras/GameplayCameras/Source/GameplayCamerasEditor/Private/AssetTools/AssetDefinition_CameraMode.cpp
@@ -0,0 +1,55 @@
+// Copyright Epic Games, Inc. All Rights Reserved.
+
+#include "AssetTools/AssetDefinition_CameraMode.h"
+
+#include "IGameplayCamerasEditorModule.h"
+#include "Styles/GameplayCamerasEditorStyle.h"
+#include "Toolkits/CameraModeEditorToolkit.h"
+#include "Toolkits/IToolkit.h"
+
+#define LOCTEXT_NAMESPACE "AssetTypeActions"
+
+#include UE_INLINE_GENERATED_CPP_BY_NAME(AssetDefinition_CameraMode)
+
+FText UAssetDefinition_CameraMode::GetAssetDisplayName() const
+{
+	return NSLOCTEXT("AssetTypeActions", "AssetTypeActions_CameraMode", "Camera Mode");
+}
+
+FLinearColor UAssetDefinition_CameraMode::GetAssetColor() const
+{
+	return FLinearColor(FColor(200, 80, 80));
+}
+
+TSoftClassPtr<UObject> UAssetDefinition_CameraMode::GetAssetClass() const
+{
+	return UCameraMode::StaticClass();
+}
+
+TConstArrayView<FAssetCategoryPath> UAssetDefinition_CameraMode::GetAssetCategories() const
+{
+	static const auto Categories = { EAssetCategoryPaths::Gameplay };
+	return Categories;
+}
+
+FAssetOpenSupport UAssetDefinition_CameraMode::GetAssetOpenSupport(const FAssetOpenSupportArgs& OpenSupportArgs) const
+{
+	return FAssetOpenSupport(
+			OpenSupportArgs.OpenMethod,
+			OpenSupportArgs.OpenMethod == EAssetOpenMethod::Edit,
+			EToolkitMode::Standalone); 
+}
+
+EAssetCommandResult UAssetDefinition_CameraMode::OpenAssets(const FAssetOpenArgs& OpenArgs) const
+{
+	IGameplayCamerasEditorModule& GameplayCamerasEditorModule = FModuleManager::LoadModuleChecked<IGameplayCamerasEditorModule>("GameplayCamerasEditor");
+	for (UCameraMode* CameraMode : OpenArgs.LoadObjects<UCameraMode>())
+	{
+		GameplayCamerasEditorModule.CreateCameraModeEditor(OpenArgs.GetToolkitMode(), OpenArgs.ToolkitHost, CameraMode);
+	}	
+
+	return EAssetCommandResult::Handled;
+}
+
+#undef LOCTEXT_NAMESPACE
+
diff --git a/Engine/Plugins/Cameras/GameplayCameras/Source/GameplayCamerasEditor/Private/AssetTools/AssetDefinition_CameraMode.h b/Engine/Plugins/Cameras/GameplayCameras/Source/GameplayCamerasEditor/Private/AssetTools/AssetDefinition_CameraMode.h
new file mode 100644
--- /dev/null
+++ b/Engine/Plugins/Cameras/GameplayCameras/Source/GameplayCamerasEditor/Private/AssetTools/AssetDefinition_CameraMode.h
@@ -0,0 +1,25 @@
+// Copyright Epic Games, Inc. All Rights Reserved.
+
+#pragma once
+
+#include "Core/CameraMode.h"
+#include "AssetDefinitionDefault.h"
+
+#include "AssetDefinition_CameraMode.generated.h"
+
+UCLASS()
+class UAssetDefinition_CameraMode : public UAssetDefinitionDefault
+{
+	GENERATED_BODY()
+
+public:
+
+	// UAssetDefinition interface
+	virtual FText GetAssetDisplayName() const override;
+	virtual FLinearColor GetAssetColor() const override;
+	virtual TSoftClassPtr<UObject> GetAssetClass() const override;
+	virtual TConstArrayView<FAssetCategoryPath> GetAssetCategories() const override;
+	virtual FAssetOpenSupport GetAssetOpenSupport(const FAssetOpenSupportArgs& OpenSupportArgs) const override;
+	virtual EAssetCommandResult OpenAssets(const FAssetOpenArgs& OpenArgs) const override;
+};
+
diff --git a/Engine/Plugins/Cameras/GameplayCameras/Source/GameplayCamerasEditor/Private/Factories/CameraAssetFactory.cpp b/Engine/Plugins/Cameras/GameplayCameras/Source/GameplayCamerasEditor/Private/Factories/CameraAssetFactory.cpp
new file mode 100644
--- /dev/null
+++ b/Engine/Plugins/Cameras/GameplayCameras/Source/GameplayCamerasEditor/Private/Factories/CameraAssetFactory.cpp
@@ -0,0 +1,37 @@
+// Copyright Epic Games, Inc. All Rights Reserved.
+
+#include "Factories/CameraAssetFactory.h"
+
+#include "Core/CameraAsset.h"
+
+#include UE_INLINE_GENERATED_CPP_BY_NAME(CameraAssetFactory)
+
+#define LOCTEXT_NAMESPACE "CameraAssetFactory"
+
+UCameraAssetFactory::UCameraAssetFactory(const FObjectInitializer& ObjectInitializer)
+	: Super(ObjectInitializer)
+{
+	bCreateNew = true;
+	bEditAfterNew = true;
+	SupportedClass = UCameraAsset::StaticClass();
+}
+
+UObject* UCameraAssetFactory::FactoryCreateNew(UClass* Class, UObject* Parent, FName Name, EObjectFlags Flags, UObject* Context, FFeedbackContext* Warn)
+{
+	UCameraAsset* NewCameraAsset = NewObject<UCameraAsset>(Parent, Class, Name, Flags | RF_Transactional);
+	return NewCameraAsset;
+}
+
+bool UCameraAssetFactory::ConfigureProperties()
+{
+	return true;
+}
+
+bool UCameraAssetFactory::ShouldShowInNewMenu() const
+{
+	return true;
+}
+
+#undef LOCTEXT_NAMESPACE
+
+
diff --git a/Engine/Plugins/Cameras/GameplayCameras/Source/GameplayCamerasEditor/Private/Factories/CameraAssetFactory.h b/Engine/Plugins/Cameras/GameplayCameras/Source/GameplayCamerasEditor/Private/Factories/CameraAssetFactory.h
new file mode 100644
--- /dev/null
+++ b/Engine/Plugins/Cameras/GameplayCameras/Source/GameplayCamerasEditor/Private/Factories/CameraAssetFactory.h
@@ -0,0 +1,26 @@
+// Copyright Epic Games, Inc. All Rights Reserved.
+
+#pragma once
+
+#include "Factories/Factory.h"
+
+#include "CameraAssetFactory.generated.h"
+
+class UCameraAsset;
+
+/**
+ * Implements a factory for UCameraAsset objects.
+ */
+UCLASS(hidecategories=Object)
+class UCameraAssetFactory : public UFactory
+{
+	GENERATED_BODY()
+
+	UCameraAssetFactory(const FObjectInitializer& ObjectInitializer);
+
+	// UFactory Interface
+	virtual UObject* FactoryCreateNew(UClass* Class, UObject* Parent, FName Name, EObjectFlags Flags, UObject* Context, FFeedbackContext* Warn) override;
+	virtual bool ConfigureProperties() override;
+	virtual bool ShouldShowInNewMenu() const override;
+};
+
diff --git a/Engine/Plugins/Cameras/GameplayCameras/Source/GameplayCamerasEditor/Private/Factories/CameraModeFactory.cpp b/Engine/Plugins/Cameras/GameplayCameras/Source/GameplayCamerasEditor/Private/Factories/CameraModeFactory.cpp
new file mode 100644
--- /dev/null
+++ b/Engine/Plugins/Cameras/GameplayCameras/Source/GameplayCamerasEditor/Private/Factories/CameraModeFactory.cpp
@@ -0,0 +1,37 @@
+// Copyright Epic Games, Inc. All Rights Reserved.
+
+#include "Factories/CameraModeFactory.h"
+
+#include "Core/CameraMode.h"
+
+#include UE_INLINE_GENERATED_CPP_BY_NAME(CameraModeFactory)
+
+#define LOCTEXT_NAMESPACE "CameraModeFactory"
+
+UCameraModeFactory::UCameraModeFactory(const FObjectInitializer& ObjectInitializer)
+	: Super(ObjectInitializer)
+{
+	bCreateNew = true;
+	bEditAfterNew = true;
+	SupportedClass = UCameraMode::StaticClass();
+}
+
+UObject* UCameraModeFactory::FactoryCreateNew(UClass* Class, UObject* Parent, FName Name, EObjectFlags Flags, UObject* Context, FFeedbackContext* Warn)
+{
+	UCameraMode* NewCameraMode = NewObject<UCameraMode>(Parent, Class, Name, Flags | RF_Transactional);
+	return NewCameraMode;
+}
+
+bool UCameraModeFactory::ConfigureProperties()
+{
+	return true;
+}
+
+bool UCameraModeFactory::ShouldShowInNewMenu() const
+{
+	return true;
+}
+
+#undef LOCTEXT_NAMESPACE
+
+
diff --git a/Engine/Plugins/Cameras/GameplayCameras/Source/GameplayCamerasEditor/Private/Factories/CameraModeFactory.h b/Engine/Plugins/Cameras/GameplayCameras/Source/GameplayCamerasEditor/Private/Factories/CameraModeFactory.h
new file mode 100644
--- /dev/null
+++ b/Engine/Plugins/Cameras/GameplayCameras/Source/GameplayCamerasEditor/Private/Factories/CameraModeFactory.h
@@ -0,0 +1,26 @@
+// Copyright Epic Games, Inc. All Rights Reserved.
+
+#pragma once
+
+#include "Factories/Factory.h"
+
+#include "CameraModeFactory.generated.h"
+
+class UCameraMode;
+
+/**
+ * Implements a factory for UCameraMode objects.
+ */
+UCLASS(hidecategories=Object)
+class UCameraModeFactory : public UFactory
+{
+	GENERATED_BODY()
+
+	UCameraModeFactory(const FObjectInitializer& ObjectInitializer);
+
+	// UFactory Interface
+	virtual UObject* FactoryCreateNew(UClass* Class, UObject* Parent, FName Name, EObjectFlags Flags, UObject* Context, FFeedbackContext* Warn) override;
+	virtual bool ConfigureProperties() override;
+	virtual bool ShouldShowInNewMenu() const override;
+};
+
diff --git a/Engine/Plugins/Cameras/GameplayCameras/Source/GameplayCamerasEditor/Private/GameplayCamerasEditorModule.cpp b/Engine/Plugins/Cameras/GameplayCameras/Source/GameplayCamerasEditor/Private/GameplayCamerasEditorModule.cpp
new file mode 100644
--- /dev/null
+++ b/Engine/Plugins/Cameras/GameplayCameras/Source/GameplayCamerasEditor/Private/GameplayCamerasEditorModule.cpp
@@ -0,0 +1,105 @@
+// Copyright Epic Games, Inc. All Rights Reserved.
+
+#include "GameplayCamerasEditorSettings.h"
+
+#include "GameplayCamerasLiveEditManager.h"
+#include "IGameplayCamerasModule.h"
+#include "IGameplayCamerasEditorModule.h"
+#include "ISettingsModule.h"
+#include "Modules/ModuleManager.h"
+#include "Styles/GameplayCamerasEditorStyle.h"
+#include "Toolkits/CameraAssetEditorToolkit.h"
+#include "Toolkits/CameraModeEditorToolkit.h"
+
+#define LOCTEXT_NAMESPACE "GameplayCamerasEditor"
+
+const FName IGameplayCamerasEditorModule::GameplayCamerasEditorAppIdentifier("GameplayCamerasEditorApp");
+
+/**
+ * Implements the FGameplayCamerasEditor module.
+ */
+class FGameplayCamerasEditorModule : public IGameplayCamerasEditorModule
+{
+public:
+	FGameplayCamerasEditorModule()
+	{
+	}
+
+	virtual void StartupModule() override
+	{
+		RegisterSettings();
+		InitializeLiveEditManager();
+	}
+
+	virtual void ShutdownModule() override
+	{
+		UnregisterSettings();
+		TeardownLiveEditManager();
+	}
+
+	virtual TSharedRef<ICameraAssetEditorToolkit> CreateCameraAssetEditor(const EToolkitMode::Type Mode, const TSharedPtr<IToolkitHost>& InitToolkitHost, UCameraAsset* CameraAsset) override
+	{
+		TSharedRef<FCameraAssetEditorToolkit> Toolkit = MakeShareable(new FCameraAssetEditorToolkit(FGameplayCamerasEditorStyle::Get()));
+		Toolkit->Initialize(Mode, InitToolkitHost, CameraAsset);
+		return Toolkit;
+	}
+
+	virtual TSharedRef<ICameraModeEditorToolkit> CreateCameraModeEditor(const EToolkitMode::Type Mode, const TSharedPtr<IToolkitHost>& InitToolkitHost, UCameraMode* CameraMode) override
+	{
+		TSharedRef<FCameraModeEditorToolkit> Toolkit = MakeShareable(new FCameraModeEditorToolkit(FGameplayCamerasEditorStyle::Get()));
+		Toolkit->Initialize(Mode, InitToolkitHost, CameraMode);
+		return Toolkit;
+	}
+
+private:
+
+	void RegisterSettings()
+	{
+		ISettingsModule* SettingsModule = FModuleManager::GetModulePtr<ISettingsModule>("Settings");
+
+		if (SettingsModule != nullptr)
+		{
+			SettingsModule->RegisterSettings("Project", "Plugins", "GameplayCamerasEditor",
+				LOCTEXT("GameplayCamerasEditorProjectSettingsName", "Gameplay Cameras Editor"),
+				LOCTEXT("GameplayCamerasEditorProjectSettingsDescription", "Configure the gameplay cameras editor."),
+				GetMutableDefault<UGameplayCamerasEditorSettings>()
+			);
+		}
+	}
+
+	void UnregisterSettings()
+	{
+		ISettingsModule* SettingsModule = FModuleManager::GetModulePtr<ISettingsModule>("Settings");
+
+		if (SettingsModule != nullptr)
+		{
+			SettingsModule->UnregisterSettings("Project", "Plugins", "GameplayCamerasEditor");
+			SettingsModule->UnregisterSettings("Editor", "ContentEditors", "GameplayCamerasEditor");
+		}
+	}
+
+	void InitializeLiveEditManager()
+	{
+		LiveEditManager = MakeShared<FGameplayCamerasLiveEditManager>();
+
+		IGameplayCamerasModule& CamerasModule = FModuleManager::LoadModuleChecked<IGameplayCamerasModule>("GameplayCameras");
+		CamerasModule.SetLiveEditManager(LiveEditManager);
+	}
+
+	void TeardownLiveEditManager()
+	{
+		IGameplayCamerasModule& CamerasModule = FModuleManager::LoadModuleChecked<IGameplayCamerasModule>("GameplayCameras");
+		CamerasModule.SetLiveEditManager(nullptr);
+
+		LiveEditManager.Reset();
+	}
+
+private:
+
+	TSharedPtr<FGameplayCamerasLiveEditManager> LiveEditManager;
+};
+
+IMPLEMENT_MODULE(FGameplayCamerasEditorModule, GameplayCamerasEditor);
+
+#undef LOCTEXT_NAMESPACE
+
diff --git a/Engine/Plugins/Cameras/GameplayCameras/Source/GameplayCamerasEditor/Private/GameplayCamerasEditorSettings.cpp b/Engine/Plugins/Cameras/GameplayCameras/Source/GameplayCamerasEditor/Private/GameplayCamerasEditorSettings.cpp
new file mode 100644
--- /dev/null
+++ b/Engine/Plugins/Cameras/GameplayCameras/Source/GameplayCamerasEditor/Private/GameplayCamerasEditorSettings.cpp
@@ -0,0 +1,11 @@
+// Copyright Epic Games, Inc. All Rights Reserved.
+
+#include "GameplayCamerasEditorSettings.h"
+
+#include UE_INLINE_GENERATED_CPP_BY_NAME(GameplayCamerasEditorSettings)
+
+UGameplayCamerasEditorSettings::UGameplayCamerasEditorSettings(const FObjectInitializer& ObjectInitializer)
+	: Super(ObjectInitializer)
+{
+}
+
diff --git a/Engine/Plugins/Cameras/GameplayCameras/Source/GameplayCamerasEditor/Private/GameplayCamerasLiveEditManager.cpp b/Engine/Plugins/Cameras/GameplayCameras/Source/GameplayCamerasEditor/Private/GameplayCamerasLiveEditManager.cpp
new file mode 100644
--- /dev/null
+++ b/Engine/Plugins/Cameras/GameplayCameras/Source/GameplayCamerasEditor/Private/GameplayCamerasLiveEditManager.cpp
@@ -0,0 +1,73 @@
+// Copyright Epic Games, Inc. All Rights Reserved.
+
+#include "GameplayCamerasLiveEditManager.h"
+
+void FGameplayCamerasLiveEditManager::RegisterInstantiatedObjects(const TMap<UObject*, UObject*> InstantiatedObjects)
+{
+	for (const TPair<UObject*, UObject*>& Pair : InstantiatedObjects)
+	{
+		FInstantiationInfo& Info = Instantiations.FindOrAdd(Pair.Key);
+		Info.InstantiatedObjects.Add(Pair.Value);
+	}
+}
+
+void FGameplayCamerasLiveEditManager::ForwardPropertyChange(const UObject* Object, const FPropertyChangedEvent& PropertyChangedEvent)
+{
+	FInstantiationInfo* Info = Instantiations.Find(Object);
+	if (!Info)
+	{
+		return;
+	}
+
+	const UClass* ObjectClass = Object->GetClass();
+	const FProperty* ChangedProperty = PropertyChangedEvent.MemberProperty;
+	if (!ensure(ChangedProperty))
+	{
+		return;
+	}
+
+	bool bForwardChange = false;
+	if (ChangedProperty->IsA<FBoolProperty>()
+			|| ChangedProperty->IsA<FNumericProperty>()
+			|| ChangedProperty->IsA<FStructProperty>())
+	{
+		bForwardChange = true;
+	}
+	if (!bForwardChange)
+	{
+		return;
+	}
+
+	const void* SourceValuePtr = ChangedProperty->ContainerPtrToValuePtr<void>(Object);
+	for (auto It = Info->InstantiatedObjects.CreateIterator(); It; ++It)
+	{
+		if (UObject* Inst = It->Get())
+		{
+			ensure(Inst->GetClass() == ObjectClass);
+			void* InstValuePtr = ChangedProperty->ContainerPtrToValuePtr<void>(Inst);
+			ChangedProperty->CopyCompleteValue(InstValuePtr, SourceValuePtr);
+		}
+		else
+		{
+			It.RemoveCurrent();
+		}
+	}
+}
+
+void FGameplayCamerasLiveEditManager::ForwardPropertyChange(const UObject* SourceObject, UObject* InstantiatedObject, const FPropertyChangedEvent& PropertyChangedEvent)
+{
+}
+
+void FGameplayCamerasLiveEditManager::AddReferencedObjects(FReferenceCollector& Collector)
+{
+	for (auto& Pair : Instantiations)
+	{
+		Collector.AddReferencedObject(Pair.Key);
+	}
+}
+
+FString FGameplayCamerasLiveEditManager::GetReferencerName() const
+{
+	return TEXT("FGameplayCamerasLiveEditManager");
+}
+
diff --git a/Engine/Plugins/Cameras/GameplayCameras/Source/GameplayCamerasEditor/Private/GameplayCamerasLiveEditManager.h b/Engine/Plugins/Cameras/GameplayCameras/Source/GameplayCamerasEditor/Private/GameplayCamerasLiveEditManager.h
new file mode 100644
--- /dev/null
+++ b/Engine/Plugins/Cameras/GameplayCameras/Source/GameplayCamerasEditor/Private/GameplayCamerasLiveEditManager.h
@@ -0,0 +1,35 @@
+// Copyright Epic Games, Inc. All Rights Reserved.
+
+#pragma once
+
+#include "IGameplayCamerasLiveEditManager.h"
+
+#include "CoreTypes.h"
+#include "UObject/GCObject.h"
+#include "UObject/WeakObjectPtr.h"
+
+class FGameplayCamerasLiveEditManager 
+	: public IGameplayCamerasLiveEditManager
+	, public FGCObject
+{
+public:
+
+	virtual void RegisterInstantiatedObjects(const TMap<UObject*, UObject*> InstantiatedObjects) override;
+
+	virtual void ForwardPropertyChange(const UObject* Object, const FPropertyChangedEvent& PropertyChangedEvent) override;
+
+private:
+
+	virtual void AddReferencedObjects(FReferenceCollector& Collector) override;
+	virtual FString GetReferencerName() const override;
+
+	void ForwardPropertyChange(const UObject* SourceObject, UObject* InstantiatedObject, const FPropertyChangedEvent& PropertyChangedEvent);
+
+private:
+	struct FInstantiationInfo
+	{
+		TArray<TWeakObjectPtr<>> InstantiatedObjects;
+	};
+	TMap<TObjectPtr<const UObject>, FInstantiationInfo> Instantiations;
+};
+
diff --git a/Engine/Plugins/Cameras/GameplayCameras/Source/GameplayCamerasEditor/Private/Styles/GameplayCamerasEditorStyle.cpp b/Engine/Plugins/Cameras/GameplayCameras/Source/GameplayCamerasEditor/Private/Styles/GameplayCamerasEditorStyle.cpp
new file mode 100644
--- /dev/null
+++ b/Engine/Plugins/Cameras/GameplayCameras/Source/GameplayCamerasEditor/Private/Styles/GameplayCamerasEditorStyle.cpp
@@ -0,0 +1,6 @@
+// Copyright Epic Games, Inc. All Rights Reserved.
+
+#include "Styles/GameplayCamerasEditorStyle.h"
+
+TSharedPtr<FGameplayCamerasEditorStyle> FGameplayCamerasEditorStyle::Singleton;
+
diff --git a/Engine/Plugins/Cameras/GameplayCameras/Source/GameplayCamerasEditor/Private/Styles/GameplayCamerasEditorStyle.h b/Engine/Plugins/Cameras/GameplayCameras/Source/GameplayCamerasEditor/Private/Styles/GameplayCamerasEditorStyle.h
new file mode 100644
--- /dev/null
+++ b/Engine/Plugins/Cameras/GameplayCameras/Source/GameplayCamerasEditor/Private/Styles/GameplayCamerasEditorStyle.h
@@ -0,0 +1,52 @@
+// Copyright Epic Games, Inc. All Rights Reserved.
+
+#pragma once
+
+#include "Misc/Paths.h"
+#include "Brushes/SlateImageBrush.h"
+#include "Styling/SlateStyleRegistry.h"
+#include "Styling/SlateStyle.h"
+
+#define IMAGE_BRUSH(RelativePath, ...) FSlateImageBrush(RootToContentDir(RelativePath, TEXT(".png")), __VA_ARGS__)
+
+/**
+ * Implements the visual style of the gameplay cameras editors.
+ */
+class FGameplayCamerasEditorStyle final : public FSlateStyleSet
+{
+public:
+	
+	FGameplayCamerasEditorStyle()
+		: FSlateStyleSet("GameplayCamerasEditorStyle")
+	{
+		const FVector2D Icon16x16(16.0f, 16.0f);
+		const FVector2D Icon24x24(24.0f, 24.0f);
+		const FVector2D Icon48x48(48.0f, 48.0f);
+		const FVector2D Icon64x64(64.0f, 64.0f);
+
+		//SetContentRoot(FPaths::EnginePluginsDir() / TEXT("Cameras/GameplayCameras/Content"));
+
+		FSlateStyleRegistry::RegisterSlateStyle(*this);
+	}
+
+	virtual ~FGameplayCamerasEditorStyle()
+	{
+		FSlateStyleRegistry::UnRegisterSlateStyle(*this);
+	}
+
+	static TSharedRef<FGameplayCamerasEditorStyle> Get()
+	{
+		if (!Singleton.IsValid())
+		{
+			Singleton = MakeShareable(new FGameplayCamerasEditorStyle);
+		}
+		return Singleton.ToSharedRef();
+	}
+
+private:
+
+	static TSharedPtr<FGameplayCamerasEditorStyle> Singleton;
+};
+
+#undef IMAGE_BRUSH
+
diff --git a/Engine/Plugins/Cameras/GameplayCameras/Source/GameplayCamerasEditor/Private/Toolkits/CameraAssetEditorToolkit.cpp b/Engine/Plugins/Cameras/GameplayCameras/Source/GameplayCamerasEditor/Private/Toolkits/CameraAssetEditorToolkit.cpp
new file mode 100644
--- /dev/null
+++ b/Engine/Plugins/Cameras/GameplayCameras/Source/GameplayCamerasEditor/Private/Toolkits/CameraAssetEditorToolkit.cpp
@@ -0,0 +1,204 @@
+// Copyright Epic Games, Inc. All Rights Reserved.
+
+#include "CameraAssetEditorToolkit.h"
+
+#include "Core/CameraAsset.h"
+#include "Framework/Docking/TabManager.h"
+#include "Framework/MultiBox/MultiBoxBuilder.h"
+#include "IGameplayCamerasEditorModule.h"
+#include "IMessageLogListing.h"
+#include "MessageLogInitializationOptions.h"
+#include "MessageLogModule.h"
+#include "Modules/ModuleManager.h"
+#include "PropertyEditorModule.h"
+#include "ScopedTransaction.h"
+#include "Selection.h"
+#include "Widgets/Docking/SDockTab.h"
+#include "WorkflowOrientedApp/WorkflowTabFactory.h"
+#include "WorkflowOrientedApp/WorkflowTabManager.h"
+
+#define LOCTEXT_NAMESPACE "CameraAssetEditorToolkit"
+
+const FName FCameraAssetEditorTabs::DetailsViewTabId(TEXT("CameraAssetEditor_DetailsView"));
+
+const FName FCameraAssetEditorApplicationModes::StandardCameraAssetEditorMode(TEXT("StandardCameraAssetEditorMode"));
+
+TSharedPtr<FTabManager::FLayout> FCameraAssetEditorApplicationModes::GetDefaultEditorLayout(TSharedPtr<ICameraAssetEditorToolkit> InCameraAssetEditor)
+{
+	const TSharedRef<FTabManager::FLayout> DefaultLayout = FTabManager::NewLayout("CameraAssetEditor_Layout")
+		->AddArea
+		(
+			FTabManager::NewPrimaryArea()->SetOrientation(Orient_Vertical)
+			->Split
+			(
+				FTabManager::NewSplitter()->SetOrientation(Orient_Horizontal)->SetSizeCoefficient(0.9f)
+				->Split
+				(
+					FTabManager::NewStack()
+					->AddTab(FCameraAssetEditorTabs::DetailsViewTabId, ETabState::OpenedTab)
+					->SetForegroundTab(FCameraAssetEditorTabs::DetailsViewTabId)
+				)
+			)
+		);
+	return DefaultLayout;
+}
+
+FCameraAssetEditorApplicationMode::FCameraAssetEditorApplicationMode(TSharedPtr<ICameraAssetEditorToolkit> InCameraAssetEditor)
+	: FApplicationMode(FCameraAssetEditorApplicationModes::StandardCameraAssetEditorMode)
+{
+	WeakCameraAssetEditor = InCameraAssetEditor;
+	//CameraAssetEditorTabFactories.RegisterFactory(InCameraAssetEditor->GraphEditorTabFactoryPtr.Pin().ToSharedRef());
+	TabLayout = FCameraAssetEditorApplicationModes::GetDefaultEditorLayout(InCameraAssetEditor);
+}
+
+void FCameraAssetEditorApplicationMode::RegisterTabFactories(TSharedPtr<FTabManager> InTabManager)
+{
+	if (TSharedPtr<ICameraAssetEditorToolkit> CameraAssetEditor = WeakCameraAssetEditor.Pin())
+	{
+		// Hack?
+		StaticCastSharedPtr<FCameraAssetEditorToolkit>(CameraAssetEditor)->InternalRegisterTabSpawners(InTabManager.ToSharedRef());
+
+		CameraAssetEditor->PushTabFactories(CameraAssetEditorTabFactories);
+	}
+
+	FApplicationMode::RegisterTabFactories(InTabManager);
+}
+
+FCameraAssetEditorToolkit::FCameraAssetEditorToolkit(const TSharedRef<ISlateStyle>& InStyle)
+	: CameraAsset(nullptr)
+	, Style(InStyle)
+{
+	DocumentManager = MakeShareable(new FDocumentTracker());
+}
+
+FCameraAssetEditorToolkit::~FCameraAssetEditorToolkit()
+{
+}
+
+void FCameraAssetEditorToolkit::Initialize(const EToolkitMode::Type Mode, const TSharedPtr<IToolkitHost>& InitToolkitHost, UCameraAsset* InCameraAsset)
+{
+	CameraAsset = InCameraAsset;
+
+	CreateWidgets();
+
+	TSharedPtr<FCameraAssetEditorToolkit> ThisPtr(SharedThis(this));
+	DocumentManager->Initialize(ThisPtr);
+
+	const bool bCreateDefaultStandaloneMenu = true;
+	const bool bCreateDefaultToolbar = true;
+
+	const TSharedRef<FTabManager::FLayout> DummyLayout = FTabManager::NewLayout("NullLayout")->AddArea(FTabManager::NewPrimaryArea());
+
+	FWorkflowCentricApplication::InitAssetEditor(
+			Mode,
+			InitToolkitHost, 
+			IGameplayCamerasEditorModule::GameplayCamerasEditorAppIdentifier, 
+			DummyLayout, 
+			bCreateDefaultStandaloneMenu, 
+			bCreateDefaultToolbar,
+			CameraAsset);
+
+	
+	AddApplicationMode(
+		FCameraAssetEditorApplicationModes::StandardCameraAssetEditorMode,
+		MakeShareable(new FCameraAssetEditorApplicationMode(SharedThis(this))));
+	SetCurrentMode(FCameraAssetEditorApplicationModes::StandardCameraAssetEditorMode);
+}
+
+void FCameraAssetEditorToolkit::CreateWidgets()
+{
+	FPropertyEditorModule& PropertyEditorModule = FModuleManager::GetModuleChecked<FPropertyEditorModule>( "PropertyEditor" );
+
+	FDetailsViewArgs DetailsViewArgs;
+	DetailsViewArgs.NameAreaSettings = FDetailsViewArgs::HideNameArea;
+	DetailsViewArgs.bHideSelectionTip = true;
+	DetailsView = PropertyEditorModule.CreateDetailView( DetailsViewArgs );
+
+	TArray<TWeakObjectPtr<>> SelectedObjects { CameraAsset };
+	DetailsView->SetObjects(SelectedObjects, true);
+
+	FMessageLogModule& MessageLogModule = FModuleManager::LoadModuleChecked<FMessageLogModule>("MessageLog");
+	FMessageLogInitializationOptions LogOptions;
+	LogOptions.bShowPages = false;
+	LogOptions.bShowFilters = false;
+	LogOptions.bAllowClear = false;
+	LogOptions.MaxPageCount = 1;
+	StatsListing = MessageLogModule.CreateLogListing("CameraAssetEditorStats", LogOptions);
+
+	Stats = MessageLogModule.CreateLogListingWidget(StatsListing.ToSharedRef());
+}
+
+void FCameraAssetEditorToolkit::RegisterTabSpawners(const TSharedRef<FTabManager>& InTabManager)
+{
+	DocumentManager->SetTabManager(InTabManager);
+	FWorkflowCentricApplication::RegisterTabSpawners(InTabManager);
+}
+
+void FCameraAssetEditorToolkit::InternalRegisterTabSpawners(const TSharedRef<FTabManager>& InTabManager)
+{
+	FAssetEditorToolkit::UnregisterTabSpawners(InTabManager);
+
+	WorkspaceMenuCategory = InTabManager->AddLocalWorkspaceMenuCategory(LOCTEXT("WorkspaceMenu_CameraAssetEditor", "Camera Asset Editor"));
+	auto WorkspaceMenuCategoryRef = WorkspaceMenuCategory.ToSharedRef();
+
+	InTabManager->RegisterTabSpawner(FCameraAssetEditorTabs::DetailsViewTabId, FOnSpawnTab::CreateSP(this, &FCameraAssetEditorToolkit::SpawnTab_DetailsView))
+		.SetDisplayName(LOCTEXT("DetailsViewTabTitle", "Details"))
+		.SetGroup(WorkspaceMenuCategoryRef)
+		.SetIcon(FSlateIcon(FAppStyle::GetAppStyleSetName(), "LevelEditor.Tabs.Details"));
+}
+
+void FCameraAssetEditorToolkit::UnregisterTabSpawners(const TSharedRef<FTabManager>& InTabManager)
+{
+	FWorkflowCentricApplication::UnregisterTabSpawners(InTabManager);
+
+	InTabManager->UnregisterTabSpawner(FCameraAssetEditorTabs::DetailsViewTabId);
+}
+
+TSharedRef<SDockTab> FCameraAssetEditorToolkit::SpawnTab_DetailsView(const FSpawnTabArgs& Args)
+{
+	TSharedPtr<SDockTab> DetailsTab = SNew(SDockTab)
+		.Label(LOCTEXT("DetailsViewTabTitle", "Details"))
+		[
+			DetailsView.ToSharedRef()
+		];
+	DetailsViewTab = DetailsTab;
+	return DetailsTab.ToSharedRef();
+}
+
+FText FCameraAssetEditorToolkit::GetBaseToolkitName() const
+{
+	return LOCTEXT("AppLabel", "Camera Asset");
+}
+
+FName FCameraAssetEditorToolkit::GetToolkitFName() const
+{
+	static FName SequencerName("CameraAssetEditor");
+	return SequencerName;
+}
+
+FString FCameraAssetEditorToolkit::GetWorldCentricTabPrefix() const
+{
+	return LOCTEXT("WorldCentricTabPrefix", "Camera Asset ").ToString();
+}
+
+FText FCameraAssetEditorToolkit::GetTabSuffix() const
+{
+	if (CameraAsset)
+	{
+		const bool bIsDirty = CameraAsset->GetOutermost()->IsDirty();
+		if (bIsDirty)
+		{
+			return LOCTEXT("TabSuffixAsterix", "*");
+		}
+	}
+
+	return FText::GetEmpty();
+}
+
+FLinearColor FCameraAssetEditorToolkit::GetWorldCentricTabColorScale() const
+{
+	return FLinearColor(0.7, 0.0f, 0.0f, 0.5f);
+}
+
+#undef LOCTEXT_NAMESPACE
+
diff --git a/Engine/Plugins/Cameras/GameplayCameras/Source/GameplayCamerasEditor/Private/Toolkits/CameraAssetEditorToolkit.h b/Engine/Plugins/Cameras/GameplayCameras/Source/GameplayCamerasEditor/Private/Toolkits/CameraAssetEditorToolkit.h
new file mode 100644
--- /dev/null
+++ b/Engine/Plugins/Cameras/GameplayCameras/Source/GameplayCamerasEditor/Private/Toolkits/CameraAssetEditorToolkit.h
@@ -0,0 +1,98 @@
+// Copyright Epic Games, Inc. All Rights Reserved.
+
+#pragma once
+
+#include "CoreMinimal.h"
+#include "Toolkits/IToolkitHost.h"
+#include "Toolkits/ICameraAssetEditorToolkit.h"
+
+class UCameraAsset;
+class FDocumentTracker;
+class FToolBarBuilder;
+class FAssetDragDropOp;
+class FClassDragDropOp;
+class FActorDragDropGraphEdOp;
+
+/**
+ * Implements an editor toolkit for camera modes.
+ */
+class FCameraAssetEditorToolkit : public ICameraAssetEditorToolkit, public FGCObject
+{
+public:
+
+	/**
+	 * Creates and initializes a new instance.
+	 *
+	 * @param InStyle The style set to use.
+	 */
+	FCameraAssetEditorToolkit(const TSharedRef<ISlateStyle>& InStyle);
+
+	/** Virtual destructor */
+	virtual ~FCameraAssetEditorToolkit();
+
+public:
+
+	/**
+	 * Initialize this asset editor.
+	 *
+	 * @param Mode Asset editing mode for this editor (standalone or world-centric).
+	 * @param InitToolkitHost When Mode is WorldCentric, this is the level editor instance to spawn this editor within.
+	 * @param CameraAsset The camera mode to edit.
+	 * @param TrackEditorDelegates Delegates to call to create auto-key handlers for this sequencer.
+	 */
+	void Initialize(const EToolkitMode::Type Mode, const TSharedPtr<IToolkitHost>& InitToolkitHost, UCameraAsset* CameraAsset);
+
+	void InternalRegisterTabSpawners(const TSharedRef<FTabManager>& InTabManager);
+
+public:
+
+	//~ FAssetEditorToolkit interface
+	virtual void AddReferencedObjects(FReferenceCollector& Collector) override
+	{
+		Collector.AddReferencedObject(CameraAsset);
+	}
+	virtual FString GetReferencerName() const override
+	{
+		return TEXT("FCameraAssetEditorToolkit");
+	}
+
+	virtual void RegisterTabSpawners(const TSharedRef<FTabManager>& InTabManager) override;
+	virtual void UnregisterTabSpawners(const TSharedRef<FTabManager>& InTabManager) override;
+
+public:
+
+	//~ IToolkit interface
+	virtual FText GetBaseToolkitName() const override;
+	virtual FName GetToolkitFName() const override;
+	virtual FString GetWorldCentricTabPrefix() const override;
+	virtual FText GetTabSuffix() const override;
+	virtual FLinearColor GetWorldCentricTabColorScale() const override;
+
+private:
+
+	void CreateWidgets();
+
+	TSharedRef<SDockTab> SpawnTab_DetailsView(const FSpawnTabArgs& Args);
+
+private:
+
+	/** Camera mode for our edit operation. */
+	TObjectPtr<UCameraAsset> CameraAsset;
+
+	/** Pointer to the style set to use for toolkits. */
+	TSharedRef<ISlateStyle> Style;
+
+	/** Document manager for workflow tabs */
+	TSharedPtr<FDocumentTracker> DocumentManager;
+
+	/** Details view */
+	TSharedPtr<class IDetailsView> DetailsView;
+
+	/** Tab that holds the details panel */
+	TWeakPtr<SDockTab> DetailsViewTab;
+
+	TSharedPtr<class SWidget> Stats;
+
+	TSharedPtr<class IMessageLogListing> StatsListing;
+};
+
diff --git a/Engine/Plugins/Cameras/GameplayCameras/Source/GameplayCamerasEditor/Private/Toolkits/CameraModeEditorToolkit.cpp b/Engine/Plugins/Cameras/GameplayCameras/Source/GameplayCamerasEditor/Private/Toolkits/CameraModeEditorToolkit.cpp
new file mode 100644
--- /dev/null
+++ b/Engine/Plugins/Cameras/GameplayCameras/Source/GameplayCamerasEditor/Private/Toolkits/CameraModeEditorToolkit.cpp
@@ -0,0 +1,204 @@
+// Copyright Epic Games, Inc. All Rights Reserved.
+
+#include "CameraModeEditorToolkit.h"
+
+#include "Core/CameraMode.h"
+#include "Framework/Docking/TabManager.h"
+#include "Framework/MultiBox/MultiBoxBuilder.h"
+#include "IGameplayCamerasEditorModule.h"
+#include "IMessageLogListing.h"
+#include "MessageLogInitializationOptions.h"
+#include "MessageLogModule.h"
+#include "Modules/ModuleManager.h"
+#include "PropertyEditorModule.h"
+#include "ScopedTransaction.h"
+#include "Selection.h"
+#include "Widgets/Docking/SDockTab.h"
+#include "WorkflowOrientedApp/WorkflowTabFactory.h"
+#include "WorkflowOrientedApp/WorkflowTabManager.h"
+
+#define LOCTEXT_NAMESPACE "CameraModeEditorToolkit"
+
+const FName FCameraModeEditorTabs::DetailsViewTabId(TEXT("CameraModeEditor_DetailsView"));
+
+const FName FCameraModeEditorApplicationModes::StandardCameraModeEditorMode(TEXT("StandardCameraModeEditorMode"));
+
+TSharedPtr<FTabManager::FLayout> FCameraModeEditorApplicationModes::GetDefaultEditorLayout(TSharedPtr<ICameraModeEditorToolkit> InCameraModeEditor)
+{
+	const TSharedRef<FTabManager::FLayout> DefaultLayout = FTabManager::NewLayout("CameraModeEditor_Layout")
+		->AddArea
+		(
+			FTabManager::NewPrimaryArea()->SetOrientation(Orient_Vertical)
+			->Split
+			(
+				FTabManager::NewSplitter()->SetOrientation(Orient_Horizontal)->SetSizeCoefficient(0.9f)
+				->Split
+				(
+					FTabManager::NewStack()
+					->AddTab(FCameraModeEditorTabs::DetailsViewTabId, ETabState::OpenedTab)
+					->SetForegroundTab(FCameraModeEditorTabs::DetailsViewTabId)
+				)
+			)
+		);
+	return DefaultLayout;
+}
+
+FCameraModeEditorApplicationMode::FCameraModeEditorApplicationMode(TSharedPtr<ICameraModeEditorToolkit> InCameraModeEditor)
+	: FApplicationMode(FCameraModeEditorApplicationModes::StandardCameraModeEditorMode)
+{
+	WeakCameraModeEditor = InCameraModeEditor;
+	//CameraModeEditorTabFactories.RegisterFactory(InCameraModeEditor->GraphEditorTabFactoryPtr.Pin().ToSharedRef());
+	TabLayout = FCameraModeEditorApplicationModes::GetDefaultEditorLayout(InCameraModeEditor);
+}
+
+void FCameraModeEditorApplicationMode::RegisterTabFactories(TSharedPtr<FTabManager> InTabManager)
+{
+	if (TSharedPtr<ICameraModeEditorToolkit> CameraModeEditor = WeakCameraModeEditor.Pin())
+	{
+		// Hack?
+		StaticCastSharedPtr<FCameraModeEditorToolkit>(CameraModeEditor)->InternalRegisterTabSpawners(InTabManager.ToSharedRef());
+
+		CameraModeEditor->PushTabFactories(CameraModeEditorTabFactories);
+	}
+
+	FApplicationMode::RegisterTabFactories(InTabManager);
+}
+
+FCameraModeEditorToolkit::FCameraModeEditorToolkit(const TSharedRef<ISlateStyle>& InStyle)
+	: CameraMode(nullptr)
+	, Style(InStyle)
+{
+	DocumentManager = MakeShareable(new FDocumentTracker());
+}
+
+FCameraModeEditorToolkit::~FCameraModeEditorToolkit()
+{
+}
+
+void FCameraModeEditorToolkit::Initialize(const EToolkitMode::Type Mode, const TSharedPtr<IToolkitHost>& InitToolkitHost, UCameraMode* InCameraMode)
+{
+	CameraMode = InCameraMode;
+
+	CreateWidgets();
+
+	TSharedPtr<FCameraModeEditorToolkit> ThisPtr(SharedThis(this));
+	DocumentManager->Initialize(ThisPtr);
+
+	const bool bCreateDefaultStandaloneMenu = true;
+	const bool bCreateDefaultToolbar = true;
+
+	const TSharedRef<FTabManager::FLayout> DummyLayout = FTabManager::NewLayout("NullLayout")->AddArea(FTabManager::NewPrimaryArea());
+
+	FWorkflowCentricApplication::InitAssetEditor(
+			Mode,
+			InitToolkitHost, 
+			IGameplayCamerasEditorModule::GameplayCamerasEditorAppIdentifier, 
+			DummyLayout, 
+			bCreateDefaultStandaloneMenu, 
+			bCreateDefaultToolbar,
+			CameraMode);
+
+	
+	AddApplicationMode(
+		FCameraModeEditorApplicationModes::StandardCameraModeEditorMode,
+		MakeShareable(new FCameraModeEditorApplicationMode(SharedThis(this))));
+	SetCurrentMode(FCameraModeEditorApplicationModes::StandardCameraModeEditorMode);
+}
+
+void FCameraModeEditorToolkit::CreateWidgets()
+{
+	FPropertyEditorModule& PropertyEditorModule = FModuleManager::GetModuleChecked<FPropertyEditorModule>( "PropertyEditor" );
+
+	FDetailsViewArgs DetailsViewArgs;
+	DetailsViewArgs.NameAreaSettings = FDetailsViewArgs::HideNameArea;
+	DetailsViewArgs.bHideSelectionTip = true;
+	DetailsView = PropertyEditorModule.CreateDetailView( DetailsViewArgs );
+
+	TArray<TWeakObjectPtr<>> SelectedObjects { CameraMode };
+	DetailsView->SetObjects(SelectedObjects, true);
+
+	FMessageLogModule& MessageLogModule = FModuleManager::LoadModuleChecked<FMessageLogModule>("MessageLog");
+	FMessageLogInitializationOptions LogOptions;
+	LogOptions.bShowPages = false;
+	LogOptions.bShowFilters = false;
+	LogOptions.bAllowClear = false;
+	LogOptions.MaxPageCount = 1;
+	StatsListing = MessageLogModule.CreateLogListing("CameraModeEditorStats", LogOptions);
+
+	Stats = MessageLogModule.CreateLogListingWidget(StatsListing.ToSharedRef());
+}
+
+void FCameraModeEditorToolkit::RegisterTabSpawners(const TSharedRef<FTabManager>& InTabManager)
+{
+	DocumentManager->SetTabManager(InTabManager);
+	FWorkflowCentricApplication::RegisterTabSpawners(InTabManager);
+}
+
+void FCameraModeEditorToolkit::InternalRegisterTabSpawners(const TSharedRef<FTabManager>& InTabManager)
+{
+	FAssetEditorToolkit::UnregisterTabSpawners(InTabManager);
+
+	WorkspaceMenuCategory = InTabManager->AddLocalWorkspaceMenuCategory(LOCTEXT("WorkspaceMenu_CameraModeEditor", "Camera Mode Editor"));
+	auto WorkspaceMenuCategoryRef = WorkspaceMenuCategory.ToSharedRef();
+
+	InTabManager->RegisterTabSpawner(FCameraModeEditorTabs::DetailsViewTabId, FOnSpawnTab::CreateSP(this, &FCameraModeEditorToolkit::SpawnTab_DetailsView))
+		.SetDisplayName(LOCTEXT("DetailsViewTabTitle", "Details"))
+		.SetGroup(WorkspaceMenuCategoryRef)
+		.SetIcon(FSlateIcon(FAppStyle::GetAppStyleSetName(), "LevelEditor.Tabs.Details"));
+}
+
+void FCameraModeEditorToolkit::UnregisterTabSpawners(const TSharedRef<FTabManager>& InTabManager)
+{
+	FWorkflowCentricApplication::UnregisterTabSpawners(InTabManager);
+
+	InTabManager->UnregisterTabSpawner(FCameraModeEditorTabs::DetailsViewTabId);
+}
+
+TSharedRef<SDockTab> FCameraModeEditorToolkit::SpawnTab_DetailsView(const FSpawnTabArgs& Args)
+{
+	TSharedPtr<SDockTab> DetailsTab = SNew(SDockTab)
+		.Label(LOCTEXT("DetailsViewTabTitle", "Details"))
+		[
+			DetailsView.ToSharedRef()
+		];
+	DetailsViewTab = DetailsTab;
+	return DetailsTab.ToSharedRef();
+}
+
+FText FCameraModeEditorToolkit::GetBaseToolkitName() const
+{
+	return LOCTEXT("AppLabel", "Camera Mode");
+}
+
+FName FCameraModeEditorToolkit::GetToolkitFName() const
+{
+	static FName SequencerName("CameraModeEditor");
+	return SequencerName;
+}
+
+FString FCameraModeEditorToolkit::GetWorldCentricTabPrefix() const
+{
+	return LOCTEXT("WorldCentricTabPrefix", "Camera Mode ").ToString();
+}
+
+FText FCameraModeEditorToolkit::GetTabSuffix() const
+{
+	if (CameraMode)
+	{
+		const bool bIsDirty = CameraMode->GetOutermost()->IsDirty();
+		if (bIsDirty)
+		{
+			return LOCTEXT("TabSuffixAsterix", "*");
+		}
+	}
+
+	return FText::GetEmpty();
+}
+
+FLinearColor FCameraModeEditorToolkit::GetWorldCentricTabColorScale() const
+{
+	return FLinearColor(0.7, 0.0f, 0.0f, 0.5f);
+}
+
+#undef LOCTEXT_NAMESPACE
+
diff --git a/Engine/Plugins/Cameras/GameplayCameras/Source/GameplayCamerasEditor/Private/Toolkits/CameraModeEditorToolkit.h b/Engine/Plugins/Cameras/GameplayCameras/Source/GameplayCamerasEditor/Private/Toolkits/CameraModeEditorToolkit.h
new file mode 100644
--- /dev/null
+++ b/Engine/Plugins/Cameras/GameplayCameras/Source/GameplayCamerasEditor/Private/Toolkits/CameraModeEditorToolkit.h
@@ -0,0 +1,98 @@
+// Copyright Epic Games, Inc. All Rights Reserved.
+
+#pragma once
+
+#include "CoreMinimal.h"
+#include "Toolkits/IToolkitHost.h"
+#include "Toolkits/ICameraModeEditorToolkit.h"
+
+class UCameraMode;
+class FDocumentTracker;
+class FToolBarBuilder;
+class FAssetDragDropOp;
+class FClassDragDropOp;
+class FActorDragDropGraphEdOp;
+
+/**
+ * Implements an editor toolkit for camera modes.
+ */
+class FCameraModeEditorToolkit : public ICameraModeEditorToolkit, public FGCObject
+{
+public:
+
+	/**
+	 * Creates and initializes a new instance.
+	 *
+	 * @param InStyle The style set to use.
+	 */
+	FCameraModeEditorToolkit(const TSharedRef<ISlateStyle>& InStyle);
+
+	/** Virtual destructor */
+	virtual ~FCameraModeEditorToolkit();
+
+public:
+
+	/**
+	 * Initialize this asset editor.
+	 *
+	 * @param Mode Asset editing mode for this editor (standalone or world-centric).
+	 * @param InitToolkitHost When Mode is WorldCentric, this is the level editor instance to spawn this editor within.
+	 * @param CameraMode The camera mode to edit.
+	 * @param TrackEditorDelegates Delegates to call to create auto-key handlers for this sequencer.
+	 */
+	void Initialize(const EToolkitMode::Type Mode, const TSharedPtr<IToolkitHost>& InitToolkitHost, UCameraMode* CameraMode);
+
+	void InternalRegisterTabSpawners(const TSharedRef<FTabManager>& InTabManager);
+
+public:
+
+	//~ FAssetEditorToolkit interface
+	virtual void AddReferencedObjects(FReferenceCollector& Collector) override
+	{
+		Collector.AddReferencedObject(CameraMode);
+	}
+	virtual FString GetReferencerName() const override
+	{
+		return TEXT("FCameraModeEditorToolkit");
+	}
+
+	virtual void RegisterTabSpawners(const TSharedRef<FTabManager>& InTabManager) override;
+	virtual void UnregisterTabSpawners(const TSharedRef<FTabManager>& InTabManager) override;
+
+public:
+
+	//~ IToolkit interface
+	virtual FText GetBaseToolkitName() const override;
+	virtual FName GetToolkitFName() const override;
+	virtual FString GetWorldCentricTabPrefix() const override;
+	virtual FText GetTabSuffix() const override;
+	virtual FLinearColor GetWorldCentricTabColorScale() const override;
+
+private:
+
+	void CreateWidgets();
+
+	TSharedRef<SDockTab> SpawnTab_DetailsView(const FSpawnTabArgs& Args);
+
+private:
+
+	/** Camera mode for our edit operation. */
+	TObjectPtr<UCameraMode> CameraMode;
+
+	/** Pointer to the style set to use for toolkits. */
+	TSharedRef<ISlateStyle> Style;
+
+	/** Document manager for workflow tabs */
+	TSharedPtr<FDocumentTracker> DocumentManager;
+
+	/** Details view */
+	TSharedPtr<class IDetailsView> DetailsView;
+
+	/** Tab that holds the details panel */
+	TWeakPtr<SDockTab> DetailsViewTab;
+
+	TSharedPtr<class SWidget> Stats;
+
+	TSharedPtr<class IMessageLogListing> StatsListing;
+};
+
diff --git a/Engine/Plugins/Cameras/GameplayCameras/Source/GameplayCamerasEditor/Public/GameplayCamerasEditorSettings.h b/Engine/Plugins/Cameras/GameplayCameras/Source/GameplayCamerasEditor/Public/GameplayCamerasEditorSettings.h
new file mode 100644
--- /dev/null
+++ b/Engine/Plugins/Cameras/GameplayCameras/Source/GameplayCamerasEditor/Public/GameplayCamerasEditorSettings.h
@@ -0,0 +1,22 @@
+// Copyright Epic Games, Inc. All Rights Reserved.
+
+#pragma once
+
+#include "CoreMinimal.h"
+#include "UObject/ObjectMacros.h"
+#include "UObject/Object.h"
+#include "UObject/Package.h"
+#include "FrameNumberDisplayFormat.h"
+
+#include "GameplayCamerasEditorSettings.generated.h"
+
+UCLASS(config=EditorPerProjectUserSettings)
+class GAMEPLAYCAMERASEDITOR_API UGameplayCamerasEditorSettings : public UObject
+{
+	GENERATED_BODY()
+
+public:
+
+	UGameplayCamerasEditorSettings(const FObjectInitializer& ObjectInitializer);	
+};
+
diff --git a/Engine/Plugins/Cameras/GameplayCameras/Source/GameplayCamerasEditor/Public/IGameplayCamerasEditorModule.h b/Engine/Plugins/Cameras/GameplayCameras/Source/GameplayCamerasEditor/Public/IGameplayCamerasEditorModule.h
new file mode 100644
--- /dev/null
+++ b/Engine/Plugins/Cameras/GameplayCameras/Source/GameplayCamerasEditor/Public/IGameplayCamerasEditorModule.h
@@ -0,0 +1,26 @@
+// Copyright Epic Games, Inc. All Rights Reserved.
+
+#pragma once
+
+#include "CoreMinimal.h"
+#include "Modules/ModuleInterface.h"
+#include "Toolkits/AssetEditorToolkit.h"
+
+class ICameraAssetEditorToolkit;
+class ICameraModeEditorToolkit;
+class UCameraAsset;
+class UCameraMode;
+
+class IGameplayCamerasEditorModule : public IModuleInterface
+{
+public:
+
+	static const FName GameplayCamerasEditorAppIdentifier;
+
+	virtual ~IGameplayCamerasEditorModule() = default;
+
+	virtual TSharedRef<ICameraAssetEditorToolkit> CreateCameraAssetEditor(const EToolkitMode::Type Mode, const TSharedPtr<IToolkitHost>& InitToolkitHost, UCameraAsset* CameraAsset) = 0;
+
+	virtual TSharedRef<ICameraModeEditorToolkit> CreateCameraModeEditor(const EToolkitMode::Type Mode, const TSharedPtr<IToolkitHost>& InitToolkitHost, UCameraMode* CameraMode) = 0;
+};
+
diff --git a/Engine/Plugins/Cameras/GameplayCameras/Source/GameplayCamerasEditor/Public/Toolkits/ICameraAssetEditorToolkit.h b/Engine/Plugins/Cameras/GameplayCameras/Source/GameplayCamerasEditor/Public/Toolkits/ICameraAssetEditorToolkit.h
new file mode 100644
--- /dev/null
+++ b/Engine/Plugins/Cameras/GameplayCameras/Source/GameplayCamerasEditor/Public/Toolkits/ICameraAssetEditorToolkit.h
@@ -0,0 +1,45 @@
+// Copyright Epic Games, Inc. All Rights Reserved.
+
+#pragma once
+
+#include "CoreMinimal.h"
+#include "Toolkits/AssetEditorToolkit.h"
+#include "WorkflowOrientedApp/ApplicationMode.h"
+#include "WorkflowOrientedApp/WorkflowCentricApplication.h"
+#include "WorkflowOrientedApp/WorkflowTabManager.h"
+
+class ICameraAssetEditorToolkit : public FWorkflowCentricApplication
+{
+};
+
+struct FCameraAssetEditorTabs
+{
+	static const FName DetailsViewTabId;
+};
+
+struct FCameraAssetEditorApplicationModes
+{
+	static const FName StandardCameraAssetEditorMode;
+
+	static TSharedPtr<FTabManager::FLayout> GetDefaultEditorLayout(TSharedPtr<ICameraAssetEditorToolkit> InCameraAssetEditor);
+
+private:
+
+	FCameraAssetEditorApplicationModes() = default;
+};
+
+class FCameraAssetEditorApplicationMode : public FApplicationMode
+{
+public:
+
+	FCameraAssetEditorApplicationMode(TSharedPtr<ICameraAssetEditorToolkit> InCameraAssetEditor);
+
+	virtual void RegisterTabFactories(TSharedPtr<FTabManager> InTabManager) override;
+
+protected:
+
+	TWeakPtr<ICameraAssetEditorToolkit> WeakCameraAssetEditor;
+
+	FWorkflowAllowedTabSet CameraAssetEditorTabFactories;
+};
+
diff --git a/Engine/Plugins/Cameras/GameplayCameras/Source/GameplayCamerasEditor/Public/Toolkits/ICameraModeEditorToolkit.h b/Engine/Plugins/Cameras/GameplayCameras/Source/GameplayCamerasEditor/Public/Toolkits/ICameraModeEditorToolkit.h
new file mode 100644
--- /dev/null
+++ b/Engine/Plugins/Cameras/GameplayCameras/Source/GameplayCamerasEditor/Public/Toolkits/ICameraModeEditorToolkit.h
@@ -0,0 +1,45 @@
+// Copyright Epic Games, Inc. All Rights Reserved.
+
+#pragma once
+
+#include "CoreMinimal.h"
+#include "Toolkits/AssetEditorToolkit.h"
+#include "WorkflowOrientedApp/ApplicationMode.h"
+#include "WorkflowOrientedApp/WorkflowCentricApplication.h"
+#include "WorkflowOrientedApp/WorkflowTabManager.h"
+
+class ICameraModeEditorToolkit : public FWorkflowCentricApplication
+{
+};
+
+struct FCameraModeEditorTabs
+{
+	static const FName DetailsViewTabId;
+};
+
+struct FCameraModeEditorApplicationModes
+{
+	static const FName StandardCameraModeEditorMode;
+
+	static TSharedPtr<FTabManager::FLayout> GetDefaultEditorLayout(TSharedPtr<ICameraModeEditorToolkit> InCameraModeEditor);
+
+private:
+
+	FCameraModeEditorApplicationModes() = default;
+};
+
+class FCameraModeEditorApplicationMode : public FApplicationMode
+{
+public:
+
+	FCameraModeEditorApplicationMode(TSharedPtr<ICameraModeEditorToolkit> InCameraModeEditor);
+
+	virtual void RegisterTabFactories(TSharedPtr<FTabManager> InTabManager) override;
+
+protected:
+
+	TWeakPtr<ICameraModeEditorToolkit> WeakCameraModeEditor;
+
+	FWorkflowAllowedTabSet CameraModeEditorTabFactories;
+};
+
