Cameras: implement recycling pool for camera modes' node trees, and fix related thing on the live-edit manager

Instantiating camera modes by cloning their node hierarchy can sometimes lead to bad performance depending on what else is going on. To reduce the chance of this happening, we add a recycling pool so that we can re-use a previously instantiated hierarchy instead of cloning a brand new one.

- Instantiated node trees are recycled instead of always cloned.
- This requires being able to call a method to reset the state of nodes on recycling.
- The camera runtime instantiator now stores state (the recycling pool) so needs to be owned by the system.
- Camera nodes know if they are instantiated, or if they have instantiations, in editor builds.
- Live edit manager only holds weakly on all objects, and does clean up its state after GC passes.

diff --git a/Engine/Plugins/Cameras/GameplayCameras/Source/GameplayCameras/Private/Core/BlendStackCameraNode.cpp b/Engine/Plugins/Cameras/GameplayCameras/Source/GameplayCameras/Private/Core/BlendStackCameraNode.cpp
--- a/Engine/Plugins/Cameras/GameplayCameras/Source/GameplayCameras/Private/Core/BlendStackCameraNode.cpp
+++ b/Engine/Plugins/Cameras/GameplayCameras/Source/GameplayCameras/Private/Core/BlendStackCameraNode.cpp
@@ -8,6 +8,7 @@
 #include "Core/CameraEvaluationContext.h"
 #include "Core/CameraMode.h"
 #include "Core/CameraRuntimeInstantiator.h"
+#include "Core/CameraSystemEvaluator.h"
 #include "Nodes/Blends/PopBlendCameraNode.h"
 
 #include UE_INLINE_GENERATED_CPP_BY_NAME(BlendStackCameraNode)
@@ -30,27 +31,34 @@ void UBlendStackCameraNode::Push(const FBlendStackCameraPushParams& Params)
 	// Create the new root node. Instantiated objects will go inside it.
 	UBlendStackRootCameraNode* EntryRootNode = NewObject<UBlendStackRootCameraNode>(this, NAME_None);
 
-	FCameraRuntimeInstantiationParams InstParams;
-	InstParams.InstantiationOuter = EntryRootNode;
-
 	// Instantiate the camera mode's node tree.
-	UCameraNode* ModeRootNode = FCameraRuntimeInstantiator::InstantiateCameraNodeTree(
-			Params.CameraMode->RootNode, InstParams);
-	EntryRootNode->RootNode = ModeRootNode;
+	FCameraRuntimeInstantiationParams NodeTreeInstParams;
+	NodeTreeInstParams.InstantiationOuter = EntryRootNode;
+	NodeTreeInstParams.bAllowRecyling = true;
+
+	FCameraRuntimeInstantiator& Instantiator = Params.Evaluator->GetRuntimeInstantiator();
+	UCameraNode* ModeRootNode = Instantiator.InstantiateCameraNodeTree(Params.CameraMode->RootNode, NodeTreeInstParams);
 
 	// Find a transition and instantiate its blend. If not transition is found,
 	// make a camera cut transition.
+	UBlendCameraNode* Blend = nullptr;
 	if (const FCameraModeTransition* Transition = FindTransition(Params))
 	{
+		FCameraRuntimeInstantiationParams BlendInstParams;
+		BlendInstParams.InstantiationOuter = EntryRootNode;
+		// No recycling on the blend.
+
 		UBlendCameraNode* ModeBlend = CastChecked<UBlendCameraNode>(
-				FCameraRuntimeInstantiator::InstantiateCameraNodeTree(Transition->Blend, InstParams));
-		EntryRootNode->Blend = ModeBlend;
+				Instantiator.InstantiateCameraNodeTree(Transition->Blend, BlendInstParams));
+		Blend = ModeBlend;
 	}
 	else
 	{
-		EntryRootNode->Blend = NewObject<UPopBlendCameraNode>(EntryRootNode, NAME_None);
+		Blend = NewObject<UPopBlendCameraNode>(EntryRootNode, NAME_None);
 	}
 
+	EntryRootNode->Initialize(Blend, ModeRootNode);
+
 	// Make a new entry and add it to the stack.
 	FCameraModeEntry NewEntry;
 	NewEntry.EvaluationContext = Params.EvaluationContext;
@@ -86,8 +94,7 @@ void UBlendStackCameraNode::OnRun(const FCameraNodeRunParams& Params, FCameraNod
 	// Start by evaluating all the root nodes in the stack.
 	for (FCameraModeEntry& Entry : Entries)
 	{
-		FCameraNodeRunParams CurParams;
-		CurParams.DeltaTime = Params.DeltaTime;
+		FCameraNodeRunParams CurParams(Params);
 		CurParams.bIsFirstFrame = Entry.bIsFirstFrame;
 
 		FCameraNodeRunResult& CurResult(Entry.Result);
@@ -143,14 +150,13 @@ void UBlendStackCameraNode::OnRun(const FCameraNodeRunParams& Params, FCameraNod
 
 		const FCameraPoseFlags ChangedFlags(CurResult.CameraPose.GetChangedFlags());
 
-		FCameraNodeRunParams CurParams;
-		CurParams.DeltaTime = Params.DeltaTime;
+		FCameraNodeRunParams CurParams(Params);
 		CurParams.bIsFirstFrame = Entry.bIsFirstFrame;
 		FCameraNodeBlendParams BlendParams(CurParams, CurResult);
 
 		FCameraNodeBlendResult BlendResult(OutResult);
 
-		UBlendCameraNode* EntryBlend = Entry.RootNode->Blend;
+		UBlendCameraNode* EntryBlend = Entry.RootNode->GetBlend();
 		if (EntryBlend)
 		{
 			EntryBlend->BlendResults(BlendParams, BlendResult);
@@ -175,7 +181,23 @@ void UBlendStackCameraNode::OnRun(const FCameraNodeRunParams& Params, FCameraNod
 	// Pop out camera modes that have been blended out.
 	if (bAutoPop && PopEntriesBelow != INDEX_NONE)
 	{
-		Entries.RemoveAt(0, PopEntriesBelow, EAllowShrinking::No);
+		FCameraRuntimeInstantiator& Instantiator = Params.Evaluator->GetRuntimeInstantiator();
+		for (int32 Index = 0; Index < PopEntriesBelow; ++Index)
+		{
+			FCameraModeEntry& Entry = Entries[0];
+			if (!Entry.bIsFrozen)
+			{
+				// Recycle the camera mode's node hierarchy if possible.
+				// Reset all nodes in it so they are back to their default state, ready to be
+				// re-used later.
+				TObjectPtr<const UCameraNode> OriginalRootNode = Entry.OriginalCameraMode->RootNode;
+				TObjectPtr<UCameraNode> RecycledRootNode = Entry.RootNode->GetRootNode();
+				Entry.RootNode->Reset(FCameraNodeResetParams());
+
+				Instantiator.RecycleInstantiatedObject(OriginalRootNode, RecycledRootNode);
+			}
+			Entries.RemoveAt(0);
+		}
 	}
 
 	// Reset first frame flags.
diff --git a/Engine/Plugins/Cameras/GameplayCameras/Source/GameplayCameras/Private/Core/BlendStackRootCameraNode.cpp b/Engine/Plugins/Cameras/GameplayCameras/Source/GameplayCameras/Private/Core/BlendStackRootCameraNode.cpp
--- a/Engine/Plugins/Cameras/GameplayCameras/Source/GameplayCameras/Private/Core/BlendStackRootCameraNode.cpp
+++ b/Engine/Plugins/Cameras/GameplayCameras/Source/GameplayCameras/Private/Core/BlendStackRootCameraNode.cpp
@@ -7,11 +7,24 @@
 
 #include UE_INLINE_GENERATED_CPP_BY_NAME(BlendStackRootCameraNode)
 
+void UBlendStackRootCameraNode::Initialize(UBlendCameraNode* InBlend, UCameraNode* InRootNode)
+{
+	Blend = InBlend;
+	RootNode = InRootNode;
+
+	TreeCache.Build(RootNode);
+}
+
 FCameraNodeChildrenView UBlendStackRootCameraNode::OnGetChildren()
 {
 	return FCameraNodeChildrenView({ Blend, RootNode });
 }
 
+void UBlendStackRootCameraNode::OnReset(const FCameraNodeResetParams& Params)
+{
+	TreeCache.ForEachNode(ECameraNodeFlags::RequiresReset, [Params](UCameraNode* Node) { Node->Reset(Params); });
+}
+
 void UBlendStackRootCameraNode::OnRun(const FCameraNodeRunParams& Params, FCameraNodeRunResult& OutResult)
 {
 	if (Blend)
diff --git a/Engine/Plugins/Cameras/GameplayCameras/Source/GameplayCameras/Private/Core/BlendStackRootCameraNode.h b/Engine/Plugins/Cameras/GameplayCameras/Source/GameplayCameras/Private/Core/BlendStackRootCameraNode.h
--- a/Engine/Plugins/Cameras/GameplayCameras/Source/GameplayCameras/Private/Core/BlendStackRootCameraNode.h
+++ b/Engine/Plugins/Cameras/GameplayCameras/Source/GameplayCameras/Private/Core/BlendStackRootCameraNode.h
@@ -3,6 +3,7 @@
 #pragma once
 
 #include "Core/CameraNode.h"
+#include "Core/CameraNodeTreeCache.h"
 
 #include "BlendStackRootCameraNode.generated.h"
 
@@ -19,12 +20,24 @@ class UBlendStackRootCameraNode : public UCameraNode
 {
 	GENERATED_BODY()
 
+public:
+
+	/** Initialize this node. */
+	void Initialize(UBlendCameraNode* InBlend, UCameraNode* InRootNode);
+
+	/** Gets the blend node. */
+	UBlendCameraNode* GetBlend() const { return Blend; }
+
+	/** Gets the root of the camera mode. */
+	UCameraNode* GetRootNode() const { return RootNode; }
+
 protected:
 
 	virtual FCameraNodeChildrenView OnGetChildren() override;
+	virtual void OnReset(const FCameraNodeResetParams& Params) override;
 	virtual void OnRun(const FCameraNodeRunParams& Params, FCameraNodeRunResult& OutResult) override;
 
-public:
+private:
 
 	/** The blend to use on the camera mode. */
 	UPROPERTY()
@@ -33,5 +46,8 @@ public:
 	/** The root of the instantied camera node tree. */
 	UPROPERTY()
 	TObjectPtr<UCameraNode> RootNode;
+
+	/** Node cache for the hierarchy below the root node. */
+	FCameraNodeTreeCache TreeCache;
 };
 
diff --git a/Engine/Plugins/Cameras/GameplayCameras/Source/GameplayCameras/Private/Core/CameraEvaluationContextStack.cpp b/Engine/Plugins/Cameras/GameplayCameras/Source/GameplayCameras/Private/Core/CameraEvaluationContextStack.cpp
--- a/Engine/Plugins/Cameras/GameplayCameras/Source/GameplayCameras/Private/Core/CameraEvaluationContextStack.cpp
+++ b/Engine/Plugins/Cameras/GameplayCameras/Source/GameplayCameras/Private/Core/CameraEvaluationContextStack.cpp
@@ -35,6 +35,8 @@ bool FCameraEvaluationContextStack::HasContext(UCameraEvaluationContext* Context
 
 void FCameraEvaluationContextStack::PushContext(UCameraEvaluationContext* Context)
 {
+	checkf(Evaluator, TEXT("Can't push context when no evaluator is set! Did you call Initialize?"));
+
 	// If we're pushing an existing context, move it to the top.
 	const int32 ExistingIndex = Entries.IndexOfByPredicate(
 			[Context](const FContextEntry& Entry) { return Entry.WeakContext == Context; });
@@ -51,14 +53,14 @@ void FCameraEvaluationContextStack::PushContext(UCameraEvaluationContext* Contex
 
 	// Instantiate the camera director.
 	FCameraRuntimeInstantiationParams InstParams;
-	InstParams.InstantiationOuter = InstantiationOuter.Get();
+	InstParams.InstantiationOuter = Evaluator;
 	if (!InstParams.InstantiationOuter)
 	{
 		InstParams.InstantiationOuter = GetTransientPackage();
 	}
 
 	const UCameraDirector* OriginalCameraDirector = Context->GetCameraAsset()->CameraDirector;
-	UCameraDirector* NewCameraDirector = FCameraRuntimeInstantiator::InstantiateCameraDirector(
+	UCameraDirector* NewCameraDirector = Evaluator->GetRuntimeInstantiator().InstantiateCameraDirector(
 			OriginalCameraDirector, InstParams);
 	
 	// Add an entry in the stack.
@@ -80,9 +82,9 @@ void FCameraEvaluationContextStack::PopContext()
 	Entries.Pop();
 }
 
-void FCameraEvaluationContextStack::SetInstantiationOuter(UObject* InInstantiationOuter)
+void FCameraEvaluationContextStack::Initialize(UCameraSystemEvaluator* InEvaluator)
 {
-	InstantiationOuter = InInstantiationOuter;
+	Evaluator = InEvaluator;
 }
 
 void FCameraEvaluationContextStack::AddReferencedObjects(FReferenceCollector& Collector)
diff --git a/Engine/Plugins/Cameras/GameplayCameras/Source/GameplayCameras/Private/Core/CameraNode.cpp b/Engine/Plugins/Cameras/GameplayCameras/Source/GameplayCameras/Private/Core/CameraNode.cpp
--- a/Engine/Plugins/Cameras/GameplayCameras/Source/GameplayCameras/Private/Core/CameraNode.cpp
+++ b/Engine/Plugins/Cameras/GameplayCameras/Source/GameplayCameras/Private/Core/CameraNode.cpp
@@ -18,6 +18,11 @@ FCameraNodeChildrenView UCameraNode::GetChildren()
 	return OnGetChildren();
 }
 
+void UCameraNode::Reset(const FCameraNodeResetParams& Params)
+{
+	OnReset(Params);
+}
+
 void UCameraNode::Run(const FCameraNodeRunParams& Params, FCameraNodeRunResult& OutResult)
 {
 	if (bIsEnabled)
@@ -32,10 +37,7 @@ void UCameraNode::PostEditChangeProperty(struct FPropertyChangedEvent& PropertyC
 {
 	Super::PostEditChangeProperty(PropertyChangedEvent);
 	
-	if (!bIsInstantiated)
-	{
-		FCameraRuntimeInstantiator::ForwardPropertyChange(this, PropertyChangedEvent);
-	}
+	FCameraRuntimeInstantiator::ForwardPropertyChange(this, PropertyChangedEvent);
 }
 
 #endif
diff --git a/Engine/Plugins/Cameras/GameplayCameras/Source/GameplayCameras/Private/Core/CameraNodeTreeCache.cpp b/Engine/Plugins/Cameras/GameplayCameras/Source/GameplayCameras/Private/Core/CameraNodeTreeCache.cpp
new file mode 100644
--- /dev/null
+++ b/Engine/Plugins/Cameras/GameplayCameras/Source/GameplayCameras/Private/Core/CameraNodeTreeCache.cpp
@@ -0,0 +1,36 @@
+// Copyright Epic Games, Inc. All Rights Reserved.
+
+#include "Core/CameraNodeTreeCache.h"
+
+void FCameraNodeTreeCache::Build(UCameraNode* InRootNode)
+{
+	CachedNodes.Reset();
+
+	TArray<UCameraNode*> IterStack;
+	IterStack.Push(InRootNode);
+	while (!IterStack.IsEmpty())
+	{
+		UCameraNode* CurNode = IterStack.Pop();
+
+		FCachedNode CachedNode{ CurNode->GetNodeFlags(), CurNode };
+		CachedNodes.Add(CachedNode);
+
+		FCameraNodeChildrenView CurChildren = CurNode->GetChildren();
+		for (UCameraNode* Child : ReverseIterate(CurChildren))
+		{
+			IterStack.Push(Child);
+		}
+	}
+}
+
+void FCameraNodeTreeCache::ForEachNode(ECameraNodeFlags FlagFilter, TFunctionRef<void(UCameraNode*)> Callback)
+{
+	for (FCachedNode& CachedNode : CachedNodes)
+	{
+		if (EnumHasAllFlags(CachedNode.Flags, FlagFilter))
+		{
+			Callback(CachedNode.Node);
+		}
+	}
+}
+
diff --git a/Engine/Plugins/Cameras/GameplayCameras/Source/GameplayCameras/Private/Core/CameraRuntimeInstantiator.cpp b/Engine/Plugins/Cameras/GameplayCameras/Source/GameplayCameras/Private/Core/CameraRuntimeInstantiator.cpp
--- a/Engine/Plugins/Cameras/GameplayCameras/Source/GameplayCameras/Private/Core/CameraRuntimeInstantiator.cpp
+++ b/Engine/Plugins/Cameras/GameplayCameras/Source/GameplayCameras/Private/Core/CameraRuntimeInstantiator.cpp
@@ -5,11 +5,26 @@
 #include "Core/CameraDirector.h"
 #include "Core/CameraNode.h"
 #include "CoreGlobals.h"
-#include "IGameplayCamerasModule.h"
+#include "HAL/IConsoleManager.h"
 #include "IGameplayCamerasLiveEditManager.h"
+#include "IGameplayCamerasModule.h"
 #include "Modules/ModuleManager.h"
 #include "UObject/UObjectGlobals.h"
 
+bool GGameplayCamerasEnableLiveEdit = true;
+static FAutoConsoleVariableRef CVarGameplayCamerasEnableLiveEdit(
+	TEXT("GameplayCameras.EnableLiveEdit"),
+	GGameplayCamerasEnableLiveEdit,
+	TEXT("(Default: true. Toggles live-editing of camera runtime objects.")
+	);
+
+bool GGameplayCamerasEnableInstantiationRecycling = true;
+static FAutoConsoleVariableRef CVarGameplayCamerasEnableInstantiationRecycling(
+	TEXT("GameplayCameras.EnableInstantiationRecycling"),
+	GGameplayCamerasEnableInstantiationRecycling,
+	TEXT("(Default: true. Toggles recycling of instantiated camera objects.")
+	);
+
 IGameplayCamerasModule* FCameraRuntimeInstantiator::GameplayCamerasModule(nullptr);
 
 IGameplayCamerasModule& FCameraRuntimeInstantiator::GetGameplayCamerasModule()
@@ -21,6 +36,15 @@ IGameplayCamerasModule& FCameraRuntimeInstantiator::GetGameplayCamerasModule()
 	return *GameplayCamerasModule;
 }
 
+void FCameraRuntimeInstantiator::AddReferencedObjects(FReferenceCollector& Collector)
+{
+	for (auto& Pair : InstantiationPool)
+	{
+		FObjectPool& ObjectPool = Pair.Value;
+		Collector.AddReferencedObjects(ObjectPool);
+	}
+}
+
 UCameraNode* FCameraRuntimeInstantiator::InstantiateCameraNodeTree(const UCameraNode* InRootNode, const FCameraRuntimeInstantiationParams& Params)
 {
 	UCameraNode* NewRootNode = CastChecked<UCameraNode>(InstantiateObject(InRootNode, Params));
@@ -35,6 +59,27 @@ UCameraDirector* FCameraRuntimeInstantiator::InstantiateCameraDirector(const UCa
 
 UObject* FCameraRuntimeInstantiator::InstantiateObject(const UObject* InSource, const FCameraRuntimeInstantiationParams& Params)
 {
+	// See if we can recycle a previously instantiated object hierarchy.
+	if (Params.bAllowRecyling && GGameplayCamerasEnableInstantiationRecycling)
+	{
+		if (FObjectPool* ObjectPool = InstantiationPool.Find(InSource))
+		{
+			TObjectPtr<UObject> Recycled = ObjectPool->Last();
+			ObjectPool->RemoveAt(ObjectPool->Num() - 1);
+			if (ObjectPool->IsEmpty())
+			{
+				InstantiationPool.Remove(InSource);
+			}
+
+			// Just return the head of the hierarchy that we remember.
+			// The live-edit manager should already have the mapping of all
+			// source-to-instantiated objects from the time the clone was
+			// created (see below).
+			return Recycled;
+		}
+	}
+
+	// No recycled instance found... clone the source object hierarchy
 	FObjectDuplicationParameters DuplicationParams = InitStaticDuplicateObjectParams(
 			InSource,
 			Params.InstantiationOuter,
@@ -42,21 +87,17 @@ UObject* FCameraRuntimeInstantiator::InstantiateObject(const UObject* InSource,
 #if WITH_EDITOR
 	TMap<UObject*, UObject*> CreatedObjects;
 	DuplicationParams.CreatedObjects = &CreatedObjects;
-#endif
+#endif  // WITH_EDITOR
 
-	UObject* OutRootObject = StaticDuplicateObjectEx(DuplicationParams);
+	UObject* OutInstantiatedObject = StaticDuplicateObjectEx(DuplicationParams);
 
 #if WITH_EDITOR
-	if (!IsRunningCommandlet())
+	// Register all instantiated objects into the live-edit manager for live-eding from the
+	// source objects to their instantiated counterparts.
+	// We only do this once on actual cloning. We expect the live-edit manager to hold onto
+	// these objects until they are GC'ed.
+	if (!IsRunningCommandlet() && GGameplayCamerasEnableLiveEdit)
 	{
-		for (const TPair<UObject*, UObject*>& Pair : CreatedObjects)
-		{
-			if (UCameraNode* CreatedNode = Cast<UCameraNode>(Pair.Value))
-			{
-				CreatedNode->SetIsInstantiated();
-			}
-		}
-
 		if (TSharedPtr<IGameplayCamerasLiveEditManager> LiveEditManager = GetGameplayCamerasModule().GetLiveEditManager())
 		{
 			LiveEditManager->RegisterInstantiatedObjects(CreatedObjects);
@@ -64,15 +105,27 @@ UObject* FCameraRuntimeInstantiator::InstantiateObject(const UObject* InSource,
 	}
 #endif
 
-	return OutRootObject;
+	return OutInstantiatedObject;
+}
+
+void FCameraRuntimeInstantiator::RecycleInstantiatedObject(const UObject* InSourceObject, UObject* InstantiatedObject)
+{
+	FObjectPool& ObjectPool = InstantiationPool.FindOrAdd(InSourceObject);
+	if (ObjectPool.Num() < 4)
+	{
+		ObjectPool.Add(InstantiatedObject);
+	}
 }
 
 void FCameraRuntimeInstantiator::ForwardPropertyChange(const UObject* Object, const FPropertyChangedEvent& PropertyChangedEvent)
 {
 #if WITH_EDITOR
-	if (TSharedPtr<IGameplayCamerasLiveEditManager> LiveEditManager = GetGameplayCamerasModule().GetLiveEditManager())
+	if (GGameplayCamerasEnableLiveEdit)
 	{
-		LiveEditManager->ForwardPropertyChange(Object, PropertyChangedEvent);
+		if (TSharedPtr<IGameplayCamerasLiveEditManager> LiveEditManager = GetGameplayCamerasModule().GetLiveEditManager())
+		{
+			LiveEditManager->ForwardPropertyChange(Object, PropertyChangedEvent);
+		}
 	}
 #endif
 }
diff --git a/Engine/Plugins/Cameras/GameplayCameras/Source/GameplayCameras/Private/Core/CameraSystemEvaluator.cpp b/Engine/Plugins/Cameras/GameplayCameras/Source/GameplayCameras/Private/Core/CameraSystemEvaluator.cpp
--- a/Engine/Plugins/Cameras/GameplayCameras/Source/GameplayCameras/Private/Core/CameraSystemEvaluator.cpp
+++ b/Engine/Plugins/Cameras/GameplayCameras/Source/GameplayCameras/Private/Core/CameraSystemEvaluator.cpp
@@ -18,13 +18,14 @@ UCameraSystemEvaluator::UCameraSystemEvaluator(const FObjectInitializer& ObjectI
 {
 	RootNode = ObjectInit.CreateDefaultSubobject<UDefaultRootCameraNode>(this, TEXT("RootNode"));
 
-	ContextStack.SetInstantiationOuter(this);
+	ContextStack.Initialize(this);
 }
 
 void UCameraSystemEvaluator::AddReferencedObjects(UObject* InThis, FReferenceCollector& Collector)
 {
 	UCameraSystemEvaluator* TypedThis = CastChecked<UCameraSystemEvaluator>(InThis);
 	TypedThis->ContextStack.AddReferencedObjects(Collector);
+	TypedThis->Instantiator.AddReferencedObjects(Collector);
 }
 
 void UCameraSystemEvaluator::PushEvaluationContext(UCameraEvaluationContext* EvaluationContext)
@@ -69,6 +70,7 @@ void UCameraSystemEvaluator::Update(const FCameraSystemEvaluationUpdateParams& P
 		if (DirectorResult.ActiveCameraModes.Num() == 1)
 		{
 			FActivateCameraModeParams CameraModeParams;
+			CameraModeParams.Evaluator = this;
 			CameraModeParams.EvaluationContext = ActiveContextInfo.EvaluationContext;
 			CameraModeParams.CameraMode = DirectorResult.ActiveCameraModes[0];
 			RootNode->ActivateCameraMode(CameraModeParams);
@@ -77,6 +79,7 @@ void UCameraSystemEvaluator::Update(const FCameraSystemEvaluationUpdateParams& P
 
 	// Run the root camera node.
 	FCameraNodeRunParams NodeParams;
+	NodeParams.Evaluator = this;
 	NodeParams.DeltaTime = Params.DeltaTime;
 
 	RootNodeResult.Reset();
@@ -96,3 +99,8 @@ void UCameraSystemEvaluator::GetEvaluatedCameraView(FMinimalViewInfo& DesiredVie
 	DesiredView.FOV = CameraPose.GetEffectiveFieldOfView();
 }
 
+FCameraRuntimeInstantiator& UCameraSystemEvaluator::GetRuntimeInstantiator()
+{
+	return Instantiator;
+}
+
diff --git a/Engine/Plugins/Cameras/GameplayCameras/Source/GameplayCameras/Private/Core/DefaultRootCameraNode.cpp b/Engine/Plugins/Cameras/GameplayCameras/Source/GameplayCameras/Private/Core/DefaultRootCameraNode.cpp
--- a/Engine/Plugins/Cameras/GameplayCameras/Source/GameplayCameras/Private/Core/DefaultRootCameraNode.cpp
+++ b/Engine/Plugins/Cameras/GameplayCameras/Source/GameplayCameras/Private/Core/DefaultRootCameraNode.cpp
@@ -72,6 +72,7 @@ void UDefaultRootCameraNode::OnActivateCameraMode(const FActivateCameraModeParam
 	if (ensure(TargetStack))
 	{
 		FBlendStackCameraPushParams PushParams;
+		PushParams.Evaluator = Params.Evaluator;
 		PushParams.EvaluationContext = Params.EvaluationContext;
 		PushParams.CameraMode = Params.CameraMode;
 		TargetStack->Push(PushParams);
diff --git a/Engine/Plugins/Cameras/GameplayCameras/Source/GameplayCameras/Public/Core/BlendStackCameraNode.h b/Engine/Plugins/Cameras/GameplayCameras/Source/GameplayCameras/Public/Core/BlendStackCameraNode.h
--- a/Engine/Plugins/Cameras/GameplayCameras/Source/GameplayCameras/Public/Core/BlendStackCameraNode.h
+++ b/Engine/Plugins/Cameras/GameplayCameras/Source/GameplayCameras/Public/Core/BlendStackCameraNode.h
@@ -17,6 +17,9 @@ struct FCameraModeTransition;
  */
 struct FBlendStackCameraPushParams
 {
+	/** The evaluator currently running.*/
+	TObjectPtr<UCameraSystemEvaluator> Evaluator;
+
 	/** The evaluation context within which a camera mode's node tree should run. */
 	TWeakObjectPtr<const UCameraEvaluationContext> EvaluationContext;
 
diff --git a/Engine/Plugins/Cameras/GameplayCameras/Source/GameplayCameras/Public/Core/CameraEvaluationContextStack.h b/Engine/Plugins/Cameras/GameplayCameras/Source/GameplayCameras/Public/Core/CameraEvaluationContextStack.h
--- a/Engine/Plugins/Cameras/GameplayCameras/Source/GameplayCameras/Public/Core/CameraEvaluationContextStack.h
+++ b/Engine/Plugins/Cameras/GameplayCameras/Source/GameplayCameras/Public/Core/CameraEvaluationContextStack.h
@@ -48,7 +48,7 @@ public:
 public:
 
 	// Internal API
-	void SetInstantiationOuter(UObject* InInstantiationOuter);
+	void Initialize(UCameraSystemEvaluator* InEvaluator);
 	void AddReferencedObjects(FReferenceCollector& Collector);
 
 private:
@@ -59,8 +59,10 @@ private:
 		TObjectPtr<UCameraDirector> CameraDirector;
 	};
 
+	/** The entries in the stack. */
 	TArray<FContextEntry> Entries;
 
-	TWeakObjectPtr<> InstantiationOuter;
+	/** The owner evaluator. */
+	TObjectPtr<UCameraSystemEvaluator> Evaluator;
 };
 
diff --git a/Engine/Plugins/Cameras/GameplayCameras/Source/GameplayCameras/Public/Core/CameraInstantiableObject.h b/Engine/Plugins/Cameras/GameplayCameras/Source/GameplayCameras/Public/Core/CameraInstantiableObject.h
new file mode 100644
--- /dev/null
+++ b/Engine/Plugins/Cameras/GameplayCameras/Source/GameplayCameras/Public/Core/CameraInstantiableObject.h
@@ -0,0 +1,50 @@
+// Copyright Epic Games, Inc. All Rights Reserved.
+
+#pragma once
+
+#include "CoreTypes.h"
+#include "Misc/EnumClassFlags.h"
+#include "UObject/Object.h"
+
+#include "CameraInstantiableObject.generated.h"
+
+/**
+ *
+ */
+UENUM()
+enum class ECameraNodeInstantiationState
+{
+	None,
+	HasInstantiations,
+	IsInstantiated
+};
+
+/**
+ *
+ */
+UCLASS(MinimalAPI)
+class UCameraInstantiableObject : public UObject
+{
+	GENERATED_BODY()
+
+#if WITH_EDITOR
+
+public:
+
+	/** Gets the instantiation state of this object. */
+	ECameraNodeInstantiationState GetInstantiationState() const { return InstantiationState; }
+	/** Sets the instantiation state of this object. */
+	void SetInstantiationState(ECameraNodeInstantiationState bInValue) { InstantiationState = bInValue; }
+
+	/** Returns whether this object is an instantiated object. */
+	bool IsInstantiated() const { return EnumHasAnyFlags(InstantiationState, ECameraNodeInstantiationState::IsInstantiated); }
+	/** Returns whether this object has any known instantiations. */
+	bool HasInstantiations() const { return EnumHasAnyFlags(InstantiationState, ECameraNodeInstantiationState::HasInstantiations); }
+
+private:
+
+	ECameraNodeInstantiationState InstantiationState = ECameraNodeInstantiationState::None;
+
+#endif  // WITH_EDITOR
+};
+
diff --git a/Engine/Plugins/Cameras/GameplayCameras/Source/GameplayCameras/Public/Core/CameraNode.h b/Engine/Plugins/Cameras/GameplayCameras/Source/GameplayCameras/Public/Core/CameraNode.h
--- a/Engine/Plugins/Cameras/GameplayCameras/Source/GameplayCameras/Public/Core/CameraNode.h
+++ b/Engine/Plugins/Cameras/GameplayCameras/Source/GameplayCameras/Public/Core/CameraNode.h
@@ -2,6 +2,7 @@
 
 #pragma once
 
+#include "Core/CameraInstantiableObject.h"
 #include "Core/CameraPose.h"
 #include "Core/CameraNodeChildrenView.h"
 #include "CoreTypes.h"
@@ -9,11 +10,15 @@
 
 #include "CameraNode.generated.h"
 
+class UCameraSystemEvaluator;
+
 /**
  * Parameter structure for running a camera node.
  */
 struct FCameraNodeRunParams
 {
+	/** The evaluation running this evaluation.*/
+	TObjectPtr<UCameraSystemEvaluator> Evaluator;
 	/** The time interval for the evaluation. */
 	float DeltaTime = 0.f;
 	/** Whether this is the first evaluation of this camera node hierarchy. */
@@ -36,11 +41,23 @@ struct FCameraNodeRunResult
 	void Reset();
 };
 
+UENUM()
+enum class ECameraNodeFlags
+{
+	None = 0,
+	RequiresReset = 1
+};
+ENUM_CLASS_FLAGS(ECameraNodeFlags);
+
+struct FCameraNodeResetParams
+{
+};
+
 /**
  * The base class for a camera node.
  */
 UCLASS(Abstract, DefaultToInstanced, EditInlineNew, MinimalAPI)
-class UCameraNode : public UObject
+class UCameraNode : public UCameraInstantiableObject
 {
 	GENERATED_BODY()
 
@@ -49,23 +66,26 @@ public:
 	/** Get the list of children under this node. */
 	FCameraNodeChildrenView GetChildren();
 
-	/** Run this node. */
-	void Run(const FCameraNodeRunParams& Params, FCameraNodeRunResult& OutResult);
-
-public:
+	/** Get the flags for this node. */
+	ECameraNodeFlags GetNodeFlags() const { return Flags; }
 
-	// Internal API
+	/** Resets this node. */
+	void Reset(const FCameraNodeResetParams& Params);
 
-	/** Whether this node was instantiated from a source node. */
-	bool GetIsInstantiated() const { return bIsInstantiated; }
-	/** Flag this node as being instantiated from a source node. */
-	void SetIsInstantiated(bool bInValue = true) { bIsInstantiated = bInValue; }
+	/** Run this node. */
+	void Run(const FCameraNodeRunParams& Params, FCameraNodeRunResult& OutResult);
 
 protected:
 
+	/** Sets the flags for this node. Should only be called once during construction. */
+	void SetNodeFlags(ECameraNodeFlags InFlags) { Flags = InFlags; }
+
 	/** Get the list of children under this node. */
 	virtual FCameraNodeChildrenView OnGetChildren() { return FCameraNodeChildrenView(); }
 
+	/** Resets this node. */
+	virtual void OnReset(const FCameraNodeResetParams& Params) {}
+
 	/** Run this node. */
 	virtual void OnRun(const FCameraNodeRunParams& Params, FCameraNodeRunResult& OutResult) {}
 
@@ -83,6 +103,7 @@ public:
 
 private:
 
-	bool bIsInstantiated = false;
+	/** The flags for this node. Should only be set once during construction. */
+	ECameraNodeFlags Flags = ECameraNodeFlags::None;
 };
 
diff --git a/Engine/Plugins/Cameras/GameplayCameras/Source/GameplayCameras/Public/Core/CameraNodeChildrenView.h b/Engine/Plugins/Cameras/GameplayCameras/Source/GameplayCameras/Public/Core/CameraNodeChildrenView.h
--- a/Engine/Plugins/Cameras/GameplayCameras/Source/GameplayCameras/Public/Core/CameraNodeChildrenView.h
+++ b/Engine/Plugins/Cameras/GameplayCameras/Source/GameplayCameras/Public/Core/CameraNodeChildrenView.h
@@ -101,37 +101,52 @@ public:
 
 	// Range iteration
 
-	struct FIterator
+	struct FBaseIterator
 	{
 		const FCameraNodeChildrenView* Owner;
 		int32 Index;
 
-		FORCEINLINE FIterator& operator++()
-		{
-			++Index;
-			return *this;
-		}
-
 		FORCEINLINE UCameraNode* operator*()
 		{
 			return (*Owner)[Index];
 		}
 
-		FORCEINLINE bool operator== (const FIterator& Other) const
+		FORCEINLINE bool operator== (const FBaseIterator& Other) const
 		{
 			return Owner == Other.Owner
 				&& Index == Other.Index;
 		}
 
-		FORCEINLINE bool operator!= (const FIterator& Other) const
+		FORCEINLINE bool operator!= (const FBaseIterator& Other) const
 		{
 			return !(*this == Other);
 		}
 	};
 
+	struct FIterator : FBaseIterator
+	{
+		FORCEINLINE FIterator& operator++()
+		{
+			++Index;
+			return *this;
+		}
+	};
+
 	FORCEINLINE FIterator begin() const { return FIterator{ this, 0 }; }
 	FORCEINLINE FIterator end() const { return FIterator{ this, Num() }; }
 
+	struct FReverseIterator : FBaseIterator
+	{
+		FORCEINLINE FReverseIterator& operator++()
+		{
+			--Index;
+			return *this;
+		}
+	};
+
+	FORCEINLINE FReverseIterator rbegin() const { return FReverseIterator{ this, Num() - 1 }; }
+	FORCEINLINE FReverseIterator rend() const { return FReverseIterator{ this, -1 }; }
+
 private:
 
 	using FStorage = TVariant<FArrayView, FArray>;
diff --git a/Engine/Plugins/Cameras/GameplayCameras/Source/GameplayCameras/Public/Core/CameraNodeTreeCache.h b/Engine/Plugins/Cameras/GameplayCameras/Source/GameplayCameras/Public/Core/CameraNodeTreeCache.h
new file mode 100644
--- /dev/null
+++ b/Engine/Plugins/Cameras/GameplayCameras/Source/GameplayCameras/Public/Core/CameraNodeTreeCache.h
@@ -0,0 +1,30 @@
+// Copyright Epic Games, Inc. All Rights Reserved.
+
+#pragma once
+
+#include "Core/CameraNode.h"
+#include "Templates/Function.h"
+
+/**
+ * A utility class that caches a hierarchy of arbitrary camera nodes as a
+ * flat depth-first list of those camera nodes.
+ */
+class FCameraNodeTreeCache
+{
+public:
+
+	/** (Re)builds the cache given the root node of a hierarchy. */
+	void Build(UCameraNode* InRootNode);
+
+	void ForEachNode(ECameraNodeFlags FlagFilter, TFunctionRef<void(UCameraNode*)> Callback);
+
+private:
+
+	struct FCachedNode
+	{
+		ECameraNodeFlags Flags;
+		TObjectPtr<UCameraNode> Node;
+	};
+	TArray<FCachedNode> CachedNodes;
+};
+
diff --git a/Engine/Plugins/Cameras/GameplayCameras/Source/GameplayCameras/Public/Core/CameraRuntimeInstantiator.h b/Engine/Plugins/Cameras/GameplayCameras/Source/GameplayCameras/Public/Core/CameraRuntimeInstantiator.h
--- a/Engine/Plugins/Cameras/GameplayCameras/Source/GameplayCameras/Public/Core/CameraRuntimeInstantiator.h
+++ b/Engine/Plugins/Cameras/GameplayCameras/Source/GameplayCameras/Public/Core/CameraRuntimeInstantiator.h
@@ -18,6 +18,7 @@ struct FCameraRuntimeInstantiationParams
 {
 	UObject* InstantiationOuter = nullptr;
 	FName InstantiationName = NAME_None;
+	bool bAllowRecyling = false;
 };
 
 /**
@@ -31,26 +32,40 @@ public:
 	 * Instantiates a tree of camera nodes by cloning it, and keeping track of the relationship between
 	 * the clones and their source assets.
 	 */
-	static UCameraNode* InstantiateCameraNodeTree(const UCameraNode* InRootNode, const FCameraRuntimeInstantiationParams& Params);
+	UCameraNode* InstantiateCameraNodeTree(const UCameraNode* InRootNode, const FCameraRuntimeInstantiationParams& Params);
 
 	/** 
 	 * Instantiates a camera directory by cloning it, and keeping track of the relationship between
 	 * the clone and its source asset.
 	 */
-	static UCameraDirector* InstantiateCameraDirector(const UCameraDirector* InCameraDirector, const FCameraRuntimeInstantiationParams& Params);
+	UCameraDirector* InstantiateCameraDirector(const UCameraDirector* InCameraDirector, const FCameraRuntimeInstantiationParams& Params);
+
+	/**
+	 * Release the root of an instantiated object hierarchy, for potential recycling later.
+	 */
+	void RecycleInstantiatedObject(const UObject* InSourceObject, UObject* InstantiatedObject);
 
 	/**
 	 * Take a property change applied on a source object and try to replicate it on any related instantiated objects.
 	 */
 	static void ForwardPropertyChange(const UObject* Object, const FPropertyChangedEvent& PropertyChangedEvent);
 
+public:
+
+	// Internal API.
+	void AddReferencedObjects(FReferenceCollector& Collector);
+
 private:
 
-	static UObject* InstantiateObject(const UObject* InSource, const FCameraRuntimeInstantiationParams& Params);
 	static IGameplayCamerasModule& GetGameplayCamerasModule();
 
+	UObject* InstantiateObject(const UObject* InSource, const FCameraRuntimeInstantiationParams& Params);
+
 private:
 
 	static IGameplayCamerasModule* GameplayCamerasModule;
+
+	using FObjectPool = TArray<TObjectPtr<UObject>>; //, TInlineAllocator<4>>;
+	TMap<TWeakObjectPtr<const UObject>, FObjectPool> InstantiationPool;
 };
 
diff --git a/Engine/Plugins/Cameras/GameplayCameras/Source/GameplayCameras/Public/Core/CameraSystemEvaluator.h b/Engine/Plugins/Cameras/GameplayCameras/Source/GameplayCameras/Public/Core/CameraSystemEvaluator.h
--- a/Engine/Plugins/Cameras/GameplayCameras/Source/GameplayCameras/Public/Core/CameraSystemEvaluator.h
+++ b/Engine/Plugins/Cameras/GameplayCameras/Source/GameplayCameras/Public/Core/CameraSystemEvaluator.h
@@ -4,8 +4,9 @@
 
 #include "Core/CameraDirector.h"
 #include "Core/CameraEvaluationContextStack.h"
-#include "Core/CameraPose.h"
 #include "Core/CameraNode.h"
+#include "Core/CameraPose.h"
+#include "Core/CameraRuntimeInstantiator.h"
 #include "CoreTypes.h"
 #include "UObject/GCObject.h"
 
@@ -70,6 +71,11 @@ public:
 	/** Get the last evaluated camera. */
 	void GetEvaluatedCameraView(FMinimalViewInfo& DesiredView);
 
+public:
+
+	/** Gets the runtime instantiator for this evaluator. */
+	FCameraRuntimeInstantiator& GetRuntimeInstantiator();
+
 public:
 
 	static void AddReferencedObjects(UObject* InThis, FReferenceCollector& Collector);
@@ -88,5 +94,8 @@ private:
 
 	/** The current overall result of the camera system. */
 	FCameraSystemEvaluationUpdateResult Result;
+
+	/** The instantiator for runtime objects. */
+	FCameraRuntimeInstantiator Instantiator;
 };
 
diff --git a/Engine/Plugins/Cameras/GameplayCameras/Source/GameplayCameras/Public/Core/RootCameraNode.h b/Engine/Plugins/Cameras/GameplayCameras/Source/GameplayCameras/Public/Core/RootCameraNode.h
--- a/Engine/Plugins/Cameras/GameplayCameras/Source/GameplayCameras/Public/Core/RootCameraNode.h
+++ b/Engine/Plugins/Cameras/GameplayCameras/Source/GameplayCameras/Public/Core/RootCameraNode.h
@@ -30,6 +30,9 @@ ENUM_CLASS_FLAGS(ECameraModeLayer)
  */
 struct FActivateCameraModeParams
 {
+	/** The evaluator currently running.*/
+	TObjectPtr<UCameraSystemEvaluator> Evaluator;
+
 	/** The evaluation context in which the camera mode runs. */
 	TWeakObjectPtr<const UCameraEvaluationContext> EvaluationContext;
 
diff --git a/Engine/Plugins/Cameras/GameplayCameras/Source/GameplayCamerasEditor/Private/GameplayCamerasEditorModule.cpp b/Engine/Plugins/Cameras/GameplayCameras/Source/GameplayCamerasEditor/Private/GameplayCamerasEditorModule.cpp
--- a/Engine/Plugins/Cameras/GameplayCameras/Source/GameplayCamerasEditor/Private/GameplayCamerasEditorModule.cpp
+++ b/Engine/Plugins/Cameras/GameplayCameras/Source/GameplayCamerasEditor/Private/GameplayCamerasEditorModule.cpp
@@ -84,15 +84,24 @@ private:
 
 		IGameplayCamerasModule& CamerasModule = FModuleManager::LoadModuleChecked<IGameplayCamerasModule>("GameplayCameras");
 		CamerasModule.SetLiveEditManager(LiveEditManager);
+
+		FCoreUObjectDelegates::GetPostGarbageCollect().AddRaw(this, &FGameplayCamerasEditorModule::OnPostGarbageCollection);
 	}
 
 	void TeardownLiveEditManager()
 	{
+		FCoreUObjectDelegates::GetPostGarbageCollect().RemoveAll(this);
+
 		IGameplayCamerasModule& CamerasModule = FModuleManager::LoadModuleChecked<IGameplayCamerasModule>("GameplayCameras");
 		CamerasModule.SetLiveEditManager(nullptr);
 
 		LiveEditManager.Reset();
 	}
+	
+	void OnPostGarbageCollection()
+	{
+		LiveEditManager->CleanUp();
+	}
 
 private:
 
diff --git a/Engine/Plugins/Cameras/GameplayCameras/Source/GameplayCamerasEditor/Private/GameplayCamerasLiveEditManager.cpp b/Engine/Plugins/Cameras/GameplayCameras/Source/GameplayCamerasEditor/Private/GameplayCamerasLiveEditManager.cpp
--- a/Engine/Plugins/Cameras/GameplayCameras/Source/GameplayCamerasEditor/Private/GameplayCamerasLiveEditManager.cpp
+++ b/Engine/Plugins/Cameras/GameplayCameras/Source/GameplayCamerasEditor/Private/GameplayCamerasLiveEditManager.cpp
@@ -2,10 +2,56 @@
 
 #include "GameplayCamerasLiveEditManager.h"
 
+#include "Core/CameraInstantiableObject.h"
+#include "UObject/UObjectGlobals.h"
+
+FGameplayCamerasLiveEditManager::FGameplayCamerasLiveEditManager()
+{
+}
+
+void FGameplayCamerasLiveEditManager::CleanUp()
+{
+	for (auto It = Instantiations.CreateIterator(); It; ++It)
+	{
+		// Clean up invalid entries on keys and values.
+		UObject* Obj = It->Key.Get();
+		if (!Obj)
+		{
+			It.RemoveCurrent();
+			continue;
+		}
+		for (auto ObjIt = It->Value.InstantiatedObjects.CreateIterator(); ObjIt; ++ObjIt)
+		{
+			if (!ObjIt->IsValid())
+			{
+				ObjIt.RemoveCurrent();
+			}
+		}
+
+		// If the source object doesn't have any instantiated objects anymore, clear its flags.
+		if (It->Value.InstantiatedObjects.IsEmpty())
+		{
+			if (UCameraInstantiableObject* SourceObject = Cast<UCameraInstantiableObject>(Obj))
+			{
+				SourceObject->SetInstantiationState(ECameraNodeInstantiationState::None);
+			}
+		}
+	}
+}
+
 void FGameplayCamerasLiveEditManager::RegisterInstantiatedObjects(const TMap<UObject*, UObject*> InstantiatedObjects)
 {
 	for (const TPair<UObject*, UObject*>& Pair : InstantiatedObjects)
 	{
+		// If the objects are instantiable, set appropriate flags on both.
+		if (UCameraInstantiableObject* SourceObject = Cast<UCameraInstantiableObject>(Pair.Key))
+		{
+			SourceObject->SetInstantiationState(ECameraNodeInstantiationState::HasInstantiations);
+
+			UCameraInstantiableObject* InstantiatedObject = CastChecked<UCameraInstantiableObject>(Pair.Value);
+			InstantiatedObject->SetInstantiationState(ECameraNodeInstantiationState::IsInstantiated);
+		}
+		
 		FInstantiationInfo& Info = Instantiations.FindOrAdd(Pair.Key);
 		Info.InstantiatedObjects.Add(Pair.Value);
 	}
@@ -54,20 +100,3 @@ void FGameplayCamerasLiveEditManager::ForwardPropertyChange(const UObject* Objec
 	}
 }
 
-void FGameplayCamerasLiveEditManager::ForwardPropertyChange(const UObject* SourceObject, UObject* InstantiatedObject, const FPropertyChangedEvent& PropertyChangedEvent)
-{
-}
-
-void FGameplayCamerasLiveEditManager::AddReferencedObjects(FReferenceCollector& Collector)
-{
-	for (auto& Pair : Instantiations)
-	{
-		Collector.AddReferencedObject(Pair.Key);
-	}
-}
-
-FString FGameplayCamerasLiveEditManager::GetReferencerName() const
-{
-	return TEXT("FGameplayCamerasLiveEditManager");
-}
-
diff --git a/Engine/Plugins/Cameras/GameplayCameras/Source/GameplayCamerasEditor/Private/GameplayCamerasLiveEditManager.h b/Engine/Plugins/Cameras/GameplayCameras/Source/GameplayCamerasEditor/Private/GameplayCamerasLiveEditManager.h
--- a/Engine/Plugins/Cameras/GameplayCameras/Source/GameplayCamerasEditor/Private/GameplayCamerasLiveEditManager.h
+++ b/Engine/Plugins/Cameras/GameplayCameras/Source/GameplayCamerasEditor/Private/GameplayCamerasLiveEditManager.h
@@ -5,31 +5,29 @@
 #include "IGameplayCamerasLiveEditManager.h"
 
 #include "CoreTypes.h"
-#include "UObject/GCObject.h"
 #include "UObject/WeakObjectPtr.h"
 
-class FGameplayCamerasLiveEditManager 
-	: public IGameplayCamerasLiveEditManager
-	, public FGCObject
+class FGameplayCamerasLiveEditManager : public IGameplayCamerasLiveEditManager
 {
 public:
 
-	virtual void RegisterInstantiatedObjects(const TMap<UObject*, UObject*> InstantiatedObjects) override;
-
-	virtual void ForwardPropertyChange(const UObject* Object, const FPropertyChangedEvent& PropertyChangedEvent) override;
+	FGameplayCamerasLiveEditManager();
 
-private:
+	/** Clean-up any invalid entries in the map of known instantiated objects. */
+	void CleanUp();
 
-	virtual void AddReferencedObjects(FReferenceCollector& Collector) override;
-	virtual FString GetReferencerName() const override;
+public:
 
-	void ForwardPropertyChange(const UObject* SourceObject, UObject* InstantiatedObject, const FPropertyChangedEvent& PropertyChangedEvent);
+	// IGameplayCamerasLiveEditManager interface
+	virtual void RegisterInstantiatedObjects(const TMap<UObject*, UObject*> InstantiatedObjects) override;
+	virtual void ForwardPropertyChange(const UObject* Object, const FPropertyChangedEvent& PropertyChangedEvent) override;
 
 private:
+
 	struct FInstantiationInfo
 	{
 		TArray<TWeakObjectPtr<>> InstantiatedObjects;
 	};
-	TMap<TObjectPtr<const UObject>, FInstantiationInfo> Instantiations;
+	TMap<TWeakObjectPtr<UObject>, FInstantiationInfo> Instantiations;
 };
 
