Cameras: new generic header for perf and log categories

diff --git a/Engine/Plugins/Cameras/GameplayCameras/Source/GameplayCameras/Private/Animations/CameraAnimationCameraModifier.cpp b/Engine/Plugins/Cameras/GameplayCameras/Source/GameplayCameras/Private/Animations/CameraAnimationCameraModifier.cpp
--- a/Engine/Plugins/Cameras/GameplayCameras/Source/GameplayCameras/Private/Animations/CameraAnimationCameraModifier.cpp
+++ b/Engine/Plugins/Cameras/GameplayCameras/Source/GameplayCameras/Private/Animations/CameraAnimationCameraModifier.cpp
@@ -8,6 +8,7 @@
 #include "Engine/Canvas.h"
 #include "Engine/Engine.h"
 #include "GameFramework/PlayerController.h"
+#include "GameplayCameras.h"
 #include "IGameplayCamerasModule.h"
 #include "Kismet/GameplayStatics.h"
 #include "UObject/Object.h"
diff --git a/Engine/Plugins/Cameras/GameplayCameras/Source/GameplayCameras/Private/Core/CameraSystemEvaluator.cpp b/Engine/Plugins/Cameras/GameplayCameras/Source/GameplayCameras/Private/Core/CameraSystemEvaluator.cpp
--- a/Engine/Plugins/Cameras/GameplayCameras/Source/GameplayCameras/Private/Core/CameraSystemEvaluator.cpp
+++ b/Engine/Plugins/Cameras/GameplayCameras/Source/GameplayCameras/Private/Core/CameraSystemEvaluator.cpp
@@ -7,6 +7,7 @@
 #include "Core/CameraEvaluationContext.h"
 #include "Core/CameraMode.h"
 #include "Core/DefaultRootCameraNode.h"
+#include "GameplayCameras.h"
 #include "IGameplayCamerasModule.h"
 
 #include UE_INLINE_GENERATED_CPP_BY_NAME(CameraSystemEvaluator)
diff --git a/Engine/Plugins/Cameras/GameplayCameras/Source/GameplayCameras/Private/GameplayCamerasModule.cpp b/Engine/Plugins/Cameras/GameplayCameras/Source/GameplayCameras/Private/GameplayCamerasModule.cpp
--- a/Engine/Plugins/Cameras/GameplayCameras/Source/GameplayCameras/Private/GameplayCamerasModule.cpp
+++ b/Engine/Plugins/Cameras/GameplayCameras/Source/GameplayCameras/Private/GameplayCamerasModule.cpp
@@ -6,8 +6,12 @@
 #include "CameraAnimationCameraModifier.h"
 #include "CameraAnimationSequencePlayer.h"
 #include "Features/IModularFeatures.h"
+#include "GameplayCameras.h"
+#include "Logging/MessageLog.h"
 #include "Modules/ModuleManager.h"
 
+DEFINE_LOG_CATEGORY(LogCameraSystem);
+
 IGameplayCamerasModule& IGameplayCamerasModule::Get()
 {
 	return FModuleManager::LoadModuleChecked<IGameplayCamerasModule>("GameplayCameras");
diff --git a/Engine/Plugins/Cameras/GameplayCameras/Source/GameplayCameras/Public/GameplayCameras.h b/Engine/Plugins/Cameras/GameplayCameras/Source/GameplayCameras/Public/GameplayCameras.h
new file mode 100644
--- /dev/null
+++ b/Engine/Plugins/Cameras/GameplayCameras/Source/GameplayCameras/Public/GameplayCameras.h
@@ -0,0 +1,12 @@
+// Copyright Epic Games, Inc. All Rights Reserved.
+
+#pragma once
+
+#include "CoreTypes.h"
+#include "Stats/Stats.h"
+
+DECLARE_STATS_GROUP(TEXT("Camera System Evaluation"), STATGROUP_CameraSystem, STATCAT_Advanced)
+DECLARE_STATS_GROUP(TEXT("Camera Animation Evaluation"), STATGROUP_CameraAnimation, STATCAT_Advanced)
+
+DECLARE_LOG_CATEGORY_EXTERN(LogCameraSystem, Log, All);
+
diff --git a/Engine/Plugins/Cameras/GameplayCameras/Source/GameplayCameras/Public/IGameplayCamerasModule.h b/Engine/Plugins/Cameras/GameplayCameras/Source/GameplayCameras/Public/IGameplayCamerasModule.h
--- a/Engine/Plugins/Cameras/GameplayCameras/Source/GameplayCameras/Public/IGameplayCamerasModule.h
+++ b/Engine/Plugins/Cameras/GameplayCameras/Source/GameplayCameras/Public/IGameplayCamerasModule.h
@@ -3,15 +3,12 @@
 #pragma once
 
 #include "Modules/ModuleInterface.h"
-#include "Stats/Stats.h"
+#include "Templates/SharedPointer.h"
 
 #if WITH_EDITOR
 class IGameplayCamerasLiveEditManager;
 #endif
 
-DECLARE_STATS_GROUP(TEXT("Camera System Evaluation"), STATGROUP_CameraSystem, STATCAT_Advanced)
-DECLARE_STATS_GROUP(TEXT("Camera Animation Evaluation"), STATGROUP_CameraAnimation, STATCAT_Advanced)
-
 class IGameplayCamerasModule : public IModuleInterface
 {
 public:
