Cameras: pass evaluation context in the camera node run params.

diff --git a/Engine/Plugins/Cameras/GameplayCameras/Source/GameplayCameras/Private/Core/BlendStackCameraNode.cpp b/Engine/Plugins/Cameras/GameplayCameras/Source/GameplayCameras/Private/Core/BlendStackCameraNode.cpp
--- a/Engine/Plugins/Cameras/GameplayCameras/Source/GameplayCameras/Private/Core/BlendStackCameraNode.cpp
+++ b/Engine/Plugins/Cameras/GameplayCameras/Source/GameplayCameras/Private/Core/BlendStackCameraNode.cpp
@@ -95,6 +95,7 @@ void UBlendStackCameraNode::OnRun(const FCameraNodeRunParams& Params, FCameraNod
 	for (FCameraModeEntry& Entry : Entries)
 	{
 		FCameraNodeRunParams CurParams(Params);
+		CurParams.EvaluationContext = Entry.EvaluationContext.Get();
 		CurParams.bIsFirstFrame = Entry.bIsFirstFrame;
 
 		FCameraNodeRunResult& CurResult(Entry.Result);
@@ -151,6 +152,7 @@ void UBlendStackCameraNode::OnRun(const FCameraNodeRunParams& Params, FCameraNod
 		const FCameraPoseFlags ChangedFlags(CurResult.CameraPose.GetChangedFlags());
 
 		FCameraNodeRunParams CurParams(Params);
+		CurParams.EvaluationContext = Entry.EvaluationContext.Get();
 		CurParams.bIsFirstFrame = Entry.bIsFirstFrame;
 		FCameraNodeBlendParams BlendParams(CurParams, CurResult);
 
diff --git a/Engine/Plugins/Cameras/GameplayCameras/Source/GameplayCameras/Public/Core/CameraNode.h b/Engine/Plugins/Cameras/GameplayCameras/Source/GameplayCameras/Public/Core/CameraNode.h
--- a/Engine/Plugins/Cameras/GameplayCameras/Source/GameplayCameras/Public/Core/CameraNode.h
+++ b/Engine/Plugins/Cameras/GameplayCameras/Source/GameplayCameras/Public/Core/CameraNode.h
@@ -10,6 +10,7 @@
 
 #include "CameraNode.generated.h"
 
+class UCameraEvaluationContext;
 class UCameraSystemEvaluator;
 
 /**
@@ -17,8 +18,10 @@ class UCameraSystemEvaluator;
  */
 struct FCameraNodeRunParams
 {
-	/** The evaluation running this evaluation.*/
+	/** The evaluation running this evaluation. */
 	TObjectPtr<UCameraSystemEvaluator> Evaluator;
+	/** The evaluation context (if any) responsible for this branch of the evaluation. */
+	TObjectPtr<const UCameraEvaluationContext> EvaluationContext;
 	/** The time interval for the evaluation. */
 	float DeltaTime = 0.f;
 	/** Whether this is the first evaluation of this camera node hierarchy. */
