Cameras: make gameplay camera component blueprintable and spawnable

diff --git a/Engine/Plugins/Cameras/GameplayCameras/Source/GameplayCameras/Public/GameFramework/GameplayCameraComponent.h b/Engine/Plugins/Cameras/GameplayCameras/Source/GameplayCameras/Public/GameFramework/GameplayCameraComponent.h
--- a/Engine/Plugins/Cameras/GameplayCameras/Source/GameplayCameras/Public/GameFramework/GameplayCameraComponent.h
+++ b/Engine/Plugins/Cameras/GameplayCameras/Source/GameplayCameras/Public/GameFramework/GameplayCameraComponent.h
@@ -16,7 +16,7 @@ class UGameplayCameraComponentEvaluationContext;
 /**
  * A component that can run a camera asset inside its own camera evaluation context.
  */
-UCLASS(BlueprintType, MinimalAPI, ClassGroup=Camera, HideCategories=(Mobility, Rendering, LOD))
+UCLASS(Blueprintable, MinimalAPI, ClassGroup=Camera, HideCategories=(Mobility, Rendering, LOD), meta=(BlueprintSpawnableComponent))
 class UGameplayCameraComponent : public USceneComponent
 {
 	GENERATED_BODY()
