Cameras: fix ownership vs passing of evaluation context

It should be owned as a weak const pointer, but passed as a resolved const pointer.

diff --git a/Engine/Plugins/Cameras/GameplayCameras/Source/GameplayCameras/Private/Core/BlendStackCameraNode.cpp b/Engine/Plugins/Cameras/GameplayCameras/Source/GameplayCameras/Private/Core/BlendStackCameraNode.cpp
--- a/Engine/Plugins/Cameras/GameplayCameras/Source/GameplayCameras/Private/Core/BlendStackCameraNode.cpp
+++ b/Engine/Plugins/Cameras/GameplayCameras/Source/GameplayCameras/Private/Core/BlendStackCameraNode.cpp
@@ -94,8 +94,15 @@ void UBlendStackCameraNode::OnRun(const FCameraNodeRunParams& Params, FCameraNod
 	// Start by evaluating all the root nodes in the stack.
 	for (FCameraModeEntry& Entry : Entries)
 	{
+		const UCameraEvaluationContext* CurContext = Entry.EvaluationContext.Get();
+		if (UNLIKELY(CurContext == nullptr))
+		{
+			Entry.Result.bIsValid = false;
+			continue;
+		}
+
 		FCameraNodeRunParams CurParams(Params);
-		CurParams.EvaluationContext = Entry.EvaluationContext.Get();
+		CurParams.EvaluationContext = CurContext;
 		CurParams.bIsFirstFrame = Entry.bIsFirstFrame;
 
 		FCameraNodeRunResult& CurResult(Entry.Result);
@@ -104,12 +111,6 @@ void UBlendStackCameraNode::OnRun(const FCameraNodeRunParams& Params, FCameraNod
 		{
 			// If the context in which this camera mode runs doesn't have a valid result,
 			// skip it.
-			const UCameraEvaluationContext* CurContext = Entry.EvaluationContext.Get();
-			if (UNLIKELY(CurContext == nullptr))
-			{
-				CurResult.bIsValid = false;
-				continue;
-			}
 			const FCameraNodeRunResult& ContextResult(CurContext->GetInitialResult());
 			if (UNLIKELY(!ContextResult.bIsValid))
 			{
diff --git a/Engine/Plugins/Cameras/GameplayCameras/Source/GameplayCameras/Public/Core/BlendStackCameraNode.h b/Engine/Plugins/Cameras/GameplayCameras/Source/GameplayCameras/Public/Core/BlendStackCameraNode.h
--- a/Engine/Plugins/Cameras/GameplayCameras/Source/GameplayCameras/Public/Core/BlendStackCameraNode.h
+++ b/Engine/Plugins/Cameras/GameplayCameras/Source/GameplayCameras/Public/Core/BlendStackCameraNode.h
@@ -21,7 +21,7 @@ struct FBlendStackCameraPushParams
 	TObjectPtr<UCameraSystemEvaluator> Evaluator;
 
 	/** The evaluation context within which a camera mode's node tree should run. */
-	TWeakObjectPtr<const UCameraEvaluationContext> EvaluationContext;
+	TObjectPtr<const UCameraEvaluationContext> EvaluationContext;
 
 	/** The source camera mode asset to instantiate and push on the blend stack. */
 	TObjectPtr<const UCameraMode> CameraMode;
diff --git a/Engine/Plugins/Cameras/GameplayCameras/Source/GameplayCameras/Public/Core/CameraDirector.h b/Engine/Plugins/Cameras/GameplayCameras/Source/GameplayCameras/Public/Core/CameraDirector.h
--- a/Engine/Plugins/Cameras/GameplayCameras/Source/GameplayCameras/Public/Core/CameraDirector.h
+++ b/Engine/Plugins/Cameras/GameplayCameras/Source/GameplayCameras/Public/Core/CameraDirector.h
@@ -19,7 +19,7 @@ struct FCameraDirectorRunParams
 	float DeltaTime = 0.f;
 
 	/** The context in which this director runs. */
-	const UCameraEvaluationContext* OwnerContext = nullptr;
+	TObjectPtr<const UCameraEvaluationContext> OwnerContext;
 };
 
 /**
diff --git a/Engine/Plugins/Cameras/GameplayCameras/Source/GameplayCameras/Public/Core/RootCameraNode.h b/Engine/Plugins/Cameras/GameplayCameras/Source/GameplayCameras/Public/Core/RootCameraNode.h
--- a/Engine/Plugins/Cameras/GameplayCameras/Source/GameplayCameras/Public/Core/RootCameraNode.h
+++ b/Engine/Plugins/Cameras/GameplayCameras/Source/GameplayCameras/Public/Core/RootCameraNode.h
@@ -34,7 +34,7 @@ struct FActivateCameraModeParams
 	TObjectPtr<UCameraSystemEvaluator> Evaluator;
 
 	/** The evaluation context in which the camera mode runs. */
-	TWeakObjectPtr<const UCameraEvaluationContext> EvaluationContext;
+	TObjectPtr<const UCameraEvaluationContext> EvaluationContext;
 
 	/** The source camera mode asset that will be instantiated. */
 	TObjectPtr<const UCameraMode> CameraMode;
