Cameras: offset camera node can work in different spaces

diff --git a/Engine/Plugins/Cameras/GameplayCameras/Source/GameplayCameras/Private/Nodes/Common/OffsetCameraNode.cpp b/Engine/Plugins/Cameras/GameplayCameras/Source/GameplayCameras/Private/Nodes/Common/OffsetCameraNode.cpp
--- a/Engine/Plugins/Cameras/GameplayCameras/Source/GameplayCameras/Private/Nodes/Common/OffsetCameraNode.cpp
+++ b/Engine/Plugins/Cameras/GameplayCameras/Source/GameplayCameras/Private/Nodes/Common/OffsetCameraNode.cpp
@@ -2,12 +2,44 @@
 
 #include "Nodes/Common/OffsetCameraNode.h"
 
+#include "Core/CameraEvaluationContext.h"
+#include "GameplayCameras.h"
+
 #include UE_INLINE_GENERATED_CPP_BY_NAME(OffsetCameraNode)
 
 void UOffsetCameraNode::OnRun(const FCameraNodeRunParams& Params, FCameraNodeRunResult& OutResult)
 {
-	const FRotator3d Rotation = OutResult.CameraPose.GetRotation();
-	const FVector3d LocalOffset = Rotation.RotateVector(Offset);
+	FVector3d LocalOffset = Offset;
+	switch(OffsetSpace)
+	{
+		case ECameraNodeSpace::CameraPose:
+		default:
+			{
+				const FRotator3d Rotation = OutResult.CameraPose.GetRotation();
+				LocalOffset = Rotation.RotateVector(Offset);
+			}
+			break;
+		case ECameraNodeSpace::Context:
+			if (Params.EvaluationContext)
+			{ 
+				const FCameraNodeRunResult& InitialResult = Params.EvaluationContext->GetInitialResult();
+				ensureMsgf(InitialResult.bIsValid,
+						TEXT("OffsetCameraNode: using invalid context result as offset space!"));
+				const FRotator3d Rotation = InitialResult.CameraPose.GetRotation();
+				LocalOffset = Rotation.RotateVector(Offset);
+			}
+			else
+			{
+				UE_LOG(LogCameraSystem, Error, 
+						TEXT("OffsetCameraNode: cannot offset in context space when there is "
+							 "no current context set."));
+				return;
+			}
+			break;
+		case ECameraNodeSpace::World:
+			// Nothing to do;
+			break;
+	}
 
 	const FVector3d Location = OutResult.CameraPose.GetLocation();
 	OutResult.CameraPose.SetLocation(Location + LocalOffset);
diff --git a/Engine/Plugins/Cameras/GameplayCameras/Source/GameplayCameras/Public/Nodes/CameraNodeTypes.h b/Engine/Plugins/Cameras/GameplayCameras/Source/GameplayCameras/Public/Nodes/CameraNodeTypes.h
new file mode 100644
--- /dev/null
+++ b/Engine/Plugins/Cameras/GameplayCameras/Source/GameplayCameras/Public/Nodes/CameraNodeTypes.h
@@ -0,0 +1,20 @@
+// Copyright Epic Games, Inc. All Rights Reserved.
+
+#pragma once
+
+#include "CoreTypes.h"
+
+/**
+ * Defines what space a camera node should operate in.
+ */
+UENUM()
+enum class ECameraNodeSpace
+{
+	/** Operates in the local camera pose space. */
+	CameraPose,
+	/** Operates in the space of the owning evaluation context's initial result. */
+	Context,
+	/** Operates in world space. */
+	World
+};
+
diff --git a/Engine/Plugins/Cameras/GameplayCameras/Source/GameplayCameras/Public/Nodes/Common/OffsetCameraNode.h b/Engine/Plugins/Cameras/GameplayCameras/Source/GameplayCameras/Public/Nodes/Common/OffsetCameraNode.h
--- a/Engine/Plugins/Cameras/GameplayCameras/Source/GameplayCameras/Public/Nodes/Common/OffsetCameraNode.h
+++ b/Engine/Plugins/Cameras/GameplayCameras/Source/GameplayCameras/Public/Nodes/Common/OffsetCameraNode.h
@@ -3,6 +3,7 @@
 #pragma once
 
 #include "Core/CameraNode.h"
+#include "Nodes/CameraNodeTypes.h"
 
 #include "OffsetCameraNode.generated.h"
 
@@ -23,5 +24,9 @@ public:
 	/** The offset to apply to the camera, in local space. */
 	UPROPERTY(EditAnywhere, Category=Common)
 	FVector3d Offset;
+
+	/** The space in which to apply the offset. */
+	UPROPERTY(EditAnywhere, Category=Common)
+	ECameraNodeSpace OffsetSpace = ECameraNodeSpace::CameraPose;
 };
 
