Cameras: rewrite the core camera system architecture using evaluator objects built from the data nodes

The previous approach, which was about data nodes running themselves, and being cloned in the runtime to instantiate them, didn't yield the expected performance levels. It ran well in 95% of cases (under 0.02ms for a representative node tree), but had unexpected spikes of more than 0.5ms in a minority of cases. This new approach yields more consistent performance (max time was around 0.13ms, median time was the same), with many more opportunities for optimization, but at the cost of more boilerplate code to implement new camera nodes.

Camera node evaluators are built by the nodes, and are responsible for running the corresponding logic. Evaluators are stored in chunks of contiguous memory owned by either the camera system itself, or by the blend stack in which they belong. There are opportunities later to pre-compute the space needed for a given camera node hierarchy, so that we can guarantee all those evaluators are stored in a single allocation chunk.

Camera directors also run with a similar design, but their evaluators are simply stored as a shared pointer, since there's only a couple of directors active at one point.

This change also introduces a simple RTTI framework, since we need to safely inspect or downcast evaluators in some cases.

The live-edit manager interface is made empty for now since we don't need to forward property changes to instantiated objects anymore (evaluators for two instances of a camera mode share the same data nodes), but we keep it because we will later need to handle hot-reloading a node hierarchy by re-building its evaluator tree.

diff --git a/Engine/Plugins/Cameras/GameplayCameras/Source/GameplayCameras/Private/Core/BlendCameraNode.cpp b/Engine/Plugins/Cameras/GameplayCameras/Source/GameplayCameras/Private/Core/BlendCameraNode.cpp
--- a/Engine/Plugins/Cameras/GameplayCameras/Source/GameplayCameras/Private/Core/BlendCameraNode.cpp
+++ b/Engine/Plugins/Cameras/GameplayCameras/Source/GameplayCameras/Private/Core/BlendCameraNode.cpp
@@ -4,7 +4,9 @@
 
 #include UE_INLINE_GENERATED_CPP_BY_NAME(BlendCameraNode)
 
-void UBlendCameraNode::BlendResults(const FCameraNodeBlendParams& Params, FCameraNodeBlendResult& OutResult)
+UE_DEFINE_CAMERA_NODE_EVALUATOR(FBlendCameraNodeEvaluator)
+
+void FBlendCameraNodeEvaluator::BlendResults(const FCameraNodeBlendParams& Params, FCameraNodeBlendResult& OutResult)
 {
 	OnBlendResults(Params, OutResult);
 }
diff --git a/Engine/Plugins/Cameras/GameplayCameras/Source/GameplayCameras/Private/Core/BlendStackCameraNode.cpp b/Engine/Plugins/Cameras/GameplayCameras/Source/GameplayCameras/Private/Core/BlendStackCameraNode.cpp
--- a/Engine/Plugins/Cameras/GameplayCameras/Source/GameplayCameras/Private/Core/BlendStackCameraNode.cpp
+++ b/Engine/Plugins/Cameras/GameplayCameras/Source/GameplayCameras/Private/Core/BlendStackCameraNode.cpp
@@ -7,13 +7,19 @@
 #include "Core/CameraAsset.h"
 #include "Core/CameraEvaluationContext.h"
 #include "Core/CameraMode.h"
-#include "Core/CameraRuntimeInstantiator.h"
 #include "Core/CameraSystemEvaluator.h"
 #include "Nodes/Blends/PopBlendCameraNode.h"
 
 #include UE_INLINE_GENERATED_CPP_BY_NAME(BlendStackCameraNode)
 
-void UBlendStackCameraNode::Push(const FBlendStackCameraPushParams& Params)
+FCameraNodeEvaluatorPtr UBlendStackCameraNode::OnBuildEvaluator(FCameraNodeEvaluatorBuilder& Builder) const
+{
+	return Builder.BuildEvaluator<FBlendStackCameraNodeEvaluator>();
+}
+
+UE_DEFINE_CAMERA_NODE_EVALUATOR(FBlendStackCameraNodeEvaluator)
+
+void FBlendStackCameraNodeEvaluator::Push(const FBlendStackCameraPushParams& Params)
 {
 	if (!Entries.IsEmpty())
 	{
@@ -21,76 +27,70 @@ void UBlendStackCameraNode::Push(const FBlendStackCameraPushParams& Params)
 		// active camera mode.
 		const FCameraModeEntry& TopEntry(Entries.Top());
 		if (!TopEntry.bIsFrozen
-				&& TopEntry.OriginalCameraMode == Params.CameraMode
+				&& TopEntry.CameraMode == Params.CameraMode
 				&& TopEntry.EvaluationContext == Params.EvaluationContext)
 		{
 			return;
 		}
 	}
 
-	// Create the new root node. Instantiated objects will go inside it.
-	UBlendStackRootCameraNode* EntryRootNode = NewObject<UBlendStackRootCameraNode>(this, NAME_None);
-
-	// Instantiate the camera mode's node tree.
-	FCameraRuntimeInstantiationParams NodeTreeInstParams;
-	NodeTreeInstParams.InstantiationOuter = EntryRootNode;
-	NodeTreeInstParams.bAllowRecyling = true;
-
-	FCameraRuntimeInstantiator& Instantiator = Params.Evaluator->GetRuntimeInstantiator();
-	UCameraNode* ModeRootNode = Instantiator.InstantiateCameraNodeTree(Params.CameraMode->RootNode, NodeTreeInstParams);
-
-	// Find a transition and instantiate its blend. If not transition is found,
-	// make a camera cut transition.
-	UBlendCameraNode* Blend = nullptr;
-	if (const FCameraModeTransition* Transition = FindTransition(Params))
+	// Create the new root node to wrap the new camera mode's root node, and the specific
+	// blend node for this transition.
+	// We need to const-cast here to be able to use our own blend stack node as the outer
+	// of the new node.
+	UObject* Outer = const_cast<UObject*>((UObject*)GetCameraNode());
+	UBlendStackRootCameraNode* EntryRootNode = NewObject<UBlendStackRootCameraNode>(Outer, NAME_None);
 	{
-		FCameraRuntimeInstantiationParams BlendInstParams;
-		BlendInstParams.InstantiationOuter = EntryRootNode;
-		// No recycling on the blend.
+		UCameraNode* ModeRootNode = Params.CameraMode->RootNode;
+		EntryRootNode->RootNode = ModeRootNode;
 
-		UBlendCameraNode* ModeBlend = CastChecked<UBlendCameraNode>(
-				Instantiator.InstantiateCameraNodeTree(Transition->Blend, BlendInstParams));
-		Blend = ModeBlend;
-	}
-	else
-	{
-		Blend = NewObject<UPopBlendCameraNode>(EntryRootNode, NAME_None);
+		// Find a transition and use its blend. If no transition is found,
+		// make a camera cut transition.
+		UBlendCameraNode* ModeBlend = nullptr;
+		if (const FCameraModeTransition* Transition = FindTransition(Params))
+		{
+			ModeBlend = Transition->Blend;
+		}
+		else
+		{
+			ModeBlend = NewObject<UPopBlendCameraNode>(EntryRootNode, NAME_None);
+		}
+		EntryRootNode->Blend = ModeBlend;
 	}
 
-	EntryRootNode->Initialize(Blend, ModeRootNode);
-
-	// Make a new entry and add it to the stack.
+	// Make the new stack entry, and use its storage buffer to build the tree of evaluators.
 	FCameraModeEntry NewEntry;
+
+	FCameraNodeEvaluatorTreeBuilderParams BuildParams;
+	BuildParams.RootCameraNode = EntryRootNode;
+	BuildParams.Evaluator = Params.Evaluator;
+	BuildParams.EvaluationContext = Params.EvaluationContext;
+	FCameraNodeEvaluator* RootEvaluator = NewEntry.EvaluatorStorage.BuildEvaluatorTree(BuildParams);
+
 	NewEntry.EvaluationContext = Params.EvaluationContext;
-	NewEntry.OriginalCameraMode = Params.CameraMode;
-	NewEntry.RootNode = EntryRootNode;
+	NewEntry.CameraMode = Params.CameraMode;
+	NewEntry.RootEvaluator = RootEvaluator->CastThisChecked<FBlendStackRootCameraNodeEvaluator>();
 	NewEntry.bIsFirstFrame = true;
 
+	// Important: we need to move the new entry here because copying evaluator storage
+	// is disabled.
 	Entries.Add(MoveTemp(NewEntry));
 }
 
-void UBlendStackCameraNode::AddReferencedObjects(UObject* InThis, FReferenceCollector& Collector)
+FCameraNodeEvaluatorChildrenView FBlendStackCameraNodeEvaluator::OnGetChildren()
 {
-	UBlendStackCameraNode* TypedThis = CastChecked<UBlendStackCameraNode>(InThis);
-	for (FCameraModeEntry& Entry : TypedThis->Entries)
-	{
-		Collector.AddReferencedObject(Entry.OriginalCameraMode);
-		Collector.AddReferencedObject(Entry.RootNode);
-	}
-}
-
-FCameraNodeChildrenView UBlendStackCameraNode::OnGetChildren()
-{
-	FCameraNodeChildrenView View;
+	FCameraNodeEvaluatorChildrenView View;
 	for (FCameraModeEntry& Entry : Entries)
 	{
-		View.Add(Entry.RootNode);
+		View.Add(Entry.RootEvaluator);
 	}
 	return View;
 }
 
-void UBlendStackCameraNode::OnRun(const FCameraNodeRunParams& Params, FCameraNodeRunResult& OutResult)
+void FBlendStackCameraNodeEvaluator::OnRun(const FCameraNodeEvaluationParams& Params, FCameraNodeEvaluationResult& OutResult)
 {
+	const UBlendStackCameraNode* BlendStackNode = GetCameraNodeAs<UBlendStackCameraNode>();
+
 	// Start by evaluating all the root nodes in the stack.
 	for (FCameraModeEntry& Entry : Entries)
 	{
@@ -101,17 +101,17 @@ void UBlendStackCameraNode::OnRun(const FCameraNodeRunParams& Params, FCameraNod
 			continue;
 		}
 
-		FCameraNodeRunParams CurParams(Params);
+		FCameraNodeEvaluationParams CurParams(Params);
 		CurParams.EvaluationContext = CurContext;
 		CurParams.bIsFirstFrame = Entry.bIsFirstFrame;
 
-		FCameraNodeRunResult& CurResult(Entry.Result);
+		FCameraNodeEvaluationResult& CurResult(Entry.Result);
 
 		if (!Entry.bIsFrozen)
 		{
 			// If the context in which this camera mode runs doesn't have a valid result,
 			// skip it.
-			const FCameraNodeRunResult& ContextResult(CurContext->GetInitialResult());
+			const FCameraNodeEvaluationResult& ContextResult(CurContext->GetInitialResult());
 			if (UNLIKELY(!ContextResult.bIsValid))
 			{
 				CurResult.bIsValid = false;
@@ -128,12 +128,12 @@ void UBlendStackCameraNode::OnRun(const FCameraNodeRunParams& Params, FCameraNod
 			CurResult.bIsValid = true;
 
 			// Run the camera mode!
-			Entry.RootNode->Run(CurParams, CurResult);
+			Entry.RootEvaluator->Run(CurParams, CurResult);
 		}
 		else
 		{
 			// Only evaluate the blend via the root node.
-			Entry.RootNode->Run(CurParams, CurResult);
+			Entry.RootEvaluator->Run(CurParams, CurResult);
 		}
 	}
 
@@ -144,7 +144,7 @@ void UBlendStackCameraNode::OnRun(const FCameraNodeRunParams& Params, FCameraNod
 	int32 PopEntriesBelow = INDEX_NONE;
 	for (FCameraModeEntry& Entry : Entries)
 	{
-		FCameraNodeRunResult& CurResult(Entry.Result);
+		FCameraNodeEvaluationResult& CurResult(Entry.Result);
 		if (UNLIKELY(!CurResult.bIsValid))
 		{
 			continue;
@@ -152,17 +152,17 @@ void UBlendStackCameraNode::OnRun(const FCameraNodeRunParams& Params, FCameraNod
 
 		const FCameraPoseFlags ChangedFlags(CurResult.CameraPose.GetChangedFlags());
 
-		FCameraNodeRunParams CurParams(Params);
+		FCameraNodeEvaluationParams CurParams(Params);
 		CurParams.EvaluationContext = Entry.EvaluationContext.Get();
 		CurParams.bIsFirstFrame = Entry.bIsFirstFrame;
 		FCameraNodeBlendParams BlendParams(CurParams, CurResult);
 
 		FCameraNodeBlendResult BlendResult(OutResult);
 
-		UBlendCameraNode* EntryBlend = Entry.RootNode->GetBlend();
-		if (EntryBlend)
+		FBlendCameraNodeEvaluator* EntryBlendEvaluator = Entry.RootEvaluator->GetBlendEvaluator();
+		if (EntryBlendEvaluator)
 		{
-			EntryBlend->BlendResults(BlendParams, BlendResult);
+			EntryBlendEvaluator->BlendResults(BlendParams, BlendResult);
 
 			if (BlendResult.bIsBlendFull && BlendResult.bIsBlendFinished)
 			{
@@ -182,23 +182,10 @@ void UBlendStackCameraNode::OnRun(const FCameraNodeRunParams& Params, FCameraNod
 	}
 
 	// Pop out camera modes that have been blended out.
-	if (bAutoPop && PopEntriesBelow != INDEX_NONE)
+	if (BlendStackNode->bAutoPop && PopEntriesBelow != INDEX_NONE)
 	{
-		FCameraRuntimeInstantiator& Instantiator = Params.Evaluator->GetRuntimeInstantiator();
 		for (int32 Index = 0; Index < PopEntriesBelow; ++Index)
 		{
-			FCameraModeEntry& Entry = Entries[0];
-			if (!Entry.bIsFrozen)
-			{
-				// Recycle the camera mode's node hierarchy if possible.
-				// Reset all nodes in it so they are back to their default state, ready to be
-				// re-used later.
-				TObjectPtr<const UCameraNode> OriginalRootNode = Entry.OriginalCameraMode->RootNode;
-				TObjectPtr<UCameraNode> RecycledRootNode = Entry.RootNode->GetRootNode();
-				Entry.RootNode->Reset(FCameraNodeResetParams());
-
-				Instantiator.RecycleInstantiatedObject(OriginalRootNode, RecycledRootNode);
-			}
 			Entries.RemoveAt(0);
 		}
 	}
@@ -210,8 +197,10 @@ void UBlendStackCameraNode::OnRun(const FCameraNodeRunParams& Params, FCameraNod
 	}
 }
 
-const FCameraModeTransition* UBlendStackCameraNode::FindTransition(const FBlendStackCameraPushParams& Params) const
+const FCameraModeTransition* FBlendStackCameraNodeEvaluator::FindTransition(const FBlendStackCameraPushParams& Params) const
 {
+	const UBlendStackCameraNode* BlendStackNode = GetCameraNodeAs<UBlendStackCameraNode>();
+
 	const UCameraEvaluationContext* ToContext = Params.EvaluationContext.Get();
 	const UCameraAsset* ToCameraAsset = ToContext ? ToContext->GetCameraAsset() : nullptr;
 	const UCameraMode* ToCameraMode = Params.CameraMode;
@@ -229,7 +218,7 @@ const FCameraModeTransition* UBlendStackCameraNode::FindTransition(const FBlendS
 
 		const UCameraEvaluationContext* FromContext = TopEntry.EvaluationContext.Get();
 		const UCameraAsset* FromCameraAsset = FromContext ? FromContext->GetCameraAsset() : nullptr;
-		const UCameraMode* FromCameraMode = TopEntry.OriginalCameraMode;
+		const UCameraMode* FromCameraMode = TopEntry.CameraMode;
 
 		if (!TopEntry.bIsFrozen)
 		{
@@ -280,7 +269,7 @@ const FCameraModeTransition* UBlendStackCameraNode::FindTransition(const FBlendS
 			}
 		}
 	}
-	else if (bBlendFirstCameraMode)
+	else if (BlendStackNode->bBlendFirstCameraMode)
 	{
 		return FindTransition(
 				ToCameraMode->EnterTransitions,
@@ -291,7 +280,7 @@ const FCameraModeTransition* UBlendStackCameraNode::FindTransition(const FBlendS
 	return nullptr;
 }
 
-const FCameraModeTransition* UBlendStackCameraNode::FindTransition(
+const FCameraModeTransition* FBlendStackCameraNodeEvaluator::FindTransition(
 			TArrayView<const FCameraModeTransition> Transitions, 
 			const UCameraMode* FromCameraMode, const UCameraAsset* FromCameraAsset, bool bFromFrozen,
 			const UCameraMode* ToCameraMode, const UCameraAsset* ToCameraAsset) const
diff --git a/Engine/Plugins/Cameras/GameplayCameras/Source/GameplayCameras/Private/Core/BlendStackRootCameraNode.cpp b/Engine/Plugins/Cameras/GameplayCameras/Source/GameplayCameras/Private/Core/BlendStackRootCameraNode.cpp
--- a/Engine/Plugins/Cameras/GameplayCameras/Source/GameplayCameras/Private/Core/BlendStackRootCameraNode.cpp
+++ b/Engine/Plugins/Cameras/GameplayCameras/Source/GameplayCameras/Private/Core/BlendStackRootCameraNode.cpp
@@ -7,33 +7,43 @@
 
 #include UE_INLINE_GENERATED_CPP_BY_NAME(BlendStackRootCameraNode)
 
-void UBlendStackRootCameraNode::Initialize(UBlendCameraNode* InBlend, UCameraNode* InRootNode)
+FCameraNodeEvaluatorPtr UBlendStackRootCameraNode::OnBuildEvaluator(FCameraNodeEvaluatorBuilder& Builder) const
 {
-	Blend = InBlend;
-	RootNode = InRootNode;
-
-	TreeCache.Build(RootNode);
+	return Builder.BuildEvaluator<FBlendStackRootCameraNodeEvaluator>();
 }
 
-FCameraNodeChildrenView UBlendStackRootCameraNode::OnGetChildren()
+UE_DEFINE_CAMERA_NODE_EVALUATOR(FBlendStackRootCameraNodeEvaluator)
+
+FCameraNodeEvaluatorChildrenView FBlendStackRootCameraNodeEvaluator::OnGetChildren()
 {
-	return FCameraNodeChildrenView({ Blend, RootNode });
+	FCameraNodeEvaluatorChildrenView Children;
+	if (BlendEvaluator)
+	{
+		Children.Add(BlendEvaluator);
+	}
+	if (RootEvaluator)
+	{
+		Children.Add(RootEvaluator);
+	}
+	return Children;
 }
 
-void UBlendStackRootCameraNode::OnReset(const FCameraNodeResetParams& Params)
+void FBlendStackRootCameraNodeEvaluator::OnInitialize(const FCameraNodeEvaluatorInitializeParams& Params)
 {
-	TreeCache.ForEachNode(ECameraNodeFlags::RequiresReset, [Params](UCameraNode* Node) { Node->Reset(Params); });
+	const UBlendStackRootCameraNode* RootNode = GetCameraNodeAs<UBlendStackRootCameraNode>();
+	BlendEvaluator = Params.BuildEvaluatorAs<FBlendCameraNodeEvaluator>(RootNode->Blend);
+	RootEvaluator = Params.BuildEvaluator(RootNode->RootNode);
 }
 
-void UBlendStackRootCameraNode::OnRun(const FCameraNodeRunParams& Params, FCameraNodeRunResult& OutResult)
+void FBlendStackRootCameraNodeEvaluator::OnRun(const FCameraNodeEvaluationParams& Params, FCameraNodeEvaluationResult& OutResult)
 {
-	if (Blend)
+	if (BlendEvaluator)
 	{
-		Blend->Run(Params, OutResult);
+		BlendEvaluator->Run(Params, OutResult);
 	}
-	if (RootNode)
+	if (RootEvaluator)
 	{
-		RootNode->Run(Params, OutResult);
+		RootEvaluator->Run(Params, OutResult);
 	}
 }
 
diff --git a/Engine/Plugins/Cameras/GameplayCameras/Source/GameplayCameras/Private/Core/BlendStackRootCameraNode.h b/Engine/Plugins/Cameras/GameplayCameras/Source/GameplayCameras/Private/Core/BlendStackRootCameraNode.h
--- a/Engine/Plugins/Cameras/GameplayCameras/Source/GameplayCameras/Private/Core/BlendStackRootCameraNode.h
+++ b/Engine/Plugins/Cameras/GameplayCameras/Source/GameplayCameras/Private/Core/BlendStackRootCameraNode.h
@@ -3,10 +3,11 @@
 #pragma once
 
 #include "Core/CameraNode.h"
-#include "Core/CameraNodeTreeCache.h"
+#include "Core/CameraNodeEvaluator.h"
 
 #include "BlendStackRootCameraNode.generated.h"
 
+class FBlendCameraNodeEvaluator;
 class UBlendCameraNode;
 class UCameraMode;
 
@@ -20,24 +21,12 @@ class UBlendStackRootCameraNode : public UCameraNode
 {
 	GENERATED_BODY()
 
-public:
-
-	/** Initialize this node. */
-	void Initialize(UBlendCameraNode* InBlend, UCameraNode* InRootNode);
-
-	/** Gets the blend node. */
-	UBlendCameraNode* GetBlend() const { return Blend; }
-
-	/** Gets the root of the camera mode. */
-	UCameraNode* GetRootNode() const { return RootNode; }
-
 protected:
 
-	virtual FCameraNodeChildrenView OnGetChildren() override;
-	virtual void OnReset(const FCameraNodeResetParams& Params) override;
-	virtual void OnRun(const FCameraNodeRunParams& Params, FCameraNodeRunResult& OutResult) override;
+	// UCameraNode interface.
+	virtual FCameraNodeEvaluatorPtr OnBuildEvaluator(FCameraNodeEvaluatorBuilder& Builder) const override;
 
-private:
+public:
 
 	/** The blend to use on the camera mode. */
 	UPROPERTY()
@@ -46,8 +35,30 @@ private:
 	/** The root of the instantied camera node tree. */
 	UPROPERTY()
 	TObjectPtr<UCameraNode> RootNode;
+};
+
+/**
+ * Evaluator for the blend stack entry root node.
+ */
+class FBlendStackRootCameraNodeEvaluator : public FCameraNodeEvaluator
+{
+	UE_DECLARE_CAMERA_NODE_EVALUATOR(FBlendStackRootCameraNodeEvaluator)
+
+public:
+
+	FBlendCameraNodeEvaluator* GetBlendEvaluator() const { return BlendEvaluator; }
+	FCameraNodeEvaluator* GetRootEvaluator() const { return RootEvaluator; }
+
+protected:
+
+	// FCameraNodeEvaluator interface.
+	virtual FCameraNodeEvaluatorChildrenView OnGetChildren() override;
+	virtual void OnInitialize(const FCameraNodeEvaluatorInitializeParams& Params) override;
+	virtual void OnRun(const FCameraNodeEvaluationParams& Params, FCameraNodeEvaluationResult& OutResult) override;
+
+private:
 
-	/** Node cache for the hierarchy below the root node. */
-	FCameraNodeTreeCache TreeCache;
+	FBlendCameraNodeEvaluator* BlendEvaluator = nullptr;
+	FCameraNodeEvaluator* RootEvaluator = nullptr;
 };
 
diff --git a/Engine/Plugins/Cameras/GameplayCameras/Source/GameplayCameras/Private/Core/CameraDirector.cpp b/Engine/Plugins/Cameras/GameplayCameras/Source/GameplayCameras/Private/Core/CameraDirector.cpp
--- a/Engine/Plugins/Cameras/GameplayCameras/Source/GameplayCameras/Private/Core/CameraDirector.cpp
+++ b/Engine/Plugins/Cameras/GameplayCameras/Source/GameplayCameras/Private/Core/CameraDirector.cpp
@@ -4,8 +4,10 @@
 
 #include UE_INLINE_GENERATED_CPP_BY_NAME(CameraDirector)
 
-void UCameraDirector::Run(const FCameraDirectorRunParams& Params, FCameraDirectorRunResult& OutResult)
+FCameraDirectorEvaluator* UCameraDirector::BuildEvaluator(FCameraDirectorEvaluatorBuilder& Builder) const
 {
-	OnRun(Params, OutResult);
+	FCameraDirectorEvaluator* NewEvaluator = OnBuildEvaluator(Builder);
+	NewEvaluator->SetPrivateCameraDirector(this);
+	return NewEvaluator;
 }
 
diff --git a/Engine/Plugins/Cameras/GameplayCameras/Source/GameplayCameras/Private/Core/CameraDirectorEvaluator.cpp b/Engine/Plugins/Cameras/GameplayCameras/Source/GameplayCameras/Private/Core/CameraDirectorEvaluator.cpp
new file mode 100644
--- /dev/null
+++ b/Engine/Plugins/Cameras/GameplayCameras/Source/GameplayCameras/Private/Core/CameraDirectorEvaluator.cpp
@@ -0,0 +1,22 @@
+// Copyright Epic Games, Inc. All Rights Reserved.
+
+#include "Core/CameraDirectorEvaluator.h"
+
+#include "Core/CameraDirector.h"
+
+UE_GAMEPLAY_CAMERAS_DEFINE_RTTI(FCameraDirectorEvaluator)
+
+FCameraDirectorEvaluator::FCameraDirectorEvaluator()
+{
+}
+
+void FCameraDirectorEvaluator::Run(const FCameraDirectorEvaluationParams& Params, FCameraDirectorEvaluationResult& OutResult)
+{
+	OnRun(Params, OutResult);
+}
+
+void FCameraDirectorEvaluator::SetPrivateCameraDirector(TObjectPtr<const UCameraDirector> InCameraDirector)
+{
+	PrivateCameraDirector = InCameraDirector;
+}
+
diff --git a/Engine/Plugins/Cameras/GameplayCameras/Source/GameplayCameras/Private/Core/CameraEvaluationContextStack.cpp b/Engine/Plugins/Cameras/GameplayCameras/Source/GameplayCameras/Private/Core/CameraEvaluationContextStack.cpp
--- a/Engine/Plugins/Cameras/GameplayCameras/Source/GameplayCameras/Private/Core/CameraEvaluationContextStack.cpp
+++ b/Engine/Plugins/Cameras/GameplayCameras/Source/GameplayCameras/Private/Core/CameraEvaluationContextStack.cpp
@@ -4,8 +4,8 @@
 
 #include "Core/CameraAsset.h"
 #include "Core/CameraDirector.h"
+#include "Core/CameraDirectorEvaluator.h"
 #include "Core/CameraEvaluationContext.h"
-#include "Core/CameraRuntimeInstantiator.h"
 #include "Core/CameraSystemEvaluator.h"
 #include "UObject/Package.h"
 
@@ -15,7 +15,7 @@ FCameraEvaluationContextInfo FCameraEvaluationContextStack::GetActiveContext() c
 	{
 		if (UCameraEvaluationContext* Context = Entry.WeakContext.Get())
 		{
-			return FCameraEvaluationContextInfo{ Context, Entry.CameraDirector };
+			return FCameraEvaluationContextInfo{ Context, Entry.CameraDirector, Entry.Evaluator };
 		}
 	}
 	return FCameraEvaluationContextInfo();
@@ -44,30 +44,24 @@ void FCameraEvaluationContextStack::PushContext(UCameraEvaluationContext* Contex
 	{
 		if (ExistingIndex < Entries.Num() - 1)
 		{
-			const FContextEntry EntryCopy(Entries[ExistingIndex]);
+			FContextEntry EntryCopy(MoveTemp(Entries[ExistingIndex]));
 			Entries.RemoveAt(ExistingIndex);
-			Entries.Add(EntryCopy);
+			Entries.Add(MoveTemp(EntryCopy));
 		}
 		return;
 	}
 
-	// Instantiate the camera director.
-	FCameraRuntimeInstantiationParams InstParams;
-	InstParams.InstantiationOuter = Evaluator;
-	if (!InstParams.InstantiationOuter)
-	{
-		InstParams.InstantiationOuter = GetTransientPackage();
-	}
+	// Make a new entry and build the director evaluator using the entry's storage.
+	FContextEntry NewEntry;
 
-	const UCameraDirector* OriginalCameraDirector = Context->GetCameraAsset()->CameraDirector;
-	UCameraDirector* NewCameraDirector = Evaluator->GetRuntimeInstantiator().InstantiateCameraDirector(
-			OriginalCameraDirector, InstParams);
+	FCameraDirectorEvaluatorBuilder Builder(NewEntry.EvaluatorStorage);
+	const UCameraDirector* CameraDirector = Context->GetCameraAsset()->CameraDirector;
+	FCameraDirectorEvaluator* DirectorEvaluator = CameraDirector->BuildEvaluator(Builder);
 	
-	// Add an entry in the stack.
-	FContextEntry NewEntry;
 	NewEntry.WeakContext = Context;
-	NewEntry.CameraDirector = NewCameraDirector;
-	Entries.Push(NewEntry);
+	NewEntry.CameraDirector = CameraDirector;
+	NewEntry.Evaluator = DirectorEvaluator;
+	Entries.Push(MoveTemp(NewEntry));
 }
 
 bool FCameraEvaluationContextStack::RemoveContext(UCameraEvaluationContext* Context)
diff --git a/Engine/Plugins/Cameras/GameplayCameras/Source/GameplayCameras/Private/Core/CameraNode.cpp b/Engine/Plugins/Cameras/GameplayCameras/Source/GameplayCameras/Private/Core/CameraNode.cpp
--- a/Engine/Plugins/Cameras/GameplayCameras/Source/GameplayCameras/Private/Core/CameraNode.cpp
+++ b/Engine/Plugins/Cameras/GameplayCameras/Source/GameplayCameras/Private/Core/CameraNode.cpp
@@ -2,42 +2,26 @@
 
 #include "Core/CameraNode.h"
 
-#include "Core/CameraRuntimeInstantiator.h"
+#include "Core/CameraNodeEvaluator.h"
 
 #include UE_INLINE_GENERATED_CPP_BY_NAME(CameraNode)
 
-void FCameraNodeRunResult::Reset()
+void FCameraNodeEvaluationResult::Reset()
 {
 	CameraPose.Reset();
 	bIsCameraCut = false;
 	bIsValid = false;
 }
 
-FCameraNodeChildrenView UCameraNode::GetChildren()
+FCameraNodeEvaluatorAllocationInfo UCameraNode::GetEvaluatorInfo() const
 {
-	return OnGetChildren();
+	return OnGetEvaluatorInfo();
 }
 
-void UCameraNode::Reset(const FCameraNodeResetParams& Params)
+FCameraNodeEvaluatorPtr UCameraNode::BuildEvaluator(FCameraNodeEvaluatorBuilder& Builder) const
 {
-	OnReset(Params);
+	FCameraNodeEvaluator* NewEvaluator = OnBuildEvaluator(Builder);
+	NewEvaluator->SetPrivateCameraNode(this);
+	return NewEvaluator;
 }
 
-void UCameraNode::Run(const FCameraNodeRunParams& Params, FCameraNodeRunResult& OutResult)
-{
-	if (bIsEnabled)
-	{
-		OnRun(Params, OutResult);
-	}
-}
-
-#if WITH_EDITOR
-
-void UCameraNode::PostEditChangeProperty(struct FPropertyChangedEvent& PropertyChangedEvent)
-{
-	Super::PostEditChangeProperty(PropertyChangedEvent);
-	
-	FCameraRuntimeInstantiator::ForwardPropertyChange(this, PropertyChangedEvent);
-}
-
-#endif
diff --git a/Engine/Plugins/Cameras/GameplayCameras/Source/GameplayCameras/Private/Core/CameraNodeEvaluator.cpp b/Engine/Plugins/Cameras/GameplayCameras/Source/GameplayCameras/Private/Core/CameraNodeEvaluator.cpp
new file mode 100644
--- /dev/null
+++ b/Engine/Plugins/Cameras/GameplayCameras/Source/GameplayCameras/Private/Core/CameraNodeEvaluator.cpp
@@ -0,0 +1,61 @@
+// Copyright Epic Games, Inc. All Rights Reserved.
+
+#include "Core/CameraNodeEvaluator.h"
+
+#include "Core/CameraNode.h"
+#include "UObject/UObjectGlobals.h"
+
+UE_GAMEPLAY_CAMERAS_DEFINE_RTTI(FCameraNodeEvaluator)
+
+FCameraNodeEvaluator* FCameraNodeEvaluatorInitializeParams::BuildEvaluator(const UCameraNode* InNode) const
+{
+	FCameraNodeEvaluator* NewEvaluator = InNode->BuildEvaluator(*Builder);
+	NewEvaluator->Initialize(*this);
+	return NewEvaluator;
+}
+
+FCameraNodeEvaluator::FCameraNodeEvaluator()
+{
+}
+
+void FCameraNodeEvaluator::SetPrivateCameraNode(TObjectPtr<const UCameraNode> InCameraNode)
+{
+	PrivateCameraNode = InCameraNode;
+}
+
+FCameraNodeEvaluatorChildrenView FCameraNodeEvaluator::GetChildren()
+{
+	return OnGetChildren();
+}
+
+void FCameraNodeEvaluator::AddReferencedObjects(FReferenceCollector& Collector)
+{
+	if (PrivateCameraNode)
+	{
+		Collector.AddReferencedObject(PrivateCameraNode);
+	}
+
+	OnAddReferencedObjects(Collector);
+
+	for (FCameraNodeEvaluator* Child : GetChildren())
+	{
+		if (Child)
+		{
+			Child->AddReferencedObjects(Collector);
+		}
+	}
+}
+
+void FCameraNodeEvaluator::Initialize(const FCameraNodeEvaluatorInitializeParams& Params)
+{
+	OnInitialize(Params);
+}
+
+void FCameraNodeEvaluator::Run(const FCameraNodeEvaluationParams& Params, FCameraNodeEvaluationResult& OutResult)
+{
+	if (PrivateCameraNode->bIsEnabled)
+	{
+		OnRun(Params, OutResult);
+	}
+}
+
diff --git a/Engine/Plugins/Cameras/GameplayCameras/Source/GameplayCameras/Private/Core/CameraNodeEvaluatorStorage.cpp b/Engine/Plugins/Cameras/GameplayCameras/Source/GameplayCameras/Private/Core/CameraNodeEvaluatorStorage.cpp
new file mode 100644
--- /dev/null
+++ b/Engine/Plugins/Cameras/GameplayCameras/Source/GameplayCameras/Private/Core/CameraNodeEvaluatorStorage.cpp
@@ -0,0 +1,82 @@
+// Copyright Epic Games, Inc. All Rights Reserved.
+
+#include "Core/CameraNodeEvaluatorStorage.h"
+
+#include "Core/CameraMode.h"
+#include "Core/CameraNode.h"
+#include "Core/CameraNodeEvaluator.h"
+#include "UObject/Package.h"
+
+FCameraNodeEvaluatorStorage::FCameraNodeEvaluatorStorage()
+{
+}
+
+FCameraNodeEvaluatorStorage::FCameraNodeEvaluatorStorage(FCameraNodeEvaluatorStorage&& Other)
+{
+	Allocations = MoveTemp(Other.Allocations);
+	EvaluatorInfos = MoveTemp(Other.EvaluatorInfos);
+}
+
+FCameraNodeEvaluatorStorage::~FCameraNodeEvaluatorStorage()
+{
+	for (FEvaluatorInfo& EvaluatorInfo : EvaluatorInfos)
+	{
+		FCameraNodeEvaluator* Evaluator(EvaluatorInfo.Ptr);
+		Evaluator->~FCameraNodeEvaluator();
+	}
+	EvaluatorInfos.Reset();
+
+	for (FAllocation& Allocation : Allocations)
+	{
+		FMemory::Free(Allocation.Memory);
+	}
+	Allocations.Reset();
+}
+
+FCameraNodeEvaluatorTreeAllocationInfo FCameraNodeEvaluatorStorage::ComputeTreeInfo(const UCameraMode* CameraMode) const
+{
+	int16 MaxSize = 0;
+	int16 MaxAlignment = 0;
+
+	// This isn't ideal, maybe we need a GetChildren API on the UCameraNode class too.
+	UPackage* Package = CameraMode->GetPackage();
+	ForEachObjectWithPackage(Package, [&MaxSize, &MaxAlignment](UObject* Object)
+	{
+		if (UCameraNode* CurNode = Cast<UCameraNode>(Object))
+		{
+			FCameraNodeEvaluatorAllocationInfo CurInfo = CurNode->GetEvaluatorInfo();
+
+			MaxAlignment = FMath::Max(MaxAlignment, CurInfo.Alignof);
+			MaxSize = Align(MaxSize, CurInfo.Alignof) + CurInfo.Sizeof;
+		}
+		return true; // Keep iterating.
+	});
+
+	return FCameraNodeEvaluatorTreeAllocationInfo{ MaxSize, MaxAlignment };
+}
+
+FCameraNodeEvaluatorPtr FCameraNodeEvaluatorStorage::BuildEvaluatorTree(const FCameraNodeEvaluatorTreeBuilderParams& Params)
+{
+	if (Params.AllocationInfo)
+	{
+		const uint16 NewCapacity = Params.AllocationInfo->TotalSizeof;
+		const uint16 NewAlignment = Params.AllocationInfo->MaxAlignof;
+
+		FAllocation& NewAllocation = Allocations.Emplace_GetRef();
+		NewAllocation.Memory = reinterpret_cast<uint8*>(FMemory::Malloc(NewCapacity, NewAlignment));
+		NewAllocation.Capacity = NewCapacity;
+		NewAllocation.Used = 0;
+	}
+
+	FCameraNodeEvaluatorBuilder Builder(*this);
+
+	FCameraNodeEvaluatorInitializeParams InitParams;
+	InitParams.Evaluator = Params.Evaluator;
+	InitParams.EvaluationContext = Params.EvaluationContext;
+	InitParams.Builder = &Builder;
+
+	FCameraNodeEvaluatorPtr RootEvaluator = InitParams.BuildEvaluator(Params.RootCameraNode);
+
+	return RootEvaluator;
+}
+
diff --git a/Engine/Plugins/Cameras/GameplayCameras/Source/GameplayCameras/Private/Core/CameraNodeTreeCache.cpp b/Engine/Plugins/Cameras/GameplayCameras/Source/GameplayCameras/Private/Core/CameraNodeTreeCache.cpp
deleted file mode 100644
diff --git a/Engine/Plugins/Cameras/GameplayCameras/Source/GameplayCameras/Private/Core/CameraObjectRtti.cpp b/Engine/Plugins/Cameras/GameplayCameras/Source/GameplayCameras/Private/Core/CameraObjectRtti.cpp
new file mode 100644
--- /dev/null
+++ b/Engine/Plugins/Cameras/GameplayCameras/Source/GameplayCameras/Private/Core/CameraObjectRtti.cpp
@@ -0,0 +1,10 @@
+// Copyright Epic Games, Inc. All Rights Reserved.
+
+#include "Core/CameraObjectRtti.h"
+
+uint32 FCameraObjectTypeID::RegisterNewID()
+{
+	static uint32 ID = 0;
+	return ID++;
+}
+
diff --git a/Engine/Plugins/Cameras/GameplayCameras/Source/GameplayCameras/Private/Core/CameraRuntimeInstantiator.cpp b/Engine/Plugins/Cameras/GameplayCameras/Source/GameplayCameras/Private/Core/CameraRuntimeInstantiator.cpp
deleted file mode 100644
diff --git a/Engine/Plugins/Cameras/GameplayCameras/Source/GameplayCameras/Private/Core/CameraSystemEvaluator.cpp b/Engine/Plugins/Cameras/GameplayCameras/Source/GameplayCameras/Private/Core/CameraSystemEvaluator.cpp
--- a/Engine/Plugins/Cameras/GameplayCameras/Source/GameplayCameras/Private/Core/CameraSystemEvaluator.cpp
+++ b/Engine/Plugins/Cameras/GameplayCameras/Source/GameplayCameras/Private/Core/CameraSystemEvaluator.cpp
@@ -4,6 +4,7 @@
 
 #include "Camera/CameraTypes.h"
 #include "Core/CameraAsset.h"
+#include "Core/CameraDirectorEvaluator.h"
 #include "Core/CameraEvaluationContext.h"
 #include "Core/CameraMode.h"
 #include "Core/DefaultRootCameraNode.h"
@@ -19,14 +20,25 @@ UCameraSystemEvaluator::UCameraSystemEvaluator(const FObjectInitializer& ObjectI
 {
 	RootNode = ObjectInit.CreateDefaultSubobject<UDefaultRootCameraNode>(this, TEXT("RootNode"));
 
-	ContextStack.Initialize(this);
+	if (!HasAnyFlags(RF_ClassDefaultObject | RF_ArchetypeObject))
+	{
+		ContextStack.Initialize(this);
+
+		FCameraNodeEvaluatorTreeBuilderParams BuildParams;
+		BuildParams.Evaluator = this;
+		BuildParams.RootCameraNode = RootNode;
+		RootEvaluator = static_cast<FRootCameraNodeEvaluator*>(RootEvaluatorStorage.BuildEvaluatorTree(BuildParams));
+	}
 }
 
 void UCameraSystemEvaluator::AddReferencedObjects(UObject* InThis, FReferenceCollector& Collector)
 {
 	UCameraSystemEvaluator* TypedThis = CastChecked<UCameraSystemEvaluator>(InThis);
 	TypedThis->ContextStack.AddReferencedObjects(Collector);
-	TypedThis->Instantiator.AddReferencedObjects(Collector);
+	if (TypedThis->RootEvaluator)
+	{
+		TypedThis->RootEvaluator->AddReferencedObjects(Collector);
+	}
 }
 
 void UCameraSystemEvaluator::PushEvaluationContext(UCameraEvaluationContext* EvaluationContext)
@@ -57,16 +69,16 @@ void UCameraSystemEvaluator::Update(const FCameraSystemEvaluationUpdateParams& P
 	}
 
 	// Run the camera director, and activate any camera mode(s) it returns to us.
-	UCameraDirector* ActiveDirector = ActiveContextInfo.CameraDirector;
-	if (ActiveDirector)
+	FCameraDirectorEvaluator* ActiveDirectorEvaluator = ActiveContextInfo.Evaluator;
+	if (ActiveDirectorEvaluator)
 	{
-		FCameraDirectorRunParams DirectorParams;
+		FCameraDirectorEvaluationParams DirectorParams;
 		DirectorParams.DeltaTime = Params.DeltaTime;
 		DirectorParams.OwnerContext = ActiveContextInfo.EvaluationContext;
 
-		FCameraDirectorRunResult DirectorResult;
+		FCameraDirectorEvaluationResult DirectorResult;
 
-		ActiveDirector->Run(DirectorParams, DirectorResult);
+		ActiveDirectorEvaluator->Run(DirectorParams, DirectorResult);
 
 		if (DirectorResult.ActiveCameraModes.Num() == 1)
 		{
@@ -74,18 +86,18 @@ void UCameraSystemEvaluator::Update(const FCameraSystemEvaluationUpdateParams& P
 			CameraModeParams.Evaluator = this;
 			CameraModeParams.EvaluationContext = ActiveContextInfo.EvaluationContext;
 			CameraModeParams.CameraMode = DirectorResult.ActiveCameraModes[0];
-			RootNode->ActivateCameraMode(CameraModeParams);
+			RootEvaluator->ActivateCameraMode(CameraModeParams);
 		}
 	}
 
 	// Run the root camera node.
-	FCameraNodeRunParams NodeParams;
+	FCameraNodeEvaluationParams NodeParams;
 	NodeParams.Evaluator = this;
 	NodeParams.DeltaTime = Params.DeltaTime;
 
 	RootNodeResult.Reset();
 
-	RootNode->Run(NodeParams, RootNodeResult);
+	RootEvaluator->Run(NodeParams, RootNodeResult);
 
 	Result.CameraPose = RootNodeResult.CameraPose;
 	Result.bIsCameraCut = RootNodeResult.bIsCameraCut;
@@ -100,8 +112,3 @@ void UCameraSystemEvaluator::GetEvaluatedCameraView(FMinimalViewInfo& DesiredVie
 	DesiredView.FOV = CameraPose.GetEffectiveFieldOfView();
 }
 
-FCameraRuntimeInstantiator& UCameraSystemEvaluator::GetRuntimeInstantiator()
-{
-	return Instantiator;
-}
-
diff --git a/Engine/Plugins/Cameras/GameplayCameras/Source/GameplayCameras/Private/Core/DefaultRootCameraNode.cpp b/Engine/Plugins/Cameras/GameplayCameras/Source/GameplayCameras/Private/Core/DefaultRootCameraNode.cpp
--- a/Engine/Plugins/Cameras/GameplayCameras/Source/GameplayCameras/Private/Core/DefaultRootCameraNode.cpp
+++ b/Engine/Plugins/Cameras/GameplayCameras/Source/GameplayCameras/Private/Core/DefaultRootCameraNode.cpp
@@ -34,38 +34,52 @@ UDefaultRootCameraNode::UDefaultRootCameraNode(const FObjectInitializer& ObjectI
 	VisualLayer = CreateBlendStack(this, ObjectInit, TEXT("VisualLayer"), false, true);
 }
 
-FCameraNodeChildrenView UDefaultRootCameraNode::OnGetChildren()
+FCameraNodeEvaluatorPtr UDefaultRootCameraNode::OnBuildEvaluator(FCameraNodeEvaluatorBuilder& Builder) const
 {
-	return FCameraNodeChildrenView({ BaseLayer, MainLayer, GlobalLayer, VisualLayer });
+	return Builder.BuildEvaluator<FDefaultRootCameraNodeEvaluator>();
 }
 
-void UDefaultRootCameraNode::OnRun(const FCameraNodeRunParams& Params, FCameraNodeRunResult& OutResult)
+UE_DEFINE_CAMERA_NODE_EVALUATOR(FDefaultRootCameraNodeEvaluator)
+
+void FDefaultRootCameraNodeEvaluator::OnInitialize(const FCameraNodeEvaluatorInitializeParams& Params)
 {
-	BaseLayer->Run(Params, OutResult);
+	const UDefaultRootCameraNode* Data = GetCameraNodeAs<UDefaultRootCameraNode>();
+	BaseLayer = Params.BuildEvaluatorAs<FBlendStackCameraNodeEvaluator>(Data->BaseLayer);
+	MainLayer = Params.BuildEvaluatorAs<FBlendStackCameraNodeEvaluator>(Data->MainLayer);
+	GlobalLayer = Params.BuildEvaluatorAs<FBlendStackCameraNodeEvaluator>(Data->GlobalLayer);
+	VisualLayer = Params.BuildEvaluatorAs<FBlendStackCameraNodeEvaluator>(Data->VisualLayer);
+}
 
-	MainLayer->Run(Params, OutResult);
+FCameraNodeEvaluatorChildrenView FDefaultRootCameraNodeEvaluator::OnGetChildren()
+{
+	return FCameraNodeEvaluatorChildrenView({ BaseLayer, MainLayer, GlobalLayer, VisualLayer });
+}
 
-	GlobalLayer->Run(Params, OutResult);
 
+void FDefaultRootCameraNodeEvaluator::OnRun(const FCameraNodeEvaluationParams& Params, FCameraNodeEvaluationResult& OutResult)
+{
+	BaseLayer->Run(Params, OutResult);
+	MainLayer->Run(Params, OutResult);
+	GlobalLayer->Run(Params, OutResult);
 	VisualLayer->Run(Params, OutResult);
 }
 
-void UDefaultRootCameraNode::OnActivateCameraMode(const FActivateCameraModeParams& Params)
+void FDefaultRootCameraNodeEvaluator::OnActivateCameraMode(const FActivateCameraModeParams& Params)
 {
-	UBlendStackCameraNode* TargetStack = nullptr;
+	FBlendStackCameraNodeEvaluator* TargetStack = nullptr;
 	switch (Params.Layer)
 	{
 		case ECameraModeLayer::Base:
-			TargetStack = CastChecked<UBlendStackCameraNode>(BaseLayer);
+			TargetStack = BaseLayer;
 			break;
 		case ECameraModeLayer::Main:
-			TargetStack = CastChecked<UBlendStackCameraNode>(MainLayer);
+			TargetStack = MainLayer;
 			break;
 		case ECameraModeLayer::Global:
-			TargetStack = CastChecked<UBlendStackCameraNode>(GlobalLayer);
+			TargetStack = GlobalLayer;
 			break;
 		case ECameraModeLayer::Visual:
-			TargetStack = CastChecked<UBlendStackCameraNode>(VisualLayer);
+			TargetStack = VisualLayer;
 			break;
 	}
 
diff --git a/Engine/Plugins/Cameras/GameplayCameras/Source/GameplayCameras/Private/Core/RootCameraNode.cpp b/Engine/Plugins/Cameras/GameplayCameras/Source/GameplayCameras/Private/Core/RootCameraNode.cpp
--- a/Engine/Plugins/Cameras/GameplayCameras/Source/GameplayCameras/Private/Core/RootCameraNode.cpp
+++ b/Engine/Plugins/Cameras/GameplayCameras/Source/GameplayCameras/Private/Core/RootCameraNode.cpp
@@ -6,7 +6,7 @@
 
 #include UE_INLINE_GENERATED_CPP_BY_NAME(RootCameraNode)
 
-void URootCameraNode::ActivateCameraMode(const FActivateCameraModeParams& Params)
+void FRootCameraNodeEvaluator::ActivateCameraMode(const FActivateCameraModeParams& Params)
 {
 	OnActivateCameraMode(Params);
 }
diff --git a/Engine/Plugins/Cameras/GameplayCameras/Source/GameplayCameras/Private/Directors/SingleCameraDirector.cpp b/Engine/Plugins/Cameras/GameplayCameras/Source/GameplayCameras/Private/Directors/SingleCameraDirector.cpp
--- a/Engine/Plugins/Cameras/GameplayCameras/Source/GameplayCameras/Private/Directors/SingleCameraDirector.cpp
+++ b/Engine/Plugins/Cameras/GameplayCameras/Source/GameplayCameras/Private/Directors/SingleCameraDirector.cpp
@@ -4,16 +4,29 @@
 
 #include UE_INLINE_GENERATED_CPP_BY_NAME(SingleCameraDirector)
 
+class FSingleCameraDirectorEvaluator : public FCameraDirectorEvaluator
+{
+	UE_DECLARE_CAMERA_DIRECTOR_EVALUATOR(FSingleCameraDirectorEvaluator)
+protected:
+	virtual void OnRun(const FCameraDirectorEvaluationParams& Params, FCameraDirectorEvaluationResult& OutResult) override
+	{
+		const USingleCameraDirector* SingleDirector = GetCameraDirectorAs<USingleCameraDirector>();
+		if (SingleDirector->CameraMode)
+		{
+			OutResult.ActiveCameraModes.Add(SingleDirector->CameraMode);
+		}
+	}
+};
+
+UE_DEFINE_CAMERA_DIRECTOR_EVALUATOR(FSingleCameraDirectorEvaluator)
+
 USingleCameraDirector::USingleCameraDirector(const FObjectInitializer& ObjectInit)
 	: Super(ObjectInit)
 {
 }
 
-void USingleCameraDirector::OnRun(const FCameraDirectorRunParams& Params, FCameraDirectorRunResult& OutResult)
+FCameraDirectorEvaluator* USingleCameraDirector::OnBuildEvaluator(FCameraDirectorEvaluatorBuilder& Builder) const
 {
-	if (CameraMode)
-	{
-		OutResult.ActiveCameraModes.Add(CameraMode);
-	}
+	return Builder.BuildEvaluator<FSingleCameraDirectorEvaluator>();
 }
 
diff --git a/Engine/Plugins/Cameras/GameplayCameras/Source/GameplayCameras/Private/Nodes/Blends/LinearBlendCameraNode.cpp b/Engine/Plugins/Cameras/GameplayCameras/Source/GameplayCameras/Private/Nodes/Blends/LinearBlendCameraNode.cpp
--- a/Engine/Plugins/Cameras/GameplayCameras/Source/GameplayCameras/Private/Nodes/Blends/LinearBlendCameraNode.cpp
+++ b/Engine/Plugins/Cameras/GameplayCameras/Source/GameplayCameras/Private/Nodes/Blends/LinearBlendCameraNode.cpp
@@ -4,8 +4,23 @@
 
 #include UE_INLINE_GENERATED_CPP_BY_NAME(LinearBlendCameraNode)
 
-void ULinearBlendCameraNode::OnComputeBlendFactor(const FCameraNodeRunParams& Params, FSimpleBlendCameraNodeRunResult& OutResult)
+class FLinearBlendCameraNodeEvaluator : public FSimpleFixedTimeBlendCameraNodeEvaluator
 {
+	UE_DECLARE_BLEND_CAMERA_NODE_EVALUATOR_EX(FLinearBlendCameraNodeEvaluator, FSimpleFixedTimeBlendCameraNodeEvaluator)
+protected:
+	virtual void OnComputeBlendFactor(const FCameraNodeEvaluationParams& Params, FSimpleBlendCameraNodeEvaluationResult& OutResult) override;
+};
+
+UE_DEFINE_BLEND_CAMERA_NODE_EVALUATOR(FLinearBlendCameraNodeEvaluator)
+
+void FLinearBlendCameraNodeEvaluator::OnComputeBlendFactor(const FCameraNodeEvaluationParams& Params, FSimpleBlendCameraNodeEvaluationResult& OutResult)
+{
+	const ULinearBlendCameraNode* BlendNode = GetCameraNodeAs<ULinearBlendCameraNode>();
 	OutResult.BlendFactor = FMath::Lerp(0.f, 1.f, GetTimeFactor());
 }
 
+FCameraNodeEvaluatorPtr ULinearBlendCameraNode::OnBuildEvaluator(FCameraNodeEvaluatorBuilder& Builder) const
+{
+	return Builder.BuildEvaluator<FLinearBlendCameraNodeEvaluator>();
+}
+
diff --git a/Engine/Plugins/Cameras/GameplayCameras/Source/GameplayCameras/Private/Nodes/Blends/PopBlendCameraNode.cpp b/Engine/Plugins/Cameras/GameplayCameras/Source/GameplayCameras/Private/Nodes/Blends/PopBlendCameraNode.cpp
--- a/Engine/Plugins/Cameras/GameplayCameras/Source/GameplayCameras/Private/Nodes/Blends/PopBlendCameraNode.cpp
+++ b/Engine/Plugins/Cameras/GameplayCameras/Source/GameplayCameras/Private/Nodes/Blends/PopBlendCameraNode.cpp
@@ -4,14 +4,24 @@
 
 #include UE_INLINE_GENERATED_CPP_BY_NAME(PopBlendCameraNode)
 
-void UPopBlendCameraNode::OnRun(const FCameraNodeRunParams& Params, FCameraNodeRunResult& OutResult)
+class FPopBlendCameraNodeEvaluator : public FBlendCameraNodeEvaluator
+{
+	UE_DECLARE_BLEND_CAMERA_NODE_EVALUATOR(FPopBlendCameraNodeEvaluator)
+protected:
+	virtual void OnRun(const FCameraNodeEvaluationParams& Params, FCameraNodeEvaluationResult& OutResult) override;
+	virtual void OnBlendResults(const FCameraNodeBlendParams& Params, FCameraNodeBlendResult& OutResult) override;
+};
+
+UE_DEFINE_BLEND_CAMERA_NODE_EVALUATOR(FPopBlendCameraNodeEvaluator)
+
+void FPopBlendCameraNodeEvaluator::OnRun(const FCameraNodeEvaluationParams& Params, FCameraNodeEvaluationResult& OutResult)
 {
 }
 
-void UPopBlendCameraNode::OnBlendResults(const FCameraNodeBlendParams& Params, FCameraNodeBlendResult& OutResult)
+void FPopBlendCameraNodeEvaluator::OnBlendResults(const FCameraNodeBlendParams& Params, FCameraNodeBlendResult& OutResult)
 {
-	const FCameraNodeRunResult& ChildResult(Params.ChildResult);
-	FCameraNodeRunResult& BlendedResult(OutResult.BlendedResult);
+	const FCameraNodeEvaluationResult& ChildResult(Params.ChildResult);
+	FCameraNodeEvaluationResult& BlendedResult(OutResult.BlendedResult);
 	
 	BlendedResult.CameraPose.OverrideChanged(ChildResult.CameraPose);
 	if (ChildResult.bIsCameraCut || Params.ChildParams.bIsFirstFrame)
@@ -23,3 +33,8 @@ void UPopBlendCameraNode::OnBlendResults(const FCameraNodeBlendParams& Params, F
 	OutResult.bIsBlendFinished = true;
 }
 
+FCameraNodeEvaluatorPtr UPopBlendCameraNode::OnBuildEvaluator(FCameraNodeEvaluatorBuilder& Builder) const
+{
+	return Builder.BuildEvaluator<FPopBlendCameraNodeEvaluator>();
+}
+
diff --git a/Engine/Plugins/Cameras/GameplayCameras/Source/GameplayCameras/Private/Nodes/Blends/SimpleBlendCameraNode.cpp b/Engine/Plugins/Cameras/GameplayCameras/Source/GameplayCameras/Private/Nodes/Blends/SimpleBlendCameraNode.cpp
--- a/Engine/Plugins/Cameras/GameplayCameras/Source/GameplayCameras/Private/Nodes/Blends/SimpleBlendCameraNode.cpp
+++ b/Engine/Plugins/Cameras/GameplayCameras/Source/GameplayCameras/Private/Nodes/Blends/SimpleBlendCameraNode.cpp
@@ -4,17 +4,19 @@
 
 #include UE_INLINE_GENERATED_CPP_BY_NAME(SimpleBlendCameraNode)
 
-void USimpleBlendCameraNode::OnRun(const FCameraNodeRunParams& Params, FCameraNodeRunResult& OutResult)
+UE_DEFINE_BLEND_CAMERA_NODE_EVALUATOR(FSimpleBlendCameraNodeEvaluator)
+
+void FSimpleBlendCameraNodeEvaluator::OnRun(const FCameraNodeEvaluationParams& Params, FCameraNodeEvaluationResult& OutResult)
 {
-	FSimpleBlendCameraNodeRunResult FactorResult;
+	FSimpleBlendCameraNodeEvaluationResult FactorResult;
 	OnComputeBlendFactor(Params, FactorResult);
 	BlendFactor = FactorResult.BlendFactor;
 }
 
-void USimpleBlendCameraNode::OnBlendResults(const FCameraNodeBlendParams& Params, FCameraNodeBlendResult& OutResult)
+void FSimpleBlendCameraNodeEvaluator::OnBlendResults(const FCameraNodeBlendParams& Params, FCameraNodeBlendResult& OutResult)
 {
-	const FCameraNodeRunResult& ChildResult(Params.ChildResult);
-	FCameraNodeRunResult& BlendedResult(OutResult.BlendedResult);
+	const FCameraNodeEvaluationResult& ChildResult(Params.ChildResult);
+	FCameraNodeEvaluationResult& BlendedResult(OutResult.BlendedResult);
 
 	BlendedResult.CameraPose.LerpChanged(
 			ChildResult.CameraPose, 
@@ -34,15 +36,24 @@ void USimpleBlendCameraNode::OnBlendResults(const FCameraNodeBlendParams& Params
 	OutResult.bIsBlendFinished = bIsBlendFinished;
 }
 
-void USimpleFixedTimeBlendCameraNode::OnRun(const FCameraNodeRunParams& Params, FCameraNodeRunResult& OutResult)
+UE_DEFINE_BLEND_CAMERA_NODE_EVALUATOR(FSimpleFixedTimeBlendCameraNodeEvaluator)
+
+void FSimpleFixedTimeBlendCameraNodeEvaluator::OnRun(const FCameraNodeEvaluationParams& Params, FCameraNodeEvaluationResult& OutResult)
 {
+	const USimpleFixedTimeBlendCameraNode* BlendNode = GetCameraNodeAs<USimpleFixedTimeBlendCameraNode>();
 	CurrentTime += Params.DeltaTime;
-	if (CurrentTime >= BlendTime)
+	if (CurrentTime >= BlendNode->BlendTime)
 	{
-		CurrentTime = BlendTime;
+		CurrentTime = BlendNode->BlendTime;
 		SetBlendFinished();
 	}
 
-	Super::OnRun(Params, OutResult);
+	FSimpleBlendCameraNodeEvaluator::OnRun(Params, OutResult);
+}
+
+float FSimpleFixedTimeBlendCameraNodeEvaluator::GetTimeFactor() const
+{
+	const USimpleFixedTimeBlendCameraNode* BlendNode = GetCameraNodeAs<USimpleFixedTimeBlendCameraNode>();
+	return CurrentTime / BlendNode->BlendTime;
 }
 
diff --git a/Engine/Plugins/Cameras/GameplayCameras/Source/GameplayCameras/Private/Nodes/Blends/SmoothBlendCameraNode.cpp b/Engine/Plugins/Cameras/GameplayCameras/Source/GameplayCameras/Private/Nodes/Blends/SmoothBlendCameraNode.cpp
--- a/Engine/Plugins/Cameras/GameplayCameras/Source/GameplayCameras/Private/Nodes/Blends/SmoothBlendCameraNode.cpp
+++ b/Engine/Plugins/Cameras/GameplayCameras/Source/GameplayCameras/Private/Nodes/Blends/SmoothBlendCameraNode.cpp
@@ -6,12 +6,22 @@
 
 #include UE_INLINE_GENERATED_CPP_BY_NAME(SmoothBlendCameraNode)
 
-void USmoothBlendCameraNode::OnComputeBlendFactor(const FCameraNodeRunParams& Params, FSimpleBlendCameraNodeRunResult& OutResult)
+class FSmoothBlendCameraNodeEvaluator : public FSimpleFixedTimeBlendCameraNodeEvaluator
+{
+	UE_DECLARE_BLEND_CAMERA_NODE_EVALUATOR_EX(FSmoothBlendCameraNodeEvaluator, FSimpleFixedTimeBlendCameraNodeEvaluator)
+protected:
+	virtual void OnComputeBlendFactor(const FCameraNodeEvaluationParams& Params, FSimpleBlendCameraNodeEvaluationResult& OutResult) override;
+};
+
+UE_DEFINE_BLEND_CAMERA_NODE_EVALUATOR(FSmoothBlendCameraNodeEvaluator)
+
+void FSmoothBlendCameraNodeEvaluator::OnComputeBlendFactor(const FCameraNodeEvaluationParams& Params, FSimpleBlendCameraNodeEvaluationResult& OutResult)
 {
 	using namespace UE::Cameras;
 
+	const USmoothBlendCameraNode* BlendNode = GetCameraNodeAs<USmoothBlendCameraNode>();
 	const float t = GetTimeFactor();
-	switch (BlendType)
+	switch (BlendNode->BlendType)
 	{
 		case ESmoothCameraBlendType::SmoothStep:
 			OutResult.BlendFactor = SmoothStep(t);
@@ -25,3 +35,8 @@ void USmoothBlendCameraNode::OnComputeBlendFactor(const FCameraNodeRunParams& Pa
 	}
 }
 
+FCameraNodeEvaluatorPtr USmoothBlendCameraNode::OnBuildEvaluator(FCameraNodeEvaluatorBuilder& Builder) const
+{
+	return Builder.BuildEvaluator<FSmoothBlendCameraNodeEvaluator>();
+}
+
diff --git a/Engine/Plugins/Cameras/GameplayCameras/Source/GameplayCameras/Private/Nodes/Common/ArrayCameraNode.cpp b/Engine/Plugins/Cameras/GameplayCameras/Source/GameplayCameras/Private/Nodes/Common/ArrayCameraNode.cpp
--- a/Engine/Plugins/Cameras/GameplayCameras/Source/GameplayCameras/Private/Nodes/Common/ArrayCameraNode.cpp
+++ b/Engine/Plugins/Cameras/GameplayCameras/Source/GameplayCameras/Private/Nodes/Common/ArrayCameraNode.cpp
@@ -2,16 +2,48 @@
 
 #include "Nodes/Common/ArrayCameraNode.h"
 
+#include "Core/CameraNodeEvaluator.h"
+
 #include UE_INLINE_GENERATED_CPP_BY_NAME(ArrayCameraNode)
 
-FCameraNodeChildrenView UArrayCameraNode::OnGetChildren()
+class FArrayCameraNodeEvaluator : public FCameraNodeEvaluator
+{
+	UE_DECLARE_CAMERA_NODE_EVALUATOR(FArrayCameraNodeEvaluator)
+
+protected:
+
+	virtual FCameraNodeEvaluatorChildrenView OnGetChildren() override;
+	virtual void OnInitialize(const FCameraNodeEvaluatorInitializeParams& Params) override;
+	virtual void OnRun(const FCameraNodeEvaluationParams& Params, FCameraNodeEvaluationResult& OutResult) override;
+
+private:
+
+	TArray<FCameraNodeEvaluator*> Children;
+};
+
+UE_DEFINE_CAMERA_NODE_EVALUATOR(FArrayCameraNodeEvaluator)
+
+FCameraNodeEvaluatorChildrenView FArrayCameraNodeEvaluator::OnGetChildren()
 {
-	return FCameraNodeChildrenView(Children);
+	return FCameraNodeEvaluatorChildrenView(Children);
 }
 
-void UArrayCameraNode::OnRun(const FCameraNodeRunParams& Params, FCameraNodeRunResult& OutResult)
+void FArrayCameraNodeEvaluator::OnInitialize(const FCameraNodeEvaluatorInitializeParams& Params)
 {
-	for (TObjectPtr<UCameraNode> Child : Children)
+	const UArrayCameraNode* ArrayNode = GetCameraNodeAs<UArrayCameraNode>();
+	for (const UCameraNode* Child : ArrayNode->Children)
+	{
+		if (Child)
+		{
+			FCameraNodeEvaluator* ChildEvaluator = Params.BuildEvaluator(Child);
+			Children.Add(ChildEvaluator);
+		}
+	}
+}
+
+void FArrayCameraNodeEvaluator::OnRun(const FCameraNodeEvaluationParams& Params, FCameraNodeEvaluationResult& OutResult)
+{
+	for (FCameraNodeEvaluator* Child : Children)
 	{
 		if (Child)
 		{
@@ -20,3 +52,8 @@ void UArrayCameraNode::OnRun(const FCameraNodeRunParams& Params, FCameraNodeRunR
 	}
 }
 
+FCameraNodeEvaluatorPtr UArrayCameraNode::OnBuildEvaluator(FCameraNodeEvaluatorBuilder& Builder) const
+{
+	return Builder.BuildEvaluator<FArrayCameraNodeEvaluator>();
+}
+
diff --git a/Engine/Plugins/Cameras/GameplayCameras/Source/GameplayCameras/Private/Nodes/Common/OffsetCameraNode.cpp b/Engine/Plugins/Cameras/GameplayCameras/Source/GameplayCameras/Private/Nodes/Common/OffsetCameraNode.cpp
--- a/Engine/Plugins/Cameras/GameplayCameras/Source/GameplayCameras/Private/Nodes/Common/OffsetCameraNode.cpp
+++ b/Engine/Plugins/Cameras/GameplayCameras/Source/GameplayCameras/Private/Nodes/Common/OffsetCameraNode.cpp
@@ -7,26 +7,38 @@
 
 #include UE_INLINE_GENERATED_CPP_BY_NAME(OffsetCameraNode)
 
-void UOffsetCameraNode::OnRun(const FCameraNodeRunParams& Params, FCameraNodeRunResult& OutResult)
+class FOffsetCameraNodeEvaluator : public FCameraNodeEvaluator
 {
-	FVector3d LocalOffset = Offset;
-	switch(OffsetSpace)
+	UE_DECLARE_CAMERA_NODE_EVALUATOR(FOffsetCameraNodeEvaluator)
+
+protected:
+
+	virtual void OnRun(const FCameraNodeEvaluationParams& Params, FCameraNodeEvaluationResult& OutResult) override;
+};
+
+UE_DEFINE_CAMERA_NODE_EVALUATOR(FOffsetCameraNodeEvaluator)
+
+void FOffsetCameraNodeEvaluator::OnRun(const FCameraNodeEvaluationParams& Params, FCameraNodeEvaluationResult& OutResult)
+{
+	const UOffsetCameraNode* OffsetNode = GetCameraNodeAs<UOffsetCameraNode>();
+	FVector3d LocalOffset = OffsetNode->Offset;
+	switch(OffsetNode->OffsetSpace)
 	{
 		case ECameraNodeSpace::CameraPose:
 		default:
 			{
 				const FRotator3d Rotation = OutResult.CameraPose.GetRotation();
-				LocalOffset = Rotation.RotateVector(Offset);
+				LocalOffset = Rotation.RotateVector(LocalOffset);
 			}
 			break;
 		case ECameraNodeSpace::Context:
 			if (Params.EvaluationContext)
 			{ 
-				const FCameraNodeRunResult& InitialResult = Params.EvaluationContext->GetInitialResult();
+				const FCameraNodeEvaluationResult& InitialResult = Params.EvaluationContext->GetInitialResult();
 				ensureMsgf(InitialResult.bIsValid,
 						TEXT("OffsetCameraNode: using invalid context result as offset space!"));
 				const FRotator3d Rotation = InitialResult.CameraPose.GetRotation();
-				LocalOffset = Rotation.RotateVector(Offset);
+				LocalOffset = Rotation.RotateVector(LocalOffset);
 			}
 			else
 			{
@@ -45,3 +57,8 @@ void UOffsetCameraNode::OnRun(const FCameraNodeRunParams& Params, FCameraNodeRun
 	OutResult.CameraPose.SetLocation(Location + LocalOffset);
 }
 
+FCameraNodeEvaluatorPtr UOffsetCameraNode::OnBuildEvaluator(FCameraNodeEvaluatorBuilder& Builder) const
+{
+	return Builder.BuildEvaluator<FOffsetCameraNodeEvaluator>();
+}
+
diff --git a/Engine/Plugins/Cameras/GameplayCameras/Source/GameplayCameras/Public/Core/BlendCameraNode.h b/Engine/Plugins/Cameras/GameplayCameras/Source/GameplayCameras/Public/Core/BlendCameraNode.h
--- a/Engine/Plugins/Cameras/GameplayCameras/Source/GameplayCameras/Public/Core/BlendCameraNode.h
+++ b/Engine/Plugins/Cameras/GameplayCameras/Source/GameplayCameras/Public/Core/BlendCameraNode.h
@@ -3,6 +3,7 @@
 #pragma once
 
 #include "Core/CameraNode.h"
+#include "Core/CameraNodeEvaluator.h"
 
 #include "BlendCameraNode.generated.h"
 
@@ -12,16 +13,16 @@
 struct FCameraNodeBlendParams
 {
 	FCameraNodeBlendParams(
-			const FCameraNodeRunParams& InChildParams,
-			const FCameraNodeRunResult& InChildResult)
+			const FCameraNodeEvaluationParams& InChildParams,
+			const FCameraNodeEvaluationResult& InChildResult)
 		: ChildParams(InChildParams)
 		, ChildResult(InChildResult)
 	{}
 
 	/** The parameters that the blend received during the evaluation. */
-	const FCameraNodeRunParams& ChildParams;
+	const FCameraNodeEvaluationParams& ChildParams;
 	/** The result that the blend should apply over another result. */
-	const FCameraNodeRunResult& ChildResult;
+	const FCameraNodeEvaluationResult& ChildResult;
 };
 
 /**
@@ -29,12 +30,12 @@ struct FCameraNodeBlendParams
  */
 struct FCameraNodeBlendResult
 {
-	FCameraNodeBlendResult(FCameraNodeRunResult& InBlendedResult)
+	FCameraNodeBlendResult(FCameraNodeEvaluationResult& InBlendedResult)
 		: BlendedResult(InBlendedResult)
 	{}
 
 	/** The result upon which another result should be blended. */
-	FCameraNodeRunResult& BlendedResult;
+	FCameraNodeEvaluationResult& BlendedResult;
 
 	/** Whether the blend has reached 100%. */
 	bool bIsBlendFull = false;
@@ -50,6 +51,14 @@ UCLASS(MinimalAPI)
 class UBlendCameraNode : public UCameraNode
 {
 	GENERATED_BODY()
+};
+
+/**
+ * Base evaluator class for blend camera nodes.
+ */
+class FBlendCameraNodeEvaluator : public FCameraNodeEvaluator
+{
+	UE_DECLARE_CAMERA_NODE_EVALUATOR(FBlendCameraNodeEvaluator)
 
 public:
 
@@ -62,3 +71,17 @@ protected:
 	virtual void OnBlendResults(const FCameraNodeBlendParams& Params, FCameraNodeBlendResult& OutResult) {}
 };
 
+// Macros for declaring and defining new blend node evaluators. They are the same
+// as the base ones for generic node evaluators, but the first one prevents you
+// from having to specify FBlendCameraNodeEvaluator as the base class, which saves
+// a little bit of typing.
+//
+#define UE_DECLARE_BLEND_CAMERA_NODE_EVALUATOR(ClassName)\
+	UE_DECLARE_CAMERA_NODE_EVALUATOR_EX(ClassName, FBlendCameraNodeEvaluator)
+
+#define UE_DECLARE_BLEND_CAMERA_NODE_EVALUATOR_EX(ClassName, BaseClassName)\
+	UE_DECLARE_CAMERA_NODE_EVALUATOR_EX(ClassName, BaseClassName)
+
+#define UE_DEFINE_BLEND_CAMERA_NODE_EVALUATOR(ClassName)\
+	UE_DEFINE_CAMERA_NODE_EVALUATOR(ClassName)
+
diff --git a/Engine/Plugins/Cameras/GameplayCameras/Source/GameplayCameras/Public/Core/BlendStackCameraNode.h b/Engine/Plugins/Cameras/GameplayCameras/Source/GameplayCameras/Public/Core/BlendStackCameraNode.h
--- a/Engine/Plugins/Cameras/GameplayCameras/Source/GameplayCameras/Public/Core/BlendStackCameraNode.h
+++ b/Engine/Plugins/Cameras/GameplayCameras/Source/GameplayCameras/Public/Core/BlendStackCameraNode.h
@@ -3,9 +3,12 @@
 #pragma once
 
 #include "Core/CameraNode.h"
+#include "Core/CameraNodeEvaluator.h"
+#include "Core/CameraNodeEvaluatorStorage.h"
 
 #include "BlendStackCameraNode.generated.h"
 
+class FBlendStackRootCameraNodeEvaluator;
 class UBlendStackRootCameraNode;
 class UCameraAsset;
 class UCameraEvaluationContext;
@@ -35,26 +38,10 @@ class UBlendStackCameraNode : public UCameraNode
 {
 	GENERATED_BODY()
 
-public:
-
-	/** Push a new camera mode onto the blend stack. */
-	void Push(const FBlendStackCameraPushParams& Params);
-
-public:
-
-	static void AddReferencedObjects(UObject* InThis, FReferenceCollector& Collector);
-
 protected:
 
 	// UCameraNode interface
-	virtual FCameraNodeChildrenView OnGetChildren() override;
-	virtual void OnRun(const FCameraNodeRunParams& Params, FCameraNodeRunResult& OutResult) override;
-
-	const FCameraModeTransition* FindTransition(const FBlendStackCameraPushParams& Params) const;
-	const FCameraModeTransition* FindTransition(
-			TArrayView<const FCameraModeTransition> Transitions, 
-			const UCameraMode* FromCameraMode, const UCameraAsset* FromCameraAsset, bool bFromFrozen,
-			const UCameraMode* ToCameraMode, const UCameraAsset* ToCameraAsset) const;
+	virtual FCameraNodeEvaluatorPtr OnBuildEvaluator(FCameraNodeEvaluatorBuilder& Builder) const override;
 
 public:
 
@@ -70,19 +57,54 @@ public:
 	 */
 	UPROPERTY()
 	bool bBlendFirstCameraMode = false;
+};
+
+/**
+ * Evaluator for a blend stack camera node.
+ */
+class FBlendStackCameraNodeEvaluator : public TCameraNodeEvaluator<UBlendStackCameraNode>
+{
+	UE_DECLARE_CAMERA_NODE_EVALUATOR(FBlendStackCameraNodeEvaluator)
+
+public:
+
+	/** Push a new camera mode onto the blend stack. */
+	void Push(const FBlendStackCameraPushParams& Params);
+
+protected:
+
+	// FCameraNodeEvaluator interface
+	virtual FCameraNodeEvaluatorChildrenView OnGetChildren() override;
+	virtual void OnRun(const FCameraNodeEvaluationParams& Params, FCameraNodeEvaluationResult& OutResult) override;
+
+	// Utility functions for finding an appropriate transition.
+	const FCameraModeTransition* FindTransition(const FBlendStackCameraPushParams& Params) const;
+	const FCameraModeTransition* FindTransition(
+			TArrayView<const FCameraModeTransition> Transitions, 
+			const UCameraMode* FromCameraMode, const UCameraAsset* FromCameraAsset, bool bFromFrozen,
+			const UCameraMode* ToCameraMode, const UCameraAsset* ToCameraAsset) const;
 
 protected:
 
 	struct FCameraModeEntry
 	{
+		/** Evaluation context in which this entry runs. */
 		TWeakObjectPtr<const UCameraEvaluationContext> EvaluationContext;
-		TObjectPtr<const UCameraMode> OriginalCameraMode;
-		TObjectPtr<UBlendStackRootCameraNode> RootNode;
-		FCameraNodeRunResult Result;
+		/** The camera mode asset that this entry runs. */
+		TObjectPtr<const UCameraMode> CameraMode;
+		/** Storage buffer for all evaluators in this node tree. */
+		FCameraNodeEvaluatorStorage EvaluatorStorage;
+		/** Root evaluator. */
+		FBlendStackRootCameraNodeEvaluator* RootEvaluator = nullptr;
+		/** Result for this node tree. */
+		FCameraNodeEvaluationResult Result;
+		/** Whether this is the first frame this entry runs. */
 		bool bIsFirstFrame = false;
+		/** Whether this entry is frozen. */
 		bool bIsFrozen = false;
 	};
 
+	/** Entries in the blend stack. */
 	TArray<FCameraModeEntry> Entries;
 };
 
diff --git a/Engine/Plugins/Cameras/GameplayCameras/Source/GameplayCameras/Public/Core/CameraDirector.h b/Engine/Plugins/Cameras/GameplayCameras/Source/GameplayCameras/Public/Core/CameraDirector.h
--- a/Engine/Plugins/Cameras/GameplayCameras/Source/GameplayCameras/Public/Core/CameraDirector.h
+++ b/Engine/Plugins/Cameras/GameplayCameras/Source/GameplayCameras/Public/Core/CameraDirector.h
@@ -4,33 +4,10 @@
 
 #include "CoreTypes.h"
 #include "UObject/Object.h"
+#include "Core/CameraDirectorEvaluator.h"
 
 #include "CameraDirector.generated.h"
 
-class UCameraEvaluationContext;
-class UCameraMode;
-
-/**
- * Parameter structure for running a camera director.
- */
-struct FCameraDirectorRunParams
-{
-	/** Time interval for the update. */
-	float DeltaTime = 0.f;
-
-	/** The context in which this director runs. */
-	TObjectPtr<const UCameraEvaluationContext> OwnerContext;
-};
-
-/**
- * Result struct for running a camera director.
- */
-struct FCameraDirectorRunResult
-{
-	/** The camera mode(s) that the director says should be active this frame. */
-	TArray<TObjectPtr<const UCameraMode>, TInlineAllocator<2>> ActiveCameraModes;
-};
-
 /**
  * Base class for a camera director.
  */
@@ -40,13 +17,13 @@ class UCameraDirector : public UObject
 	GENERATED_BODY()
 
 public:
-	
-	/** Runs the camera director to determine what camera mode(s) should be active this frame. */
-	void Run(const FCameraDirectorRunParams& Params, FCameraDirectorRunResult& OutResult);
+
+	/** Build the evaluator for this director. */
+	FCameraDirectorEvaluator* BuildEvaluator(FCameraDirectorEvaluatorBuilder& Builder) const;
 
 protected:
 
-	/** Runs the camera director to determine what camera mode(s) should be active this frame. */
-	virtual void OnRun(const FCameraDirectorRunParams& Params, FCameraDirectorRunResult& OutResult) {}
+	/** Build the evaluator for this director. */
+	virtual FCameraDirectorEvaluator* OnBuildEvaluator(FCameraDirectorEvaluatorBuilder& Builder) const { return nullptr; }
 };
 
diff --git a/Engine/Plugins/Cameras/GameplayCameras/Source/GameplayCameras/Public/Core/CameraDirectorEvaluator.h b/Engine/Plugins/Cameras/GameplayCameras/Source/GameplayCameras/Public/Core/CameraDirectorEvaluator.h
new file mode 100644
--- /dev/null
+++ b/Engine/Plugins/Cameras/GameplayCameras/Source/GameplayCameras/Public/Core/CameraDirectorEvaluator.h
@@ -0,0 +1,138 @@
+// Copyright Epic Games, Inc. All Rights Reserved.
+
+#pragma once
+
+#include "Core/CameraObjectRtti.h"
+#include "CoreTypes.h"
+#include "Templates/UniquePtr.h"
+#include "UObject/ObjectPtr.h"
+
+class FCameraDirectorEvaluator;
+class FCameraDirectorEvaluatorStorage;
+class UCameraDirector;
+class UCameraEvaluationContext;
+class UCameraMode;
+
+/**
+ * Parameter structure for running a camera director.
+ */
+struct FCameraDirectorEvaluationParams
+{
+	/** Time interval for the update. */
+	float DeltaTime = 0.f;
+
+	/** The context in which this director runs. */
+	TObjectPtr<const UCameraEvaluationContext> OwnerContext;
+};
+
+/**
+ * Result structure for running a camera director.
+ */
+struct FCameraDirectorEvaluationResult
+{
+	/** The camera mode(s) that the director says should be active this frame. */
+	TArray<TObjectPtr<const UCameraMode>, TInlineAllocator<2>> ActiveCameraModes;
+};
+
+/**
+ * Structure for building director evaluators.
+ */
+struct FCameraDirectorEvaluatorBuilder
+{
+	FCameraDirectorEvaluatorBuilder(FCameraDirectorEvaluatorStorage& InStorage)
+		: Storage(InStorage)
+	{}
+
+	/** Builds a director evaluator of the given type. */
+	template<typename EvaluatorType, typename ...ArgTypes>
+	EvaluatorType* BuildEvaluator(ArgTypes&&... InArgs);
+
+private:
+
+	FCameraDirectorEvaluatorStorage& Storage;
+};
+
+/**
+ * Storage for a director evaluator.
+ */
+class FCameraDirectorEvaluatorStorage
+{
+public:
+
+	/** Gets the stored evaluator, if any. */
+	FCameraDirectorEvaluator* GetEvaluator() const { return Evaluator.Get(); }
+
+private:
+
+	template<typename EvaluatorType, typename ...ArgTypes>
+	EvaluatorType* BuildEvaluator(ArgTypes&&... InArgs);
+
+	TSharedPtr<FCameraDirectorEvaluator> Evaluator;
+
+	friend struct FCameraDirectorEvaluatorBuilder;
+};
+
+/**
+ * Base class for camera director evaluators.
+ */
+class FCameraDirectorEvaluator
+{
+	UE_GAMEPLAY_CAMERAS_DECLARE_RTTI_BASE(FCameraDirectorEvaluator)
+
+public:
+
+	FCameraDirectorEvaluator();
+	virtual ~FCameraDirectorEvaluator() {}
+	
+	/** Runs the camera director to determine what camera mode(s) should be active this frame. */
+	void Run(const FCameraDirectorEvaluationParams& Params, FCameraDirectorEvaluationResult& OutResult);
+
+	/** Gets the camera director. */
+	template<typename CameraDirectorType>
+	const CameraDirectorType* GetCameraDirectorAs() const
+	{
+		return Cast<CameraDirectorType>(PrivateCameraDirector);
+	}
+
+public:
+
+	// Internal API.
+	void SetPrivateCameraDirector(TObjectPtr<const UCameraDirector> InCameraDirector);
+
+protected:
+
+	/** Runs the camera director to determine what camera mode(s) should be active this frame. */
+	virtual void OnRun(const FCameraDirectorEvaluationParams& Params, FCameraDirectorEvaluationResult& OutResult) {}
+
+private:
+
+	/** The camera director this evaluator is running. */
+	TObjectPtr<const UCameraDirector> PrivateCameraDirector;
+};
+
+// Utility macros for declaring and defining camera director evaluators.
+//
+#define UE_DECLARE_CAMERA_DIRECTOR_EVALUATOR(ClassName)\
+	UE_GAMEPLAY_CAMERAS_DECLARE_RTTI(ClassName, FCameraDirectorEvaluator)
+
+#define UE_DECLARE_CAMERA_DIRECTOR_EVALUATOR_EX(ClassName, BaseClassName)\
+	UE_GAMEPLAY_CAMERAS_DECLARE_RTTI(ClassName, BaseClassName)
+
+#define UE_DEFINE_CAMERA_DIRECTOR_EVALUATOR(ClassName)\
+	UE_GAMEPLAY_CAMERAS_DEFINE_RTTI(ClassName)
+
+template<typename EvaluatorType, typename ...ArgTypes>
+EvaluatorType* FCameraDirectorEvaluatorBuilder::BuildEvaluator(ArgTypes&&... InArgs)
+{
+	return Storage.BuildEvaluator<EvaluatorType>(Forward<ArgTypes>(InArgs)...);
+}
+
+template<typename EvaluatorType, typename ...ArgTypes>
+EvaluatorType* FCameraDirectorEvaluatorStorage::BuildEvaluator(ArgTypes&&... InArgs)
+{
+	// We should only build one evaluator.
+	ensure(Evaluator == nullptr);
+	Evaluator = MakeShared<EvaluatorType>(Forward<ArgTypes>(InArgs)...);
+	return Evaluator->CastThisChecked<EvaluatorType>();
+}
+
diff --git a/Engine/Plugins/Cameras/GameplayCameras/Source/GameplayCameras/Public/Core/CameraEvaluationContext.h b/Engine/Plugins/Cameras/GameplayCameras/Source/GameplayCameras/Public/Core/CameraEvaluationContext.h
--- a/Engine/Plugins/Cameras/GameplayCameras/Source/GameplayCameras/Public/Core/CameraEvaluationContext.h
+++ b/Engine/Plugins/Cameras/GameplayCameras/Source/GameplayCameras/Public/Core/CameraEvaluationContext.h
@@ -4,7 +4,7 @@
 
 #include "CoreTypes.h"
 #include "UObject/ObjectPtr.h"
-#include "Core/CameraNode.h"
+#include "Core/CameraNodeEvaluator.h"
 
 #include "CameraEvaluationContext.generated.h"
 
@@ -26,7 +26,7 @@ public:
 	UCameraAsset* GetCameraAsset() const { return CameraAsset; }
 
 	/** Gets the initial evaluation result for all camera modes in this context. */
-	const FCameraNodeRunResult& GetInitialResult() const { return InitialResult; }
+	const FCameraNodeEvaluationResult& GetInitialResult() const { return InitialResult; }
 
 protected:
 
@@ -35,6 +35,6 @@ protected:
 	TObjectPtr<UCameraAsset> CameraAsset;
 
 	/** The initial result for all camera modes in this context. */
-	FCameraNodeRunResult InitialResult;
+	FCameraNodeEvaluationResult InitialResult;
 };
 
diff --git a/Engine/Plugins/Cameras/GameplayCameras/Source/GameplayCameras/Public/Core/CameraEvaluationContextStack.h b/Engine/Plugins/Cameras/GameplayCameras/Source/GameplayCameras/Public/Core/CameraEvaluationContextStack.h
--- a/Engine/Plugins/Cameras/GameplayCameras/Source/GameplayCameras/Public/Core/CameraEvaluationContextStack.h
+++ b/Engine/Plugins/Cameras/GameplayCameras/Source/GameplayCameras/Public/Core/CameraEvaluationContextStack.h
@@ -3,8 +3,10 @@
 #pragma once
 
 #include "Containers/Array.h"
+#include "Core/CameraDirectorEvaluator.h"
 #include "UObject/WeakObjectPtr.h"
 
+class FCameraDirectorEvaluator;
 class UCameraAsset;
 class UCameraDirector;
 class UCameraEvaluationContext;
@@ -17,7 +19,9 @@ struct FCameraEvaluationContextInfo
 	UCameraEvaluationContext* EvaluationContext = nullptr;
 
 	/** The instantiated camera director running in this context. */
-	UCameraDirector* CameraDirector = nullptr;
+	const UCameraDirector* CameraDirector = nullptr;
+
+	FCameraDirectorEvaluator* Evaluator = nullptr;
 
 	/** Returns whether this structure has a valid context and director. */
 	bool IsValid() const { return EvaluationContext && CameraDirector; }
@@ -56,7 +60,9 @@ private:
 	struct FContextEntry
 	{
 		TWeakObjectPtr<UCameraEvaluationContext> WeakContext;
-		TObjectPtr<UCameraDirector> CameraDirector;
+		TObjectPtr<const UCameraDirector> CameraDirector;
+		FCameraDirectorEvaluatorStorage EvaluatorStorage;
+		FCameraDirectorEvaluator* Evaluator = nullptr;
 	};
 
 	/** The entries in the stack. */
diff --git a/Engine/Plugins/Cameras/GameplayCameras/Source/GameplayCameras/Public/Core/CameraInstantiableObject.h b/Engine/Plugins/Cameras/GameplayCameras/Source/GameplayCameras/Public/Core/CameraInstantiableObject.h
deleted file mode 100644
diff --git a/Engine/Plugins/Cameras/GameplayCameras/Source/GameplayCameras/Public/Core/CameraNode.h b/Engine/Plugins/Cameras/GameplayCameras/Source/GameplayCameras/Public/Core/CameraNode.h
--- a/Engine/Plugins/Cameras/GameplayCameras/Source/GameplayCameras/Public/Core/CameraNode.h
+++ b/Engine/Plugins/Cameras/GameplayCameras/Source/GameplayCameras/Public/Core/CameraNode.h
@@ -2,111 +2,38 @@
 
 #pragma once
 
-#include "Core/CameraInstantiableObject.h"
-#include "Core/CameraPose.h"
-#include "Core/CameraNodeChildrenView.h"
+#include "Core/CameraNodeEvaluatorBuilder.h"
 #include "CoreTypes.h"
 #include "UObject/Object.h"
 
 #include "CameraNode.generated.h"
 
-class UCameraEvaluationContext;
-class UCameraSystemEvaluator;
-
-/**
- * Parameter structure for running a camera node.
- */
-struct FCameraNodeRunParams
-{
-	/** The evaluation running this evaluation. */
-	TObjectPtr<UCameraSystemEvaluator> Evaluator;
-	/** The evaluation context (if any) responsible for this branch of the evaluation. */
-	TObjectPtr<const UCameraEvaluationContext> EvaluationContext;
-	/** The time interval for the evaluation. */
-	float DeltaTime = 0.f;
-	/** Whether this is the first evaluation of this camera node hierarchy. */
-	bool bIsFirstFrame = false;
-};
-
-/**
- * Input/output result structure for running a camera node.
- */
-struct FCameraNodeRunResult
-{
-	/** The camera pose. */
-	FCameraPose CameraPose;
-	/** Whether the current frame is a camera cut. */
-	bool bIsCameraCut = false;
-	/** Whether this result is valid. */
-	bool bIsValid = false;
-
-	/** Reset this result to its default (non-valid) state. */
-	void Reset();
-};
-
-UENUM()
-enum class ECameraNodeFlags
-{
-	None = 0,
-	RequiresReset = 1
-};
-ENUM_CLASS_FLAGS(ECameraNodeFlags);
-
-struct FCameraNodeResetParams
-{
-};
-
 /**
  * The base class for a camera node.
  */
 UCLASS(Abstract, DefaultToInstanced, EditInlineNew, MinimalAPI)
-class UCameraNode : public UCameraInstantiableObject
+class UCameraNode : public UObject
 {
 	GENERATED_BODY()
 
 public:
 
-	/** Get the list of children under this node. */
-	FCameraNodeChildrenView GetChildren();
-
-	/** Get the flags for this node. */
-	ECameraNodeFlags GetNodeFlags() const { return Flags; }
-
-	/** Resets this node. */
-	void Reset(const FCameraNodeResetParams& Params);
-
-	/** Run this node. */
-	void Run(const FCameraNodeRunParams& Params, FCameraNodeRunResult& OutResult);
+	/** Gets optional info about this node's evaluator. */
+	FCameraNodeEvaluatorAllocationInfo GetEvaluatorInfo() const;
+	/** Builds the evaluator for this node. */
+	FCameraNodeEvaluatorPtr BuildEvaluator(FCameraNodeEvaluatorBuilder& Builder) const;
 
 protected:
 
-	/** Sets the flags for this node. Should only be called once during construction. */
-	void SetNodeFlags(ECameraNodeFlags InFlags) { Flags = InFlags; }
-
-	/** Get the list of children under this node. */
-	virtual FCameraNodeChildrenView OnGetChildren() { return FCameraNodeChildrenView(); }
-
-	/** Resets this node. */
-	virtual void OnReset(const FCameraNodeResetParams& Params) {}
-
-	/** Run this node. */
-	virtual void OnRun(const FCameraNodeRunParams& Params, FCameraNodeRunResult& OutResult) {}
-
-#if WITH_EDITOR
-
-	virtual void PostEditChangeProperty(struct FPropertyChangedEvent& PropertyChangedEvent) override;
-
-#endif
+	/** Gets optional info about this node's evaluator. */
+	virtual FCameraNodeEvaluatorAllocationInfo OnGetEvaluatorInfo() const { return FCameraNodeEvaluatorAllocationInfo(); }
+	/** Builds the evaluator for this node. */
+	virtual FCameraNodeEvaluatorPtr OnBuildEvaluator(FCameraNodeEvaluatorBuilder& Builder) const { return nullptr; }
 
 public:
 
 	/** Specifies whether this node is enabled. */
 	UPROPERTY(EditAnywhere, Category=Common)
 	bool bIsEnabled = true;
-
-private:
-
-	/** The flags for this node. Should only be set once during construction. */
-	ECameraNodeFlags Flags = ECameraNodeFlags::None;
 };
 
diff --git a/Engine/Plugins/Cameras/GameplayCameras/Source/GameplayCameras/Public/Core/CameraNodeChildrenView.h b/Engine/Plugins/Cameras/GameplayCameras/Source/GameplayCameras/Public/Core/CameraNodeChildrenView.h
deleted file mode 100644
diff --git a/Engine/Plugins/Cameras/GameplayCameras/Source/GameplayCameras/Public/Core/CameraNodeEvaluator.h b/Engine/Plugins/Cameras/GameplayCameras/Source/GameplayCameras/Public/Core/CameraNodeEvaluator.h
new file mode 100644
--- /dev/null
+++ b/Engine/Plugins/Cameras/GameplayCameras/Source/GameplayCameras/Public/Core/CameraNodeEvaluator.h
@@ -0,0 +1,159 @@
+// Copyright Epic Games, Inc. All Rights Reserved.
+
+#pragma once
+
+#include "Core/CameraObjectRtti.h"
+#include "Core/CameraPose.h"
+#include "Core/ObjectChildrenView.h"
+#include "CoreTypes.h"
+#include "UObject/ObjectPtr.h"
+
+class FReferenceCollector;
+class FCameraNodeEvaluator;
+class UCameraEvaluationContext;
+class UCameraNode;
+class UCameraSystemEvaluator;
+struct FCameraNodeEvaluatorBuilder;
+
+/**
+ * Parameter structure for running a camera node evaluator.
+ */
+struct FCameraNodeEvaluationParams
+{
+	/** The evaluation running this evaluation. */
+	TObjectPtr<UCameraSystemEvaluator> Evaluator;
+	/** The evaluation context (if any) responsible for this branch of the evaluation. */
+	TObjectPtr<const UCameraEvaluationContext> EvaluationContext;
+	/** The time interval for the evaluation. */
+	float DeltaTime = 0.f;
+	/** Whether this is the first evaluation of this camera node hierarchy. */
+	bool bIsFirstFrame = false;
+};
+
+/**
+ * Input/output result structure for running a camera node evaluator.
+ */
+struct FCameraNodeEvaluationResult
+{
+	/** The camera pose. */
+	FCameraPose CameraPose;
+	/** Whether the current frame is a camera cut. */
+	bool bIsCameraCut = false;
+	/** Whether this result is valid. */
+	bool bIsValid = false;
+
+	/** Reset this result to its default (non-valid) state. */
+	void Reset();
+};
+
+/**
+ * Structure for initializing a camera node evaluator.
+ */
+struct FCameraNodeEvaluatorInitializeParams
+{
+	/** The evaluation running this evaluation. */
+	TObjectPtr<UCameraSystemEvaluator> Evaluator;
+	/** The evaluation context (if any) responsible for this branch of the evaluation. */
+	TObjectPtr<const UCameraEvaluationContext> EvaluationContext;
+	/** Builder object for building children evaluators. */
+	FCameraNodeEvaluatorBuilder* Builder = nullptr;
+
+	/** Builds an evaluator for the given camera node. */
+	FCameraNodeEvaluator* BuildEvaluator(const UCameraNode* InNode) const;
+
+	/** Builds an evaluator for the given camera node, and down-cast it to the given type. */
+	template<typename EvaluatorType>
+	EvaluatorType* BuildEvaluatorAs(const UCameraNode* InNode) const;
+};
+
+/** View on a camera node evaluator's children. */
+using FCameraNodeEvaluatorChildrenView = TObjectChildrenView<FCameraNodeEvaluator*>;
+
+/**
+ * Base class for objects responsible for running a camera node.
+ */
+class FCameraNodeEvaluator
+{
+	UE_GAMEPLAY_CAMERAS_DECLARE_RTTI_BASE(FCameraNodeEvaluator)
+
+public:
+
+	FCameraNodeEvaluator();
+	virtual ~FCameraNodeEvaluator() {}
+
+	/** Initialize this evaluator. */
+	void Initialize(const FCameraNodeEvaluatorInitializeParams& Params);
+
+	/** Get the list of children under this evaluator. */
+	FCameraNodeEvaluatorChildrenView GetChildren();
+
+	/** Run this evaluator. */
+	void Run(const FCameraNodeEvaluationParams& Params, FCameraNodeEvaluationResult& OutResult);
+
+	/** Collect referenced UObjects. */
+	void AddReferencedObjects(FReferenceCollector& Collector);
+
+	/** Get the camera node. */
+	template<typename CameraNodeType>
+	const CameraNodeType* GetCameraNodeAs() const
+	{
+		return Cast<CameraNodeType>(PrivateCameraNode);
+	}
+
+	// Internal API.
+	void SetPrivateCameraNode(TObjectPtr<const UCameraNode> InCameraNode);
+
+protected:
+
+	/** Initialize this evaluator. */
+	virtual void OnInitialize(const FCameraNodeEvaluatorInitializeParams& Params) {}
+
+	/** Get the list of children under this evaluator. */
+	virtual FCameraNodeEvaluatorChildrenView OnGetChildren() { return FCameraNodeEvaluatorChildrenView(); }
+
+	/** Run this evaluator. */
+	virtual void OnRun(const FCameraNodeEvaluationParams& Params, FCameraNodeEvaluationResult& OutResult) {}
+
+	/** Collect referenced UObjects. */
+	virtual void OnAddReferencedObjects(FReferenceCollector& Collector) {}
+
+private:
+
+	/** The camera node to run. */
+	TObjectPtr<const UCameraNode> PrivateCameraNode;
+};
+
+// Utility macros for declaring and defining camera node evaluators.
+//
+#define UE_DECLARE_CAMERA_NODE_EVALUATOR(ClassName)\
+	UE_GAMEPLAY_CAMERAS_DECLARE_RTTI(ClassName, FCameraNodeEvaluator)
+
+#define UE_DECLARE_CAMERA_NODE_EVALUATOR_EX(ClassName, BaseClassName)\
+	UE_GAMEPLAY_CAMERAS_DECLARE_RTTI(ClassName, BaseClassName)
+
+#define UE_DEFINE_CAMERA_NODE_EVALUATOR(ClassName)\
+	UE_GAMEPLAY_CAMERAS_DEFINE_RTTI(ClassName)
+
+/** Utility base class for camera node evaluators of a specific camera node type. */
+template<typename CameraNodeType>
+class TCameraNodeEvaluator : public FCameraNodeEvaluator
+{
+public:
+
+	/** Gets the camera node. */
+	const CameraNodeType* GetCameraNode() const
+	{
+		return GetCameraNodeAs<CameraNodeType>();
+	}
+
+	friend CameraNodeType;
+};
+
+template<typename EvaluatorType>
+EvaluatorType* FCameraNodeEvaluatorInitializeParams::BuildEvaluatorAs(const UCameraNode* InNode) const
+{
+	FCameraNodeEvaluator* NewEvaluator = BuildEvaluator(InNode);
+	check(NewEvaluator);
+	return NewEvaluator->CastThisChecked<EvaluatorType>();
+}
+
diff --git a/Engine/Plugins/Cameras/GameplayCameras/Source/GameplayCameras/Public/Core/CameraNodeEvaluatorBuilder.h b/Engine/Plugins/Cameras/GameplayCameras/Source/GameplayCameras/Public/Core/CameraNodeEvaluatorBuilder.h
new file mode 100644
--- /dev/null
+++ b/Engine/Plugins/Cameras/GameplayCameras/Source/GameplayCameras/Public/Core/CameraNodeEvaluatorBuilder.h
@@ -0,0 +1,30 @@
+// Copyright Epic Games, Inc. All Rights Reserved.
+
+#pragma once
+
+#include "Core/CameraNodeEvaluatorStorage.h"
+
+/** Structure for building camera node evaluators. */
+struct FCameraNodeEvaluatorBuilder
+{
+	FCameraNodeEvaluatorBuilder(FCameraNodeEvaluatorStorage& InStorage)
+		: Storage(InStorage)
+	{}
+
+	/** Builds a camera node evaluator of a given type. */
+	template<typename EvaluatorType, typename ...ArgTypes>
+	EvaluatorType* BuildEvaluator(ArgTypes&&... InArgs);
+
+private:
+
+	FCameraNodeEvaluatorStorage& Storage;
+
+	friend class FCameraNodeTreeState;
+};
+
+template<typename EvaluatorType, typename ...ArgTypes>
+EvaluatorType* FCameraNodeEvaluatorBuilder::BuildEvaluator(ArgTypes&&... InArgs)
+{
+	return Storage.BuildEvaluator<EvaluatorType>(Forward<ArgTypes>(InArgs)...);
+}
+
diff --git a/Engine/Plugins/Cameras/GameplayCameras/Source/GameplayCameras/Public/Core/CameraNodeEvaluatorFwd.h b/Engine/Plugins/Cameras/GameplayCameras/Source/GameplayCameras/Public/Core/CameraNodeEvaluatorFwd.h
new file mode 100644
--- /dev/null
+++ b/Engine/Plugins/Cameras/GameplayCameras/Source/GameplayCameras/Public/Core/CameraNodeEvaluatorFwd.h
@@ -0,0 +1,30 @@
+// Copyright Epic Games, Inc. All Rights Reserved.
+
+#pragma once
+
+#include "CoreTypes.h"
+
+class FCameraNodeEvaluator;
+class FCameraNodeEvaluatorStorage;
+struct FCameraNodeEvaluatorBuilder;
+
+using FCameraNodeEvaluatorPtr = FCameraNodeEvaluator*;
+
+/** Allocation information for a node evaluator. */
+struct FCameraNodeEvaluatorAllocationInfo
+{
+	int16 Sizeof = 0;
+	int16 Alignof = 0;
+};
+
+/** Allocation information for a node evaluator, auto-setup for a given type. */
+template<typename EvaluatorType>
+struct TCameraNodeEvaluatorAllocationInfo : FCameraNodeEvaluatorAllocationInfo
+{
+	TCameraNodeEvaluatorAllocationInfo() 
+	{
+		Sizeof = sizeof(EvaluatorType);
+		Alignof = alignof(EvaluatorType);
+	}
+};
+
diff --git a/Engine/Plugins/Cameras/GameplayCameras/Source/GameplayCameras/Public/Core/CameraNodeEvaluatorStorage.h b/Engine/Plugins/Cameras/GameplayCameras/Source/GameplayCameras/Public/Core/CameraNodeEvaluatorStorage.h
new file mode 100644
--- /dev/null
+++ b/Engine/Plugins/Cameras/GameplayCameras/Source/GameplayCameras/Public/Core/CameraNodeEvaluatorStorage.h
@@ -0,0 +1,131 @@
+// Copyright Epic Games, Inc. All Rights Reserved.
+
+#pragma once
+
+#include "Core/CameraNodeEvaluatorFwd.h"
+#include "CoreTypes.h"
+#include "HAL/UnrealMemory.h"
+#include "Math/UnrealMathUtility.h"
+#include "UObject/ObjectPtr.h"
+
+class FCameraNodeEvaluator;
+class FCameraNodeEvaluatorStorage;
+class FReferenceCollector;
+class UCameraEvaluationContext;
+class UCameraMode;
+class UCameraNode;
+class UCameraSystemEvaluator;
+
+/** Allocation information for an entire tree of node evaluators. */
+struct FCameraNodeEvaluatorTreeAllocationInfo
+{
+	int16 TotalSizeof = 0;
+	int16 MaxAlignof = 0;
+};
+
+/** Structure for building an entire tree of node evaluators. */
+struct FCameraNodeEvaluatorTreeBuilderParams
+{
+	/** The root node of the tree. */
+	TObjectPtr<const UCameraNode> RootCameraNode;
+	/** The evaluator running this evaluation. */
+	TObjectPtr<UCameraSystemEvaluator> Evaluator;
+	/** The evaluation context (if any) responsible for this branch of the evaluation. */
+	TObjectPtr<const UCameraEvaluationContext> EvaluationContext;
+	/** An optional allocation information to optimize storage. */
+	FCameraNodeEvaluatorTreeAllocationInfo* AllocationInfo = nullptr;
+};
+
+/**
+ * A class responsible for storing a tree of camera node evaluators.
+ */
+class FCameraNodeEvaluatorStorage
+{
+public:
+
+	FCameraNodeEvaluatorStorage();
+	FCameraNodeEvaluatorStorage(FCameraNodeEvaluatorStorage&& Other);
+	~FCameraNodeEvaluatorStorage();
+
+	FCameraNodeEvaluatorStorage(const FCameraNodeEvaluatorStorage&) = delete;
+	FCameraNodeEvaluatorStorage& operator=(const FCameraNodeEvaluatorStorage&) = delete;
+
+public:
+
+	/** Compute allocation information for the given tree of camera nodes. */
+	FCameraNodeEvaluatorTreeAllocationInfo ComputeTreeInfo(const UCameraMode* CameraMode) const;
+	/** Build the tree of evaluators for the given tree of camera nodes. */
+	FCameraNodeEvaluatorPtr BuildEvaluatorTree(const FCameraNodeEvaluatorTreeBuilderParams& Params);
+
+private:
+
+	template<typename EvaluatorType, typename ...ArgTypes>
+	EvaluatorType* BuildEvaluator(ArgTypes&&... InArgs);
+
+private:
+
+	struct FAllocation
+	{
+		uint8* Memory = nullptr;
+		uint16 Capacity = 0;
+		uint16 Used = 0;
+	};
+	TArray<FAllocation> Allocations;
+
+	struct FEvaluatorInfo
+	{
+		FCameraNodeEvaluator* Ptr = nullptr;
+	};
+	TArray<FEvaluatorInfo> EvaluatorInfos;
+
+	friend struct FCameraNodeEvaluatorBuilder;
+};
+
+template<typename EvaluatorType, typename ...ArgTypes>
+EvaluatorType* FCameraNodeEvaluatorStorage::BuildEvaluator(ArgTypes&&... InArgs)
+{
+	static const uint16 DefaultCapacity = 128;
+	static const uint16 DefaultAlignment = 32;
+
+	uint16 StateSizeof = sizeof(EvaluatorType);
+	uint16 StateAlignof = alignof(EvaluatorType);
+
+	// Search for any allocation bucket that has enough room for the evaluator
+	// we want to build.
+	uint8* TargetPtr = nullptr;
+	FAllocation* TargetAllocation = nullptr;
+	for (FAllocation& Allocation : Allocations)
+	{
+		uint8* PossiblePtr = Align(Allocation.Memory + Allocation.Used, StateAlignof);
+		uint16 NewUsed = (PossiblePtr + StateSizeof) - Allocation.Memory;
+		if (NewUsed <= Allocation.Capacity)
+		{
+			TargetPtr = PossiblePtr;
+			TargetAllocation = &Allocation;
+			TargetAllocation->Used = NewUsed;
+			break;
+		}
+	}
+	// If we didn't find anything, we need to make a new allocation bucket.
+	if (TargetPtr == nullptr)
+	{
+		const uint16 NewCapacity = FMath::Max(DefaultCapacity, StateSizeof);
+		const uint16 NewAlignment = FMath::Max(DefaultAlignment, StateAlignof);
+
+		FAllocation& NewAllocation = Allocations.Emplace_GetRef();
+		NewAllocation.Memory = reinterpret_cast<uint8*>(FMemory::Malloc(NewCapacity, NewAlignment));
+		NewAllocation.Capacity = NewCapacity;
+		NewAllocation.Used = StateSizeof;
+
+		TargetPtr = NewAllocation.Memory;
+		TargetAllocation = &NewAllocation;
+	}
+	check(TargetPtr && TargetAllocation);
+
+	EvaluatorType* NewEvaluator = new(TargetPtr) EvaluatorType(Forward<ArgTypes>(InArgs)...);
+
+	EvaluatorInfos.Add({ NewEvaluator });
+
+	return NewEvaluator;
+}
+
diff --git a/Engine/Plugins/Cameras/GameplayCameras/Source/GameplayCameras/Public/Core/CameraNodeTreeCache.h b/Engine/Plugins/Cameras/GameplayCameras/Source/GameplayCameras/Public/Core/CameraNodeTreeCache.h
deleted file mode 100644
diff --git a/Engine/Plugins/Cameras/GameplayCameras/Source/GameplayCameras/Public/Core/CameraObjectRtti.h b/Engine/Plugins/Cameras/GameplayCameras/Source/GameplayCameras/Public/Core/CameraObjectRtti.h
new file mode 100644
--- /dev/null
+++ b/Engine/Plugins/Cameras/GameplayCameras/Source/GameplayCameras/Public/Core/CameraObjectRtti.h
@@ -0,0 +1,94 @@
+// Copyright Epic Games, Inc. All Rights Reserved.
+
+#pragma once
+
+#include "CoreTypes.h"
+#include "UObject/NameTypes.h"
+
+/**
+ * Identifier for a given RTTI-enabled class.
+ */
+struct FCameraObjectTypeID
+{
+	friend uint32 GetTypeHash(FCameraObjectTypeID In)
+	{
+		return In.ID;
+	}
+
+	friend bool operator<(FCameraObjectTypeID A, FCameraObjectTypeID B)
+	{
+		return A.ID < B.ID;
+	}
+
+	friend bool operator==(FCameraObjectTypeID A, FCameraObjectTypeID B)
+	{
+		return A.ID == B.ID;
+	}
+
+	uint32 GetTypeID() const
+	{
+		return ID;
+	}
+
+	FCameraObjectTypeID(const FName& InName, uint32 InID)
+		: Name(InName)
+		, ID(InID)
+	{}
+
+protected:
+
+	static uint32 RegisterNewID();
+
+protected:
+
+	FName Name;
+	uint32 ID;
+};
+
+/**
+ * Strongly-typed ID wrapper for an RTTI-enabled class.
+ */
+template<typename T>
+struct TCameraObjectTypeID : FCameraObjectTypeID
+{
+private:
+
+	TCameraObjectTypeID(const FName& InName, uint32 InID) : FCameraObjectTypeID(InName, InID) {}
+
+	static TCameraObjectTypeID RegisterNewID(const FName& InClassName)
+	{
+		return TCameraObjectTypeID(InClassName, FCameraObjectTypeID::RegisterNewID());
+	}
+
+	friend T;
+};
+
+// Macros for enabling simple RTTI information on a class hierarchy.
+//
+// The first macro is for the root of the class hierarchy, while the second is for all
+// other classes below it. The third macro goes in the cpp file of each class.
+//
+#define UE_GAMEPLAY_CAMERAS_DECLARE_RTTI_BASE(ClassName)\
+	public:\
+		static const TCameraObjectTypeID<ClassName>& StaticTypeID() { return ClassName::PrivateTypeID; }\
+		virtual const FCameraObjectTypeID& GetTypeID() const { return ClassName::PrivateTypeID; }\
+		virtual bool IsKindOf(const FCameraObjectTypeID& InTypeID) const { return InTypeID == ClassName::PrivateTypeID; }\
+		template<typename Type> bool IsKindOf() const { return IsKindOf(Type::StaticTypeID()); }\
+		template<typename Type> Type* CastThis() { return IsKindOf<Type>() ? static_cast<Type*>(this) : nullptr; }\
+		template<typename Type> const Type* CastThis() const { return IsKindOf<Type>() ? static_cast<Type*>(this) : nullptr; }\
+		template<typename Type> Type* CastThisChecked() { check(IsKindOf<Type>()); return static_cast<Type*>(this); }\
+		template<typename Type> const Type* CastThisChecked() const { check(IsKindOf<Type>()); return static_cast<Type*>(this); }\
+	private:\
+		static const TCameraObjectTypeID<ClassName> PrivateTypeID;
+
+#define UE_GAMEPLAY_CAMERAS_DECLARE_RTTI(ClassName, BaseClassName)\
+	public:\
+		static const TCameraObjectTypeID<ClassName>& StaticTypeID() { return ClassName::PrivateTypeID; }\
+		virtual const FCameraObjectTypeID& GetTypeID() const override { return ClassName::PrivateTypeID; }\
+		virtual bool IsKindOf(const FCameraObjectTypeID& InTypeID) const override { return (InTypeID == ClassName::PrivateTypeID) || BaseClassName::IsKindOf(InTypeID); }\
+	private:\
+		static const TCameraObjectTypeID<ClassName> PrivateTypeID;
+
+#define UE_GAMEPLAY_CAMERAS_DEFINE_RTTI(ClassName)\
+	const TCameraObjectTypeID<ClassName> ClassName::PrivateTypeID = TCameraObjectTypeID<ClassName>::RegisterNewID(#ClassName);
+
diff --git a/Engine/Plugins/Cameras/GameplayCameras/Source/GameplayCameras/Public/Core/CameraRuntimeInstantiator.h b/Engine/Plugins/Cameras/GameplayCameras/Source/GameplayCameras/Public/Core/CameraRuntimeInstantiator.h
deleted file mode 100644
diff --git a/Engine/Plugins/Cameras/GameplayCameras/Source/GameplayCameras/Public/Core/CameraSystemEvaluator.h b/Engine/Plugins/Cameras/GameplayCameras/Source/GameplayCameras/Public/Core/CameraSystemEvaluator.h
--- a/Engine/Plugins/Cameras/GameplayCameras/Source/GameplayCameras/Public/Core/CameraSystemEvaluator.h
+++ b/Engine/Plugins/Cameras/GameplayCameras/Source/GameplayCameras/Public/Core/CameraSystemEvaluator.h
@@ -5,13 +5,15 @@
 #include "Core/CameraDirector.h"
 #include "Core/CameraEvaluationContextStack.h"
 #include "Core/CameraNode.h"
+#include "Core/CameraNodeEvaluator.h"
+#include "Core/CameraNodeEvaluatorStorage.h"
 #include "Core/CameraPose.h"
-#include "Core/CameraRuntimeInstantiator.h"
 #include "CoreTypes.h"
 #include "UObject/GCObject.h"
 
 #include "CameraSystemEvaluator.generated.h"
 
+class FRootCameraNodeEvaluator;
 class UCameraDirector;
 class UCameraEvaluationContext;
 class UCameraMode;
@@ -71,11 +73,6 @@ public:
 	/** Get the last evaluated camera. */
 	void GetEvaluatedCameraView(FMinimalViewInfo& DesiredView);
 
-public:
-
-	/** Gets the runtime instantiator for this evaluator. */
-	FCameraRuntimeInstantiator& GetRuntimeInstantiator();
-
 public:
 
 	static void AddReferencedObjects(UObject* InThis, FReferenceCollector& Collector);
@@ -89,13 +86,16 @@ private:
 	/** The stack of active evaluation context. */
 	FCameraEvaluationContextStack ContextStack;
 
+	/** Storage buffer for the root evaluator. */
+	FCameraNodeEvaluatorStorage RootEvaluatorStorage;
+
+	/** The root evaluator. */
+	FRootCameraNodeEvaluator* RootEvaluator;
+
 	/** The current result of the root camera node. */
-	FCameraNodeRunResult RootNodeResult;
+	FCameraNodeEvaluationResult RootNodeResult;
 
 	/** The current overall result of the camera system. */
 	FCameraSystemEvaluationUpdateResult Result;
-
-	/** The instantiator for runtime objects. */
-	FCameraRuntimeInstantiator Instantiator;
 };
 
diff --git a/Engine/Plugins/Cameras/GameplayCameras/Source/GameplayCameras/Public/Core/DefaultRootCameraNode.h b/Engine/Plugins/Cameras/GameplayCameras/Source/GameplayCameras/Public/Core/DefaultRootCameraNode.h
--- a/Engine/Plugins/Cameras/GameplayCameras/Source/GameplayCameras/Public/Core/DefaultRootCameraNode.h
+++ b/Engine/Plugins/Cameras/GameplayCameras/Source/GameplayCameras/Public/Core/DefaultRootCameraNode.h
@@ -7,6 +7,7 @@
 #include "DefaultRootCameraNode.generated.h"
 
 class UBlendStackCameraNode;
+class FBlendStackCameraNodeEvaluator;
 
 /**
  * The default implementation of a root camera node.
@@ -22,23 +23,43 @@ public:
 
 protected:
 
-	virtual FCameraNodeChildrenView OnGetChildren() override;
-	virtual void OnRun(const FCameraNodeRunParams& Params, FCameraNodeRunResult& OutResult) override;
+	virtual FCameraNodeEvaluatorPtr OnBuildEvaluator(FCameraNodeEvaluatorBuilder& Builder) const override;
 
-	virtual void OnActivateCameraMode(const FActivateCameraModeParams& Params) override;
-
-private:
+public:
 
 	UPROPERTY(Instanced)
-	TObjectPtr<UCameraNode> BaseLayer;
+	TObjectPtr<UBlendStackCameraNode> BaseLayer;
 
 	UPROPERTY(Instanced)
-	TObjectPtr<UCameraNode> MainLayer;
+	TObjectPtr<UBlendStackCameraNode> MainLayer;
 
 	UPROPERTY(Instanced)
-	TObjectPtr<UCameraNode> GlobalLayer;
+	TObjectPtr<UBlendStackCameraNode> GlobalLayer;
 
 	UPROPERTY(Instanced)
-	TObjectPtr<UCameraNode> VisualLayer;
+	TObjectPtr<UBlendStackCameraNode> VisualLayer;
+};
+
+/**
+ * Evaluator for the default root camera node.
+ */
+class FDefaultRootCameraNodeEvaluator : public FRootCameraNodeEvaluator
+{
+	UE_DECLARE_CAMERA_NODE_EVALUATOR(FDefaultRootCameraNodeEvaluator)
+
+protected:
+
+	// FRootCameraNodeEvaluator interface.
+	virtual FCameraNodeEvaluatorChildrenView OnGetChildren() override;
+	virtual void OnInitialize(const FCameraNodeEvaluatorInitializeParams& Params) override;
+	virtual void OnRun(const FCameraNodeEvaluationParams& Params, FCameraNodeEvaluationResult& OutResult) override;
+	virtual void OnActivateCameraMode(const FActivateCameraModeParams& Params) override;
+
+private:
+
+	FBlendStackCameraNodeEvaluator* BaseLayer;
+	FBlendStackCameraNodeEvaluator* MainLayer;
+	FBlendStackCameraNodeEvaluator* GlobalLayer;
+	FBlendStackCameraNodeEvaluator* VisualLayer;
 };
 
diff --git a/Engine/Plugins/Cameras/GameplayCameras/Source/GameplayCameras/Public/Core/ObjectChildrenView.h b/Engine/Plugins/Cameras/GameplayCameras/Source/GameplayCameras/Public/Core/ObjectChildrenView.h
new file mode 100644
--- /dev/null
+++ b/Engine/Plugins/Cameras/GameplayCameras/Source/GameplayCameras/Public/Core/ObjectChildrenView.h
@@ -0,0 +1,155 @@
+// Copyright Epic Games, Inc. All Rights Reserved.
+
+#pragma once
+
+#include "Containers/Array.h"
+#include "Containers/ArrayView.h"
+#include "Misc/TVariant.h"
+#include "Templates/UnrealTypeTraits.h"
+#include "UObject/ObjectPtr.h"
+
+#include <initializer_list>
+
+/**
+ * Structure that describes a list of children of an object.
+ *
+ * This structure can either provide a TArrayView<> on an existing container of children,
+ * or store within itself a list of arbitrary camera node children.
+ */
+template<typename ChildType>
+struct TObjectChildrenView
+{
+public:
+
+	using FArrayView = TArrayView<ChildType>;
+	using FArray = TArray<ChildType, TInlineAllocator<4>>;
+
+	/** An empty view. */
+	TObjectChildrenView()
+	{
+	}
+
+	/** Sets the view to the given TArrayView<>. */
+	TObjectChildrenView(FArrayView&& InArrayView)
+	{
+		Storage.template Set<FArrayView>(MoveTemp(InArrayView));
+	}
+
+	/** Sets the view to pointer storage and adds the given list of children pointers. */
+	TObjectChildrenView(std::initializer_list<ChildType> InChildren)
+	{
+		Storage.template Set<FArray>(FArray(InChildren));
+	}
+
+	/**
+	 * Sets the view to pointer storage (if not already done) and adds the given
+	 * pointer to the list.
+	 */
+	void Add(typename TCallTraits<ChildType>::ParamType InChild)
+	{
+		if (Storage.GetIndex() != FStorage::template IndexOfType<FArray>())
+		{
+			Storage.template Set<FArray>(FArray());
+		}
+		FArray& Array = Storage.template Get<FArray>();
+		Array.Add(InChild);
+	}
+
+	/** Whether this view has any children. */
+	bool IsEmpty() const
+	{
+		switch (Storage.GetIndex())
+		{
+			case FStorage::template IndexOfType<FArrayView>():
+				return Storage.template Get<FArrayView>().IsEmpty();
+			case FStorage::template IndexOfType<FArray>():
+				return Storage.template Get<FArray>().IsEmpty();
+			default:
+				return true;
+		}
+	}
+
+	/** Returns the number of children. */
+	int32 Num() const
+	{
+		switch (Storage.GetIndex())
+		{
+			case FStorage::template IndexOfType<FArrayView>():
+				return Storage.template Get<FArrayView>().Num();
+			case FStorage::template IndexOfType<FArray>():
+				return Storage.template Get<FArray>().Num();
+			default:
+				return 0;
+		}
+	}
+
+	/** Gets the i'th child. */
+	ChildType operator[](int32 Index) const
+	{
+		switch (Storage.GetIndex())
+		{
+			case FStorage::template IndexOfType<FArrayView>():
+				return Storage.template Get<FArrayView>()[Index];
+			case FStorage::template IndexOfType<FArray>():
+				return Storage.template Get<FArray>()[Index];
+			default:
+				return nullptr;
+		}
+	}
+
+public:
+
+	// Range iteration
+
+	struct FBaseIterator
+	{
+		const TObjectChildrenView* Owner;
+		int32 Index;
+
+		FORCEINLINE ChildType operator*()
+		{
+			return (*Owner)[Index];
+		}
+
+		FORCEINLINE bool operator== (const FBaseIterator& Other) const
+		{
+			return Owner == Other.Owner
+				&& Index == Other.Index;
+		}
+
+		FORCEINLINE bool operator!= (const FBaseIterator& Other) const
+		{
+			return !(*this == Other);
+		}
+	};
+
+	struct FIterator : FBaseIterator
+	{
+		FORCEINLINE FIterator& operator++()
+		{
+			++FBaseIterator::Index;
+			return *this;
+		}
+	};
+
+	FORCEINLINE FIterator begin() const { return FIterator{ this, 0 }; }
+	FORCEINLINE FIterator end() const { return FIterator{ this, Num() }; }
+
+	struct FReverseIterator : FBaseIterator
+	{
+		FORCEINLINE FReverseIterator& operator++()
+		{
+			--FBaseIterator::Index;
+			return *this;
+		}
+	};
+
+	FORCEINLINE FReverseIterator rbegin() const { return FReverseIterator{ this, Num() - 1 }; }
+	FORCEINLINE FReverseIterator rend() const { return FReverseIterator{ this, -1 }; }
+
+private:
+
+	using FStorage = TVariant<FArrayView, FArray>;
+	FStorage Storage;
+};
+
diff --git a/Engine/Plugins/Cameras/GameplayCameras/Source/GameplayCameras/Public/Core/RootCameraNode.h b/Engine/Plugins/Cameras/GameplayCameras/Source/GameplayCameras/Public/Core/RootCameraNode.h
--- a/Engine/Plugins/Cameras/GameplayCameras/Source/GameplayCameras/Public/Core/RootCameraNode.h
+++ b/Engine/Plugins/Cameras/GameplayCameras/Source/GameplayCameras/Public/Core/RootCameraNode.h
@@ -3,11 +3,13 @@
 #pragma once
 
 #include "Core/CameraNode.h"
+#include "Core/CameraNodeEvaluator.h"
 
 #include "RootCameraNode.generated.h"
 
 class UCameraEvaluationContext;
 class UCameraMode;
+class UCameraSystemEvaluator;
 
 /**
  * Defines evaluation layers for camera modes.
@@ -51,7 +53,13 @@ UCLASS(MinimalAPI, Abstract)
 class URootCameraNode : public UCameraNode
 {
 	GENERATED_BODY()
+};
 
+/**
+ * Base class for the evaluator of a root camera node.
+ */
+class FRootCameraNodeEvaluator : public FCameraNodeEvaluator
+{
 public:
 
 	/** Activates a camera mode. */
diff --git a/Engine/Plugins/Cameras/GameplayCameras/Source/GameplayCameras/Public/Directors/SingleCameraDirector.h b/Engine/Plugins/Cameras/GameplayCameras/Source/GameplayCameras/Public/Directors/SingleCameraDirector.h
--- a/Engine/Plugins/Cameras/GameplayCameras/Source/GameplayCameras/Public/Directors/SingleCameraDirector.h
+++ b/Engine/Plugins/Cameras/GameplayCameras/Source/GameplayCameras/Public/Directors/SingleCameraDirector.h
@@ -20,7 +20,8 @@ public:
 
 protected:
 
-	virtual void OnRun(const FCameraDirectorRunParams& Params, FCameraDirectorRunResult& OutResult) override;
+	// UCameraDirector interface.
+	virtual FCameraDirectorEvaluator* OnBuildEvaluator(FCameraDirectorEvaluatorBuilder& Builder) const override;
 
 public:
 
diff --git a/Engine/Plugins/Cameras/GameplayCameras/Source/GameplayCameras/Public/IGameplayCamerasLiveEditManager.h b/Engine/Plugins/Cameras/GameplayCameras/Source/GameplayCameras/Public/IGameplayCamerasLiveEditManager.h
--- a/Engine/Plugins/Cameras/GameplayCameras/Source/GameplayCameras/Public/IGameplayCamerasLiveEditManager.h
+++ b/Engine/Plugins/Cameras/GameplayCameras/Source/GameplayCameras/Public/IGameplayCamerasLiveEditManager.h
@@ -16,18 +16,5 @@ class IGameplayCamerasLiveEditManager : public TSharedFromThis<IGameplayCamerasL
 public:
 
 	virtual ~IGameplayCamerasLiveEditManager() {}
-
-	/**
-	 * Register a new set of instantiated objects.
-	 *
-	 * @param InstantiatedObjects  A mapping between a source object and an instantiated object.
-	 */
-	virtual void RegisterInstantiatedObjects(const TMap<UObject*, UObject*> InstantiatedObjects) = 0;
-
-	/**
-	 * Request that a property change on the given source object should be replicated on any
-	 * known related instantied objects.
-	 */
-	virtual void ForwardPropertyChange(const UObject* Object, const FPropertyChangedEvent& PropertyChangedEvent) = 0;
 };
 
diff --git a/Engine/Plugins/Cameras/GameplayCameras/Source/GameplayCameras/Public/IGameplayCamerasModule.h b/Engine/Plugins/Cameras/GameplayCameras/Source/GameplayCameras/Public/IGameplayCamerasModule.h
--- a/Engine/Plugins/Cameras/GameplayCameras/Source/GameplayCameras/Public/IGameplayCamerasModule.h
+++ b/Engine/Plugins/Cameras/GameplayCameras/Source/GameplayCameras/Public/IGameplayCamerasModule.h
@@ -3,7 +3,6 @@
 #pragma once
 
 #include "Modules/ModuleInterface.h"
-#include "Templates/SharedPointer.h"
 
 #if WITH_EDITOR
 class IGameplayCamerasLiveEditManager;
diff --git a/Engine/Plugins/Cameras/GameplayCameras/Source/GameplayCameras/Public/Nodes/Blends/LinearBlendCameraNode.h b/Engine/Plugins/Cameras/GameplayCameras/Source/GameplayCameras/Public/Nodes/Blends/LinearBlendCameraNode.h
--- a/Engine/Plugins/Cameras/GameplayCameras/Source/GameplayCameras/Public/Nodes/Blends/LinearBlendCameraNode.h
+++ b/Engine/Plugins/Cameras/GameplayCameras/Source/GameplayCameras/Public/Nodes/Blends/LinearBlendCameraNode.h
@@ -14,6 +14,8 @@ class ULinearBlendCameraNode : public USimpleFixedTimeBlendCameraNode
 {
 	GENERATED_BODY()
 
-	virtual void OnComputeBlendFactor(const FCameraNodeRunParams& Params, FSimpleBlendCameraNodeRunResult& OutResult) override;
+protected:
+
+	virtual FCameraNodeEvaluatorPtr OnBuildEvaluator(FCameraNodeEvaluatorBuilder& Builder) const override;
 };
 
diff --git a/Engine/Plugins/Cameras/GameplayCameras/Source/GameplayCameras/Public/Nodes/Blends/PopBlendCameraNode.h b/Engine/Plugins/Cameras/GameplayCameras/Source/GameplayCameras/Public/Nodes/Blends/PopBlendCameraNode.h
--- a/Engine/Plugins/Cameras/GameplayCameras/Source/GameplayCameras/Public/Nodes/Blends/PopBlendCameraNode.h
+++ b/Engine/Plugins/Cameras/GameplayCameras/Source/GameplayCameras/Public/Nodes/Blends/PopBlendCameraNode.h
@@ -16,7 +16,6 @@ class UPopBlendCameraNode : public UBlendCameraNode
 
 protected:
 
-	virtual void OnRun(const FCameraNodeRunParams& Params, FCameraNodeRunResult& OutResult) override;
-	virtual void OnBlendResults(const FCameraNodeBlendParams& Params, FCameraNodeBlendResult& OutResult) override;
+	virtual FCameraNodeEvaluatorPtr OnBuildEvaluator(FCameraNodeEvaluatorBuilder& Builder) const override;
 };
 
diff --git a/Engine/Plugins/Cameras/GameplayCameras/Source/GameplayCameras/Public/Nodes/Blends/SimpleBlendCameraNode.h b/Engine/Plugins/Cameras/GameplayCameras/Source/GameplayCameras/Public/Nodes/Blends/SimpleBlendCameraNode.h
--- a/Engine/Plugins/Cameras/GameplayCameras/Source/GameplayCameras/Public/Nodes/Blends/SimpleBlendCameraNode.h
+++ b/Engine/Plugins/Cameras/GameplayCameras/Source/GameplayCameras/Public/Nodes/Blends/SimpleBlendCameraNode.h
@@ -9,7 +9,7 @@
 /**
  * Result structure for defining a simple scalar-factor-based blend.
  */
-struct FSimpleBlendCameraNodeRunResult
+struct FSimpleBlendCameraNodeEvaluationResult
 {
 	float BlendFactor = 0.f;
 };
@@ -21,6 +21,26 @@ UCLASS(MinimalAPI, Abstract)
 class USimpleBlendCameraNode : public UBlendCameraNode
 {
 	GENERATED_BODY()
+};
+
+/**
+ * Base class for a blend camera node that uses a simple scalar factor over a fixed time.
+ */
+UCLASS(MinimalAPI, Abstract)
+class USimpleFixedTimeBlendCameraNode : public USimpleBlendCameraNode
+{
+	GENERATED_BODY()
+
+public:
+
+	/** Duration of the blend. */
+	UPROPERTY(EditAnywhere, Category=Blending)
+	float BlendTime = 1.f;
+};
+
+class FSimpleBlendCameraNodeEvaluator : public FBlendCameraNodeEvaluator
+{
+	UE_DECLARE_BLEND_CAMERA_NODE_EVALUATOR(FSimpleBlendCameraNodeEvaluator)
 
 public:
 
@@ -29,10 +49,10 @@ public:
 
 protected:
 
-	virtual void OnRun(const FCameraNodeRunParams& Params, FCameraNodeRunResult& OutResult) override;
+	virtual void OnRun(const FCameraNodeEvaluationParams& Params, FCameraNodeEvaluationResult& OutResult) override;
 	virtual void OnBlendResults(const FCameraNodeBlendParams& Params, FCameraNodeBlendResult& OutResult) override;
 
-	virtual void OnComputeBlendFactor(const FCameraNodeRunParams& Params, FSimpleBlendCameraNodeRunResult& OutResult) {}
+	virtual void OnComputeBlendFactor(const FCameraNodeEvaluationParams& Params, FSimpleBlendCameraNodeEvaluationResult& OutResult) {}
 
 	void SetBlendFinished() { bIsBlendFinished = true; }
 
@@ -42,25 +62,15 @@ private:
 	bool bIsBlendFinished = false;
 };
 
-/**
- * Base class for a blend camera node that uses a simple scalar factor over a fixed time.
- */
-UCLASS(MinimalAPI, Abstract)
-class USimpleFixedTimeBlendCameraNode : public USimpleBlendCameraNode
+class FSimpleFixedTimeBlendCameraNodeEvaluator : public FSimpleBlendCameraNodeEvaluator
 {
-	GENERATED_BODY()
-
-public:
-
-	/** Duration of the blend. */
-	UPROPERTY(EditAnywhere, Category=Blending)
-	float BlendTime = 1.f;
+	UE_DECLARE_BLEND_CAMERA_NODE_EVALUATOR_EX(FSimpleFixedTimeBlendCameraNodeEvaluator, FSimpleBlendCameraNodeEvaluator)
 
 protected:
 
-	virtual void OnRun(const FCameraNodeRunParams& Params, FCameraNodeRunResult& OutResult) override;
+	virtual void OnRun(const FCameraNodeEvaluationParams& Params, FCameraNodeEvaluationResult& OutResult) override;
 
-	float GetTimeFactor() const { return CurrentTime / BlendTime; }
+	float GetTimeFactor() const;
 
 private:
 
diff --git a/Engine/Plugins/Cameras/GameplayCameras/Source/GameplayCameras/Public/Nodes/Blends/SmoothBlendCameraNode.h b/Engine/Plugins/Cameras/GameplayCameras/Source/GameplayCameras/Public/Nodes/Blends/SmoothBlendCameraNode.h
--- a/Engine/Plugins/Cameras/GameplayCameras/Source/GameplayCameras/Public/Nodes/Blends/SmoothBlendCameraNode.h
+++ b/Engine/Plugins/Cameras/GameplayCameras/Source/GameplayCameras/Public/Nodes/Blends/SmoothBlendCameraNode.h
@@ -26,7 +26,7 @@ class USmoothBlendCameraNode : public USimpleFixedTimeBlendCameraNode
 
 protected:
 
-	virtual void OnComputeBlendFactor(const FCameraNodeRunParams& Params, FSimpleBlendCameraNodeRunResult& OutResult) override;
+	virtual FCameraNodeEvaluatorPtr OnBuildEvaluator(FCameraNodeEvaluatorBuilder& Builder) const override;
 	
 public:
 
diff --git a/Engine/Plugins/Cameras/GameplayCameras/Source/GameplayCameras/Public/Nodes/Common/ArrayCameraNode.h b/Engine/Plugins/Cameras/GameplayCameras/Source/GameplayCameras/Public/Nodes/Common/ArrayCameraNode.h
--- a/Engine/Plugins/Cameras/GameplayCameras/Source/GameplayCameras/Public/Nodes/Common/ArrayCameraNode.h
+++ b/Engine/Plugins/Cameras/GameplayCameras/Source/GameplayCameras/Public/Nodes/Common/ArrayCameraNode.h
@@ -17,8 +17,8 @@ class UArrayCameraNode : public UCameraNode
 
 protected:
 
-	virtual FCameraNodeChildrenView OnGetChildren() override;
-	virtual void OnRun(const FCameraNodeRunParams& Params, FCameraNodeRunResult& OutResult) override;
+	// UCameraNode interface.
+	virtual FCameraNodeEvaluatorPtr OnBuildEvaluator(FCameraNodeEvaluatorBuilder& Builder) const override;
 
 public:
 
diff --git a/Engine/Plugins/Cameras/GameplayCameras/Source/GameplayCameras/Public/Nodes/Common/OffsetCameraNode.h b/Engine/Plugins/Cameras/GameplayCameras/Source/GameplayCameras/Public/Nodes/Common/OffsetCameraNode.h
--- a/Engine/Plugins/Cameras/GameplayCameras/Source/GameplayCameras/Public/Nodes/Common/OffsetCameraNode.h
+++ b/Engine/Plugins/Cameras/GameplayCameras/Source/GameplayCameras/Public/Nodes/Common/OffsetCameraNode.h
@@ -17,7 +17,8 @@ class UOffsetCameraNode : public UCameraNode
 
 protected:
 
-	virtual void OnRun(const FCameraNodeRunParams& Params, FCameraNodeRunResult& OutResult) override;
+	// UCameraNode interface.
+	virtual FCameraNodeEvaluatorPtr OnBuildEvaluator(FCameraNodeEvaluatorBuilder& Builder) const override;
 
 public:
 
diff --git a/Engine/Plugins/Cameras/GameplayCameras/Source/GameplayCamerasEditor/Private/GameplayCamerasEditorModule.cpp b/Engine/Plugins/Cameras/GameplayCameras/Source/GameplayCamerasEditor/Private/GameplayCamerasEditorModule.cpp
--- a/Engine/Plugins/Cameras/GameplayCameras/Source/GameplayCamerasEditor/Private/GameplayCamerasEditorModule.cpp
+++ b/Engine/Plugins/Cameras/GameplayCameras/Source/GameplayCamerasEditor/Private/GameplayCamerasEditorModule.cpp
@@ -100,7 +100,6 @@ private:
 	
 	void OnPostGarbageCollection()
 	{
-		LiveEditManager->CleanUp();
 	}
 
 private:
diff --git a/Engine/Plugins/Cameras/GameplayCameras/Source/GameplayCamerasEditor/Private/GameplayCamerasLiveEditManager.cpp b/Engine/Plugins/Cameras/GameplayCameras/Source/GameplayCamerasEditor/Private/GameplayCamerasLiveEditManager.cpp
--- a/Engine/Plugins/Cameras/GameplayCameras/Source/GameplayCamerasEditor/Private/GameplayCamerasLiveEditManager.cpp
+++ b/Engine/Plugins/Cameras/GameplayCameras/Source/GameplayCamerasEditor/Private/GameplayCamerasLiveEditManager.cpp
@@ -2,101 +2,9 @@
 
 #include "GameplayCamerasLiveEditManager.h"
 
-#include "Core/CameraInstantiableObject.h"
 #include "UObject/UObjectGlobals.h"
 
 FGameplayCamerasLiveEditManager::FGameplayCamerasLiveEditManager()
 {
 }
 
-void FGameplayCamerasLiveEditManager::CleanUp()
-{
-	for (auto It = Instantiations.CreateIterator(); It; ++It)
-	{
-		// Clean up invalid entries on keys and values.
-		UObject* Obj = It->Key.Get();
-		if (!Obj)
-		{
-			It.RemoveCurrent();
-			continue;
-		}
-		for (auto ObjIt = It->Value.InstantiatedObjects.CreateIterator(); ObjIt; ++ObjIt)
-		{
-			if (!ObjIt->IsValid())
-			{
-				ObjIt.RemoveCurrent();
-			}
-		}
-
-		// If the source object doesn't have any instantiated objects anymore, clear its flags.
-		if (It->Value.InstantiatedObjects.IsEmpty())
-		{
-			if (UCameraInstantiableObject* SourceObject = Cast<UCameraInstantiableObject>(Obj))
-			{
-				SourceObject->SetInstantiationState(ECameraNodeInstantiationState::None);
-			}
-		}
-	}
-}
-
-void FGameplayCamerasLiveEditManager::RegisterInstantiatedObjects(const TMap<UObject*, UObject*> InstantiatedObjects)
-{
-	for (const TPair<UObject*, UObject*>& Pair : InstantiatedObjects)
-	{
-		// If the objects are instantiable, set appropriate flags on both.
-		if (UCameraInstantiableObject* SourceObject = Cast<UCameraInstantiableObject>(Pair.Key))
-		{
-			SourceObject->SetInstantiationState(ECameraNodeInstantiationState::HasInstantiations);
-
-			UCameraInstantiableObject* InstantiatedObject = CastChecked<UCameraInstantiableObject>(Pair.Value);
-			InstantiatedObject->SetInstantiationState(ECameraNodeInstantiationState::IsInstantiated);
-		}
-		
-		FInstantiationInfo& Info = Instantiations.FindOrAdd(Pair.Key);
-		Info.InstantiatedObjects.Add(Pair.Value);
-	}
-}
-
-void FGameplayCamerasLiveEditManager::ForwardPropertyChange(const UObject* Object, const FPropertyChangedEvent& PropertyChangedEvent)
-{
-	FInstantiationInfo* Info = Instantiations.Find(Object);
-	if (!Info)
-	{
-		return;
-	}
-
-	const UClass* ObjectClass = Object->GetClass();
-	const FProperty* ChangedProperty = PropertyChangedEvent.MemberProperty;
-	if (!ensure(ChangedProperty))
-	{
-		return;
-	}
-
-	bool bForwardChange = false;
-	if (ChangedProperty->IsA<FBoolProperty>()
-			|| ChangedProperty->IsA<FNumericProperty>()
-			|| ChangedProperty->IsA<FStructProperty>())
-	{
-		bForwardChange = true;
-	}
-	if (!bForwardChange)
-	{
-		return;
-	}
-
-	const void* SourceValuePtr = ChangedProperty->ContainerPtrToValuePtr<void>(Object);
-	for (auto It = Info->InstantiatedObjects.CreateIterator(); It; ++It)
-	{
-		if (UObject* Inst = It->Get())
-		{
-			ensure(Inst->GetClass() == ObjectClass);
-			void* InstValuePtr = ChangedProperty->ContainerPtrToValuePtr<void>(Inst);
-			ChangedProperty->CopyCompleteValue(InstValuePtr, SourceValuePtr);
-		}
-		else
-		{
-			It.RemoveCurrent();
-		}
-	}
-}
-
diff --git a/Engine/Plugins/Cameras/GameplayCameras/Source/GameplayCamerasEditor/Private/GameplayCamerasLiveEditManager.h b/Engine/Plugins/Cameras/GameplayCameras/Source/GameplayCamerasEditor/Private/GameplayCamerasLiveEditManager.h
--- a/Engine/Plugins/Cameras/GameplayCameras/Source/GameplayCamerasEditor/Private/GameplayCamerasLiveEditManager.h
+++ b/Engine/Plugins/Cameras/GameplayCameras/Source/GameplayCamerasEditor/Private/GameplayCamerasLiveEditManager.h
@@ -12,22 +12,5 @@ class FGameplayCamerasLiveEditManager : public IGameplayCamerasLiveEditManager
 public:
 
 	FGameplayCamerasLiveEditManager();
-
-	/** Clean-up any invalid entries in the map of known instantiated objects. */
-	void CleanUp();
-
-public:
-
-	// IGameplayCamerasLiveEditManager interface
-	virtual void RegisterInstantiatedObjects(const TMap<UObject*, UObject*> InstantiatedObjects) override;
-	virtual void ForwardPropertyChange(const UObject* Object, const FPropertyChangedEvent& PropertyChangedEvent) override;
-
-private:
-
-	struct FInstantiationInfo
-	{
-		TArray<TWeakObjectPtr<>> InstantiatedObjects;
-	};
-	TMap<TWeakObjectPtr<UObject>, FInstantiationInfo> Instantiations;
 };
 
