Cameras: add critical damper utility class, and dampen position camera node

#changelist validated
#virtualized

diff --git a/Engine/Plugins/Cameras/GameplayCameras/Source/GameplayCameras/Private/Math/CriticalDamper.cpp b/Engine/Plugins/Cameras/GameplayCameras/Source/GameplayCameras/Private/Math/CriticalDamper.cpp
new file mode 100644
--- /dev/null
+++ b/Engine/Plugins/Cameras/GameplayCameras/Source/GameplayCameras/Private/Math/CriticalDamper.cpp
@@ -0,0 +1,249 @@
+// Copyright Epic Games, Inc. All Rights Reserved.
+
+#include "Math/CriticalDamper.h"
+
+#include "HAL/IConsoleManager.h"
+#include "Math/UnrealMathUtility.h"
+
+#include UE_INLINE_GENERATED_CPP_BY_NAME(CriticalDamper)
+
+float GGameplayCamerasCriticalDamperStabilizationThreshold = 0.00001f;
+static FAutoConsoleVariableRef CVarGameplayCamerasCriticalDamperStabilizationThreshold(
+	TEXT("GameplayCameras.CriticalDamper.StabilizationThreshold"),
+	GGameplayCamerasCriticalDamperStabilizationThreshold,
+	TEXT("(Default: 0.00001f (in meters and meters/sec). Determines when to clamp a critical damper to 0 to stabilize it.")
+	);
+
+FCriticalDamper::FCriticalDamper()
+{
+	Reset(0, 0);
+	SetW0(-1);
+}
+
+FCriticalDamper::FCriticalDamper(float InW0)
+{
+	Reset(0, 0);
+	SetW0(InW0);
+}
+
+void FCriticalDamper::Reset(float InX0, float InX0Derivative)
+{
+	X0 = InX0;
+	X0Derivative = InX0Derivative;
+}
+
+float FCriticalDamper::Update(float DeltaTime)
+{
+	return Update(X0, DeltaTime);
+}
+
+float FCriticalDamper::Update(float X, float DeltaTime)
+{
+	// Last frame we were at X0. This frame we are at X, so something forcibly moved
+	// us by the difference.
+	const float ForcedMovement = (X - X0);
+	InternalUpdate(ForcedMovement, DeltaTime);
+	return X0;
+}
+
+float FCriticalDamper::Update(float PreviousDamped, float NextTarget, float DeltaTime)
+{
+	const float X = (NextTarget - PreviousDamped);
+	Update(X, DeltaTime);
+	// Return the position of the system relative to current frame's target.
+	return NextTarget - X0;
+}
+
+void FCriticalDamper::InternalUpdate(float ForcedMovement, float DeltaTime)
+{
+	if (W0 > 0.f)
+	{
+		// We need to move the base of the spring by ForcedMovement, and then run the
+		// simulation of the spring-mass system. This first step is because we typically 
+		// have the base of the spring attached to a moving object (the player character,
+		// the player's moving vehicle, etc.) This means we effectively have a "driven
+		// harmonic oscillator" system.
+		//
+		// References:
+		//		https://en.wikipedia.org/wiki/Damping
+		//      https://en.wikipedia.org/wiki/Harmonic_oscillator
+		//		https://mathworld.wolfram.com/CriticallyDampedSimpleHarmonicMotion.html
+		//		https://tutorial.math.lamar.edu/classes/de/nonhomogeneousde.aspx
+		//
+		// Notations:
+		//		x' is the first derivative of x, aka dx/dt
+		//		x'' is the second derivative of x, aka (d^2)x/d(t^2)
+		//
+		// Starting almost from first principles, we know that the movement of the mass in
+		// the spring-mass system is:
+		//
+		//		m*x'' = sum(forces)
+		//
+		// One force is the spring, so:
+		//
+		//		m*x'' = -k*x - c*x' + ExternalForces
+		//		m*x'' + c*x' + k*x = ExternalForces
+		//		x'' + c/m*x' + k/m*x = ExternalForces/m
+		//
+		// ...where m is the mass, k is the spring constant, and c is the viscous damping
+		// coefficient.
+		//
+		// Let's introduce some notations:
+		//
+		//		w0 = sqrt(k/m)  aka the natural frequency of the system
+		//		zeta = c/(2*sqrt(m*k))  aka the damping ratio
+		//
+		// Our equation becomes:
+		//
+		//		x'' + 2*zeta*w0*x' + w0^2*x = ExternalForces/m
+		//
+		// We want a critically damped system, which means we want to force zeta=1:
+		//
+		//		x'' + 2*w0*x' + w0^2*x = ExternalForces/m
+		// 
+		// Now let's look at ExternalForces. These exist when ForcedMovement is non
+		// zero, i.e. we have some external force that moves the base of the spring.
+		// This typically happens when the spring is attached to some moving player
+		// controlled entity like a character or a vehicle. The acceleration or
+		// deceleration of that entity expand or compresses the spring.
+		//
+		// When the base of the spring isn't moving ExternalForces is zero, and our 
+		// equation is:
+		//
+		//		x'' + 2*w0*x' + w0^2*x = 0
+		//
+		// This is a homogenous linear differential equation. Let's put a pin in that,
+		// and call it HE.
+		//
+		// Now, when the base of the spring is indeed moving, the force introduced in 
+		// the system is based on the spring itself:
+		//
+		//		ExternalForces = -k*D
+		//
+		// ...where D is the distance forced by the movement of the base. This is the
+		// equation for the added pull/push of the spring on the mass due to how much
+		// the base moved.
+		//
+		// We assume that this forced movement happened at constant speed over 
+		// DeltaTime (it may not have, but that's irrelevant since DeltaTime is our
+		// sampling rate) so:
+		//
+		//		ExternalForces = -k*(D0 + Dv*t)
+		//
+		// ...where D0 is the previous position of the spring base, Dv is the speed
+		// of the spring base, and t is the time. In practice D0 is always 0 since 
+		// we reset our simulation space to converge towards 0 every frame. So we
+		// can rewrite our general equation as:
+		//
+		//		x'' + 2*w0*x' + w0^2*x = -k*Dv*t/m
+		//		x'' + 2*w0*x' + w0^2*x = -w0^2*Dv*t
+		//
+		// This is a non-homogeneous linear differential equation. Let's call it NHE.
+		// The homogeneous one that we mentioned earlier (that we called HE, when 
+		// ExternalForces are 0), is therefore its "associated" homogeneous differential 
+		// equation, or "complementary" homogeneous differential equation.
+		//
+		// There's a theorem that says, very roughly, that if YP is a "particular 
+		// solution" to the NHE, we can just add it to the general solution of the HE
+		// and obtain a general solution to the NHE.
+		//
+		// Let's start by finding YP, a particular solution to the NHE. Since the 
+		// constant in the NHE is of the form c*t where c is a constant, we can use the
+		// "method of undertermined coefficients" to make an educated guess as to what
+		// a possible solution is. We can imagine that one solution is of the form:
+		//
+		//		x = P0 + P1*t
+		//
+		// To test this theory, let's derive this form and plug the result in the NHE:
+		//
+		//		x' = P1
+		//		x'' = 0
+		//		x'' + 2*w0*x' + w0^2*x = -w0^2*Dv*t
+		//		0 + 2*w0*P1 + w0^2*(P0 + P1*t) = -w0^2*Dv*t
+		//		2*w0*P1 + w0^2*(P0 + P1*t) = -w0^2*Dv*t
+		//		
+		// We solve this by setting t=0:
+		//
+		//		2*w0*P1 + w0^2*P0 = 0
+		//		P0 = -2*w0*P1/w0^2
+		//		P0 = -2*P1/w0
+		//
+		//	And so:
+		//
+		//		2*w0*P1 + w0^2*(P0 + P1*t) = -w0^2*Dv*t
+		//		2*w0*P1 + w0^2*(-2*P1/w0 + P1*t) = -w0^2*Dv*t
+		//		2*w0*P1 - 2*w0*P1 + w0^2*P1*t = -w0^2*Dv*t
+		//		w0^2*P1*t = -w0^2*Dv*t
+		//		P1 = -Dv
+		//
+		// Therefore:
+		//
+		//		P0 = 2*Dv/w0
+		//
+		// Going back to the full solution:
+		//
+		//		x = (A + B*t)*e^(-w0*t) + (P0 + P1*t)
+		//
+		// We derive this to get the velocity:
+		//
+		//		x' = -w0*A*e^(-w0*t) + B*e^(-w0*t) - w0*B*t*e^(-w0*t) + P1
+		//		x' = (-w0*A + B - w0*B*t)*e^(-w0*t) + P1
+		//
+		// Now we go back to figuring out A and B. Wolfram (see references) did that by
+		// using the initial conditions, so let's do that too with this solution that has
+		// the extra elements at the end from the NHE. At t=0, our two equations for x
+		// and x' become:
+		//
+		//		x(0) = A + P0
+		//		x'(0) = (-w0*A + B) + P1
+		//
+		// Solving for A and B:
+		//
+		//		A = x(0) - P0
+		//		B = x'(0) + w0*A - P1
+		//
+		// Now all we need to do is assemble all those pieces and write the code!
+
+		const float Dv = (ForcedMovement != 0.f) ? (ForcedMovement / DeltaTime) : 0.f;
+
+		const float P0 = 2.f * Dv / W0;
+		const float P1 = -Dv;
+
+		const float ExpMinusW0Dt = FMath::Exp(-W0 * DeltaTime);
+		const float A = X0 - P0;
+		const float B = X0Derivative + W0 * A - P1;
+
+		const float Xt = (A + B * DeltaTime) * ExpMinusW0Dt + P0 + P1 * DeltaTime;
+		const float XtDerivative = (-W0 * A + B - W0 * B * DeltaTime) * ExpMinusW0Dt + P1;
+
+		// Let's set the evaluation result of this frame as X0 and X0Derivative, which
+		// both store the last evaluated state, and also serve as the "initial conditions" 
+		// for next frame's evaluation.
+		//
+		// Note that we add ForcedMovement to X0 because the base of the spring has moved
+		// by that amount during this frame. This means that X0 is now relative to this
+		// new base, which used to be located at the origin at the beginning of the frame.
+		// We want X0 to be instead relative to 0.
+		X0 = Xt + ForcedMovement;
+		X0Derivative = XtDerivative;
+
+		// Floating point precision isn't good enough to let us mathematically converge
+		// towards 0 in an optimal manner, without overshooting. So overshooting by very
+		// small fractions can occur. We just stop the spring-mass system dead in its
+		// tracks when we figure that it's very close... since we test both position and
+		// velocity, this will only happen when it was about the settle anyway.
+		if (FMath::Abs(X0) <= GGameplayCamerasCriticalDamperStabilizationThreshold &&
+				FMath::Abs(X0Derivative) <= GGameplayCamerasCriticalDamperStabilizationThreshold)
+		{
+			X0 = 0;
+			X0Derivative = 0;
+		}
+	}
+	else if (W0 <= 0.f)
+	{
+		// The spring-mass system is disabled... let's just stick exactly to our target.
+		X0 = 0;
+		X0Derivative = 0;
+	}
+}
+
diff --git a/Engine/Plugins/Cameras/GameplayCameras/Source/GameplayCameras/Private/Nodes/Common/DampenPositionCameraNode.cpp b/Engine/Plugins/Cameras/GameplayCameras/Source/GameplayCameras/Private/Nodes/Common/DampenPositionCameraNode.cpp
new file mode 100644
--- /dev/null
+++ b/Engine/Plugins/Cameras/GameplayCameras/Source/GameplayCameras/Private/Nodes/Common/DampenPositionCameraNode.cpp
@@ -0,0 +1,109 @@
+// Copyright Epic Games, Inc. All Rights Reserved.
+
+#include "Nodes/Common/DampenPositionCameraNode.h"
+
+#include "Core/CameraEvaluationContext.h"
+#include "Core/CameraNodeEvaluator.h"
+#include "GameplayCameras.h"
+#include "Math/CriticalDamper.h"
+#include "Templates/Tuple.h"
+
+#include UE_INLINE_GENERATED_CPP_BY_NAME(DampenPositionCameraNode)
+
+class FDampenPositionCameraNodeEvaluator : public FCameraNodeEvaluator
+{
+	UE_DECLARE_CAMERA_NODE_EVALUATOR(FDampenPositionCameraNodeEvaluator)
+
+protected:
+
+	virtual void OnInitialize(const FCameraNodeEvaluatorInitializeParams& Params) override;
+	virtual void OnRun(const FCameraNodeEvaluationParams& Params, FCameraNodeEvaluationResult& OutResult) override;
+
+private:
+
+	FCriticalDamper ForwardDamper;
+	FCriticalDamper LateralDamper;
+	FCriticalDamper VerticalDamper;
+
+	FVector3d PreviousLocation;
+};
+
+UE_DEFINE_CAMERA_NODE_EVALUATOR(FDampenPositionCameraNodeEvaluator)
+
+void FDampenPositionCameraNodeEvaluator::OnInitialize(const FCameraNodeEvaluatorInitializeParams& Params)
+{
+	const UDampenPositionCameraNode* DampenNode = GetCameraNodeAs<UDampenPositionCameraNode>();
+
+	ForwardDamper.SetW0(DampenNode->ForwardDampingFactor);
+	ForwardDamper.Reset(0, 0);
+
+	LateralDamper.SetW0(DampenNode->LateralDampingFactor);
+	LateralDamper.Reset(0, 0);
+
+	VerticalDamper.SetW0(DampenNode->VerticalDampingFactor);
+	VerticalDamper.Reset(0, 0);
+
+	const FCameraNodeEvaluationResult& InitialResult = Params.EvaluationContext->GetInitialResult();
+	PreviousLocation = InitialResult.CameraPose.GetLocation();
+}
+
+void FDampenPositionCameraNodeEvaluator::OnRun(const FCameraNodeEvaluationParams& Params, FCameraNodeEvaluationResult& OutResult)
+{
+	// We want the dampen the given camera position, which means it's trying
+	// to converge towards the one given in the result (which we set as our 
+	// next target), but will be lagging behind.
+	const FVector3d NextTarget = OutResult.CameraPose.GetLocation();
+	FVector3d NextLocation = NextTarget;
+
+	if (!Params.bIsFirstFrame && !OutResult.bIsCameraCut)
+	{
+		FTransform3d Transform = OutResult.CameraPose.GetTransform();
+		FRotator3d Rotation = OutResult.CameraPose.GetRotation();
+
+		using FAxisDamper = TTuple<FVector3d, FCriticalDamper*>;
+		FAxisDamper AxisDampers[3]
+		{
+			{ Rotation.RotateVector(FVector3d::ForwardVector), &ForwardDamper },
+			{ Rotation.RotateVector(FVector3d::RightVector),& LateralDamper },
+			{ Rotation.RotateVector(FVector3d::UpVector), &VerticalDamper }
+		};
+
+		// The next target has moved further away compared to the previous target,
+		// so we're lagging behind even more than before. Compute this new lag vector.
+		const FVector3d NewLagVector = NextTarget - PreviousLocation;
+		// Let's start at our previous (dampened) location, and see by how much we
+		// can catch up on our lag this frame.
+		FVector3d NewDampedLocation = PreviousLocation;
+
+		for (FAxisDamper& AxisDamper : AxisDampers)
+		{
+			FVector3d Axis(AxisDamper.Key);
+			FCriticalDamper* Damper(AxisDamper.Value);
+
+			// Compute lag on the forward/lateral/vertical axis, and pass this new
+			// lag distance as the new position of the damper. Update it to know 
+			// how much we catch up, and offset last frame's position by that amount.
+			double NewLagDistance = FVector3d::DotProduct(NewLagVector, Axis);
+			// TODO: use GetWorld()->GetWorldSettings()->WorldToMeters
+			Damper->Update(NewLagDistance / 100.0, Params.DeltaTime);
+			NewDampedLocation += Axis * (NewLagDistance - Damper->GetX0() * 100.0);
+		}
+		
+		NextLocation = NewDampedLocation;
+	}
+
+	PreviousLocation = NextLocation;
+
+	OutResult.CameraPose.SetLocation(NextLocation);
+}
+
+UDampenPositionCameraNode::UDampenPositionCameraNode(const FObjectInitializer& ObjectInit)
+	: Super(ObjectInit)
+{
+}
+
+FCameraNodeEvaluatorPtr UDampenPositionCameraNode::OnBuildEvaluator(FCameraNodeEvaluatorBuilder& Builder) const
+{
+	return Builder.BuildEvaluator<FDampenPositionCameraNodeEvaluator>();
+}
+
diff --git a/Engine/Plugins/Cameras/GameplayCameras/Source/GameplayCameras/Private/Tests/CriticalDamperTests.cpp b/Engine/Plugins/Cameras/GameplayCameras/Source/GameplayCameras/Private/Tests/CriticalDamperTests.cpp
new file mode 100644
--- /dev/null
+++ b/Engine/Plugins/Cameras/GameplayCameras/Source/GameplayCameras/Private/Tests/CriticalDamperTests.cpp
@@ -0,0 +1,115 @@
+// Copyright Epic Games, Inc. All Rights Reserved.
+
+#include "CoreTypes.h"
+#include "Math/CriticalDamper.h"
+#include "Misc/AutomationTest.h"
+
+#define LOCTEXT_NAMESPACE "CriticalDamperTests"
+
+IMPLEMENT_SIMPLE_AUTOMATION_TEST(FGameplayCamerasCriticalDamperTest, "System.Engine.GameplayCameras.CriticalDamper", EAutomationTestFlags::EditorContext | EAutomationTestFlags::EngineFilter)
+bool FGameplayCamerasCriticalDamperTest::RunTest(const FString& Parameters)
+{
+	{
+		FCriticalDamper Damper(10.f);
+		float Test = Damper.Update(0, 1.f);
+		UTEST_EQUAL("Dampen 0", Test, 0.f);
+	}
+
+	{
+		FCriticalDamper Damper(10.f);
+		Damper.Reset(1, 0);
+		for (int32 Frame = 0; Frame < 30; ++Frame)
+		{
+			Damper.Update(0.34f);
+		}
+		UTEST_EQUAL("Dampen 1", Damper.GetX0(), 0.f);
+	}
+
+	{
+		FCriticalDamper Damper(10.f);
+		float Value = 1.f;
+		for (int32 Frame = 0; Frame < 30; ++Frame)
+		{
+			Value = Damper.Update(Value, 0.34f);
+		}
+		UTEST_EQUAL("Dampen 1 with value passing", Damper.GetX0(), 0.f);
+	}
+
+	{
+		FCriticalDamper Damper(10.f);
+		float Value = 1.f;
+		for (int32 Frame = 0; Frame < 30; ++Frame)
+		{
+			Value = Damper.Update(Value, 0, 0.34f);
+		}
+		UTEST_EQUAL("Dampen 1 with state passing", Damper.GetX0(), 0.f);
+	}
+
+	{
+		FCriticalDamper Damper(10.f);
+		float Value = -1.f;
+		for (int32 Frame = 0; Frame < 30; ++Frame)
+		{
+			Value = Damper.Update(Value, 0.34f);
+		}
+		UTEST_EQUAL("Dampen -1", Damper.GetX0(), 0.f);
+	}
+
+	{
+		FCriticalDamper Damper(10.f);
+		float Value = 5.f;
+		for (int32 Frame = 0; Frame < 40; ++Frame)
+		{
+			Value = Damper.Update(Value, 2.f, 0.34f);
+		}
+		UTEST_EQUAL("Dampen 5->2 (X0)", Damper.GetX0(), 0.f);
+		UTEST_EQUAL_TOLERANCE("Dampen 5->2", Value, 2.f, 0.00001f);
+	}
+
+	{
+		FCriticalDamper Damper1(10.f);
+		Damper1.Reset(1, 0);
+		Damper1.Update(0.34f);
+		Damper1.Update(0.34f);
+
+		FCriticalDamper Damper2(10.f);
+		Damper2.Reset(Damper1.GetX0(), Damper1.GetX0Derivative());
+
+		UTEST_EQUAL("Framerate equivalence", Damper2.GetX0(), Damper1.GetX0());
+		UTEST_EQUAL("Framerate equivalence", Damper2.GetX0Derivative(), Damper1.GetX0Derivative());
+		UTEST_EQUAL("Framerate equivalence", Damper2.GetW0(), Damper1.GetW0());
+		Damper1.Update(0.34f);
+		Damper1.Update(0.34f);
+		Damper2.Update(0.68f);
+		UTEST_EQUAL_TOLERANCE("Framerate equivalence", Damper2.GetX0(), Damper1.GetX0(), 0.001f);
+		UTEST_EQUAL_TOLERANCE("Framerate equivalence", Damper2.GetX0Derivative(), Damper1.GetX0Derivative(), 0.001f);
+	}
+
+	{
+		FCriticalDamper Damper1(10.f);
+		Damper1.Reset(0, 0);
+		float X1 = 5.f;
+		X1 = Damper1.Update(X1, 6, 0.34f);
+		X1 = Damper1.Update(X1, 7, 0.34f);
+		X1 = Damper1.Update(X1, 8, 0.34f);
+		X1 = Damper1.Update(X1, 9, 0.34f);
+		X1 = Damper1.Update(X1, 10, 0.34f);
+		X1 = Damper1.Update(X1, 11, 0.34f);
+
+		FCriticalDamper Damper2(10.f);
+		Damper2.Reset(0, 0);
+		float X2 = 5.f;
+		X2 = Damper2.Update(X2, 7, 0.68f);
+		X2 = Damper2.Update(X2, 9, 0.68f);
+		X2 = Damper2.Update(X2, 11, 0.68f);
+
+		UTEST_EQUAL_TOLERANCE("Framerate equivalence", X2, X1, 0.001f);
+		UTEST_EQUAL_TOLERANCE("Framerate equivalence", Damper2.GetX0(), Damper1.GetX0(), 0.001f);
+		UTEST_EQUAL_TOLERANCE("Framerate equivalence", Damper2.GetX0Derivative(), Damper1.GetX0Derivative(), 0.001f);
+	}
+
+	return true;
+}
+
+#undef LOCTEXT_NAMESPACE
+
diff --git a/Engine/Plugins/Cameras/GameplayCameras/Source/GameplayCameras/Public/Math/CriticalDamper.h b/Engine/Plugins/Cameras/GameplayCameras/Source/GameplayCameras/Public/Math/CriticalDamper.h
new file mode 100644
--- /dev/null
+++ b/Engine/Plugins/Cameras/GameplayCameras/Source/GameplayCameras/Public/Math/CriticalDamper.h
@@ -0,0 +1,112 @@
+// Copyright Epic Games, Inc. All Rights Reserved.
+
+#pragma once
+
+#include "CriticalDamper.generated.h"
+
+/**
+ * Implements a critically damped spring-mass system.
+ * See https://en.wikipedia.org/wiki/Damping
+ *
+ * Important notes:
+ *
+ * - A critically damped spring-mass system converges as fast as possible
+ *   towards its state of equilibrium without overshooting and/or oscillating.
+ *   That's the theory, though. In practice, probably because of floating 
+ *   point precision, the system can slightly overshoot, so we clamp the result
+ *   when it's just about to "settle".
+ *   See the `GameplayCameras.CriticalDamper.StabilizationThreshold` cvar for
+ *   tweaking how that happens.
+ *
+ * - In this implementation, the state of equilibrium is 0 (zero). That is, the
+ *   target that the system is trying to reach is 0. In most practical cases, 
+ *   however, we want to converge towards an arbitrary value, which can even
+ *   move if the system is being "dragged around". The different overloads of
+ *   the `Update` method are there for these different use-cases. Don't mix and
+ *   match calls, though: it's recommended to stick to one of these overloads.
+ *
+ * - See the math notes in the corresponding cpp file for details about the
+ *   algorithm and its implementation.
+ */
+USTRUCT()
+struct FCriticalDamper
+{
+	GENERATED_BODY()
+
+public:
+
+	FCriticalDamper();
+	FCriticalDamper(float InW0);
+
+	/**
+	 * Resets the initial conditions (aka the previous frame's state) of the
+	 * spring-mass system.
+	 */
+	void Reset(float InX0, float InX0Derivative);
+
+	/**
+	 * Gets the damping factor.
+	 * This is the undamped angular frequency of the oscillator system.
+	 * It's technically expressed in seconds.
+	 */
+	float GetW0() const { return W0; }
+	/**
+	 * Sets the damping factor.
+	 * This is the undamped angular frequency of the oscillator system.
+	 * It's technically expressed in seconds.
+	 */
+	void SetW0(float InW0) { W0 = InW0; }
+
+	/** Gets the last position of the spring-mass system. */
+	float GetX0() const { return X0; }
+	/** Gets the last velocity of the spring-mass system. */
+	float GetX0Derivative() const { return X0Derivative; }
+
+	/**
+	 * Updates the system to make it converge towards 0.
+	 *
+	 * @param DeltaTime  The elapsed time.
+	 */
+	float Update(float DeltaTime);
+	/**
+	 * Updates the system, given a possibly forced position the system was moved to.
+	 * The forced position should be last frame's result (also available with GetX0())
+	 * if the target wasn't moved.
+	 * The system will converge towards 0.
+	 *
+	 * @param X          The forced position of the system.
+	 * @param DeltaTime  The elapsed time.
+	 * @return           The new position of the system.
+	 */
+	float Update(float X, float DeltaTime);
+	/** 
+	 * Updates the system for damping towards arbitrary non-zero values.
+	 * This call takes last frame's target and this frame's target. The difference 
+	 * between the two indicates by how much the target was forcibly moved, if at all.
+	 * The result (also available with GetX0()) is returned relative to the current
+	 * frame's target (i.e. the second parameter).
+	 *
+	 * @param PreviousDamped  The previous value returned by this function.
+	 * @param NextTarget      The new target for the system.
+	 * @param DeltaTime       The elapsed time.
+	 * @return			      The new damped value.
+	 */
+	float Update(float PreviousDamped, float NextTarget, float DeltaTime);
+
+private:
+
+	void InternalUpdate(float ForcedMovement, float DeltaTime);
+
+private:
+
+	/** Angular frequency of the oscillator. */
+	UPROPERTY()
+	float W0;
+	/** Initial position of the spring-mass system, aka the previous frame's position. */
+	UPROPERTY()
+	float X0;
+	/** Initial velocity of the spring-mass system, aka the previous frame's velocity. */
+	UPROPERTY()
+	float X0Derivative;
+};
+
diff --git a/Engine/Plugins/Cameras/GameplayCameras/Source/GameplayCameras/Public/Nodes/Common/DampenPositionCameraNode.h b/Engine/Plugins/Cameras/GameplayCameras/Source/GameplayCameras/Public/Nodes/Common/DampenPositionCameraNode.h
new file mode 100644
--- /dev/null
+++ b/Engine/Plugins/Cameras/GameplayCameras/Source/GameplayCameras/Public/Nodes/Common/DampenPositionCameraNode.h
@@ -0,0 +1,41 @@
+// Copyright Epic Games, Inc. All Rights Reserved.
+
+#pragma once
+
+#include "Core/CameraNode.h"
+#include "Nodes/CameraNodeTypes.h"
+
+#include "DampenPositionCameraNode.generated.h"
+
+/**
+ * A camera node that offsets the location of the camera.
+ */
+UCLASS(MinimalAPI)
+class UDampenPositionCameraNode : public UCameraNode
+{
+	GENERATED_BODY()
+
+public:	
+
+	UDampenPositionCameraNode(const FObjectInitializer& ObjectInit);
+
+protected:
+
+	// UCameraNode interface.
+	virtual FCameraNodeEvaluatorPtr OnBuildEvaluator(FCameraNodeEvaluatorBuilder& Builder) const override;
+
+public:
+
+	/** Damping factor for forward/backward camera movement. */
+	UPROPERTY(EditAnywhere, Category=Damping)
+	float ForwardDampingFactor = 0.f;
+
+	/** Damping factor for left/right camera movement. */
+	UPROPERTY(EditAnywhere, Category=Damping)
+	float LateralDampingFactor = 0.f;
+
+	/** Damping factor for up/down camera movement. */
+	UPROPERTY(EditAnywhere, Category=Damping)
+	float VerticalDampingFactor = 0.f;
+};
+
