Cameras: add camera variable table to node evaluator result

diff --git a/Engine/Plugins/Cameras/GameplayCameras/Source/GameplayCameras/Private/Core/CameraVariableTable.cpp b/Engine/Plugins/Cameras/GameplayCameras/Source/GameplayCameras/Private/Core/CameraVariableTable.cpp
new file mode 100644
--- /dev/null
+++ b/Engine/Plugins/Cameras/GameplayCameras/Source/GameplayCameras/Private/Core/CameraVariableTable.cpp
@@ -0,0 +1,391 @@
+// Copyright Epic Games, Inc. All Rights Reserved.
+
+#include "Core/CameraVariableTable.h"
+
+#include "HAL/UnrealMemory.h"
+#include "Math/UnrealMath.h"
+
+namespace UE::Cameras::Private
+{
+
+static const FString GUnavailableVariableDebugName(TEXT("<no debug info>"));
+
+bool IsVariableInMask(uint32 VariableId, const FCameraVariableTableFlags* InMask, bool bInvertMask)
+{
+	if (InMask)
+	{
+		const bool bIsInMask = InMask->VariableIds.Contains(VariableId);
+		return bInvertMask ? !bIsInMask : bIsInMask;
+	}
+	return true;
+}
+
+}  // namespace UE::Cameras::Private
+
+FCameraVariableTable::FCameraVariableTable()
+{
+}
+
+FCameraVariableTable::FCameraVariableTable(FCameraVariableTable&& Other)
+	: Entries(MoveTemp(Other.Entries))
+	, Memory(Other.Memory)
+	, Capacity(Other.Capacity)
+	, Used(Other.Used)
+{
+	Other.Memory = nullptr;
+	Other.Capacity = 0;
+	Other.Used = 0;
+}
+
+FCameraVariableTable& FCameraVariableTable::operator=(FCameraVariableTable&& Other)
+{
+	Entries = MoveTemp(Other.Entries);
+	Memory = Other.Memory;
+	Capacity = Other.Capacity;
+	Used = Other.Used;
+
+	Other.Memory = nullptr;
+	Other.Capacity = 0;
+	Other.Used = 0;
+
+	return *this;
+}
+
+FCameraVariableTable::~FCameraVariableTable()
+{
+	if (Memory)
+	{
+		FMemory::Free(Memory);
+		Memory = nullptr;
+		Capacity = Used = 0;
+	}
+}
+
+void FCameraVariableTable::Initialize(const FCameraVariableTableAllocationInfo& AllocationInfo)
+{
+	// Reset any previous state.
+	Entries.Reset();
+	if (Memory)
+	{
+		FMemory::Free(Memory);
+		Memory = nullptr;
+	}
+
+	// Compute the total buffer size we need, and create our entries as we go.
+	uint32 TotalSizeOf = 0;
+	uint32 MaxAlignOf = 0;
+	uint32 CurSizeOf, CurAlignOf;
+	for (const FCameraVariableDefinition& VariableDefinition : AllocationInfo.VariableDefinitions)
+	{
+		GetVariableTypeAllocationInfo(VariableDefinition.VariableType, CurSizeOf, CurAlignOf);
+		const uint32 NewEntryOffset = Align(TotalSizeOf, CurAlignOf);
+		TotalSizeOf = NewEntryOffset + CurSizeOf;
+		MaxAlignOf = FMath::Max(MaxAlignOf, CurAlignOf);
+
+		FEntry NewEntry;
+		NewEntry.Id = VariableDefinition.VariableId;
+		NewEntry.Type = VariableDefinition.VariableType;
+		NewEntry.Offset = NewEntryOffset;
+#if WITH_EDITORONLY_DATA
+		NewEntry.DebugName = VariableDefinition.VariableName;
+#endif
+		Entries.Add(NewEntry.Id, NewEntry);
+	}
+
+	// Allocate the memory buffer.
+	Memory = reinterpret_cast<uint8*>(FMemory::Malloc(TotalSizeOf, MaxAlignOf));
+	Capacity = TotalSizeOf;
+	Used = 0;
+
+	// Go back to our entries and initialize each entry to the default value for that variable type.
+	for (auto& Pair : Entries)
+	{
+		const FEntry& Entry = Pair.Value;
+		uint8* ValuePtr = Memory + Entry.Offset;
+		switch (Entry.Type)
+		{
+#define UE_CAMERA_VARIABLE_FOR_TYPE(ValueType, ValueName)\
+			case ECameraVariableType::ValueName:\
+				{\
+					ValueType* TypedValuePtr = reinterpret_cast<ValueType*>(ValuePtr);\
+					*TypedValuePtr = ValueType();\
+				}\
+				break;
+UE_CAMERA_VARIABLE_FOR_ALL_TYPES()
+#undef UE_CAMERA_VARIABLE_FOR_TYPE
+		}
+	}
+}
+
+void FCameraVariableTable::AddVariable(const FCameraVariableDefinition& VariableDefinition)
+{
+	uint32 SizeOf, AlignOf;
+	GetVariableTypeAllocationInfo(VariableDefinition.VariableType, SizeOf, AlignOf);
+
+	uint8* VariablePtr = Align(Memory + Used, AlignOf);
+	uint32 NewUsed = (VariablePtr + SizeOf) - Memory;
+
+	if (NewUsed > Capacity)
+	{
+		ReallocateBuffer(Capacity * 2);
+
+		VariablePtr = Align(Memory + Used, AlignOf);
+	}
+
+	Used = NewUsed;
+
+	FEntry NewEntry;
+	NewEntry.Id = VariableDefinition.VariableId;
+	NewEntry.Type = VariableDefinition.VariableType;
+	NewEntry.Offset = VariablePtr - Memory;
+	NewEntry.Flags = VariableDefinition.bIsPrivate ? EEntryFlags::Private : EEntryFlags::None;
+
+	Entries.Add(VariableDefinition.VariableId, NewEntry);
+}
+
+void FCameraVariableTable::ReallocateBuffer(uint32 NewCapacity)
+{
+	static const uint32 DefaultAlignment = 32;
+
+	uint8* NewMemory = reinterpret_cast<uint8*>(FMemory::Malloc(NewCapacity, DefaultAlignment));
+	FMemory::Memmove(NewMemory, Memory, Capacity);
+	Memory = NewMemory;
+	Capacity = NewCapacity;
+}
+
+bool FCameraVariableTable::GetVariableTypeAllocationInfo(ECameraVariableType VariableType, uint32& OutSizeOf, uint32& OutAlignOf)
+{
+	switch (VariableType)
+	{
+#define UE_CAMERA_VARIABLE_FOR_TYPE(ValueType, ValueName)\
+		case ECameraVariableType::ValueName:\
+			OutSizeOf = sizeof(ValueType);\
+			OutAlignOf = sizeof(ValueType);\
+			return true;
+UE_CAMERA_VARIABLE_FOR_ALL_TYPES()
+#undef UE_CAMERA_VARIABLE_FOR_TYPE
+	}
+	return false;
+}
+
+bool FCameraVariableTable::ContainsValue(uint32 VariableId) const
+{
+	return Entries.Contains(VariableId);
+}
+
+bool FCameraVariableTable::IsValueWritten(uint32 VariableId) const
+{
+	if (const FEntry* Entry = Entries.Find(VariableId))
+	{
+		return EnumHasAnyFlags(Entry->Flags, EEntryFlags::Written);
+	}
+	return false;
+}
+
+void FCameraVariableTable::UnsetValue(uint32 VariableId)
+{
+	if (FEntry* Entry = Entries.Find(VariableId))
+	{
+		EnumRemoveFlags(Entry->Flags, EEntryFlags::Written | EEntryFlags::WrittenThisFrame);
+	}
+}
+
+void FCameraVariableTable::UnsetAllValues()
+{
+	for (auto& Pair : Entries)
+	{
+		EnumRemoveFlags(Pair.Value.Flags, EEntryFlags::Written | EEntryFlags::WrittenThisFrame);
+	}
+}
+
+bool FCameraVariableTable::IsValueWrittenThisFrame(uint32 VariableId) const
+{
+	if (const FEntry* Entry = Entries.Find(VariableId))
+	{
+		return EnumHasAnyFlags(Entry->Flags, EEntryFlags::WrittenThisFrame);
+	}
+	return false;
+}
+
+void FCameraVariableTable::ResetAllWrittenThisFrame()
+{
+	for (auto& Pair : Entries)
+	{
+		EnumRemoveFlags(Pair.Value.Flags, EEntryFlags::WrittenThisFrame);
+	}
+}
+
+void FCameraVariableTable::OverrideChanged(const FCameraVariableTable& OtherTable)
+{
+	InternalOverrideChanged(OtherTable, nullptr, false, nullptr);
+}
+
+void FCameraVariableTable::OverrideChanged(const FCameraVariableTable& OtherTable, const FCameraVariableTableFlags& InMask, bool bInvertMask, FCameraVariableTableFlags& OutMask)
+{
+	InternalOverrideChanged(OtherTable, &InMask, bInvertMask, &OutMask);
+}
+
+void FCameraVariableTable::InternalOverrideChanged(const FCameraVariableTable& OtherTable, const FCameraVariableTableFlags* InMask, bool bInvertMask, FCameraVariableTableFlags* OutMask)
+{
+	using namespace UE::Cameras::Private;
+
+	for (auto& OtherPair : OtherTable.Entries)
+	{
+		// Look for entries in the other table that have been written to, and aren't private.
+		const FEntry& OtherEntry  = OtherPair.Value;
+		const EEntryFlags OtherFlags = OtherEntry.Flags;
+		if (EnumHasAnyFlags(OtherFlags, EEntryFlags::Written)
+				&& !EnumHasAnyFlags(OtherFlags, EEntryFlags::Private)
+				&& IsVariableInMask(OtherEntry.Id, InMask, bInvertMask))
+		{
+			// See if we know this variable.
+			FEntry* ThisEntry = Entries.Find(OtherEntry.Id);
+			if (ThisEntry)
+			{
+				// We already have the other table's variable in our table. Let's check
+				// that the types match, and then copy the memory.
+#if WITH_EDITORONLY_DATA
+				checkf(ThisEntry->DebugName == OtherEntry.DebugName,
+						TEXT("Camera variable name collision! Expected variable '%d' to be named '%s', but other table has '%s'!"),
+						ThisEntry->Id, *ThisEntry->DebugName, *OtherEntry.DebugName);
+#endif
+
+#if WITH_EDITORONLY_DATA
+				const FString& DebugName = ThisEntry->DebugName;
+#else
+				const FString& DebugName = GUnavailableVariableDebugName;
+#endif
+				checkf(ThisEntry->Type == OtherEntry.Type, 
+						TEXT("Camera variable name collision! Expected '%d' (%s) to be of type '%s' but other table has type '%s'!"),
+						ThisEntry->Id, *DebugName,
+						*UEnum::GetValueAsString(ThisEntry->Type), *UEnum::GetValueAsString(OtherEntry.Type));
+			}
+			else
+			{
+				// We don't have this variable in our table. Let's add it.
+				FCameraVariableDefinition NewVariableDefinition;
+				NewVariableDefinition.VariableId = OtherEntry.Id;
+				NewVariableDefinition.VariableType = OtherEntry.Type;
+#if WITH_EDITORONLY_DATA
+				NewVariableDefinition.VariableName = OtherEntry.DebugName;
+#endif
+				AddVariable(NewVariableDefinition);
+
+				ThisEntry = Entries.Find(OtherEntry.Id);
+			}
+
+			if (ensure(ThisEntry))
+			{
+				uint32 ValueSize = 0, ValueAlignment = 0;
+				GetVariableTypeAllocationInfo(ThisEntry->Type, ValueSize, ValueAlignment);
+				check(ValueSize != 0);
+
+				uint8* ThisValuePtr = Memory + ThisEntry->Offset;
+				const uint8* OtherValuePtr = OtherTable.Memory + OtherEntry.Offset;
+				FMemory::Memcpy(ThisValuePtr, OtherValuePtr, ValueSize);
+				EnumAddFlags(ThisEntry->Flags, EEntryFlags::Written);
+
+				if (OutMask)
+				{
+					OutMask->VariableIds.Add(ThisEntry->Id);
+				}
+			}
+		}
+	}
+}
+
+void FCameraVariableTable::LerpChanged(const FCameraVariableTable& ToTable, float Factor)
+{
+	InternalLerpChanged(ToTable, Factor, nullptr, false, nullptr);
+}
+
+void FCameraVariableTable::LerpChanged(const FCameraVariableTable& ToTable, float Factor, const FCameraVariableTableFlags& InMask, bool bInvertMask, FCameraVariableTableFlags& OutMask)
+{
+	InternalLerpChanged(ToTable, Factor, &InMask, false, &OutMask);
+}
+
+void FCameraVariableTable::InternalLerpChanged(const FCameraVariableTable& ToTable, float Factor, const FCameraVariableTableFlags* InMask, bool bInvertMask, FCameraVariableTableFlags* OutMask)
+{
+	using namespace UE::Cameras::Private;
+
+	for (auto& ToPair : ToTable.Entries)
+	{
+		// Look for entries in the other table that have been written to, and aren't private.
+		const FEntry& ToEntry  = ToPair.Value;
+		const EEntryFlags ToFlags = ToEntry.Flags;
+		if (EnumHasAnyFlags(ToFlags, EEntryFlags::Written)
+				&& !EnumHasAnyFlags(ToFlags, EEntryFlags::Private)
+				&& IsVariableInMask(ToEntry.Id, InMask, bInvertMask))
+		{
+			// See if we know this variable.
+			FEntry* FromEntry = Entries.Find(ToEntry.Id);
+			if (FromEntry)
+			{
+				// We already have the other table's variable in our table. Let's check
+				// that the types match, and then interpolate the values.
+#if WITH_EDITORONLY_DATA
+				checkf(FromEntry->DebugName == ToEntry.DebugName,
+						TEXT("Camera variable name collision! Expected variable '%d' to be named '%s', but other table has '%s'!"),
+						FromEntry->Id, *FromEntry->DebugName, *ToEntry.DebugName);
+#endif
+
+#if WITH_EDITORONLY_DATA
+				const FString& DebugName = FromEntry->DebugName;
+#else
+				const FString& DebugName = GUnavailableVariableDebugName;
+#endif
+				checkf(FromEntry->Type == ToEntry.Type, 
+						TEXT("Camera variable name collision! Expected '%d' (%s) to be of type '%s' but other table has type '%s'!"),
+						FromEntry->Id, *DebugName,
+						*UEnum::GetValueAsString(FromEntry->Type), *UEnum::GetValueAsString(ToEntry.Type));
+
+				uint8* FromValuePtr = Memory + FromEntry->Offset;
+				const uint8* ToValuePtr = ToTable.Memory + ToEntry.Offset;
+				switch (FromEntry->Type)
+				{
+#define UE_CAMERA_VARIABLE_FOR_TYPE(ValueType, ValueName)\
+					case ECameraVariableType::ValueName:\
+						{\
+							ValueType& FromValue = *reinterpret_cast<ValueType*>(FromValuePtr);\
+							const ValueType& ToValue = *reinterpret_cast<const ValueType*>(ToValuePtr);\
+							ValueType InterpValue = TCameraVariableInterpolation<ValueType>::Interpolate(*FromEntry, FromValue, ToValue, Factor);\
+							FromValue = InterpValue;\
+						}\
+						break;
+UE_CAMERA_VARIABLE_FOR_ALL_TYPES()
+#undef UE_CAMERA_VARIABLE_FOR_TYPE
+				}
+				EnumAddFlags(FromEntry->Flags, EEntryFlags::Written);
+			}
+			else
+			{
+				// We don't have this variable in our table. Let's add it.
+				FCameraVariableDefinition NewVariableDefinition;
+				NewVariableDefinition.VariableId = ToEntry.Id;
+				NewVariableDefinition.VariableType = ToEntry.Type;
+#if WITH_EDITORONLY_DATA
+				NewVariableDefinition.VariableName = ToEntry.DebugName;
+#endif
+				AddVariable(NewVariableDefinition);
+
+				FromEntry = Entries.Find(ToEntry.Id);
+				check(FromEntry);
+
+				uint32 ValueSize = 0, ValueAlignment = 0;
+				GetVariableTypeAllocationInfo(FromEntry->Type, ValueSize, ValueAlignment);
+				check(ValueSize != 0);
+
+				uint8* FromValuePtr = Memory + FromEntry->Offset;
+				const uint8* ToValuePtr = ToTable.Memory + ToEntry.Offset;
+				FMemory::Memcpy(FromValuePtr, ToValuePtr, ValueSize);
+				EnumAddFlags(FromEntry->Flags, EEntryFlags::Written);
+			}
+
+			if (OutMask)
+			{
+				OutMask->VariableIds.Add(FromEntry->Id);
+			}
+		}
+	}
+}
+
diff --git a/Engine/Plugins/Cameras/GameplayCameras/Source/GameplayCameras/Private/Nodes/Blends/SimpleBlendCameraNode.cpp b/Engine/Plugins/Cameras/GameplayCameras/Source/GameplayCameras/Private/Nodes/Blends/SimpleBlendCameraNode.cpp
--- a/Engine/Plugins/Cameras/GameplayCameras/Source/GameplayCameras/Private/Nodes/Blends/SimpleBlendCameraNode.cpp
+++ b/Engine/Plugins/Cameras/GameplayCameras/Source/GameplayCameras/Private/Nodes/Blends/SimpleBlendCameraNode.cpp
@@ -25,6 +25,8 @@ void FSimpleBlendCameraNodeEvaluator::OnBlendResults(const FCameraNodeBlendParam
 			false,
 			BlendedResult.CameraPose.GetChangedFlags());
 
+	BlendedResult.VariableTable.LerpChanged(ChildResult.VariableTable, BlendFactor);
+
 	// If we have even a fraction of a camera cut, we need to make the
 	// whole result into a camera cut.
 	if (BlendFactor > 0.f && ChildResult.bIsCameraCut)
diff --git a/Engine/Plugins/Cameras/GameplayCameras/Source/GameplayCameras/Public/Core/CameraNodeEvaluator.h b/Engine/Plugins/Cameras/GameplayCameras/Source/GameplayCameras/Public/Core/CameraNodeEvaluator.h
--- a/Engine/Plugins/Cameras/GameplayCameras/Source/GameplayCameras/Public/Core/CameraNodeEvaluator.h
+++ b/Engine/Plugins/Cameras/GameplayCameras/Source/GameplayCameras/Public/Core/CameraNodeEvaluator.h
@@ -4,6 +4,7 @@
 
 #include "Core/CameraObjectRtti.h"
 #include "Core/CameraPose.h"
+#include "Core/CameraVariableTable.h"
 #include "Core/ObjectChildrenView.h"
 #include "CoreTypes.h"
 #include "UObject/ObjectPtr.h"
@@ -37,6 +38,10 @@ struct FCameraNodeEvaluationResult
 {
 	/** The camera pose. */
 	FCameraPose CameraPose;
+
+	/** The variable table. */
+	FCameraVariableTable VariableTable;
+
 	/** Whether the current frame is a camera cut. */
 	bool bIsCameraCut = false;
 	/** Whether this result is valid. */
diff --git a/Engine/Plugins/Cameras/GameplayCameras/Source/GameplayCameras/Public/Core/CameraVariableTable.h b/Engine/Plugins/Cameras/GameplayCameras/Source/GameplayCameras/Public/Core/CameraVariableTable.h
new file mode 100644
--- /dev/null
+++ b/Engine/Plugins/Cameras/GameplayCameras/Source/GameplayCameras/Public/Core/CameraVariableTable.h
@@ -0,0 +1,256 @@
+// Copyright Epic Games, Inc. All Rights Reserved.
+
+#pragma once
+
+#include "Containers/Map.h"
+#include "Containers/Set.h"
+#include "Core/CameraVariableTableFwd.h"
+#include "CoreTypes.h"
+#include "Misc/EnumClassFlags.h"
+#include "Templates/UnrealTypeTraits.h"
+#include "UObject/NameTypes.h"
+
+template<typename ValueType>
+struct TCameraVariableTraits;
+
+template<typename ValueType>
+struct TCameraVariableInterpolation;
+
+/**
+ * A structure that keeps track of which variables have been processed in a
+ * camera variable table.
+ */
+struct FCameraVariableTableFlags
+{
+	/** The list of processed variable IDs. */
+	TSet<uint32> VariableIds;
+};
+
+/**
+ * The camera variable table is a container for a collection of arbitrary values
+ * of various types. Only certain basic types are supported (most primitive types).
+ * 
+ * This table serves both as an implementation of the usual "blackboard" design, where
+ * gameplay systems can push any appropriate values into the camera system, and as a
+ * place for camera node evaluators to stash various things.
+ *
+ * The main function of the variable table is that it is blended along with the camera
+ * rig it belongs to. Any matching values between to blended tables with be themselves
+ * blended, except for values flagged as "private".
+ *
+ * Internally, the variable table is allocated as one continuous block of memory, plus
+ * a map of metadata keyed by variable ID. A variable ID can be anything, but will
+ * generally be the hash of the variable name.
+ */
+class FCameraVariableTable
+{
+public:
+
+	FCameraVariableTable();
+	FCameraVariableTable(FCameraVariableTable&& Other);
+	FCameraVariableTable& operator=(FCameraVariableTable&& Other);
+	~FCameraVariableTable();
+
+	FCameraVariableTable(const FCameraVariableTable&) = delete;
+	FCameraVariableTable& operator=(const FCameraVariableTable&) = delete;
+
+	/** Initializes the variable table so that it fits the provided allocation info. */
+	void Initialize(const FCameraVariableTableAllocationInfo& AllocationInfo);
+
+	/** Adds a variable to the table.
+	 *
+	 * This may re-allocate the internal memory buffer. It's recommended to pre-compute
+	 * the allocation information needed for a table, and initialize it once. */
+	void AddVariable(const FCameraVariableDefinition& VariableDefinition);
+
+public:
+
+	// Getter methods.
+
+	template<typename ValueType>
+	const ValueType* FindValue(uint32 VariableId) const;
+
+	template<typename ValueType>
+	const ValueType& GetValue(uint32 VariableId) const;
+
+	template<typename ValueType>
+	ValueType GetValue(uint32 VariableId, typename TCallTraits<ValueType>::ParamType DefaultValue) const;
+
+	template<typename ValueType>
+	bool TryGetValue(uint32 VariableId, ValueType& OutValue) const;
+
+	bool ContainsValue(uint32 VariableId) const;
+
+public:
+
+	// Setter methods.
+	
+	template<typename ValueType>
+	void SetValue(uint32 VariableId, typename TCallTraits<ValueType>::ParamType Value);
+
+	template<typename ValueType>
+	bool TrySetValue(uint32 VariableId, typename TCallTraits<ValueType>::ParamType Value);
+
+public:
+
+	// Interpolation.
+	
+	void OverrideChanged(const FCameraVariableTable& OtherTable);
+	void OverrideChanged(const FCameraVariableTable& OtherTable, const FCameraVariableTableFlags& InMask, bool bInvertMask, FCameraVariableTableFlags& OutMask);
+
+	void LerpChanged(const FCameraVariableTable& ToTable, float Factor);
+	void LerpChanged(const FCameraVariableTable& ToTable, float Factor, const FCameraVariableTableFlags& InMask, bool bInvertMask, FCameraVariableTableFlags& OutMask);
+
+public:
+
+	// Lower level API.
+
+	bool IsValueWritten(uint32 VariableId) const;
+	void UnsetValue(uint32 VariableId);
+	void UnsetAllValues();
+
+	bool IsValueWrittenThisFrame(uint32 VariableId) const;
+	void ResetAllWrittenThisFrame();
+
+private:
+
+	static bool GetVariableTypeAllocationInfo(ECameraVariableType VariableType, uint32& OutSizeOf, uint32& OutAlignOf);
+
+	void ReallocateBuffer(uint32 NewCapacity);
+
+	template<typename ValueType>
+	bool CheckVariableType(ECameraVariableType InType)
+	{
+		ensure(TCameraVariableTraits<ValueType>::Type == InType);
+	}
+
+	void InternalOverrideChanged(const FCameraVariableTable& OtherTable, const FCameraVariableTableFlags* InMask, bool bInvertMask, FCameraVariableTableFlags* OutMask);
+	void InternalLerpChanged(const FCameraVariableTable& ToTable, float Factor, const FCameraVariableTableFlags* InMask, bool bInvertMask, FCameraVariableTableFlags* OutMask);
+
+private:
+
+	enum class EEntryFlags : uint8
+	{
+		None = 0,
+		Private = 1 << 0,
+		Written = 1 << 1,
+		WrittenThisFrame = 1 << 2
+	};
+	FRIEND_ENUM_CLASS_FLAGS(EEntryFlags)
+
+	struct FEntry
+	{
+		uint32 Id;
+		ECameraVariableType Type;
+		uint32 Offset;
+		mutable EEntryFlags Flags;
+#if WITH_EDITORONLY_DATA
+		FString DebugName;
+#endif
+	};
+
+	TMap<uint32, FEntry> Entries;
+
+	uint8* Memory = nullptr;
+	uint32 Capacity = 0;
+	uint32 Used = 0;
+
+	template<typename T>
+	friend struct TCameraVariableInterpolation;
+};
+
+ENUM_CLASS_FLAGS(FCameraVariableTable::EEntryFlags)
+
+template<typename ValueType>
+const ValueType* FCameraVariableTable::FindValue(uint32 VariableId) const
+{
+	if (const FEntry* Entry = Entries.Find(VariableId))
+	{
+		CheckVariableType<ValueType>(Entry->Type);
+		return reinterpret_cast<ValueType*>(Memory + Entry->Offset);
+	}
+	return nullptr;
+}
+
+template<typename ValueType>
+const ValueType& FCameraVariableTable::GetValue(uint32 VariableId) const
+{
+	const FEntry& Entry = Entries.FindChecked(VariableId);
+	CheckVariableType<ValueType>(Entry.Type);
+	return reinterpret_cast<ValueType*>(Memory + Entry.Offset);
+}
+
+template<typename ValueType>
+ValueType FCameraVariableTable::GetValue(uint32 VariableId, typename TCallTraits<ValueType>::ParamType DefaultValue) const
+{
+	if (const ValueType* Value = FindValue<ValueType>(VariableId))
+	{
+		return *Value;
+	}
+	return DefaultValue;
+}
+
+template<typename ValueType>
+bool FCameraVariableTable::TryGetValue(uint32 VariableId, ValueType& OutValue) const
+{
+	if (const ValueType* Value = FindValue<ValueType>(VariableId))
+	{
+		OutValue = *Value;
+		return true;
+	}
+	return false;
+}
+
+template<typename ValueType>
+void FCameraVariableTable::SetValue(uint32 VariableId, typename TCallTraits<ValueType>::ParamType Value)
+{
+	const FEntry& Entry = Entries.FindChecked(VariableId);
+	CheckVariableType<ValueType>(Entry.Type);
+	ValueType* ValuePtr = reinterpret_cast<ValueType*>(Memory + Entry.Offset);
+	*ValuePtr = Value;
+}
+
+template<typename ValueType>
+bool FCameraVariableTable::TrySetValue(uint32 VariableId, typename TCallTraits<ValueType>::ParamType Value)
+{
+	if (const FEntry* Entry = Entries.Find(VariableId))
+	{
+		CheckVariableType<ValueType>(Entry->Type);
+		ValueType* ValuePtr = reinterpret_cast<ValueType*>(Memory + Entry->Offset);
+		return true;
+	}
+	return false;
+}
+
+#define UE_CAMERA_VARIABLE_FOR_TYPE(ValueType, ValueName)\
+	template<>\
+	struct TCameraVariableTraits<ValueType>\
+	{\
+		ECameraVariableType Type = ECameraVariableType::ValueName;\
+	};
+UE_CAMERA_VARIABLE_FOR_ALL_TYPES()
+#undef UE_CAMERA_VARIABLE_FOR_TYPE
+
+template<typename ValueType>
+struct TCameraVariableInterpolation
+{
+	using ValueParam = typename TCallTraits<ValueType>::ParamType;
+	static ValueType Interpolate(const FCameraVariableTable::FEntry& TableEntry, ValueParam From, ValueParam To, float Factor)
+	{
+		return FMath::LerpStable(From, To, Factor);
+	}
+};
+
+template<typename T>
+struct TCameraVariableInterpolation<UE::Math::TTransform<T>>
+{
+	using ValueType = UE::Math::TTransform<T>;
+	using ValueParam = typename TCallTraits<ValueType>::ParamType;
+	static ValueType Interpolate(const FCameraVariableTable::FEntry& TableEntry, ValueParam From, ValueParam To, float Factor)
+	{
+		ValueType Result(From);
+		Result.BlendWith(To, Factor);
+		return Result;
+	}
+};
+
diff --git a/Engine/Plugins/Cameras/GameplayCameras/Source/GameplayCameras/Public/Core/CameraVariableTableFwd.h b/Engine/Plugins/Cameras/GameplayCameras/Source/GameplayCameras/Public/Core/CameraVariableTableFwd.h
new file mode 100644
--- /dev/null
+++ b/Engine/Plugins/Cameras/GameplayCameras/Source/GameplayCameras/Public/Core/CameraVariableTableFwd.h
@@ -0,0 +1,73 @@
+// Copyright Epic Games, Inc. All Rights Reserved.
+
+#pragma once
+
+#include "Containers/UnrealString.h"
+#include "CoreTypes.h"
+
+#include "CameraVariableTableFwd.generated.h"
+
+#define UE_CAMERA_VARIABLE_FOR_ALL_TYPES()\
+	UE_CAMERA_VARIABLE_FOR_TYPE(bool, Boolean)\
+	UE_CAMERA_VARIABLE_FOR_TYPE(int32, Integer32)\
+	UE_CAMERA_VARIABLE_FOR_TYPE(float, Float)\
+	UE_CAMERA_VARIABLE_FOR_TYPE(double, Double)\
+	UE_CAMERA_VARIABLE_FOR_TYPE(FVector2f, Vector2f)\
+	UE_CAMERA_VARIABLE_FOR_TYPE(FVector2d, Vector2d)\
+	UE_CAMERA_VARIABLE_FOR_TYPE(FVector3f, Vector3f)\
+	UE_CAMERA_VARIABLE_FOR_TYPE(FVector3d, Vector3d)\
+	UE_CAMERA_VARIABLE_FOR_TYPE(FVector4f, Vector4f)\
+	UE_CAMERA_VARIABLE_FOR_TYPE(FVector4d, Vector4d)\
+	UE_CAMERA_VARIABLE_FOR_TYPE(FRotator3f, Rotator3f)\
+	UE_CAMERA_VARIABLE_FOR_TYPE(FRotator3d, Rotator3d)\
+	UE_CAMERA_VARIABLE_FOR_TYPE(FTransform3f, Transform3f)\
+	UE_CAMERA_VARIABLE_FOR_TYPE(FTransform3d, Transform3d)
+
+UENUM()
+enum class ECameraVariableType
+{
+	Boolean,
+	Integer32,
+	Float,
+	Double,
+	Vector2f,
+	Vector2d,
+	Vector3f,
+	Vector3d,
+	Vector4f,
+	Vector4d,
+	Rotator3f,
+	Rotator3d,
+	Transform3f,
+	Transform3d
+};
+
+USTRUCT()
+struct FCameraVariableDefinition
+{
+	GENERATED_BODY()
+
+	UPROPERTY()
+	uint32 VariableId = 0;
+
+	UPROPERTY()
+	ECameraVariableType VariableType;
+
+	UPROPERTY()
+	bool bIsPrivate = false;
+
+#if WITH_EDITORONLY_DATA
+	UPROPERTY()
+	FString VariableName;
+#endif
+};
+
+USTRUCT()
+struct FCameraVariableTableAllocationInfo
+{
+	GENERATED_BODY()
+
+	UPROPERTY()
+	TArray<FCameraVariableDefinition> VariableDefinitions;
+};
+
