Cameras: camera node allocation info now contains both evaluator and variable table info

diff --git a/Engine/Plugins/Cameras/GameplayCameras/Source/GameplayCameras/Private/Core/CameraNode.cpp b/Engine/Plugins/Cameras/GameplayCameras/Source/GameplayCameras/Private/Core/CameraNode.cpp
--- a/Engine/Plugins/Cameras/GameplayCameras/Source/GameplayCameras/Private/Core/CameraNode.cpp
+++ b/Engine/Plugins/Cameras/GameplayCameras/Source/GameplayCameras/Private/Core/CameraNode.cpp
@@ -13,9 +13,9 @@ void FCameraNodeEvaluationResult::Reset()
 	bIsValid = false;
 }
 
-FCameraNodeEvaluatorAllocationInfo UCameraNode::GetEvaluatorInfo() const
+FCameraNodeAllocationInfo UCameraNode::GetAllocationInfo() const
 {
-	return OnGetEvaluatorInfo();
+	return OnGetAllocationInfo();
 }
 
 FCameraNodeEvaluatorPtr UCameraNode::BuildEvaluator(FCameraNodeEvaluatorBuilder& Builder) const
diff --git a/Engine/Plugins/Cameras/GameplayCameras/Source/GameplayCameras/Private/Core/CameraNodeEvaluatorStorage.cpp b/Engine/Plugins/Cameras/GameplayCameras/Source/GameplayCameras/Private/Core/CameraNodeEvaluatorStorage.cpp
--- a/Engine/Plugins/Cameras/GameplayCameras/Source/GameplayCameras/Private/Core/CameraNodeEvaluatorStorage.cpp
+++ b/Engine/Plugins/Cameras/GameplayCameras/Source/GameplayCameras/Private/Core/CameraNodeEvaluatorStorage.cpp
@@ -44,10 +44,10 @@ FCameraNodeEvaluatorTreeAllocationInfo FCameraNodeEvaluatorStorage::ComputeTreeI
 	{
 		if (UCameraNode* CurNode = Cast<UCameraNode>(Object))
 		{
-			FCameraNodeEvaluatorAllocationInfo CurInfo = CurNode->GetEvaluatorInfo();
+			FCameraNodeAllocationInfo CurInfo = CurNode->GetAllocationInfo();
 
-			MaxAlignment = FMath::Max(MaxAlignment, CurInfo.Alignof);
-			MaxSize = Align(MaxSize, CurInfo.Alignof) + CurInfo.Sizeof;
+			MaxAlignment = FMath::Max(MaxAlignment, CurInfo.EvaluatorInfo.Alignof);
+			MaxSize = Align(MaxSize, CurInfo.EvaluatorInfo.Alignof) + CurInfo.EvaluatorInfo.Sizeof;
 		}
 		return true; // Keep iterating.
 	});
diff --git a/Engine/Plugins/Cameras/GameplayCameras/Source/GameplayCameras/Public/Core/CameraNode.h b/Engine/Plugins/Cameras/GameplayCameras/Source/GameplayCameras/Public/Core/CameraNode.h
--- a/Engine/Plugins/Cameras/GameplayCameras/Source/GameplayCameras/Public/Core/CameraNode.h
+++ b/Engine/Plugins/Cameras/GameplayCameras/Source/GameplayCameras/Public/Core/CameraNode.h
@@ -3,11 +3,27 @@
 #pragma once
 
 #include "Core/CameraNodeEvaluatorBuilder.h"
+#include "Core/CameraVariableTableFwd.h"
 #include "CoreTypes.h"
 #include "UObject/Object.h"
 
 #include "CameraNode.generated.h"
 
+/**
+ * Structure describing various allocations needed by a camera node.
+ */
+USTRUCT()
+struct FCameraNodeAllocationInfo
+{
+	GENERATED_BODY()
+
+	UPROPERTY()
+	FCameraNodeEvaluatorAllocationInfo EvaluatorInfo;
+
+	UPROPERTY()
+	FCameraVariableTableAllocationInfo VariableTableInfo;
+};
+
 /**
  * The base class for a camera node.
  */
@@ -18,15 +34,17 @@ class UCameraNode : public UObject
 
 public:
 
-	/** Gets optional info about this node's evaluator. */
-	FCameraNodeEvaluatorAllocationInfo GetEvaluatorInfo() const;
+	/** Gets optional info about this node's required allocations at runtime. */
+	FCameraNodeAllocationInfo GetAllocationInfo() const;
+
 	/** Builds the evaluator for this node. */
 	FCameraNodeEvaluatorPtr BuildEvaluator(FCameraNodeEvaluatorBuilder& Builder) const;
 
 protected:
 
-	/** Gets optional info about this node's evaluator. */
-	virtual FCameraNodeEvaluatorAllocationInfo OnGetEvaluatorInfo() const { return FCameraNodeEvaluatorAllocationInfo(); }
+	/** Gets optional info about this node's required allocations at runtime. */
+	virtual FCameraNodeAllocationInfo OnGetAllocationInfo() const { return FCameraNodeAllocationInfo(); }
+
 	/** Builds the evaluator for this node. */
 	virtual FCameraNodeEvaluatorPtr OnBuildEvaluator(FCameraNodeEvaluatorBuilder& Builder) const { return nullptr; }
 
diff --git a/Engine/Plugins/Cameras/GameplayCameras/Source/GameplayCameras/Public/Core/CameraNodeEvaluatorFwd.h b/Engine/Plugins/Cameras/GameplayCameras/Source/GameplayCameras/Public/Core/CameraNodeEvaluatorFwd.h
--- a/Engine/Plugins/Cameras/GameplayCameras/Source/GameplayCameras/Public/Core/CameraNodeEvaluatorFwd.h
+++ b/Engine/Plugins/Cameras/GameplayCameras/Source/GameplayCameras/Public/Core/CameraNodeEvaluatorFwd.h
@@ -4,6 +4,8 @@
 
 #include "CoreTypes.h"
 
+#include "CameraNodeEvaluatorFwd.generated.h"
+
 class FCameraNodeEvaluator;
 class FCameraNodeEvaluatorStorage;
 struct FCameraNodeEvaluatorBuilder;
@@ -11,9 +13,15 @@ struct FCameraNodeEvaluatorBuilder;
 using FCameraNodeEvaluatorPtr = FCameraNodeEvaluator*;
 
 /** Allocation information for a node evaluator. */
+USTRUCT()
 struct FCameraNodeEvaluatorAllocationInfo
 {
+	GENERATED_BODY()
+
+	UPROPERTY()
 	int16 Sizeof = 0;
+
+	UPROPERTY()
 	int16 Alignof = 0;
 };
 
