Cameras: add player controller to evaluation context

diff --git a/Engine/Plugins/Cameras/GameplayCameras/Source/GameplayCameras/Private/GameFramework/GameplayCameraComponent.cpp b/Engine/Plugins/Cameras/GameplayCameras/Source/GameplayCameras/Private/GameFramework/GameplayCameraComponent.cpp
--- a/Engine/Plugins/Cameras/GameplayCameras/Source/GameplayCameras/Private/GameFramework/GameplayCameraComponent.cpp
+++ b/Engine/Plugins/Cameras/GameplayCameras/Source/GameplayCameras/Private/GameFramework/GameplayCameraComponent.cpp
@@ -68,7 +68,7 @@ void UGameplayCameraComponent::ActivateCamera(APlayerController* PlayerControlle
 	if (EvaluationContext == nullptr)
 	{
 		EvaluationContext = NewObject<UGameplayCameraComponentEvaluationContext>(this, TEXT("EvaluationContext"));
-		EvaluationContext->Initialize(this);
+		EvaluationContext->Initialize(this, PlayerController);
 	}
 
 	UCameraSystemEvaluator* Evaluator = CameraSystem->GetCameraSystemComponent()->GetCameraSystemEvaluator();
@@ -141,8 +141,9 @@ UGameplayCameraComponentEvaluationContext::UGameplayCameraComponentEvaluationCon
 {
 }
 
-void UGameplayCameraComponentEvaluationContext::Initialize(UGameplayCameraComponent* Owner)
+void UGameplayCameraComponentEvaluationContext::Initialize(UGameplayCameraComponent* Owner, APlayerController* InPlayerController)
 {
+	PlayerController = InPlayerController;
 	CameraAsset = Owner->Camera;
 }
 
diff --git a/Engine/Plugins/Cameras/GameplayCameras/Source/GameplayCameras/Public/Core/CameraEvaluationContext.h b/Engine/Plugins/Cameras/GameplayCameras/Source/GameplayCameras/Public/Core/CameraEvaluationContext.h
--- a/Engine/Plugins/Cameras/GameplayCameras/Source/GameplayCameras/Public/Core/CameraEvaluationContext.h
+++ b/Engine/Plugins/Cameras/GameplayCameras/Source/GameplayCameras/Public/Core/CameraEvaluationContext.h
@@ -9,6 +9,7 @@
 #include "CameraEvaluationContext.generated.h"
 
 class UCameraAsset;
+class APlayerController;
 
 /**
  * Base class for providing a context to running camera modes.
@@ -22,6 +23,12 @@ public:
 
 	UCameraEvaluationContext(const FObjectInitializer& ObjectInit);
 
+	/**
+	 * Gets the player controller (if any) in control of the cameras running inside
+	 * of this evaluation context.
+	 */
+	APlayerController* GetPlayerController() const { return PlayerController; }
+
 	/** Gets the camera asset that is hosted in this context. */
 	UCameraAsset* GetCameraAsset() const { return CameraAsset; }
 
@@ -30,6 +37,13 @@ public:
 
 protected:
 
+	/**
+	 * The player controller (if any) in control of the cameras running inside
+	 * of this evaluation context.
+	 */
+	UPROPERTY()
+	TObjectPtr<APlayerController> PlayerController;
+
 	/** The camera asset hosted in this context. */
 	UPROPERTY()
 	TObjectPtr<UCameraAsset> CameraAsset;
diff --git a/Engine/Plugins/Cameras/GameplayCameras/Source/GameplayCameras/Public/GameFramework/GameplayCameraComponent.h b/Engine/Plugins/Cameras/GameplayCameras/Source/GameplayCameras/Public/GameFramework/GameplayCameraComponent.h
--- a/Engine/Plugins/Cameras/GameplayCameras/Source/GameplayCameras/Public/GameFramework/GameplayCameraComponent.h
+++ b/Engine/Plugins/Cameras/GameplayCameras/Source/GameplayCameras/Public/GameFramework/GameplayCameraComponent.h
@@ -75,7 +75,7 @@ public:
 
 	UGameplayCameraComponentEvaluationContext(const FObjectInitializer& ObjectInit);
 
-	void Initialize(UGameplayCameraComponent* Owner);
+	void Initialize(UGameplayCameraComponent* Owner, APlayerController* InPlayerController);
 	void Update(UGameplayCameraComponent* Owner);
 };
 
