Cameras: fix missing function return

diff --git a/Engine/Plugins/Cameras/GameplayCameras/Source/GameplayCameras/Public/Core/CameraVariableTable.h b/Engine/Plugins/Cameras/GameplayCameras/Source/GameplayCameras/Public/Core/CameraVariableTable.h
--- a/Engine/Plugins/Cameras/GameplayCameras/Source/GameplayCameras/Public/Core/CameraVariableTable.h
+++ b/Engine/Plugins/Cameras/GameplayCameras/Source/GameplayCameras/Public/Core/CameraVariableTable.h
@@ -121,7 +121,7 @@ private:
 	template<typename ValueType>
 	bool CheckVariableType(ECameraVariableType InType)
 	{
-		ensure(TCameraVariableTraits<ValueType>::Type == InType);
+		return ensure(TCameraVariableTraits<ValueType>::Type == InType);
 	}
 
 	void InternalOverrideChanged(const FCameraVariableTable& OtherTable, const FCameraVariableTableFlags* InMask, bool bInvertMask, FCameraVariableTableFlags* OutMask);
