Cameras: add boom arm node

diff --git a/Engine/Plugins/Cameras/GameplayCameras/Source/GameplayCameras/Private/Nodes/Common/BoomArmCameraNode.cpp b/Engine/Plugins/Cameras/GameplayCameras/Source/GameplayCameras/Private/Nodes/Common/BoomArmCameraNode.cpp
new file mode 100644
--- /dev/null
+++ b/Engine/Plugins/Cameras/GameplayCameras/Source/GameplayCameras/Private/Nodes/Common/BoomArmCameraNode.cpp
@@ -0,0 +1,59 @@
+// Copyright Epic Games, Inc. All Rights Reserved.
+
+#include "Nodes/Common/BoomArmCameraNode.h"
+
+#include "Core/CameraEvaluationContext.h"
+#include "GameFramework/Pawn.h"
+#include "GameFramework/PlayerController.h"
+#include "GameplayCameras.h"
+
+#include UE_INLINE_GENERATED_CPP_BY_NAME(BoomArmCameraNode)
+
+class FBoomArmCameraNodeEvaluator : public FCameraNodeEvaluator
+{
+	UE_DECLARE_CAMERA_NODE_EVALUATOR(FBoomArmCameraNodeEvaluator)
+
+protected:
+
+	virtual void OnRun(const FCameraNodeEvaluationParams& Params, FCameraNodeEvaluationResult& OutResult) override;
+};
+
+UE_DEFINE_CAMERA_NODE_EVALUATOR(FBoomArmCameraNodeEvaluator)
+
+void FBoomArmCameraNodeEvaluator::OnRun(const FCameraNodeEvaluationParams& Params, FCameraNodeEvaluationResult& OutResult)
+{
+	FRotator3d BoomRotation;
+	if (Params.EvaluationContext)
+	{
+		if (APlayerController* PlayerController = Params.EvaluationContext->GetPlayerController())
+		{
+			if (APawn* Pawn = PlayerController->GetPawn())
+			{
+				const FRotator3d PawnViewRotation = Pawn->GetViewRotation();
+				BoomRotation = PawnViewRotation;
+			}
+		}
+	}
+
+	const UBoomArmCameraNode* BoomArmNode = GetCameraNodeAs<UBoomArmCameraNode>();
+
+	// Here we want to logically apply transform in this order:
+	//
+	// FinalTransform = BoomOffset * BoomRotation * CameraPose.Location
+	//
+	// Since FTransform3d applies rotation first and translation second, we can save one multiplication
+	// by using the fact that BoomRotation is, well, just a rotation, and CameraPose.Location is of
+	// course just a translation. So we can put them both in the same transform:
+	const FTransform3d BoomPivot(BoomRotation, OutResult.CameraPose.GetLocation());
+	const FTransform3d BoomOffset(BoomArmNode->BoomOffset);
+
+	const FTransform3d FinalTransform(BoomOffset * BoomPivot);
+
+	OutResult.CameraPose.SetTransform(FinalTransform);
+}
+
+FCameraNodeEvaluatorPtr UBoomArmCameraNode::OnBuildEvaluator(FCameraNodeEvaluatorBuilder& Builder) const
+{
+	return Builder.BuildEvaluator<FBoomArmCameraNodeEvaluator>();
+}
+
diff --git a/Engine/Plugins/Cameras/GameplayCameras/Source/GameplayCameras/Public/Nodes/Common/BoomArmCameraNode.h b/Engine/Plugins/Cameras/GameplayCameras/Source/GameplayCameras/Public/Nodes/Common/BoomArmCameraNode.h
new file mode 100644
--- /dev/null
+++ b/Engine/Plugins/Cameras/GameplayCameras/Source/GameplayCameras/Public/Nodes/Common/BoomArmCameraNode.h
@@ -0,0 +1,28 @@
+// Copyright Epic Games, Inc. All Rights Reserved.
+
+#pragma once
+
+#include "Core/CameraNode.h"
+#include "Nodes/CameraNodeTypes.h"
+
+#include "BoomArmCameraNode.generated.h"
+
+/**
+ * A camera node that can rotate the camera in yaw and pitch based on player input.
+ */
+UCLASS(MinimalAPI)
+class UBoomArmCameraNode : public UCameraNode
+{
+	GENERATED_BODY()
+
+protected:
+
+	// UCameraNode interface.
+	virtual FCameraNodeEvaluatorPtr OnBuildEvaluator(FCameraNodeEvaluatorBuilder& Builder) const override;
+
+public:
+
+	UPROPERTY(EditAnywhere, Category=Common)
+	FVector3d BoomOffset;
+};
+
