Cameras: add BP method to deactivate gameplay camera component

diff --git a/Engine/Plugins/Cameras/GameplayCameras/Source/GameplayCameras/Private/GameFramework/GameplayCameraComponent.cpp b/Engine/Plugins/Cameras/GameplayCameras/Source/GameplayCameras/Private/GameFramework/GameplayCameraComponent.cpp
--- a/Engine/Plugins/Cameras/GameplayCameras/Source/GameplayCameras/Private/GameFramework/GameplayCameraComponent.cpp
+++ b/Engine/Plugins/Cameras/GameplayCameras/Source/GameplayCameras/Private/GameFramework/GameplayCameraComponent.cpp
@@ -34,6 +34,11 @@ UGameplayCameraComponent::UGameplayCameraComponent(const FObjectInitializer& Obj
 
 void UGameplayCameraComponent::ActivateCamera(int32 PlayerIndex)
 {
+	if (ActivatedForPlayerIndex >= 0)
+	{
+		return;
+	}
+
 	UWorld* World = GetWorld();
 	if (!ensure(World))
 	{
@@ -48,10 +53,37 @@ void UGameplayCameraComponent::ActivateCamera(int32 PlayerIndex)
 		if (APlayerController* PlayerController = It->Get())
 		{
 			ActivateCamera(PlayerController);
+			ActivatedForPlayerIndex = PlayerIndex;
 		}
 	}
 }
 
+void UGameplayCameraComponent::DeactivateCamera()
+{
+	if (ActivatedForPlayerIndex < 0)
+	{
+		return;
+	}
+
+	UWorld* World = GetWorld();
+	if (!ensure(World))
+	{
+		return;
+	}
+
+	if (ensure(ActivatedForPlayerIndex < World->GetNumPlayerControllers()))
+	{
+		FConstPlayerControllerIterator It = World->GetPlayerControllerIterator();
+		It += ActivatedForPlayerIndex;
+		if (APlayerController* PlayerController = It->Get())
+		{
+			DeactivateCamera(PlayerController);
+		}
+	}
+
+	ActivatedForPlayerIndex = INDEX_NONE;
+}
+
 void UGameplayCameraComponent::ActivateCamera(APlayerController* PlayerController)
 {
 	if (!ensure(PlayerController && PlayerController->PlayerCameraManager))
@@ -77,6 +109,28 @@ void UGameplayCameraComponent::ActivateCamera(APlayerController* PlayerControlle
 	Activate();
 }
 
+void UGameplayCameraComponent::DeactivateCamera(APlayerController* PlayerController)
+{
+	if (!ensure(PlayerController && PlayerController->PlayerCameraManager))
+	{
+		return;
+	}
+	
+	AGameplayCameraSystemActor* CameraSystem = Cast<AGameplayCameraSystemActor>(PlayerController->PlayerCameraManager->GetViewTarget());
+	if (!ensure(CameraSystem))
+	{
+		return;
+	}
+
+	if (EvaluationContext != nullptr)
+	{
+		UCameraSystemEvaluator* Evaluator = CameraSystem->GetCameraSystemComponent()->GetCameraSystemEvaluator();
+		Evaluator->RemoveEvaluationContext(EvaluationContext);
+	}
+
+	Deactivate();
+}
+
 void UGameplayCameraComponent::OnRegister()
 {
 #if WITH_EDITORONLY_DATA
diff --git a/Engine/Plugins/Cameras/GameplayCameras/Source/GameplayCameras/Public/GameFramework/GameplayCameraComponent.h b/Engine/Plugins/Cameras/GameplayCameras/Source/GameplayCameras/Public/GameFramework/GameplayCameraComponent.h
--- a/Engine/Plugins/Cameras/GameplayCameras/Source/GameplayCameras/Public/GameFramework/GameplayCameraComponent.h
+++ b/Engine/Plugins/Cameras/GameplayCameras/Source/GameplayCameras/Public/GameFramework/GameplayCameraComponent.h
@@ -28,6 +28,9 @@ public:
 	UFUNCTION(BlueprintCallable, Category=Camera)
 	GAMEPLAYCAMERAS_API void ActivateCamera(int32 PlayerIndex = 0);
 
+	UFUNCTION(BlueprintCallable, Category=Camera)
+	GAMEPLAYCAMERAS_API void DeactivateCamera();
+
 public:
 
 	// UActorComponent interface
@@ -38,6 +41,7 @@ public:
 private:
 
 	void ActivateCamera(APlayerController* PlayerController);
+	void DeactivateCamera(APlayerController* PlayerController);
 
 #if WITH_EDITORONLY_DATA
 
@@ -64,6 +68,10 @@ protected:
 	TObjectPtr<UStaticMeshComponent> PreviewMeshComponent;
 
 #endif	// WITH_EDITORONLY_DATA
+	
+private:
+
+	int32 ActivatedForPlayerIndex = INDEX_NONE;
 };
 
 UCLASS(MinimalAPI)
