Cameras: update damping factor from the data node in editor builds to support live-editing

diff --git a/Engine/Plugins/Cameras/GameplayCameras/Source/GameplayCameras/Private/Nodes/Common/DampenPositionCameraNode.cpp b/Engine/Plugins/Cameras/GameplayCameras/Source/GameplayCameras/Private/Nodes/Common/DampenPositionCameraNode.cpp
--- a/Engine/Plugins/Cameras/GameplayCameras/Source/GameplayCameras/Private/Nodes/Common/DampenPositionCameraNode.cpp
+++ b/Engine/Plugins/Cameras/GameplayCameras/Source/GameplayCameras/Private/Nodes/Common/DampenPositionCameraNode.cpp
@@ -49,6 +49,13 @@ void FDampenPositionCameraNodeEvaluator::OnInitialize(const FCameraNodeEvaluator
 
 void FDampenPositionCameraNodeEvaluator::OnRun(const FCameraNodeEvaluationParams& Params, FCameraNodeEvaluationResult& OutResult)
 {
+#if WITH_EDITOR
+	const UDampenPositionCameraNode* DampenNode = GetCameraNodeAs<UDampenPositionCameraNode>();
+	ForwardDamper.SetW0(DampenNode->ForwardDampingFactor);
+	LateralDamper.SetW0(DampenNode->LateralDampingFactor);
+	VerticalDamper.SetW0(DampenNode->VerticalDampingFactor);
+#endif
+
 	// We want the dampen the given camera position, which means it's trying
 	// to converge towards the one given in the result (which we set as our 
 	// next target), but will be lagging behind.
