Cameras: fix static analysis

diff --git a/Engine/Plugins/Cameras/GameplayCameras/Source/GameplayCameras/Public/Core/CameraVariableTableFwd.h b/Engine/Plugins/Cameras/GameplayCameras/Source/GameplayCameras/Public/Core/CameraVariableTableFwd.h
--- a/Engine/Plugins/Cameras/GameplayCameras/Source/GameplayCameras/Public/Core/CameraVariableTableFwd.h
+++ b/Engine/Plugins/Cameras/GameplayCameras/Source/GameplayCameras/Public/Core/CameraVariableTableFwd.h
@@ -51,7 +51,7 @@ struct FCameraVariableDefinition
 	uint32 VariableId = 0;
 
 	UPROPERTY()
-	ECameraVariableType VariableType;
+	ECameraVariableType VariableType = ECameraVariableType::Boolean;
 
 	UPROPERTY()
 	bool bIsPrivate = false;
