Cameras: move gameplay camera assets into a sub-menu of the Content Browser.

diff --git a/Engine/Plugins/Cameras/GameplayCameras/Source/GameplayCamerasEditor/Private/AssetTools/AssetDefinition_CameraAsset.cpp b/Engine/Plugins/Cameras/GameplayCameras/Source/GameplayCamerasEditor/Private/AssetTools/AssetDefinition_CameraAsset.cpp
--- a/Engine/Plugins/Cameras/GameplayCameras/Source/GameplayCamerasEditor/Private/AssetTools/AssetDefinition_CameraAsset.cpp
+++ b/Engine/Plugins/Cameras/GameplayCameras/Source/GameplayCamerasEditor/Private/AssetTools/AssetDefinition_CameraAsset.cpp
@@ -7,13 +7,13 @@
 #include "Toolkits/CameraAssetEditorToolkit.h"
 #include "Toolkits/IToolkit.h"
 
-#define LOCTEXT_NAMESPACE "AssetTypeActions"
+#define LOCTEXT_NAMESPACE "AssetDefinition_CameraAsset"
 
 #include UE_INLINE_GENERATED_CPP_BY_NAME(AssetDefinition_CameraAsset)
 
 FText UAssetDefinition_CameraAsset::GetAssetDisplayName() const
 {
-	return NSLOCTEXT("AssetTypeActions", "AssetTypeActions_CameraAsset", "Camera Asset");
+	return LOCTEXT("AssetDisplayName", "Camera Asset");
 }
 
 FLinearColor UAssetDefinition_CameraAsset::GetAssetColor() const
@@ -28,7 +28,7 @@ TSoftClassPtr<UObject> UAssetDefinition_CameraAsset::GetAssetClass() const
 
 TConstArrayView<FAssetCategoryPath> UAssetDefinition_CameraAsset::GetAssetCategories() const
 {
-	static const auto Categories = { EAssetCategoryPaths::Gameplay };
+	static const auto Categories = { FAssetCategoryPath(EAssetCategoryPaths::Gameplay, LOCTEXT("AssetSubCategory", "Cameras")) };
 	return Categories;
 }
 
diff --git a/Engine/Plugins/Cameras/GameplayCameras/Source/GameplayCamerasEditor/Private/AssetTools/AssetDefinition_CameraMode.cpp b/Engine/Plugins/Cameras/GameplayCameras/Source/GameplayCamerasEditor/Private/AssetTools/AssetDefinition_CameraMode.cpp
--- a/Engine/Plugins/Cameras/GameplayCameras/Source/GameplayCamerasEditor/Private/AssetTools/AssetDefinition_CameraMode.cpp
+++ b/Engine/Plugins/Cameras/GameplayCameras/Source/GameplayCamerasEditor/Private/AssetTools/AssetDefinition_CameraMode.cpp
@@ -4,16 +4,16 @@
 
 #include "IGameplayCamerasEditorModule.h"
 #include "Styles/GameplayCamerasEditorStyle.h"
-#include "Toolkits/CameraModeEditorToolkit.h"
+#include "Toolkits/CameraModeAssetEditorToolkit.h"
 #include "Toolkits/IToolkit.h"
 
-#define LOCTEXT_NAMESPACE "AssetTypeActions"
+#define LOCTEXT_NAMESPACE "AssetDefinition_CameraMode"
 
 #include UE_INLINE_GENERATED_CPP_BY_NAME(AssetDefinition_CameraMode)
 
 FText UAssetDefinition_CameraMode::GetAssetDisplayName() const
 {
-	return NSLOCTEXT("AssetTypeActions", "AssetTypeActions_CameraMode", "Camera Mode");
+	return LOCTEXT("AssetDisplayName", "Camera Mode");
 }
 
 FLinearColor UAssetDefinition_CameraMode::GetAssetColor() const
@@ -28,7 +28,7 @@ TSoftClassPtr<UObject> UAssetDefinition_CameraMode::GetAssetClass() const
 
 TConstArrayView<FAssetCategoryPath> UAssetDefinition_CameraMode::GetAssetCategories() const
 {
-	static const auto Categories = { EAssetCategoryPaths::Gameplay };
+	static const auto Categories = { FAssetCategoryPath(EAssetCategoryPaths::Gameplay, LOCTEXT("AssetSubCategory", "Cameras")) };
 	return Categories;
 }
 
