Cameras: rewrite camera-related editors using the new recommended editor framework

diff --git a/Engine/Plugins/Cameras/GameplayCameras/Source/GameplayCamerasEditor/GameplayCamerasEditor.Build.cs b/Engine/Plugins/Cameras/GameplayCameras/Source/GameplayCamerasEditor/GameplayCamerasEditor.Build.cs
--- a/Engine/Plugins/Cameras/GameplayCameras/Source/GameplayCamerasEditor/GameplayCamerasEditor.Build.cs
+++ b/Engine/Plugins/Cameras/GameplayCameras/Source/GameplayCamerasEditor/GameplayCamerasEditor.Build.cs
@@ -22,9 +22,11 @@ public class GameplayCamerasEditor : ModuleRules
 				"AssetRegistry",
 				"GameplayCameras",
 				"BlueprintGraph",
+				"EditorSubsystem",
 				"Core",
 				"CoreUObject",
 				"EditorFramework",
+				"InteractiveToolsFramework",
 				"Engine",
 				"InputCore",
 				"Kismet",
diff --git a/Engine/Plugins/Cameras/GameplayCameras/Source/GameplayCamerasEditor/Private/AssetTools/CameraAssetEditor.cpp b/Engine/Plugins/Cameras/GameplayCameras/Source/GameplayCamerasEditor/Private/AssetTools/CameraAssetEditor.cpp
new file mode 100644
--- /dev/null
+++ b/Engine/Plugins/Cameras/GameplayCameras/Source/GameplayCamerasEditor/Private/AssetTools/CameraAssetEditor.cpp
@@ -0,0 +1,27 @@
+// Copyright Epic Games, Inc. All Rights Reserved.
+
+#include "AssetTools/CameraAssetEditor.h"
+
+#include "Core/CameraAsset.h"
+#include "EditorModeManager.h"
+#include "Toolkits/CameraAssetEditorToolkit.h"
+
+#include UE_INLINE_GENERATED_CPP_BY_NAME(CameraAssetEditor)
+
+void UCameraAssetEditor::Initialize(TObjectPtr<UCameraAsset> InCameraAsset)
+{
+	CameraAsset = InCameraAsset;
+
+	Super::Initialize();
+}
+
+void UCameraAssetEditor::GetObjectsToEdit(TArray<UObject*>& InObjectsToEdit)
+{
+	InObjectsToEdit.Add(CameraAsset.Get());
+}
+
+TSharedPtr<FBaseAssetToolkit> UCameraAssetEditor::CreateToolkit()
+{
+	return MakeShared<FCameraAssetEditorToolkit>(this);
+}
+
diff --git a/Engine/Plugins/Cameras/GameplayCameras/Source/GameplayCamerasEditor/Private/AssetTools/CameraModeAssetEditor.cpp b/Engine/Plugins/Cameras/GameplayCameras/Source/GameplayCamerasEditor/Private/AssetTools/CameraModeAssetEditor.cpp
new file mode 100644
--- /dev/null
+++ b/Engine/Plugins/Cameras/GameplayCameras/Source/GameplayCamerasEditor/Private/AssetTools/CameraModeAssetEditor.cpp
@@ -0,0 +1,27 @@
+// Copyright Epic Games, Inc. All Rights Reserved.
+
+#include "AssetTools/CameraModeAssetEditor.h"
+
+#include "Core/CameraMode.h"
+#include "EditorModeManager.h"
+#include "Toolkits/CameraModeAssetEditorToolkit.h"
+
+#include UE_INLINE_GENERATED_CPP_BY_NAME(CameraModeAssetEditor)
+
+void UCameraModeAssetEditor::Initialize(TObjectPtr<UCameraMode> InCameraModeAsset)
+{
+	CameraModeAsset = InCameraModeAsset;
+
+	Super::Initialize();
+}
+
+void UCameraModeAssetEditor::GetObjectsToEdit(TArray<UObject*>& InObjectsToEdit)
+{
+	InObjectsToEdit.Add(CameraModeAsset.Get());
+}
+
+TSharedPtr<FBaseAssetToolkit> UCameraModeAssetEditor::CreateToolkit()
+{
+	return MakeShared<FCameraModeAssetEditorToolkit>(this);
+}
+
diff --git a/Engine/Plugins/Cameras/GameplayCameras/Source/GameplayCamerasEditor/Private/GameplayCamerasEditorModule.cpp b/Engine/Plugins/Cameras/GameplayCameras/Source/GameplayCamerasEditor/Private/GameplayCamerasEditorModule.cpp
--- a/Engine/Plugins/Cameras/GameplayCameras/Source/GameplayCamerasEditor/Private/GameplayCamerasEditorModule.cpp
+++ b/Engine/Plugins/Cameras/GameplayCameras/Source/GameplayCamerasEditor/Private/GameplayCamerasEditorModule.cpp
@@ -7,9 +7,11 @@
 #include "IGameplayCamerasEditorModule.h"
 #include "ISettingsModule.h"
 #include "Modules/ModuleManager.h"
+#include "AssetTools/CameraAssetEditor.h"
+#include "AssetTools/CameraModeAssetEditor.h"
 #include "Styles/GameplayCamerasEditorStyle.h"
 #include "Toolkits/CameraAssetEditorToolkit.h"
-#include "Toolkits/CameraModeEditorToolkit.h"
+#include "Toolkits/CameraModeAssetEditorToolkit.h"
 
 #define LOCTEXT_NAMESPACE "GameplayCamerasEditor"
 
@@ -37,18 +39,20 @@ public:
 		TeardownLiveEditManager();
 	}
 
-	virtual TSharedRef<ICameraAssetEditorToolkit> CreateCameraAssetEditor(const EToolkitMode::Type Mode, const TSharedPtr<IToolkitHost>& InitToolkitHost, UCameraAsset* CameraAsset) override
+	virtual UCameraAssetEditor* CreateCameraAssetEditor(const EToolkitMode::Type Mode, const TSharedPtr<IToolkitHost>& InitToolkitHost, UCameraAsset* CameraAsset) override
 	{
-		TSharedRef<FCameraAssetEditorToolkit> Toolkit = MakeShareable(new FCameraAssetEditorToolkit(FGameplayCamerasEditorStyle::Get()));
-		Toolkit->Initialize(Mode, InitToolkitHost, CameraAsset);
-		return Toolkit;
+		UAssetEditorSubsystem* AssetEditorSubsystem = GEditor->GetEditorSubsystem<UAssetEditorSubsystem>();
+		UCameraAssetEditor* AssetEditor = NewObject<UCameraAssetEditor>(AssetEditorSubsystem, NAME_None, RF_Transient);
+		AssetEditor->Initialize(CameraAsset);
+		return AssetEditor;
 	}
 
-	virtual TSharedRef<ICameraModeEditorToolkit> CreateCameraModeEditor(const EToolkitMode::Type Mode, const TSharedPtr<IToolkitHost>& InitToolkitHost, UCameraMode* CameraMode) override
+	virtual UCameraModeAssetEditor* CreateCameraModeEditor(const EToolkitMode::Type Mode, const TSharedPtr<IToolkitHost>& InitToolkitHost, UCameraMode* CameraMode) override
 	{
-		TSharedRef<FCameraModeEditorToolkit> Toolkit = MakeShareable(new FCameraModeEditorToolkit(FGameplayCamerasEditorStyle::Get()));
-		Toolkit->Initialize(Mode, InitToolkitHost, CameraMode);
-		return Toolkit;
+		UAssetEditorSubsystem* AssetEditorSubsystem = GEditor->GetEditorSubsystem<UAssetEditorSubsystem>();
+		UCameraModeAssetEditor* AssetEditor = NewObject<UCameraModeAssetEditor>(AssetEditorSubsystem, NAME_None, RF_Transient);
+		AssetEditor->Initialize(CameraMode);
+		return AssetEditor;
 	}
 
 private:
diff --git a/Engine/Plugins/Cameras/GameplayCameras/Source/GameplayCamerasEditor/Private/Toolkits/CameraAssetEditorToolkit.cpp b/Engine/Plugins/Cameras/GameplayCameras/Source/GameplayCamerasEditor/Private/Toolkits/CameraAssetEditorToolkit.cpp
--- a/Engine/Plugins/Cameras/GameplayCameras/Source/GameplayCamerasEditor/Private/Toolkits/CameraAssetEditorToolkit.cpp
+++ b/Engine/Plugins/Cameras/GameplayCameras/Source/GameplayCamerasEditor/Private/Toolkits/CameraAssetEditorToolkit.cpp
@@ -1,31 +1,27 @@
 // Copyright Epic Games, Inc. All Rights Reserved.
 
-#include "CameraAssetEditorToolkit.h"
+#include "Toolkits/CameraAssetEditorToolkit.h"
 
+#include "AssetTools/CameraAssetEditor.h"
 #include "Core/CameraAsset.h"
-#include "Framework/Docking/TabManager.h"
-#include "Framework/MultiBox/MultiBoxBuilder.h"
-#include "IGameplayCamerasEditorModule.h"
+#include "EditorModeManager.h"
+#include "Framework/Docking/LayoutExtender.h"
 #include "IMessageLogListing.h"
 #include "MessageLogInitializationOptions.h"
 #include "MessageLogModule.h"
 #include "Modules/ModuleManager.h"
 #include "PropertyEditorModule.h"
-#include "ScopedTransaction.h"
-#include "Selection.h"
 #include "Widgets/Docking/SDockTab.h"
-#include "WorkflowOrientedApp/WorkflowTabFactory.h"
-#include "WorkflowOrientedApp/WorkflowTabManager.h"
 
 #define LOCTEXT_NAMESPACE "CameraAssetEditorToolkit"
 
-const FName FCameraAssetEditorTabs::DetailsViewTabId(TEXT("CameraAssetEditor_DetailsView"));
+const FName FCameraAssetEditorToolkit::DetailsViewTabId(TEXT("CameraAssetEditor_DetailsView"));
 
-const FName FCameraAssetEditorApplicationModes::StandardCameraAssetEditorMode(TEXT("StandardCameraAssetEditorMode"));
-
-TSharedPtr<FTabManager::FLayout> FCameraAssetEditorApplicationModes::GetDefaultEditorLayout(TSharedPtr<ICameraAssetEditorToolkit> InCameraAssetEditor)
+FCameraAssetEditorToolkit::FCameraAssetEditorToolkit(UCameraAssetEditor* InOwningAssetEditor)
+	: FBaseAssetToolkit(InOwningAssetEditor)
 {
-	const TSharedRef<FTabManager::FLayout> DefaultLayout = FTabManager::NewLayout("CameraAssetEditor_Layout")
+	// Override base class default layout.
+	StandaloneDefaultLayout = FTabManager::NewLayout("CameraAssetEditor_Layout")
 		->AddArea
 		(
 			FTabManager::NewPrimaryArea()->SetOrientation(Orient_Vertical)
@@ -35,88 +31,52 @@ TSharedPtr<FTabManager::FLayout> FCameraAssetEditorApplicationModes::GetDefaultE
 				->Split
 				(
 					FTabManager::NewStack()
-					->AddTab(FCameraAssetEditorTabs::DetailsViewTabId, ETabState::OpenedTab)
-					->SetForegroundTab(FCameraAssetEditorTabs::DetailsViewTabId)
+					->AddTab(DetailsViewTabId, ETabState::OpenedTab)
+					->SetForegroundTab(DetailsViewTabId)
 				)
 			)
 		);
-	return DefaultLayout;
 }
 
-FCameraAssetEditorApplicationMode::FCameraAssetEditorApplicationMode(TSharedPtr<ICameraAssetEditorToolkit> InCameraAssetEditor)
-	: FApplicationMode(FCameraAssetEditorApplicationModes::StandardCameraAssetEditorMode)
+FCameraAssetEditorToolkit::~FCameraAssetEditorToolkit()
 {
-	WeakCameraAssetEditor = InCameraAssetEditor;
-	//CameraAssetEditorTabFactories.RegisterFactory(InCameraAssetEditor->GraphEditorTabFactoryPtr.Pin().ToSharedRef());
-	TabLayout = FCameraAssetEditorApplicationModes::GetDefaultEditorLayout(InCameraAssetEditor);
 }
 
-void FCameraAssetEditorApplicationMode::RegisterTabFactories(TSharedPtr<FTabManager> InTabManager)
-{
-	if (TSharedPtr<ICameraAssetEditorToolkit> CameraAssetEditor = WeakCameraAssetEditor.Pin())
-	{
-		// Hack?
-		StaticCastSharedPtr<FCameraAssetEditorToolkit>(CameraAssetEditor)->InternalRegisterTabSpawners(InTabManager.ToSharedRef());
-
-		CameraAssetEditor->PushTabFactories(CameraAssetEditorTabFactories);
-	}
-
-	FApplicationMode::RegisterTabFactories(InTabManager);
-}
-
-FCameraAssetEditorToolkit::FCameraAssetEditorToolkit(const TSharedRef<ISlateStyle>& InStyle)
-	: CameraAsset(nullptr)
-	, Style(InStyle)
+void FCameraAssetEditorToolkit::RegisterTabSpawners(const TSharedRef<FTabManager>& InTabManager)
 {
-	DocumentManager = MakeShareable(new FDocumentTracker());
-}
+	// Skip FBaseAssetToolkit here because we don't want a viewport tab.
+	FAssetEditorToolkit::RegisterTabSpawners(InTabManager);
 
-FCameraAssetEditorToolkit::~FCameraAssetEditorToolkit()
-{
+	InTabManager->RegisterTabSpawner(DetailsViewTabId, FOnSpawnTab::CreateSP(this, &FCameraAssetEditorToolkit::SpawnTab_Details))
+		.SetDisplayName(LOCTEXT("Details", "Details"))
+		.SetGroup(AssetEditorTabsCategory.ToSharedRef())
+		.SetIcon(FSlateIcon(FAppStyle::GetAppStyleSetName(), "LevelEditor.Tabs.Details"));
 }
 
-void FCameraAssetEditorToolkit::Initialize(const EToolkitMode::Type Mode, const TSharedPtr<IToolkitHost>& InitToolkitHost, UCameraAsset* InCameraAsset)
+void FCameraAssetEditorToolkit::UnregisterTabSpawners(const TSharedRef<FTabManager>& InTabManager)
 {
-	CameraAsset = InCameraAsset;
-
-	CreateWidgets();
-
-	TSharedPtr<FCameraAssetEditorToolkit> ThisPtr(SharedThis(this));
-	DocumentManager->Initialize(ThisPtr);
-
-	const bool bCreateDefaultStandaloneMenu = true;
-	const bool bCreateDefaultToolbar = true;
-
-	const TSharedRef<FTabManager::FLayout> DummyLayout = FTabManager::NewLayout("NullLayout")->AddArea(FTabManager::NewPrimaryArea());
-
-	FWorkflowCentricApplication::InitAssetEditor(
-			Mode,
-			InitToolkitHost, 
-			IGameplayCamerasEditorModule::GameplayCamerasEditorAppIdentifier, 
-			DummyLayout, 
-			bCreateDefaultStandaloneMenu, 
-			bCreateDefaultToolbar,
-			CameraAsset);
+	// Skip FBaseAssetToolkit here because we don't want a viewport tab.
+	FAssetEditorToolkit::UnregisterTabSpawners(InTabManager);
 
-	
-	AddApplicationMode(
-		FCameraAssetEditorApplicationModes::StandardCameraAssetEditorMode,
-		MakeShareable(new FCameraAssetEditorApplicationMode(SharedThis(this))));
-	SetCurrentMode(FCameraAssetEditorApplicationModes::StandardCameraAssetEditorMode);
+	InTabManager->UnregisterTabSpawner(DetailsViewTabId);
 }
 
 void FCameraAssetEditorToolkit::CreateWidgets()
 {
-	FPropertyEditorModule& PropertyEditorModule = FModuleManager::GetModuleChecked<FPropertyEditorModule>( "PropertyEditor" );
+	// Skip FBaseAssetToolkit here because we don't want a viewport tab.
+	// ...no up-call...
 
+	// Do most of FBaseAssetToolkit's work except for the viewport.
+	RegisterToolbar();
+	LayoutExtender = MakeShared<FLayoutExtender>();
+
+	FPropertyEditorModule& PropertyEditorModule = FModuleManager::GetModuleChecked<FPropertyEditorModule>("PropertyEditor");
 	FDetailsViewArgs DetailsViewArgs;
 	DetailsViewArgs.NameAreaSettings = FDetailsViewArgs::HideNameArea;
 	DetailsViewArgs.bHideSelectionTip = true;
-	DetailsView = PropertyEditorModule.CreateDetailView( DetailsViewArgs );
-
-	TArray<TWeakObjectPtr<>> SelectedObjects { CameraAsset };
-	DetailsView->SetObjects(SelectedObjects, true);
+	DetailsView = PropertyEditorModule.CreateDetailView(DetailsViewArgs);
 
+	// Create the message log.
 	FMessageLogModule& MessageLogModule = FModuleManager::LoadModuleChecked<FMessageLogModule>("MessageLog");
 	FMessageLogInitializationOptions LogOptions;
 	LogOptions.bShowPages = false;
@@ -125,44 +85,10 @@ void FCameraAssetEditorToolkit::CreateWidgets()
 	LogOptions.MaxPageCount = 1;
 	StatsListing = MessageLogModule.CreateLogListing("CameraAssetEditorStats", LogOptions);
 
-	Stats = MessageLogModule.CreateLogListingWidget(StatsListing.ToSharedRef());
-}
-
-void FCameraAssetEditorToolkit::RegisterTabSpawners(const TSharedRef<FTabManager>& InTabManager)
-{
-	DocumentManager->SetTabManager(InTabManager);
-	FWorkflowCentricApplication::RegisterTabSpawners(InTabManager);
-}
-
-void FCameraAssetEditorToolkit::InternalRegisterTabSpawners(const TSharedRef<FTabManager>& InTabManager)
-{
-	FAssetEditorToolkit::UnregisterTabSpawners(InTabManager);
-
-	WorkspaceMenuCategory = InTabManager->AddLocalWorkspaceMenuCategory(LOCTEXT("WorkspaceMenu_CameraAssetEditor", "Camera Asset Editor"));
-	auto WorkspaceMenuCategoryRef = WorkspaceMenuCategory.ToSharedRef();
-
-	InTabManager->RegisterTabSpawner(FCameraAssetEditorTabs::DetailsViewTabId, FOnSpawnTab::CreateSP(this, &FCameraAssetEditorToolkit::SpawnTab_DetailsView))
-		.SetDisplayName(LOCTEXT("DetailsViewTabTitle", "Details"))
-		.SetGroup(WorkspaceMenuCategoryRef)
-		.SetIcon(FSlateIcon(FAppStyle::GetAppStyleSetName(), "LevelEditor.Tabs.Details"));
-}
-
-void FCameraAssetEditorToolkit::UnregisterTabSpawners(const TSharedRef<FTabManager>& InTabManager)
-{
-	FWorkflowCentricApplication::UnregisterTabSpawners(InTabManager);
+	Stats = MessageLogModule.CreateLogListingWidget(StatsListing.ToSharedRef());}
 
-	InTabManager->UnregisterTabSpawner(FCameraAssetEditorTabs::DetailsViewTabId);
-}
-
-TSharedRef<SDockTab> FCameraAssetEditorToolkit::SpawnTab_DetailsView(const FSpawnTabArgs& Args)
+void FCameraAssetEditorToolkit::PostInitAssetEditor()
 {
-	TSharedPtr<SDockTab> DetailsTab = SNew(SDockTab)
-		.Label(LOCTEXT("DetailsViewTabTitle", "Details"))
-		[
-			DetailsView.ToSharedRef()
-		];
-	DetailsViewTab = DetailsTab;
-	return DetailsTab.ToSharedRef();
 }
 
 FText FCameraAssetEditorToolkit::GetBaseToolkitName() const
@@ -181,20 +107,6 @@ FString FCameraAssetEditorToolkit::GetWorldCentricTabPrefix() const
 	return LOCTEXT("WorldCentricTabPrefix", "Camera Asset ").ToString();
 }
 
-FText FCameraAssetEditorToolkit::GetTabSuffix() const
-{
-	if (CameraAsset)
-	{
-		const bool bIsDirty = CameraAsset->GetOutermost()->IsDirty();
-		if (bIsDirty)
-		{
-			return LOCTEXT("TabSuffixAsterix", "*");
-		}
-	}
-
-	return FText::GetEmpty();
-}
-
 FLinearColor FCameraAssetEditorToolkit::GetWorldCentricTabColorScale() const
 {
 	return FLinearColor(0.7, 0.0f, 0.0f, 0.5f);
diff --git a/Engine/Plugins/Cameras/GameplayCameras/Source/GameplayCamerasEditor/Private/Toolkits/CameraAssetEditorToolkit.h b/Engine/Plugins/Cameras/GameplayCameras/Source/GameplayCamerasEditor/Private/Toolkits/CameraAssetEditorToolkit.h
deleted file mode 100644
diff --git a/Engine/Plugins/Cameras/GameplayCameras/Source/GameplayCamerasEditor/Private/Toolkits/CameraModeAssetEditorToolkit.cpp b/Engine/Plugins/Cameras/GameplayCameras/Source/GameplayCamerasEditor/Private/Toolkits/CameraModeAssetEditorToolkit.cpp
new file mode 100644
--- /dev/null
+++ b/Engine/Plugins/Cameras/GameplayCameras/Source/GameplayCamerasEditor/Private/Toolkits/CameraModeAssetEditorToolkit.cpp
@@ -0,0 +1,116 @@
+// Copyright Epic Games, Inc. All Rights Reserved.
+
+#include "Toolkits/CameraModeAssetEditorToolkit.h"
+
+#include "AssetTools/CameraModeAssetEditor.h"
+#include "Core/CameraMode.h"
+#include "EditorModeManager.h"
+#include "Framework/Docking/LayoutExtender.h"
+#include "IMessageLogListing.h"
+#include "MessageLogInitializationOptions.h"
+#include "MessageLogModule.h"
+#include "Modules/ModuleManager.h"
+#include "PropertyEditorModule.h"
+#include "Widgets/Docking/SDockTab.h"
+
+#define LOCTEXT_NAMESPACE "CameraModeAssetEditorToolkit"
+
+const FName FCameraModeAssetEditorToolkit::DetailsViewTabId(TEXT("CameraModeAssetEditor_DetailsView"));
+
+FCameraModeAssetEditorToolkit::FCameraModeAssetEditorToolkit(UCameraModeAssetEditor* InOwningAssetEditor)
+	: FBaseAssetToolkit(InOwningAssetEditor)
+{
+	// Override base class default layout.
+	StandaloneDefaultLayout = FTabManager::NewLayout("CameraModeAssetEditor_Layout")
+		->AddArea
+		(
+			FTabManager::NewPrimaryArea()->SetOrientation(Orient_Vertical)
+			->Split
+			(
+				FTabManager::NewSplitter()->SetOrientation(Orient_Horizontal)->SetSizeCoefficient(0.9f)
+				->Split
+				(
+					FTabManager::NewStack()
+					->AddTab(DetailsViewTabId, ETabState::OpenedTab)
+					->SetForegroundTab(DetailsViewTabId)
+				)
+			)
+		);
+}
+
+FCameraModeAssetEditorToolkit::~FCameraModeAssetEditorToolkit()
+{
+}
+
+void FCameraModeAssetEditorToolkit::RegisterTabSpawners(const TSharedRef<FTabManager>& InTabManager)
+{
+	// Skip FBaseAssetToolkit here because we don't want a viewport tab.
+	FAssetEditorToolkit::RegisterTabSpawners(InTabManager);
+
+	InTabManager->RegisterTabSpawner(DetailsViewTabId, FOnSpawnTab::CreateSP(this, &FCameraModeAssetEditorToolkit::SpawnTab_Details))
+		.SetDisplayName(LOCTEXT("Details", "Details"))
+		.SetGroup(AssetEditorTabsCategory.ToSharedRef())
+		.SetIcon(FSlateIcon(FAppStyle::GetAppStyleSetName(), "LevelEditor.Tabs.Details"));
+}
+
+void FCameraModeAssetEditorToolkit::UnregisterTabSpawners(const TSharedRef<FTabManager>& InTabManager)
+{
+	// Skip FBaseAssetToolkit here because we don't want a viewport tab.
+	FAssetEditorToolkit::UnregisterTabSpawners(InTabManager);
+
+	InTabManager->UnregisterTabSpawner(DetailsViewTabId);
+}
+
+void FCameraModeAssetEditorToolkit::CreateWidgets()
+{
+	// Skip FBaseAssetToolkit here because we don't want a viewport tab.
+	// ...no up-call...
+
+	// Do most of FBaseAssetToolkit's work except for the viewport.
+	RegisterToolbar();
+	LayoutExtender = MakeShared<FLayoutExtender>();
+
+	FPropertyEditorModule& PropertyEditorModule = FModuleManager::GetModuleChecked<FPropertyEditorModule>("PropertyEditor");
+	FDetailsViewArgs DetailsViewArgs;
+	DetailsViewArgs.NameAreaSettings = FDetailsViewArgs::HideNameArea;
+	DetailsViewArgs.bHideSelectionTip = true;
+	DetailsView = PropertyEditorModule.CreateDetailView(DetailsViewArgs);
+
+	// Create the message log.
+	FMessageLogModule& MessageLogModule = FModuleManager::LoadModuleChecked<FMessageLogModule>("MessageLog");
+	FMessageLogInitializationOptions LogOptions;
+	LogOptions.bShowPages = false;
+	LogOptions.bShowFilters = false;
+	LogOptions.bAllowClear = false;
+	LogOptions.MaxPageCount = 1;
+	StatsListing = MessageLogModule.CreateLogListing("CameraModeAssetEditorStats", LogOptions);
+
+	Stats = MessageLogModule.CreateLogListingWidget(StatsListing.ToSharedRef());}
+
+void FCameraModeAssetEditorToolkit::PostInitAssetEditor()
+{
+}
+
+FText FCameraModeAssetEditorToolkit::GetBaseToolkitName() const
+{
+	return LOCTEXT("AppLabel", "Camera Mode Asset");
+}
+
+FName FCameraModeAssetEditorToolkit::GetToolkitFName() const
+{
+	static FName SequencerName("CameraModeAssetEditor");
+	return SequencerName;
+}
+
+FString FCameraModeAssetEditorToolkit::GetWorldCentricTabPrefix() const
+{
+	return LOCTEXT("WorldCentricTabPrefix", "Camera Mode Asset ").ToString();
+}
+
+FLinearColor FCameraModeAssetEditorToolkit::GetWorldCentricTabColorScale() const
+{
+	return FLinearColor(0.7, 0.0f, 0.0f, 0.5f);
+}
+
+#undef LOCTEXT_NAMESPACE
+
diff --git a/Engine/Plugins/Cameras/GameplayCameras/Source/GameplayCamerasEditor/Private/Toolkits/CameraModeEditorToolkit.cpp b/Engine/Plugins/Cameras/GameplayCameras/Source/GameplayCamerasEditor/Private/Toolkits/CameraModeEditorToolkit.cpp
deleted file mode 100644
diff --git a/Engine/Plugins/Cameras/GameplayCameras/Source/GameplayCamerasEditor/Private/Toolkits/CameraModeEditorToolkit.h b/Engine/Plugins/Cameras/GameplayCameras/Source/GameplayCamerasEditor/Private/Toolkits/CameraModeEditorToolkit.h
deleted file mode 100644
diff --git a/Engine/Plugins/Cameras/GameplayCameras/Source/GameplayCamerasEditor/Public/AssetTools/CameraAssetEditor.h b/Engine/Plugins/Cameras/GameplayCameras/Source/GameplayCamerasEditor/Public/AssetTools/CameraAssetEditor.h
new file mode 100644
--- /dev/null
+++ b/Engine/Plugins/Cameras/GameplayCameras/Source/GameplayCamerasEditor/Public/AssetTools/CameraAssetEditor.h
@@ -0,0 +1,34 @@
+// Copyright Epic Games, Inc. All Rights Reserved.
+
+#pragma once
+
+#include "Tools/UAssetEditor.h"
+
+#include "CameraAssetEditor.generated.h"
+
+class FBaseAssetToolkit;
+class UCameraAsset;
+
+/**
+ * Editor for a camera asset.
+ */
+UCLASS(Transient)
+class UCameraAssetEditor : public UAssetEditor
+{
+	GENERATED_BODY()
+
+public:
+
+	void Initialize(TObjectPtr<UCameraAsset> InCameraAsset);
+
+public:
+
+	// UAssetEditor interface
+	virtual void GetObjectsToEdit(TArray<UObject*>& InObjectsToEdit) override;
+	virtual TSharedPtr<FBaseAssetToolkit> CreateToolkit() override;
+
+private:
+
+	TObjectPtr<UCameraAsset> CameraAsset;
+};
+
diff --git a/Engine/Plugins/Cameras/GameplayCameras/Source/GameplayCamerasEditor/Public/AssetTools/CameraModeAssetEditor.h b/Engine/Plugins/Cameras/GameplayCameras/Source/GameplayCamerasEditor/Public/AssetTools/CameraModeAssetEditor.h
new file mode 100644
--- /dev/null
+++ b/Engine/Plugins/Cameras/GameplayCameras/Source/GameplayCamerasEditor/Public/AssetTools/CameraModeAssetEditor.h
@@ -0,0 +1,34 @@
+// Copyright Epic Games, Inc. All Rights Reserved.
+
+#pragma once
+
+#include "Tools/UAssetEditor.h"
+
+#include "CameraModeAssetEditor.generated.h"
+
+class FBaseAssetToolkit;
+class UCameraMode;
+
+/**
+ * Editor for a camera mode asset.
+ */
+UCLASS(Transient)
+class UCameraModeAssetEditor : public UAssetEditor
+{
+	GENERATED_BODY()
+
+public:
+
+	void Initialize(TObjectPtr<UCameraMode> InCameraModeAsset);
+
+public:
+
+	// UAssetEditor interface
+	virtual void GetObjectsToEdit(TArray<UObject*>& InObjectsToEdit) override;
+	virtual TSharedPtr<FBaseAssetToolkit> CreateToolkit() override;
+
+private:
+
+	TObjectPtr<UCameraMode> CameraModeAsset;
+};
+
diff --git a/Engine/Plugins/Cameras/GameplayCameras/Source/GameplayCamerasEditor/Public/IGameplayCamerasEditorModule.h b/Engine/Plugins/Cameras/GameplayCameras/Source/GameplayCamerasEditor/Public/IGameplayCamerasEditorModule.h
--- a/Engine/Plugins/Cameras/GameplayCameras/Source/GameplayCamerasEditor/Public/IGameplayCamerasEditorModule.h
+++ b/Engine/Plugins/Cameras/GameplayCameras/Source/GameplayCamerasEditor/Public/IGameplayCamerasEditorModule.h
@@ -6,11 +6,14 @@
 #include "Modules/ModuleInterface.h"
 #include "Toolkits/AssetEditorToolkit.h"
 
-class ICameraAssetEditorToolkit;
-class ICameraModeEditorToolkit;
 class UCameraAsset;
+class UCameraAssetEditor;
 class UCameraMode;
+class UCameraModeAssetEditor;
 
+/**
+ * The gameplay cameras editor module.
+ */
 class IGameplayCamerasEditorModule : public IModuleInterface
 {
 public:
@@ -19,8 +22,10 @@ public:
 
 	virtual ~IGameplayCamerasEditorModule() = default;
 
-	virtual TSharedRef<ICameraAssetEditorToolkit> CreateCameraAssetEditor(const EToolkitMode::Type Mode, const TSharedPtr<IToolkitHost>& InitToolkitHost, UCameraAsset* CameraAsset) = 0;
+	/** Creates an editor for the given camera asset */
+	virtual UCameraAssetEditor* CreateCameraAssetEditor(const EToolkitMode::Type Mode, const TSharedPtr<IToolkitHost>& InitToolkitHost, UCameraAsset* CameraAsset) = 0;
 
-	virtual TSharedRef<ICameraModeEditorToolkit> CreateCameraModeEditor(const EToolkitMode::Type Mode, const TSharedPtr<IToolkitHost>& InitToolkitHost, UCameraMode* CameraMode) = 0;
+	/** Creates an editor for the given camera mode asset */
+	virtual UCameraModeAssetEditor* CreateCameraModeEditor(const EToolkitMode::Type Mode, const TSharedPtr<IToolkitHost>& InitToolkitHost, UCameraMode* CameraMode) = 0;
 };
 
diff --git a/Engine/Plugins/Cameras/GameplayCameras/Source/GameplayCamerasEditor/Public/Toolkits/CameraAssetEditorToolkit.h b/Engine/Plugins/Cameras/GameplayCameras/Source/GameplayCamerasEditor/Public/Toolkits/CameraAssetEditorToolkit.h
new file mode 100644
--- /dev/null
+++ b/Engine/Plugins/Cameras/GameplayCameras/Source/GameplayCamerasEditor/Public/Toolkits/CameraAssetEditorToolkit.h
@@ -0,0 +1,63 @@
+// Copyright Epic Games, Inc. All Rights Reserved.
+
+#pragma once
+
+#include "CoreMinimal.h"
+#include "Tools/BaseAssetToolkit.h"
+#include "UObject/GCObject.h"
+
+class IMessageLogListing;
+class UCameraAsset;
+class UCameraAssetEditor;
+class SWidget;
+
+/**
+ * Editor toolkit for a camera asset.
+ */
+class FCameraAssetEditorToolkit 
+	: public FBaseAssetToolkit
+	, public FGCObject
+{
+public:
+
+	FCameraAssetEditorToolkit(UCameraAssetEditor* InOwningAssetEditor);
+	~FCameraAssetEditorToolkit();
+
+protected:
+
+	// FBaseAssetToolkit interface
+	virtual void RegisterTabSpawners(const TSharedRef<class FTabManager>& TabManager) override;
+	virtual void UnregisterTabSpawners(const TSharedRef<FTabManager>& InTabManager) override;
+	virtual void CreateWidgets() override;
+	virtual void PostInitAssetEditor() override;
+
+	// IToolkit interface
+	virtual FText GetBaseToolkitName() const override;
+	virtual FName GetToolkitFName() const override;
+	virtual FString GetWorldCentricTabPrefix() const override;
+	virtual FLinearColor GetWorldCentricTabColorScale() const override;
+
+	// FGCObject interface
+	virtual void AddReferencedObjects(FReferenceCollector& Collector) override
+	{
+		Collector.AddReferencedObject(CameraAsset);
+	}
+	virtual FString GetReferencerName() const override
+	{
+		return TEXT("FCameraAssetEditorToolkit");
+	}
+
+private:
+
+	static const FName DetailsViewTabId;
+
+	/** Message log widget */
+	TSharedPtr<SWidget> Stats;
+
+	/** Message log listing */
+	TSharedPtr<IMessageLogListing> StatsListing;
+
+	/** The asset being edited */
+	TObjectPtr<UCameraAsset> CameraAsset;
+};
+
diff --git a/Engine/Plugins/Cameras/GameplayCameras/Source/GameplayCamerasEditor/Public/Toolkits/CameraModeAssetEditorToolkit.h b/Engine/Plugins/Cameras/GameplayCameras/Source/GameplayCamerasEditor/Public/Toolkits/CameraModeAssetEditorToolkit.h
new file mode 100644
--- /dev/null
+++ b/Engine/Plugins/Cameras/GameplayCameras/Source/GameplayCamerasEditor/Public/Toolkits/CameraModeAssetEditorToolkit.h
@@ -0,0 +1,63 @@
+// Copyright Epic Games, Inc. All Rights Reserved.
+
+#pragma once
+
+#include "CoreMinimal.h"
+#include "Tools/BaseAssetToolkit.h"
+#include "UObject/GCObject.h"
+
+class IMessageLogListing;
+class UCameraMode;
+class UCameraModeAssetEditor;
+class SWidget;
+
+/**
+ * Editor toolkit for a camera mode asset.
+ */
+class FCameraModeAssetEditorToolkit 
+	: public FBaseAssetToolkit
+	, public FGCObject
+{
+public:
+
+	FCameraModeAssetEditorToolkit(UCameraModeAssetEditor* InOwningAssetEditor);
+	~FCameraModeAssetEditorToolkit();
+
+protected:
+
+	// FBaseAssetToolkit interface
+	virtual void RegisterTabSpawners(const TSharedRef<class FTabManager>& TabManager) override;
+	virtual void UnregisterTabSpawners(const TSharedRef<FTabManager>& InTabManager) override;
+	virtual void CreateWidgets() override;
+	virtual void PostInitAssetEditor() override;
+
+	// IToolkit interface
+	virtual FText GetBaseToolkitName() const override;
+	virtual FName GetToolkitFName() const override;
+	virtual FString GetWorldCentricTabPrefix() const override;
+	virtual FLinearColor GetWorldCentricTabColorScale() const override;
+
+	// FGCObject interface
+	virtual void AddReferencedObjects(FReferenceCollector& Collector) override
+	{
+		Collector.AddReferencedObject(CameraModeAsset);
+	}
+	virtual FString GetReferencerName() const override
+	{
+		return TEXT("FCameraModeAssetEditorToolkit");
+	}
+
+private:
+
+	static const FName DetailsViewTabId;
+
+	/** Message log widget */
+	TSharedPtr<SWidget> Stats;
+
+	/** Message log listing */
+	TSharedPtr<IMessageLogListing> StatsListing;
+
+	/** The asset being edited */
+	TObjectPtr<UCameraMode> CameraModeAsset;
+};
+
diff --git a/Engine/Plugins/Cameras/GameplayCameras/Source/GameplayCamerasEditor/Public/Toolkits/ICameraAssetEditorToolkit.h b/Engine/Plugins/Cameras/GameplayCameras/Source/GameplayCamerasEditor/Public/Toolkits/ICameraAssetEditorToolkit.h
deleted file mode 100644
diff --git a/Engine/Plugins/Cameras/GameplayCameras/Source/GameplayCamerasEditor/Public/Toolkits/ICameraModeEditorToolkit.h b/Engine/Plugins/Cameras/GameplayCameras/Source/GameplayCamerasEditor/Public/Toolkits/ICameraModeEditorToolkit.h
deleted file mode 100644
