Cameras: add camera variable assets, and camera parameters that can be driven by these variables

Camera variable assets are ways for users to safely define the name and type of a camera variable, which can then be placed in the variable table. This system will be usable as a "blackboard" of sorts, where gameplay values can be pushed into the camera system and used there to drive camera behaviors. This can be done with camera parameters, which allow camera nodes to accept both fixed values and references to variable assets.

diff --git a/Engine/Plugins/Cameras/GameplayCameras/Source/GameplayCameras/Private/Core/CameraParameters.cpp b/Engine/Plugins/Cameras/GameplayCameras/Source/GameplayCameras/Private/Core/CameraParameters.cpp
new file mode 100644
--- /dev/null
+++ b/Engine/Plugins/Cameras/GameplayCameras/Source/GameplayCameras/Private/Core/CameraParameters.cpp
@@ -0,0 +1,161 @@
+// Copyright Epic Games, Inc. All Rights Reserved.
+
+#include "Core/CameraParameters.h"
+
+#include "UObject/UnrealNames.h"
+
+bool FBooleanCameraParameter::SerializeFromMismatchedTag(const FPropertyTag& Tag, FStructuredArchive::FSlot Slot)
+{
+	if (Tag.Type == NAME_BoolProperty)
+	{
+		Slot << Value;
+		return true;
+	}
+
+	return false;
+}
+
+bool FInteger32CameraParameter::SerializeFromMismatchedTag(const FPropertyTag& Tag, FStructuredArchive::FSlot Slot)
+{
+	if (Tag.Type == NAME_IntProperty || Tag.Type == NAME_Int32Property)
+	{
+		Slot << Value;
+		return true;
+	}
+
+	return false;
+}
+
+bool FFloatCameraParameter::SerializeFromMismatchedTag(const FPropertyTag& Tag, FStructuredArchive::FSlot Slot)
+{
+	if (Tag.Type == NAME_FloatProperty)
+	{
+		Slot << Value;
+		return true;
+	}
+
+	return false;
+}
+
+
+bool FDoubleCameraParameter::SerializeFromMismatchedTag(const FPropertyTag& Tag, FStructuredArchive::FSlot Slot)
+{
+	if (Tag.Type == NAME_StructProperty && Tag.StructName == NAME_DoubleProperty)
+	{
+		Slot << Value;
+		return true;
+	}
+
+	return false;
+}
+
+bool FVector2fCameraParameter::SerializeFromMismatchedTag(const FPropertyTag& Tag, FStructuredArchive::FSlot Slot)
+{
+	if (Tag.Type == NAME_StructProperty && (Tag.StructName == NAME_Vector2f || Tag.StructName == NAME_Vector2D))
+	{
+		Slot << Value;
+		return true;
+	}
+
+	return false;
+}
+
+bool FVector2dCameraParameter::SerializeFromMismatchedTag(const FPropertyTag& Tag, FStructuredArchive::FSlot Slot)
+{
+	if (Tag.Type == NAME_StructProperty && Tag.StructName == NAME_Vector2d)
+	{
+		Slot << Value;
+		return true;
+	}
+
+	return false;
+}
+
+bool FVector3fCameraParameter::SerializeFromMismatchedTag(const FPropertyTag& Tag, FStructuredArchive::FSlot Slot)
+{
+	if (Tag.Type == NAME_StructProperty && Tag.StructName == NAME_Vector3f)
+	{
+		Slot << Value;
+		return true;
+	}
+
+	return false;
+}
+
+bool FVector3dCameraParameter::SerializeFromMismatchedTag(const FPropertyTag& Tag, FStructuredArchive::FSlot Slot)
+{
+	if (Tag.Type == NAME_StructProperty && Tag.StructName == NAME_Vector3d)
+	{
+		Slot << Value;
+		return true;
+	}
+
+	return false;
+}
+
+bool FVector4fCameraParameter::SerializeFromMismatchedTag(const FPropertyTag& Tag, FStructuredArchive::FSlot Slot)
+{
+	if (Tag.Type == NAME_StructProperty && Tag.StructName == NAME_Vector4f)
+	{
+		Slot << Value;
+		return true;
+	}
+
+	return false;
+}
+
+bool FVector4dCameraParameter::SerializeFromMismatchedTag(const FPropertyTag& Tag, FStructuredArchive::FSlot Slot)
+{
+	if (Tag.Type == NAME_StructProperty && Tag.StructName == NAME_Vector4d)
+	{
+		Slot << Value;
+		return true;
+	}
+
+	return false;
+}
+
+bool FRotator3fCameraParameter::SerializeFromMismatchedTag(const FPropertyTag& Tag, FStructuredArchive::FSlot Slot)
+{
+	if (Tag.Type == NAME_StructProperty && Tag.StructName == NAME_Rotator3f)
+	{
+		Slot << Value;
+		return true;
+	}
+
+	return false;
+}
+
+bool FRotator3dCameraParameter::SerializeFromMismatchedTag(const FPropertyTag& Tag, FStructuredArchive::FSlot Slot)
+{
+	if (Tag.Type == NAME_StructProperty && Tag.StructName == NAME_Rotator3d)
+	{
+		Slot << Value;
+		return true;
+	}
+
+	return false;
+}
+
+bool FTransform3fCameraParameter::SerializeFromMismatchedTag(const FPropertyTag& Tag, FStructuredArchive::FSlot Slot)
+{
+	if (Tag.Type == NAME_StructProperty && Tag.StructName == NAME_Transform3f)
+	{
+		Slot << Value;
+		return true;
+	}
+
+	return false;
+}
+
+bool FTransform3dCameraParameter::SerializeFromMismatchedTag(const FPropertyTag& Tag, FStructuredArchive::FSlot Slot)
+{
+	if (Tag.Type == NAME_StructProperty && Tag.StructName == NAME_Transform3d)
+	{
+		Slot << Value;
+		return true;
+	}
+
+	return false;
+}
+
diff --git a/Engine/Plugins/Cameras/GameplayCameras/Source/GameplayCameras/Private/Core/CameraVariableAssets.cpp b/Engine/Plugins/Cameras/GameplayCameras/Source/GameplayCameras/Private/Core/CameraVariableAssets.cpp
new file mode 100644
--- /dev/null
+++ b/Engine/Plugins/Cameras/GameplayCameras/Source/GameplayCameras/Private/Core/CameraVariableAssets.cpp
@@ -0,0 +1,11 @@
+// Copyright Epic Games, Inc. All Rights Reserved.
+
+#include "Core/CameraVariableAssets.h"
+
+#include UE_INLINE_GENERATED_CPP_BY_NAME(CameraVariableAssets)
+
+UCameraVariableAsset::UCameraVariableAsset(const FObjectInitializer& ObjectInit)
+	: Super(ObjectInit)
+{
+}
+
diff --git a/Engine/Plugins/Cameras/GameplayCameras/Source/GameplayCameras/Public/Core/CameraParameterReader.h b/Engine/Plugins/Cameras/GameplayCameras/Source/GameplayCameras/Public/Core/CameraParameterReader.h
new file mode 100644
--- /dev/null
+++ b/Engine/Plugins/Cameras/GameplayCameras/Source/GameplayCameras/Public/Core/CameraParameterReader.h
@@ -0,0 +1,72 @@
+// Copyright Epic Games, Inc. All Rights Reserved.
+
+#pragma once
+
+#include "Core/CameraVariableTable.h"
+#include "CoreTypes.h"
+#include "Templates/UnrealTypeTraits.h"
+
+template<typename ValueType>
+class FCameraParameterReader
+{
+public:
+
+	FCameraParameterReader() {}
+
+	template<typename ParameterType>
+	FCameraParameterReader(const ParameterType& Parameter)
+	{
+		Initialize(Parameter);
+	}
+
+	/**
+	 * Initializes the reader around the given parameter.
+	 */
+	template<typename ParameterType>
+	void Initialize(const ParameterType& Parameter)
+	{
+		static_assert(
+				std::is_same<ValueType, typename ParameterType::ValueType>(),
+				"The given parameter is of the wrong type for this reader! Value types must be the same.");
+
+		DefaultValuePtr = &Parameter.Value;
+		ensureMsgf(DefaultValuePtr, TEXT("The given parameter doesn't have a value!"));
+
+		if (Parameter.Variable)
+		{
+			VariableId = Parameter.Variable->GetVariableId();
+			DefaultValuePtr = reinterpret_cast<const ValueType*>(Parameter.Variable->GetDefaultValuePtr());
+			ensureMsgf(DefaultValuePtr, TEXT("The given parameter's driving variable doesn't have a default value!"));
+		}
+	}
+
+	/**
+	 * Gets the actual value for the parameter.
+	 */
+	typename TCallTraits<ValueType>::ParamType Get(const FCameraVariableTable& VariableTable) const
+	{
+		checkf(DefaultValuePtr, TEXT("Parameter reader has no value pointer!"));
+		if (VariableId == 0)
+		{
+			// No variable is driving the parameter, just return the parameter value.
+			return *DefaultValuePtr;
+		}
+		else
+		{
+			// The parameter is driven by a variable. Find it in the variable table.
+			if (const ValueType* ActualValue = VariableTable.FindValue<ValueType>(VariableId))
+			{
+				return *ActualValue;
+			}
+			return *DefaultValuePtr;
+		}
+	}
+
+private:
+
+	/** Pointer to the value in the parameter. */
+	const ValueType* DefaultValuePtr = nullptr;
+	/** The ID of the variable driving the parameter, if any. */
+	uint32 VariableId = 0;
+};
+
diff --git a/Engine/Plugins/Cameras/GameplayCameras/Source/GameplayCameras/Public/Core/CameraParameters.h b/Engine/Plugins/Cameras/GameplayCameras/Source/GameplayCameras/Public/Core/CameraParameters.h
new file mode 100644
--- /dev/null
+++ b/Engine/Plugins/Cameras/GameplayCameras/Source/GameplayCameras/Public/Core/CameraParameters.h
@@ -0,0 +1,262 @@
+// Copyright Epic Games, Inc. All Rights Reserved.
+
+#pragma once
+
+#include "Core/CameraVariableAssets.h"
+#include "CoreTypes.h"
+#include "Math/MathFwd.h"
+
+#include "CameraParameters.generated.h"
+
+/** Boolean camera parameter. */
+USTRUCT()
+struct FBooleanCameraParameter
+{
+	GENERATED_BODY()
+
+	using ValueType = bool;
+
+	UPROPERTY(EditAnywhere, Category=Common)
+	bool Value = false;
+
+	UPROPERTY(EditAnywhere, Category=Common)
+	TObjectPtr<UBooleanCameraVariable> Variable;
+
+	bool SerializeFromMismatchedTag(const FPropertyTag& Tag, FStructuredArchive::FSlot Slot);
+};
+
+/** Integer camera parameter. */
+USTRUCT()
+struct FInteger32CameraParameter
+{
+	GENERATED_BODY()
+
+	using ValueType = int32;
+
+	UPROPERTY(EditAnywhere, Category=Common)
+	int32 Value = 0;
+
+	UPROPERTY(EditAnywhere, Category=Common)
+	TObjectPtr<UInteger32CameraVariable> Variable;
+
+	bool SerializeFromMismatchedTag(const FPropertyTag& Tag, FStructuredArchive::FSlot Slot);
+};
+
+/** Float camera parameter. */
+USTRUCT()
+struct FFloatCameraParameter
+{
+	GENERATED_BODY()
+
+	using ValueType = float;
+
+	UPROPERTY(EditAnywhere, Category=Common)
+	float Value = 0.f;
+
+	UPROPERTY(EditAnywhere, Category=Common)
+	TObjectPtr<UFloatCameraVariable> Variable;
+
+	bool SerializeFromMismatchedTag(const FPropertyTag& Tag, FStructuredArchive::FSlot Slot);
+};
+
+/** Double camera parameter. */
+USTRUCT()
+struct FDoubleCameraParameter
+{
+	GENERATED_BODY()
+
+	using ValueType = double;
+
+	UPROPERTY(EditAnywhere, Category=Common)
+	double Value = 0.0;
+
+	UPROPERTY(EditAnywhere, Category=Common)
+	TObjectPtr<UDoubleCameraVariable> Variable;
+
+	bool SerializeFromMismatchedTag(const FPropertyTag& Tag, FStructuredArchive::FSlot Slot);
+};
+
+/** Vector2f camera parameter. */
+USTRUCT()
+struct FVector2fCameraParameter
+{
+	GENERATED_BODY()
+
+	using ValueType = FVector2f;
+
+	UPROPERTY(EditAnywhere, Category=Common)
+	FVector2f Value;
+
+	UPROPERTY(EditAnywhere, Category=Common)
+	TObjectPtr<UVector2fCameraVariable> Variable;
+
+	bool SerializeFromMismatchedTag(const FPropertyTag& Tag, FStructuredArchive::FSlot Slot);
+};
+
+/** Vector2d camera parameter. */
+USTRUCT()
+struct FVector2dCameraParameter
+{
+	GENERATED_BODY()
+
+	using ValueType = FVector2D;
+
+	UPROPERTY(EditAnywhere, Category=Common)
+	FVector2D Value;
+
+	UPROPERTY(EditAnywhere, Category=Common)
+	TObjectPtr<UVector2dCameraVariable> Variable;
+
+	bool SerializeFromMismatchedTag(const FPropertyTag& Tag, FStructuredArchive::FSlot Slot);
+};
+
+/** Vector3f camera parameter. */
+USTRUCT()
+struct FVector3fCameraParameter
+{
+	GENERATED_BODY()
+
+	using ValueType = FVector3f;
+
+	UPROPERTY(EditAnywhere, Category=Common)
+	FVector3f Value;
+
+	UPROPERTY(EditAnywhere, Category=Common)
+	TObjectPtr<UVector3fCameraVariable> Variable;
+
+	bool SerializeFromMismatchedTag(const FPropertyTag& Tag, FStructuredArchive::FSlot Slot);
+};
+
+/** Vector3d camera parameter. */
+USTRUCT()
+struct FVector3dCameraParameter
+{
+	GENERATED_BODY()
+
+	using ValueType = FVector3d;
+
+	UPROPERTY(EditAnywhere, Category=Common)
+	FVector3d Value;
+
+	UPROPERTY(EditAnywhere, Category=Common)
+	TObjectPtr<UVector3dCameraVariable> Variable;
+
+	bool SerializeFromMismatchedTag(const FPropertyTag& Tag, FStructuredArchive::FSlot Slot);
+};
+
+/** Vector4f camera parameter. */
+USTRUCT()
+struct FVector4fCameraParameter
+{
+	GENERATED_BODY()
+
+	using ValueType = FVector4f;
+
+	UPROPERTY(EditAnywhere, Category=Common)
+	FVector4f Value;
+
+	UPROPERTY(EditAnywhere, Category=Common)
+	TObjectPtr<UVector4fCameraVariable> Variable;
+
+	bool SerializeFromMismatchedTag(const FPropertyTag& Tag, FStructuredArchive::FSlot Slot);
+};
+
+/** Vector4d camera parameter. */
+USTRUCT()
+struct FVector4dCameraParameter
+{
+	GENERATED_BODY()
+
+	using ValueType = FVector4d;
+
+	UPROPERTY(EditAnywhere, Category=Common)
+	FVector4d Value;
+
+	UPROPERTY(EditAnywhere, Category=Common)
+	TObjectPtr<UVector4dCameraVariable> Variable;
+
+	bool SerializeFromMismatchedTag(const FPropertyTag& Tag, FStructuredArchive::FSlot Slot);
+};
+
+/** Rotator3f camera parameter. */
+USTRUCT()
+struct FRotator3fCameraParameter
+{
+	GENERATED_BODY()
+
+	using ValueType = FRotator3f;
+
+	UPROPERTY(EditAnywhere, Category=Common)
+	FRotator3f Value;
+
+	UPROPERTY(EditAnywhere, Category=Common)
+	TObjectPtr<URotator3fCameraVariable> Variable;
+
+	bool SerializeFromMismatchedTag(const FPropertyTag& Tag, FStructuredArchive::FSlot Slot);
+};
+
+/** Rotator3d camera parameter. */
+USTRUCT()
+struct FRotator3dCameraParameter
+{
+	GENERATED_BODY()
+
+	using ValueType = FRotator3d;
+
+	UPROPERTY(EditAnywhere, Category=Common)
+	FRotator3d Value;
+
+	UPROPERTY(EditAnywhere, Category=Common)
+	TObjectPtr<URotator3dCameraVariable> Variable;
+
+	bool SerializeFromMismatchedTag(const FPropertyTag& Tag, FStructuredArchive::FSlot Slot);
+};
+
+/** Transform3f camera parameter. */
+USTRUCT()
+struct FTransform3fCameraParameter
+{
+	GENERATED_BODY()
+
+	using ValueType = FTransform3f;
+
+	UPROPERTY(EditAnywhere, Category=Common)
+	FTransform3f Value;
+
+	UPROPERTY(EditAnywhere, Category=Common)
+	TObjectPtr<UTransform3fCameraVariable> Variable;
+
+	bool SerializeFromMismatchedTag(const FPropertyTag& Tag, FStructuredArchive::FSlot Slot);
+};
+
+/** Transform3d camera parameter. */
+USTRUCT()
+struct FTransform3dCameraParameter
+{
+	GENERATED_BODY()
+
+	using ValueType = FTransform3d;
+
+	UPROPERTY(EditAnywhere, Category=Common)
+	FTransform3d Value;
+
+	UPROPERTY(EditAnywhere, Category=Common)
+	TObjectPtr<UTransform3dCameraVariable> Variable;
+
+	bool SerializeFromMismatchedTag(const FPropertyTag& Tag, FStructuredArchive::FSlot Slot);
+};
+
+// Any camera parameter might replace a previously non-parameterized property (i.e. a "fixed" property
+// of the underlying type, like bool, int32, float, etc.)
+// When someone upgrades the fixed property to a parameterized property, any previously saved data will
+// run into a mismatched tag. So the parameters will handle that by loading the saved value inside of
+// them.
+#define UE_CAMERA_VARIABLE_FOR_TYPE(ValueType, ValueName)\
+	template<> struct TStructOpsTypeTraits<F##ValueName##CameraParameter>\
+		: public TStructOpsTypeTraitsBase2<F##ValueName##CameraParameter>\
+	{\
+		enum { WithStructuredSerializeFromMismatchedTag = true };\
+	};
+UE_CAMERA_VARIABLE_FOR_ALL_TYPES()
+#undef UE_CAMERA_VARIABLE_FOR_TYPE
+
diff --git a/Engine/Plugins/Cameras/GameplayCameras/Source/GameplayCameras/Public/Core/CameraVariableAssets.h b/Engine/Plugins/Cameras/GameplayCameras/Source/GameplayCameras/Public/Core/CameraVariableAssets.h
new file mode 100644
--- /dev/null
+++ b/Engine/Plugins/Cameras/GameplayCameras/Source/GameplayCameras/Public/Core/CameraVariableAssets.h
@@ -0,0 +1,295 @@
+// Copyright Epic Games, Inc. All Rights Reserved.
+
+#pragma once
+
+#include "Core/CameraVariableTableFwd.h"
+#include "CoreTypes.h"
+#include "Engine/DataAsset.h"
+#include "Math/MathFwd.h"
+
+#include "CameraVariableAssets.generated.h"
+
+/**
+ * The base asset class for all camera variables.
+ */
+UCLASS(Abstract)
+class GAMEPLAYCAMERAS_API UCameraVariableAsset : public UDataAsset
+{
+	GENERATED_BODY()
+
+public:
+
+	UCameraVariableAsset(const FObjectInitializer& ObjectInit);
+
+	uint32 GetVariableId() const { return VariableId; }
+
+	virtual ECameraVariableType GetVariableType() const PURE_VIRTUAL(UCameraVariableAsset::GetVariableType, return ECameraVariableType::Boolean;);
+	virtual const uint8* GetDefaultValuePtr() const PURE_VIRTUAL(UCameraVariableAsset::GetDefaultValuePtr, return nullptr;);
+
+public:
+
+	/** Whether this variable auto-resets to its default value every frame. */
+	UPROPERTY(EditAnywhere, Category=Camera)
+	bool bAutoReset = false;
+
+private:
+
+	UPROPERTY()
+	uint32 VariableId = 0;
+
+	UPROPERTY()
+	bool bIsLocal = false;
+};
+
+/** Boolean camera variable. */
+UCLASS()
+class GAMEPLAYCAMERAS_API UBooleanCameraVariable : public UCameraVariableAsset
+{
+	GENERATED_BODY()
+
+public:
+
+	virtual ECameraVariableType GetVariableType() const override { return ECameraVariableType::Boolean; }
+	virtual const uint8* GetDefaultValuePtr() const override { return reinterpret_cast<const uint8*>(&bDefaultValue); }
+
+public:
+
+	/** The default value of this variable. */
+	UPROPERTY(EditAnywhere, Category=Camera)
+	bool bDefaultValue = false;
+};
+
+/** Integer camera variable. */
+UCLASS()
+class GAMEPLAYCAMERAS_API UInteger32CameraVariable : public UCameraVariableAsset
+{
+	GENERATED_BODY()
+
+public:
+
+	virtual ECameraVariableType GetVariableType() const override { return ECameraVariableType::Integer32; }
+	virtual const uint8* GetDefaultValuePtr() const override { return reinterpret_cast<const uint8*>(&DefaultValue); }
+
+public:
+
+	/** The default value of this variable. */
+	UPROPERTY(EditAnywhere, Category=Camera)
+	int32 DefaultValue = 0;
+};
+
+/** Float camera variable. */
+UCLASS()
+class GAMEPLAYCAMERAS_API UFloatCameraVariable : public UCameraVariableAsset
+{
+	GENERATED_BODY()
+
+public:
+
+	virtual ECameraVariableType GetVariableType() const override { return ECameraVariableType::Float; }
+	virtual const uint8* GetDefaultValuePtr() const override { return reinterpret_cast<const uint8*>(&DefaultValue); }
+
+public:
+
+	/** The default value of this variable. */
+	UPROPERTY(EditAnywhere, Category=Camera)
+	float DefaultValue = 0.f;
+};
+
+/** Double camera variable. */
+UCLASS()
+class GAMEPLAYCAMERAS_API UDoubleCameraVariable : public UCameraVariableAsset
+{
+	GENERATED_BODY()
+
+public:
+
+	virtual ECameraVariableType GetVariableType() const override { return ECameraVariableType::Double; }
+	virtual const uint8* GetDefaultValuePtr() const override { return reinterpret_cast<const uint8*>(&DefaultValue); }
+
+public:
+
+	/** The default value of this variable. */
+	UPROPERTY(EditAnywhere, Category=Camera)
+	double DefaultValue = 0.0;
+};
+
+/** Vector2f camera variable. */
+UCLASS()
+class GAMEPLAYCAMERAS_API UVector2fCameraVariable : public UCameraVariableAsset
+{
+	GENERATED_BODY()
+
+public:
+
+	virtual ECameraVariableType GetVariableType() const override { return ECameraVariableType::Vector2f; }
+	virtual const uint8* GetDefaultValuePtr() const override { return reinterpret_cast<const uint8*>(&DefaultValue); }
+
+public:
+
+	/** The default value of this variable. */
+	UPROPERTY(EditAnywhere, Category=Camera)
+	FVector2f DefaultValue;
+};
+
+/** Vector2d camera variable. */
+UCLASS()
+class GAMEPLAYCAMERAS_API UVector2dCameraVariable : public UCameraVariableAsset
+{
+	GENERATED_BODY()
+
+public:
+
+	virtual ECameraVariableType GetVariableType() const override { return ECameraVariableType::Vector2d; }
+	virtual const uint8* GetDefaultValuePtr() const override { return reinterpret_cast<const uint8*>(&DefaultValue); }
+
+public:
+
+	/** The default value of this variable. */
+	UPROPERTY(EditAnywhere, Category=Camera)
+	FVector2D DefaultValue;
+};
+
+/** Vector3f camera variable. */
+UCLASS()
+class GAMEPLAYCAMERAS_API UVector3fCameraVariable : public UCameraVariableAsset
+{
+	GENERATED_BODY()
+
+public:
+
+	virtual ECameraVariableType GetVariableType() const override { return ECameraVariableType::Vector3f; }
+	virtual const uint8* GetDefaultValuePtr() const override { return reinterpret_cast<const uint8*>(&DefaultValue); }
+
+public:
+
+	/** The default value of this variable. */
+	UPROPERTY(EditAnywhere, Category=Camera)
+	FVector3f DefaultValue;
+};
+
+/** Vector3d camera variable. */
+UCLASS()
+class GAMEPLAYCAMERAS_API UVector3dCameraVariable : public UCameraVariableAsset
+{
+	GENERATED_BODY()
+
+public:
+
+	virtual ECameraVariableType GetVariableType() const override { return ECameraVariableType::Vector3d; }
+	virtual const uint8* GetDefaultValuePtr() const override { return reinterpret_cast<const uint8*>(&DefaultValue); }
+
+public:
+
+	/** The default value of this variable. */
+	UPROPERTY(EditAnywhere, Category=Camera)
+	FVector3d DefaultValue;
+};
+
+/** Vector4f camera variable. */
+UCLASS()
+class GAMEPLAYCAMERAS_API UVector4fCameraVariable : public UCameraVariableAsset
+{
+	GENERATED_BODY()
+
+public:
+
+	virtual ECameraVariableType GetVariableType() const override { return ECameraVariableType::Vector4f; }
+	virtual const uint8* GetDefaultValuePtr() const override { return reinterpret_cast<const uint8*>(&DefaultValue); }
+
+public:
+
+	/** The default value of this variable. */
+	UPROPERTY(EditAnywhere, Category=Camera)
+	FVector4f DefaultValue;
+};
+
+/** Vector4d camera variable. */
+UCLASS()
+class GAMEPLAYCAMERAS_API UVector4dCameraVariable : public UCameraVariableAsset
+{
+	GENERATED_BODY()
+
+public:
+
+	virtual ECameraVariableType GetVariableType() const override { return ECameraVariableType::Vector4d; }
+	virtual const uint8* GetDefaultValuePtr() const override { return reinterpret_cast<const uint8*>(&DefaultValue); }
+
+public:
+
+	/** The default value of this variable. */
+	UPROPERTY(EditAnywhere, Category=Camera)
+	FVector4d DefaultValue;
+};
+
+/** Rotator3f camera variable. */
+UCLASS()
+class GAMEPLAYCAMERAS_API URotator3fCameraVariable : public UCameraVariableAsset
+{
+	GENERATED_BODY()
+
+public:
+
+	virtual ECameraVariableType GetVariableType() const override { return ECameraVariableType::Rotator3f; }
+	virtual const uint8* GetDefaultValuePtr() const override { return reinterpret_cast<const uint8*>(&DefaultValue); }
+
+public:
+
+	/** The default value of this variable. */
+	UPROPERTY(EditAnywhere, Category=Camera)
+	FRotator3f DefaultValue;
+};
+
+/** Rotator3d camera variable. */
+UCLASS()
+class GAMEPLAYCAMERAS_API URotator3dCameraVariable : public UCameraVariableAsset
+{
+	GENERATED_BODY()
+
+public:
+
+	virtual ECameraVariableType GetVariableType() const override { return ECameraVariableType::Rotator3d; }
+	virtual const uint8* GetDefaultValuePtr() const override { return reinterpret_cast<const uint8*>(&DefaultValue); }
+
+public:
+
+	/** The default value of this variable. */
+	UPROPERTY(EditAnywhere, Category=Camera)
+	FRotator3d DefaultValue;
+};
+
+/** Transform3f camera variable. */
+UCLASS()
+class GAMEPLAYCAMERAS_API UTransform3fCameraVariable : public UCameraVariableAsset
+{
+	GENERATED_BODY()
+
+public:
+
+	virtual ECameraVariableType GetVariableType() const override { return ECameraVariableType::Transform3f; }
+	virtual const uint8* GetDefaultValuePtr() const override { return reinterpret_cast<const uint8*>(&DefaultValue); }
+
+public:
+
+	/** The default value of this variable. */
+	UPROPERTY(EditAnywhere, Category=Camera)
+	FTransform3f DefaultValue;
+};
+
+/** Transform3d camera variable. */
+UCLASS()
+class GAMEPLAYCAMERAS_API UTransform3dCameraVariable : public UCameraVariableAsset
+{
+	GENERATED_BODY()
+
+public:
+
+	virtual ECameraVariableType GetVariableType() const override { return ECameraVariableType::Transform3d; }
+	virtual const uint8* GetDefaultValuePtr() const override { return reinterpret_cast<const uint8*>(&DefaultValue); }
+
+public:
+
+	/** The default value of this variable. */
+	UPROPERTY(EditAnywhere, Category=Camera)
+	FTransform3d DefaultValue;
+};
+
diff --git a/Engine/Plugins/Cameras/GameplayCameras/Source/GameplayCameras/Public/Core/CameraVariableTable.h b/Engine/Plugins/Cameras/GameplayCameras/Source/GameplayCameras/Public/Core/CameraVariableTable.h
--- a/Engine/Plugins/Cameras/GameplayCameras/Source/GameplayCameras/Public/Core/CameraVariableTable.h
+++ b/Engine/Plugins/Cameras/GameplayCameras/Source/GameplayCameras/Public/Core/CameraVariableTable.h
@@ -116,14 +116,14 @@ private:
 
 	static bool GetVariableTypeAllocationInfo(ECameraVariableType VariableType, uint32& OutSizeOf, uint32& OutAlignOf);
 
-	void ReallocateBuffer(uint32 NewCapacity);
-
 	template<typename ValueType>
-	bool CheckVariableType(ECameraVariableType InType)
+	static bool CheckVariableType(ECameraVariableType InType)
 	{
 		return ensure(TCameraVariableTraits<ValueType>::Type == InType);
 	}
 
+	void ReallocateBuffer(uint32 NewCapacity);
+
 	void InternalOverrideChanged(const FCameraVariableTable& OtherTable, const FCameraVariableTableFlags* InMask, bool bInvertMask, FCameraVariableTableFlags* OutMask);
 	void InternalLerpChanged(const FCameraVariableTable& ToTable, float Factor, const FCameraVariableTableFlags* InMask, bool bInvertMask, FCameraVariableTableFlags* OutMask);
 
@@ -226,7 +226,7 @@ bool FCameraVariableTable::TrySetValue(uint32 VariableId, typename TCallTraits<V
 	template<>\
 	struct TCameraVariableTraits<ValueType>\
 	{\
-		ECameraVariableType Type = ECameraVariableType::ValueName;\
+		static const ECameraVariableType Type = ECameraVariableType::ValueName;\
 	};
 UE_CAMERA_VARIABLE_FOR_ALL_TYPES()
 #undef UE_CAMERA_VARIABLE_FOR_TYPE
diff --git a/Engine/Plugins/Cameras/GameplayCameras/Source/GameplayCamerasEditor/Private/AssetTools/AssetDefinition_CameraVariableAssets.cpp b/Engine/Plugins/Cameras/GameplayCameras/Source/GameplayCamerasEditor/Private/AssetTools/AssetDefinition_CameraVariableAssets.cpp
new file mode 100644
--- /dev/null
+++ b/Engine/Plugins/Cameras/GameplayCameras/Source/GameplayCamerasEditor/Private/AssetTools/AssetDefinition_CameraVariableAssets.cpp
@@ -0,0 +1,44 @@
+// Copyright Epic Games, Inc. All Rights Reserved.
+
+#include "AssetTools/AssetDefinition_CameraVariableAssets.h"
+
+#include "Core/CameraVariableAssets.h"
+#include "Toolkits/SimpleAssetEditor.h"
+
+#define LOCTEXT_NAMESPACE "AssetDefinition_CameraVariableAssets"
+
+#include UE_INLINE_GENERATED_CPP_BY_NAME(AssetDefinition_CameraVariableAssets)
+
+FLinearColor UAssetDefinition_CameraVariableAsset::GetAssetColor() const
+{
+	return FLinearColor(FColor(200, 80, 80));
+}
+
+TConstArrayView<FAssetCategoryPath> UAssetDefinition_CameraVariableAsset::GetAssetCategories() const
+{
+	static const auto Categories = { FAssetCategoryPath(EAssetCategoryPaths::Gameplay, LOCTEXT("AssetSubCategory", "Cameras")) };
+	return Categories;
+}
+
+EAssetCommandResult UAssetDefinition_CameraVariableAsset::OpenAssets(const FAssetOpenArgs& OpenArgs) const
+{
+	TArray<UObject*> Objects = OpenArgs.LoadObjects<UObject>();
+	FSimpleAssetEditor::CreateEditor(OpenArgs.GetToolkitMode(), OpenArgs.ToolkitHost, Objects);
+
+	return EAssetCommandResult::Handled;
+}
+
+#define UE_CAMERA_VARIABLE_FOR_TYPE(ValueType, ValueName)\
+	FText UAssetDefinition_##ValueName##CameraVariable::GetAssetDisplayName() const\
+	{\
+		return LOCTEXT("AssetDisplayName", "#ValueName Camera Variable");\
+	}\
+	TSoftClassPtr<UObject> UAssetDefinition_##ValueName##CameraVariable::GetAssetClass() const\
+	{\
+		return U##ValueName##CameraVariable::StaticClass();\
+	}
+UE_CAMERA_VARIABLE_FOR_ALL_TYPES()
+#undef UE_CAMERA_VARIABLE_FOR_TYPE
+
+#undef LOCTEXT_NAMESPACE
+
diff --git a/Engine/Plugins/Cameras/GameplayCameras/Source/GameplayCamerasEditor/Private/AssetTools/AssetDefinition_CameraVariableAssets.h b/Engine/Plugins/Cameras/GameplayCameras/Source/GameplayCamerasEditor/Private/AssetTools/AssetDefinition_CameraVariableAssets.h
new file mode 100644
--- /dev/null
+++ b/Engine/Plugins/Cameras/GameplayCameras/Source/GameplayCamerasEditor/Private/AssetTools/AssetDefinition_CameraVariableAssets.h
@@ -0,0 +1,161 @@
+// Copyright Epic Games, Inc. All Rights Reserved.
+
+#pragma once
+
+#include "AssetDefinitionDefault.h"
+
+#include "AssetDefinition_CameraVariableAssets.generated.h"
+
+UCLASS(Abstract, MinimalAPI)
+class UAssetDefinition_CameraVariableAsset : public UAssetDefinitionDefault
+{
+	GENERATED_BODY()
+
+public:
+
+	// UAssetDefinition interface
+	virtual FLinearColor GetAssetColor() const override;
+	virtual TConstArrayView<FAssetCategoryPath> GetAssetCategories() const override;
+	virtual EAssetCommandResult OpenAssets(const FAssetOpenArgs& OpenArgs) const override;
+};
+
+UCLASS()
+class UAssetDefinition_BooleanCameraVariable : public UAssetDefinition_CameraVariableAsset
+{
+	GENERATED_BODY()
+
+	// UAssetDefinition interface
+	virtual FText GetAssetDisplayName() const override;
+	virtual TSoftClassPtr<UObject> GetAssetClass() const override;
+};
+
+UCLASS()
+class UAssetDefinition_Integer32CameraVariable : public UAssetDefinition_CameraVariableAsset
+{
+	GENERATED_BODY()
+
+	// UAssetDefinition interface
+	virtual FText GetAssetDisplayName() const override;
+	virtual TSoftClassPtr<UObject> GetAssetClass() const override;
+};
+
+UCLASS()
+class UAssetDefinition_FloatCameraVariable : public UAssetDefinition_CameraVariableAsset
+{
+	GENERATED_BODY()
+
+	// UAssetDefinition interface
+	virtual FText GetAssetDisplayName() const override;
+	virtual TSoftClassPtr<UObject> GetAssetClass() const override;
+};
+
+UCLASS()
+class UAssetDefinition_DoubleCameraVariable : public UAssetDefinition_CameraVariableAsset
+{
+	GENERATED_BODY()
+
+	// UAssetDefinition interface
+	virtual FText GetAssetDisplayName() const override;
+	virtual TSoftClassPtr<UObject> GetAssetClass() const override;
+};
+
+UCLASS()
+class UAssetDefinition_Vector2fCameraVariable : public UAssetDefinition_CameraVariableAsset
+{
+	GENERATED_BODY()
+
+	// UAssetDefinition interface
+	virtual FText GetAssetDisplayName() const override;
+	virtual TSoftClassPtr<UObject> GetAssetClass() const override;
+};
+
+UCLASS()
+class UAssetDefinition_Vector2dCameraVariable : public UAssetDefinition_CameraVariableAsset
+{
+	GENERATED_BODY()
+
+	// UAssetDefinition interface
+	virtual FText GetAssetDisplayName() const override;
+	virtual TSoftClassPtr<UObject> GetAssetClass() const override;
+};
+
+UCLASS()
+class UAssetDefinition_Vector3fCameraVariable : public UAssetDefinition_CameraVariableAsset
+{
+	GENERATED_BODY()
+
+	// UAssetDefinition interface
+	virtual FText GetAssetDisplayName() const override;
+	virtual TSoftClassPtr<UObject> GetAssetClass() const override;
+};
+
+UCLASS()
+class UAssetDefinition_Vector3dCameraVariable : public UAssetDefinition_CameraVariableAsset
+{
+	GENERATED_BODY()
+
+	// UAssetDefinition interface
+	virtual FText GetAssetDisplayName() const override;
+	virtual TSoftClassPtr<UObject> GetAssetClass() const override;
+};
+
+UCLASS()
+class UAssetDefinition_Vector4fCameraVariable : public UAssetDefinition_CameraVariableAsset
+{
+	GENERATED_BODY()
+
+	// UAssetDefinition interface
+	virtual FText GetAssetDisplayName() const override;
+	virtual TSoftClassPtr<UObject> GetAssetClass() const override;
+};
+
+UCLASS()
+class UAssetDefinition_Vector4dCameraVariable : public UAssetDefinition_CameraVariableAsset
+{
+	GENERATED_BODY()
+
+	// UAssetDefinition interface
+	virtual FText GetAssetDisplayName() const override;
+	virtual TSoftClassPtr<UObject> GetAssetClass() const override;
+};
+
+UCLASS()
+class UAssetDefinition_Rotator3fCameraVariable : public UAssetDefinition_CameraVariableAsset
+{
+	GENERATED_BODY()
+
+	// UAssetDefinition interface
+	virtual FText GetAssetDisplayName() const override;
+	virtual TSoftClassPtr<UObject> GetAssetClass() const override;
+};
+
+UCLASS()
+class UAssetDefinition_Rotator3dCameraVariable : public UAssetDefinition_CameraVariableAsset
+{
+	GENERATED_BODY()
+
+	// UAssetDefinition interface
+	virtual FText GetAssetDisplayName() const override;
+	virtual TSoftClassPtr<UObject> GetAssetClass() const override;
+};
+
+UCLASS()
+class UAssetDefinition_Transform3fCameraVariable : public UAssetDefinition_CameraVariableAsset
+{
+	GENERATED_BODY()
+
+	// UAssetDefinition interface
+	virtual FText GetAssetDisplayName() const override;
+	virtual TSoftClassPtr<UObject> GetAssetClass() const override;
+};
+
+UCLASS()
+class UAssetDefinition_Transform3dCameraVariable : public UAssetDefinition_CameraVariableAsset
+{
+	GENERATED_BODY()
+
+	// UAssetDefinition interface
+	virtual FText GetAssetDisplayName() const override;
+	virtual TSoftClassPtr<UObject> GetAssetClass() const override;
+};
+
diff --git a/Engine/Plugins/Cameras/GameplayCameras/Source/GameplayCamerasEditor/Private/Factories/CameraVariableFactories.cpp b/Engine/Plugins/Cameras/GameplayCameras/Source/GameplayCamerasEditor/Private/Factories/CameraVariableFactories.cpp
new file mode 100644
--- /dev/null
+++ b/Engine/Plugins/Cameras/GameplayCameras/Source/GameplayCamerasEditor/Private/Factories/CameraVariableFactories.cpp
@@ -0,0 +1,37 @@
+// Copyright Epic Games, Inc. All Rights Reserved.
+
+#include "Factories/CameraVariableFactories.h"
+
+#include "Core/CameraVariableAssets.h"
+
+#include UE_INLINE_GENERATED_CPP_BY_NAME(CameraVariableFactories)
+
+#define LOCTEXT_NAMESPACE "CameraVariableFactories"
+
+UCameraVariableAssetFactory::UCameraVariableAssetFactory(const FObjectInitializer& ObjectInit)
+	: Super(ObjectInit)
+{
+	bCreateNew = true;
+	bEditAfterNew = true;
+}
+
+bool UCameraVariableAssetFactory::ShouldShowInNewMenu() const
+{
+	return true;
+}
+
+#define UE_CAMERA_VARIABLE_FOR_TYPE(ValueType, ValueName)\
+	U##ValueName##CameraVariableFactory::U##ValueName##CameraVariableFactory(const FObjectInitializer& ObjectInit)\
+		: Super(ObjectInit)\
+	{\
+		SupportedClass = U##ValueName##CameraVariable::StaticClass();\
+	}\
+	UObject* U##ValueName##CameraVariableFactory::FactoryCreateNew(UClass* Class, UObject* Parent, FName Name, EObjectFlags Flags, UObject* Context, FFeedbackContext* Warn)\
+	{\
+		return NewObject<U##ValueName##CameraVariable>(Parent, Class, Name, Flags | RF_Transactional);\
+	}
+UE_CAMERA_VARIABLE_FOR_ALL_TYPES()
+#undef UE_CAMERA_VARIABLE_FOR_TYPE
+
+#undef LOCTEXT_NAMESPACE
+
diff --git a/Engine/Plugins/Cameras/GameplayCameras/Source/GameplayCamerasEditor/Private/Factories/CameraVariableFactories.h b/Engine/Plugins/Cameras/GameplayCameras/Source/GameplayCamerasEditor/Private/Factories/CameraVariableFactories.h
new file mode 100644
--- /dev/null
+++ b/Engine/Plugins/Cameras/GameplayCameras/Source/GameplayCamerasEditor/Private/Factories/CameraVariableFactories.h
@@ -0,0 +1,180 @@
+// Copyright Epic Games, Inc. All Rights Reserved.
+
+#pragma once
+
+#include "Factories/Factory.h"
+
+#include "CameraVariableFactories.generated.h"
+
+class UCameraVariableAsset;
+
+/**
+ * Implements a factory for camera variable assets.
+ */
+UCLASS(Abstract, hidecategories=Object)
+class UCameraVariableAssetFactory : public UFactory
+{
+	GENERATED_BODY()
+
+public:
+
+	UCameraVariableAssetFactory(const FObjectInitializer& ObjectInit);
+
+	// UFactory Interface
+	virtual bool ShouldShowInNewMenu() const override;
+};
+
+UCLASS()
+class UBooleanCameraVariableFactory : public UCameraVariableAssetFactory
+{
+	GENERATED_BODY()
+
+	UBooleanCameraVariableFactory(const FObjectInitializer& ObjectInit);
+
+	// UFactory Interface
+	virtual UObject* FactoryCreateNew(UClass* Class, UObject* Parent, FName Name, EObjectFlags Flags, UObject* Context, FFeedbackContext* Warn) override;
+};
+
+UCLASS()
+class UInteger32CameraVariableFactory : public UCameraVariableAssetFactory
+{
+	GENERATED_BODY()
+
+	UInteger32CameraVariableFactory(const FObjectInitializer& ObjectInit);
+
+	// UFactory Interface
+	virtual UObject* FactoryCreateNew(UClass* Class, UObject* Parent, FName Name, EObjectFlags Flags, UObject* Context, FFeedbackContext* Warn) override;
+};
+
+UCLASS()
+class UFloatCameraVariableFactory : public UCameraVariableAssetFactory
+{
+	GENERATED_BODY()
+
+	UFloatCameraVariableFactory(const FObjectInitializer& ObjectInit);
+
+	// UFactory Interface
+	virtual UObject* FactoryCreateNew(UClass* Class, UObject* Parent, FName Name, EObjectFlags Flags, UObject* Context, FFeedbackContext* Warn) override;
+};
+
+UCLASS()
+class UDoubleCameraVariableFactory : public UCameraVariableAssetFactory
+{
+	GENERATED_BODY()
+
+	UDoubleCameraVariableFactory(const FObjectInitializer& ObjectInit);
+
+	// UFactory Interface
+	virtual UObject* FactoryCreateNew(UClass* Class, UObject* Parent, FName Name, EObjectFlags Flags, UObject* Context, FFeedbackContext* Warn) override;
+};
+
+UCLASS()
+class UVector2fCameraVariableFactory : public UCameraVariableAssetFactory
+{
+	GENERATED_BODY()
+
+	UVector2fCameraVariableFactory(const FObjectInitializer& ObjectInit);
+
+	// UFactory Interface
+	virtual UObject* FactoryCreateNew(UClass* Class, UObject* Parent, FName Name, EObjectFlags Flags, UObject* Context, FFeedbackContext* Warn) override;
+};
+
+UCLASS()
+class UVector2dCameraVariableFactory : public UCameraVariableAssetFactory
+{
+	GENERATED_BODY()
+
+	UVector2dCameraVariableFactory(const FObjectInitializer& ObjectInit);
+
+	// UFactory Interface
+	virtual UObject* FactoryCreateNew(UClass* Class, UObject* Parent, FName Name, EObjectFlags Flags, UObject* Context, FFeedbackContext* Warn) override;
+};
+
+UCLASS()
+class UVector3fCameraVariableFactory : public UCameraVariableAssetFactory
+{
+	GENERATED_BODY()
+
+	UVector3fCameraVariableFactory(const FObjectInitializer& ObjectInit);
+
+	// UFactory Interface
+	virtual UObject* FactoryCreateNew(UClass* Class, UObject* Parent, FName Name, EObjectFlags Flags, UObject* Context, FFeedbackContext* Warn) override;
+};
+
+UCLASS()
+class UVector3dCameraVariableFactory : public UCameraVariableAssetFactory
+{
+	GENERATED_BODY()
+
+	UVector3dCameraVariableFactory(const FObjectInitializer& ObjectInit);
+
+	// UFactory Interface
+	virtual UObject* FactoryCreateNew(UClass* Class, UObject* Parent, FName Name, EObjectFlags Flags, UObject* Context, FFeedbackContext* Warn) override;
+};
+
+UCLASS()
+class UVector4fCameraVariableFactory : public UCameraVariableAssetFactory
+{
+	GENERATED_BODY()
+
+	UVector4fCameraVariableFactory(const FObjectInitializer& ObjectInit);
+
+	// UFactory Interface
+	virtual UObject* FactoryCreateNew(UClass* Class, UObject* Parent, FName Name, EObjectFlags Flags, UObject* Context, FFeedbackContext* Warn) override;
+};
+
+UCLASS()
+class UVector4dCameraVariableFactory : public UCameraVariableAssetFactory
+{
+	GENERATED_BODY()
+
+	UVector4dCameraVariableFactory(const FObjectInitializer& ObjectInit);
+
+	// UFactory Interface
+	virtual UObject* FactoryCreateNew(UClass* Class, UObject* Parent, FName Name, EObjectFlags Flags, UObject* Context, FFeedbackContext* Warn) override;
+};
+
+UCLASS()
+class URotator3fCameraVariableFactory : public UCameraVariableAssetFactory
+{
+	GENERATED_BODY()
+
+	URotator3fCameraVariableFactory(const FObjectInitializer& ObjectInit);
+
+	// UFactory Interface
+	virtual UObject* FactoryCreateNew(UClass* Class, UObject* Parent, FName Name, EObjectFlags Flags, UObject* Context, FFeedbackContext* Warn) override;
+};
+
+UCLASS()
+class URotator3dCameraVariableFactory : public UCameraVariableAssetFactory
+{
+	GENERATED_BODY()
+
+	URotator3dCameraVariableFactory(const FObjectInitializer& ObjectInit);
+
+	// UFactory Interface
+	virtual UObject* FactoryCreateNew(UClass* Class, UObject* Parent, FName Name, EObjectFlags Flags, UObject* Context, FFeedbackContext* Warn) override;
+};
+
+UCLASS()
+class UTransform3fCameraVariableFactory : public UCameraVariableAssetFactory
+{
+	GENERATED_BODY()
+
+	UTransform3fCameraVariableFactory(const FObjectInitializer& ObjectInit);
+
+	// UFactory Interface
+	virtual UObject* FactoryCreateNew(UClass* Class, UObject* Parent, FName Name, EObjectFlags Flags, UObject* Context, FFeedbackContext* Warn) override;
+};
+
+UCLASS()
+class UTransform3dCameraVariableFactory : public UCameraVariableAssetFactory
+{
+	GENERATED_BODY()
+
+	UTransform3dCameraVariableFactory(const FObjectInitializer& ObjectInit);
+
+	// UFactory Interface
+	virtual UObject* FactoryCreateNew(UClass* Class, UObject* Parent, FName Name, EObjectFlags Flags, UObject* Context, FFeedbackContext* Warn) override;
+};
+
