Cameras: turn the offset in the offset node into a camera parameter

This allows the offset to be driven by a camera variable. Upgrade of previously saved data should be gracefully handled by the camera parameters' own serialization.

diff --git a/Engine/Plugins/Cameras/GameplayCameras/Source/GameplayCameras/Private/Nodes/Common/OffsetCameraNode.cpp b/Engine/Plugins/Cameras/GameplayCameras/Source/GameplayCameras/Private/Nodes/Common/OffsetCameraNode.cpp
--- a/Engine/Plugins/Cameras/GameplayCameras/Source/GameplayCameras/Private/Nodes/Common/OffsetCameraNode.cpp
+++ b/Engine/Plugins/Cameras/GameplayCameras/Source/GameplayCameras/Private/Nodes/Common/OffsetCameraNode.cpp
@@ -3,6 +3,7 @@
 #include "Nodes/Common/OffsetCameraNode.h"
 
 #include "Core/CameraEvaluationContext.h"
+#include "Core/CameraParameterReader.h"
 #include "GameplayCameras.h"
 
 #include UE_INLINE_GENERATED_CPP_BY_NAME(OffsetCameraNode)
@@ -13,15 +14,27 @@ class FOffsetCameraNodeEvaluator : public FCameraNodeEvaluator
 
 protected:
 
+	virtual void OnInitialize(const FCameraNodeEvaluatorInitializeParams& Params) override;
 	virtual void OnRun(const FCameraNodeEvaluationParams& Params, FCameraNodeEvaluationResult& OutResult) override;
+
+private:
+
+	FCameraParameterReader<FVector3d> OffsetReader;
 };
 
 UE_DEFINE_CAMERA_NODE_EVALUATOR(FOffsetCameraNodeEvaluator)
 
+void FOffsetCameraNodeEvaluator::OnInitialize(const FCameraNodeEvaluatorInitializeParams& Params)
+{
+	const UOffsetCameraNode* OffsetNode = GetCameraNodeAs<UOffsetCameraNode>();
+	OffsetReader.Initialize(OffsetNode->Offset);
+}
+
 void FOffsetCameraNodeEvaluator::OnRun(const FCameraNodeEvaluationParams& Params, FCameraNodeEvaluationResult& OutResult)
 {
+	FVector3d LocalOffset = OffsetReader.Get(OutResult.VariableTable);
+
 	const UOffsetCameraNode* OffsetNode = GetCameraNodeAs<UOffsetCameraNode>();
-	FVector3d LocalOffset = OffsetNode->Offset;
 	switch(OffsetNode->OffsetSpace)
 	{
 		case ECameraNodeSpace::CameraPose:
diff --git a/Engine/Plugins/Cameras/GameplayCameras/Source/GameplayCameras/Public/Nodes/Common/OffsetCameraNode.h b/Engine/Plugins/Cameras/GameplayCameras/Source/GameplayCameras/Public/Nodes/Common/OffsetCameraNode.h
--- a/Engine/Plugins/Cameras/GameplayCameras/Source/GameplayCameras/Public/Nodes/Common/OffsetCameraNode.h
+++ b/Engine/Plugins/Cameras/GameplayCameras/Source/GameplayCameras/Public/Nodes/Common/OffsetCameraNode.h
@@ -3,6 +3,7 @@
 #pragma once
 
 #include "Core/CameraNode.h"
+#include "Core/CameraParameters.h"
 #include "Nodes/CameraNodeTypes.h"
 
 #include "OffsetCameraNode.generated.h"
@@ -24,7 +25,7 @@ public:
 
 	/** The offset to apply to the camera, in local space. */
 	UPROPERTY(EditAnywhere, Category=Common)
-	FVector3d Offset;
+	FVector3dCameraParameter Offset;
 
 	/** The space in which to apply the offset. */
 	UPROPERTY(EditAnywhere, Category=Common)
