Cameras: add children API to camera nodes

diff --git a/Engine/Plugins/Cameras/GameplayCameras/Source/GameplayCameras/Private/Core/BlendStackRootCameraNode.cpp b/Engine/Plugins/Cameras/GameplayCameras/Source/GameplayCameras/Private/Core/BlendStackRootCameraNode.cpp
--- a/Engine/Plugins/Cameras/GameplayCameras/Source/GameplayCameras/Private/Core/BlendStackRootCameraNode.cpp
+++ b/Engine/Plugins/Cameras/GameplayCameras/Source/GameplayCameras/Private/Core/BlendStackRootCameraNode.cpp
@@ -7,6 +7,20 @@
 
 #include UE_INLINE_GENERATED_CPP_BY_NAME(BlendStackRootCameraNode)
 
+FCameraNodeChildrenView UBlendStackRootCameraNode::OnGetChildren()
+{
+	FCameraNodeChildrenView Children;
+	if (Blend)
+	{
+		Children.Add(Blend);
+	}
+	if (RootNode)
+	{
+		Children.Add(RootNode);
+	}
+	return Children;
+}
+
 FCameraNodeEvaluatorPtr UBlendStackRootCameraNode::OnBuildEvaluator(FCameraNodeEvaluatorBuilder& Builder) const
 {
 	return Builder.BuildEvaluator<FBlendStackRootCameraNodeEvaluator>();
diff --git a/Engine/Plugins/Cameras/GameplayCameras/Source/GameplayCameras/Private/Core/BlendStackRootCameraNode.h b/Engine/Plugins/Cameras/GameplayCameras/Source/GameplayCameras/Private/Core/BlendStackRootCameraNode.h
--- a/Engine/Plugins/Cameras/GameplayCameras/Source/GameplayCameras/Private/Core/BlendStackRootCameraNode.h
+++ b/Engine/Plugins/Cameras/GameplayCameras/Source/GameplayCameras/Private/Core/BlendStackRootCameraNode.h
@@ -24,6 +24,7 @@ class UBlendStackRootCameraNode : public UCameraNode
 protected:
 
 	// UCameraNode interface.
+	virtual FCameraNodeChildrenView OnGetChildren() override;
 	virtual FCameraNodeEvaluatorPtr OnBuildEvaluator(FCameraNodeEvaluatorBuilder& Builder) const override;
 
 public:
diff --git a/Engine/Plugins/Cameras/GameplayCameras/Source/GameplayCameras/Private/Core/CameraNode.cpp b/Engine/Plugins/Cameras/GameplayCameras/Source/GameplayCameras/Private/Core/CameraNode.cpp
--- a/Engine/Plugins/Cameras/GameplayCameras/Source/GameplayCameras/Private/Core/CameraNode.cpp
+++ b/Engine/Plugins/Cameras/GameplayCameras/Source/GameplayCameras/Private/Core/CameraNode.cpp
@@ -13,6 +13,11 @@ void FCameraNodeEvaluationResult::Reset()
 	bIsValid = false;
 }
 
+FCameraNodeChildrenView UCameraNode::GetChildren()
+{
+	return OnGetChildren();
+}
+
 FCameraNodeAllocationInfo UCameraNode::GetAllocationInfo() const
 {
 	return OnGetAllocationInfo();
diff --git a/Engine/Plugins/Cameras/GameplayCameras/Source/GameplayCameras/Private/Nodes/Common/ArrayCameraNode.cpp b/Engine/Plugins/Cameras/GameplayCameras/Source/GameplayCameras/Private/Nodes/Common/ArrayCameraNode.cpp
--- a/Engine/Plugins/Cameras/GameplayCameras/Source/GameplayCameras/Private/Nodes/Common/ArrayCameraNode.cpp
+++ b/Engine/Plugins/Cameras/GameplayCameras/Source/GameplayCameras/Private/Nodes/Common/ArrayCameraNode.cpp
@@ -52,6 +52,11 @@ void FArrayCameraNodeEvaluator::OnRun(const FCameraNodeEvaluationParams& Params,
 	}
 }
 
+FCameraNodeChildrenView UArrayCameraNode::OnGetChildren()
+{
+	return FCameraNodeChildrenView(Children);
+}
+
 FCameraNodeEvaluatorPtr UArrayCameraNode::OnBuildEvaluator(FCameraNodeEvaluatorBuilder& Builder) const
 {
 	return Builder.BuildEvaluator<FArrayCameraNodeEvaluator>();
diff --git a/Engine/Plugins/Cameras/GameplayCameras/Source/GameplayCameras/Public/Core/CameraNode.h b/Engine/Plugins/Cameras/GameplayCameras/Source/GameplayCameras/Public/Core/CameraNode.h
--- a/Engine/Plugins/Cameras/GameplayCameras/Source/GameplayCameras/Public/Core/CameraNode.h
+++ b/Engine/Plugins/Cameras/GameplayCameras/Source/GameplayCameras/Public/Core/CameraNode.h
@@ -4,6 +4,7 @@
 
 #include "Core/CameraNodeEvaluatorBuilder.h"
 #include "Core/CameraVariableTableFwd.h"
+#include "Core/ObjectChildrenView.h"
 #include "CoreTypes.h"
 #include "UObject/Object.h"
 
@@ -24,6 +25,9 @@ struct FCameraNodeAllocationInfo
 	FCameraVariableTableAllocationInfo VariableTableInfo;
 };
 
+/** View on a camera node's children. */
+using FCameraNodeChildrenView = TObjectChildrenView<TObjectPtr<UCameraNode>>;
+
 /**
  * The base class for a camera node.
  */
@@ -34,6 +38,9 @@ class UCameraNode : public UObject
 
 public:
 
+	/** Get the list of children under this node. */
+	FCameraNodeChildrenView GetChildren();
+
 	/** Gets optional info about this node's required allocations at runtime. */
 	FCameraNodeAllocationInfo GetAllocationInfo() const;
 
@@ -42,6 +49,9 @@ public:
 
 protected:
 
+	/** Get the list of children under this node. */
+	virtual FCameraNodeChildrenView OnGetChildren() { return FCameraNodeChildrenView(); }
+
 	/** Gets optional info about this node's required allocations at runtime. */
 	virtual FCameraNodeAllocationInfo OnGetAllocationInfo() const { return FCameraNodeAllocationInfo(); }
 
diff --git a/Engine/Plugins/Cameras/GameplayCameras/Source/GameplayCameras/Public/Nodes/Common/ArrayCameraNode.h b/Engine/Plugins/Cameras/GameplayCameras/Source/GameplayCameras/Public/Nodes/Common/ArrayCameraNode.h
--- a/Engine/Plugins/Cameras/GameplayCameras/Source/GameplayCameras/Public/Nodes/Common/ArrayCameraNode.h
+++ b/Engine/Plugins/Cameras/GameplayCameras/Source/GameplayCameras/Public/Nodes/Common/ArrayCameraNode.h
@@ -18,6 +18,7 @@ class UArrayCameraNode : public UCameraNode
 protected:
 
 	// UCameraNode interface.
+	virtual FCameraNodeChildrenView OnGetChildren() override;
 	virtual FCameraNodeEvaluatorPtr OnBuildEvaluator(FCameraNodeEvaluatorBuilder& Builder) const override;
 
 public:
