Cameras: add basic hot-reloading mechanism for camera mode editor

The editor now has a "Build" button that immediately rebuilds the running evaluator tree of a camera mode (if any) after that camera mode has been edited.

diff --git a/Engine/Plugins/Cameras/GameplayCameras/Source/GameplayCameras/Private/Core/BlendStackCameraNode.cpp b/Engine/Plugins/Cameras/GameplayCameras/Source/GameplayCameras/Private/Core/BlendStackCameraNode.cpp
--- a/Engine/Plugins/Cameras/GameplayCameras/Source/GameplayCameras/Private/Core/BlendStackCameraNode.cpp
+++ b/Engine/Plugins/Cameras/GameplayCameras/Source/GameplayCameras/Private/Core/BlendStackCameraNode.cpp
@@ -8,6 +8,9 @@
 #include "Core/CameraEvaluationContext.h"
 #include "Core/CameraMode.h"
 #include "Core/CameraSystemEvaluator.h"
+#include "IGameplayCamerasModule.h"
+#include "IGameplayCamerasLiveEditManager.h"
+#include "Modules/ModuleManager.h"
 #include "Nodes/Blends/PopBlendCameraNode.h"
 
 #include UE_INLINE_GENERATED_CPP_BY_NAME(BlendStackCameraNode)
@@ -69,9 +72,20 @@ void FBlendStackCameraNodeEvaluator::Push(const FBlendStackCameraPushParams& Par
 
 	NewEntry.EvaluationContext = Params.EvaluationContext;
 	NewEntry.CameraMode = Params.CameraMode;
+	NewEntry.RootNode = EntryRootNode;
 	NewEntry.RootEvaluator = RootEvaluator->CastThisChecked<FBlendStackRootCameraNodeEvaluator>();
 	NewEntry.bIsFirstFrame = true;
 
+#if WITH_EDITOR
+	IGameplayCamerasModule& GameplayCamerasModule = FModuleManager::GetModuleChecked<IGameplayCamerasModule>("GameplayCameras");
+	TSharedPtr<IGameplayCamerasLiveEditManager> LiveEditManager = GameplayCamerasModule.GetLiveEditManager();
+	Params.CameraMode->GatherPackages(NewEntry.ListenedPackages);
+	for (const UPackage* ListenPackage : NewEntry.ListenedPackages)
+	{
+		LiveEditManager->AddListener(ListenPackage, this);
+	}
+#endif  // WITH_EDITOR
+
 	// Important: we need to move the new entry here because copying evaluator storage
 	// is disabled.
 	Entries.Add(MoveTemp(NewEntry));
@@ -87,6 +101,11 @@ FCameraNodeEvaluatorChildrenView FBlendStackCameraNodeEvaluator::OnGetChildren()
 	return View;
 }
 
+void FBlendStackCameraNodeEvaluator::OnInitialize(const FCameraNodeEvaluatorInitializeParams& Params)
+{
+	OwningEvaluator = Params.Evaluator;
+}
+
 void FBlendStackCameraNodeEvaluator::OnRun(const FCameraNodeEvaluationParams& Params, FCameraNodeEvaluationResult& OutResult)
 {
 	const UBlendStackCameraNode* BlendStackNode = GetCameraNodeAs<UBlendStackCameraNode>();
@@ -184,8 +203,21 @@ void FBlendStackCameraNodeEvaluator::OnRun(const FCameraNodeEvaluationParams& Pa
 	// Pop out camera modes that have been blended out.
 	if (BlendStackNode->bAutoPop && PopEntriesBelow != INDEX_NONE)
 	{
+#if WITH_EDITOR
+		IGameplayCamerasModule& GameplayCamerasModule = FModuleManager::GetModuleChecked<IGameplayCamerasModule>("GameplayCameras");
+		TSharedPtr<IGameplayCamerasLiveEditManager> LiveEditManager = GameplayCamerasModule.GetLiveEditManager();
+#endif  // WITH_EDITOR
+
 		for (int32 Index = 0; Index < PopEntriesBelow; ++Index)
 		{
+#if WITH_EDITOR
+			const FCameraModeEntry& FirstEntry = Entries[0];
+			for (const UPackage* ListenPackage : FirstEntry.ListenedPackages)
+			{
+				LiveEditManager->RemoveListener(ListenPackage, this);
+			}
+#endif  // WITH_EDITOR
+
 			Entries.RemoveAt(0);
 		}
 	}
@@ -313,3 +345,38 @@ const FCameraModeTransition* FBlendStackCameraNodeEvaluator::FindTransition(
 	return nullptr;
 }
 
+#if WITH_EDITOR
+
+void FBlendStackCameraNodeEvaluator::OnPostBuildAsset(const FGameplayCameraAssetBuildEvent& BuildEvent)
+{
+	for (FCameraModeEntry& Entry : Entries)
+	{
+		const bool bRebuildEntry = Entry.ListenedPackages.Contains(BuildEvent.AssetPackage);
+		if (bRebuildEntry)
+		{
+			Entry.EvaluatorStorage.DestroyEvaluatorTree();
+
+			// Re-assign the root node in case the camera mode's root was changed.
+			Entry.RootNode->RootNode = Entry.CameraMode->RootNode;
+
+			// Remove the blend on the root node, since we don't want the reloaded camera mode to re-blend-in
+			// for no good reason.
+			Entry.RootNode->Blend = nullptr;
+
+			// Rebuild the evaluator tree.
+			FCameraNodeEvaluatorTreeBuilderParams BuildParams;
+			BuildParams.RootCameraNode = Entry.RootNode;
+			BuildParams.Evaluator = OwningEvaluator;
+			BuildParams.EvaluationContext = Entry.EvaluationContext.Get();
+			FCameraNodeEvaluator* RootEvaluator = Entry.EvaluatorStorage.BuildEvaluatorTree(BuildParams);
+
+			Entry.RootEvaluator = RootEvaluator->CastThisChecked<FBlendStackRootCameraNodeEvaluator>();
+
+			// This is the first frame for this new hierarchy of evaluators.
+			Entry.bIsFirstFrame = true;
+		}
+	}
+}
+
+#endif  // WITH_EDITOR
+
diff --git a/Engine/Plugins/Cameras/GameplayCameras/Source/GameplayCameras/Private/Core/CameraMode.cpp b/Engine/Plugins/Cameras/GameplayCameras/Source/GameplayCameras/Private/Core/CameraMode.cpp
--- a/Engine/Plugins/Cameras/GameplayCameras/Source/GameplayCameras/Private/Core/CameraMode.cpp
+++ b/Engine/Plugins/Cameras/GameplayCameras/Source/GameplayCameras/Private/Core/CameraMode.cpp
@@ -2,5 +2,31 @@
 
 #include "Core/CameraMode.h"
 
+#include "Core/CameraNode.h"
+
 #include UE_INLINE_GENERATED_CPP_BY_NAME(CameraMode)
 
+#if WITH_EDITOR
+
+void UCameraMode::GatherPackages(FCameraModePackages& OutPackages) const
+{
+	TArray<UCameraNode*> NodeStack;
+	if (RootNode)
+	{
+		NodeStack.Add(RootNode);
+	}
+	while (!NodeStack.IsEmpty())
+	{
+		UCameraNode* CurrentNode = NodeStack.Pop();
+		const UPackage* CurrentPackage = CurrentNode->GetOutermost();
+		OutPackages.AddUnique(CurrentPackage);
+
+		FCameraNodeChildrenView CurrentChildren = CurrentNode->GetChildren();
+		for (UCameraNode* CurrentChild : ReverseIterate(CurrentChildren))
+		{
+			NodeStack.Add(CurrentChild);
+		}
+	}
+}
+
+#endif  // WITH_EDITOR
diff --git a/Engine/Plugins/Cameras/GameplayCameras/Source/GameplayCameras/Private/Core/CameraNodeEvaluatorStorage.cpp b/Engine/Plugins/Cameras/GameplayCameras/Source/GameplayCameras/Private/Core/CameraNodeEvaluatorStorage.cpp
--- a/Engine/Plugins/Cameras/GameplayCameras/Source/GameplayCameras/Private/Core/CameraNodeEvaluatorStorage.cpp
+++ b/Engine/Plugins/Cameras/GameplayCameras/Source/GameplayCameras/Private/Core/CameraNodeEvaluatorStorage.cpp
@@ -17,23 +17,22 @@ FCameraNodeEvaluatorStorage::FCameraNodeEvaluatorStorage(FCameraNodeEvaluatorSto
 	EvaluatorInfos = MoveTemp(Other.EvaluatorInfos);
 }
 
-FCameraNodeEvaluatorStorage::~FCameraNodeEvaluatorStorage()
+FCameraNodeEvaluatorStorage& FCameraNodeEvaluatorStorage::operator=(FCameraNodeEvaluatorStorage&& Other)
 {
-	for (FEvaluatorInfo& EvaluatorInfo : EvaluatorInfos)
+	if (ensure(this != &Other))
 	{
-		FCameraNodeEvaluator* Evaluator(EvaluatorInfo.Ptr);
-		Evaluator->~FCameraNodeEvaluator();
+		Allocations = MoveTemp(Other.Allocations);
+		EvaluatorInfos = MoveTemp(Other.EvaluatorInfos);
 	}
-	EvaluatorInfos.Reset();
+	return *this;
+}
 
-	for (FAllocation& Allocation : Allocations)
-	{
-		FMemory::Free(Allocation.Memory);
-	}
-	Allocations.Reset();
+FCameraNodeEvaluatorStorage::~FCameraNodeEvaluatorStorage()
+{
+	DestroyEvaluatorTree(true);
 }
 
-FCameraNodeEvaluatorTreeAllocationInfo FCameraNodeEvaluatorStorage::ComputeTreeInfo(const UCameraMode* CameraMode) const
+FCameraNodeEvaluatorTreeAllocationInfo FCameraNodeEvaluatorStorage::ComputeTreeInfo(const UCameraMode* CameraMode)
 {
 	int16 MaxSize = 0;
 	int16 MaxAlignment = 0;
@@ -80,3 +79,29 @@ FCameraNodeEvaluatorPtr FCameraNodeEvaluatorStorage::BuildEvaluatorTree(const FC
 	return RootEvaluator;
 }
 
+void FCameraNodeEvaluatorStorage::DestroyEvaluatorTree(bool bFreeAllocations)
+{
+	for (FEvaluatorInfo& EvaluatorInfo : EvaluatorInfos)
+	{
+		FCameraNodeEvaluator* Evaluator(EvaluatorInfo.Ptr);
+		Evaluator->~FCameraNodeEvaluator();
+	}
+	EvaluatorInfos.Reset();
+
+	if (bFreeAllocations)
+	{
+		for (FAllocation& Allocation : Allocations)
+		{
+			FMemory::Free(Allocation.Memory);
+		}
+		Allocations.Reset();
+	}
+	else
+	{
+		for (FAllocation& Allocation : Allocations)
+		{
+			Allocation.Used = 0;
+		}
+	}
+}
+
diff --git a/Engine/Plugins/Cameras/GameplayCameras/Source/GameplayCameras/Public/Core/BlendStackCameraNode.h b/Engine/Plugins/Cameras/GameplayCameras/Source/GameplayCameras/Public/Core/BlendStackCameraNode.h
--- a/Engine/Plugins/Cameras/GameplayCameras/Source/GameplayCameras/Public/Core/BlendStackCameraNode.h
+++ b/Engine/Plugins/Cameras/GameplayCameras/Source/GameplayCameras/Public/Core/BlendStackCameraNode.h
@@ -2,9 +2,11 @@
 
 #pragma once
 
+#include "Core/CameraMode.h"
 #include "Core/CameraNode.h"
 #include "Core/CameraNodeEvaluator.h"
 #include "Core/CameraNodeEvaluatorStorage.h"
+#include "IGameplayCamerasLiveEditListener.h"
 
 #include "BlendStackCameraNode.generated.h"
 
@@ -62,7 +64,11 @@ public:
 /**
  * Evaluator for a blend stack camera node.
  */
-class FBlendStackCameraNodeEvaluator : public TCameraNodeEvaluator<UBlendStackCameraNode>
+class FBlendStackCameraNodeEvaluator 
+	: public TCameraNodeEvaluator<UBlendStackCameraNode>
+#if WITH_EDITOR
+	, public IGameplayCamerasLiveEditListener
+#endif
 {
 	UE_DECLARE_CAMERA_NODE_EVALUATOR(FBlendStackCameraNodeEvaluator)
 
@@ -75,8 +81,16 @@ protected:
 
 	// FCameraNodeEvaluator interface
 	virtual FCameraNodeEvaluatorChildrenView OnGetChildren() override;
+	virtual void OnInitialize(const FCameraNodeEvaluatorInitializeParams& Params) override;
 	virtual void OnRun(const FCameraNodeEvaluationParams& Params, FCameraNodeEvaluationResult& OutResult) override;
 
+#if WITH_EDITOR
+	// IGameplayCamerasLiveEditListener interface
+	virtual void OnPostBuildAsset(const FGameplayCameraAssetBuildEvent& BuildEvent) override;
+#endif
+
+protected:
+
 	// Utility functions for finding an appropriate transition.
 	const FCameraModeTransition* FindTransition(const FBlendStackCameraPushParams& Params) const;
 	const FCameraModeTransition* FindTransition(
@@ -92,6 +106,8 @@ protected:
 		TWeakObjectPtr<const UCameraEvaluationContext> EvaluationContext;
 		/** The camera mode asset that this entry runs. */
 		TObjectPtr<const UCameraMode> CameraMode;
+		/** The root node. */
+		TObjectPtr<UBlendStackRootCameraNode> RootNode;
 		/** Storage buffer for all evaluators in this node tree. */
 		FCameraNodeEvaluatorStorage EvaluatorStorage;
 		/** Root evaluator. */
@@ -102,8 +118,14 @@ protected:
 		bool bIsFirstFrame = false;
 		/** Whether this entry is frozen. */
 		bool bIsFrozen = false;
+#if WITH_EDITOR
+		FCameraModePackages ListenedPackages;
+#endif  // WITH_EDITOR
 	};
 
+	/** The camera system evaluator running this node. */
+	TObjectPtr<UCameraSystemEvaluator> OwningEvaluator;
+
 	/** Entries in the blend stack. */
 	TArray<FCameraModeEntry> Entries;
 };
diff --git a/Engine/Plugins/Cameras/GameplayCameras/Source/GameplayCameras/Public/Core/CameraMode.h b/Engine/Plugins/Cameras/GameplayCameras/Source/GameplayCameras/Public/Core/CameraMode.h
--- a/Engine/Plugins/Cameras/GameplayCameras/Source/GameplayCameras/Public/Core/CameraMode.h
+++ b/Engine/Plugins/Cameras/GameplayCameras/Source/GameplayCameras/Public/Core/CameraMode.h
@@ -10,6 +10,24 @@
 
 class UCameraNode;
 
+/**
+ *
+ */
+UENUM()
+enum class ECameraModeBuildStatus : uint8
+{
+	Clean,
+	CleanWithWarnings,
+	WithErrors,
+	Dirty
+};
+
+/**
+ * List of packages that contain the definition of a camera mode.
+ * In most cases there's only one, but with nested assets there could be more.
+ */
+using FCameraModePackages = TArray<const UPackage*, TInlineAllocator<4>>;
+
 /**
  * A camera mode asset, which runs a hierarchy of camera nodes to drive 
  * the behavior of a camera.
@@ -19,6 +37,12 @@ class UCameraMode : public UObject
 {
 	GENERATED_BODY()
 
+public:
+
+#if WITH_EDITOR
+	GAMEPLAYCAMERAS_API void GatherPackages(FCameraModePackages& OutPackages) const;
+#endif  // WITH_EDITOR
+
 public:
 
 	/** Root camera node. */
@@ -32,5 +56,8 @@ public:
 	/** List of exist transitions for this camera mode. */
 	UPROPERTY(EditAnywhere, Category=Blending)
 	TArray<FCameraModeTransition> ExitTransitions;
+
+	UPROPERTY(Transient)
+	ECameraModeBuildStatus BuildStatus = ECameraModeBuildStatus::Dirty;
 };
 
diff --git a/Engine/Plugins/Cameras/GameplayCameras/Source/GameplayCameras/Public/Core/CameraNodeEvaluatorStorage.h b/Engine/Plugins/Cameras/GameplayCameras/Source/GameplayCameras/Public/Core/CameraNodeEvaluatorStorage.h
--- a/Engine/Plugins/Cameras/GameplayCameras/Source/GameplayCameras/Public/Core/CameraNodeEvaluatorStorage.h
+++ b/Engine/Plugins/Cameras/GameplayCameras/Source/GameplayCameras/Public/Core/CameraNodeEvaluatorStorage.h
@@ -45,6 +45,7 @@ public:
 
 	FCameraNodeEvaluatorStorage();
 	FCameraNodeEvaluatorStorage(FCameraNodeEvaluatorStorage&& Other);
+	FCameraNodeEvaluatorStorage& operator=(FCameraNodeEvaluatorStorage&& Other);
 	~FCameraNodeEvaluatorStorage();
 
 	FCameraNodeEvaluatorStorage(const FCameraNodeEvaluatorStorage&) = delete;
@@ -53,9 +54,12 @@ public:
 public:
 
 	/** Compute allocation information for the given tree of camera nodes. */
-	FCameraNodeEvaluatorTreeAllocationInfo ComputeTreeInfo(const UCameraMode* CameraMode) const;
+	static FCameraNodeEvaluatorTreeAllocationInfo ComputeTreeInfo(const UCameraMode* CameraMode);
+
 	/** Build the tree of evaluators for the given tree of camera nodes. */
 	FCameraNodeEvaluatorPtr BuildEvaluatorTree(const FCameraNodeEvaluatorTreeBuilderParams& Params);
+	/** Destroy any allocated evaluators. */
+	void DestroyEvaluatorTree(bool bFreeAllocations = false);
 
 private:
 
diff --git a/Engine/Plugins/Cameras/GameplayCameras/Source/GameplayCameras/Public/IGameplayCamerasLiveEditListener.h b/Engine/Plugins/Cameras/GameplayCameras/Source/GameplayCameras/Public/IGameplayCamerasLiveEditListener.h
new file mode 100644
--- /dev/null
+++ b/Engine/Plugins/Cameras/GameplayCameras/Source/GameplayCameras/Public/IGameplayCamerasLiveEditListener.h
@@ -0,0 +1,35 @@
+// Copyright Epic Games, Inc. All Rights Reserved.
+
+#pragma once
+
+#if WITH_EDITOR
+
+#include "CoreTypes.h"
+#include "Templates/SharedPointer.h"
+#include "UObject/UnrealType.h"
+
+class UPackage;
+
+struct FGameplayCameraAssetBuildEvent
+{
+	const UPackage* AssetPackage = nullptr;
+};
+
+/**
+ * Interface for an object that can handle an asset being hot-reloaded at runtime.
+ */
+class IGameplayCamerasLiveEditListener
+{
+public:
+
+	virtual ~IGameplayCamerasLiveEditListener() {}
+
+	void PostBuildAsset(const FGameplayCameraAssetBuildEvent& BuildEvent) { OnPostBuildAsset(BuildEvent); }
+
+protected:
+
+	virtual void OnPostBuildAsset(const FGameplayCameraAssetBuildEvent& BuildEvent) = 0;
+};
+
+#endif  // WITH_EDITOR
+
diff --git a/Engine/Plugins/Cameras/GameplayCameras/Source/GameplayCameras/Public/IGameplayCamerasLiveEditManager.h b/Engine/Plugins/Cameras/GameplayCameras/Source/GameplayCameras/Public/IGameplayCamerasLiveEditManager.h
--- a/Engine/Plugins/Cameras/GameplayCameras/Source/GameplayCameras/Public/IGameplayCamerasLiveEditManager.h
+++ b/Engine/Plugins/Cameras/GameplayCameras/Source/GameplayCameras/Public/IGameplayCamerasLiveEditManager.h
@@ -2,19 +2,32 @@
 
 #pragma once
 
+#if WITH_EDITOR
+
 #include "CoreTypes.h"
 #include "Templates/SharedPointer.h"
 #include "UObject/UnrealType.h"
 
-class UObject;
+class UPackage;
+class IGameplayCamerasLiveEditListener;
 
 /**
- * Interface for an object that can handle live-editing features of the camera assets.
+ * Interface for an object that can centralize the live-editing features of the camera system.
  */
 class IGameplayCamerasLiveEditManager : public TSharedFromThis<IGameplayCamerasLiveEditManager>
 {
 public:
 
 	virtual ~IGameplayCamerasLiveEditManager() {}
+
+	/** Notify all listeners to reload cameras related to the given package. */
+	virtual void NotifyPostBuildAsset(const UPackage* InAssetPackage) const = 0;
+
+	/** Add a listener for the given package. */
+	virtual void AddListener(const UPackage* InAssetPackage, IGameplayCamerasLiveEditListener* Listener) = 0;
+	/** Removes a listener for the given package. */
+	virtual void RemoveListener(const UPackage* InAssetPackage, IGameplayCamerasLiveEditListener* Listener) = 0;
 };
 
+#endif  // WITH_EDITOR
+
diff --git a/Engine/Plugins/Cameras/GameplayCameras/Source/GameplayCameras/Public/IGameplayCamerasModule.h b/Engine/Plugins/Cameras/GameplayCameras/Source/GameplayCameras/Public/IGameplayCamerasModule.h
--- a/Engine/Plugins/Cameras/GameplayCameras/Source/GameplayCameras/Public/IGameplayCamerasModule.h
+++ b/Engine/Plugins/Cameras/GameplayCameras/Source/GameplayCameras/Public/IGameplayCamerasModule.h
@@ -3,7 +3,7 @@
 #pragma once
 
 #include "Modules/ModuleInterface.h"
-#include "Templates/SharedPointerFwd.h"
+#include "Templates/SharedPointer.h"
 
 #if WITH_EDITOR
 class IGameplayCamerasLiveEditManager;
diff --git a/Engine/Plugins/Cameras/GameplayCameras/Source/GameplayCamerasEditor/Private/Commands/CameraAssetEditorCommands.cpp b/Engine/Plugins/Cameras/GameplayCameras/Source/GameplayCamerasEditor/Private/Commands/CameraAssetEditorCommands.cpp
new file mode 100644
--- /dev/null
+++ b/Engine/Plugins/Cameras/GameplayCameras/Source/GameplayCamerasEditor/Private/Commands/CameraAssetEditorCommands.cpp
@@ -0,0 +1,29 @@
+// Copyright Epic Games, Inc. All Rights Reserved.
+
+#include "Commands/CameraAssetEditorCommands.h"
+
+#include "Framework/Commands/Commands.h"
+#include "Framework/Commands/InputChord.h"
+#include "Framework/Commands/UICommandInfo.h"
+#include "Styling/AppStyle.h"
+
+#define LOCTEXT_NAMESPACE "SequencerCommands"
+
+FCameraAssetEditorCommands::FCameraAssetEditorCommands()
+	: TCommands<FCameraAssetEditorCommands>(
+			"GameplayCameras",
+			NSLOCTEXT("Contexts", "GameplayCameras", "Gameplay Cameras"),
+			NAME_None,
+			FAppStyle::GetAppStyleSetName()
+		)
+{
+}
+
+void FCameraAssetEditorCommands::RegisterCommands()
+{
+	UI_COMMAND(Build, "Build", "Builds the asset and refreshes it in PIE",
+			EUserInterfaceActionType::Button, FInputChord(EKeys::F7));
+}
+ 
+#undef LOCTEXT_NAMESPACE
+
diff --git a/Engine/Plugins/Cameras/GameplayCameras/Source/GameplayCamerasEditor/Private/Commands/CameraModeAssetEditorCommands.cpp b/Engine/Plugins/Cameras/GameplayCameras/Source/GameplayCamerasEditor/Private/Commands/CameraModeAssetEditorCommands.cpp
new file mode 100644
--- /dev/null
+++ b/Engine/Plugins/Cameras/GameplayCameras/Source/GameplayCamerasEditor/Private/Commands/CameraModeAssetEditorCommands.cpp
@@ -0,0 +1,29 @@
+// Copyright Epic Games, Inc. All Rights Reserved.
+
+#include "Commands/CameraModeAssetEditorCommands.h"
+
+#include "Framework/Commands/Commands.h"
+#include "Framework/Commands/InputChord.h"
+#include "Framework/Commands/UICommandInfo.h"
+#include "Styling/AppStyle.h"
+
+#define LOCTEXT_NAMESPACE "SequencerCommands"
+
+FCameraModeAssetEditorCommands::FCameraModeAssetEditorCommands()
+	: TCommands<FCameraModeAssetEditorCommands>(
+			"GameplayCameras",
+			NSLOCTEXT("Contexts", "GameplayCameras", "Gameplay Cameras"),
+			NAME_None,
+			FAppStyle::GetAppStyleSetName()
+		)
+{
+}
+
+void FCameraModeAssetEditorCommands::RegisterCommands()
+{
+	UI_COMMAND(Build, "Build", "Builds the asset and refreshes it in PIE",
+			EUserInterfaceActionType::Button, FInputChord(EKeys::F7));
+}
+ 
+#undef LOCTEXT_NAMESPACE
+
diff --git a/Engine/Plugins/Cameras/GameplayCameras/Source/GameplayCamerasEditor/Private/GameplayCamerasEditorModule.cpp b/Engine/Plugins/Cameras/GameplayCameras/Source/GameplayCamerasEditor/Private/GameplayCamerasEditorModule.cpp
--- a/Engine/Plugins/Cameras/GameplayCameras/Source/GameplayCamerasEditor/Private/GameplayCamerasEditorModule.cpp
+++ b/Engine/Plugins/Cameras/GameplayCameras/Source/GameplayCamerasEditor/Private/GameplayCamerasEditorModule.cpp
@@ -2,20 +2,24 @@
 
 #include "GameplayCamerasEditorSettings.h"
 
+#include "AssetTools/CameraAssetEditor.h"
+#include "AssetTools/CameraModeAssetEditor.h"
+#include "Commands/CameraAssetEditorCommands.h"
+#include "Commands/CameraModeAssetEditorCommands.h"
 #include "GameplayCamerasLiveEditManager.h"
-#include "IGameplayCamerasModule.h"
 #include "IGameplayCamerasEditorModule.h"
+#include "IGameplayCamerasModule.h"
 #include "ISettingsModule.h"
 #include "Modules/ModuleManager.h"
-#include "AssetTools/CameraAssetEditor.h"
-#include "AssetTools/CameraModeAssetEditor.h"
 #include "Styles/GameplayCamerasEditorStyle.h"
+#include "ToolMenus.h"
 #include "Toolkits/CameraAssetEditorToolkit.h"
 #include "Toolkits/CameraModeAssetEditorToolkit.h"
 
 #define LOCTEXT_NAMESPACE "GameplayCamerasEditor"
 
 const FName IGameplayCamerasEditorModule::GameplayCamerasEditorAppIdentifier("GameplayCamerasEditorApp");
+const FName IGameplayCamerasEditorModule::CameraModeAssetEditorToolBarName("CameraModeAssetEditor.ToolBar");
 
 /**
  * Implements the FGameplayCamerasEditor module.
@@ -31,10 +35,18 @@ public:
 	{
 		RegisterSettings();
 		InitializeLiveEditManager();
+
+		UToolMenus::RegisterStartupCallback(FSimpleMulticastDelegate::FDelegate::CreateRaw(
+					this, &FGameplayCamerasEditorModule::RegisterMenus));
 	}
 
 	virtual void ShutdownModule() override
 	{
+		UToolMenus::UnRegisterStartupCallback(this);
+
+		FCameraAssetEditorCommands::Unregister();
+		FCameraModeAssetEditorCommands::Unregister();
+
 		UnregisterSettings();
 		TeardownLiveEditManager();
 	}
@@ -82,6 +94,12 @@ private:
 		}
 	}
 
+	void RegisterMenus()
+	{
+		FCameraAssetEditorCommands::Register();
+		FCameraModeAssetEditorCommands::Register();	
+	}
+
 	void InitializeLiveEditManager()
 	{
 		LiveEditManager = MakeShared<FGameplayCamerasLiveEditManager>();
diff --git a/Engine/Plugins/Cameras/GameplayCameras/Source/GameplayCamerasEditor/Private/GameplayCamerasLiveEditManager.cpp b/Engine/Plugins/Cameras/GameplayCameras/Source/GameplayCamerasEditor/Private/GameplayCamerasLiveEditManager.cpp
--- a/Engine/Plugins/Cameras/GameplayCameras/Source/GameplayCamerasEditor/Private/GameplayCamerasLiveEditManager.cpp
+++ b/Engine/Plugins/Cameras/GameplayCameras/Source/GameplayCamerasEditor/Private/GameplayCamerasLiveEditManager.cpp
@@ -2,9 +2,62 @@
 
 #include "GameplayCamerasLiveEditManager.h"
 
+#include "IGameplayCamerasLiveEditListener.h"
+#include "UObject/Package.h"
 #include "UObject/UObjectGlobals.h"
 
 FGameplayCamerasLiveEditManager::FGameplayCamerasLiveEditManager()
 {
 }
 
+void FGameplayCamerasLiveEditManager::NotifyPostBuildAsset(const UPackage* InAssetPackage) const
+{
+	if (const FListenerArray* Listeners = ListenerMap.Find(InAssetPackage))
+	{
+		FGameplayCameraAssetBuildEvent BuildEvent;
+		BuildEvent.AssetPackage = InAssetPackage;
+
+		for (IGameplayCamerasLiveEditListener* Listener : *Listeners)
+		{
+			Listener->PostBuildAsset(BuildEvent);
+		}
+	}
+}
+
+void FGameplayCamerasLiveEditManager::AddListener(const UPackage* InAssetPackage, IGameplayCamerasLiveEditListener* Listener)
+{
+	if (ensure(InAssetPackage && Listener))
+	{
+		FListenerArray& Listeners = ListenerMap.FindOrAdd(InAssetPackage);
+		Listeners.Add(Listener);
+	}
+}
+
+void FGameplayCamerasLiveEditManager::RemoveListener(const UPackage* InAssetPackage, IGameplayCamerasLiveEditListener* Listener)
+{
+	if (ensure(InAssetPackage && Listener))
+	{
+		FListenerArray* Listeners = ListenerMap.Find(InAssetPackage);
+		if (ensure(Listeners))
+		{
+			const int32 NumRemoved = Listeners->RemoveSwap(Listener);
+			ensure(NumRemoved == 1);
+			if (Listeners->IsEmpty())
+			{
+				ListenerMap.Remove(InAssetPackage);
+			}
+		}
+	}
+}
+
+void FGameplayCamerasLiveEditManager::RemoveGarbage()
+{
+	for (auto It = ListenerMap.CreateIterator(); It; ++It)
+	{
+		if (!It.Key().IsValid())
+		{
+			It.RemoveCurrent();
+		}
+	}
+}
+
diff --git a/Engine/Plugins/Cameras/GameplayCameras/Source/GameplayCamerasEditor/Private/GameplayCamerasLiveEditManager.h b/Engine/Plugins/Cameras/GameplayCameras/Source/GameplayCamerasEditor/Private/GameplayCamerasLiveEditManager.h
--- a/Engine/Plugins/Cameras/GameplayCameras/Source/GameplayCamerasEditor/Private/GameplayCamerasLiveEditManager.h
+++ b/Engine/Plugins/Cameras/GameplayCameras/Source/GameplayCamerasEditor/Private/GameplayCamerasLiveEditManager.h
@@ -7,10 +7,29 @@
 #include "CoreTypes.h"
 #include "UObject/WeakObjectPtr.h"
 
+class UPackage;
+
 class FGameplayCamerasLiveEditManager : public IGameplayCamerasLiveEditManager
 {
 public:
 
 	FGameplayCamerasLiveEditManager();
+
+public:
+
+	// IGameplayCamerasLiveEditManager interface
+	virtual void NotifyPostBuildAsset(const UPackage* InAssetPackage) const override;
+	virtual void AddListener(const UPackage* InAssetPackage, IGameplayCamerasLiveEditListener* Listener) override;
+	virtual void RemoveListener(const UPackage* InAssetPackage, IGameplayCamerasLiveEditListener* Listener) override;
+
+private:
+
+	void RemoveGarbage();
+
+private:
+
+	using FListenerArray = TArray<IGameplayCamerasLiveEditListener*>;
+	using FListenerMap = TMap<TWeakObjectPtr<const UPackage>, FListenerArray>;
+	FListenerMap ListenerMap;
 };
 
diff --git a/Engine/Plugins/Cameras/GameplayCameras/Source/GameplayCamerasEditor/Private/Toolkits/CameraModeAssetEditorToolkit.cpp b/Engine/Plugins/Cameras/GameplayCameras/Source/GameplayCamerasEditor/Private/Toolkits/CameraModeAssetEditorToolkit.cpp
--- a/Engine/Plugins/Cameras/GameplayCameras/Source/GameplayCamerasEditor/Private/Toolkits/CameraModeAssetEditorToolkit.cpp
+++ b/Engine/Plugins/Cameras/GameplayCameras/Source/GameplayCamerasEditor/Private/Toolkits/CameraModeAssetEditorToolkit.cpp
@@ -3,9 +3,12 @@
 #include "Toolkits/CameraModeAssetEditorToolkit.h"
 
 #include "AssetTools/CameraModeAssetEditor.h"
+#include "Commands/CameraModeAssetEditorCommands.h"
 #include "Core/CameraMode.h"
 #include "EditorModeManager.h"
 #include "Framework/Docking/LayoutExtender.h"
+#include "IGameplayCamerasModule.h"
+#include "IGameplayCamerasLiveEditManager.h"
 #include "IMessageLogListing.h"
 #include "MessageLogInitializationOptions.h"
 #include "MessageLogModule.h"
@@ -19,6 +22,8 @@ const FName FCameraModeAssetEditorToolkit::DetailsViewTabId(TEXT("CameraModeAsse
 
 FCameraModeAssetEditorToolkit::FCameraModeAssetEditorToolkit(UCameraModeAssetEditor* InOwningAssetEditor)
 	: FBaseAssetToolkit(InOwningAssetEditor)
+	, CommandBindings(new FUICommandList())
+	, CameraModeAsset(InOwningAssetEditor->GetCameraModeAsset())
 {
 	// Override base class default layout.
 	StandaloneDefaultLayout = FTabManager::NewLayout("CameraModeAssetEditor_Layout")
@@ -74,6 +79,7 @@ void FCameraModeAssetEditorToolkit::CreateWidgets()
 	FDetailsViewArgs DetailsViewArgs;
 	DetailsViewArgs.NameAreaSettings = FDetailsViewArgs::HideNameArea;
 	DetailsViewArgs.bHideSelectionTip = true;
+	DetailsViewArgs.NotifyHook = this;
 	DetailsView = PropertyEditorModule.CreateDetailView(DetailsViewArgs);
 
 	// Create the message log.
@@ -85,10 +91,89 @@ void FCameraModeAssetEditorToolkit::CreateWidgets()
 	LogOptions.MaxPageCount = 1;
 	StatsListing = MessageLogModule.CreateLogListing("CameraModeAssetEditorStats", LogOptions);
 
-	Stats = MessageLogModule.CreateLogListingWidget(StatsListing.ToSharedRef());}
+	Stats = MessageLogModule.CreateLogListingWidget(StatsListing.ToSharedRef());
+}
+
+void FCameraModeAssetEditorToolkit::RegisterToolbar()
+{
+	FName ParentName;
+	const FName MenuName = GetToolMenuToolbarName(ParentName);
+	UToolMenus* ToolMenus = UToolMenus::Get();
+	if (!ToolMenus->IsMenuRegistered(MenuName))
+	{
+		const FCameraModeAssetEditorCommands& Commands = FCameraModeAssetEditorCommands::Get();
+
+		FToolMenuOwnerScoped ToolMenuOwnerScope(this);
+
+		UToolMenu* ToolBarMenu = UToolMenus::Get()->RegisterMenu(
+				MenuName, ParentName, EMultiBoxType::ToolBar);
+
+		FToolMenuSection& ToolBarSection = ToolBarMenu->FindOrAddSection("Build");
+		FToolMenuEntry BuildButton = FToolMenuEntry::InitToolBarButton(Commands.Build);
+		BuildButton.Icon = TAttribute<FSlateIcon>(
+				SharedThis(this), &FCameraModeAssetEditorToolkit::GetBuildButtonIcon);
+		BuildButton.ToolTip = TAttribute<FText>(
+				SharedThis(this), &FCameraModeAssetEditorToolkit::GetBuildButtonTooltip);
+		ToolBarSection.AddEntry(BuildButton);
+	}
+}
+
+FSlateIcon FCameraModeAssetEditorToolkit::GetBuildButtonIcon() const
+{
+	static const FName BuildStatusBackground("Blueprint.CompileStatus.Background");
+	static const FName BuildStatusUnknown("Blueprint.CompileStatus.Overlay.Unknown");
+	static const FName BuildStatusError("Blueprint.CompileStatus.Overlay.Error");
+	static const FName BuildStatusGood("Blueprint.CompileStatus.Overlay.Good");
+	static const FName BuildStatusWarning("Blueprint.CompileStatus.Overlay.Warning");
+
+	switch (CameraModeAsset->BuildStatus)
+	{
+		default:
+		case ECameraModeBuildStatus::Dirty:
+			return FSlateIcon(
+					FAppStyle::GetAppStyleSetName(), BuildStatusBackground, 
+					NAME_None, BuildStatusUnknown);
+		case ECameraModeBuildStatus::WithErrors:
+			return FSlateIcon(
+					FAppStyle::GetAppStyleSetName(), BuildStatusBackground, 
+					NAME_None, BuildStatusError);
+		case ECameraModeBuildStatus::Clean:
+			return FSlateIcon(
+					FAppStyle::GetAppStyleSetName(), BuildStatusBackground, 
+					NAME_None, BuildStatusGood);
+		case ECameraModeBuildStatus::CleanWithWarnings:
+			return FSlateIcon(
+					FAppStyle::GetAppStyleSetName(), BuildStatusBackground, 
+					NAME_None, BuildStatusWarning);
+	}
+}
+
+FText FCameraModeAssetEditorToolkit::GetBuildButtonTooltip() const
+{
+	switch (CameraModeAsset->BuildStatus)
+	{
+		default:
+		case ECameraModeBuildStatus::Dirty:
+			return LOCTEXT("BuildButtonStatusDirty", "Dirty or unknown, should rebuild");
+		case ECameraModeBuildStatus::WithErrors:
+			return LOCTEXT("BuildButtonStatusWithErrors", "There were errors during the build, see the log window for details");
+		case ECameraModeBuildStatus::Clean:
+			return LOCTEXT("BuildButtonStatusClean", "Good to go");
+		case ECameraModeBuildStatus::CleanWithWarnings:
+			return LOCTEXT("BuildButtonStatusCleanWithWarnings", "There were warnings during the build, see the log window for details");
+	}
+}
 
 void FCameraModeAssetEditorToolkit::PostInitAssetEditor()
 {
+	const FCameraModeAssetEditorCommands& Commands = FCameraModeAssetEditorCommands::Get();
+
+	ToolkitCommands->MapAction(
+		Commands.Build,
+		FExecuteAction::CreateSP(this, &FCameraModeAssetEditorToolkit::OnBuild));
+
+	IGameplayCamerasModule& GameplayCamerasModule = FModuleManager::GetModuleChecked<IGameplayCamerasModule>("GameplayCameras");
+	LiveEditManager = GameplayCamerasModule.GetLiveEditManager();
 }
 
 FText FCameraModeAssetEditorToolkit::GetBaseToolkitName() const
@@ -112,5 +197,23 @@ FLinearColor FCameraModeAssetEditorToolkit::GetWorldCentricTabColorScale() const
 	return FLinearColor(0.7, 0.0f, 0.0f, 0.5f);
 }
 
+void FCameraModeAssetEditorToolkit::NotifyPostChange(const FPropertyChangedEvent& PropertyChangedEvent, FProperty* PropertyThatChanged)
+{
+	CameraModeAsset->BuildStatus = ECameraModeBuildStatus::Dirty;
+}
+
+void FCameraModeAssetEditorToolkit::OnBuild()
+{
+	CameraModeAsset->BuildStatus = ECameraModeBuildStatus::Clean;
+
+	FCameraModePackages BuiltPackages;
+	CameraModeAsset->GatherPackages(BuiltPackages);
+
+	for (const UPackage* BuiltPackage : BuiltPackages)
+	{
+		LiveEditManager->NotifyPostBuildAsset(BuiltPackage);
+	}
+}
+
 #undef LOCTEXT_NAMESPACE
 
diff --git a/Engine/Plugins/Cameras/GameplayCameras/Source/GameplayCamerasEditor/Public/AssetTools/CameraModeAssetEditor.h b/Engine/Plugins/Cameras/GameplayCameras/Source/GameplayCamerasEditor/Public/AssetTools/CameraModeAssetEditor.h
--- a/Engine/Plugins/Cameras/GameplayCameras/Source/GameplayCamerasEditor/Public/AssetTools/CameraModeAssetEditor.h
+++ b/Engine/Plugins/Cameras/GameplayCameras/Source/GameplayCamerasEditor/Public/AssetTools/CameraModeAssetEditor.h
@@ -21,6 +21,8 @@ public:
 
 	void Initialize(TObjectPtr<UCameraMode> InCameraModeAsset);
 
+	UCameraMode* GetCameraModeAsset() const { return CameraModeAsset; }
+
 public:
 
 	// UAssetEditor interface
diff --git a/Engine/Plugins/Cameras/GameplayCameras/Source/GameplayCamerasEditor/Public/Commands/CameraAssetEditorCommands.h b/Engine/Plugins/Cameras/GameplayCameras/Source/GameplayCamerasEditor/Public/Commands/CameraAssetEditorCommands.h
new file mode 100644
--- /dev/null
+++ b/Engine/Plugins/Cameras/GameplayCameras/Source/GameplayCamerasEditor/Public/Commands/CameraAssetEditorCommands.h
@@ -0,0 +1,24 @@
+// Copyright Epic Games, Inc. All Rights Reserved.
+
+#pragma once
+
+#include "CoreTypes.h"
+#include "Framework/Commands/Commands.h"
+
+class FUICommandInfo;
+class UClass;
+struct FInputChord;
+
+class FCameraAssetEditorCommands : public TCommands<FCameraAssetEditorCommands>
+{
+public:
+
+	FCameraAssetEditorCommands();
+
+	virtual void RegisterCommands() override;
+	
+public:
+
+	TSharedPtr<FUICommandInfo> Build;
+};
+
diff --git a/Engine/Plugins/Cameras/GameplayCameras/Source/GameplayCamerasEditor/Public/Commands/CameraModeAssetEditorCommands.h b/Engine/Plugins/Cameras/GameplayCameras/Source/GameplayCamerasEditor/Public/Commands/CameraModeAssetEditorCommands.h
new file mode 100644
--- /dev/null
+++ b/Engine/Plugins/Cameras/GameplayCameras/Source/GameplayCamerasEditor/Public/Commands/CameraModeAssetEditorCommands.h
@@ -0,0 +1,24 @@
+// Copyright Epic Games, Inc. All Rights Reserved.
+
+#pragma once
+
+#include "CoreTypes.h"
+#include "Framework/Commands/Commands.h"
+
+class FUICommandInfo;
+class UClass;
+struct FInputChord;
+
+class FCameraModeAssetEditorCommands : public TCommands<FCameraModeAssetEditorCommands>
+{
+public:
+
+	FCameraModeAssetEditorCommands();
+
+	virtual void RegisterCommands() override;
+	
+public:
+
+	TSharedPtr<FUICommandInfo> Build;
+};
+
diff --git a/Engine/Plugins/Cameras/GameplayCameras/Source/GameplayCamerasEditor/Public/IGameplayCamerasEditorModule.h b/Engine/Plugins/Cameras/GameplayCameras/Source/GameplayCamerasEditor/Public/IGameplayCamerasEditorModule.h
--- a/Engine/Plugins/Cameras/GameplayCameras/Source/GameplayCamerasEditor/Public/IGameplayCamerasEditorModule.h
+++ b/Engine/Plugins/Cameras/GameplayCameras/Source/GameplayCamerasEditor/Public/IGameplayCamerasEditorModule.h
@@ -20,8 +20,12 @@ public:
 
 	static const FName GameplayCamerasEditorAppIdentifier;
 
+	static const FName CameraModeAssetEditorToolBarName;
+
 	virtual ~IGameplayCamerasEditorModule() = default;
 
+public:
+
 	/** Creates an editor for the given camera asset */
 	virtual UCameraAssetEditor* CreateCameraAssetEditor(const EToolkitMode::Type Mode, const TSharedPtr<IToolkitHost>& InitToolkitHost, UCameraAsset* CameraAsset) = 0;
 
diff --git a/Engine/Plugins/Cameras/GameplayCameras/Source/GameplayCamerasEditor/Public/Toolkits/CameraModeAssetEditorToolkit.h b/Engine/Plugins/Cameras/GameplayCameras/Source/GameplayCamerasEditor/Public/Toolkits/CameraModeAssetEditorToolkit.h
--- a/Engine/Plugins/Cameras/GameplayCameras/Source/GameplayCamerasEditor/Public/Toolkits/CameraModeAssetEditorToolkit.h
+++ b/Engine/Plugins/Cameras/GameplayCameras/Source/GameplayCamerasEditor/Public/Toolkits/CameraModeAssetEditorToolkit.h
@@ -3,9 +3,11 @@
 #pragma once
 
 #include "CoreMinimal.h"
+#include "Misc/NotifyHook.h"
 #include "Tools/BaseAssetToolkit.h"
 #include "UObject/GCObject.h"
 
+class IGameplayCamerasLiveEditManager;
 class IMessageLogListing;
 class UCameraMode;
 class UCameraModeAssetEditor;
@@ -17,6 +19,7 @@ class SWidget;
 class FCameraModeAssetEditorToolkit 
 	: public FBaseAssetToolkit
 	, public FGCObject
+	, public FNotifyHook
 {
 public:
 
@@ -29,6 +32,7 @@ protected:
 	virtual void RegisterTabSpawners(const TSharedRef<class FTabManager>& TabManager) override;
 	virtual void UnregisterTabSpawners(const TSharedRef<FTabManager>& InTabManager) override;
 	virtual void CreateWidgets() override;
+	virtual void RegisterToolbar() override;
 	virtual void PostInitAssetEditor() override;
 
 	// IToolkit interface
@@ -47,10 +51,23 @@ protected:
 		return TEXT("FCameraModeAssetEditorToolkit");
 	}
 
+	// FNotifyHook interface
+	virtual void NotifyPostChange(const FPropertyChangedEvent& PropertyChangedEvent, FProperty* PropertyThatChanged) override;
+
+private:
+
+	FSlateIcon GetBuildButtonIcon() const;
+	FText GetBuildButtonTooltip() const;
+
+	void OnBuild();
+
 private:
 
 	static const FName DetailsViewTabId;
 
+	/** Command bindings */
+	TSharedRef<FUICommandList> CommandBindings;
+
 	/** Message log widget */
 	TSharedPtr<SWidget> Stats;
 
@@ -59,5 +76,8 @@ private:
 
 	/** The asset being edited */
 	TObjectPtr<UCameraMode> CameraModeAsset;
+
+	/** Live edit manager for updating the assets in the runtime */
+	TSharedPtr<IGameplayCamerasLiveEditManager> LiveEditManager;
 };
 
