Cameras: rename "camera mode" to "camera rig"

While "camera mode" is a somewhat standard term for a specific sort of gameplay camera, it's actually too specific. The term "camera rig" is more generic and can describe other hierarchies of nodes that don't have anything to do with a given gameplay camera behavior, such as those nodes that run before or after the blend stack, or those that are re-used as "prefabs" of sorts. A "camera mode" is therefore only a top-level camera rig that happens to run in the main blend stack.

Of course, another reason for this change is to avoid the easy confusion between "camera mode" and "camera node".

diff --git a/Engine/Plugins/Cameras/GameplayCameras/Source/GameplayCameras/Private/Core/BlendStackCameraNode.cpp b/Engine/Plugins/Cameras/GameplayCameras/Source/GameplayCameras/Private/Core/BlendStackCameraNode.cpp
--- a/Engine/Plugins/Cameras/GameplayCameras/Source/GameplayCameras/Private/Core/BlendStackCameraNode.cpp
+++ b/Engine/Plugins/Cameras/GameplayCameras/Source/GameplayCameras/Private/Core/BlendStackCameraNode.cpp
@@ -6,10 +6,10 @@
 #include "Core/BlendStackRootCameraNode.h"
 #include "Core/CameraAsset.h"
 #include "Core/CameraEvaluationContext.h"
-#include "Core/CameraMode.h"
+#include "Core/CameraRigAsset.h"
 #include "Core/CameraSystemEvaluator.h"
-#include "IGameplayCamerasModule.h"
 #include "IGameplayCamerasLiveEditManager.h"
+#include "IGameplayCamerasModule.h"
 #include "Modules/ModuleManager.h"
 #include "Nodes/Blends/PopBlendCameraNode.h"
 
@@ -27,30 +27,30 @@ void FBlendStackCameraNodeEvaluator::Push(const FBlendStackCameraPushParams& Par
 	if (!Entries.IsEmpty())
 	{
 		// Don't push anything if what is being requested is already the
-		// active camera mode.
-		const FCameraModeEntry& TopEntry(Entries.Top());
+		// active camera rig.
+		const FCameraRigEntry& TopEntry(Entries.Top());
 		if (!TopEntry.bIsFrozen
-				&& TopEntry.CameraMode == Params.CameraMode
+				&& TopEntry.CameraRig == Params.CameraRig
 				&& TopEntry.EvaluationContext == Params.EvaluationContext)
 		{
 			return;
 		}
 	}
 
-	// Create the new root node to wrap the new camera mode's root node, and the specific
+	// Create the new root node to wrap the new camera rig's root node, and the specific
 	// blend node for this transition.
 	// We need to const-cast here to be able to use our own blend stack node as the outer
 	// of the new node.
 	UObject* Outer = const_cast<UObject*>((UObject*)GetCameraNode());
 	UBlendStackRootCameraNode* EntryRootNode = NewObject<UBlendStackRootCameraNode>(Outer, NAME_None);
 	{
-		UCameraNode* ModeRootNode = Params.CameraMode->RootNode;
+		UCameraNode* ModeRootNode = Params.CameraRig->RootNode;
 		EntryRootNode->RootNode = ModeRootNode;
 
 		// Find a transition and use its blend. If no transition is found,
 		// make a camera cut transition.
 		UBlendCameraNode* ModeBlend = nullptr;
-		if (const FCameraModeTransition* Transition = FindTransition(Params))
+		if (const FCameraRigTransition* Transition = FindTransition(Params))
 		{
 			ModeBlend = Transition->Blend;
 		}
@@ -62,7 +62,7 @@ void FBlendStackCameraNodeEvaluator::Push(const FBlendStackCameraPushParams& Par
 	}
 
 	// Make the new stack entry, and use its storage buffer to build the tree of evaluators.
-	FCameraModeEntry NewEntry;
+	FCameraRigEntry NewEntry;
 
 	FCameraNodeEvaluatorTreeBuilderParams BuildParams;
 	BuildParams.RootCameraNode = EntryRootNode;
@@ -71,7 +71,7 @@ void FBlendStackCameraNodeEvaluator::Push(const FBlendStackCameraPushParams& Par
 	FCameraNodeEvaluator* RootEvaluator = NewEntry.EvaluatorStorage.BuildEvaluatorTree(BuildParams);
 
 	NewEntry.EvaluationContext = Params.EvaluationContext;
-	NewEntry.CameraMode = Params.CameraMode;
+	NewEntry.CameraRig = Params.CameraRig;
 	NewEntry.RootNode = EntryRootNode;
 	NewEntry.RootEvaluator = RootEvaluator->CastThisChecked<FBlendStackRootCameraNodeEvaluator>();
 	NewEntry.bIsFirstFrame = true;
@@ -79,7 +79,7 @@ void FBlendStackCameraNodeEvaluator::Push(const FBlendStackCameraPushParams& Par
 #if WITH_EDITOR
 	IGameplayCamerasModule& GameplayCamerasModule = FModuleManager::GetModuleChecked<IGameplayCamerasModule>("GameplayCameras");
 	TSharedPtr<IGameplayCamerasLiveEditManager> LiveEditManager = GameplayCamerasModule.GetLiveEditManager();
-	Params.CameraMode->GatherPackages(NewEntry.ListenedPackages);
+	Params.CameraRig->GatherPackages(NewEntry.ListenedPackages);
 	for (const UPackage* ListenPackage : NewEntry.ListenedPackages)
 	{
 		LiveEditManager->AddListener(ListenPackage, this);
@@ -94,7 +94,7 @@ void FBlendStackCameraNodeEvaluator::Push(const FBlendStackCameraPushParams& Par
 FCameraNodeEvaluatorChildrenView FBlendStackCameraNodeEvaluator::OnGetChildren()
 {
 	FCameraNodeEvaluatorChildrenView View;
-	for (FCameraModeEntry& Entry : Entries)
+	for (FCameraRigEntry& Entry : Entries)
 	{
 		View.Add(Entry.RootEvaluator);
 	}
@@ -111,7 +111,7 @@ void FBlendStackCameraNodeEvaluator::OnRun(const FCameraNodeEvaluationParams& Pa
 	const UBlendStackCameraNode* BlendStackNode = GetCameraNodeAs<UBlendStackCameraNode>();
 
 	// Start by evaluating all the root nodes in the stack.
-	for (FCameraModeEntry& Entry : Entries)
+	for (FCameraRigEntry& Entry : Entries)
 	{
 		const UCameraEvaluationContext* CurContext = Entry.EvaluationContext.Get();
 		if (UNLIKELY(CurContext == nullptr))
@@ -128,7 +128,7 @@ void FBlendStackCameraNodeEvaluator::OnRun(const FCameraNodeEvaluationParams& Pa
 
 		if (!Entry.bIsFrozen)
 		{
-			// If the context in which this camera mode runs doesn't have a valid result,
+			// If the context in which this camera rig runs doesn't have a valid result,
 			// skip it.
 			const FCameraNodeEvaluationResult& ContextResult(CurContext->GetInitialResult());
 			if (UNLIKELY(!ContextResult.bIsValid))
@@ -146,7 +146,7 @@ void FBlendStackCameraNodeEvaluator::OnRun(const FCameraNodeEvaluationParams& Pa
 			CurResult.bIsCameraCut = OutResult.bIsCameraCut || ContextResult.bIsCameraCut;
 			CurResult.bIsValid = true;
 
-			// Run the camera mode!
+			// Run the camera rig!
 			Entry.RootEvaluator->Run(CurParams, CurResult);
 		}
 		else
@@ -157,11 +157,11 @@ void FBlendStackCameraNodeEvaluator::OnRun(const FCameraNodeEvaluationParams& Pa
 	}
 
 	// Now blend all the results, keeping track of blends that have reached 100% so
-	// that we can remove any camera modes below (since the would have been completely
+	// that we can remove any camera rigs below (since the would have been completely
 	// blended out by that).
 	int32 EntryIndex = 0;
 	int32 PopEntriesBelow = INDEX_NONE;
-	for (FCameraModeEntry& Entry : Entries)
+	for (FCameraRigEntry& Entry : Entries)
 	{
 		FCameraNodeEvaluationResult& CurResult(Entry.Result);
 		if (UNLIKELY(!CurResult.bIsValid))
@@ -200,7 +200,7 @@ void FBlendStackCameraNodeEvaluator::OnRun(const FCameraNodeEvaluationParams& Pa
 		++EntryIndex;
 	}
 
-	// Pop out camera modes that have been blended out.
+	// Pop out camera rigs that have been blended out.
 	if (BlendStackNode->bAutoPop && PopEntriesBelow != INDEX_NONE)
 	{
 #if WITH_EDITOR
@@ -211,7 +211,7 @@ void FBlendStackCameraNodeEvaluator::OnRun(const FCameraNodeEvaluationParams& Pa
 		for (int32 Index = 0; Index < PopEntriesBelow; ++Index)
 		{
 #if WITH_EDITOR
-			const FCameraModeEntry& FirstEntry = Entries[0];
+			const FCameraRigEntry& FirstEntry = Entries[0];
 			for (const UPackage* ListenPackage : FirstEntry.ListenedPackages)
 			{
 				LiveEditManager->RemoveListener(ListenPackage, this);
@@ -223,42 +223,42 @@ void FBlendStackCameraNodeEvaluator::OnRun(const FCameraNodeEvaluationParams& Pa
 	}
 
 	// Reset first frame flags.
-	for (FCameraModeEntry& Entry : Entries)
+	for (FCameraRigEntry& Entry : Entries)
 	{
 		Entry.bIsFirstFrame = false;
 	}
 }
 
-const FCameraModeTransition* FBlendStackCameraNodeEvaluator::FindTransition(const FBlendStackCameraPushParams& Params) const
+const FCameraRigTransition* FBlendStackCameraNodeEvaluator::FindTransition(const FBlendStackCameraPushParams& Params) const
 {
 	const UBlendStackCameraNode* BlendStackNode = GetCameraNodeAs<UBlendStackCameraNode>();
 
 	const UCameraEvaluationContext* ToContext = Params.EvaluationContext.Get();
 	const UCameraAsset* ToCameraAsset = ToContext ? ToContext->GetCameraAsset() : nullptr;
-	const UCameraMode* ToCameraMode = Params.CameraMode;
+	const UCameraRigAsset* ToCameraRig = Params.CameraRig;
 
-	// Find a transition that works for blending towards ToCameraMode.
+	// Find a transition that works for blending towards ToCameraRig.
 	// If the stack isn't empty, we need to find a transition that works between the previous and 
-	// next camera modes. If the stack is empty, we blend the new camera mode in from nothing if
+	// next camera rigs. If the stack is empty, we blend the new camera rig in from nothing if
 	// appropriate.
 	if (!Entries.IsEmpty())
 	{
-		const FCameraModeTransition* TransitionToUse = nullptr;
+		const FCameraRigTransition* TransitionToUse = nullptr;
 
-		// Start by looking at exit transitions on the last active (top) camera mode.
-		const FCameraModeEntry& TopEntry = Entries.Top();
+		// Start by looking at exit transitions on the last active (top) camera rig.
+		const FCameraRigEntry& TopEntry = Entries.Top();
 
 		const UCameraEvaluationContext* FromContext = TopEntry.EvaluationContext.Get();
 		const UCameraAsset* FromCameraAsset = FromContext ? FromContext->GetCameraAsset() : nullptr;
-		const UCameraMode* FromCameraMode = TopEntry.CameraMode;
+		const UCameraRigAsset* FromCameraRig = TopEntry.CameraRig;
 
 		if (!TopEntry.bIsFrozen)
 		{
-			// Look for exit transitions on the last active camera mode itself.
+			// Look for exit transitions on the last active camera rig itself.
 			TransitionToUse = FindTransition(
-					FromCameraMode->ExitTransitions,
-					FromCameraMode, FromCameraAsset, false,
-					ToCameraMode, ToCameraAsset);
+					FromCameraRig->ExitTransitions,
+					FromCameraRig, FromCameraAsset, false,
+					ToCameraRig, ToCameraAsset);
 			if (TransitionToUse)
 			{
 				return TransitionToUse;
@@ -269,8 +269,8 @@ const FCameraModeTransition* FBlendStackCameraNodeEvaluator::FindTransition(cons
 			{
 				TransitionToUse = FindTransition(
 						FromCameraAsset->ExitTransitions,
-						FromCameraMode, FromCameraAsset, false,
-						ToCameraMode, ToCameraAsset);
+						FromCameraRig, FromCameraAsset, false,
+						ToCameraRig, ToCameraAsset);
 				if (TransitionToUse)
 				{
 					return TransitionToUse;
@@ -278,11 +278,11 @@ const FCameraModeTransition* FBlendStackCameraNodeEvaluator::FindTransition(cons
 			}
 		}
 
-		// Now look at enter transitions on the new camera mode.
+		// Now look at enter transitions on the new camera rig.
 		TransitionToUse = FindTransition(
-				ToCameraMode->EnterTransitions,
-				FromCameraMode, FromCameraAsset, TopEntry.bIsFrozen,
-				ToCameraMode, ToCameraAsset);
+				ToCameraRig->EnterTransitions,
+				FromCameraRig, FromCameraAsset, TopEntry.bIsFrozen,
+				ToCameraRig, ToCameraAsset);
 		if (TransitionToUse)
 		{
 			return TransitionToUse;
@@ -293,41 +293,41 @@ const FCameraModeTransition* FBlendStackCameraNodeEvaluator::FindTransition(cons
 		{
 			TransitionToUse = FindTransition(
 					ToCameraAsset->EnterTransitions,
-					FromCameraMode, FromCameraAsset, TopEntry.bIsFrozen,
-					ToCameraMode, ToCameraAsset);
+					FromCameraRig, FromCameraAsset, TopEntry.bIsFrozen,
+					ToCameraRig, ToCameraAsset);
 			if (TransitionToUse)
 			{
 				return TransitionToUse;
 			}
 		}
 	}
-	else if (BlendStackNode->bBlendFirstCameraMode)
+	else if (BlendStackNode->bBlendFirstCameraRig)
 	{
 		return FindTransition(
-				ToCameraMode->EnterTransitions,
+				ToCameraRig->EnterTransitions,
 				nullptr, nullptr, false,
-				ToCameraMode, ToCameraAsset);
+				ToCameraRig, ToCameraAsset);
 	}
 
 	return nullptr;
 }
 
-const FCameraModeTransition* FBlendStackCameraNodeEvaluator::FindTransition(
-			TArrayView<const FCameraModeTransition> Transitions, 
-			const UCameraMode* FromCameraMode, const UCameraAsset* FromCameraAsset, bool bFromFrozen,
-			const UCameraMode* ToCameraMode, const UCameraAsset* ToCameraAsset) const
+const FCameraRigTransition* FBlendStackCameraNodeEvaluator::FindTransition(
+			TArrayView<const FCameraRigTransition> Transitions, 
+			const UCameraRigAsset* FromCameraRig, const UCameraAsset* FromCameraAsset, bool bFromFrozen,
+			const UCameraRigAsset* ToCameraRig, const UCameraAsset* ToCameraAsset) const
 {
-	FCameraModeTransitionConditionMatchParams MatchParams;
-	MatchParams.FromCameraMode = FromCameraMode;
+	FCameraRigTransitionConditionMatchParams MatchParams;
+	MatchParams.FromCameraRig = FromCameraRig;
 	MatchParams.FromCameraAsset = FromCameraAsset;
-	MatchParams.ToCameraMode = ToCameraMode;
+	MatchParams.ToCameraRig = ToCameraRig;
 	MatchParams.ToCameraAsset = ToCameraAsset;
 
 	// The transition should be used if all its conditions pass.
-	for (const FCameraModeTransition& Transition : Transitions)
+	for (const FCameraRigTransition& Transition : Transitions)
 	{
 		bool bConditionsPass = true;
-		for (const UCameraModeTransitionCondition* Condition : Transition.Conditions)
+		for (const UCameraRigTransitionCondition* Condition : Transition.Conditions)
 		{
 			if (!Condition->TransitionMatches(MatchParams))
 			{
@@ -349,17 +349,17 @@ const FCameraModeTransition* FBlendStackCameraNodeEvaluator::FindTransition(
 
 void FBlendStackCameraNodeEvaluator::OnPostBuildAsset(const FGameplayCameraAssetBuildEvent& BuildEvent)
 {
-	for (FCameraModeEntry& Entry : Entries)
+	for (FCameraRigEntry& Entry : Entries)
 	{
 		const bool bRebuildEntry = Entry.ListenedPackages.Contains(BuildEvent.AssetPackage);
 		if (bRebuildEntry)
 		{
 			Entry.EvaluatorStorage.DestroyEvaluatorTree();
 
-			// Re-assign the root node in case the camera mode's root was changed.
-			Entry.RootNode->RootNode = Entry.CameraMode->RootNode;
+			// Re-assign the root node in case the camera rig's root was changed.
+			Entry.RootNode->RootNode = Entry.CameraRig->RootNode;
 
-			// Remove the blend on the root node, since we don't want the reloaded camera mode to re-blend-in
+			// Remove the blend on the root node, since we don't want the reloaded camera rig to re-blend-in
 			// for no good reason.
 			Entry.RootNode->Blend = nullptr;
 
diff --git a/Engine/Plugins/Cameras/GameplayCameras/Source/GameplayCameras/Private/Core/BlendStackRootCameraNode.cpp b/Engine/Plugins/Cameras/GameplayCameras/Source/GameplayCameras/Private/Core/BlendStackRootCameraNode.cpp
--- a/Engine/Plugins/Cameras/GameplayCameras/Source/GameplayCameras/Private/Core/BlendStackRootCameraNode.cpp
+++ b/Engine/Plugins/Cameras/GameplayCameras/Source/GameplayCameras/Private/Core/BlendStackRootCameraNode.cpp
@@ -3,7 +3,7 @@
 #include "Core/BlendStackRootCameraNode.h"
 
 #include "Core/BlendCameraNode.h"
-#include "Core/CameraMode.h"
+#include "Core/CameraRigAsset.h"
 
 #include UE_INLINE_GENERATED_CPP_BY_NAME(BlendStackRootCameraNode)
 
diff --git a/Engine/Plugins/Cameras/GameplayCameras/Source/GameplayCameras/Private/Core/BlendStackRootCameraNode.h b/Engine/Plugins/Cameras/GameplayCameras/Source/GameplayCameras/Private/Core/BlendStackRootCameraNode.h
--- a/Engine/Plugins/Cameras/GameplayCameras/Source/GameplayCameras/Private/Core/BlendStackRootCameraNode.h
+++ b/Engine/Plugins/Cameras/GameplayCameras/Source/GameplayCameras/Private/Core/BlendStackRootCameraNode.h
@@ -9,11 +9,11 @@
 
 class FBlendCameraNodeEvaluator;
 class UBlendCameraNode;
-class UCameraMode;
+class UCameraRigAsset;
 
 /**
- * Root camera node for running a camera mode in a blend stack.
- * This camera node wraps both the camera mode's root node, and the
+ * Root camera node for running a camera rig in a blend stack.
+ * This camera node wraps both the camera rig's root node, and the
  * blend node used to blend it.
  */
 UCLASS(MinimalAPI)
@@ -29,7 +29,7 @@ protected:
 
 public:
 
-	/** The blend to use on the camera mode. */
+	/** The blend to use on the camera rig. */
 	UPROPERTY()
 	TObjectPtr<UBlendCameraNode> Blend;
 
diff --git a/Engine/Plugins/Cameras/GameplayCameras/Source/GameplayCameras/Private/Core/CameraMode.cpp b/Engine/Plugins/Cameras/GameplayCameras/Source/GameplayCameras/Private/Core/CameraMode.cpp
deleted file mode 100644
diff --git a/Engine/Plugins/Cameras/GameplayCameras/Source/GameplayCameras/Private/Core/CameraModeTransition.cpp b/Engine/Plugins/Cameras/GameplayCameras/Source/GameplayCameras/Private/Core/CameraModeTransition.cpp
deleted file mode 100644
diff --git a/Engine/Plugins/Cameras/GameplayCameras/Source/GameplayCameras/Private/Core/CameraNodeEvaluatorStorage.cpp b/Engine/Plugins/Cameras/GameplayCameras/Source/GameplayCameras/Private/Core/CameraNodeEvaluatorStorage.cpp
--- a/Engine/Plugins/Cameras/GameplayCameras/Source/GameplayCameras/Private/Core/CameraNodeEvaluatorStorage.cpp
+++ b/Engine/Plugins/Cameras/GameplayCameras/Source/GameplayCameras/Private/Core/CameraNodeEvaluatorStorage.cpp
@@ -2,9 +2,9 @@
 
 #include "Core/CameraNodeEvaluatorStorage.h"
 
-#include "Core/CameraMode.h"
 #include "Core/CameraNode.h"
 #include "Core/CameraNodeEvaluator.h"
+#include "Core/CameraRigAsset.h"
 #include "UObject/Package.h"
 
 FCameraNodeEvaluatorStorage::FCameraNodeEvaluatorStorage()
@@ -32,13 +32,13 @@ FCameraNodeEvaluatorStorage::~FCameraNodeEvaluatorStorage()
 	DestroyEvaluatorTree(true);
 }
 
-FCameraNodeEvaluatorTreeAllocationInfo FCameraNodeEvaluatorStorage::ComputeTreeInfo(const UCameraMode* CameraMode)
+FCameraNodeEvaluatorTreeAllocationInfo FCameraNodeEvaluatorStorage::ComputeTreeInfo(const UCameraRigAsset* CameraRig)
 {
 	int16 MaxSize = 0;
 	int16 MaxAlignment = 0;
 
 	// This isn't ideal, maybe we need a GetChildren API on the UCameraNode class too.
-	UPackage* Package = CameraMode->GetPackage();
+	UPackage* Package = CameraRig->GetPackage();
 	ForEachObjectWithPackage(Package, [&MaxSize, &MaxAlignment](UObject* Object)
 	{
 		if (UCameraNode* CurNode = Cast<UCameraNode>(Object))
diff --git a/Engine/Plugins/Cameras/GameplayCameras/Source/GameplayCameras/Private/Core/CameraRigAsset.cpp b/Engine/Plugins/Cameras/GameplayCameras/Source/GameplayCameras/Private/Core/CameraRigAsset.cpp
new file mode 100644
--- /dev/null
+++ b/Engine/Plugins/Cameras/GameplayCameras/Source/GameplayCameras/Private/Core/CameraRigAsset.cpp
@@ -0,0 +1,32 @@
+// Copyright Epic Games, Inc. All Rights Reserved.
+
+#include "Core/CameraRigAsset.h"
+
+#include "Core/CameraNode.h"
+
+#include UE_INLINE_GENERATED_CPP_BY_NAME(CameraRigAsset)
+
+#if WITH_EDITOR
+
+void UCameraRigAsset::GatherPackages(FCameraRigPackages& OutPackages) const
+{
+	TArray<UCameraNode*> NodeStack;
+	if (RootNode)
+	{
+		NodeStack.Add(RootNode);
+	}
+	while (!NodeStack.IsEmpty())
+	{
+		UCameraNode* CurrentNode = NodeStack.Pop();
+		const UPackage* CurrentPackage = CurrentNode->GetOutermost();
+		OutPackages.AddUnique(CurrentPackage);
+
+		FCameraNodeChildrenView CurrentChildren = CurrentNode->GetChildren();
+		for (UCameraNode* CurrentChild : ReverseIterate(CurrentChildren))
+		{
+			NodeStack.Add(CurrentChild);
+		}
+	}
+}
+
+#endif  // WITH_EDITOR
diff --git a/Engine/Plugins/Cameras/GameplayCameras/Source/GameplayCameras/Private/Core/CameraRigTransition.cpp b/Engine/Plugins/Cameras/GameplayCameras/Source/GameplayCameras/Private/Core/CameraRigTransition.cpp
new file mode 100644
--- /dev/null
+++ b/Engine/Plugins/Cameras/GameplayCameras/Source/GameplayCameras/Private/Core/CameraRigTransition.cpp
@@ -0,0 +1,11 @@
+// Copyright Epic Games, Inc. All Rights Reserved.
+
+#include "Core/CameraRigTransition.h"
+
+#include UE_INLINE_GENERATED_CPP_BY_NAME(CameraRigTransition)
+
+bool UCameraRigTransitionCondition::TransitionMatches(const FCameraRigTransitionConditionMatchParams& Params) const
+{
+	return OnTransitionMatches(Params);
+}
+
diff --git a/Engine/Plugins/Cameras/GameplayCameras/Source/GameplayCameras/Private/Core/CameraSystemEvaluator.cpp b/Engine/Plugins/Cameras/GameplayCameras/Source/GameplayCameras/Private/Core/CameraSystemEvaluator.cpp
--- a/Engine/Plugins/Cameras/GameplayCameras/Source/GameplayCameras/Private/Core/CameraSystemEvaluator.cpp
+++ b/Engine/Plugins/Cameras/GameplayCameras/Source/GameplayCameras/Private/Core/CameraSystemEvaluator.cpp
@@ -6,7 +6,7 @@
 #include "Core/CameraAsset.h"
 #include "Core/CameraDirectorEvaluator.h"
 #include "Core/CameraEvaluationContext.h"
-#include "Core/CameraMode.h"
+#include "Core/CameraRigAsset.h"
 #include "Core/DefaultRootCameraNode.h"
 #include "GameplayCameras.h"
 #include "IGameplayCamerasModule.h"
@@ -68,7 +68,7 @@ void UCameraSystemEvaluator::Update(const FCameraSystemEvaluationUpdateParams& P
 		return;
 	}
 
-	// Run the camera director, and activate any camera mode(s) it returns to us.
+	// Run the camera director, and activate any camera rig(s) it returns to us.
 	FCameraDirectorEvaluator* ActiveDirectorEvaluator = ActiveContextInfo.Evaluator;
 	if (ActiveDirectorEvaluator)
 	{
@@ -80,13 +80,13 @@ void UCameraSystemEvaluator::Update(const FCameraSystemEvaluationUpdateParams& P
 
 		ActiveDirectorEvaluator->Run(DirectorParams, DirectorResult);
 
-		if (DirectorResult.ActiveCameraModes.Num() == 1)
+		if (DirectorResult.ActiveCameraRigs.Num() == 1)
 		{
-			FActivateCameraModeParams CameraModeParams;
-			CameraModeParams.Evaluator = this;
-			CameraModeParams.EvaluationContext = ActiveContextInfo.EvaluationContext;
-			CameraModeParams.CameraMode = DirectorResult.ActiveCameraModes[0];
-			RootEvaluator->ActivateCameraMode(CameraModeParams);
+			FActivateCameraRigParams CameraRigParams;
+			CameraRigParams.Evaluator = this;
+			CameraRigParams.EvaluationContext = ActiveContextInfo.EvaluationContext;
+			CameraRigParams.CameraRig = DirectorResult.ActiveCameraRigs[0];
+			RootEvaluator->ActivateCameraRig(CameraRigParams);
 		}
 	}
 
diff --git a/Engine/Plugins/Cameras/GameplayCameras/Source/GameplayCameras/Private/Core/DefaultRootCameraNode.cpp b/Engine/Plugins/Cameras/GameplayCameras/Source/GameplayCameras/Private/Core/DefaultRootCameraNode.cpp
--- a/Engine/Plugins/Cameras/GameplayCameras/Source/GameplayCameras/Private/Core/DefaultRootCameraNode.cpp
+++ b/Engine/Plugins/Cameras/GameplayCameras/Source/GameplayCameras/Private/Core/DefaultRootCameraNode.cpp
@@ -3,7 +3,7 @@
 #include "Core/DefaultRootCameraNode.h"
 
 #include "Core/BlendStackCameraNode.h"
-#include "Core/CameraMode.h"
+#include "Core/CameraRigAsset.h"
 
 #include UE_INLINE_GENERATED_CPP_BY_NAME(DefaultRootCameraNode)
 
@@ -12,12 +12,12 @@ namespace UE::Cameras::Private
 
 TObjectPtr<UBlendStackCameraNode> CreateBlendStack(
 		UObject* This, const FObjectInitializer& ObjectInit,
-		const FName& Name, bool bAutoPop = true, bool bBlendFirstCameraMode = false)
+		const FName& Name, bool bAutoPop = true, bool bBlendFirstCameraRig = false)
 {
 	TObjectPtr<UBlendStackCameraNode> NewBlendStack = ObjectInit.CreateDefaultSubobject<UBlendStackCameraNode>(
 			This, Name);
 	NewBlendStack->bAutoPop = bAutoPop;
-	NewBlendStack->bBlendFirstCameraMode = bBlendFirstCameraMode;
+	NewBlendStack->bBlendFirstCameraRig = bBlendFirstCameraRig;
 	return NewBlendStack;
 }
 
@@ -64,21 +64,21 @@ void FDefaultRootCameraNodeEvaluator::OnRun(const FCameraNodeEvaluationParams& P
 	VisualLayer->Run(Params, OutResult);
 }
 
-void FDefaultRootCameraNodeEvaluator::OnActivateCameraMode(const FActivateCameraModeParams& Params)
+void FDefaultRootCameraNodeEvaluator::OnActivateCameraRig(const FActivateCameraRigParams& Params)
 {
 	FBlendStackCameraNodeEvaluator* TargetStack = nullptr;
 	switch (Params.Layer)
 	{
-		case ECameraModeLayer::Base:
+		case ECameraRigLayer::Base:
 			TargetStack = BaseLayer;
 			break;
-		case ECameraModeLayer::Main:
+		case ECameraRigLayer::Main:
 			TargetStack = MainLayer;
 			break;
-		case ECameraModeLayer::Global:
+		case ECameraRigLayer::Global:
 			TargetStack = GlobalLayer;
 			break;
-		case ECameraModeLayer::Visual:
+		case ECameraRigLayer::Visual:
 			TargetStack = VisualLayer;
 			break;
 	}
@@ -88,7 +88,7 @@ void FDefaultRootCameraNodeEvaluator::OnActivateCameraMode(const FActivateCamera
 		FBlendStackCameraPushParams PushParams;
 		PushParams.Evaluator = Params.Evaluator;
 		PushParams.EvaluationContext = Params.EvaluationContext;
-		PushParams.CameraMode = Params.CameraMode;
+		PushParams.CameraRig = Params.CameraRig;
 		TargetStack->Push(PushParams);
 	}
 }
diff --git a/Engine/Plugins/Cameras/GameplayCameras/Source/GameplayCameras/Private/Core/RootCameraNode.cpp b/Engine/Plugins/Cameras/GameplayCameras/Source/GameplayCameras/Private/Core/RootCameraNode.cpp
--- a/Engine/Plugins/Cameras/GameplayCameras/Source/GameplayCameras/Private/Core/RootCameraNode.cpp
+++ b/Engine/Plugins/Cameras/GameplayCameras/Source/GameplayCameras/Private/Core/RootCameraNode.cpp
@@ -2,12 +2,12 @@
 
 #include "Core/RootCameraNode.h"
 
-#include "Core/CameraMode.h"
+#include "Core/CameraRigAsset.h"
 
 #include UE_INLINE_GENERATED_CPP_BY_NAME(RootCameraNode)
 
-void FRootCameraNodeEvaluator::ActivateCameraMode(const FActivateCameraModeParams& Params)
+void FRootCameraNodeEvaluator::ActivateCameraRig(const FActivateCameraRigParams& Params)
 {
-	OnActivateCameraMode(Params);
+	OnActivateCameraRig(Params);
 }
 
diff --git a/Engine/Plugins/Cameras/GameplayCameras/Source/GameplayCameras/Private/Directors/SingleCameraDirector.cpp b/Engine/Plugins/Cameras/GameplayCameras/Source/GameplayCameras/Private/Directors/SingleCameraDirector.cpp
--- a/Engine/Plugins/Cameras/GameplayCameras/Source/GameplayCameras/Private/Directors/SingleCameraDirector.cpp
+++ b/Engine/Plugins/Cameras/GameplayCameras/Source/GameplayCameras/Private/Directors/SingleCameraDirector.cpp
@@ -11,9 +11,9 @@ protected:
 	virtual void OnRun(const FCameraDirectorEvaluationParams& Params, FCameraDirectorEvaluationResult& OutResult) override
 	{
 		const USingleCameraDirector* SingleDirector = GetCameraDirectorAs<USingleCameraDirector>();
-		if (SingleDirector->CameraMode)
+		if (SingleDirector->CameraRig)
 		{
-			OutResult.ActiveCameraModes.Add(SingleDirector->CameraMode);
+			OutResult.ActiveCameraRigs.Add(SingleDirector->CameraRig);
 		}
 	}
 };
diff --git a/Engine/Plugins/Cameras/GameplayCameras/Source/GameplayCameras/Public/Core/BlendStackCameraNode.h b/Engine/Plugins/Cameras/GameplayCameras/Source/GameplayCameras/Public/Core/BlendStackCameraNode.h
--- a/Engine/Plugins/Cameras/GameplayCameras/Source/GameplayCameras/Public/Core/BlendStackCameraNode.h
+++ b/Engine/Plugins/Cameras/GameplayCameras/Source/GameplayCameras/Public/Core/BlendStackCameraNode.h
@@ -2,10 +2,10 @@
 
 #pragma once
 
-#include "Core/CameraMode.h"
 #include "Core/CameraNode.h"
 #include "Core/CameraNodeEvaluator.h"
 #include "Core/CameraNodeEvaluatorStorage.h"
+#include "Core/CameraRigAsset.h"
 #include "IGameplayCamerasLiveEditListener.h"
 
 #include "BlendStackCameraNode.generated.h"
@@ -14,22 +14,22 @@ class FBlendStackRootCameraNodeEvaluator;
 class UBlendStackRootCameraNode;
 class UCameraAsset;
 class UCameraEvaluationContext;
-class UCameraMode;
-struct FCameraModeTransition;
+class UCameraRigAsset;
+struct FCameraRigTransition;
 
 /**
- * Parameter structure for pushing a camera mode onto a blend stack.
+ * Parameter structure for pushing a camera rig onto a blend stack.
  */
 struct FBlendStackCameraPushParams
 {
 	/** The evaluator currently running.*/
 	TObjectPtr<UCameraSystemEvaluator> Evaluator;
 
-	/** The evaluation context within which a camera mode's node tree should run. */
+	/** The evaluation context within which a camera rig's node tree should run. */
 	TObjectPtr<const UCameraEvaluationContext> EvaluationContext;
 
-	/** The source camera mode asset to instantiate and push on the blend stack. */
-	TObjectPtr<const UCameraMode> CameraMode;
+	/** The source camera rig asset to instantiate and push on the blend stack. */
+	TObjectPtr<const UCameraRigAsset> CameraRig;
 };
 
 /**
@@ -48,17 +48,17 @@ protected:
 public:
 
 	/** 
-	 * Whether to automatically pop camera modes out of the stack when another mode
+	 * Whether to automatically pop camera rigs out of the stack when another rig
 	 * has reached 100% blend above them.
 	 */
 	UPROPERTY()
 	bool bAutoPop = true;
 
 	/**
-	 * Whether to blend-in the first camera mode when the stack is previously empty.
+	 * Whether to blend-in the first camera rig when the stack is previously empty.
 	 */
 	UPROPERTY()
-	bool bBlendFirstCameraMode = false;
+	bool bBlendFirstCameraRig = false;
 };
 
 /**
@@ -74,7 +74,7 @@ class FBlendStackCameraNodeEvaluator
 
 public:
 
-	/** Push a new camera mode onto the blend stack. */
+	/** Push a new camera rig onto the blend stack. */
 	void Push(const FBlendStackCameraPushParams& Params);
 
 protected:
@@ -92,20 +92,20 @@ protected:
 protected:
 
 	// Utility functions for finding an appropriate transition.
-	const FCameraModeTransition* FindTransition(const FBlendStackCameraPushParams& Params) const;
-	const FCameraModeTransition* FindTransition(
-			TArrayView<const FCameraModeTransition> Transitions, 
-			const UCameraMode* FromCameraMode, const UCameraAsset* FromCameraAsset, bool bFromFrozen,
-			const UCameraMode* ToCameraMode, const UCameraAsset* ToCameraAsset) const;
+	const FCameraRigTransition* FindTransition(const FBlendStackCameraPushParams& Params) const;
+	const FCameraRigTransition* FindTransition(
+			TArrayView<const FCameraRigTransition> Transitions, 
+			const UCameraRigAsset* FromCameraRig, const UCameraAsset* FromCameraAsset, bool bFromFrozen,
+			const UCameraRigAsset* ToCameraRig, const UCameraAsset* ToCameraAsset) const;
 
 protected:
 
-	struct FCameraModeEntry
+	struct FCameraRigEntry
 	{
 		/** Evaluation context in which this entry runs. */
 		TWeakObjectPtr<const UCameraEvaluationContext> EvaluationContext;
-		/** The camera mode asset that this entry runs. */
-		TObjectPtr<const UCameraMode> CameraMode;
+		/** The camera rig asset that this entry runs. */
+		TObjectPtr<const UCameraRigAsset> CameraRig;
 		/** The root node. */
 		TObjectPtr<UBlendStackRootCameraNode> RootNode;
 		/** Storage buffer for all evaluators in this node tree. */
@@ -119,7 +119,7 @@ protected:
 		/** Whether this entry is frozen. */
 		bool bIsFrozen = false;
 #if WITH_EDITOR
-		FCameraModePackages ListenedPackages;
+		FCameraRigPackages ListenedPackages;
 #endif  // WITH_EDITOR
 	};
 
@@ -127,6 +127,6 @@ protected:
 	TObjectPtr<UCameraSystemEvaluator> OwningEvaluator;
 
 	/** Entries in the blend stack. */
-	TArray<FCameraModeEntry> Entries;
+	TArray<FCameraRigEntry> Entries;
 };
 
diff --git a/Engine/Plugins/Cameras/GameplayCameras/Source/GameplayCameras/Public/Core/CameraAsset.h b/Engine/Plugins/Cameras/GameplayCameras/Source/GameplayCameras/Public/Core/CameraAsset.h
--- a/Engine/Plugins/Cameras/GameplayCameras/Source/GameplayCameras/Public/Core/CameraAsset.h
+++ b/Engine/Plugins/Cameras/GameplayCameras/Source/GameplayCameras/Public/Core/CameraAsset.h
@@ -2,7 +2,7 @@
 
 #pragma once
 
-#include "Core/CameraModeTransition.h"
+#include "Core/CameraRigTransition.h"
 #include "CoreTypes.h"
 #include "UObject/ObjectPtr.h"
 
@@ -24,13 +24,13 @@ public:
 	UPROPERTY(Instanced, EditAnywhere, Category=Director)
 	TObjectPtr<UCameraDirector> CameraDirector;
 
-	/** A list of default enter transitions for all the camera modes in this asset. */
+	/** A list of default enter transitions for all the camera rigs in this asset. */
 	UPROPERTY(EditAnywhere, Category=Blending)
-	TArray<FCameraModeTransition> EnterTransitions;
+	TArray<FCameraRigTransition> EnterTransitions;
 
-	/** A list of default exit transitions for all the camera modes in this asset. */
+	/** A list of default exit transitions for all the camera rigs in this asset. */
 	UPROPERTY(EditAnywhere, Category=Blending)
-	TArray<FCameraModeTransition> ExitTransitions;
+	TArray<FCameraRigTransition> ExitTransitions;
 };
 
 
diff --git a/Engine/Plugins/Cameras/GameplayCameras/Source/GameplayCameras/Public/Core/CameraDirectorEvaluator.h b/Engine/Plugins/Cameras/GameplayCameras/Source/GameplayCameras/Public/Core/CameraDirectorEvaluator.h
--- a/Engine/Plugins/Cameras/GameplayCameras/Source/GameplayCameras/Public/Core/CameraDirectorEvaluator.h
+++ b/Engine/Plugins/Cameras/GameplayCameras/Source/GameplayCameras/Public/Core/CameraDirectorEvaluator.h
@@ -11,7 +11,7 @@ class FCameraDirectorEvaluator;
 class FCameraDirectorEvaluatorStorage;
 class UCameraDirector;
 class UCameraEvaluationContext;
-class UCameraMode;
+class UCameraRigAsset;
 
 /**
  * Parameter structure for running a camera director.
@@ -30,8 +30,8 @@ struct FCameraDirectorEvaluationParams
  */
 struct FCameraDirectorEvaluationResult
 {
-	/** The camera mode(s) that the director says should be active this frame. */
-	TArray<TObjectPtr<const UCameraMode>, TInlineAllocator<2>> ActiveCameraModes;
+	/** The camera rig(s) that the director says should be active this frame. */
+	TArray<TObjectPtr<const UCameraRigAsset>, TInlineAllocator<2>> ActiveCameraRigs;
 };
 
 /**
@@ -84,7 +84,7 @@ public:
 	FCameraDirectorEvaluator();
 	virtual ~FCameraDirectorEvaluator() {}
 	
-	/** Runs the camera director to determine what camera mode(s) should be active this frame. */
+	/** Runs the camera director to determine what camera rig(s) should be active this frame. */
 	void Run(const FCameraDirectorEvaluationParams& Params, FCameraDirectorEvaluationResult& OutResult);
 
 	/** Gets the camera director. */
@@ -101,7 +101,7 @@ public:
 
 protected:
 
-	/** Runs the camera director to determine what camera mode(s) should be active this frame. */
+	/** Runs the camera director to determine what camera rig(s) should be active this frame. */
 	virtual void OnRun(const FCameraDirectorEvaluationParams& Params, FCameraDirectorEvaluationResult& OutResult) {}
 
 private:
diff --git a/Engine/Plugins/Cameras/GameplayCameras/Source/GameplayCameras/Public/Core/CameraEvaluationContext.h b/Engine/Plugins/Cameras/GameplayCameras/Source/GameplayCameras/Public/Core/CameraEvaluationContext.h
--- a/Engine/Plugins/Cameras/GameplayCameras/Source/GameplayCameras/Public/Core/CameraEvaluationContext.h
+++ b/Engine/Plugins/Cameras/GameplayCameras/Source/GameplayCameras/Public/Core/CameraEvaluationContext.h
@@ -12,7 +12,7 @@ class UCameraAsset;
 class APlayerController;
 
 /**
- * Base class for providing a context to running camera modes.
+ * Base class for providing a context to running camera rigs.
  */
 UCLASS(MinimalAPI)
 class UCameraEvaluationContext : public UObject
@@ -32,7 +32,7 @@ public:
 	/** Gets the camera asset that is hosted in this context. */
 	UCameraAsset* GetCameraAsset() const { return CameraAsset; }
 
-	/** Gets the initial evaluation result for all camera modes in this context. */
+	/** Gets the initial evaluation result for all camera rigs in this context. */
 	const FCameraNodeEvaluationResult& GetInitialResult() const { return InitialResult; }
 
 protected:
@@ -48,7 +48,7 @@ protected:
 	UPROPERTY()
 	TObjectPtr<UCameraAsset> CameraAsset;
 
-	/** The initial result for all camera modes in this context. */
+	/** The initial result for all camera rigs in this context. */
 	FCameraNodeEvaluationResult InitialResult;
 };
 
diff --git a/Engine/Plugins/Cameras/GameplayCameras/Source/GameplayCameras/Public/Core/CameraMode.h b/Engine/Plugins/Cameras/GameplayCameras/Source/GameplayCameras/Public/Core/CameraMode.h
deleted file mode 100644
diff --git a/Engine/Plugins/Cameras/GameplayCameras/Source/GameplayCameras/Public/Core/CameraModeTransition.h b/Engine/Plugins/Cameras/GameplayCameras/Source/GameplayCameras/Public/Core/CameraModeTransition.h
deleted file mode 100644
diff --git a/Engine/Plugins/Cameras/GameplayCameras/Source/GameplayCameras/Public/Core/CameraNodeEvaluatorStorage.h b/Engine/Plugins/Cameras/GameplayCameras/Source/GameplayCameras/Public/Core/CameraNodeEvaluatorStorage.h
--- a/Engine/Plugins/Cameras/GameplayCameras/Source/GameplayCameras/Public/Core/CameraNodeEvaluatorStorage.h
+++ b/Engine/Plugins/Cameras/GameplayCameras/Source/GameplayCameras/Public/Core/CameraNodeEvaluatorStorage.h
@@ -12,7 +12,7 @@ class FCameraNodeEvaluator;
 class FCameraNodeEvaluatorStorage;
 class FReferenceCollector;
 class UCameraEvaluationContext;
-class UCameraMode;
+class UCameraRigAsset;
 class UCameraNode;
 class UCameraSystemEvaluator;
 
@@ -54,7 +54,7 @@ public:
 public:
 
 	/** Compute allocation information for the given tree of camera nodes. */
-	static FCameraNodeEvaluatorTreeAllocationInfo ComputeTreeInfo(const UCameraMode* CameraMode);
+	static FCameraNodeEvaluatorTreeAllocationInfo ComputeTreeInfo(const UCameraRigAsset* CameraRig);
 
 	/** Build the tree of evaluators for the given tree of camera nodes. */
 	FCameraNodeEvaluatorPtr BuildEvaluatorTree(const FCameraNodeEvaluatorTreeBuilderParams& Params);
diff --git a/Engine/Plugins/Cameras/GameplayCameras/Source/GameplayCameras/Public/Core/CameraRigAsset.h b/Engine/Plugins/Cameras/GameplayCameras/Source/GameplayCameras/Public/Core/CameraRigAsset.h
new file mode 100644
--- /dev/null
+++ b/Engine/Plugins/Cameras/GameplayCameras/Source/GameplayCameras/Public/Core/CameraRigAsset.h
@@ -0,0 +1,63 @@
+// Copyright Epic Games, Inc. All Rights Reserved.
+
+#pragma once
+
+#include "Core/CameraRigTransition.h"
+#include "CoreTypes.h"
+#include "UObject/ObjectPtr.h"
+
+#include "CameraRigAsset.generated.h"
+
+class UCameraNode;
+
+/**
+ *
+ */
+UENUM()
+enum class ECameraRigBuildStatus : uint8
+{
+	Clean,
+	CleanWithWarnings,
+	WithErrors,
+	Dirty
+};
+
+/**
+ * List of packages that contain the definition of a camera rig.
+ * In most cases there's only one, but with nested assets there could be more.
+ */
+using FCameraRigPackages = TArray<const UPackage*, TInlineAllocator<4>>;
+
+/**
+ * A camera rig asset, which runs a hierarchy of camera nodes to drive 
+ * the behavior of a camera.
+ */
+UCLASS(MinimalAPI)
+class UCameraRigAsset : public UObject
+{
+	GENERATED_BODY()
+
+public:
+
+#if WITH_EDITOR
+	GAMEPLAYCAMERAS_API void GatherPackages(FCameraRigPackages& OutPackages) const;
+#endif  // WITH_EDITOR
+
+public:
+
+	/** Root camera node. */
+	UPROPERTY(EditAnywhere, Instanced, Category=Common)
+	TObjectPtr<UCameraNode> RootNode;
+
+	/** List of enter transitions for this camera rig. */
+	UPROPERTY(EditAnywhere, Category=Blending)
+	TArray<FCameraRigTransition> EnterTransitions;
+
+	/** List of exist transitions for this camera rig. */
+	UPROPERTY(EditAnywhere, Category=Blending)
+	TArray<FCameraRigTransition> ExitTransitions;
+
+	UPROPERTY(Transient)
+	ECameraRigBuildStatus BuildStatus = ECameraRigBuildStatus::Dirty;
+};
+
diff --git a/Engine/Plugins/Cameras/GameplayCameras/Source/GameplayCameras/Public/Core/CameraRigTransition.h b/Engine/Plugins/Cameras/GameplayCameras/Source/GameplayCameras/Public/Core/CameraRigTransition.h
new file mode 100644
--- /dev/null
+++ b/Engine/Plugins/Cameras/GameplayCameras/Source/GameplayCameras/Public/Core/CameraRigTransition.h
@@ -0,0 +1,66 @@
+// Copyright Epic Games, Inc. All Rights Reserved.
+
+#pragma once
+
+#include "Containers/Array.h"
+#include "CoreTypes.h"
+#include "UObject/ObjectPtr.h"
+
+#include "CameraRigTransition.generated.h"
+
+class UBlendCameraNode;
+class UCameraAsset;
+class UCameraRigAsset;
+
+/**
+ * Parameter structure for camera transitions.
+ */
+struct FCameraRigTransitionConditionMatchParams
+{
+	/** The previous camera rig. */
+	const UCameraRigAsset* FromCameraRig = nullptr;
+	/** The previous camera asset. */
+	const UCameraAsset* FromCameraAsset = nullptr;
+
+	/** The next camera rig. */
+	const UCameraRigAsset* ToCameraRig = nullptr;
+	/** The next camera asset. */
+	const UCameraAsset* ToCameraAsset = nullptr;
+};
+
+/**
+ * Base class for a camera transition condition.
+ */
+UCLASS(Abstract, DefaultToInstanced, MinimalAPI)
+class UCameraRigTransitionCondition : public UObject
+{
+	GENERATED_BODY()
+
+public:
+
+	/** Evaluates whether this transition should be used. */
+	bool TransitionMatches(const FCameraRigTransitionConditionMatchParams& Params) const;
+
+protected:
+
+	/** Evaluates whether this transition should be used. */
+	virtual bool OnTransitionMatches(const FCameraRigTransitionConditionMatchParams& Params) const { return false; }
+};
+
+/**
+ * A camera transition.
+ */
+USTRUCT()
+struct FCameraRigTransition
+{
+	GENERATED_BODY()
+
+	/** The list of conditions that must pass for this transition to be used. */
+	UPROPERTY(EditAnywhere, Instanced, Category=Common)
+	TArray<TObjectPtr<UCameraRigTransitionCondition>> Conditions;
+
+	/** The blend to use to blend a given camera rig in or out. */
+	UPROPERTY(EditAnywhere, Instanced, Category=Common)
+	TObjectPtr<UBlendCameraNode> Blend;
+};
+
diff --git a/Engine/Plugins/Cameras/GameplayCameras/Source/GameplayCameras/Public/Core/CameraSystemEvaluator.h b/Engine/Plugins/Cameras/GameplayCameras/Source/GameplayCameras/Public/Core/CameraSystemEvaluator.h
--- a/Engine/Plugins/Cameras/GameplayCameras/Source/GameplayCameras/Public/Core/CameraSystemEvaluator.h
+++ b/Engine/Plugins/Cameras/GameplayCameras/Source/GameplayCameras/Public/Core/CameraSystemEvaluator.h
@@ -16,7 +16,7 @@
 class FRootCameraNodeEvaluator;
 class UCameraDirector;
 class UCameraEvaluationContext;
-class UCameraMode;
+class UCameraRigAsset;
 class URootCameraNode;
 struct FMinimalViewInfo;
 
diff --git a/Engine/Plugins/Cameras/GameplayCameras/Source/GameplayCameras/Public/Core/DefaultRootCameraNode.h b/Engine/Plugins/Cameras/GameplayCameras/Source/GameplayCameras/Public/Core/DefaultRootCameraNode.h
--- a/Engine/Plugins/Cameras/GameplayCameras/Source/GameplayCameras/Public/Core/DefaultRootCameraNode.h
+++ b/Engine/Plugins/Cameras/GameplayCameras/Source/GameplayCameras/Public/Core/DefaultRootCameraNode.h
@@ -53,7 +53,7 @@ protected:
 	virtual FCameraNodeEvaluatorChildrenView OnGetChildren() override;
 	virtual void OnInitialize(const FCameraNodeEvaluatorInitializeParams& Params) override;
 	virtual void OnRun(const FCameraNodeEvaluationParams& Params, FCameraNodeEvaluationResult& OutResult) override;
-	virtual void OnActivateCameraMode(const FActivateCameraModeParams& Params) override;
+	virtual void OnActivateCameraRig(const FActivateCameraRigParams& Params) override;
 
 private:
 
diff --git a/Engine/Plugins/Cameras/GameplayCameras/Source/GameplayCameras/Public/Core/RootCameraNode.h b/Engine/Plugins/Cameras/GameplayCameras/Source/GameplayCameras/Public/Core/RootCameraNode.h
--- a/Engine/Plugins/Cameras/GameplayCameras/Source/GameplayCameras/Public/Core/RootCameraNode.h
+++ b/Engine/Plugins/Cameras/GameplayCameras/Source/GameplayCameras/Public/Core/RootCameraNode.h
@@ -8,14 +8,14 @@
 #include "RootCameraNode.generated.h"
 
 class UCameraEvaluationContext;
-class UCameraMode;
+class UCameraRigAsset;
 class UCameraSystemEvaluator;
 
 /**
- * Defines evaluation layers for camera modes.
+ * Defines evaluation layers for camera rigs.
  */
 UENUM()
-enum class ECameraModeLayer
+enum class ECameraRigLayer
 {
 	Base,
 	Main,
@@ -25,24 +25,24 @@ enum class ECameraModeLayer
 	User1,
 	User2
 };
-ENUM_CLASS_FLAGS(ECameraModeLayer)
+ENUM_CLASS_FLAGS(ECameraRigLayer)
 
 /**
- * Parameter structure for activating a new camera mode.
+ * Parameter structure for activating a new camera rig.
  */
-struct FActivateCameraModeParams
+struct FActivateCameraRigParams
 {
 	/** The evaluator currently running.*/
 	TObjectPtr<UCameraSystemEvaluator> Evaluator;
 
-	/** The evaluation context in which the camera mode runs. */
+	/** The evaluation context in which the camera rig runs. */
 	TObjectPtr<const UCameraEvaluationContext> EvaluationContext;
 
-	/** The source camera mode asset that will be instantiated. */
-	TObjectPtr<const UCameraMode> CameraMode;
+	/** The source camera rig asset that will be instantiated. */
+	TObjectPtr<const UCameraRigAsset> CameraRig;
 
-	/** The evaluation layer on which to instantiate the camera mode. */
-	ECameraModeLayer Layer = ECameraModeLayer::Main;
+	/** The evaluation layer on which to instantiate the camera rig. */
+	ECameraRigLayer Layer = ECameraRigLayer::Main;
 };
 
 /**
@@ -62,12 +62,12 @@ class FRootCameraNodeEvaluator : public FCameraNodeEvaluator
 {
 public:
 
-	/** Activates a camera mode. */
-	void ActivateCameraMode(const FActivateCameraModeParams& Params);
+	/** Activates a camera rig. */
+	void ActivateCameraRig(const FActivateCameraRigParams& Params);
 
 private:
 
-	/** Activates a camera mode. */
-	virtual void OnActivateCameraMode(const FActivateCameraModeParams& Params) {}
+	/** Activates a camera rig. */
+	virtual void OnActivateCameraRig(const FActivateCameraRigParams& Params) {}
 };
 
diff --git a/Engine/Plugins/Cameras/GameplayCameras/Source/GameplayCameras/Public/Directors/SingleCameraDirector.h b/Engine/Plugins/Cameras/GameplayCameras/Source/GameplayCameras/Public/Directors/SingleCameraDirector.h
--- a/Engine/Plugins/Cameras/GameplayCameras/Source/GameplayCameras/Public/Directors/SingleCameraDirector.h
+++ b/Engine/Plugins/Cameras/GameplayCameras/Source/GameplayCameras/Public/Directors/SingleCameraDirector.h
@@ -7,7 +7,7 @@
 #include "SingleCameraDirector.generated.h"
 
 /**
- * A simple camera director that only ever returns one single camera mode.
+ * A simple camera director that only ever returns one single camera rig.
  */
 UCLASS(EditInlineNew)
 class USingleCameraDirector : public UCameraDirector
@@ -25,8 +25,8 @@ protected:
 
 public:
 
-	/** The camera mode to run every frame. */
+	/** The camera rig to run every frame. */
 	UPROPERTY(EditAnywhere, Category=Common)
-	TObjectPtr<UCameraMode> CameraMode;
+	TObjectPtr<UCameraRigAsset> CameraRig;
 };
 
diff --git a/Engine/Plugins/Cameras/GameplayCameras/Source/GameplayCameras/Public/GameFramework/GameplayCameraSystemComponent.h b/Engine/Plugins/Cameras/GameplayCameras/Source/GameplayCameras/Public/GameFramework/GameplayCameraSystemComponent.h
--- a/Engine/Plugins/Cameras/GameplayCameras/Source/GameplayCameras/Public/GameFramework/GameplayCameraSystemComponent.h
+++ b/Engine/Plugins/Cameras/GameplayCameras/Source/GameplayCameras/Public/GameFramework/GameplayCameraSystemComponent.h
@@ -11,7 +11,7 @@
 #include "GameplayCameraSystemComponent.generated.h"
 
 class UCameraSystemEvaluator;
-class UCameraMode;
+class UCameraRigAsset;
 struct FMinimalViewInfo;
 
 /**
diff --git a/Engine/Plugins/Cameras/GameplayCameras/Source/GameplayCameras/Public/Nodes/Blends/SmoothBlendCameraNode.h b/Engine/Plugins/Cameras/GameplayCameras/Source/GameplayCameras/Public/Nodes/Blends/SmoothBlendCameraNode.h
--- a/Engine/Plugins/Cameras/GameplayCameras/Source/GameplayCameras/Public/Nodes/Blends/SmoothBlendCameraNode.h
+++ b/Engine/Plugins/Cameras/GameplayCameras/Source/GameplayCameras/Public/Nodes/Blends/SmoothBlendCameraNode.h
@@ -17,7 +17,7 @@ enum class ESmoothCameraBlendType
 };
 
 /**
- * A blend camera mode that implements the smoothstep and smoothersteps algorithms.
+ * A blend camera rig that implements the smoothstep and smoothersteps algorithms.
  */
 UCLASS(MinimalAPI)
 class USmoothBlendCameraNode : public USimpleFixedTimeBlendCameraNode
diff --git a/Engine/Plugins/Cameras/GameplayCameras/Source/GameplayCamerasEditor/Private/AssetTools/AssetDefinition_CameraMode.cpp b/Engine/Plugins/Cameras/GameplayCameras/Source/GameplayCamerasEditor/Private/AssetTools/AssetDefinition_CameraMode.cpp
deleted file mode 100644
diff --git a/Engine/Plugins/Cameras/GameplayCameras/Source/GameplayCamerasEditor/Private/AssetTools/AssetDefinition_CameraMode.h b/Engine/Plugins/Cameras/GameplayCameras/Source/GameplayCamerasEditor/Private/AssetTools/AssetDefinition_CameraMode.h
deleted file mode 100644
diff --git a/Engine/Plugins/Cameras/GameplayCameras/Source/GameplayCamerasEditor/Private/AssetTools/AssetDefinition_CameraRigAsset.cpp b/Engine/Plugins/Cameras/GameplayCameras/Source/GameplayCamerasEditor/Private/AssetTools/AssetDefinition_CameraRigAsset.cpp
new file mode 100644
--- /dev/null
+++ b/Engine/Plugins/Cameras/GameplayCameras/Source/GameplayCamerasEditor/Private/AssetTools/AssetDefinition_CameraRigAsset.cpp
@@ -0,0 +1,55 @@
+// Copyright Epic Games, Inc. All Rights Reserved.
+
+#include "AssetTools/AssetDefinition_CameraRigAsset.h"
+
+#include "IGameplayCamerasEditorModule.h"
+#include "Styles/GameplayCamerasEditorStyle.h"
+#include "Toolkits/CameraRigAssetEditorToolkit.h"
+#include "Toolkits/IToolkit.h"
+
+#define LOCTEXT_NAMESPACE "AssetDefinition_CameraRigAsset"
+
+#include UE_INLINE_GENERATED_CPP_BY_NAME(AssetDefinition_CameraRigAsset)
+
+FText UAssetDefinition_CameraRigAsset::GetAssetDisplayName() const
+{
+	return LOCTEXT("AssetDisplayName", "Camera Rig");
+}
+
+FLinearColor UAssetDefinition_CameraRigAsset::GetAssetColor() const
+{
+	return FLinearColor(FColor(200, 80, 80));
+}
+
+TSoftClassPtr<UObject> UAssetDefinition_CameraRigAsset::GetAssetClass() const
+{
+	return UCameraRigAsset::StaticClass();
+}
+
+TConstArrayView<FAssetCategoryPath> UAssetDefinition_CameraRigAsset::GetAssetCategories() const
+{
+	static const auto Categories = { FAssetCategoryPath(EAssetCategoryPaths::Gameplay, LOCTEXT("AssetSubCategory", "Cameras")) };
+	return Categories;
+}
+
+FAssetOpenSupport UAssetDefinition_CameraRigAsset::GetAssetOpenSupport(const FAssetOpenSupportArgs& OpenSupportArgs) const
+{
+	return FAssetOpenSupport(
+			OpenSupportArgs.OpenMethod,
+			OpenSupportArgs.OpenMethod == EAssetOpenMethod::Edit,
+			EToolkitMode::Standalone); 
+}
+
+EAssetCommandResult UAssetDefinition_CameraRigAsset::OpenAssets(const FAssetOpenArgs& OpenArgs) const
+{
+	IGameplayCamerasEditorModule& GameplayCamerasEditorModule = FModuleManager::LoadModuleChecked<IGameplayCamerasEditorModule>("GameplayCamerasEditor");
+	for (UCameraRigAsset* CameraRig : OpenArgs.LoadObjects<UCameraRigAsset>())
+	{
+		GameplayCamerasEditorModule.CreateCameraRigEditor(OpenArgs.GetToolkitMode(), OpenArgs.ToolkitHost, CameraRig);
+	}	
+
+	return EAssetCommandResult::Handled;
+}
+
+#undef LOCTEXT_NAMESPACE
+
diff --git a/Engine/Plugins/Cameras/GameplayCameras/Source/GameplayCamerasEditor/Private/AssetTools/AssetDefinition_CameraRigAsset.h b/Engine/Plugins/Cameras/GameplayCameras/Source/GameplayCamerasEditor/Private/AssetTools/AssetDefinition_CameraRigAsset.h
new file mode 100644
--- /dev/null
+++ b/Engine/Plugins/Cameras/GameplayCameras/Source/GameplayCamerasEditor/Private/AssetTools/AssetDefinition_CameraRigAsset.h
@@ -0,0 +1,25 @@
+// Copyright Epic Games, Inc. All Rights Reserved.
+
+#pragma once
+
+#include "AssetDefinitionDefault.h"
+#include "Core/CameraRigAsset.h"
+
+#include "AssetDefinition_CameraRigAsset.generated.h"
+
+UCLASS()
+class UAssetDefinition_CameraRigAsset : public UAssetDefinitionDefault
+{
+	GENERATED_BODY()
+
+public:
+
+	// UAssetDefinition interface
+	virtual FText GetAssetDisplayName() const override;
+	virtual FLinearColor GetAssetColor() const override;
+	virtual TSoftClassPtr<UObject> GetAssetClass() const override;
+	virtual TConstArrayView<FAssetCategoryPath> GetAssetCategories() const override;
+	virtual FAssetOpenSupport GetAssetOpenSupport(const FAssetOpenSupportArgs& OpenSupportArgs) const override;
+	virtual EAssetCommandResult OpenAssets(const FAssetOpenArgs& OpenArgs) const override;
+};
+
diff --git a/Engine/Plugins/Cameras/GameplayCameras/Source/GameplayCamerasEditor/Private/AssetTools/CameraModeAssetEditor.cpp b/Engine/Plugins/Cameras/GameplayCameras/Source/GameplayCamerasEditor/Private/AssetTools/CameraModeAssetEditor.cpp
deleted file mode 100644
diff --git a/Engine/Plugins/Cameras/GameplayCameras/Source/GameplayCamerasEditor/Private/AssetTools/CameraRigAssetEditor.cpp b/Engine/Plugins/Cameras/GameplayCameras/Source/GameplayCamerasEditor/Private/AssetTools/CameraRigAssetEditor.cpp
new file mode 100644
--- /dev/null
+++ b/Engine/Plugins/Cameras/GameplayCameras/Source/GameplayCamerasEditor/Private/AssetTools/CameraRigAssetEditor.cpp
@@ -0,0 +1,27 @@
+// Copyright Epic Games, Inc. All Rights Reserved.
+
+#include "AssetTools/CameraRigAssetEditor.h"
+
+#include "Core/CameraRigAsset.h"
+#include "EditorModeManager.h"
+#include "Toolkits/CameraRigAssetEditorToolkit.h"
+
+#include UE_INLINE_GENERATED_CPP_BY_NAME(CameraRigAssetEditor)
+
+void UCameraRigAssetEditor::Initialize(TObjectPtr<UCameraRigAsset> InCameraRigAsset)
+{
+	CameraRigAsset = InCameraRigAsset;
+
+	Super::Initialize();
+}
+
+void UCameraRigAssetEditor::GetObjectsToEdit(TArray<UObject*>& InObjectsToEdit)
+{
+	InObjectsToEdit.Add(CameraRigAsset.Get());
+}
+
+TSharedPtr<FBaseAssetToolkit> UCameraRigAssetEditor::CreateToolkit()
+{
+	return MakeShared<FCameraRigAssetEditorToolkit>(this);
+}
+
diff --git a/Engine/Plugins/Cameras/GameplayCameras/Source/GameplayCamerasEditor/Private/Commands/CameraModeAssetEditorCommands.cpp b/Engine/Plugins/Cameras/GameplayCameras/Source/GameplayCamerasEditor/Private/Commands/CameraModeAssetEditorCommands.cpp
deleted file mode 100644
diff --git a/Engine/Plugins/Cameras/GameplayCameras/Source/GameplayCamerasEditor/Private/Commands/CameraRigAssetEditorCommands.cpp b/Engine/Plugins/Cameras/GameplayCameras/Source/GameplayCamerasEditor/Private/Commands/CameraRigAssetEditorCommands.cpp
new file mode 100644
--- /dev/null
+++ b/Engine/Plugins/Cameras/GameplayCameras/Source/GameplayCamerasEditor/Private/Commands/CameraRigAssetEditorCommands.cpp
@@ -0,0 +1,29 @@
+// Copyright Epic Games, Inc. All Rights Reserved.
+
+#include "Commands/CameraRigAssetEditorCommands.h"
+
+#include "Framework/Commands/Commands.h"
+#include "Framework/Commands/InputChord.h"
+#include "Framework/Commands/UICommandInfo.h"
+#include "Styling/AppStyle.h"
+
+#define LOCTEXT_NAMESPACE "CameraRigAssetEditorCommands"
+
+FCameraRigAssetEditorCommands::FCameraRigAssetEditorCommands()
+	: TCommands<FCameraRigAssetEditorCommands>(
+			"GameplayCameras",
+			NSLOCTEXT("Contexts", "GameplayCameras", "Gameplay Cameras"),
+			NAME_None,
+			FAppStyle::GetAppStyleSetName()
+		)
+{
+}
+
+void FCameraRigAssetEditorCommands::RegisterCommands()
+{
+	UI_COMMAND(Build, "Build", "Builds the asset and refreshes it in PIE",
+			EUserInterfaceActionType::Button, FInputChord(EKeys::F7));
+}
+ 
+#undef LOCTEXT_NAMESPACE
+
diff --git a/Engine/Plugins/Cameras/GameplayCameras/Source/GameplayCamerasEditor/Private/Factories/CameraModeFactory.cpp b/Engine/Plugins/Cameras/GameplayCameras/Source/GameplayCamerasEditor/Private/Factories/CameraModeFactory.cpp
deleted file mode 100644
diff --git a/Engine/Plugins/Cameras/GameplayCameras/Source/GameplayCamerasEditor/Private/Factories/CameraModeFactory.h b/Engine/Plugins/Cameras/GameplayCameras/Source/GameplayCamerasEditor/Private/Factories/CameraModeFactory.h
deleted file mode 100644
diff --git a/Engine/Plugins/Cameras/GameplayCameras/Source/GameplayCamerasEditor/Private/Factories/CameraRigAssetFactory.cpp b/Engine/Plugins/Cameras/GameplayCameras/Source/GameplayCamerasEditor/Private/Factories/CameraRigAssetFactory.cpp
new file mode 100644
--- /dev/null
+++ b/Engine/Plugins/Cameras/GameplayCameras/Source/GameplayCamerasEditor/Private/Factories/CameraRigAssetFactory.cpp
@@ -0,0 +1,37 @@
+// Copyright Epic Games, Inc. All Rights Reserved.
+
+#include "Factories/CameraRigAssetFactory.h"
+
+#include "Core/CameraRigAsset.h"
+
+#include UE_INLINE_GENERATED_CPP_BY_NAME(CameraRigAssetFactory)
+
+#define LOCTEXT_NAMESPACE "CameraRigAssetFactory"
+
+UCameraRigAssetFactory::UCameraRigAssetFactory(const FObjectInitializer& ObjectInitializer)
+	: Super(ObjectInitializer)
+{
+	bCreateNew = true;
+	bEditAfterNew = true;
+	SupportedClass = UCameraRigAsset::StaticClass();
+}
+
+UObject* UCameraRigAssetFactory::FactoryCreateNew(UClass* Class, UObject* Parent, FName Name, EObjectFlags Flags, UObject* Context, FFeedbackContext* Warn)
+{
+	UCameraRigAsset* NewCameraRig = NewObject<UCameraRigAsset>(Parent, Class, Name, Flags | RF_Transactional);
+	return NewCameraRig;
+}
+
+bool UCameraRigAssetFactory::ConfigureProperties()
+{
+	return true;
+}
+
+bool UCameraRigAssetFactory::ShouldShowInNewMenu() const
+{
+	return true;
+}
+
+#undef LOCTEXT_NAMESPACE
+
+
diff --git a/Engine/Plugins/Cameras/GameplayCameras/Source/GameplayCamerasEditor/Private/Factories/CameraRigAssetFactory.h b/Engine/Plugins/Cameras/GameplayCameras/Source/GameplayCamerasEditor/Private/Factories/CameraRigAssetFactory.h
new file mode 100644
--- /dev/null
+++ b/Engine/Plugins/Cameras/GameplayCameras/Source/GameplayCamerasEditor/Private/Factories/CameraRigAssetFactory.h
@@ -0,0 +1,26 @@
+// Copyright Epic Games, Inc. All Rights Reserved.
+
+#pragma once
+
+#include "Factories/Factory.h"
+
+#include "CameraRigAssetFactory.generated.h"
+
+class UCameraRigAsset;
+
+/**
+ * Implements a factory for UCameraRigAsset objects.
+ */
+UCLASS(hidecategories=Object)
+class UCameraRigAssetFactory : public UFactory
+{
+	GENERATED_BODY()
+
+	UCameraRigAssetFactory(const FObjectInitializer& ObjectInitializer);
+
+	// UFactory Interface
+	virtual UObject* FactoryCreateNew(UClass* Class, UObject* Parent, FName Name, EObjectFlags Flags, UObject* Context, FFeedbackContext* Warn) override;
+	virtual bool ConfigureProperties() override;
+	virtual bool ShouldShowInNewMenu() const override;
+};
+
diff --git a/Engine/Plugins/Cameras/GameplayCameras/Source/GameplayCamerasEditor/Private/GameplayCamerasEditorModule.cpp b/Engine/Plugins/Cameras/GameplayCameras/Source/GameplayCamerasEditor/Private/GameplayCamerasEditorModule.cpp
--- a/Engine/Plugins/Cameras/GameplayCameras/Source/GameplayCamerasEditor/Private/GameplayCamerasEditorModule.cpp
+++ b/Engine/Plugins/Cameras/GameplayCameras/Source/GameplayCamerasEditor/Private/GameplayCamerasEditorModule.cpp
@@ -3,9 +3,9 @@
 #include "GameplayCamerasEditorSettings.h"
 
 #include "AssetTools/CameraAssetEditor.h"
-#include "AssetTools/CameraModeAssetEditor.h"
+#include "AssetTools/CameraRigAssetEditor.h"
 #include "Commands/CameraAssetEditorCommands.h"
-#include "Commands/CameraModeAssetEditorCommands.h"
+#include "Commands/CameraRigAssetEditorCommands.h"
 #include "GameplayCamerasLiveEditManager.h"
 #include "IGameplayCamerasEditorModule.h"
 #include "IGameplayCamerasModule.h"
@@ -14,12 +14,12 @@
 #include "Styles/GameplayCamerasEditorStyle.h"
 #include "ToolMenus.h"
 #include "Toolkits/CameraAssetEditorToolkit.h"
-#include "Toolkits/CameraModeAssetEditorToolkit.h"
+#include "Toolkits/CameraRigAssetEditorToolkit.h"
 
 #define LOCTEXT_NAMESPACE "GameplayCamerasEditor"
 
 const FName IGameplayCamerasEditorModule::GameplayCamerasEditorAppIdentifier("GameplayCamerasEditorApp");
-const FName IGameplayCamerasEditorModule::CameraModeAssetEditorToolBarName("CameraModeAssetEditor.ToolBar");
+const FName IGameplayCamerasEditorModule::CameraRigAssetEditorToolBarName("CameraRigAssetEditor.ToolBar");
 
 /**
  * Implements the FGameplayCamerasEditor module.
@@ -45,7 +45,7 @@ public:
 		UToolMenus::UnRegisterStartupCallback(this);
 
 		FCameraAssetEditorCommands::Unregister();
-		FCameraModeAssetEditorCommands::Unregister();
+		FCameraRigAssetEditorCommands::Unregister();
 
 		UnregisterSettings();
 		TeardownLiveEditManager();
@@ -59,11 +59,11 @@ public:
 		return AssetEditor;
 	}
 
-	virtual UCameraModeAssetEditor* CreateCameraModeEditor(const EToolkitMode::Type Mode, const TSharedPtr<IToolkitHost>& InitToolkitHost, UCameraMode* CameraMode) override
+	virtual UCameraRigAssetEditor* CreateCameraRigEditor(const EToolkitMode::Type Mode, const TSharedPtr<IToolkitHost>& InitToolkitHost, UCameraRigAsset* CameraRig) override
 	{
 		UAssetEditorSubsystem* AssetEditorSubsystem = GEditor->GetEditorSubsystem<UAssetEditorSubsystem>();
-		UCameraModeAssetEditor* AssetEditor = NewObject<UCameraModeAssetEditor>(AssetEditorSubsystem, NAME_None, RF_Transient);
-		AssetEditor->Initialize(CameraMode);
+		UCameraRigAssetEditor* AssetEditor = NewObject<UCameraRigAssetEditor>(AssetEditorSubsystem, NAME_None, RF_Transient);
+		AssetEditor->Initialize(CameraRig);
 		return AssetEditor;
 	}
 
@@ -97,7 +97,7 @@ private:
 	void RegisterMenus()
 	{
 		FCameraAssetEditorCommands::Register();
-		FCameraModeAssetEditorCommands::Register();	
+		FCameraRigAssetEditorCommands::Register();	
 	}
 
 	void InitializeLiveEditManager()
diff --git a/Engine/Plugins/Cameras/GameplayCameras/Source/GameplayCamerasEditor/Private/Toolkits/CameraModeAssetEditorToolkit.cpp b/Engine/Plugins/Cameras/GameplayCameras/Source/GameplayCamerasEditor/Private/Toolkits/CameraModeAssetEditorToolkit.cpp
deleted file mode 100644
diff --git a/Engine/Plugins/Cameras/GameplayCameras/Source/GameplayCamerasEditor/Private/Toolkits/CameraRigAssetEditorToolkit.cpp b/Engine/Plugins/Cameras/GameplayCameras/Source/GameplayCamerasEditor/Private/Toolkits/CameraRigAssetEditorToolkit.cpp
new file mode 100644
--- /dev/null
+++ b/Engine/Plugins/Cameras/GameplayCameras/Source/GameplayCamerasEditor/Private/Toolkits/CameraRigAssetEditorToolkit.cpp
@@ -0,0 +1,219 @@
+// Copyright Epic Games, Inc. All Rights Reserved.
+
+#include "Toolkits/CameraRigAssetEditorToolkit.h"
+
+#include "AssetTools/CameraRigAssetEditor.h"
+#include "Commands/CameraRigAssetEditorCommands.h"
+#include "Core/CameraRigAsset.h"
+#include "EditorModeManager.h"
+#include "Framework/Docking/LayoutExtender.h"
+#include "IGameplayCamerasLiveEditManager.h"
+#include "IGameplayCamerasModule.h"
+#include "IMessageLogListing.h"
+#include "MessageLogInitializationOptions.h"
+#include "MessageLogModule.h"
+#include "Modules/ModuleManager.h"
+#include "PropertyEditorModule.h"
+#include "Widgets/Docking/SDockTab.h"
+
+#define LOCTEXT_NAMESPACE "CameraRigAssetEditorToolkit"
+
+const FName FCameraRigAssetEditorToolkit::DetailsViewTabId(TEXT("CameraRigAssetEditor_DetailsView"));
+
+FCameraRigAssetEditorToolkit::FCameraRigAssetEditorToolkit(UCameraRigAssetEditor* InOwningAssetEditor)
+	: FBaseAssetToolkit(InOwningAssetEditor)
+	, CommandBindings(new FUICommandList())
+	, CameraRigAsset(InOwningAssetEditor->GetCameraRigAsset())
+{
+	// Override base class default layout.
+	StandaloneDefaultLayout = FTabManager::NewLayout("CameraRigAssetEditor_Layout")
+		->AddArea
+		(
+			FTabManager::NewPrimaryArea()->SetOrientation(Orient_Vertical)
+			->Split
+			(
+				FTabManager::NewSplitter()->SetOrientation(Orient_Horizontal)->SetSizeCoefficient(0.9f)
+				->Split
+				(
+					FTabManager::NewStack()
+					->AddTab(DetailsViewTabId, ETabState::OpenedTab)
+					->SetForegroundTab(DetailsViewTabId)
+				)
+			)
+		);
+}
+
+FCameraRigAssetEditorToolkit::~FCameraRigAssetEditorToolkit()
+{
+}
+
+void FCameraRigAssetEditorToolkit::RegisterTabSpawners(const TSharedRef<FTabManager>& InTabManager)
+{
+	// Skip FBaseAssetToolkit here because we don't want a viewport tab.
+	FAssetEditorToolkit::RegisterTabSpawners(InTabManager);
+
+	InTabManager->RegisterTabSpawner(DetailsViewTabId, FOnSpawnTab::CreateSP(this, &FCameraRigAssetEditorToolkit::SpawnTab_Details))
+		.SetDisplayName(LOCTEXT("Details", "Details"))
+		.SetGroup(AssetEditorTabsCategory.ToSharedRef())
+		.SetIcon(FSlateIcon(FAppStyle::GetAppStyleSetName(), "LevelEditor.Tabs.Details"));
+}
+
+void FCameraRigAssetEditorToolkit::UnregisterTabSpawners(const TSharedRef<FTabManager>& InTabManager)
+{
+	// Skip FBaseAssetToolkit here because we don't want a viewport tab.
+	FAssetEditorToolkit::UnregisterTabSpawners(InTabManager);
+
+	InTabManager->UnregisterTabSpawner(DetailsViewTabId);
+}
+
+void FCameraRigAssetEditorToolkit::CreateWidgets()
+{
+	// Skip FBaseAssetToolkit here because we don't want a viewport tab.
+	// ...no up-call...
+
+	// Do most of FBaseAssetToolkit's work except for the viewport.
+	RegisterToolbar();
+	LayoutExtender = MakeShared<FLayoutExtender>();
+
+	FPropertyEditorModule& PropertyEditorModule = FModuleManager::GetModuleChecked<FPropertyEditorModule>("PropertyEditor");
+	FDetailsViewArgs DetailsViewArgs;
+	DetailsViewArgs.NameAreaSettings = FDetailsViewArgs::HideNameArea;
+	DetailsViewArgs.bHideSelectionTip = true;
+	DetailsViewArgs.NotifyHook = this;
+	DetailsView = PropertyEditorModule.CreateDetailView(DetailsViewArgs);
+
+	// Create the message log.
+	FMessageLogModule& MessageLogModule = FModuleManager::LoadModuleChecked<FMessageLogModule>("MessageLog");
+	FMessageLogInitializationOptions LogOptions;
+	LogOptions.bShowPages = false;
+	LogOptions.bShowFilters = false;
+	LogOptions.bAllowClear = false;
+	LogOptions.MaxPageCount = 1;
+	StatsListing = MessageLogModule.CreateLogListing("CameraRigAssetEditorStats", LogOptions);
+
+	Stats = MessageLogModule.CreateLogListingWidget(StatsListing.ToSharedRef());
+}
+
+void FCameraRigAssetEditorToolkit::RegisterToolbar()
+{
+	FName ParentName;
+	const FName MenuName = GetToolMenuToolbarName(ParentName);
+	UToolMenus* ToolMenus = UToolMenus::Get();
+	if (!ToolMenus->IsMenuRegistered(MenuName))
+	{
+		const FCameraRigAssetEditorCommands& Commands = FCameraRigAssetEditorCommands::Get();
+
+		FToolMenuOwnerScoped ToolMenuOwnerScope(this);
+
+		UToolMenu* ToolBarMenu = UToolMenus::Get()->RegisterMenu(
+				MenuName, ParentName, EMultiBoxType::ToolBar);
+
+		FToolMenuSection& ToolBarSection = ToolBarMenu->FindOrAddSection("Build");
+		FToolMenuEntry BuildButton = FToolMenuEntry::InitToolBarButton(Commands.Build);
+		BuildButton.Icon = TAttribute<FSlateIcon>(
+				SharedThis(this), &FCameraRigAssetEditorToolkit::GetBuildButtonIcon);
+		BuildButton.ToolTip = TAttribute<FText>(
+				SharedThis(this), &FCameraRigAssetEditorToolkit::GetBuildButtonTooltip);
+		ToolBarSection.AddEntry(BuildButton);
+	}
+}
+
+FSlateIcon FCameraRigAssetEditorToolkit::GetBuildButtonIcon() const
+{
+	static const FName BuildStatusBackground("Blueprint.CompileStatus.Background");
+	static const FName BuildStatusUnknown("Blueprint.CompileStatus.Overlay.Unknown");
+	static const FName BuildStatusError("Blueprint.CompileStatus.Overlay.Error");
+	static const FName BuildStatusGood("Blueprint.CompileStatus.Overlay.Good");
+	static const FName BuildStatusWarning("Blueprint.CompileStatus.Overlay.Warning");
+
+	switch (CameraRigAsset->BuildStatus)
+	{
+		default:
+		case ECameraRigBuildStatus::Dirty:
+			return FSlateIcon(
+					FAppStyle::GetAppStyleSetName(), BuildStatusBackground, 
+					NAME_None, BuildStatusUnknown);
+		case ECameraRigBuildStatus::WithErrors:
+			return FSlateIcon(
+					FAppStyle::GetAppStyleSetName(), BuildStatusBackground, 
+					NAME_None, BuildStatusError);
+		case ECameraRigBuildStatus::Clean:
+			return FSlateIcon(
+					FAppStyle::GetAppStyleSetName(), BuildStatusBackground, 
+					NAME_None, BuildStatusGood);
+		case ECameraRigBuildStatus::CleanWithWarnings:
+			return FSlateIcon(
+					FAppStyle::GetAppStyleSetName(), BuildStatusBackground, 
+					NAME_None, BuildStatusWarning);
+	}
+}
+
+FText FCameraRigAssetEditorToolkit::GetBuildButtonTooltip() const
+{
+	switch (CameraRigAsset->BuildStatus)
+	{
+		default:
+		case ECameraRigBuildStatus::Dirty:
+			return LOCTEXT("BuildButtonStatusDirty", "Dirty or unknown, should rebuild");
+		case ECameraRigBuildStatus::WithErrors:
+			return LOCTEXT("BuildButtonStatusWithErrors", "There were errors during the build, see the log window for details");
+		case ECameraRigBuildStatus::Clean:
+			return LOCTEXT("BuildButtonStatusClean", "Good to go");
+		case ECameraRigBuildStatus::CleanWithWarnings:
+			return LOCTEXT("BuildButtonStatusCleanWithWarnings", "There were warnings during the build, see the log window for details");
+	}
+}
+
+void FCameraRigAssetEditorToolkit::PostInitAssetEditor()
+{
+	const FCameraRigAssetEditorCommands& Commands = FCameraRigAssetEditorCommands::Get();
+
+	ToolkitCommands->MapAction(
+		Commands.Build,
+		FExecuteAction::CreateSP(this, &FCameraRigAssetEditorToolkit::OnBuild));
+
+	IGameplayCamerasModule& GameplayCamerasModule = FModuleManager::GetModuleChecked<IGameplayCamerasModule>("GameplayCameras");
+	LiveEditManager = GameplayCamerasModule.GetLiveEditManager();
+}
+
+FText FCameraRigAssetEditorToolkit::GetBaseToolkitName() const
+{
+	return LOCTEXT("AppLabel", "Camera Rig Asset");
+}
+
+FName FCameraRigAssetEditorToolkit::GetToolkitFName() const
+{
+	static FName SequencerName("CameraRigAssetEditor");
+	return SequencerName;
+}
+
+FString FCameraRigAssetEditorToolkit::GetWorldCentricTabPrefix() const
+{
+	return LOCTEXT("WorldCentricTabPrefix", "Camera Rig Asset ").ToString();
+}
+
+FLinearColor FCameraRigAssetEditorToolkit::GetWorldCentricTabColorScale() const
+{
+	return FLinearColor(0.7, 0.0f, 0.0f, 0.5f);
+}
+
+void FCameraRigAssetEditorToolkit::NotifyPostChange(const FPropertyChangedEvent& PropertyChangedEvent, FProperty* PropertyThatChanged)
+{
+	CameraRigAsset->BuildStatus = ECameraRigBuildStatus::Dirty;
+}
+
+void FCameraRigAssetEditorToolkit::OnBuild()
+{
+	CameraRigAsset->BuildStatus = ECameraRigBuildStatus::Clean;
+
+	FCameraRigPackages BuiltPackages;
+	CameraRigAsset->GatherPackages(BuiltPackages);
+
+	for (const UPackage* BuiltPackage : BuiltPackages)
+	{
+		LiveEditManager->NotifyPostBuildAsset(BuiltPackage);
+	}
+}
+
+#undef LOCTEXT_NAMESPACE
+
diff --git a/Engine/Plugins/Cameras/GameplayCameras/Source/GameplayCamerasEditor/Public/AssetTools/CameraModeAssetEditor.h b/Engine/Plugins/Cameras/GameplayCameras/Source/GameplayCamerasEditor/Public/AssetTools/CameraModeAssetEditor.h
deleted file mode 100644
diff --git a/Engine/Plugins/Cameras/GameplayCameras/Source/GameplayCamerasEditor/Public/AssetTools/CameraRigAssetEditor.h b/Engine/Plugins/Cameras/GameplayCameras/Source/GameplayCamerasEditor/Public/AssetTools/CameraRigAssetEditor.h
new file mode 100644
--- /dev/null
+++ b/Engine/Plugins/Cameras/GameplayCameras/Source/GameplayCamerasEditor/Public/AssetTools/CameraRigAssetEditor.h
@@ -0,0 +1,36 @@
+// Copyright Epic Games, Inc. All Rights Reserved.
+
+#pragma once
+
+#include "Tools/UAssetEditor.h"
+
+#include "CameraRigAssetEditor.generated.h"
+
+class FBaseAssetToolkit;
+class UCameraRigAsset;
+
+/**
+ * Editor for a camera rig asset.
+ */
+UCLASS(Transient)
+class UCameraRigAssetEditor : public UAssetEditor
+{
+	GENERATED_BODY()
+
+public:
+
+	void Initialize(TObjectPtr<UCameraRigAsset> InCameraRigAsset);
+
+	UCameraRigAsset* GetCameraRigAsset() const { return CameraRigAsset; }
+
+public:
+
+	// UAssetEditor interface
+	virtual void GetObjectsToEdit(TArray<UObject*>& InObjectsToEdit) override;
+	virtual TSharedPtr<FBaseAssetToolkit> CreateToolkit() override;
+
+private:
+
+	TObjectPtr<UCameraRigAsset> CameraRigAsset;
+};
+
diff --git a/Engine/Plugins/Cameras/GameplayCameras/Source/GameplayCamerasEditor/Public/Commands/CameraModeAssetEditorCommands.h b/Engine/Plugins/Cameras/GameplayCameras/Source/GameplayCamerasEditor/Public/Commands/CameraModeAssetEditorCommands.h
deleted file mode 100644
diff --git a/Engine/Plugins/Cameras/GameplayCameras/Source/GameplayCamerasEditor/Public/Commands/CameraRigAssetEditorCommands.h b/Engine/Plugins/Cameras/GameplayCameras/Source/GameplayCamerasEditor/Public/Commands/CameraRigAssetEditorCommands.h
new file mode 100644
--- /dev/null
+++ b/Engine/Plugins/Cameras/GameplayCameras/Source/GameplayCamerasEditor/Public/Commands/CameraRigAssetEditorCommands.h
@@ -0,0 +1,24 @@
+// Copyright Epic Games, Inc. All Rights Reserved.
+
+#pragma once
+
+#include "CoreTypes.h"
+#include "Framework/Commands/Commands.h"
+
+class FUICommandInfo;
+class UClass;
+struct FInputChord;
+
+class FCameraRigAssetEditorCommands : public TCommands<FCameraRigAssetEditorCommands>
+{
+public:
+
+	FCameraRigAssetEditorCommands();
+
+	virtual void RegisterCommands() override;
+	
+public:
+
+	TSharedPtr<FUICommandInfo> Build;
+};
+
diff --git a/Engine/Plugins/Cameras/GameplayCameras/Source/GameplayCamerasEditor/Public/IGameplayCamerasEditorModule.h b/Engine/Plugins/Cameras/GameplayCameras/Source/GameplayCamerasEditor/Public/IGameplayCamerasEditorModule.h
--- a/Engine/Plugins/Cameras/GameplayCameras/Source/GameplayCamerasEditor/Public/IGameplayCamerasEditorModule.h
+++ b/Engine/Plugins/Cameras/GameplayCameras/Source/GameplayCamerasEditor/Public/IGameplayCamerasEditorModule.h
@@ -8,8 +8,8 @@
 
 class UCameraAsset;
 class UCameraAssetEditor;
-class UCameraMode;
-class UCameraModeAssetEditor;
+class UCameraRigAsset;
+class UCameraRigAssetEditor;
 
 /**
  * The gameplay cameras editor module.
@@ -20,7 +20,7 @@ public:
 
 	static const FName GameplayCamerasEditorAppIdentifier;
 
-	static const FName CameraModeAssetEditorToolBarName;
+	static const FName CameraRigAssetEditorToolBarName;
 
 	virtual ~IGameplayCamerasEditorModule() = default;
 
@@ -29,7 +29,7 @@ public:
 	/** Creates an editor for the given camera asset */
 	virtual UCameraAssetEditor* CreateCameraAssetEditor(const EToolkitMode::Type Mode, const TSharedPtr<IToolkitHost>& InitToolkitHost, UCameraAsset* CameraAsset) = 0;
 
-	/** Creates an editor for the given camera mode asset */
-	virtual UCameraModeAssetEditor* CreateCameraModeEditor(const EToolkitMode::Type Mode, const TSharedPtr<IToolkitHost>& InitToolkitHost, UCameraMode* CameraMode) = 0;
+	/** Creates an editor for the given camera rig asset */
+	virtual UCameraRigAssetEditor* CreateCameraRigEditor(const EToolkitMode::Type Mode, const TSharedPtr<IToolkitHost>& InitToolkitHost, UCameraRigAsset* CameraRig) = 0;
 };
 
diff --git a/Engine/Plugins/Cameras/GameplayCameras/Source/GameplayCamerasEditor/Public/Toolkits/CameraModeAssetEditorToolkit.h b/Engine/Plugins/Cameras/GameplayCameras/Source/GameplayCamerasEditor/Public/Toolkits/CameraModeAssetEditorToolkit.h
deleted file mode 100644
diff --git a/Engine/Plugins/Cameras/GameplayCameras/Source/GameplayCamerasEditor/Public/Toolkits/CameraRigAssetEditorToolkit.h b/Engine/Plugins/Cameras/GameplayCameras/Source/GameplayCamerasEditor/Public/Toolkits/CameraRigAssetEditorToolkit.h
new file mode 100644
--- /dev/null
+++ b/Engine/Plugins/Cameras/GameplayCameras/Source/GameplayCamerasEditor/Public/Toolkits/CameraRigAssetEditorToolkit.h
@@ -0,0 +1,83 @@
+// Copyright Epic Games, Inc. All Rights Reserved.
+
+#pragma once
+
+#include "CoreMinimal.h"
+#include "Misc/NotifyHook.h"
+#include "Tools/BaseAssetToolkit.h"
+#include "UObject/GCObject.h"
+
+class IGameplayCamerasLiveEditManager;
+class IMessageLogListing;
+class UCameraRigAsset;
+class UCameraRigAssetEditor;
+class SWidget;
+
+/**
+ * Editor toolkit for a camera rig asset.
+ */
+class FCameraRigAssetEditorToolkit 
+	: public FBaseAssetToolkit
+	, public FGCObject
+	, public FNotifyHook
+{
+public:
+
+	FCameraRigAssetEditorToolkit(UCameraRigAssetEditor* InOwningAssetEditor);
+	~FCameraRigAssetEditorToolkit();
+
+protected:
+
+	// FBaseAssetToolkit interface
+	virtual void RegisterTabSpawners(const TSharedRef<class FTabManager>& TabManager) override;
+	virtual void UnregisterTabSpawners(const TSharedRef<FTabManager>& InTabManager) override;
+	virtual void CreateWidgets() override;
+	virtual void RegisterToolbar() override;
+	virtual void PostInitAssetEditor() override;
+
+	// IToolkit interface
+	virtual FText GetBaseToolkitName() const override;
+	virtual FName GetToolkitFName() const override;
+	virtual FString GetWorldCentricTabPrefix() const override;
+	virtual FLinearColor GetWorldCentricTabColorScale() const override;
+
+	// FGCObject interface
+	virtual void AddReferencedObjects(FReferenceCollector& Collector) override
+	{
+		Collector.AddReferencedObject(CameraRigAsset);
+	}
+	virtual FString GetReferencerName() const override
+	{
+		return TEXT("FCameraRigAssetEditorToolkit");
+	}
+
+	// FNotifyHook interface
+	virtual void NotifyPostChange(const FPropertyChangedEvent& PropertyChangedEvent, FProperty* PropertyThatChanged) override;
+
+private:
+
+	FSlateIcon GetBuildButtonIcon() const;
+	FText GetBuildButtonTooltip() const;
+
+	void OnBuild();
+
+private:
+
+	static const FName DetailsViewTabId;
+
+	/** Command bindings */
+	TSharedRef<FUICommandList> CommandBindings;
+
+	/** Message log widget */
+	TSharedPtr<SWidget> Stats;
+
+	/** Message log listing */
+	TSharedPtr<IMessageLogListing> StatsListing;
+
+	/** The asset being edited */
+	TObjectPtr<UCameraRigAsset> CameraRigAsset;
+
+	/** Live edit manager for updating the assets in the runtime */
+	TSharedPtr<IGameplayCamerasLiveEditManager> LiveEditManager;
+};
+
