Cameras: fix clang-complete build

diff --git a/Engine/Plugins/Cameras/GameplayCameras/Source/GameplayCameras/Private/GameplayCamerasModule.cpp b/Engine/Plugins/Cameras/GameplayCameras/Source/GameplayCameras/Private/GameplayCamerasModule.cpp
--- a/Engine/Plugins/Cameras/GameplayCameras/Source/GameplayCameras/Private/GameplayCamerasModule.cpp
+++ b/Engine/Plugins/Cameras/GameplayCameras/Source/GameplayCameras/Private/GameplayCamerasModule.cpp
@@ -9,7 +9,6 @@
 #include "GameplayCameras.h"
 #include "Logging/MessageLog.h"
 #include "Modules/ModuleManager.h"
-#include "Templates/SharedPointer.h"
 
 DEFINE_LOG_CATEGORY(LogCameraSystem);
 
