Cameras: make the camera system and evaluation context classes non-UObjects

This prevents unnecessary dependencies on the UObject framework, and the possibilty of running vast amounts of camera contexts.
The only reason camera contexts were UObjects were for type-casting anyway, and that's now possible with the simple RTTI system that was added for node evaluators.

diff --git a/Engine/Plugins/Cameras/GameplayCameras/Source/GameplayCameras/Private/Core/BlendStackCameraNode.cpp b/Engine/Plugins/Cameras/GameplayCameras/Source/GameplayCameras/Private/Core/BlendStackCameraNode.cpp
--- a/Engine/Plugins/Cameras/GameplayCameras/Source/GameplayCameras/Private/Core/BlendStackCameraNode.cpp
+++ b/Engine/Plugins/Cameras/GameplayCameras/Source/GameplayCameras/Private/Core/BlendStackCameraNode.cpp
@@ -113,8 +113,8 @@ void FBlendStackCameraNodeEvaluator::OnRun(const FCameraNodeEvaluationParams& Pa
 	// Start by evaluating all the root nodes in the stack.
 	for (FCameraRigEntry& Entry : Entries)
 	{
-		const UCameraEvaluationContext* CurContext = Entry.EvaluationContext.Get();
-		if (UNLIKELY(CurContext == nullptr))
+		TSharedPtr<const FCameraEvaluationContext> CurContext = Entry.EvaluationContext.Pin();
+		if (UNLIKELY(!CurContext.IsValid()))
 		{
 			Entry.Result.bIsValid = false;
 			continue;
@@ -172,7 +172,7 @@ void FBlendStackCameraNodeEvaluator::OnRun(const FCameraNodeEvaluationParams& Pa
 		const FCameraPoseFlags ChangedFlags(CurResult.CameraPose.GetChangedFlags());
 
 		FCameraNodeEvaluationParams CurParams(Params);
-		CurParams.EvaluationContext = Entry.EvaluationContext.Get();
+		CurParams.EvaluationContext = Entry.EvaluationContext.Pin();
 		CurParams.bIsFirstFrame = Entry.bIsFirstFrame;
 		FCameraNodeBlendParams BlendParams(CurParams, CurResult);
 
@@ -233,7 +233,7 @@ const FCameraRigTransition* FBlendStackCameraNodeEvaluator::FindTransition(const
 {
 	const UBlendStackCameraNode* BlendStackNode = GetCameraNodeAs<UBlendStackCameraNode>();
 
-	const UCameraEvaluationContext* ToContext = Params.EvaluationContext.Get();
+	TSharedPtr<const FCameraEvaluationContext> ToContext = Params.EvaluationContext;
 	const UCameraAsset* ToCameraAsset = ToContext ? ToContext->GetCameraAsset() : nullptr;
 	const UCameraRigAsset* ToCameraRig = Params.CameraRig;
 
@@ -248,7 +248,7 @@ const FCameraRigTransition* FBlendStackCameraNodeEvaluator::FindTransition(const
 		// Start by looking at exit transitions on the last active (top) camera rig.
 		const FCameraRigEntry& TopEntry = Entries.Top();
 
-		const UCameraEvaluationContext* FromContext = TopEntry.EvaluationContext.Get();
+		TSharedPtr<const FCameraEvaluationContext> FromContext = TopEntry.EvaluationContext.Pin();
 		const UCameraAsset* FromCameraAsset = FromContext ? FromContext->GetCameraAsset() : nullptr;
 		const UCameraRigAsset* FromCameraRig = TopEntry.CameraRig;
 
@@ -345,6 +345,15 @@ const FCameraRigTransition* FBlendStackCameraNodeEvaluator::FindTransition(
 	return nullptr;
 }
 
+void FBlendStackCameraNodeEvaluator::OnAddReferencedObjects(FReferenceCollector& Collector)
+{
+	for (FCameraRigEntry& Entry : Entries)
+	{
+		Collector.AddReferencedObject(Entry.CameraRig);
+		Collector.AddReferencedObject(Entry.RootNode);
+	}
+}
+
 #if WITH_EDITOR
 
 void FBlendStackCameraNodeEvaluator::OnPostBuildAsset(const FGameplayCameraAssetBuildEvent& BuildEvent)
@@ -367,7 +376,7 @@ void FBlendStackCameraNodeEvaluator::OnPostBuildAsset(const FGameplayCameraAsset
 			FCameraNodeEvaluatorTreeBuilderParams BuildParams;
 			BuildParams.RootCameraNode = Entry.RootNode;
 			BuildParams.Evaluator = OwningEvaluator;
-			BuildParams.EvaluationContext = Entry.EvaluationContext.Get();
+			BuildParams.EvaluationContext = Entry.EvaluationContext.Pin();
 			FCameraNodeEvaluator* RootEvaluator = Entry.EvaluatorStorage.BuildEvaluatorTree(BuildParams);
 
 			Entry.RootEvaluator = RootEvaluator->CastThisChecked<FBlendStackRootCameraNodeEvaluator>();
diff --git a/Engine/Plugins/Cameras/GameplayCameras/Source/GameplayCameras/Private/Core/CameraEvaluationContext.cpp b/Engine/Plugins/Cameras/GameplayCameras/Source/GameplayCameras/Private/Core/CameraEvaluationContext.cpp
--- a/Engine/Plugins/Cameras/GameplayCameras/Source/GameplayCameras/Private/Core/CameraEvaluationContext.cpp
+++ b/Engine/Plugins/Cameras/GameplayCameras/Source/GameplayCameras/Private/Core/CameraEvaluationContext.cpp
@@ -2,10 +2,13 @@
 
 #include "Core/CameraEvaluationContext.h"
 
-#include UE_INLINE_GENERATED_CPP_BY_NAME(CameraEvaluationContext)
+UE_GAMEPLAY_CAMERAS_DEFINE_RTTI(FCameraEvaluationContext)
 
-UCameraEvaluationContext::UCameraEvaluationContext(const FObjectInitializer& ObjectInit)
-	: Super(ObjectInit)
+FCameraEvaluationContext::FCameraEvaluationContext()
+{
+}
+
+FCameraEvaluationContext::~FCameraEvaluationContext()
 {
 }
 
diff --git a/Engine/Plugins/Cameras/GameplayCameras/Source/GameplayCameras/Private/Core/CameraEvaluationContextStack.cpp b/Engine/Plugins/Cameras/GameplayCameras/Source/GameplayCameras/Private/Core/CameraEvaluationContextStack.cpp
--- a/Engine/Plugins/Cameras/GameplayCameras/Source/GameplayCameras/Private/Core/CameraEvaluationContextStack.cpp
+++ b/Engine/Plugins/Cameras/GameplayCameras/Source/GameplayCameras/Private/Core/CameraEvaluationContextStack.cpp
@@ -13,7 +13,7 @@ FCameraEvaluationContextInfo FCameraEvaluationContextStack::GetActiveContext() c
 {
 	for (const FContextEntry& Entry : ReverseIterate(Entries))
 	{
-		if (UCameraEvaluationContext* Context = Entry.WeakContext.Get())
+		if (TSharedPtr<FCameraEvaluationContext> Context = Entry.WeakContext.Pin())
 		{
 			return FCameraEvaluationContextInfo{ Context, Entry.CameraDirector, Entry.Evaluator };
 		}
@@ -21,11 +21,11 @@ FCameraEvaluationContextInfo FCameraEvaluationContextStack::GetActiveContext() c
 	return FCameraEvaluationContextInfo();
 }
 
-bool FCameraEvaluationContextStack::HasContext(UCameraEvaluationContext* Context) const
+bool FCameraEvaluationContextStack::HasContext(TSharedRef<FCameraEvaluationContext> Context) const
 {
 	for (const FContextEntry& Entry : ReverseIterate(Entries))
 	{
-		if (Context == Entry.WeakContext.Get())
+		if (Context == Entry.WeakContext)
 		{
 			return true;
 		}
@@ -33,7 +33,7 @@ bool FCameraEvaluationContextStack::HasContext(UCameraEvaluationContext* Context
 	return false;
 }
 
-void FCameraEvaluationContextStack::PushContext(UCameraEvaluationContext* Context)
+void FCameraEvaluationContextStack::PushContext(TSharedRef<FCameraEvaluationContext> Context)
 {
 	checkf(Evaluator, TEXT("Can't push context when no evaluator is set! Did you call Initialize?"));
 
@@ -64,7 +64,7 @@ void FCameraEvaluationContextStack::PushContext(UCameraEvaluationContext* Contex
 	Entries.Push(MoveTemp(NewEntry));
 }
 
-bool FCameraEvaluationContextStack::RemoveContext(UCameraEvaluationContext* Context)
+bool FCameraEvaluationContextStack::RemoveContext(TSharedRef<FCameraEvaluationContext> Context)
 {
 	const int32 NumRemoved = Entries.RemoveAll(
 			[Context](FContextEntry& Entry) { return Entry.WeakContext == Context; });
@@ -76,7 +76,7 @@ void FCameraEvaluationContextStack::PopContext()
 	Entries.Pop();
 }
 
-void FCameraEvaluationContextStack::Initialize(UCameraSystemEvaluator* InEvaluator)
+void FCameraEvaluationContextStack::Initialize(TSharedRef<FCameraSystemEvaluator> InEvaluator)
 {
 	Evaluator = InEvaluator;
 }
diff --git a/Engine/Plugins/Cameras/GameplayCameras/Source/GameplayCameras/Private/Core/CameraSystemEvaluator.cpp b/Engine/Plugins/Cameras/GameplayCameras/Source/GameplayCameras/Private/Core/CameraSystemEvaluator.cpp
--- a/Engine/Plugins/Cameras/GameplayCameras/Source/GameplayCameras/Private/Core/CameraSystemEvaluator.cpp
+++ b/Engine/Plugins/Cameras/GameplayCameras/Source/GameplayCameras/Private/Core/CameraSystemEvaluator.cpp
@@ -10,53 +10,75 @@
 #include "Core/DefaultRootCameraNode.h"
 #include "GameplayCameras.h"
 #include "IGameplayCamerasModule.h"
-
-#include UE_INLINE_GENERATED_CPP_BY_NAME(CameraSystemEvaluator)
+#include "UObject/Package.h"
+#include "UObject/UObjectGlobals.h"
 
 DECLARE_CYCLE_STAT(TEXT("Camera System Eval"), CameraSystemEval_Total, STATGROUP_CameraSystem);
 
-UCameraSystemEvaluator::UCameraSystemEvaluator(const FObjectInitializer& ObjectInit)
-	: Super(ObjectInit)
+FCameraSystemEvaluator::FCameraSystemEvaluator()
 {
-	RootNode = ObjectInit.CreateDefaultSubobject<UDefaultRootCameraNode>(this, TEXT("RootNode"));
+}
 
-	if (!HasAnyFlags(RF_ClassDefaultObject | RF_ArchetypeObject))
+void FCameraSystemEvaluator::Initialize(TObjectPtr<UObject> InOwner)
+{
+	FCameraSystemEvaluatorCreateParams Params;
+	Params.Owner = InOwner;
+	Initialize(Params);
+}
+
+void FCameraSystemEvaluator::Initialize(const FCameraSystemEvaluatorCreateParams& Params)
+{
+	UObject* Owner = Params.Owner;
+	if (!Owner)
 	{
-		ContextStack.Initialize(this);
+		Owner = GetTransientPackage();
+	}
 
-		FCameraNodeEvaluatorTreeBuilderParams BuildParams;
-		BuildParams.Evaluator = this;
-		BuildParams.RootCameraNode = RootNode;
-		RootEvaluator = static_cast<FRootCameraNodeEvaluator*>(RootEvaluatorStorage.BuildEvaluatorTree(BuildParams));
+	if (Params.RootNodeFactory)
+	{
+		RootNode = Params.RootNodeFactory();
+	}
+	else
+	{
+		RootNode = NewObject<UDefaultRootCameraNode>(Owner, TEXT("RootNode"));
 	}
+
+	TSharedRef<FCameraSystemEvaluator> This(SharedThis(this));
+
+	ContextStack.Initialize(This);
+
+	FCameraNodeEvaluatorTreeBuilderParams BuildParams;
+	BuildParams.Evaluator = This;
+	BuildParams.RootCameraNode = RootNode;
+	RootEvaluator = static_cast<FRootCameraNodeEvaluator*>(RootEvaluatorStorage.BuildEvaluatorTree(BuildParams));
 }
 
-void UCameraSystemEvaluator::AddReferencedObjects(UObject* InThis, FReferenceCollector& Collector)
+void FCameraSystemEvaluator::AddReferencedObjects(FReferenceCollector& Collector)
 {
-	UCameraSystemEvaluator* TypedThis = CastChecked<UCameraSystemEvaluator>(InThis);
-	TypedThis->ContextStack.AddReferencedObjects(Collector);
-	if (TypedThis->RootEvaluator)
+	Collector.AddReferencedObject(RootNode);
+	ContextStack.AddReferencedObjects(Collector);
+	if (RootEvaluator)
 	{
-		TypedThis->RootEvaluator->AddReferencedObjects(Collector);
+		RootEvaluator->AddReferencedObjects(Collector);
 	}
 }
 
-void UCameraSystemEvaluator::PushEvaluationContext(UCameraEvaluationContext* EvaluationContext)
+void FCameraSystemEvaluator::PushEvaluationContext(TSharedRef<FCameraEvaluationContext> EvaluationContext)
 {
 	ContextStack.PushContext(EvaluationContext);
 }
 
-void UCameraSystemEvaluator::RemoveEvaluationContext(UCameraEvaluationContext* EvaluationContext)
+void FCameraSystemEvaluator::RemoveEvaluationContext(TSharedRef<FCameraEvaluationContext> EvaluationContext)
 {
 	ContextStack.RemoveContext(EvaluationContext);
 }
 
-void UCameraSystemEvaluator::PopEvaluationContext()
+void FCameraSystemEvaluator::PopEvaluationContext()
 {
 	ContextStack.PopContext();
 }
 
-void UCameraSystemEvaluator::Update(const FCameraSystemEvaluationUpdateParams& Params)
+void FCameraSystemEvaluator::Update(const FCameraSystemEvaluationUpdateParams& Params)
 {
 	SCOPE_CYCLE_COUNTER(CameraSystemEval_Total);
 
@@ -68,6 +90,8 @@ void UCameraSystemEvaluator::Update(const FCameraSystemEvaluationUpdateParams& P
 		return;
 	}
 
+	TSharedPtr<FCameraSystemEvaluator> This(SharedThis(this));
+
 	// Run the camera director, and activate any camera rig(s) it returns to us.
 	FCameraDirectorEvaluator* ActiveDirectorEvaluator = ActiveContextInfo.Evaluator;
 	if (ActiveDirectorEvaluator)
@@ -83,7 +107,7 @@ void UCameraSystemEvaluator::Update(const FCameraSystemEvaluationUpdateParams& P
 		if (DirectorResult.ActiveCameraRigs.Num() == 1)
 		{
 			FActivateCameraRigParams CameraRigParams;
-			CameraRigParams.Evaluator = this;
+			CameraRigParams.Evaluator = This;
 			CameraRigParams.EvaluationContext = ActiveContextInfo.EvaluationContext;
 			CameraRigParams.CameraRig = DirectorResult.ActiveCameraRigs[0];
 			RootEvaluator->ActivateCameraRig(CameraRigParams);
@@ -92,7 +116,7 @@ void UCameraSystemEvaluator::Update(const FCameraSystemEvaluationUpdateParams& P
 
 	// Run the root camera node.
 	FCameraNodeEvaluationParams NodeParams;
-	NodeParams.Evaluator = this;
+	NodeParams.Evaluator = This;
 	NodeParams.DeltaTime = Params.DeltaTime;
 
 	RootNodeResult.Reset();
@@ -104,7 +128,7 @@ void UCameraSystemEvaluator::Update(const FCameraSystemEvaluationUpdateParams& P
 	Result.bIsValid = true;
 }
 
-void UCameraSystemEvaluator::GetEvaluatedCameraView(FMinimalViewInfo& DesiredView)
+void FCameraSystemEvaluator::GetEvaluatedCameraView(FMinimalViewInfo& DesiredView)
 {
 	const FCameraPose& CameraPose = Result.CameraPose;
 	DesiredView.Location = CameraPose.GetLocation();
diff --git a/Engine/Plugins/Cameras/GameplayCameras/Source/GameplayCameras/Private/GameFramework/GameplayCameraComponent.cpp b/Engine/Plugins/Cameras/GameplayCameras/Source/GameplayCameras/Private/GameFramework/GameplayCameraComponent.cpp
--- a/Engine/Plugins/Cameras/GameplayCameras/Source/GameplayCameras/Private/GameFramework/GameplayCameraComponent.cpp
+++ b/Engine/Plugins/Cameras/GameplayCameras/Source/GameplayCameras/Private/GameFramework/GameplayCameraComponent.cpp
@@ -97,14 +97,14 @@ void UGameplayCameraComponent::ActivateCamera(APlayerController* PlayerControlle
 		return;
 	}
 
-	if (EvaluationContext == nullptr)
+	if (!EvaluationContext.IsValid())
 	{
-		EvaluationContext = NewObject<UGameplayCameraComponentEvaluationContext>(this, TEXT("EvaluationContext"));
+		EvaluationContext = MakeShared<FGameplayCameraComponentEvaluationContext>();
 		EvaluationContext->Initialize(this, PlayerController);
 	}
 
-	UCameraSystemEvaluator* Evaluator = CameraSystem->GetCameraSystemComponent()->GetCameraSystemEvaluator();
-	Evaluator->PushEvaluationContext(EvaluationContext);
+	TSharedPtr<FCameraSystemEvaluator> Evaluator = CameraSystem->GetCameraSystemComponent()->GetCameraSystemEvaluator();
+	Evaluator->PushEvaluationContext(EvaluationContext.ToSharedRef());
 
 	Activate();
 }
@@ -122,10 +122,10 @@ void UGameplayCameraComponent::DeactivateCamera(APlayerController* PlayerControl
 		return;
 	}
 
-	if (EvaluationContext != nullptr)
+	if (EvaluationContext.IsValid())
 	{
-		UCameraSystemEvaluator* Evaluator = CameraSystem->GetCameraSystemComponent()->GetCameraSystemEvaluator();
-		Evaluator->RemoveEvaluationContext(EvaluationContext);
+		TSharedPtr<FCameraSystemEvaluator> Evaluator = CameraSystem->GetCameraSystemComponent()->GetCameraSystemEvaluator();
+		Evaluator->RemoveEvaluationContext(EvaluationContext.ToSharedRef());
 	}
 
 	Deactivate();
@@ -190,18 +190,15 @@ void UGameplayCameraComponent::UpdatePreviewMeshTransform()
 
 #endif
 
-UGameplayCameraComponentEvaluationContext::UGameplayCameraComponentEvaluationContext(const FObjectInitializer& ObjectInit)
-	: Super(ObjectInit)
-{
-}
+UE_DEFINE_CAMERA_EVALUATION_CONTEXT(FGameplayCameraComponentEvaluationContext)
 
-void UGameplayCameraComponentEvaluationContext::Initialize(UGameplayCameraComponent* Owner, APlayerController* InPlayerController)
+void FGameplayCameraComponentEvaluationContext::Initialize(UGameplayCameraComponent* Owner, APlayerController* InPlayerController)
 {
 	PlayerController = InPlayerController;
 	CameraAsset = Owner->Camera;
 }
 
-void UGameplayCameraComponentEvaluationContext::Update(UGameplayCameraComponent* Owner)
+void FGameplayCameraComponentEvaluationContext::Update(UGameplayCameraComponent* Owner)
 {
 	const FTransform& OwnerTransform = Owner->GetComponentTransform();
 	InitialResult.CameraPose.SetTransform(OwnerTransform);
diff --git a/Engine/Plugins/Cameras/GameplayCameras/Source/GameplayCameras/Private/GameFramework/GameplayCameraSystemComponent.cpp b/Engine/Plugins/Cameras/GameplayCameras/Source/GameplayCameras/Private/GameFramework/GameplayCameraSystemComponent.cpp
--- a/Engine/Plugins/Cameras/GameplayCameras/Source/GameplayCameras/Private/GameFramework/GameplayCameraSystemComponent.cpp
+++ b/Engine/Plugins/Cameras/GameplayCameras/Source/GameplayCameras/Private/GameFramework/GameplayCameraSystemComponent.cpp
@@ -28,8 +28,6 @@ UGameplayCameraSystemComponent::UGameplayCameraSystemComponent(const FObjectInit
 		PreviewMesh = EditorCameraMesh.Object;
 	}
 #endif  // WITH_EDITORONLY_DATA
-
-	Evaluator = ObjectInit.CreateDefaultSubobject<UCameraSystemEvaluator>(this, "CameraSystemEvaluator");
 }
 
 void UGameplayCameraSystemComponent::GetCameraView(float DeltaTime, FMinimalViewInfo& DesiredView)
@@ -43,6 +41,12 @@ void UGameplayCameraSystemComponent::GetCameraView(float DeltaTime, FMinimalView
 
 void UGameplayCameraSystemComponent::OnRegister()
 {
+	if (!Evaluator.IsValid())
+	{
+		Evaluator = MakeShared<FCameraSystemEvaluator>();
+		Evaluator->Initialize(this);
+	}
+
 #if WITH_EDITORONLY_DATA
 	if (PreviewMesh && !PreviewMeshComponent)
 	{
@@ -92,6 +96,15 @@ bool UGameplayCameraSystemComponent::GetEditorPreviewInfo(float DeltaTime, FMini
 
 #endif  // WITH_EDITOR
 
+void UGameplayCameraSystemComponent::AddReferencedObjects(UObject* InThis, FReferenceCollector& Collector)
+{
+	UGameplayCameraSystemComponent* TypedThis = CastChecked<UGameplayCameraSystemComponent>(InThis);
+	if (TypedThis->Evaluator.IsValid())
+	{
+		TypedThis->Evaluator->AddReferencedObjects(Collector);
+	}
+}
+
 void UGameplayCameraSystemComponent::OnBecomeViewTarget()
 {
 }
diff --git a/Engine/Plugins/Cameras/GameplayCameras/Source/GameplayCameras/Public/Core/BlendStackCameraNode.h b/Engine/Plugins/Cameras/GameplayCameras/Source/GameplayCameras/Public/Core/BlendStackCameraNode.h
--- a/Engine/Plugins/Cameras/GameplayCameras/Source/GameplayCameras/Public/Core/BlendStackCameraNode.h
+++ b/Engine/Plugins/Cameras/GameplayCameras/Source/GameplayCameras/Public/Core/BlendStackCameraNode.h
@@ -11,9 +11,10 @@
 #include "BlendStackCameraNode.generated.h"
 
 class FBlendStackRootCameraNodeEvaluator;
+class FCameraEvaluationContext;
+class FCameraSystemEvaluator;
 class UBlendStackRootCameraNode;
 class UCameraAsset;
-class UCameraEvaluationContext;
 class UCameraRigAsset;
 struct FCameraRigTransition;
 
@@ -23,10 +24,10 @@ struct FCameraRigTransition;
 struct FBlendStackCameraPushParams
 {
 	/** The evaluator currently running.*/
-	TObjectPtr<UCameraSystemEvaluator> Evaluator;
+	TSharedPtr<FCameraSystemEvaluator> Evaluator;
 
 	/** The evaluation context within which a camera rig's node tree should run. */
-	TObjectPtr<const UCameraEvaluationContext> EvaluationContext;
+	TSharedPtr<const FCameraEvaluationContext> EvaluationContext;
 
 	/** The source camera rig asset to instantiate and push on the blend stack. */
 	TObjectPtr<const UCameraRigAsset> CameraRig;
@@ -83,6 +84,7 @@ protected:
 	virtual FCameraNodeEvaluatorChildrenView OnGetChildren() override;
 	virtual void OnInitialize(const FCameraNodeEvaluatorInitializeParams& Params) override;
 	virtual void OnRun(const FCameraNodeEvaluationParams& Params, FCameraNodeEvaluationResult& OutResult) override;
+	virtual void OnAddReferencedObjects(FReferenceCollector& Collector) override;
 
 #if WITH_EDITOR
 	// IGameplayCamerasLiveEditListener interface
@@ -103,7 +105,7 @@ protected:
 	struct FCameraRigEntry
 	{
 		/** Evaluation context in which this entry runs. */
-		TWeakObjectPtr<const UCameraEvaluationContext> EvaluationContext;
+		TWeakPtr<const FCameraEvaluationContext> EvaluationContext;
 		/** The camera rig asset that this entry runs. */
 		TObjectPtr<const UCameraRigAsset> CameraRig;
 		/** The root node. */
@@ -124,7 +126,7 @@ protected:
 	};
 
 	/** The camera system evaluator running this node. */
-	TObjectPtr<UCameraSystemEvaluator> OwningEvaluator;
+	TSharedPtr<FCameraSystemEvaluator> OwningEvaluator;
 
 	/** Entries in the blend stack. */
 	TArray<FCameraRigEntry> Entries;
diff --git a/Engine/Plugins/Cameras/GameplayCameras/Source/GameplayCameras/Public/Core/CameraDirectorEvaluator.h b/Engine/Plugins/Cameras/GameplayCameras/Source/GameplayCameras/Public/Core/CameraDirectorEvaluator.h
--- a/Engine/Plugins/Cameras/GameplayCameras/Source/GameplayCameras/Public/Core/CameraDirectorEvaluator.h
+++ b/Engine/Plugins/Cameras/GameplayCameras/Source/GameplayCameras/Public/Core/CameraDirectorEvaluator.h
@@ -9,8 +9,8 @@
 
 class FCameraDirectorEvaluator;
 class FCameraDirectorEvaluatorStorage;
+class FCameraEvaluationContext;
 class UCameraDirector;
-class UCameraEvaluationContext;
 class UCameraRigAsset;
 
 /**
@@ -22,7 +22,7 @@ struct FCameraDirectorEvaluationParams
 	float DeltaTime = 0.f;
 
 	/** The context in which this director runs. */
-	TObjectPtr<const UCameraEvaluationContext> OwnerContext;
+	TSharedPtr<const FCameraEvaluationContext> OwnerContext;
 };
 
 /**
diff --git a/Engine/Plugins/Cameras/GameplayCameras/Source/GameplayCameras/Public/Core/CameraEvaluationContext.h b/Engine/Plugins/Cameras/GameplayCameras/Source/GameplayCameras/Public/Core/CameraEvaluationContext.h
--- a/Engine/Plugins/Cameras/GameplayCameras/Source/GameplayCameras/Public/Core/CameraEvaluationContext.h
+++ b/Engine/Plugins/Cameras/GameplayCameras/Source/GameplayCameras/Public/Core/CameraEvaluationContext.h
@@ -5,8 +5,8 @@
 #include "CoreTypes.h"
 #include "UObject/ObjectPtr.h"
 #include "Core/CameraNodeEvaluator.h"
-
-#include "CameraEvaluationContext.generated.h"
+#include "Core/CameraObjectRtti.h"
+#include "Templates/SharedPointer.h"
 
 class UCameraAsset;
 class APlayerController;
@@ -14,14 +14,17 @@ class APlayerController;
 /**
  * Base class for providing a context to running camera rigs.
  */
-UCLASS(MinimalAPI)
-class UCameraEvaluationContext : public UObject
+class FCameraEvaluationContext : public TSharedFromThis<FCameraEvaluationContext>
 {
-	GENERATED_BODY()
+	UE_GAMEPLAY_CAMERAS_DECLARE_RTTI_BASE(FCameraEvaluationContext)
 
 public:
 
-	UCameraEvaluationContext(const FObjectInitializer& ObjectInit);
+	/** Constructs an evaluation context. */
+	FCameraEvaluationContext();
+
+	/** Destroys this evaluation context. */
+	virtual ~FCameraEvaluationContext();
 
 	/**
 	 * Gets the player controller (if any) in control of the cameras running inside
@@ -35,20 +38,33 @@ public:
 	/** Gets the initial evaluation result for all camera rigs in this context. */
 	const FCameraNodeEvaluationResult& GetInitialResult() const { return InitialResult; }
 
+public:
+
+	void AddReferencedObjects(FReferenceCollector& Collector);
+
 protected:
 
 	/**
 	 * The player controller (if any) in control of the cameras running inside
 	 * of this evaluation context.
 	 */
-	UPROPERTY()
 	TObjectPtr<APlayerController> PlayerController;
 
 	/** The camera asset hosted in this context. */
-	UPROPERTY()
 	TObjectPtr<UCameraAsset> CameraAsset;
 
 	/** The initial result for all camera rigs in this context. */
 	FCameraNodeEvaluationResult InitialResult;
 };
 
+// Utility macros for declaring and defining camera evaluation contexts.
+//
+#define UE_DECLARE_CAMERA_EVALUATION_CONTEXT(ClassName)\
+	UE_GAMEPLAY_CAMERAS_DECLARE_RTTI(ClassName, FCameraEvaluationContext)
+
+#define UE_DECLARE_CAMERA_EVALUATION_CONTEXT_EX(ClassName, BaseClassName)\
+	UE_GAMEPLAY_CAMERAS_DECLARE_RTTI(ClassName, BaseClassName)
+
+#define UE_DEFINE_CAMERA_EVALUATION_CONTEXT(ClassName)\
+	UE_GAMEPLAY_CAMERAS_DEFINE_RTTI(ClassName)
+
diff --git a/Engine/Plugins/Cameras/GameplayCameras/Source/GameplayCameras/Public/Core/CameraEvaluationContextStack.h b/Engine/Plugins/Cameras/GameplayCameras/Source/GameplayCameras/Public/Core/CameraEvaluationContextStack.h
--- a/Engine/Plugins/Cameras/GameplayCameras/Source/GameplayCameras/Public/Core/CameraEvaluationContextStack.h
+++ b/Engine/Plugins/Cameras/GameplayCameras/Source/GameplayCameras/Public/Core/CameraEvaluationContextStack.h
@@ -7,20 +7,21 @@
 #include "UObject/WeakObjectPtr.h"
 
 class FCameraDirectorEvaluator;
+class FCameraEvaluationContext;
+class FCameraSystemEvaluator;
 class UCameraAsset;
 class UCameraDirector;
-class UCameraEvaluationContext;
-class UCameraSystemEvaluator;
 
 /** Information about a running camera evaluation context. */
 struct FCameraEvaluationContextInfo
 {
 	/** The evaluation context. */
-	UCameraEvaluationContext* EvaluationContext = nullptr;
+	TSharedPtr<FCameraEvaluationContext> EvaluationContext;
 
 	/** The instantiated camera director running in this context. */
-	const UCameraDirector* CameraDirector = nullptr;
+	TObjectPtr<const UCameraDirector> CameraDirector;
 
+	/** The evaluator for the running camera director. */
 	FCameraDirectorEvaluator* Evaluator = nullptr;
 
 	/** Returns whether this structure has a valid context and director. */
@@ -38,13 +39,13 @@ public:
 	FCameraEvaluationContextInfo GetActiveContext() const;
 
 	/** Returns whether the given context exists in the stack. */
-	bool HasContext(UCameraEvaluationContext* Context) const;
+	bool HasContext(TSharedRef<FCameraEvaluationContext> Context) const;
 
 	/** Push a new context on the stack and instantiate its director. */
-	void PushContext(UCameraEvaluationContext* Context);
+	void PushContext(TSharedRef<FCameraEvaluationContext> Context);
 
 	/** Remove an existing context from the stack. */
-	bool RemoveContext(UCameraEvaluationContext* Context);
+	bool RemoveContext(TSharedRef<FCameraEvaluationContext> Context);
 
 	/** Pop the active (top) context. */
 	void PopContext();
@@ -52,14 +53,14 @@ public:
 public:
 
 	// Internal API
-	void Initialize(UCameraSystemEvaluator* InEvaluator);
+	void Initialize(TSharedRef<FCameraSystemEvaluator> InEvaluator);
 	void AddReferencedObjects(FReferenceCollector& Collector);
 
 private:
 
 	struct FContextEntry
 	{
-		TWeakObjectPtr<UCameraEvaluationContext> WeakContext;
+		TWeakPtr<FCameraEvaluationContext> WeakContext;
 		TObjectPtr<const UCameraDirector> CameraDirector;
 		FCameraDirectorEvaluatorStorage EvaluatorStorage;
 		FCameraDirectorEvaluator* Evaluator = nullptr;
@@ -69,6 +70,6 @@ private:
 	TArray<FContextEntry> Entries;
 
 	/** The owner evaluator. */
-	TObjectPtr<UCameraSystemEvaluator> Evaluator;
+	TSharedPtr<FCameraSystemEvaluator> Evaluator;
 };
 
diff --git a/Engine/Plugins/Cameras/GameplayCameras/Source/GameplayCameras/Public/Core/CameraNodeEvaluator.h b/Engine/Plugins/Cameras/GameplayCameras/Source/GameplayCameras/Public/Core/CameraNodeEvaluator.h
--- a/Engine/Plugins/Cameras/GameplayCameras/Source/GameplayCameras/Public/Core/CameraNodeEvaluator.h
+++ b/Engine/Plugins/Cameras/GameplayCameras/Source/GameplayCameras/Public/Core/CameraNodeEvaluator.h
@@ -9,11 +9,11 @@
 #include "CoreTypes.h"
 #include "UObject/ObjectPtr.h"
 
-class FReferenceCollector;
+class FCameraEvaluationContext;
 class FCameraNodeEvaluator;
-class UCameraEvaluationContext;
+class FCameraSystemEvaluator;
+class FReferenceCollector;
 class UCameraNode;
-class UCameraSystemEvaluator;
 struct FCameraNodeEvaluatorBuilder;
 
 /**
@@ -22,9 +22,9 @@ struct FCameraNodeEvaluatorBuilder;
 struct FCameraNodeEvaluationParams
 {
 	/** The evaluation running this evaluation. */
-	TObjectPtr<UCameraSystemEvaluator> Evaluator;
+	TSharedPtr<FCameraSystemEvaluator> Evaluator;
 	/** The evaluation context (if any) responsible for this branch of the evaluation. */
-	TObjectPtr<const UCameraEvaluationContext> EvaluationContext;
+	TSharedPtr<const FCameraEvaluationContext> EvaluationContext;
 	/** The time interval for the evaluation. */
 	float DeltaTime = 0.f;
 	/** Whether this is the first evaluation of this camera node hierarchy. */
@@ -57,9 +57,9 @@ struct FCameraNodeEvaluationResult
 struct FCameraNodeEvaluatorInitializeParams
 {
 	/** The evaluation running this evaluation. */
-	TObjectPtr<UCameraSystemEvaluator> Evaluator;
+	TSharedPtr<FCameraSystemEvaluator> Evaluator;
 	/** The evaluation context (if any) responsible for this branch of the evaluation. */
-	TObjectPtr<const UCameraEvaluationContext> EvaluationContext;
+	TSharedPtr<const FCameraEvaluationContext> EvaluationContext;
 	/** Builder object for building children evaluators. */
 	FCameraNodeEvaluatorBuilder* Builder = nullptr;
 
@@ -105,6 +105,8 @@ public:
 		return Cast<CameraNodeType>(PrivateCameraNode);
 	}
 
+public:
+
 	// Internal API.
 	void SetPrivateCameraNode(TObjectPtr<const UCameraNode> InCameraNode);
 
diff --git a/Engine/Plugins/Cameras/GameplayCameras/Source/GameplayCameras/Public/Core/CameraNodeEvaluatorStorage.h b/Engine/Plugins/Cameras/GameplayCameras/Source/GameplayCameras/Public/Core/CameraNodeEvaluatorStorage.h
--- a/Engine/Plugins/Cameras/GameplayCameras/Source/GameplayCameras/Public/Core/CameraNodeEvaluatorStorage.h
+++ b/Engine/Plugins/Cameras/GameplayCameras/Source/GameplayCameras/Public/Core/CameraNodeEvaluatorStorage.h
@@ -8,13 +8,13 @@
 #include "Math/UnrealMathUtility.h"
 #include "UObject/ObjectPtr.h"
 
+class FCameraEvaluationContext;
 class FCameraNodeEvaluator;
 class FCameraNodeEvaluatorStorage;
+class FCameraSystemEvaluator;
 class FReferenceCollector;
-class UCameraEvaluationContext;
 class UCameraRigAsset;
 class UCameraNode;
-class UCameraSystemEvaluator;
 
 /** Allocation information for an entire tree of node evaluators. */
 struct FCameraNodeEvaluatorTreeAllocationInfo
@@ -29,9 +29,9 @@ struct FCameraNodeEvaluatorTreeBuilderParams
 	/** The root node of the tree. */
 	TObjectPtr<const UCameraNode> RootCameraNode;
 	/** The evaluator running this evaluation. */
-	TObjectPtr<UCameraSystemEvaluator> Evaluator;
+	TSharedPtr<FCameraSystemEvaluator> Evaluator;
 	/** The evaluation context (if any) responsible for this branch of the evaluation. */
-	TObjectPtr<const UCameraEvaluationContext> EvaluationContext;
+	TSharedPtr<const FCameraEvaluationContext> EvaluationContext;
 	/** An optional allocation information to optimize storage. */
 	FCameraNodeEvaluatorTreeAllocationInfo* AllocationInfo = nullptr;
 };
diff --git a/Engine/Plugins/Cameras/GameplayCameras/Source/GameplayCameras/Public/Core/CameraSystemEvaluator.h b/Engine/Plugins/Cameras/GameplayCameras/Source/GameplayCameras/Public/Core/CameraSystemEvaluator.h
--- a/Engine/Plugins/Cameras/GameplayCameras/Source/GameplayCameras/Public/Core/CameraSystemEvaluator.h
+++ b/Engine/Plugins/Cameras/GameplayCameras/Source/GameplayCameras/Public/Core/CameraSystemEvaluator.h
@@ -9,17 +9,27 @@
 #include "Core/CameraNodeEvaluatorStorage.h"
 #include "Core/CameraPose.h"
 #include "CoreTypes.h"
+#include "Templates/SharedPointer.h"
 #include "UObject/GCObject.h"
 
-#include "CameraSystemEvaluator.generated.h"
-
+class FCameraEvaluationContext;
 class FRootCameraNodeEvaluator;
 class UCameraDirector;
-class UCameraEvaluationContext;
 class UCameraRigAsset;
 class URootCameraNode;
 struct FMinimalViewInfo;
 
+/**
+ * Parameter structure for initializing a new camera system evaluator.
+ */
+struct FCameraSystemEvaluatorCreateParams
+{
+	TObjectPtr<UObject> Owner;
+
+	using FRootNodeFactory = TFunction<URootCameraNode*()>;
+	FRootNodeFactory RootNodeFactory;
+};
+
 /**
  * Parameter structure for updating the camera system.
  */
@@ -47,21 +57,24 @@ struct FCameraSystemEvaluationUpdateResult
 /**
  * The main camera system evaluator class.
  */
-UCLASS(MinimalAPI)
-class UCameraSystemEvaluator : public UObject
+class FCameraSystemEvaluator : public TSharedFromThis<FCameraSystemEvaluator>
 {
-	GENERATED_BODY()
-
 public:
 
-	UCameraSystemEvaluator(const FObjectInitializer& ObjectInit);
+	/** Builds a new camera system. Initialize must be called before the system is used. */
+	FCameraSystemEvaluator();
+
+	/** Initializes the camera system. */
+	void Initialize(TObjectPtr<UObject> InOwner = nullptr);
+	/** Initializes the camera system. */
+	void Initialize(const FCameraSystemEvaluatorCreateParams& Params);
 
 public:
 
 	/** Push a new evaluation context on the stack. */
-	void PushEvaluationContext(UCameraEvaluationContext* EvaluationContext);
+	void PushEvaluationContext(TSharedRef<FCameraEvaluationContext> EvaluationContext);
 	/** Remove an existing evaluation context from the stack. */
-	void RemoveEvaluationContext(UCameraEvaluationContext* EvaluationContext);
+	void RemoveEvaluationContext(TSharedRef<FCameraEvaluationContext> EvaluationContext);
 	/** Pop the active (top) evaluation context from the stack. */
 	void PopEvaluationContext();
 
@@ -75,12 +88,11 @@ public:
 
 public:
 
-	static void AddReferencedObjects(UObject* InThis, FReferenceCollector& Collector);
+	void AddReferencedObjects(FReferenceCollector& Collector);
 
 private:
 
 	/** The root camera node. */
-	UPROPERTY()
 	TObjectPtr<URootCameraNode> RootNode;
 
 	/** The stack of active evaluation context. */
diff --git a/Engine/Plugins/Cameras/GameplayCameras/Source/GameplayCameras/Public/Core/RootCameraNode.h b/Engine/Plugins/Cameras/GameplayCameras/Source/GameplayCameras/Public/Core/RootCameraNode.h
--- a/Engine/Plugins/Cameras/GameplayCameras/Source/GameplayCameras/Public/Core/RootCameraNode.h
+++ b/Engine/Plugins/Cameras/GameplayCameras/Source/GameplayCameras/Public/Core/RootCameraNode.h
@@ -7,9 +7,9 @@
 
 #include "RootCameraNode.generated.h"
 
-class UCameraEvaluationContext;
+class FCameraEvaluationContext;
+class FCameraSystemEvaluator;
 class UCameraRigAsset;
-class UCameraSystemEvaluator;
 
 /**
  * Defines evaluation layers for camera rigs.
@@ -33,10 +33,10 @@ ENUM_CLASS_FLAGS(ECameraRigLayer)
 struct FActivateCameraRigParams
 {
 	/** The evaluator currently running.*/
-	TObjectPtr<UCameraSystemEvaluator> Evaluator;
+	TSharedPtr<FCameraSystemEvaluator> Evaluator;
 
 	/** The evaluation context in which the camera rig runs. */
-	TObjectPtr<const UCameraEvaluationContext> EvaluationContext;
+	TSharedPtr<const FCameraEvaluationContext> EvaluationContext;
 
 	/** The source camera rig asset that will be instantiated. */
 	TObjectPtr<const UCameraRigAsset> CameraRig;
diff --git a/Engine/Plugins/Cameras/GameplayCameras/Source/GameplayCameras/Public/GameFramework/GameplayCameraComponent.h b/Engine/Plugins/Cameras/GameplayCameras/Source/GameplayCameras/Public/GameFramework/GameplayCameraComponent.h
--- a/Engine/Plugins/Cameras/GameplayCameras/Source/GameplayCameras/Public/GameFramework/GameplayCameraComponent.h
+++ b/Engine/Plugins/Cameras/GameplayCameras/Source/GameplayCameras/Public/GameFramework/GameplayCameraComponent.h
@@ -10,8 +10,8 @@
 
 #include "GameplayCameraComponent.generated.h"
 
+class FGameplayCameraComponentEvaluationContext;
 class UCameraAsset;
-class UGameplayCameraComponentEvaluationContext;
 
 /**
  * A component that can run a camera asset inside its own camera evaluation context.
@@ -56,8 +56,7 @@ public:
 
 protected:
 
-	UPROPERTY(Transient)
-	TObjectPtr<UGameplayCameraComponentEvaluationContext> EvaluationContext;
+	TSharedPtr<FGameplayCameraComponentEvaluationContext> EvaluationContext;
 	
 #if WITH_EDITORONLY_DATA
 
@@ -74,15 +73,15 @@ private:
 	int32 ActivatedForPlayerIndex = INDEX_NONE;
 };
 
-UCLASS(MinimalAPI)
-class UGameplayCameraComponentEvaluationContext : public UCameraEvaluationContext
+/**
+ * Evaluation context for the gameplay camera component.
+ */
+class FGameplayCameraComponentEvaluationContext : public FCameraEvaluationContext
 {
-	GENERATED_BODY()
+	UE_DECLARE_CAMERA_EVALUATION_CONTEXT(FGameplayCameraComponentEvaluationContext)
 
 public:
 
-	UGameplayCameraComponentEvaluationContext(const FObjectInitializer& ObjectInit);
-
 	void Initialize(UGameplayCameraComponent* Owner, APlayerController* InPlayerController);
 	void Update(UGameplayCameraComponent* Owner);
 };
diff --git a/Engine/Plugins/Cameras/GameplayCameras/Source/GameplayCameras/Public/GameFramework/GameplayCameraSystemComponent.h b/Engine/Plugins/Cameras/GameplayCameras/Source/GameplayCameras/Public/GameFramework/GameplayCameraSystemComponent.h
--- a/Engine/Plugins/Cameras/GameplayCameras/Source/GameplayCameras/Public/GameFramework/GameplayCameraSystemComponent.h
+++ b/Engine/Plugins/Cameras/GameplayCameras/Source/GameplayCameras/Public/GameFramework/GameplayCameraSystemComponent.h
@@ -10,7 +10,7 @@
 
 #include "GameplayCameraSystemComponent.generated.h"
 
-class UCameraSystemEvaluator;
+class FCameraSystemEvaluator;
 class UCameraRigAsset;
 struct FMinimalViewInfo;
 
@@ -26,8 +26,7 @@ public:
 
 	UGameplayCameraSystemComponent(const FObjectInitializer& ObjectInit);
 
-	UFUNCTION(BlueprintGetter)
-	UCameraSystemEvaluator* GetCameraSystemEvaluator() { return Evaluator; }
+	TSharedPtr<FCameraSystemEvaluator> GetCameraSystemEvaluator() { return Evaluator; }
 
 	GAMEPLAYCAMERAS_API void GetCameraView(float DeltaTime, FMinimalViewInfo& DesiredView);
 
@@ -42,6 +41,9 @@ public:
 #if WITH_EDITOR
 	virtual bool GetEditorPreviewInfo(float DeltaTime, FMinimalViewInfo& ViewOut) override;
 #endif  // WITH_EDITOR
+	
+	// UObject interface
+	static void AddReferencedObjects(UObject* InThis, FReferenceCollector& Collector);
 
 public:
 
@@ -51,8 +53,7 @@ public:
 
 private:
 	
-	UPROPERTY(Instanced, Transient)
-	TObjectPtr<UCameraSystemEvaluator> Evaluator;
+	TSharedPtr<FCameraSystemEvaluator> Evaluator;
 
 #if WITH_EDITORONLY_DATA
 
