Cameras: always leave a valid blend node on the blend stack entries

This for instance fixes a crash if a transition data is invalid.

diff --git a/Engine/Plugins/Cameras/GameplayCameras/Source/GameplayCameras/Private/Core/BlendStackCameraNode.cpp b/Engine/Plugins/Cameras/GameplayCameras/Source/GameplayCameras/Private/Core/BlendStackCameraNode.cpp
--- a/Engine/Plugins/Cameras/GameplayCameras/Source/GameplayCameras/Private/Core/BlendStackCameraNode.cpp
+++ b/Engine/Plugins/Cameras/GameplayCameras/Source/GameplayCameras/Private/Core/BlendStackCameraNode.cpp
@@ -54,7 +54,7 @@ void FBlendStackCameraNodeEvaluator::Push(const FBlendStackCameraPushParams& Par
 		{
 			ModeBlend = Transition->Blend;
 		}
-		else
+		if (!ModeBlend)
 		{
 			ModeBlend = NewObject<UPopBlendCameraNode>(EntryRootNode, NAME_None);
 		}
@@ -370,7 +370,7 @@ void FBlendStackCameraNodeEvaluator::OnPostBuildAsset(const FGameplayCameraAsset
 
 			// Remove the blend on the root node, since we don't want the reloaded camera rig to re-blend-in
 			// for no good reason.
-			Entry.RootNode->Blend = nullptr;
+			Entry.RootNode->Blend = NewObject<UPopBlendCameraNode>(Entry.RootNode, NAME_None);
 
 			// Rebuild the evaluator tree.
 			FCameraNodeEvaluatorTreeBuilderParams BuildParams;
