Cameras: use ref-counting for live-edit listening of packages being rebuilt in editor

This solves a bug when multiple instances of a camera rig are being blended together. It would cause multiple adding/removing of the blend stack as a listener for that package.

diff --git a/Engine/Plugins/Cameras/GameplayCameras/Source/GameplayCameras/Private/Core/BlendStackCameraNode.cpp b/Engine/Plugins/Cameras/GameplayCameras/Source/GameplayCameras/Private/Core/BlendStackCameraNode.cpp
--- a/Engine/Plugins/Cameras/GameplayCameras/Source/GameplayCameras/Private/Core/BlendStackCameraNode.cpp
+++ b/Engine/Plugins/Cameras/GameplayCameras/Source/GameplayCameras/Private/Core/BlendStackCameraNode.cpp
@@ -82,7 +82,12 @@ void FBlendStackCameraNodeEvaluator::Push(const FBlendStackCameraPushParams& Par
 	Params.CameraRig->GatherPackages(NewEntry.ListenedPackages);
 	for (const UPackage* ListenPackage : NewEntry.ListenedPackages)
 	{
-		LiveEditManager->AddListener(ListenPackage, this);
+		int32& NumListens = AllListenedPackages.FindOrAdd(ListenPackage, 0);
+		if (NumListens == 0)
+		{
+			LiveEditManager->AddListener(ListenPackage, this);
+		}
+		++NumListens;
 	}
 #endif  // WITH_EDITOR
 
@@ -214,7 +219,16 @@ void FBlendStackCameraNodeEvaluator::OnRun(const FCameraNodeEvaluationParams& Pa
 			const FCameraRigEntry& FirstEntry = Entries[0];
 			for (const UPackage* ListenPackage : FirstEntry.ListenedPackages)
 			{
-				LiveEditManager->RemoveListener(ListenPackage, this);
+				int32* NumListens = AllListenedPackages.Find(ListenPackage);
+				if (ensure(NumListens))
+				{
+					--(*NumListens);
+					if (*NumListens == 0)
+					{
+						LiveEditManager->RemoveListener(ListenPackage, this);
+						AllListenedPackages.Remove(ListenPackage);
+					}
+				}
 			}
 #endif  // WITH_EDITOR
 
diff --git a/Engine/Plugins/Cameras/GameplayCameras/Source/GameplayCameras/Public/Core/BlendStackCameraNode.h b/Engine/Plugins/Cameras/GameplayCameras/Source/GameplayCameras/Public/Core/BlendStackCameraNode.h
--- a/Engine/Plugins/Cameras/GameplayCameras/Source/GameplayCameras/Public/Core/BlendStackCameraNode.h
+++ b/Engine/Plugins/Cameras/GameplayCameras/Source/GameplayCameras/Public/Core/BlendStackCameraNode.h
@@ -130,5 +130,9 @@ protected:
 
 	/** Entries in the blend stack. */
 	TArray<FCameraRigEntry> Entries;
+
+#if WITH_EDITOR
+	TMap<const UPackage*, int32> AllListenedPackages;
+#endif  // WITH_EDITOR
 };
 
