Cameras: deactivate camera component before re-activating it

If we are re-activating the camera for a different player, deactivate the previous player's camera.
If we are re-activating the camera for the same player as previously, we let it simply move the component's evaluation context to the top of the camera system stack.

diff --git a/Engine/Plugins/Cameras/GameplayCameras/Source/GameplayCameras/Private/GameFramework/GameplayCameraComponent.cpp b/Engine/Plugins/Cameras/GameplayCameras/Source/GameplayCameras/Private/GameFramework/GameplayCameraComponent.cpp
--- a/Engine/Plugins/Cameras/GameplayCameras/Source/GameplayCameras/Private/GameFramework/GameplayCameraComponent.cpp
+++ b/Engine/Plugins/Cameras/GameplayCameras/Source/GameplayCameras/Private/GameFramework/GameplayCameraComponent.cpp
@@ -34,9 +34,9 @@ UGameplayCameraComponent::UGameplayCameraComponent(const FObjectInitializer& Obj
 
 void UGameplayCameraComponent::ActivateCamera(int32 PlayerIndex)
 {
-	if (ActivatedForPlayerIndex >= 0)
+	if (ActivatedForPlayerIndex >= 0 && ActivatedForPlayerIndex != PlayerIndex)
 	{
-		return;
+		DeactivateCamera();
 	}
 
 	UWorld* World = GetWorld();
