Cameras: move all non-data classes to the UE::Cameras namespace

diff --git a/Engine/Plugins/Cameras/GameplayCameras/Source/GameplayCameras/Private/Core/BlendCameraNode.cpp b/Engine/Plugins/Cameras/GameplayCameras/Source/GameplayCameras/Private/Core/BlendCameraNode.cpp
--- a/Engine/Plugins/Cameras/GameplayCameras/Source/GameplayCameras/Private/Core/BlendCameraNode.cpp
+++ b/Engine/Plugins/Cameras/GameplayCameras/Source/GameplayCameras/Private/Core/BlendCameraNode.cpp
@@ -4,6 +4,9 @@
 
 #include UE_INLINE_GENERATED_CPP_BY_NAME(BlendCameraNode)
 
+namespace UE::Cameras
+{
+
 UE_DEFINE_CAMERA_NODE_EVALUATOR(FBlendCameraNodeEvaluator)
 
 void FBlendCameraNodeEvaluator::BlendResults(const FCameraNodeBlendParams& Params, FCameraNodeBlendResult& OutResult)
@@ -11,3 +14,5 @@ void FBlendCameraNodeEvaluator::BlendResults(const FCameraNodeBlendParams& Param
 	OnBlendResults(Params, OutResult);
 }
 
+}  // namespace UE::Cameras
+
diff --git a/Engine/Plugins/Cameras/GameplayCameras/Source/GameplayCameras/Private/Core/BlendStackCameraNode.cpp b/Engine/Plugins/Cameras/GameplayCameras/Source/GameplayCameras/Private/Core/BlendStackCameraNode.cpp
--- a/Engine/Plugins/Cameras/GameplayCameras/Source/GameplayCameras/Private/Core/BlendStackCameraNode.cpp
+++ b/Engine/Plugins/Cameras/GameplayCameras/Source/GameplayCameras/Private/Core/BlendStackCameraNode.cpp
@@ -17,9 +17,13 @@
 
 FCameraNodeEvaluatorPtr UBlendStackCameraNode::OnBuildEvaluator(FCameraNodeEvaluatorBuilder& Builder) const
 {
+	using namespace UE::Cameras;
 	return Builder.BuildEvaluator<FBlendStackCameraNodeEvaluator>();
 }
 
+namespace UE::Cameras
+{
+
 UE_DEFINE_CAMERA_NODE_EVALUATOR(FBlendStackCameraNodeEvaluator)
 
 void FBlendStackCameraNodeEvaluator::Push(const FBlendStackCameraPushParams& Params)
@@ -403,3 +407,5 @@ void FBlendStackCameraNodeEvaluator::OnPostBuildAsset(const FGameplayCameraAsset
 
 #endif  // WITH_EDITOR
 
+}  // namespace UE::Cameras
+
diff --git a/Engine/Plugins/Cameras/GameplayCameras/Source/GameplayCameras/Private/Core/BlendStackRootCameraNode.cpp b/Engine/Plugins/Cameras/GameplayCameras/Source/GameplayCameras/Private/Core/BlendStackRootCameraNode.cpp
--- a/Engine/Plugins/Cameras/GameplayCameras/Source/GameplayCameras/Private/Core/BlendStackRootCameraNode.cpp
+++ b/Engine/Plugins/Cameras/GameplayCameras/Source/GameplayCameras/Private/Core/BlendStackRootCameraNode.cpp
@@ -23,9 +23,13 @@ FCameraNodeChildrenView UBlendStackRootCameraNode::OnGetChildren()
 
 FCameraNodeEvaluatorPtr UBlendStackRootCameraNode::OnBuildEvaluator(FCameraNodeEvaluatorBuilder& Builder) const
 {
+	using namespace UE::Cameras;
 	return Builder.BuildEvaluator<FBlendStackRootCameraNodeEvaluator>();
 }
 
+namespace UE::Cameras
+{
+
 UE_DEFINE_CAMERA_NODE_EVALUATOR(FBlendStackRootCameraNodeEvaluator)
 
 FCameraNodeEvaluatorChildrenView FBlendStackRootCameraNodeEvaluator::OnGetChildren()
@@ -61,3 +65,5 @@ void FBlendStackRootCameraNodeEvaluator::OnRun(const FCameraNodeEvaluationParams
 	}
 }
 
+}  // namespace UE::Cameras
+
diff --git a/Engine/Plugins/Cameras/GameplayCameras/Source/GameplayCameras/Private/Core/BlendStackRootCameraNode.h b/Engine/Plugins/Cameras/GameplayCameras/Source/GameplayCameras/Private/Core/BlendStackRootCameraNode.h
--- a/Engine/Plugins/Cameras/GameplayCameras/Source/GameplayCameras/Private/Core/BlendStackRootCameraNode.h
+++ b/Engine/Plugins/Cameras/GameplayCameras/Source/GameplayCameras/Private/Core/BlendStackRootCameraNode.h
@@ -7,7 +7,6 @@
 
 #include "BlendStackRootCameraNode.generated.h"
 
-class FBlendCameraNodeEvaluator;
 class UBlendCameraNode;
 class UCameraRigAsset;
 
@@ -38,6 +37,11 @@ public:
 	TObjectPtr<UCameraNode> RootNode;
 };
 
+namespace UE::Cameras
+{
+
+class FBlendCameraNodeEvaluator;
+
 /**
  * Evaluator for the blend stack entry root node.
  */
@@ -63,3 +67,5 @@ private:
 	FCameraNodeEvaluator* RootEvaluator = nullptr;
 };
 
+}  // namespace UE::Cameras
+
diff --git a/Engine/Plugins/Cameras/GameplayCameras/Source/GameplayCameras/Private/Core/CameraDirector.cpp b/Engine/Plugins/Cameras/GameplayCameras/Source/GameplayCameras/Private/Core/CameraDirector.cpp
--- a/Engine/Plugins/Cameras/GameplayCameras/Source/GameplayCameras/Private/Core/CameraDirector.cpp
+++ b/Engine/Plugins/Cameras/GameplayCameras/Source/GameplayCameras/Private/Core/CameraDirector.cpp
@@ -4,8 +4,9 @@
 
 #include UE_INLINE_GENERATED_CPP_BY_NAME(CameraDirector)
 
-FCameraDirectorEvaluator* UCameraDirector::BuildEvaluator(FCameraDirectorEvaluatorBuilder& Builder) const
+FCameraDirectorEvaluatorPtr UCameraDirector::BuildEvaluator(FCameraDirectorEvaluatorBuilder& Builder) const
 {
+	using namespace UE::Cameras;
 	FCameraDirectorEvaluator* NewEvaluator = OnBuildEvaluator(Builder);
 	NewEvaluator->SetPrivateCameraDirector(this);
 	return NewEvaluator;
diff --git a/Engine/Plugins/Cameras/GameplayCameras/Source/GameplayCameras/Private/Core/CameraDirectorEvaluator.cpp b/Engine/Plugins/Cameras/GameplayCameras/Source/GameplayCameras/Private/Core/CameraDirectorEvaluator.cpp
--- a/Engine/Plugins/Cameras/GameplayCameras/Source/GameplayCameras/Private/Core/CameraDirectorEvaluator.cpp
+++ b/Engine/Plugins/Cameras/GameplayCameras/Source/GameplayCameras/Private/Core/CameraDirectorEvaluator.cpp
@@ -4,6 +4,9 @@
 
 #include "Core/CameraDirector.h"
 
+namespace UE::Cameras
+{
+
 UE_GAMEPLAY_CAMERAS_DEFINE_RTTI(FCameraDirectorEvaluator)
 
 FCameraDirectorEvaluator::FCameraDirectorEvaluator()
@@ -20,3 +23,5 @@ void FCameraDirectorEvaluator::SetPrivateCameraDirector(TObjectPtr<const UCamera
 	PrivateCameraDirector = InCameraDirector;
 }
 
+}  // namespace UE::Cameras
+
diff --git a/Engine/Plugins/Cameras/GameplayCameras/Source/GameplayCameras/Private/Core/CameraEvaluationContext.cpp b/Engine/Plugins/Cameras/GameplayCameras/Source/GameplayCameras/Private/Core/CameraEvaluationContext.cpp
--- a/Engine/Plugins/Cameras/GameplayCameras/Source/GameplayCameras/Private/Core/CameraEvaluationContext.cpp
+++ b/Engine/Plugins/Cameras/GameplayCameras/Source/GameplayCameras/Private/Core/CameraEvaluationContext.cpp
@@ -2,6 +2,9 @@
 
 #include "Core/CameraEvaluationContext.h"
 
+namespace UE::Cameras
+{
+
 UE_GAMEPLAY_CAMERAS_DEFINE_RTTI(FCameraEvaluationContext)
 
 FCameraEvaluationContext::FCameraEvaluationContext()
@@ -12,3 +15,5 @@ FCameraEvaluationContext::~FCameraEvaluationContext()
 {
 }
 
+}  // namespace UE::Cameras
+
diff --git a/Engine/Plugins/Cameras/GameplayCameras/Source/GameplayCameras/Private/Core/CameraEvaluationContextStack.cpp b/Engine/Plugins/Cameras/GameplayCameras/Source/GameplayCameras/Private/Core/CameraEvaluationContextStack.cpp
--- a/Engine/Plugins/Cameras/GameplayCameras/Source/GameplayCameras/Private/Core/CameraEvaluationContextStack.cpp
+++ b/Engine/Plugins/Cameras/GameplayCameras/Source/GameplayCameras/Private/Core/CameraEvaluationContextStack.cpp
@@ -9,6 +9,9 @@
 #include "Core/CameraSystemEvaluator.h"
 #include "UObject/Package.h"
 
+namespace UE::Cameras
+{
+
 FCameraEvaluationContextInfo FCameraEvaluationContextStack::GetActiveContext() const
 {
 	for (const FContextEntry& Entry : ReverseIterate(Entries))
@@ -89,3 +92,5 @@ void FCameraEvaluationContextStack::AddReferencedObjects(FReferenceCollector& Co
 	}
 }
 
+}  // namespace UE::Cameras
+
diff --git a/Engine/Plugins/Cameras/GameplayCameras/Source/GameplayCameras/Private/Core/CameraNode.cpp b/Engine/Plugins/Cameras/GameplayCameras/Source/GameplayCameras/Private/Core/CameraNode.cpp
--- a/Engine/Plugins/Cameras/GameplayCameras/Source/GameplayCameras/Private/Core/CameraNode.cpp
+++ b/Engine/Plugins/Cameras/GameplayCameras/Source/GameplayCameras/Private/Core/CameraNode.cpp
@@ -6,6 +6,9 @@
 
 #include UE_INLINE_GENERATED_CPP_BY_NAME(CameraNode)
 
+namespace UE::Cameras
+{
+
 void FCameraNodeEvaluationResult::Reset()
 {
 	CameraPose.Reset();
@@ -13,6 +16,8 @@ void FCameraNodeEvaluationResult::Reset()
 	bIsValid = false;
 }
 
+}  // namespace UE::Cameras
+
 FCameraNodeChildrenView UCameraNode::GetChildren()
 {
 	return OnGetChildren();
@@ -25,6 +30,7 @@ FCameraNodeAllocationInfo UCameraNode::GetAllocationInfo() const
 
 FCameraNodeEvaluatorPtr UCameraNode::BuildEvaluator(FCameraNodeEvaluatorBuilder& Builder) const
 {
+	using namespace UE::Cameras;
 	FCameraNodeEvaluator* NewEvaluator = OnBuildEvaluator(Builder);
 	NewEvaluator->SetPrivateCameraNode(this);
 	return NewEvaluator;
diff --git a/Engine/Plugins/Cameras/GameplayCameras/Source/GameplayCameras/Private/Core/CameraNodeEvaluator.cpp b/Engine/Plugins/Cameras/GameplayCameras/Source/GameplayCameras/Private/Core/CameraNodeEvaluator.cpp
--- a/Engine/Plugins/Cameras/GameplayCameras/Source/GameplayCameras/Private/Core/CameraNodeEvaluator.cpp
+++ b/Engine/Plugins/Cameras/GameplayCameras/Source/GameplayCameras/Private/Core/CameraNodeEvaluator.cpp
@@ -5,6 +5,9 @@
 #include "Core/CameraNode.h"
 #include "UObject/UObjectGlobals.h"
 
+namespace UE::Cameras
+{
+
 UE_GAMEPLAY_CAMERAS_DEFINE_RTTI(FCameraNodeEvaluator)
 
 FCameraNodeEvaluator* FCameraNodeEvaluatorInitializeParams::BuildEvaluator(const UCameraNode* InNode) const
@@ -59,3 +62,5 @@ void FCameraNodeEvaluator::Run(const FCameraNodeEvaluationParams& Params, FCamer
 	}
 }
 
+}  // namespace UE::Cameras
+
diff --git a/Engine/Plugins/Cameras/GameplayCameras/Source/GameplayCameras/Private/Core/CameraNodeEvaluatorStorage.cpp b/Engine/Plugins/Cameras/GameplayCameras/Source/GameplayCameras/Private/Core/CameraNodeEvaluatorStorage.cpp
--- a/Engine/Plugins/Cameras/GameplayCameras/Source/GameplayCameras/Private/Core/CameraNodeEvaluatorStorage.cpp
+++ b/Engine/Plugins/Cameras/GameplayCameras/Source/GameplayCameras/Private/Core/CameraNodeEvaluatorStorage.cpp
@@ -7,6 +7,9 @@
 #include "Core/CameraRigAsset.h"
 #include "UObject/Package.h"
 
+namespace UE::Cameras
+{
+
 FCameraNodeEvaluatorStorage::FCameraNodeEvaluatorStorage()
 {
 }
@@ -105,3 +108,5 @@ void FCameraNodeEvaluatorStorage::DestroyEvaluatorTree(bool bFreeAllocations)
 	}
 }
 
+}  // namespace UE::Cameras
+
diff --git a/Engine/Plugins/Cameras/GameplayCameras/Source/GameplayCameras/Private/Core/CameraObjectRtti.cpp b/Engine/Plugins/Cameras/GameplayCameras/Source/GameplayCameras/Private/Core/CameraObjectRtti.cpp
--- a/Engine/Plugins/Cameras/GameplayCameras/Source/GameplayCameras/Private/Core/CameraObjectRtti.cpp
+++ b/Engine/Plugins/Cameras/GameplayCameras/Source/GameplayCameras/Private/Core/CameraObjectRtti.cpp
@@ -2,9 +2,14 @@
 
 #include "Core/CameraObjectRtti.h"
 
+namespace UE::Cameras
+{
+
 uint32 FCameraObjectTypeID::RegisterNewID()
 {
 	static uint32 ID = 0;
 	return ID++;
 }
 
+}  // namespace UE::Cameras
+
diff --git a/Engine/Plugins/Cameras/GameplayCameras/Source/GameplayCameras/Private/Core/CameraSystemEvaluator.cpp b/Engine/Plugins/Cameras/GameplayCameras/Source/GameplayCameras/Private/Core/CameraSystemEvaluator.cpp
--- a/Engine/Plugins/Cameras/GameplayCameras/Source/GameplayCameras/Private/Core/CameraSystemEvaluator.cpp
+++ b/Engine/Plugins/Cameras/GameplayCameras/Source/GameplayCameras/Private/Core/CameraSystemEvaluator.cpp
@@ -15,6 +15,9 @@
 
 DECLARE_CYCLE_STAT(TEXT("Camera System Eval"), CameraSystemEval_Total, STATGROUP_CameraSystem);
 
+namespace UE::Cameras
+{
+
 FCameraSystemEvaluator::FCameraSystemEvaluator()
 {
 }
@@ -136,3 +139,5 @@ void FCameraSystemEvaluator::GetEvaluatedCameraView(FMinimalViewInfo& DesiredVie
 	DesiredView.FOV = CameraPose.GetEffectiveFieldOfView();
 }
 
+}  // namespace UE::Cameras
+
diff --git a/Engine/Plugins/Cameras/GameplayCameras/Source/GameplayCameras/Private/Core/CameraVariableTable.cpp b/Engine/Plugins/Cameras/GameplayCameras/Source/GameplayCameras/Private/Core/CameraVariableTable.cpp
--- a/Engine/Plugins/Cameras/GameplayCameras/Source/GameplayCameras/Private/Core/CameraVariableTable.cpp
+++ b/Engine/Plugins/Cameras/GameplayCameras/Source/GameplayCameras/Private/Core/CameraVariableTable.cpp
@@ -5,7 +5,10 @@
 #include "HAL/UnrealMemory.h"
 #include "Math/UnrealMath.h"
 
-namespace UE::Cameras::Private
+namespace UE::Cameras
+{
+
+namespace Private
 {
 
 static const FString GUnavailableVariableDebugName(TEXT("<no debug info>"));
@@ -20,7 +23,7 @@ bool IsVariableInMask(uint32 VariableId, const FCameraVariableTableFlags* InMask
 	return true;
 }
 
-}  // namespace UE::Cameras::Private
+}  // namespace Private
 
 FCameraVariableTable::FCameraVariableTable()
 {
@@ -389,3 +392,5 @@ UE_CAMERA_VARIABLE_FOR_ALL_TYPES()
 	}
 }
 
+}  // namespace UE::Cameras
+
diff --git a/Engine/Plugins/Cameras/GameplayCameras/Source/GameplayCameras/Private/Core/DefaultRootCameraNode.cpp b/Engine/Plugins/Cameras/GameplayCameras/Source/GameplayCameras/Private/Core/DefaultRootCameraNode.cpp
--- a/Engine/Plugins/Cameras/GameplayCameras/Source/GameplayCameras/Private/Core/DefaultRootCameraNode.cpp
+++ b/Engine/Plugins/Cameras/GameplayCameras/Source/GameplayCameras/Private/Core/DefaultRootCameraNode.cpp
@@ -36,9 +36,13 @@ UDefaultRootCameraNode::UDefaultRootCameraNode(const FObjectInitializer& ObjectI
 
 FCameraNodeEvaluatorPtr UDefaultRootCameraNode::OnBuildEvaluator(FCameraNodeEvaluatorBuilder& Builder) const
 {
+	using namespace UE::Cameras;
 	return Builder.BuildEvaluator<FDefaultRootCameraNodeEvaluator>();
 }
 
+namespace UE::Cameras
+{
+
 UE_DEFINE_CAMERA_NODE_EVALUATOR(FDefaultRootCameraNodeEvaluator)
 
 void FDefaultRootCameraNodeEvaluator::OnInitialize(const FCameraNodeEvaluatorInitializeParams& Params)
@@ -93,3 +97,5 @@ void FDefaultRootCameraNodeEvaluator::OnActivateCameraRig(const FActivateCameraR
 	}
 }
 
+}  // namespace UE::Cameras
+
diff --git a/Engine/Plugins/Cameras/GameplayCameras/Source/GameplayCameras/Private/Core/RootCameraNode.cpp b/Engine/Plugins/Cameras/GameplayCameras/Source/GameplayCameras/Private/Core/RootCameraNode.cpp
--- a/Engine/Plugins/Cameras/GameplayCameras/Source/GameplayCameras/Private/Core/RootCameraNode.cpp
+++ b/Engine/Plugins/Cameras/GameplayCameras/Source/GameplayCameras/Private/Core/RootCameraNode.cpp
@@ -6,8 +6,13 @@
 
 #include UE_INLINE_GENERATED_CPP_BY_NAME(RootCameraNode)
 
+namespace UE::Cameras
+{
+
 void FRootCameraNodeEvaluator::ActivateCameraRig(const FActivateCameraRigParams& Params)
 {
 	OnActivateCameraRig(Params);
 }
 
+}  // namespace UE::Cameras
+
diff --git a/Engine/Plugins/Cameras/GameplayCameras/Source/GameplayCameras/Private/Directors/SingleCameraDirector.cpp b/Engine/Plugins/Cameras/GameplayCameras/Source/GameplayCameras/Private/Directors/SingleCameraDirector.cpp
--- a/Engine/Plugins/Cameras/GameplayCameras/Source/GameplayCameras/Private/Directors/SingleCameraDirector.cpp
+++ b/Engine/Plugins/Cameras/GameplayCameras/Source/GameplayCameras/Private/Directors/SingleCameraDirector.cpp
@@ -4,6 +4,9 @@
 
 #include UE_INLINE_GENERATED_CPP_BY_NAME(SingleCameraDirector)
 
+namespace UE::Cameras
+{
+
 class FSingleCameraDirectorEvaluator : public FCameraDirectorEvaluator
 {
 	UE_DECLARE_CAMERA_DIRECTOR_EVALUATOR(FSingleCameraDirectorEvaluator)
@@ -20,13 +23,16 @@ protected:
 
 UE_DEFINE_CAMERA_DIRECTOR_EVALUATOR(FSingleCameraDirectorEvaluator)
 
+}  // namespace UE::Cameras
+
 USingleCameraDirector::USingleCameraDirector(const FObjectInitializer& ObjectInit)
 	: Super(ObjectInit)
 {
 }
 
-FCameraDirectorEvaluator* USingleCameraDirector::OnBuildEvaluator(FCameraDirectorEvaluatorBuilder& Builder) const
+FCameraDirectorEvaluatorPtr USingleCameraDirector::OnBuildEvaluator(FCameraDirectorEvaluatorBuilder& Builder) const
 {
+	using namespace UE::Cameras;
 	return Builder.BuildEvaluator<FSingleCameraDirectorEvaluator>();
 }
 
diff --git a/Engine/Plugins/Cameras/GameplayCameras/Source/GameplayCameras/Private/GameFramework/GameplayCameraComponent.cpp b/Engine/Plugins/Cameras/GameplayCameras/Source/GameplayCameras/Private/GameFramework/GameplayCameraComponent.cpp
--- a/Engine/Plugins/Cameras/GameplayCameras/Source/GameplayCameras/Private/GameFramework/GameplayCameraComponent.cpp
+++ b/Engine/Plugins/Cameras/GameplayCameras/Source/GameplayCameras/Private/GameFramework/GameplayCameraComponent.cpp
@@ -86,6 +86,8 @@ void UGameplayCameraComponent::DeactivateCamera()
 
 void UGameplayCameraComponent::ActivateCamera(APlayerController* PlayerController)
 {
+	using namespace UE::Cameras;
+
 	if (!ensure(PlayerController && PlayerController->PlayerCameraManager))
 	{
 		return;
@@ -111,6 +113,8 @@ void UGameplayCameraComponent::ActivateCamera(APlayerController* PlayerControlle
 
 void UGameplayCameraComponent::DeactivateCamera(APlayerController* PlayerController)
 {
+	using namespace UE::Cameras;
+
 	if (!ensure(PlayerController && PlayerController->PlayerCameraManager))
 	{
 		return;
@@ -190,6 +194,9 @@ void UGameplayCameraComponent::UpdatePreviewMeshTransform()
 
 #endif
 
+namespace UE::Cameras
+{
+
 UE_DEFINE_CAMERA_EVALUATION_CONTEXT(FGameplayCameraComponentEvaluationContext)
 
 void FGameplayCameraComponentEvaluationContext::Initialize(UGameplayCameraComponent* Owner, APlayerController* InPlayerController)
@@ -205,5 +212,7 @@ void FGameplayCameraComponentEvaluationContext::Update(UGameplayCameraComponent*
 	InitialResult.bIsValid = true;
 }
 
+}  // namespace UE::Cameras
+
 #undef LOCTEXT_NAMESPACE
 
diff --git a/Engine/Plugins/Cameras/GameplayCameras/Source/GameplayCameras/Private/GameFramework/GameplayCameraSystemComponent.cpp b/Engine/Plugins/Cameras/GameplayCameras/Source/GameplayCameras/Private/GameFramework/GameplayCameraSystemComponent.cpp
--- a/Engine/Plugins/Cameras/GameplayCameras/Source/GameplayCameras/Private/GameFramework/GameplayCameraSystemComponent.cpp
+++ b/Engine/Plugins/Cameras/GameplayCameras/Source/GameplayCameras/Private/GameFramework/GameplayCameraSystemComponent.cpp
@@ -32,6 +32,8 @@ UGameplayCameraSystemComponent::UGameplayCameraSystemComponent(const FObjectInit
 
 void UGameplayCameraSystemComponent::GetCameraView(float DeltaTime, FMinimalViewInfo& DesiredView)
 {
+	using namespace UE::Cameras;
+
 	FCameraSystemEvaluationUpdateParams UpdateParams;
 	UpdateParams.DeltaTime = DeltaTime;
 	Evaluator->Update(UpdateParams);
@@ -41,6 +43,8 @@ void UGameplayCameraSystemComponent::GetCameraView(float DeltaTime, FMinimalView
 
 void UGameplayCameraSystemComponent::OnRegister()
 {
+	using namespace UE::Cameras;
+
 	if (!Evaluator.IsValid())
 	{
 		Evaluator = MakeShared<FCameraSystemEvaluator>();
diff --git a/Engine/Plugins/Cameras/GameplayCameras/Source/GameplayCameras/Private/Math/CriticalDamper.cpp b/Engine/Plugins/Cameras/GameplayCameras/Source/GameplayCameras/Private/Math/CriticalDamper.cpp
--- a/Engine/Plugins/Cameras/GameplayCameras/Source/GameplayCameras/Private/Math/CriticalDamper.cpp
+++ b/Engine/Plugins/Cameras/GameplayCameras/Source/GameplayCameras/Private/Math/CriticalDamper.cpp
@@ -5,8 +5,6 @@
 #include "HAL/IConsoleManager.h"
 #include "Math/UnrealMathUtility.h"
 
-#include UE_INLINE_GENERATED_CPP_BY_NAME(CriticalDamper)
-
 float GGameplayCamerasCriticalDamperStabilizationThreshold = 0.00001f;
 static FAutoConsoleVariableRef CVarGameplayCamerasCriticalDamperStabilizationThreshold(
 	TEXT("GameplayCameras.CriticalDamper.StabilizationThreshold"),
@@ -14,6 +12,9 @@ static FAutoConsoleVariableRef CVarGameplayCamerasCriticalDamperStabilizationThr
 	TEXT("(Default: 0.00001f (in meters and meters/sec). Determines when to clamp a critical damper to 0 to stabilize it.")
 	);
 
+namespace UE::Cameras
+{
+
 FCriticalDamper::FCriticalDamper()
 {
 	Reset(0, 0);
@@ -247,3 +248,5 @@ void FCriticalDamper::InternalUpdate(float ForcedMovement, float DeltaTime)
 	}
 }
 
+}  // namespace UE::Cameras
+
diff --git a/Engine/Plugins/Cameras/GameplayCameras/Source/GameplayCameras/Private/Nodes/Blends/LinearBlendCameraNode.cpp b/Engine/Plugins/Cameras/GameplayCameras/Source/GameplayCameras/Private/Nodes/Blends/LinearBlendCameraNode.cpp
--- a/Engine/Plugins/Cameras/GameplayCameras/Source/GameplayCameras/Private/Nodes/Blends/LinearBlendCameraNode.cpp
+++ b/Engine/Plugins/Cameras/GameplayCameras/Source/GameplayCameras/Private/Nodes/Blends/LinearBlendCameraNode.cpp
@@ -4,6 +4,9 @@
 
 #include UE_INLINE_GENERATED_CPP_BY_NAME(LinearBlendCameraNode)
 
+namespace UE::Cameras
+{
+
 class FLinearBlendCameraNodeEvaluator : public FSimpleFixedTimeBlendCameraNodeEvaluator
 {
 	UE_DECLARE_BLEND_CAMERA_NODE_EVALUATOR_EX(FLinearBlendCameraNodeEvaluator, FSimpleFixedTimeBlendCameraNodeEvaluator)
@@ -19,8 +22,11 @@ void FLinearBlendCameraNodeEvaluator::OnComputeBlendFactor(const FCameraNodeEval
 	OutResult.BlendFactor = FMath::Lerp(0.f, 1.f, GetTimeFactor());
 }
 
+}  // namespace UE::Cameras
+
 FCameraNodeEvaluatorPtr ULinearBlendCameraNode::OnBuildEvaluator(FCameraNodeEvaluatorBuilder& Builder) const
 {
+	using namespace UE::Cameras;
 	return Builder.BuildEvaluator<FLinearBlendCameraNodeEvaluator>();
 }
 
diff --git a/Engine/Plugins/Cameras/GameplayCameras/Source/GameplayCameras/Private/Nodes/Blends/PopBlendCameraNode.cpp b/Engine/Plugins/Cameras/GameplayCameras/Source/GameplayCameras/Private/Nodes/Blends/PopBlendCameraNode.cpp
--- a/Engine/Plugins/Cameras/GameplayCameras/Source/GameplayCameras/Private/Nodes/Blends/PopBlendCameraNode.cpp
+++ b/Engine/Plugins/Cameras/GameplayCameras/Source/GameplayCameras/Private/Nodes/Blends/PopBlendCameraNode.cpp
@@ -4,6 +4,9 @@
 
 #include UE_INLINE_GENERATED_CPP_BY_NAME(PopBlendCameraNode)
 
+namespace UE::Cameras
+{
+
 class FPopBlendCameraNodeEvaluator : public FBlendCameraNodeEvaluator
 {
 	UE_DECLARE_BLEND_CAMERA_NODE_EVALUATOR(FPopBlendCameraNodeEvaluator)
@@ -33,8 +36,11 @@ void FPopBlendCameraNodeEvaluator::OnBlendResults(const FCameraNodeBlendParams&
 	OutResult.bIsBlendFinished = true;
 }
 
+}  // namespace UE::Cameras
+
 FCameraNodeEvaluatorPtr UPopBlendCameraNode::OnBuildEvaluator(FCameraNodeEvaluatorBuilder& Builder) const
 {
+	using namespace UE::Cameras;
 	return Builder.BuildEvaluator<FPopBlendCameraNodeEvaluator>();
 }
 
diff --git a/Engine/Plugins/Cameras/GameplayCameras/Source/GameplayCameras/Private/Nodes/Blends/SimpleBlendCameraNode.cpp b/Engine/Plugins/Cameras/GameplayCameras/Source/GameplayCameras/Private/Nodes/Blends/SimpleBlendCameraNode.cpp
--- a/Engine/Plugins/Cameras/GameplayCameras/Source/GameplayCameras/Private/Nodes/Blends/SimpleBlendCameraNode.cpp
+++ b/Engine/Plugins/Cameras/GameplayCameras/Source/GameplayCameras/Private/Nodes/Blends/SimpleBlendCameraNode.cpp
@@ -4,6 +4,9 @@
 
 #include UE_INLINE_GENERATED_CPP_BY_NAME(SimpleBlendCameraNode)
 
+namespace UE::Cameras
+{
+
 UE_DEFINE_BLEND_CAMERA_NODE_EVALUATOR(FSimpleBlendCameraNodeEvaluator)
 
 void FSimpleBlendCameraNodeEvaluator::OnRun(const FCameraNodeEvaluationParams& Params, FCameraNodeEvaluationResult& OutResult)
@@ -59,3 +62,5 @@ float FSimpleFixedTimeBlendCameraNodeEvaluator::GetTimeFactor() const
 	return CurrentTime / BlendNode->BlendTime;
 }
 
+}  // namespace UE::Cameras
+
diff --git a/Engine/Plugins/Cameras/GameplayCameras/Source/GameplayCameras/Private/Nodes/Blends/SmoothBlendCameraNode.cpp b/Engine/Plugins/Cameras/GameplayCameras/Source/GameplayCameras/Private/Nodes/Blends/SmoothBlendCameraNode.cpp
--- a/Engine/Plugins/Cameras/GameplayCameras/Source/GameplayCameras/Private/Nodes/Blends/SmoothBlendCameraNode.cpp
+++ b/Engine/Plugins/Cameras/GameplayCameras/Source/GameplayCameras/Private/Nodes/Blends/SmoothBlendCameraNode.cpp
@@ -6,6 +6,9 @@
 
 #include UE_INLINE_GENERATED_CPP_BY_NAME(SmoothBlendCameraNode)
 
+namespace UE::Cameras
+{
+
 class FSmoothBlendCameraNodeEvaluator : public FSimpleFixedTimeBlendCameraNodeEvaluator
 {
 	UE_DECLARE_BLEND_CAMERA_NODE_EVALUATOR_EX(FSmoothBlendCameraNodeEvaluator, FSimpleFixedTimeBlendCameraNodeEvaluator)
@@ -35,8 +38,11 @@ void FSmoothBlendCameraNodeEvaluator::OnComputeBlendFactor(const FCameraNodeEval
 	}
 }
 
+}  // namespace UE::Cameras
+
 FCameraNodeEvaluatorPtr USmoothBlendCameraNode::OnBuildEvaluator(FCameraNodeEvaluatorBuilder& Builder) const
 {
+	using namespace UE::Cameras;
 	return Builder.BuildEvaluator<FSmoothBlendCameraNodeEvaluator>();
 }
 
diff --git a/Engine/Plugins/Cameras/GameplayCameras/Source/GameplayCameras/Private/Nodes/Common/ArrayCameraNode.cpp b/Engine/Plugins/Cameras/GameplayCameras/Source/GameplayCameras/Private/Nodes/Common/ArrayCameraNode.cpp
--- a/Engine/Plugins/Cameras/GameplayCameras/Source/GameplayCameras/Private/Nodes/Common/ArrayCameraNode.cpp
+++ b/Engine/Plugins/Cameras/GameplayCameras/Source/GameplayCameras/Private/Nodes/Common/ArrayCameraNode.cpp
@@ -6,6 +6,9 @@
 
 #include UE_INLINE_GENERATED_CPP_BY_NAME(ArrayCameraNode)
 
+namespace UE::Cameras
+{
+
 class FArrayCameraNodeEvaluator : public FCameraNodeEvaluator
 {
 	UE_DECLARE_CAMERA_NODE_EVALUATOR(FArrayCameraNodeEvaluator)
@@ -52,6 +55,8 @@ void FArrayCameraNodeEvaluator::OnRun(const FCameraNodeEvaluationParams& Params,
 	}
 }
 
+}  // namespace UE::Cameras
+
 FCameraNodeChildrenView UArrayCameraNode::OnGetChildren()
 {
 	return FCameraNodeChildrenView(Children);
@@ -59,6 +64,7 @@ FCameraNodeChildrenView UArrayCameraNode::OnGetChildren()
 
 FCameraNodeEvaluatorPtr UArrayCameraNode::OnBuildEvaluator(FCameraNodeEvaluatorBuilder& Builder) const
 {
+	using namespace UE::Cameras;
 	return Builder.BuildEvaluator<FArrayCameraNodeEvaluator>();
 }
 
diff --git a/Engine/Plugins/Cameras/GameplayCameras/Source/GameplayCameras/Private/Nodes/Common/BoomArmCameraNode.cpp b/Engine/Plugins/Cameras/GameplayCameras/Source/GameplayCameras/Private/Nodes/Common/BoomArmCameraNode.cpp
--- a/Engine/Plugins/Cameras/GameplayCameras/Source/GameplayCameras/Private/Nodes/Common/BoomArmCameraNode.cpp
+++ b/Engine/Plugins/Cameras/GameplayCameras/Source/GameplayCameras/Private/Nodes/Common/BoomArmCameraNode.cpp
@@ -9,6 +9,9 @@
 
 #include UE_INLINE_GENERATED_CPP_BY_NAME(BoomArmCameraNode)
 
+namespace UE::Cameras
+{
+
 class FBoomArmCameraNodeEvaluator : public FCameraNodeEvaluator
 {
 	UE_DECLARE_CAMERA_NODE_EVALUATOR(FBoomArmCameraNodeEvaluator)
@@ -52,8 +55,11 @@ void FBoomArmCameraNodeEvaluator::OnRun(const FCameraNodeEvaluationParams& Param
 	OutResult.CameraPose.SetTransform(FinalTransform);
 }
 
+}  // namespace UE::Cameras
+
 FCameraNodeEvaluatorPtr UBoomArmCameraNode::OnBuildEvaluator(FCameraNodeEvaluatorBuilder& Builder) const
 {
+	using namespace UE::Cameras;
 	return Builder.BuildEvaluator<FBoomArmCameraNodeEvaluator>();
 }
 
diff --git a/Engine/Plugins/Cameras/GameplayCameras/Source/GameplayCameras/Private/Nodes/Common/DampenPositionCameraNode.cpp b/Engine/Plugins/Cameras/GameplayCameras/Source/GameplayCameras/Private/Nodes/Common/DampenPositionCameraNode.cpp
--- a/Engine/Plugins/Cameras/GameplayCameras/Source/GameplayCameras/Private/Nodes/Common/DampenPositionCameraNode.cpp
+++ b/Engine/Plugins/Cameras/GameplayCameras/Source/GameplayCameras/Private/Nodes/Common/DampenPositionCameraNode.cpp
@@ -10,6 +10,9 @@
 
 #include UE_INLINE_GENERATED_CPP_BY_NAME(DampenPositionCameraNode)
 
+namespace UE::Cameras
+{
+
 class FDampenPositionCameraNodeEvaluator : public FCameraNodeEvaluator
 {
 	UE_DECLARE_CAMERA_NODE_EVALUATOR(FDampenPositionCameraNodeEvaluator)
@@ -104,6 +107,8 @@ void FDampenPositionCameraNodeEvaluator::OnRun(const FCameraNodeEvaluationParams
 	OutResult.CameraPose.SetLocation(NextLocation);
 }
 
+}  // namespace UE::Cameras
+
 UDampenPositionCameraNode::UDampenPositionCameraNode(const FObjectInitializer& ObjectInit)
 	: Super(ObjectInit)
 {
@@ -111,6 +116,7 @@ UDampenPositionCameraNode::UDampenPositionCameraNode(const FObjectInitializer& O
 
 FCameraNodeEvaluatorPtr UDampenPositionCameraNode::OnBuildEvaluator(FCameraNodeEvaluatorBuilder& Builder) const
 {
+	using namespace UE::Cameras;
 	return Builder.BuildEvaluator<FDampenPositionCameraNodeEvaluator>();
 }
 
diff --git a/Engine/Plugins/Cameras/GameplayCameras/Source/GameplayCameras/Private/Nodes/Common/OffsetCameraNode.cpp b/Engine/Plugins/Cameras/GameplayCameras/Source/GameplayCameras/Private/Nodes/Common/OffsetCameraNode.cpp
--- a/Engine/Plugins/Cameras/GameplayCameras/Source/GameplayCameras/Private/Nodes/Common/OffsetCameraNode.cpp
+++ b/Engine/Plugins/Cameras/GameplayCameras/Source/GameplayCameras/Private/Nodes/Common/OffsetCameraNode.cpp
@@ -8,6 +8,9 @@
 
 #include UE_INLINE_GENERATED_CPP_BY_NAME(OffsetCameraNode)
 
+namespace UE::Cameras
+{
+
 class FOffsetCameraNodeEvaluator : public FCameraNodeEvaluator
 {
 	UE_DECLARE_CAMERA_NODE_EVALUATOR(FOffsetCameraNodeEvaluator)
@@ -70,8 +73,11 @@ void FOffsetCameraNodeEvaluator::OnRun(const FCameraNodeEvaluationParams& Params
 	OutResult.CameraPose.SetLocation(Location + LocalOffset);
 }
 
+}  // namespace UE::Cameras
+
 FCameraNodeEvaluatorPtr UOffsetCameraNode::OnBuildEvaluator(FCameraNodeEvaluatorBuilder& Builder) const
 {
+	using namespace UE::Cameras;
 	return Builder.BuildEvaluator<FOffsetCameraNodeEvaluator>();
 }
 
diff --git a/Engine/Plugins/Cameras/GameplayCameras/Source/GameplayCameras/Private/Tests/CriticalDamperTests.cpp b/Engine/Plugins/Cameras/GameplayCameras/Source/GameplayCameras/Private/Tests/CriticalDamperTests.cpp
--- a/Engine/Plugins/Cameras/GameplayCameras/Source/GameplayCameras/Private/Tests/CriticalDamperTests.cpp
+++ b/Engine/Plugins/Cameras/GameplayCameras/Source/GameplayCameras/Private/Tests/CriticalDamperTests.cpp
@@ -9,6 +9,8 @@
 IMPLEMENT_SIMPLE_AUTOMATION_TEST(FGameplayCamerasCriticalDamperTest, "System.Engine.GameplayCameras.CriticalDamper", EAutomationTestFlags::EditorContext | EAutomationTestFlags::EngineFilter)
 bool FGameplayCamerasCriticalDamperTest::RunTest(const FString& Parameters)
 {
+	using namespace UE::Cameras;
+
 	{
 		FCriticalDamper Damper(10.f);
 		float Test = Damper.Update(0, 1.f);
diff --git a/Engine/Plugins/Cameras/GameplayCameras/Source/GameplayCameras/Public/Core/BlendCameraNode.h b/Engine/Plugins/Cameras/GameplayCameras/Source/GameplayCameras/Public/Core/BlendCameraNode.h
--- a/Engine/Plugins/Cameras/GameplayCameras/Source/GameplayCameras/Public/Core/BlendCameraNode.h
+++ b/Engine/Plugins/Cameras/GameplayCameras/Source/GameplayCameras/Public/Core/BlendCameraNode.h
@@ -7,6 +7,18 @@
 
 #include "BlendCameraNode.generated.h"
 
+/**
+ * Base class for blend camera nodes.
+ */
+UCLASS(MinimalAPI)
+class UBlendCameraNode : public UCameraNode
+{
+	GENERATED_BODY()
+};
+
+namespace UE::Cameras
+{
+
 /**
  * Parameter struct for blending camera node tree results.
  */
@@ -44,15 +56,6 @@ struct FCameraNodeBlendResult
 	bool bIsBlendFinished = false;
 };
 
-/**
- * Base class for blend camera nodes.
- */
-UCLASS(MinimalAPI)
-class UBlendCameraNode : public UCameraNode
-{
-	GENERATED_BODY()
-};
-
 /**
  * Base evaluator class for blend camera nodes.
  */
@@ -71,13 +74,15 @@ protected:
 	virtual void OnBlendResults(const FCameraNodeBlendParams& Params, FCameraNodeBlendResult& OutResult) {}
 };
 
+}  // namespace UE::Cameras
+
 // Macros for declaring and defining new blend node evaluators. They are the same
 // as the base ones for generic node evaluators, but the first one prevents you
 // from having to specify FBlendCameraNodeEvaluator as the base class, which saves
 // a little bit of typing.
 //
 #define UE_DECLARE_BLEND_CAMERA_NODE_EVALUATOR(ClassName)\
-	UE_DECLARE_CAMERA_NODE_EVALUATOR_EX(ClassName, FBlendCameraNodeEvaluator)
+	UE_DECLARE_CAMERA_NODE_EVALUATOR_EX(ClassName, ::UE::Cameras::FBlendCameraNodeEvaluator)
 
 #define UE_DECLARE_BLEND_CAMERA_NODE_EVALUATOR_EX(ClassName, BaseClassName)\
 	UE_DECLARE_CAMERA_NODE_EVALUATOR_EX(ClassName, BaseClassName)
diff --git a/Engine/Plugins/Cameras/GameplayCameras/Source/GameplayCameras/Public/Core/BlendStackCameraNode.h b/Engine/Plugins/Cameras/GameplayCameras/Source/GameplayCameras/Public/Core/BlendStackCameraNode.h
--- a/Engine/Plugins/Cameras/GameplayCameras/Source/GameplayCameras/Public/Core/BlendStackCameraNode.h
+++ b/Engine/Plugins/Cameras/GameplayCameras/Source/GameplayCameras/Public/Core/BlendStackCameraNode.h
@@ -10,28 +10,19 @@
 
 #include "BlendStackCameraNode.generated.h"
 
-class FBlendStackRootCameraNodeEvaluator;
-class FCameraEvaluationContext;
-class FCameraSystemEvaluator;
+struct FCameraRigTransition;
 class UBlendStackRootCameraNode;
 class UCameraAsset;
 class UCameraRigAsset;
-struct FCameraRigTransition;
 
-/**
- * Parameter structure for pushing a camera rig onto a blend stack.
- */
-struct FBlendStackCameraPushParams
+namespace UE::Cameras
 {
-	/** The evaluator currently running.*/
-	TSharedPtr<FCameraSystemEvaluator> Evaluator;
 
-	/** The evaluation context within which a camera rig's node tree should run. */
-	TSharedPtr<const FCameraEvaluationContext> EvaluationContext;
+class FBlendStackRootCameraNodeEvaluator;
+class FCameraEvaluationContext;
+class FCameraSystemEvaluator;
 
-	/** The source camera rig asset to instantiate and push on the blend stack. */
-	TObjectPtr<const UCameraRigAsset> CameraRig;
-};
+}  // namespace UE::Cameras
 
 /**
  * A blend stack implemented as a camera node.
@@ -62,6 +53,24 @@ public:
 	bool bBlendFirstCameraRig = false;
 };
 
+namespace UE::Cameras
+{
+
+/**
+ * Parameter structure for pushing a camera rig onto a blend stack.
+ */
+struct FBlendStackCameraPushParams
+{
+	/** The evaluator currently running.*/
+	TSharedPtr<FCameraSystemEvaluator> Evaluator;
+
+	/** The evaluation context within which a camera rig's node tree should run. */
+	TSharedPtr<const FCameraEvaluationContext> EvaluationContext;
+
+	/** The source camera rig asset to instantiate and push on the blend stack. */
+	TObjectPtr<const UCameraRigAsset> CameraRig;
+};
+
 /**
  * Evaluator for a blend stack camera node.
  */
@@ -136,3 +145,5 @@ protected:
 #endif  // WITH_EDITOR
 };
 
+}  // namespace UE::Cameras
+
diff --git a/Engine/Plugins/Cameras/GameplayCameras/Source/GameplayCameras/Public/Core/CameraDirector.h b/Engine/Plugins/Cameras/GameplayCameras/Source/GameplayCameras/Public/Core/CameraDirector.h
--- a/Engine/Plugins/Cameras/GameplayCameras/Source/GameplayCameras/Public/Core/CameraDirector.h
+++ b/Engine/Plugins/Cameras/GameplayCameras/Source/GameplayCameras/Public/Core/CameraDirector.h
@@ -18,12 +18,14 @@ class UCameraDirector : public UObject
 
 public:
 
+	using FCameraDirectorEvaluatorBuilder = UE::Cameras::FCameraDirectorEvaluatorBuilder;
+
 	/** Build the evaluator for this director. */
-	FCameraDirectorEvaluator* BuildEvaluator(FCameraDirectorEvaluatorBuilder& Builder) const;
+	FCameraDirectorEvaluatorPtr BuildEvaluator(FCameraDirectorEvaluatorBuilder& Builder) const;
 
 protected:
 
 	/** Build the evaluator for this director. */
-	virtual FCameraDirectorEvaluator* OnBuildEvaluator(FCameraDirectorEvaluatorBuilder& Builder) const { return nullptr; }
+	virtual FCameraDirectorEvaluatorPtr OnBuildEvaluator(FCameraDirectorEvaluatorBuilder& Builder) const { return nullptr; }
 };
 
diff --git a/Engine/Plugins/Cameras/GameplayCameras/Source/GameplayCameras/Public/Core/CameraDirectorEvaluator.h b/Engine/Plugins/Cameras/GameplayCameras/Source/GameplayCameras/Public/Core/CameraDirectorEvaluator.h
--- a/Engine/Plugins/Cameras/GameplayCameras/Source/GameplayCameras/Public/Core/CameraDirectorEvaluator.h
+++ b/Engine/Plugins/Cameras/GameplayCameras/Source/GameplayCameras/Public/Core/CameraDirectorEvaluator.h
@@ -7,11 +7,15 @@
 #include "Templates/UniquePtr.h"
 #include "UObject/ObjectPtr.h"
 
+class UCameraDirector;
+class UCameraRigAsset;
+
+namespace UE::Cameras
+{
+
 class FCameraDirectorEvaluator;
 class FCameraDirectorEvaluatorStorage;
 class FCameraEvaluationContext;
-class UCameraDirector;
-class UCameraRigAsset;
 
 /**
  * Parameter structure for running a camera director.
@@ -110,17 +114,6 @@ private:
 	TObjectPtr<const UCameraDirector> PrivateCameraDirector;
 };
 
-// Utility macros for declaring and defining camera director evaluators.
-//
-#define UE_DECLARE_CAMERA_DIRECTOR_EVALUATOR(ClassName)\
-	UE_GAMEPLAY_CAMERAS_DECLARE_RTTI(ClassName, FCameraDirectorEvaluator)
-
-#define UE_DECLARE_CAMERA_DIRECTOR_EVALUATOR_EX(ClassName, BaseClassName)\
-	UE_GAMEPLAY_CAMERAS_DECLARE_RTTI(ClassName, BaseClassName)
-
-#define UE_DEFINE_CAMERA_DIRECTOR_EVALUATOR(ClassName)\
-	UE_GAMEPLAY_CAMERAS_DEFINE_RTTI(ClassName)
-
 template<typename EvaluatorType, typename ...ArgTypes>
 EvaluatorType* FCameraDirectorEvaluatorBuilder::BuildEvaluator(ArgTypes&&... InArgs)
 {
@@ -136,3 +129,19 @@ EvaluatorType* FCameraDirectorEvaluatorStorage::BuildEvaluator(ArgTypes&&... InA
 	return Evaluator->CastThisChecked<EvaluatorType>();
 }
 
+}  // namespace UE::Cameras
+
+// Typedef to avoid having to deal with namespaces in UCameraNode subclasses.
+using FCameraDirectorEvaluatorPtr = UE::Cameras::FCameraDirectorEvaluator*;
+
+// Utility macros for declaring and defining camera director evaluators.
+//
+#define UE_DECLARE_CAMERA_DIRECTOR_EVALUATOR(ClassName)\
+	UE_GAMEPLAY_CAMERAS_DECLARE_RTTI(ClassName, ::UE::Cameras::FCameraDirectorEvaluator)
+
+#define UE_DECLARE_CAMERA_DIRECTOR_EVALUATOR_EX(ClassName, BaseClassName)\
+	UE_GAMEPLAY_CAMERAS_DECLARE_RTTI(ClassName, BaseClassName)
+
+#define UE_DEFINE_CAMERA_DIRECTOR_EVALUATOR(ClassName)\
+	UE_GAMEPLAY_CAMERAS_DEFINE_RTTI(ClassName)
+
diff --git a/Engine/Plugins/Cameras/GameplayCameras/Source/GameplayCameras/Public/Core/CameraEvaluationContext.h b/Engine/Plugins/Cameras/GameplayCameras/Source/GameplayCameras/Public/Core/CameraEvaluationContext.h
--- a/Engine/Plugins/Cameras/GameplayCameras/Source/GameplayCameras/Public/Core/CameraEvaluationContext.h
+++ b/Engine/Plugins/Cameras/GameplayCameras/Source/GameplayCameras/Public/Core/CameraEvaluationContext.h
@@ -11,6 +11,9 @@
 class UCameraAsset;
 class APlayerController;
 
+namespace UE::Cameras
+{
+
 /**
  * Base class for providing a context to running camera rigs.
  */
@@ -57,10 +60,12 @@ protected:
 	FCameraNodeEvaluationResult InitialResult;
 };
 
+}  // namespace UE::Cameras
+
 // Utility macros for declaring and defining camera evaluation contexts.
 //
 #define UE_DECLARE_CAMERA_EVALUATION_CONTEXT(ClassName)\
-	UE_GAMEPLAY_CAMERAS_DECLARE_RTTI(ClassName, FCameraEvaluationContext)
+	UE_GAMEPLAY_CAMERAS_DECLARE_RTTI(ClassName, ::UE::Cameras::FCameraEvaluationContext)
 
 #define UE_DECLARE_CAMERA_EVALUATION_CONTEXT_EX(ClassName, BaseClassName)\
 	UE_GAMEPLAY_CAMERAS_DECLARE_RTTI(ClassName, BaseClassName)
diff --git a/Engine/Plugins/Cameras/GameplayCameras/Source/GameplayCameras/Public/Core/CameraEvaluationContextStack.h b/Engine/Plugins/Cameras/GameplayCameras/Source/GameplayCameras/Public/Core/CameraEvaluationContextStack.h
--- a/Engine/Plugins/Cameras/GameplayCameras/Source/GameplayCameras/Public/Core/CameraEvaluationContextStack.h
+++ b/Engine/Plugins/Cameras/GameplayCameras/Source/GameplayCameras/Public/Core/CameraEvaluationContextStack.h
@@ -6,11 +6,15 @@
 #include "Core/CameraDirectorEvaluator.h"
 #include "UObject/WeakObjectPtr.h"
 
+class UCameraAsset;
+class UCameraDirector;
+
+namespace UE::Cameras
+{
+
 class FCameraDirectorEvaluator;
 class FCameraEvaluationContext;
 class FCameraSystemEvaluator;
-class UCameraAsset;
-class UCameraDirector;
 
 /** Information about a running camera evaluation context. */
 struct FCameraEvaluationContextInfo
@@ -73,3 +77,5 @@ private:
 	TSharedPtr<FCameraSystemEvaluator> Evaluator;
 };
 
+}  // namespace UE::Cameras
+
diff --git a/Engine/Plugins/Cameras/GameplayCameras/Source/GameplayCameras/Public/Core/CameraNode.h b/Engine/Plugins/Cameras/GameplayCameras/Source/GameplayCameras/Public/Core/CameraNode.h
--- a/Engine/Plugins/Cameras/GameplayCameras/Source/GameplayCameras/Public/Core/CameraNode.h
+++ b/Engine/Plugins/Cameras/GameplayCameras/Source/GameplayCameras/Public/Core/CameraNode.h
@@ -26,7 +26,7 @@ struct FCameraNodeAllocationInfo
 };
 
 /** View on a camera node's children. */
-using FCameraNodeChildrenView = TObjectChildrenView<TObjectPtr<UCameraNode>>;
+using FCameraNodeChildrenView = UE::Cameras::TObjectChildrenView<TObjectPtr<UCameraNode>>;
 
 /**
  * The base class for a camera node.
@@ -37,6 +37,8 @@ class UCameraNode : public UObject
 	GENERATED_BODY()
 
 public:
+	
+	using FCameraNodeEvaluatorBuilder = UE::Cameras::FCameraNodeEvaluatorBuilder;
 
 	/** Get the list of children under this node. */
 	FCameraNodeChildrenView GetChildren();
diff --git a/Engine/Plugins/Cameras/GameplayCameras/Source/GameplayCameras/Public/Core/CameraNodeEvaluator.h b/Engine/Plugins/Cameras/GameplayCameras/Source/GameplayCameras/Public/Core/CameraNodeEvaluator.h
--- a/Engine/Plugins/Cameras/GameplayCameras/Source/GameplayCameras/Public/Core/CameraNodeEvaluator.h
+++ b/Engine/Plugins/Cameras/GameplayCameras/Source/GameplayCameras/Public/Core/CameraNodeEvaluator.h
@@ -9,11 +9,15 @@
 #include "CoreTypes.h"
 #include "UObject/ObjectPtr.h"
 
+class FReferenceCollector;
+class UCameraNode;
+
+namespace UE::Cameras
+{
+
 class FCameraEvaluationContext;
 class FCameraNodeEvaluator;
 class FCameraSystemEvaluator;
-class FReferenceCollector;
-class UCameraNode;
 struct FCameraNodeEvaluatorBuilder;
 
 /**
@@ -130,17 +134,6 @@ private:
 	TObjectPtr<const UCameraNode> PrivateCameraNode;
 };
 
-// Utility macros for declaring and defining camera node evaluators.
-//
-#define UE_DECLARE_CAMERA_NODE_EVALUATOR(ClassName)\
-	UE_GAMEPLAY_CAMERAS_DECLARE_RTTI(ClassName, FCameraNodeEvaluator)
-
-#define UE_DECLARE_CAMERA_NODE_EVALUATOR_EX(ClassName, BaseClassName)\
-	UE_GAMEPLAY_CAMERAS_DECLARE_RTTI(ClassName, BaseClassName)
-
-#define UE_DEFINE_CAMERA_NODE_EVALUATOR(ClassName)\
-	UE_GAMEPLAY_CAMERAS_DEFINE_RTTI(ClassName)
-
 /** Utility base class for camera node evaluators of a specific camera node type. */
 template<typename CameraNodeType>
 class TCameraNodeEvaluator : public FCameraNodeEvaluator
@@ -164,3 +157,16 @@ EvaluatorType* FCameraNodeEvaluatorInitializeParams::BuildEvaluatorAs(const UCam
 	return NewEvaluator->CastThisChecked<EvaluatorType>();
 }
 
+}  // namespace UE::Cameras
+
+// Utility macros for declaring and defining camera node evaluators.
+//
+#define UE_DECLARE_CAMERA_NODE_EVALUATOR(ClassName)\
+	UE_GAMEPLAY_CAMERAS_DECLARE_RTTI(ClassName, FCameraNodeEvaluator)
+
+#define UE_DECLARE_CAMERA_NODE_EVALUATOR_EX(ClassName, BaseClassName)\
+	UE_GAMEPLAY_CAMERAS_DECLARE_RTTI(ClassName, BaseClassName)
+
+#define UE_DEFINE_CAMERA_NODE_EVALUATOR(ClassName)\
+	UE_GAMEPLAY_CAMERAS_DEFINE_RTTI(ClassName)
+
diff --git a/Engine/Plugins/Cameras/GameplayCameras/Source/GameplayCameras/Public/Core/CameraNodeEvaluatorBuilder.h b/Engine/Plugins/Cameras/GameplayCameras/Source/GameplayCameras/Public/Core/CameraNodeEvaluatorBuilder.h
--- a/Engine/Plugins/Cameras/GameplayCameras/Source/GameplayCameras/Public/Core/CameraNodeEvaluatorBuilder.h
+++ b/Engine/Plugins/Cameras/GameplayCameras/Source/GameplayCameras/Public/Core/CameraNodeEvaluatorBuilder.h
@@ -4,6 +4,9 @@
 
 #include "Core/CameraNodeEvaluatorStorage.h"
 
+namespace UE::Cameras
+{
+
 /** Structure for building camera node evaluators. */
 struct FCameraNodeEvaluatorBuilder
 {
@@ -28,3 +31,5 @@ EvaluatorType* FCameraNodeEvaluatorBuilder::BuildEvaluator(ArgTypes&&... InArgs)
 	return Storage.BuildEvaluator<EvaluatorType>(Forward<ArgTypes>(InArgs)...);
 }
 
+}  // namespace UE::Cameras
+
diff --git a/Engine/Plugins/Cameras/GameplayCameras/Source/GameplayCameras/Public/Core/CameraNodeEvaluatorFwd.h b/Engine/Plugins/Cameras/GameplayCameras/Source/GameplayCameras/Public/Core/CameraNodeEvaluatorFwd.h
--- a/Engine/Plugins/Cameras/GameplayCameras/Source/GameplayCameras/Public/Core/CameraNodeEvaluatorFwd.h
+++ b/Engine/Plugins/Cameras/GameplayCameras/Source/GameplayCameras/Public/Core/CameraNodeEvaluatorFwd.h
@@ -6,11 +6,17 @@
 
 #include "CameraNodeEvaluatorFwd.generated.h"
 
+namespace UE::Cameras
+{
+
 class FCameraNodeEvaluator;
 class FCameraNodeEvaluatorStorage;
 struct FCameraNodeEvaluatorBuilder;
 
-using FCameraNodeEvaluatorPtr = FCameraNodeEvaluator*;
+}  // namespace UE::Cameras
+
+// Typedef to avoid having to deal with namespaces in UCameraNode subclasses.
+using FCameraNodeEvaluatorPtr = UE::Cameras::FCameraNodeEvaluator*;
 
 /** Allocation information for a node evaluator. */
 USTRUCT()
diff --git a/Engine/Plugins/Cameras/GameplayCameras/Source/GameplayCameras/Public/Core/CameraNodeEvaluatorStorage.h b/Engine/Plugins/Cameras/GameplayCameras/Source/GameplayCameras/Public/Core/CameraNodeEvaluatorStorage.h
--- a/Engine/Plugins/Cameras/GameplayCameras/Source/GameplayCameras/Public/Core/CameraNodeEvaluatorStorage.h
+++ b/Engine/Plugins/Cameras/GameplayCameras/Source/GameplayCameras/Public/Core/CameraNodeEvaluatorStorage.h
@@ -8,13 +8,17 @@
 #include "Math/UnrealMathUtility.h"
 #include "UObject/ObjectPtr.h"
 
+class FReferenceCollector;
+class UCameraRigAsset;
+class UCameraNode;
+
+namespace UE::Cameras
+{
+
 class FCameraEvaluationContext;
 class FCameraNodeEvaluator;
 class FCameraNodeEvaluatorStorage;
 class FCameraSystemEvaluator;
-class FReferenceCollector;
-class UCameraRigAsset;
-class UCameraNode;
 
 /** Allocation information for an entire tree of node evaluators. */
 struct FCameraNodeEvaluatorTreeAllocationInfo
@@ -133,3 +137,5 @@ EvaluatorType* FCameraNodeEvaluatorStorage::BuildEvaluator(ArgTypes&&... InArgs)
 	return NewEvaluator;
 }
 
+}  // namespace UE::Cameras
+
diff --git a/Engine/Plugins/Cameras/GameplayCameras/Source/GameplayCameras/Public/Core/CameraObjectRtti.h b/Engine/Plugins/Cameras/GameplayCameras/Source/GameplayCameras/Public/Core/CameraObjectRtti.h
--- a/Engine/Plugins/Cameras/GameplayCameras/Source/GameplayCameras/Public/Core/CameraObjectRtti.h
+++ b/Engine/Plugins/Cameras/GameplayCameras/Source/GameplayCameras/Public/Core/CameraObjectRtti.h
@@ -5,6 +5,9 @@
 #include "CoreTypes.h"
 #include "UObject/NameTypes.h"
 
+namespace UE::Cameras
+{
+
 /**
  * Identifier for a given RTTI-enabled class.
  */
@@ -63,6 +66,8 @@ private:
 	friend T;
 };
 
+}  // namespace UE::Cameras
+
 // Macros for enabling simple RTTI information on a class hierarchy.
 //
 // The first macro is for the root of the class hierarchy, while the second is for all
@@ -70,7 +75,7 @@ private:
 //
 #define UE_GAMEPLAY_CAMERAS_DECLARE_RTTI_BASE(ClassName)\
 	public:\
-		static const TCameraObjectTypeID<ClassName>& StaticTypeID() { return ClassName::PrivateTypeID; }\
+		static const ::UE::Cameras::TCameraObjectTypeID<ClassName>& StaticTypeID() { return ClassName::PrivateTypeID; }\
 		virtual const FCameraObjectTypeID& GetTypeID() const { return ClassName::PrivateTypeID; }\
 		virtual bool IsKindOf(const FCameraObjectTypeID& InTypeID) const { return InTypeID == ClassName::PrivateTypeID; }\
 		template<typename Type> bool IsKindOf() const { return IsKindOf(Type::StaticTypeID()); }\
@@ -79,16 +84,16 @@ private:
 		template<typename Type> Type* CastThisChecked() { check(IsKindOf<Type>()); return static_cast<Type*>(this); }\
 		template<typename Type> const Type* CastThisChecked() const { check(IsKindOf<Type>()); return static_cast<Type*>(this); }\
 	private:\
-		static const TCameraObjectTypeID<ClassName> PrivateTypeID;
+		static const ::UE::Cameras::TCameraObjectTypeID<ClassName> PrivateTypeID;
 
 #define UE_GAMEPLAY_CAMERAS_DECLARE_RTTI(ClassName, BaseClassName)\
 	public:\
-		static const TCameraObjectTypeID<ClassName>& StaticTypeID() { return ClassName::PrivateTypeID; }\
+		static const ::UE::Cameras::TCameraObjectTypeID<ClassName>& StaticTypeID() { return ClassName::PrivateTypeID; }\
 		virtual const FCameraObjectTypeID& GetTypeID() const override { return ClassName::PrivateTypeID; }\
 		virtual bool IsKindOf(const FCameraObjectTypeID& InTypeID) const override { return (InTypeID == ClassName::PrivateTypeID) || BaseClassName::IsKindOf(InTypeID); }\
 	private:\
-		static const TCameraObjectTypeID<ClassName> PrivateTypeID;
+		static const ::UE::Cameras::TCameraObjectTypeID<ClassName> PrivateTypeID;
 
 #define UE_GAMEPLAY_CAMERAS_DEFINE_RTTI(ClassName)\
-	const TCameraObjectTypeID<ClassName> ClassName::PrivateTypeID = TCameraObjectTypeID<ClassName>::RegisterNewID(#ClassName);
+	const ::UE::Cameras::TCameraObjectTypeID<ClassName> ClassName::PrivateTypeID = ::UE::Cameras::TCameraObjectTypeID<ClassName>::RegisterNewID(#ClassName);
 
diff --git a/Engine/Plugins/Cameras/GameplayCameras/Source/GameplayCameras/Public/Core/CameraParameterReader.h b/Engine/Plugins/Cameras/GameplayCameras/Source/GameplayCameras/Public/Core/CameraParameterReader.h
--- a/Engine/Plugins/Cameras/GameplayCameras/Source/GameplayCameras/Public/Core/CameraParameterReader.h
+++ b/Engine/Plugins/Cameras/GameplayCameras/Source/GameplayCameras/Public/Core/CameraParameterReader.h
@@ -6,6 +6,9 @@
 #include "CoreTypes.h"
 #include "Templates/UnrealTypeTraits.h"
 
+namespace UE::Cameras
+{
+
 template<typename ValueType>
 class FCameraParameterReader
 {
@@ -70,3 +73,5 @@ private:
 	uint32 VariableId = 0;
 };
 
+}  // namespace UE::Cameras
+
diff --git a/Engine/Plugins/Cameras/GameplayCameras/Source/GameplayCameras/Public/Core/CameraSystemEvaluator.h b/Engine/Plugins/Cameras/GameplayCameras/Source/GameplayCameras/Public/Core/CameraSystemEvaluator.h
--- a/Engine/Plugins/Cameras/GameplayCameras/Source/GameplayCameras/Public/Core/CameraSystemEvaluator.h
+++ b/Engine/Plugins/Cameras/GameplayCameras/Source/GameplayCameras/Public/Core/CameraSystemEvaluator.h
@@ -12,13 +12,17 @@
 #include "Templates/SharedPointer.h"
 #include "UObject/GCObject.h"
 
-class FCameraEvaluationContext;
-class FRootCameraNodeEvaluator;
 class UCameraDirector;
 class UCameraRigAsset;
 class URootCameraNode;
 struct FMinimalViewInfo;
 
+namespace UE::Cameras
+{
+
+class FCameraEvaluationContext;
+class FRootCameraNodeEvaluator;
+
 /**
  * Parameter structure for initializing a new camera system evaluator.
  */
@@ -111,3 +115,5 @@ private:
 	FCameraSystemEvaluationUpdateResult Result;
 };
 
+}  // namespace UE::Cameras
+
diff --git a/Engine/Plugins/Cameras/GameplayCameras/Source/GameplayCameras/Public/Core/CameraVariableTable.h b/Engine/Plugins/Cameras/GameplayCameras/Source/GameplayCameras/Public/Core/CameraVariableTable.h
--- a/Engine/Plugins/Cameras/GameplayCameras/Source/GameplayCameras/Public/Core/CameraVariableTable.h
+++ b/Engine/Plugins/Cameras/GameplayCameras/Source/GameplayCameras/Public/Core/CameraVariableTable.h
@@ -10,6 +10,9 @@
 #include "Templates/UnrealTypeTraits.h"
 #include "UObject/NameTypes.h"
 
+namespace UE::Cameras
+{
+
 template<typename ValueType>
 struct TCameraVariableTraits;
 
@@ -254,3 +257,5 @@ struct TCameraVariableInterpolation<UE::Math::TTransform<T>>
 	}
 };
 
+}  // namespace UE::Cameras
+
diff --git a/Engine/Plugins/Cameras/GameplayCameras/Source/GameplayCameras/Public/Core/DefaultRootCameraNode.h b/Engine/Plugins/Cameras/GameplayCameras/Source/GameplayCameras/Public/Core/DefaultRootCameraNode.h
--- a/Engine/Plugins/Cameras/GameplayCameras/Source/GameplayCameras/Public/Core/DefaultRootCameraNode.h
+++ b/Engine/Plugins/Cameras/GameplayCameras/Source/GameplayCameras/Public/Core/DefaultRootCameraNode.h
@@ -7,7 +7,6 @@
 #include "DefaultRootCameraNode.generated.h"
 
 class UBlendStackCameraNode;
-class FBlendStackCameraNodeEvaluator;
 
 /**
  * The default implementation of a root camera node.
@@ -40,6 +39,11 @@ public:
 	TObjectPtr<UBlendStackCameraNode> VisualLayer;
 };
 
+namespace UE::Cameras
+{
+
+class FBlendStackCameraNodeEvaluator;
+
 /**
  * Evaluator for the default root camera node.
  */
@@ -63,3 +67,5 @@ private:
 	FBlendStackCameraNodeEvaluator* VisualLayer;
 };
 
+}  // namespace UE::Cameras
+
diff --git a/Engine/Plugins/Cameras/GameplayCameras/Source/GameplayCameras/Public/Core/ObjectChildrenView.h b/Engine/Plugins/Cameras/GameplayCameras/Source/GameplayCameras/Public/Core/ObjectChildrenView.h
--- a/Engine/Plugins/Cameras/GameplayCameras/Source/GameplayCameras/Public/Core/ObjectChildrenView.h
+++ b/Engine/Plugins/Cameras/GameplayCameras/Source/GameplayCameras/Public/Core/ObjectChildrenView.h
@@ -10,6 +10,9 @@
 
 #include <initializer_list>
 
+namespace UE::Cameras
+{
+
 /**
  * Structure that describes a list of children of an object.
  *
@@ -153,3 +156,5 @@ private:
 	FStorage Storage;
 };
 
+}  // namespace UE::Cameras
+
diff --git a/Engine/Plugins/Cameras/GameplayCameras/Source/GameplayCameras/Public/Core/RootCameraNode.h b/Engine/Plugins/Cameras/GameplayCameras/Source/GameplayCameras/Public/Core/RootCameraNode.h
--- a/Engine/Plugins/Cameras/GameplayCameras/Source/GameplayCameras/Public/Core/RootCameraNode.h
+++ b/Engine/Plugins/Cameras/GameplayCameras/Source/GameplayCameras/Public/Core/RootCameraNode.h
@@ -7,8 +7,6 @@
 
 #include "RootCameraNode.generated.h"
 
-class FCameraEvaluationContext;
-class FCameraSystemEvaluator;
 class UCameraRigAsset;
 
 /**
@@ -27,6 +25,22 @@ enum class ECameraRigLayer
 };
 ENUM_CLASS_FLAGS(ECameraRigLayer)
 
+/**
+ * The base class for a camera node that can act as the root of the
+ * camera system evaluation.
+ */
+UCLASS(MinimalAPI, Abstract)
+class URootCameraNode : public UCameraNode
+{
+	GENERATED_BODY()
+};
+
+namespace UE::Cameras
+{
+
+class FCameraEvaluationContext;
+class FCameraSystemEvaluator;
+
 /**
  * Parameter structure for activating a new camera rig.
  */
@@ -45,16 +59,6 @@ struct FActivateCameraRigParams
 	ECameraRigLayer Layer = ECameraRigLayer::Main;
 };
 
-/**
- * The base class for a camera node that can act as the root of the
- * camera system evaluation.
- */
-UCLASS(MinimalAPI, Abstract)
-class URootCameraNode : public UCameraNode
-{
-	GENERATED_BODY()
-};
-
 /**
  * Base class for the evaluator of a root camera node.
  */
@@ -71,3 +75,5 @@ private:
 	virtual void OnActivateCameraRig(const FActivateCameraRigParams& Params) {}
 };
 
+}  // namespace UE::Cameras
+
diff --git a/Engine/Plugins/Cameras/GameplayCameras/Source/GameplayCameras/Public/Directors/SingleCameraDirector.h b/Engine/Plugins/Cameras/GameplayCameras/Source/GameplayCameras/Public/Directors/SingleCameraDirector.h
--- a/Engine/Plugins/Cameras/GameplayCameras/Source/GameplayCameras/Public/Directors/SingleCameraDirector.h
+++ b/Engine/Plugins/Cameras/GameplayCameras/Source/GameplayCameras/Public/Directors/SingleCameraDirector.h
@@ -21,7 +21,7 @@ public:
 protected:
 
 	// UCameraDirector interface.
-	virtual FCameraDirectorEvaluator* OnBuildEvaluator(FCameraDirectorEvaluatorBuilder& Builder) const override;
+	virtual FCameraDirectorEvaluatorPtr OnBuildEvaluator(FCameraDirectorEvaluatorBuilder& Builder) const override;
 
 public:
 
diff --git a/Engine/Plugins/Cameras/GameplayCameras/Source/GameplayCameras/Public/GameFramework/GameplayCameraComponent.h b/Engine/Plugins/Cameras/GameplayCameras/Source/GameplayCameras/Public/GameFramework/GameplayCameraComponent.h
--- a/Engine/Plugins/Cameras/GameplayCameras/Source/GameplayCameras/Public/GameFramework/GameplayCameraComponent.h
+++ b/Engine/Plugins/Cameras/GameplayCameras/Source/GameplayCameras/Public/GameFramework/GameplayCameraComponent.h
@@ -10,9 +10,15 @@
 
 #include "GameplayCameraComponent.generated.h"
 
-class FGameplayCameraComponentEvaluationContext;
 class UCameraAsset;
 
+namespace UE::Cameras
+{
+
+class FGameplayCameraComponentEvaluationContext;
+
+}  // namespace UE::Cameras
+
 /**
  * A component that can run a camera asset inside its own camera evaluation context.
  */
@@ -56,6 +62,8 @@ public:
 
 protected:
 
+	using FGameplayCameraComponentEvaluationContext = UE::Cameras::FGameplayCameraComponentEvaluationContext;
+
 	TSharedPtr<FGameplayCameraComponentEvaluationContext> EvaluationContext;
 	
 #if WITH_EDITORONLY_DATA
@@ -73,6 +81,9 @@ private:
 	int32 ActivatedForPlayerIndex = INDEX_NONE;
 };
 
+namespace UE::Cameras
+{
+
 /**
  * Evaluation context for the gameplay camera component.
  */
@@ -86,3 +97,5 @@ public:
 	void Update(UGameplayCameraComponent* Owner);
 };
 
+}  // namespace UE::Cameras
+
diff --git a/Engine/Plugins/Cameras/GameplayCameras/Source/GameplayCameras/Public/GameFramework/GameplayCameraSystemComponent.h b/Engine/Plugins/Cameras/GameplayCameras/Source/GameplayCameras/Public/GameFramework/GameplayCameraSystemComponent.h
--- a/Engine/Plugins/Cameras/GameplayCameras/Source/GameplayCameras/Public/GameFramework/GameplayCameraSystemComponent.h
+++ b/Engine/Plugins/Cameras/GameplayCameras/Source/GameplayCameras/Public/GameFramework/GameplayCameraSystemComponent.h
@@ -10,10 +10,16 @@
 
 #include "GameplayCameraSystemComponent.generated.h"
 
-class FCameraSystemEvaluator;
 class UCameraRigAsset;
 struct FMinimalViewInfo;
 
+namespace UE::Cameras
+{
+
+class FCameraSystemEvaluator;
+
+}  // namespace UE::Cameras
+
 /**
  * A component that hosts a camera system.
  */
@@ -24,6 +30,8 @@ class UGameplayCameraSystemComponent : public USceneComponent
 
 public:
 
+	using FCameraSystemEvaluator = UE::Cameras::FCameraSystemEvaluator;
+
 	UGameplayCameraSystemComponent(const FObjectInitializer& ObjectInit);
 
 	TSharedPtr<FCameraSystemEvaluator> GetCameraSystemEvaluator() { return Evaluator; }
diff --git a/Engine/Plugins/Cameras/GameplayCameras/Source/GameplayCameras/Public/IGameplayCamerasLiveEditListener.h b/Engine/Plugins/Cameras/GameplayCameras/Source/GameplayCameras/Public/IGameplayCamerasLiveEditListener.h
--- a/Engine/Plugins/Cameras/GameplayCameras/Source/GameplayCameras/Public/IGameplayCamerasLiveEditListener.h
+++ b/Engine/Plugins/Cameras/GameplayCameras/Source/GameplayCameras/Public/IGameplayCamerasLiveEditListener.h
@@ -2,14 +2,17 @@
 
 #pragma once
 
-#if WITH_EDITOR
-
 #include "CoreTypes.h"
 #include "Templates/SharedPointer.h"
 #include "UObject/UnrealType.h"
 
 class UPackage;
 
+namespace UE::Cameras
+{
+
+#if WITH_EDITOR
+
 struct FGameplayCameraAssetBuildEvent
 {
 	const UPackage* AssetPackage = nullptr;
@@ -33,3 +36,5 @@ protected:
 
 #endif  // WITH_EDITOR
 
+}  // namespace UE::Cameras
+
diff --git a/Engine/Plugins/Cameras/GameplayCameras/Source/GameplayCameras/Public/IGameplayCamerasLiveEditManager.h b/Engine/Plugins/Cameras/GameplayCameras/Source/GameplayCameras/Public/IGameplayCamerasLiveEditManager.h
--- a/Engine/Plugins/Cameras/GameplayCameras/Source/GameplayCameras/Public/IGameplayCamerasLiveEditManager.h
+++ b/Engine/Plugins/Cameras/GameplayCameras/Source/GameplayCameras/Public/IGameplayCamerasLiveEditManager.h
@@ -2,13 +2,17 @@
 
 #pragma once
 
-#if WITH_EDITOR
-
 #include "CoreTypes.h"
 #include "Templates/SharedPointer.h"
 #include "UObject/UnrealType.h"
 
 class UPackage;
+
+namespace UE::Cameras
+{
+
+#if WITH_EDITOR
+
 class IGameplayCamerasLiveEditListener;
 
 /**
@@ -31,3 +35,5 @@ public:
 
 #endif  // WITH_EDITOR
 
+}  // namespace UE::Cameras
+
diff --git a/Engine/Plugins/Cameras/GameplayCameras/Source/GameplayCameras/Public/IGameplayCamerasModule.h b/Engine/Plugins/Cameras/GameplayCameras/Source/GameplayCameras/Public/IGameplayCamerasModule.h
--- a/Engine/Plugins/Cameras/GameplayCameras/Source/GameplayCameras/Public/IGameplayCamerasModule.h
+++ b/Engine/Plugins/Cameras/GameplayCameras/Source/GameplayCameras/Public/IGameplayCamerasModule.h
@@ -5,13 +5,19 @@
 #include "Modules/ModuleInterface.h"
 #include "Templates/SharedPointer.h"
 
+namespace UE::Cameras
+{
+
 #if WITH_EDITOR
 class IGameplayCamerasLiveEditManager;
 #endif
 
+}  // namespace UE::Cameras
+
 class IGameplayCamerasModule : public IModuleInterface
 {
 public:
+
 	/**
 	 * Singleton-like access to ICameraModule
 	 *
@@ -20,6 +26,8 @@ public:
 	static IGameplayCamerasModule& Get();
 
 #if WITH_EDITOR
+	using IGameplayCamerasLiveEditManager = UE::Cameras::IGameplayCamerasLiveEditManager;
+
 	/**
 	 * Gets the live edit manager.
 	 */
diff --git a/Engine/Plugins/Cameras/GameplayCameras/Source/GameplayCameras/Public/Math/CriticalDamper.h b/Engine/Plugins/Cameras/GameplayCameras/Source/GameplayCameras/Public/Math/CriticalDamper.h
--- a/Engine/Plugins/Cameras/GameplayCameras/Source/GameplayCameras/Public/Math/CriticalDamper.h
+++ b/Engine/Plugins/Cameras/GameplayCameras/Source/GameplayCameras/Public/Math/CriticalDamper.h
@@ -2,7 +2,10 @@
 
 #pragma once
 
-#include "CriticalDamper.generated.h"
+#include "CoreTypes.h"
+
+namespace UE::Cameras
+{
 
 /**
  * Implements a critically damped spring-mass system.
@@ -28,11 +31,8 @@
  * - See the math notes in the corresponding cpp file for details about the
  *   algorithm and its implementation.
  */
-USTRUCT()
 struct FCriticalDamper
 {
-	GENERATED_BODY()
-
 public:
 
 	FCriticalDamper();
@@ -100,13 +100,12 @@ private:
 private:
 
 	/** Angular frequency of the oscillator. */
-	UPROPERTY()
 	float W0;
 	/** Initial position of the spring-mass system, aka the previous frame's position. */
-	UPROPERTY()
 	float X0;
 	/** Initial velocity of the spring-mass system, aka the previous frame's velocity. */
-	UPROPERTY()
 	float X0Derivative;
 };
 
+}  // namespace UE::Cameras
+
diff --git a/Engine/Plugins/Cameras/GameplayCameras/Source/GameplayCameras/Public/Nodes/Blends/SimpleBlendCameraNode.h b/Engine/Plugins/Cameras/GameplayCameras/Source/GameplayCameras/Public/Nodes/Blends/SimpleBlendCameraNode.h
--- a/Engine/Plugins/Cameras/GameplayCameras/Source/GameplayCameras/Public/Nodes/Blends/SimpleBlendCameraNode.h
+++ b/Engine/Plugins/Cameras/GameplayCameras/Source/GameplayCameras/Public/Nodes/Blends/SimpleBlendCameraNode.h
@@ -6,14 +6,6 @@
 
 #include "SimpleBlendCameraNode.generated.h"
 
-/**
- * Result structure for defining a simple scalar-factor-based blend.
- */
-struct FSimpleBlendCameraNodeEvaluationResult
-{
-	float BlendFactor = 0.f;
-};
-
 /**
  * Base class for a blend camera node that uses a simple scalar factor.
  */
@@ -38,6 +30,17 @@ public:
 	float BlendTime = 1.f;
 };
 
+namespace UE::Cameras
+{
+
+/**
+ * Result structure for defining a simple scalar-factor-based blend.
+ */
+struct FSimpleBlendCameraNodeEvaluationResult
+{
+	float BlendFactor = 0.f;
+};
+
 class FSimpleBlendCameraNodeEvaluator : public FBlendCameraNodeEvaluator
 {
 	UE_DECLARE_BLEND_CAMERA_NODE_EVALUATOR(FSimpleBlendCameraNodeEvaluator)
@@ -77,3 +80,5 @@ private:
 	float CurrentTime = 0.f;
 };
 
+}  // namespace UE::Cameras
+
diff --git a/Engine/Plugins/Cameras/GameplayCameras/Source/GameplayCamerasEditor/Private/GameplayCamerasEditorModule.cpp b/Engine/Plugins/Cameras/GameplayCameras/Source/GameplayCamerasEditor/Private/GameplayCamerasEditorModule.cpp
--- a/Engine/Plugins/Cameras/GameplayCameras/Source/GameplayCamerasEditor/Private/GameplayCamerasEditorModule.cpp
+++ b/Engine/Plugins/Cameras/GameplayCameras/Source/GameplayCamerasEditor/Private/GameplayCamerasEditorModule.cpp
@@ -102,6 +102,8 @@ private:
 
 	void InitializeLiveEditManager()
 	{
+		using namespace UE::Cameras;
+
 		LiveEditManager = MakeShared<FGameplayCamerasLiveEditManager>();
 
 		IGameplayCamerasModule& CamerasModule = FModuleManager::LoadModuleChecked<IGameplayCamerasModule>("GameplayCameras");
@@ -126,7 +128,7 @@ private:
 
 private:
 
-	TSharedPtr<FGameplayCamerasLiveEditManager> LiveEditManager;
+	TSharedPtr<UE::Cameras::FGameplayCamerasLiveEditManager> LiveEditManager;
 };
 
 IMPLEMENT_MODULE(FGameplayCamerasEditorModule, GameplayCamerasEditor);
diff --git a/Engine/Plugins/Cameras/GameplayCameras/Source/GameplayCamerasEditor/Private/GameplayCamerasLiveEditManager.cpp b/Engine/Plugins/Cameras/GameplayCameras/Source/GameplayCamerasEditor/Private/GameplayCamerasLiveEditManager.cpp
--- a/Engine/Plugins/Cameras/GameplayCameras/Source/GameplayCamerasEditor/Private/GameplayCamerasLiveEditManager.cpp
+++ b/Engine/Plugins/Cameras/GameplayCameras/Source/GameplayCamerasEditor/Private/GameplayCamerasLiveEditManager.cpp
@@ -6,6 +6,9 @@
 #include "UObject/Package.h"
 #include "UObject/UObjectGlobals.h"
 
+namespace UE::Cameras
+{
+
 FGameplayCamerasLiveEditManager::FGameplayCamerasLiveEditManager()
 {
 }
@@ -61,3 +64,5 @@ void FGameplayCamerasLiveEditManager::RemoveGarbage()
 	}
 }
 
+}  // namespace UE::Cameras
+
diff --git a/Engine/Plugins/Cameras/GameplayCameras/Source/GameplayCamerasEditor/Private/GameplayCamerasLiveEditManager.h b/Engine/Plugins/Cameras/GameplayCameras/Source/GameplayCamerasEditor/Private/GameplayCamerasLiveEditManager.h
--- a/Engine/Plugins/Cameras/GameplayCameras/Source/GameplayCamerasEditor/Private/GameplayCamerasLiveEditManager.h
+++ b/Engine/Plugins/Cameras/GameplayCameras/Source/GameplayCamerasEditor/Private/GameplayCamerasLiveEditManager.h
@@ -9,6 +9,9 @@
 
 class UPackage;
 
+namespace UE::Cameras
+{
+
 class FGameplayCamerasLiveEditManager : public IGameplayCamerasLiveEditManager
 {
 public:
@@ -33,3 +36,5 @@ private:
 	FListenerMap ListenerMap;
 };
 
+}  // namespace UE::Cameras
+
diff --git a/Engine/Plugins/Cameras/GameplayCameras/Source/GameplayCamerasEditor/Public/Toolkits/CameraRigAssetEditorToolkit.h b/Engine/Plugins/Cameras/GameplayCameras/Source/GameplayCamerasEditor/Public/Toolkits/CameraRigAssetEditorToolkit.h
--- a/Engine/Plugins/Cameras/GameplayCameras/Source/GameplayCamerasEditor/Public/Toolkits/CameraRigAssetEditorToolkit.h
+++ b/Engine/Plugins/Cameras/GameplayCameras/Source/GameplayCamerasEditor/Public/Toolkits/CameraRigAssetEditorToolkit.h
@@ -7,12 +7,18 @@
 #include "Tools/BaseAssetToolkit.h"
 #include "UObject/GCObject.h"
 
-class IGameplayCamerasLiveEditManager;
 class IMessageLogListing;
 class UCameraRigAsset;
 class UCameraRigAssetEditor;
 class SWidget;
 
+namespace UE::Cameras
+{
+
+class IGameplayCamerasLiveEditManager;
+
+}  // namespace UE::Cameras
+
 /**
  * Editor toolkit for a camera rig asset.
  */
@@ -78,6 +84,6 @@ private:
 	TObjectPtr<UCameraRigAsset> CameraRigAsset;
 
 	/** Live edit manager for updating the assets in the runtime */
-	TSharedPtr<IGameplayCamerasLiveEditManager> LiveEditManager;
+	TSharedPtr<UE::Cameras::IGameplayCamerasLiveEditManager> LiveEditManager;
 };
 
