Cameras: process the variable table in the blend stack, generate variable ids automatically

diff --git a/Engine/Plugins/Cameras/GameplayCameras/Source/GameplayCameras/Private/Core/BlendStackCameraNode.cpp b/Engine/Plugins/Cameras/GameplayCameras/Source/GameplayCameras/Private/Core/BlendStackCameraNode.cpp
--- a/Engine/Plugins/Cameras/GameplayCameras/Source/GameplayCameras/Private/Core/BlendStackCameraNode.cpp
+++ b/Engine/Plugins/Cameras/GameplayCameras/Source/GameplayCameras/Private/Core/BlendStackCameraNode.cpp
@@ -149,9 +149,12 @@ void FBlendStackCameraNodeEvaluator::OnRun(const FCameraNodeEvaluationParams& Pa
 			// Start with the input given to us.
 			CurResult.CameraPose = OutResult.CameraPose;
 			CurResult.CameraPose.ClearAllChangedFlags();
+			CurResult.VariableTable.OverrideAll(OutResult.VariableTable);
+			CurResult.VariableTable.ClearAllWrittenThisFrameFlags();
 
 			// Override it with whatever the evaluation context has set on its result.
 			CurResult.CameraPose.OverrideChanged(ContextResult.CameraPose);
+			CurResult.VariableTable.OverrideAll(ContextResult.VariableTable);
 			CurResult.bIsCameraCut = OutResult.bIsCameraCut || ContextResult.bIsCameraCut;
 			CurResult.bIsValid = true;
 
diff --git a/Engine/Plugins/Cameras/GameplayCameras/Source/GameplayCameras/Private/Core/CameraVariableAssets.cpp b/Engine/Plugins/Cameras/GameplayCameras/Source/GameplayCameras/Private/Core/CameraVariableAssets.cpp
--- a/Engine/Plugins/Cameras/GameplayCameras/Source/GameplayCameras/Private/Core/CameraVariableAssets.cpp
+++ b/Engine/Plugins/Cameras/GameplayCameras/Source/GameplayCameras/Private/Core/CameraVariableAssets.cpp
@@ -9,3 +9,15 @@ UCameraVariableAsset::UCameraVariableAsset(const FObjectInitializer& ObjectInit)
 {
 }
 
+void UCameraVariableAsset::Serialize(FArchive& Ar)
+{
+	Super::Serialize(Ar);
+
+#if WITH_EDITORONLY_DATA
+	if ((Ar.IsLoading() && VariableId == 0) || Ar.IsSaving())
+	{
+		VariableId = GetTypeHash(GetFullName());
+	}
+#endif
+}
+
diff --git a/Engine/Plugins/Cameras/GameplayCameras/Source/GameplayCameras/Private/Core/CameraVariableTable.cpp b/Engine/Plugins/Cameras/GameplayCameras/Source/GameplayCameras/Private/Core/CameraVariableTable.cpp
--- a/Engine/Plugins/Cameras/GameplayCameras/Source/GameplayCameras/Private/Core/CameraVariableTable.cpp
+++ b/Engine/Plugins/Cameras/GameplayCameras/Source/GameplayCameras/Private/Core/CameraVariableTable.cpp
@@ -210,7 +210,7 @@ bool FCameraVariableTable::IsValueWrittenThisFrame(uint32 VariableId) const
 	return false;
 }
 
-void FCameraVariableTable::ResetAllWrittenThisFrame()
+void FCameraVariableTable::ClearAllWrittenThisFrameFlags()
 {
 	for (auto& Pair : Entries)
 	{
@@ -218,17 +218,22 @@ void FCameraVariableTable::ResetAllWrittenThisFrame()
 	}
 }
 
+void FCameraVariableTable::OverrideAll(const FCameraVariableTable& OtherTable)
+{
+	InternalOverride(OtherTable, nullptr, false, nullptr, false);
+}
+
 void FCameraVariableTable::OverrideChanged(const FCameraVariableTable& OtherTable)
 {
-	InternalOverrideChanged(OtherTable, nullptr, false, nullptr);
+	InternalOverride(OtherTable, nullptr, false, nullptr, true);
 }
 
 void FCameraVariableTable::OverrideChanged(const FCameraVariableTable& OtherTable, const FCameraVariableTableFlags& InMask, bool bInvertMask, FCameraVariableTableFlags& OutMask)
 {
-	InternalOverrideChanged(OtherTable, &InMask, bInvertMask, &OutMask);
+	InternalOverride(OtherTable, &InMask, bInvertMask, &OutMask, true);
 }
 
-void FCameraVariableTable::InternalOverrideChanged(const FCameraVariableTable& OtherTable, const FCameraVariableTableFlags* InMask, bool bInvertMask, FCameraVariableTableFlags* OutMask)
+void FCameraVariableTable::InternalOverride(const FCameraVariableTable& OtherTable, const FCameraVariableTableFlags* InMask, bool bInvertMask, FCameraVariableTableFlags* OutMask, bool bChangedOnly)
 {
 	using namespace UE::Cameras::Private;
 
@@ -238,6 +243,7 @@ void FCameraVariableTable::InternalOverrideChanged(const FCameraVariableTable& O
 		const FEntry& OtherEntry  = OtherPair.Value;
 		const EEntryFlags OtherFlags = OtherEntry.Flags;
 		if (EnumHasAnyFlags(OtherFlags, EEntryFlags::Written)
+				&& (!bChangedOnly || EnumHasAnyFlags(OtherFlags, EEntryFlags::WrittenThisFrame))
 				&& !EnumHasAnyFlags(OtherFlags, EEntryFlags::Private)
 				&& IsVariableInMask(OtherEntry.Id, InMask, bInvertMask))
 		{
@@ -299,15 +305,15 @@ void FCameraVariableTable::InternalOverrideChanged(const FCameraVariableTable& O
 
 void FCameraVariableTable::LerpChanged(const FCameraVariableTable& ToTable, float Factor)
 {
-	InternalLerpChanged(ToTable, Factor, nullptr, false, nullptr);
+	InternalLerp(ToTable, Factor, nullptr, false, nullptr, true);
 }
 
 void FCameraVariableTable::LerpChanged(const FCameraVariableTable& ToTable, float Factor, const FCameraVariableTableFlags& InMask, bool bInvertMask, FCameraVariableTableFlags& OutMask)
 {
-	InternalLerpChanged(ToTable, Factor, &InMask, false, &OutMask);
+	InternalLerp(ToTable, Factor, &InMask, false, &OutMask, true);
 }
 
-void FCameraVariableTable::InternalLerpChanged(const FCameraVariableTable& ToTable, float Factor, const FCameraVariableTableFlags* InMask, bool bInvertMask, FCameraVariableTableFlags* OutMask)
+void FCameraVariableTable::InternalLerp(const FCameraVariableTable& ToTable, float Factor, const FCameraVariableTableFlags* InMask, bool bInvertMask, FCameraVariableTableFlags* OutMask, bool bChangedOnly)
 {
 	using namespace UE::Cameras::Private;
 
@@ -317,6 +323,7 @@ void FCameraVariableTable::InternalLerpChanged(const FCameraVariableTable& ToTab
 		const FEntry& ToEntry  = ToPair.Value;
 		const EEntryFlags ToFlags = ToEntry.Flags;
 		if (EnumHasAnyFlags(ToFlags, EEntryFlags::Written)
+				&& (!bChangedOnly || EnumHasAnyFlags(ToFlags, EEntryFlags::WrittenThisFrame))
 				&& !EnumHasAnyFlags(ToFlags, EEntryFlags::Private)
 				&& IsVariableInMask(ToEntry.Id, InMask, bInvertMask))
 		{
diff --git a/Engine/Plugins/Cameras/GameplayCameras/Source/GameplayCameras/Public/Core/CameraVariableAssets.h b/Engine/Plugins/Cameras/GameplayCameras/Source/GameplayCameras/Public/Core/CameraVariableAssets.h
--- a/Engine/Plugins/Cameras/GameplayCameras/Source/GameplayCameras/Public/Core/CameraVariableAssets.h
+++ b/Engine/Plugins/Cameras/GameplayCameras/Source/GameplayCameras/Public/Core/CameraVariableAssets.h
@@ -26,6 +26,11 @@ public:
 	virtual ECameraVariableType GetVariableType() const PURE_VIRTUAL(UCameraVariableAsset::GetVariableType, return ECameraVariableType::Boolean;);
 	virtual const uint8* GetDefaultValuePtr() const PURE_VIRTUAL(UCameraVariableAsset::GetDefaultValuePtr, return nullptr;);
 
+public:
+
+	// UObject interface
+	virtual void Serialize(FArchive& Ar) override;
+
 public:
 
 	/** Whether this variable auto-resets to its default value every frame. */
diff --git a/Engine/Plugins/Cameras/GameplayCameras/Source/GameplayCameras/Public/Core/CameraVariableTable.h b/Engine/Plugins/Cameras/GameplayCameras/Source/GameplayCameras/Public/Core/CameraVariableTable.h
--- a/Engine/Plugins/Cameras/GameplayCameras/Source/GameplayCameras/Public/Core/CameraVariableTable.h
+++ b/Engine/Plugins/Cameras/GameplayCameras/Source/GameplayCameras/Public/Core/CameraVariableTable.h
@@ -63,7 +63,8 @@ public:
 	/** Adds a variable to the table.
 	 *
 	 * This may re-allocate the internal memory buffer. It's recommended to pre-compute
-	 * the allocation information needed for a table, and initialize it once. */
+	 * the allocation information needed for a table, and initialize it once.
+	 */
 	void AddVariable(const FCameraVariableDefinition& VariableDefinition);
 
 public:
@@ -98,6 +99,7 @@ public:
 
 	// Interpolation.
 	
+	void OverrideAll(const FCameraVariableTable& OtherTable);
 	void OverrideChanged(const FCameraVariableTable& OtherTable);
 	void OverrideChanged(const FCameraVariableTable& OtherTable, const FCameraVariableTableFlags& InMask, bool bInvertMask, FCameraVariableTableFlags& OutMask);
 
@@ -113,7 +115,7 @@ public:
 	void UnsetAllValues();
 
 	bool IsValueWrittenThisFrame(uint32 VariableId) const;
-	void ResetAllWrittenThisFrame();
+	void ClearAllWrittenThisFrameFlags();
 
 private:
 
@@ -127,8 +129,8 @@ private:
 
 	void ReallocateBuffer(uint32 NewCapacity);
 
-	void InternalOverrideChanged(const FCameraVariableTable& OtherTable, const FCameraVariableTableFlags* InMask, bool bInvertMask, FCameraVariableTableFlags* OutMask);
-	void InternalLerpChanged(const FCameraVariableTable& ToTable, float Factor, const FCameraVariableTableFlags* InMask, bool bInvertMask, FCameraVariableTableFlags* OutMask);
+	void InternalOverride(const FCameraVariableTable& OtherTable, const FCameraVariableTableFlags* InMask, bool bInvertMask, FCameraVariableTableFlags* OutMask, bool bChangedOnly);
+	void InternalLerp(const FCameraVariableTable& ToTable, float Factor, const FCameraVariableTableFlags* InMask, bool bInvertMask, FCameraVariableTableFlags* OutMask, bool bChangedOnly);
 
 private:
 
@@ -180,7 +182,7 @@ const ValueType& FCameraVariableTable::GetValue(uint32 VariableId) const
 {
 	const FEntry& Entry = Entries.FindChecked(VariableId);
 	CheckVariableType<ValueType>(Entry.Type);
-	return reinterpret_cast<ValueType*>(Memory + Entry.Offset);
+	return *reinterpret_cast<ValueType*>(Memory + Entry.Offset);
 }
 
 template<typename ValueType>
@@ -207,19 +209,22 @@ bool FCameraVariableTable::TryGetValue(uint32 VariableId, ValueType& OutValue) c
 template<typename ValueType>
 void FCameraVariableTable::SetValue(uint32 VariableId, typename TCallTraits<ValueType>::ParamType Value)
 {
-	const FEntry& Entry = Entries.FindChecked(VariableId);
+	FEntry& Entry = Entries.FindChecked(VariableId);
 	CheckVariableType<ValueType>(Entry.Type);
 	ValueType* ValuePtr = reinterpret_cast<ValueType*>(Memory + Entry.Offset);
 	*ValuePtr = Value;
+	Entry.Flags |= EEntryFlags::Written | EEntryFlags::WrittenThisFrame;
 }
 
 template<typename ValueType>
 bool FCameraVariableTable::TrySetValue(uint32 VariableId, typename TCallTraits<ValueType>::ParamType Value)
 {
-	if (const FEntry* Entry = Entries.Find(VariableId))
+	if (FEntry* Entry = Entries.Find(VariableId))
 	{
 		CheckVariableType<ValueType>(Entry->Type);
 		ValueType* ValuePtr = reinterpret_cast<ValueType*>(Memory + Entry->Offset);
+		*ValuePtr = Value;
+		Entry->Flags |= EEntryFlags::Written | EEntryFlags::WrittenThisFrame;
 		return true;
 	}
 	return false;
