Cameras: add SetValue method to the variable table

diff --git a/Engine/Plugins/Cameras/GameplayCameras/Source/GameplayCameras/Private/Core/CameraVariableAssets.cpp b/Engine/Plugins/Cameras/GameplayCameras/Source/GameplayCameras/Private/Core/CameraVariableAssets.cpp
--- a/Engine/Plugins/Cameras/GameplayCameras/Source/GameplayCameras/Private/Core/CameraVariableAssets.cpp
+++ b/Engine/Plugins/Cameras/GameplayCameras/Source/GameplayCameras/Private/Core/CameraVariableAssets.cpp
@@ -9,6 +9,17 @@ UCameraVariableAsset::UCameraVariableAsset(const FObjectInitializer& ObjectInit)
 {
 }
 
+FCameraVariableDefinition UCameraVariableAsset::GetVariableDefinition() const
+{
+	FCameraVariableDefinition VariableDefinition;
+	VariableDefinition.VariableId = GetVariableId();
+	VariableDefinition.VariableType = GetVariableType();
+#if WITH_EDITORONLY_DATA
+	VariableDefinition.VariableName = GetName();
+#endif
+	return VariableDefinition;
+}
+
 void UCameraVariableAsset::Serialize(FArchive& Ar)
 {
 	Super::Serialize(Ar);
diff --git a/Engine/Plugins/Cameras/GameplayCameras/Source/GameplayCameras/Public/Core/CameraVariableTable.h b/Engine/Plugins/Cameras/GameplayCameras/Source/GameplayCameras/Public/Core/CameraVariableTable.h
--- a/Engine/Plugins/Cameras/GameplayCameras/Source/GameplayCameras/Public/Core/CameraVariableTable.h
+++ b/Engine/Plugins/Cameras/GameplayCameras/Source/GameplayCameras/Public/Core/CameraVariableTable.h
@@ -95,6 +95,12 @@ public:
 	template<typename ValueType>
 	bool TrySetValue(uint32 VariableId, typename TCallTraits<ValueType>::ParamType Value);
 
+	template<typename VariableAssetType>
+	void SetValue(
+			const VariableAssetType* VariableAsset, 
+			typename TCallTraits<typename VariableAssetType::ValueType>::ParamType Value, 
+			bool bCreateIfMissing = false);
+
 public:
 
 	// Interpolation.
@@ -230,6 +236,28 @@ bool FCameraVariableTable::TrySetValue(uint32 VariableId, typename TCallTraits<V
 	return false;
 }
 
+template<typename VariableAssetType>
+void FCameraVariableTable::SetValue(
+		const VariableAssetType* VariableAsset, 
+		typename TCallTraits<typename VariableAssetType::ValueType>::ParamType Value, 
+		bool bCreateIfMissing)
+{
+	if (ensure(VariableAsset))
+	{
+		if (TrySetValue<typename VariableAssetType::ValueType>(VariableAsset->GetVariableId(), Value))
+		{
+			return;
+		}
+
+		if (bCreateIfMissing)
+		{
+			FCameraVariableDefinition VariableDefinition = VariableAsset->GetVariableDefinition();
+			AddVariable(VariableDefinition);
+			SetValue<typename VariableAssetType::ValueType>(VariableDefinition.VariableId, Value);
+		}
+	}
+}
+
 #define UE_CAMERA_VARIABLE_FOR_TYPE(ValueType, ValueName)\
 	template<>\
 	struct TCameraVariableTraits<ValueType>\
