Cameras: change how evaluator trees are built, and add allocation info on camera rigs

Evaluator trees are now built in 2 phases: the first phase actually constructs them, and the second phase initializes them. This separates the two different concerns, but also makes it easier to build a temporary (non-functional) evaluator tree for the purposes of checking how much memory we need to allocate the "real" one. This is what the new allocation builder class does to set the allocation on a camera rig.

In addition to that, this new allocation builder class also traverses the node tree to gather all used camera variables and figure out the optimal variable table layout.

Both the evaluator allocation info and variable table layout info are saved on the camera rig when it is built.

diff --git a/Engine/Plugins/Cameras/GameplayCameras/Source/GameplayCameras/Private/Core/BlendStackCameraNode.cpp b/Engine/Plugins/Cameras/GameplayCameras/Source/GameplayCameras/Private/Core/BlendStackCameraNode.cpp
--- a/Engine/Plugins/Cameras/GameplayCameras/Source/GameplayCameras/Private/Core/BlendStackCameraNode.cpp
+++ b/Engine/Plugins/Cameras/GameplayCameras/Source/GameplayCameras/Private/Core/BlendStackCameraNode.cpp
@@ -68,7 +68,7 @@ void FBlendStackCameraNodeEvaluator::Push(const FBlendStackCameraPushParams& Par
 	// Make the new stack entry, and use its storage buffer to build the tree of evaluators.
 	FCameraRigEntry NewEntry;
 
-	FCameraNodeEvaluatorTreeBuilderParams BuildParams;
+	FCameraNodeEvaluatorTreeBuildParams BuildParams;
 	BuildParams.RootCameraNode = EntryRootNode;
 	BuildParams.Evaluator = Params.Evaluator;
 	BuildParams.EvaluationContext = Params.EvaluationContext;
@@ -394,7 +394,7 @@ void FBlendStackCameraNodeEvaluator::OnPostBuildAsset(const FGameplayCameraAsset
 			Entry.RootNode->Blend = NewObject<UPopBlendCameraNode>(Entry.RootNode, NAME_None);
 
 			// Rebuild the evaluator tree.
-			FCameraNodeEvaluatorTreeBuilderParams BuildParams;
+			FCameraNodeEvaluatorTreeBuildParams BuildParams;
 			BuildParams.RootCameraNode = Entry.RootNode;
 			BuildParams.Evaluator = OwningEvaluator;
 			BuildParams.EvaluationContext = Entry.EvaluationContext.Pin();
diff --git a/Engine/Plugins/Cameras/GameplayCameras/Source/GameplayCameras/Private/Core/BlendStackRootCameraNode.cpp b/Engine/Plugins/Cameras/GameplayCameras/Source/GameplayCameras/Private/Core/BlendStackRootCameraNode.cpp
--- a/Engine/Plugins/Cameras/GameplayCameras/Source/GameplayCameras/Private/Core/BlendStackRootCameraNode.cpp
+++ b/Engine/Plugins/Cameras/GameplayCameras/Source/GameplayCameras/Private/Core/BlendStackRootCameraNode.cpp
@@ -46,7 +46,7 @@ FCameraNodeEvaluatorChildrenView FBlendStackRootCameraNodeEvaluator::OnGetChildr
 	return Children;
 }
 
-void FBlendStackRootCameraNodeEvaluator::OnInitialize(const FCameraNodeEvaluatorInitializeParams& Params)
+void FBlendStackRootCameraNodeEvaluator::OnBuild(const FCameraNodeEvaluatorBuildParams& Params)
 {
 	const UBlendStackRootCameraNode* RootNode = GetCameraNodeAs<UBlendStackRootCameraNode>();
 	BlendEvaluator = Params.BuildEvaluatorAs<FBlendCameraNodeEvaluator>(RootNode->Blend);
diff --git a/Engine/Plugins/Cameras/GameplayCameras/Source/GameplayCameras/Private/Core/BlendStackRootCameraNode.h b/Engine/Plugins/Cameras/GameplayCameras/Source/GameplayCameras/Private/Core/BlendStackRootCameraNode.h
--- a/Engine/Plugins/Cameras/GameplayCameras/Source/GameplayCameras/Private/Core/BlendStackRootCameraNode.h
+++ b/Engine/Plugins/Cameras/GameplayCameras/Source/GameplayCameras/Private/Core/BlendStackRootCameraNode.h
@@ -58,7 +58,7 @@ protected:
 
 	// FCameraNodeEvaluator interface.
 	virtual FCameraNodeEvaluatorChildrenView OnGetChildren() override;
-	virtual void OnInitialize(const FCameraNodeEvaluatorInitializeParams& Params) override;
+	virtual void OnBuild(const FCameraNodeEvaluatorBuildParams& Params) override;
 	virtual void OnRun(const FCameraNodeEvaluationParams& Params, FCameraNodeEvaluationResult& OutResult) override;
 
 private:
diff --git a/Engine/Plugins/Cameras/GameplayCameras/Source/GameplayCameras/Private/Core/CameraNode.cpp b/Engine/Plugins/Cameras/GameplayCameras/Source/GameplayCameras/Private/Core/CameraNode.cpp
--- a/Engine/Plugins/Cameras/GameplayCameras/Source/GameplayCameras/Private/Core/CameraNode.cpp
+++ b/Engine/Plugins/Cameras/GameplayCameras/Source/GameplayCameras/Private/Core/CameraNode.cpp
@@ -6,26 +6,14 @@
 
 #include UE_INLINE_GENERATED_CPP_BY_NAME(CameraNode)
 
-namespace UE::Cameras
-{
-
-void FCameraNodeEvaluationResult::Reset()
-{
-	CameraPose.Reset();
-	bIsCameraCut = false;
-	bIsValid = false;
-}
-
-}  // namespace UE::Cameras
-
 FCameraNodeChildrenView UCameraNode::GetChildren()
 {
 	return OnGetChildren();
 }
 
-FCameraNodeAllocationInfo UCameraNode::GetAllocationInfo() const
+void UCameraNode::BuildAllocationInfo(FCameraRigAllocationInfo& AllocationInfo) const
 {
-	return OnGetAllocationInfo();
+	OnBuildAllocationInfo(AllocationInfo);
 }
 
 FCameraNodeEvaluatorPtr UCameraNode::BuildEvaluator(FCameraNodeEvaluatorBuilder& Builder) const
diff --git a/Engine/Plugins/Cameras/GameplayCameras/Source/GameplayCameras/Private/Core/CameraNodeEvaluator.cpp b/Engine/Plugins/Cameras/GameplayCameras/Source/GameplayCameras/Private/Core/CameraNodeEvaluator.cpp
--- a/Engine/Plugins/Cameras/GameplayCameras/Source/GameplayCameras/Private/Core/CameraNodeEvaluator.cpp
+++ b/Engine/Plugins/Cameras/GameplayCameras/Source/GameplayCameras/Private/Core/CameraNodeEvaluator.cpp
@@ -10,10 +10,17 @@ namespace UE::Cameras
 
 UE_GAMEPLAY_CAMERAS_DEFINE_RTTI(FCameraNodeEvaluator)
 
-FCameraNodeEvaluator* FCameraNodeEvaluatorInitializeParams::BuildEvaluator(const UCameraNode* InNode) const
+void FCameraNodeEvaluationResult::Reset()
 {
-	FCameraNodeEvaluator* NewEvaluator = InNode->BuildEvaluator(*Builder);
-	NewEvaluator->Initialize(*this);
+	CameraPose.Reset();
+	bIsCameraCut = false;
+	bIsValid = false;
+}
+
+FCameraNodeEvaluator* FCameraNodeEvaluatorBuildParams::BuildEvaluator(const UCameraNode* InNode) const
+{
+	FCameraNodeEvaluator* NewEvaluator = InNode->BuildEvaluator(Builder);
+	NewEvaluator->Build(*this);
 	return NewEvaluator;
 }
 
@@ -49,9 +56,22 @@ void FCameraNodeEvaluator::AddReferencedObjects(FReferenceCollector& Collector)
 	}
 }
 
+void FCameraNodeEvaluator::Build(const FCameraNodeEvaluatorBuildParams& Params)
+{
+	OnBuild(Params);
+}
+
 void FCameraNodeEvaluator::Initialize(const FCameraNodeEvaluatorInitializeParams& Params)
 {
 	OnInitialize(Params);
+
+	for (FCameraNodeEvaluator* Child : GetChildren())
+	{
+		if (Child)
+		{
+			Child->Initialize(Params);
+		}
+	}
 }
 
 void FCameraNodeEvaluator::Run(const FCameraNodeEvaluationParams& Params, FCameraNodeEvaluationResult& OutResult)
diff --git a/Engine/Plugins/Cameras/GameplayCameras/Source/GameplayCameras/Private/Core/CameraNodeEvaluatorStorage.cpp b/Engine/Plugins/Cameras/GameplayCameras/Source/GameplayCameras/Private/Core/CameraNodeEvaluatorStorage.cpp
--- a/Engine/Plugins/Cameras/GameplayCameras/Source/GameplayCameras/Private/Core/CameraNodeEvaluatorStorage.cpp
+++ b/Engine/Plugins/Cameras/GameplayCameras/Source/GameplayCameras/Private/Core/CameraNodeEvaluatorStorage.cpp
@@ -35,29 +35,7 @@ FCameraNodeEvaluatorStorage::~FCameraNodeEvaluatorStorage()
 	DestroyEvaluatorTree(true);
 }
 
-FCameraNodeEvaluatorTreeAllocationInfo FCameraNodeEvaluatorStorage::ComputeTreeInfo(const UCameraRigAsset* CameraRig)
-{
-	int16 MaxSize = 0;
-	int16 MaxAlignment = 0;
-
-	// This isn't ideal, maybe we need a GetChildren API on the UCameraNode class too.
-	UPackage* Package = CameraRig->GetPackage();
-	ForEachObjectWithPackage(Package, [&MaxSize, &MaxAlignment](UObject* Object)
-	{
-		if (UCameraNode* CurNode = Cast<UCameraNode>(Object))
-		{
-			FCameraNodeAllocationInfo CurInfo = CurNode->GetAllocationInfo();
-
-			MaxAlignment = FMath::Max(MaxAlignment, CurInfo.EvaluatorInfo.Alignof);
-			MaxSize = Align(MaxSize, CurInfo.EvaluatorInfo.Alignof) + CurInfo.EvaluatorInfo.Sizeof;
-		}
-		return true; // Keep iterating.
-	});
-
-	return FCameraNodeEvaluatorTreeAllocationInfo{ MaxSize, MaxAlignment };
-}
-
-FCameraNodeEvaluatorPtr FCameraNodeEvaluatorStorage::BuildEvaluatorTree(const FCameraNodeEvaluatorTreeBuilderParams& Params)
+FCameraNodeEvaluatorPtr FCameraNodeEvaluatorStorage::BuildEvaluatorTree(const FCameraNodeEvaluatorTreeBuildParams& Params)
 {
 	if (Params.AllocationInfo)
 	{
@@ -71,13 +49,18 @@ FCameraNodeEvaluatorPtr FCameraNodeEvaluatorStorage::BuildEvaluatorTree(const FC
 	}
 
 	FCameraNodeEvaluatorBuilder Builder(*this);
+	FCameraNodeEvaluatorBuildParams BuildParams(Builder);
 
-	FCameraNodeEvaluatorInitializeParams InitParams;
-	InitParams.Evaluator = Params.Evaluator;
-	InitParams.EvaluationContext = Params.EvaluationContext;
-	InitParams.Builder = &Builder;
+	FCameraNodeEvaluatorPtr RootEvaluator = BuildParams.BuildEvaluator(Params.RootCameraNode);
+
+	if (Params.bInitialize)
+	{
+		FCameraNodeEvaluatorInitializeParams InitParams;
+		InitParams.Evaluator = Params.Evaluator;
+		InitParams.EvaluationContext = Params.EvaluationContext;
 
-	FCameraNodeEvaluatorPtr RootEvaluator = InitParams.BuildEvaluator(Params.RootCameraNode);
+		RootEvaluator->Initialize(InitParams);
+	}
 
 	return RootEvaluator;
 }
@@ -108,5 +91,25 @@ void FCameraNodeEvaluatorStorage::DestroyEvaluatorTree(bool bFreeAllocations)
 	}
 }
 
+void FCameraNodeEvaluatorStorage::GetAllocationInfo(FCameraNodeEvaluatorAllocationInfo& OutAllocationInfo)
+{
+	OutAllocationInfo.TotalSizeof = 0;
+	OutAllocationInfo.MaxAlignof = 0;
+
+	if (!Allocations.IsEmpty())
+	{
+		OutAllocationInfo.MaxAlignof = Allocations[0].Alignment;
+
+		for (FAllocation& Allocation : Allocations)
+		{
+			if (OutAllocationInfo.TotalSizeof > 0)
+			{
+				OutAllocationInfo.TotalSizeof = Align(OutAllocationInfo.TotalSizeof, Allocation.Alignment);
+			}
+			OutAllocationInfo.TotalSizeof += Allocation.Used;
+		}
+	}
+}
+
 }  // namespace UE::Cameras
 
diff --git a/Engine/Plugins/Cameras/GameplayCameras/Source/GameplayCameras/Private/Core/CameraRigAllocationInfoBuilder.cpp b/Engine/Plugins/Cameras/GameplayCameras/Source/GameplayCameras/Private/Core/CameraRigAllocationInfoBuilder.cpp
new file mode 100644
--- /dev/null
+++ b/Engine/Plugins/Cameras/GameplayCameras/Source/GameplayCameras/Private/Core/CameraRigAllocationInfoBuilder.cpp
@@ -0,0 +1,115 @@
+// Copyright Epic Games, Inc. All Rights Reserved.
+
+#include "Core/CameraRigAllocationInfoBuilder.h"
+
+#include "Core/CameraNode.h"
+#include "Core/CameraNodeEvaluatorStorage.h"
+#include "Core/CameraParameters.h"
+#include "Core/CameraVariableAssets.h"
+#include "Templates/UnrealTemplate.h"
+
+namespace UE::Cameras
+{
+
+bool FCameraRigAllocationInfoBuilder::bParameterHandlersInitialized;
+FCameraRigAllocationInfoBuilder::FParameterHandlerMap FCameraRigAllocationInfoBuilder::ParameterHandlers;
+
+void FCameraRigAllocationInfoBuilder::BuildAllocationInfo(const UCameraRigAsset* CameraRigAsset, FCameraRigAllocationInfo& AllocationInfo)
+{
+	if (!ensureMsgf(CameraRigAsset, TEXT("The provided camera rig is null!")))
+	{
+		return;
+	}
+	if (!ensureMsgf(CameraRigAsset->RootNode, TEXT("The provided camera rig has no root node!")))
+	{
+		return;
+	}
+
+	EnsureParameterHandlersInitialized();
+	TGuardValue<FCameraRigAllocationInfo*> CurrentAllocationInfoGuard(CurrentAllocationInfo, &AllocationInfo);
+
+	// Build a mock tree of evaluators.
+	FCameraNodeEvaluatorTreeBuildParams BuildParams;
+	BuildParams.RootCameraNode = CameraRigAsset->RootNode;
+	BuildParams.bInitialize = false;
+	FCameraNodeEvaluatorStorage Storage;
+	Storage.BuildEvaluatorTree(BuildParams);
+
+	// Get the size of the allocation.
+	Storage.GetAllocationInfo(AllocationInfo.EvaluatorInfo);
+
+	// Compute the rest, especially the variable table allocation info.
+	TArray<UCameraNode*> NodeStack;
+	NodeStack.Add(CameraRigAsset->RootNode);
+	while (!NodeStack.IsEmpty())
+	{
+		UCameraNode* CurrentNode = NodeStack.Pop();
+		BuildAllocationInfo(CurrentNode);
+
+		for (UCameraNode* Child : ReverseIterate(CurrentNode->GetChildren()))
+		{
+			if (Child)
+			{
+				NodeStack.Add(Child);
+			}
+		}
+	}
+}
+
+void FCameraRigAllocationInfoBuilder::EnsureParameterHandlersInitialized()
+{
+	if (!bParameterHandlersInitialized)
+	{
+		bParameterHandlersInitialized = true;
+
+#define UE_CAMERA_VARIABLE_FOR_TYPE(ValueType, ValueName)\
+		ParameterHandlers.Add(\
+				F##ValueName##CameraParameter::StaticStruct(),\
+				[](UObject* Node, FStructProperty* Property, FCameraRigAllocationInfo& AllocationInfo)\
+				{\
+					check(Property->Struct == F##ValueName##CameraParameter::StaticStruct());\
+					auto* ParameterPtr = Property->ContainerPtrToValuePtr<F##ValueName##CameraParameter>(Node);\
+					FCameraRigAllocationInfoBuilder::GatherVariable(ParameterPtr->Variable, AllocationInfo.VariableTableInfo);\
+				});
+UE_CAMERA_VARIABLE_FOR_ALL_TYPES()
+#undef UE_CAMERA_VARIABLE_FOR_TYPE
+	}
+}
+
+void FCameraRigAllocationInfoBuilder::BuildAllocationInfo(const UCameraNode* CameraNode)
+{
+	checkSlow(CurrentAllocationInfo);
+
+	// Look for properties that are camera parameters, and gather what camera variables
+	// they reference.
+	UClass* CameraNodeClass = CameraNode->GetClass();
+	for (TFieldIterator<FProperty> It(CameraNodeClass); It; ++It)
+	{
+		FStructProperty* StructProperty = CastField<FStructProperty>(*It);
+		if (!StructProperty)
+		{
+			continue;
+		}
+
+		if (FParameterHandler* ParameterHandler = ParameterHandlers.Find(StructProperty->Struct))
+		{
+			UObject* CameraNodeObj = const_cast<UObject*>(static_cast<const UObject*>(CameraNode));
+			(*ParameterHandler)(CameraNodeObj, StructProperty, *CurrentAllocationInfo);
+		}
+	}
+
+	// Let the camera node add any custom variables or extra memory.
+	CameraNode->BuildAllocationInfo(*CurrentAllocationInfo);
+}
+
+void FCameraRigAllocationInfoBuilder::GatherVariable(const UCameraVariableAsset* InVariable, FCameraVariableTableAllocationInfo& AllocationInfo)
+{
+	if (InVariable)
+	{
+		FCameraVariableDefinition VariableDefinition = InVariable->GetVariableDefinition();
+		AllocationInfo.VariableDefinitions.Add(VariableDefinition);
+	}
+}
+
+}  // namespace UE::Cameras
+
diff --git a/Engine/Plugins/Cameras/GameplayCameras/Source/GameplayCameras/Private/Core/CameraSystemEvaluator.cpp b/Engine/Plugins/Cameras/GameplayCameras/Source/GameplayCameras/Private/Core/CameraSystemEvaluator.cpp
--- a/Engine/Plugins/Cameras/GameplayCameras/Source/GameplayCameras/Private/Core/CameraSystemEvaluator.cpp
+++ b/Engine/Plugins/Cameras/GameplayCameras/Source/GameplayCameras/Private/Core/CameraSystemEvaluator.cpp
@@ -50,7 +50,7 @@ void FCameraSystemEvaluator::Initialize(const FCameraSystemEvaluatorCreateParams
 
 	ContextStack.Initialize(This);
 
-	FCameraNodeEvaluatorTreeBuilderParams BuildParams;
+	FCameraNodeEvaluatorTreeBuildParams BuildParams;
 	BuildParams.Evaluator = This;
 	BuildParams.RootCameraNode = RootNode;
 	RootEvaluator = static_cast<FRootCameraNodeEvaluator*>(RootEvaluatorStorage.BuildEvaluatorTree(BuildParams));
diff --git a/Engine/Plugins/Cameras/GameplayCameras/Source/GameplayCameras/Private/Core/DefaultRootCameraNode.cpp b/Engine/Plugins/Cameras/GameplayCameras/Source/GameplayCameras/Private/Core/DefaultRootCameraNode.cpp
--- a/Engine/Plugins/Cameras/GameplayCameras/Source/GameplayCameras/Private/Core/DefaultRootCameraNode.cpp
+++ b/Engine/Plugins/Cameras/GameplayCameras/Source/GameplayCameras/Private/Core/DefaultRootCameraNode.cpp
@@ -45,7 +45,7 @@ namespace UE::Cameras
 
 UE_DEFINE_CAMERA_NODE_EVALUATOR(FDefaultRootCameraNodeEvaluator)
 
-void FDefaultRootCameraNodeEvaluator::OnInitialize(const FCameraNodeEvaluatorInitializeParams& Params)
+void FDefaultRootCameraNodeEvaluator::OnBuild(const FCameraNodeEvaluatorBuildParams& Params)
 {
 	const UDefaultRootCameraNode* Data = GetCameraNodeAs<UDefaultRootCameraNode>();
 	BaseLayer = Params.BuildEvaluatorAs<FBlendStackCameraNodeEvaluator>(Data->BaseLayer);
diff --git a/Engine/Plugins/Cameras/GameplayCameras/Source/GameplayCameras/Private/Nodes/Common/ArrayCameraNode.cpp b/Engine/Plugins/Cameras/GameplayCameras/Source/GameplayCameras/Private/Nodes/Common/ArrayCameraNode.cpp
--- a/Engine/Plugins/Cameras/GameplayCameras/Source/GameplayCameras/Private/Nodes/Common/ArrayCameraNode.cpp
+++ b/Engine/Plugins/Cameras/GameplayCameras/Source/GameplayCameras/Private/Nodes/Common/ArrayCameraNode.cpp
@@ -16,7 +16,7 @@ class FArrayCameraNodeEvaluator : public FCameraNodeEvaluator
 protected:
 
 	virtual FCameraNodeEvaluatorChildrenView OnGetChildren() override;
-	virtual void OnInitialize(const FCameraNodeEvaluatorInitializeParams& Params) override;
+	virtual void OnBuild(const FCameraNodeEvaluatorBuildParams& Params) override;
 	virtual void OnRun(const FCameraNodeEvaluationParams& Params, FCameraNodeEvaluationResult& OutResult) override;
 
 private:
@@ -31,7 +31,7 @@ FCameraNodeEvaluatorChildrenView FArrayCameraNodeEvaluator::OnGetChildren()
 	return FCameraNodeEvaluatorChildrenView(Children);
 }
 
-void FArrayCameraNodeEvaluator::OnInitialize(const FCameraNodeEvaluatorInitializeParams& Params)
+void FArrayCameraNodeEvaluator::OnBuild(const FCameraNodeEvaluatorBuildParams& Params)
 {
 	const UArrayCameraNode* ArrayNode = GetCameraNodeAs<UArrayCameraNode>();
 	for (const UCameraNode* Child : ArrayNode->Children)
diff --git a/Engine/Plugins/Cameras/GameplayCameras/Source/GameplayCameras/Public/Core/CameraNode.h b/Engine/Plugins/Cameras/GameplayCameras/Source/GameplayCameras/Public/Core/CameraNode.h
--- a/Engine/Plugins/Cameras/GameplayCameras/Source/GameplayCameras/Public/Core/CameraNode.h
+++ b/Engine/Plugins/Cameras/GameplayCameras/Source/GameplayCameras/Public/Core/CameraNode.h
@@ -10,20 +10,7 @@
 
 #include "CameraNode.generated.h"
 
-/**
- * Structure describing various allocations needed by a camera node.
- */
-USTRUCT()
-struct FCameraNodeAllocationInfo
-{
-	GENERATED_BODY()
-
-	UPROPERTY()
-	FCameraNodeEvaluatorAllocationInfo EvaluatorInfo;
-
-	UPROPERTY()
-	FCameraVariableTableAllocationInfo VariableTableInfo;
-};
+struct FCameraRigAllocationInfo;
 
 /** View on a camera node's children. */
 using FCameraNodeChildrenView = UE::Cameras::TObjectChildrenView<TObjectPtr<UCameraNode>>;
@@ -44,7 +31,7 @@ public:
 	FCameraNodeChildrenView GetChildren();
 
 	/** Gets optional info about this node's required allocations at runtime. */
-	FCameraNodeAllocationInfo GetAllocationInfo() const;
+	void BuildAllocationInfo(FCameraRigAllocationInfo& AllocationInfo) const;
 
 	/** Builds the evaluator for this node. */
 	FCameraNodeEvaluatorPtr BuildEvaluator(FCameraNodeEvaluatorBuilder& Builder) const;
@@ -55,7 +42,7 @@ protected:
 	virtual FCameraNodeChildrenView OnGetChildren() { return FCameraNodeChildrenView(); }
 
 	/** Gets optional info about this node's required allocations at runtime. */
-	virtual FCameraNodeAllocationInfo OnGetAllocationInfo() const { return FCameraNodeAllocationInfo(); }
+	virtual void OnBuildAllocationInfo(FCameraRigAllocationInfo& AllocationInfo) const {}
 
 	/** Builds the evaluator for this node. */
 	virtual FCameraNodeEvaluatorPtr OnBuildEvaluator(FCameraNodeEvaluatorBuilder& Builder) const { return nullptr; }
diff --git a/Engine/Plugins/Cameras/GameplayCameras/Source/GameplayCameras/Public/Core/CameraNodeEvaluator.h b/Engine/Plugins/Cameras/GameplayCameras/Source/GameplayCameras/Public/Core/CameraNodeEvaluator.h
--- a/Engine/Plugins/Cameras/GameplayCameras/Source/GameplayCameras/Public/Core/CameraNodeEvaluator.h
+++ b/Engine/Plugins/Cameras/GameplayCameras/Source/GameplayCameras/Public/Core/CameraNodeEvaluator.h
@@ -56,16 +56,13 @@ struct FCameraNodeEvaluationResult
 };
 
 /**
- * Structure for initializing a camera node evaluator.
+ * Structure for building the tree of camera node evaluators.
  */
-struct FCameraNodeEvaluatorInitializeParams
+struct FCameraNodeEvaluatorBuildParams
 {
-	/** The evaluation running this evaluation. */
-	TSharedPtr<FCameraSystemEvaluator> Evaluator;
-	/** The evaluation context (if any) responsible for this branch of the evaluation. */
-	TSharedPtr<const FCameraEvaluationContext> EvaluationContext;
-	/** Builder object for building children evaluators. */
-	FCameraNodeEvaluatorBuilder* Builder = nullptr;
+	FCameraNodeEvaluatorBuildParams(FCameraNodeEvaluatorBuilder& InBuilder)
+		: Builder(InBuilder)
+	{}
 
 	/** Builds an evaluator for the given camera node. */
 	FCameraNodeEvaluator* BuildEvaluator(const UCameraNode* InNode) const;
@@ -73,6 +70,22 @@ struct FCameraNodeEvaluatorInitializeParams
 	/** Builds an evaluator for the given camera node, and down-cast it to the given type. */
 	template<typename EvaluatorType>
 	EvaluatorType* BuildEvaluatorAs(const UCameraNode* InNode) const;
+
+private:
+
+	/** Builder object for building children evaluators. */
+	FCameraNodeEvaluatorBuilder& Builder;
+};
+
+/**
+ * Structure for initializing a camera node evaluator.
+ */
+struct FCameraNodeEvaluatorInitializeParams
+{
+	/** The evaluation running this evaluation. */
+	TSharedPtr<FCameraSystemEvaluator> Evaluator;
+	/** The evaluation context (if any) responsible for this branch of the evaluation. */
+	TSharedPtr<const FCameraEvaluationContext> EvaluationContext;
 };
 
 /** View on a camera node evaluator's children. */
@@ -90,6 +103,9 @@ public:
 	FCameraNodeEvaluator();
 	virtual ~FCameraNodeEvaluator() {}
 
+	/** Called to build any children evaluators. */
+	void Build(const FCameraNodeEvaluatorBuildParams& Params);
+
 	/** Initialize this evaluator. */
 	void Initialize(const FCameraNodeEvaluatorInitializeParams& Params);
 
@@ -116,6 +132,9 @@ public:
 
 protected:
 
+	/** Called to build any children evaluators. */
+	virtual void OnBuild(const FCameraNodeEvaluatorBuildParams& Params) {}
+
 	/** Initialize this evaluator. */
 	virtual void OnInitialize(const FCameraNodeEvaluatorInitializeParams& Params) {}
 
@@ -150,7 +169,7 @@ public:
 };
 
 template<typename EvaluatorType>
-EvaluatorType* FCameraNodeEvaluatorInitializeParams::BuildEvaluatorAs(const UCameraNode* InNode) const
+EvaluatorType* FCameraNodeEvaluatorBuildParams::BuildEvaluatorAs(const UCameraNode* InNode) const
 {
 	FCameraNodeEvaluator* NewEvaluator = BuildEvaluator(InNode);
 	check(NewEvaluator);
diff --git a/Engine/Plugins/Cameras/GameplayCameras/Source/GameplayCameras/Public/Core/CameraNodeEvaluatorBuilder.h b/Engine/Plugins/Cameras/GameplayCameras/Source/GameplayCameras/Public/Core/CameraNodeEvaluatorBuilder.h
--- a/Engine/Plugins/Cameras/GameplayCameras/Source/GameplayCameras/Public/Core/CameraNodeEvaluatorBuilder.h
+++ b/Engine/Plugins/Cameras/GameplayCameras/Source/GameplayCameras/Public/Core/CameraNodeEvaluatorBuilder.h
@@ -21,8 +21,6 @@ struct FCameraNodeEvaluatorBuilder
 private:
 
 	FCameraNodeEvaluatorStorage& Storage;
-
-	friend class FCameraNodeTreeState;
 };
 
 template<typename EvaluatorType, typename ...ArgTypes>
diff --git a/Engine/Plugins/Cameras/GameplayCameras/Source/GameplayCameras/Public/Core/CameraNodeEvaluatorFwd.h b/Engine/Plugins/Cameras/GameplayCameras/Source/GameplayCameras/Public/Core/CameraNodeEvaluatorFwd.h
--- a/Engine/Plugins/Cameras/GameplayCameras/Source/GameplayCameras/Public/Core/CameraNodeEvaluatorFwd.h
+++ b/Engine/Plugins/Cameras/GameplayCameras/Source/GameplayCameras/Public/Core/CameraNodeEvaluatorFwd.h
@@ -18,17 +18,19 @@ struct FCameraNodeEvaluatorBuilder;
 // Typedef to avoid having to deal with namespaces in UCameraNode subclasses.
 using FCameraNodeEvaluatorPtr = UE::Cameras::FCameraNodeEvaluator*;
 
-/** Allocation information for a node evaluator. */
+/** Allocation information for an entire tree of node evaluators. */
 USTRUCT()
 struct FCameraNodeEvaluatorAllocationInfo
 {
 	GENERATED_BODY()
 
+	/** Total required size for the node evaluators. */
 	UPROPERTY()
-	int16 Sizeof = 0;
+	int16 TotalSizeof = 0;
 
+	/** Maximum required alignment for the node evaluators. */
 	UPROPERTY()
-	int16 Alignof = 0;
+	int16 MaxAlignof = 0;
 };
 
 /** Allocation information for a node evaluator, auto-setup for a given type. */
@@ -37,8 +39,8 @@ struct TCameraNodeEvaluatorAllocationInfo : FCameraNodeEvaluatorAllocationInfo
 {
 	TCameraNodeEvaluatorAllocationInfo() 
 	{
-		Sizeof = sizeof(EvaluatorType);
-		Alignof = alignof(EvaluatorType);
+		TotalSizeof = sizeof(EvaluatorType);
+		MaxAlignof = alignof(EvaluatorType);
 	}
 };
 
diff --git a/Engine/Plugins/Cameras/GameplayCameras/Source/GameplayCameras/Public/Core/CameraNodeEvaluatorStorage.h b/Engine/Plugins/Cameras/GameplayCameras/Source/GameplayCameras/Public/Core/CameraNodeEvaluatorStorage.h
--- a/Engine/Plugins/Cameras/GameplayCameras/Source/GameplayCameras/Public/Core/CameraNodeEvaluatorStorage.h
+++ b/Engine/Plugins/Cameras/GameplayCameras/Source/GameplayCameras/Public/Core/CameraNodeEvaluatorStorage.h
@@ -20,15 +20,8 @@ class FCameraNodeEvaluator;
 class FCameraNodeEvaluatorStorage;
 class FCameraSystemEvaluator;
 
-/** Allocation information for an entire tree of node evaluators. */
-struct FCameraNodeEvaluatorTreeAllocationInfo
-{
-	int16 TotalSizeof = 0;
-	int16 MaxAlignof = 0;
-};
-
 /** Structure for building an entire tree of node evaluators. */
-struct FCameraNodeEvaluatorTreeBuilderParams
+struct FCameraNodeEvaluatorTreeBuildParams
 {
 	/** The root node of the tree. */
 	TObjectPtr<const UCameraNode> RootCameraNode;
@@ -37,7 +30,10 @@ struct FCameraNodeEvaluatorTreeBuilderParams
 	/** The evaluation context (if any) responsible for this branch of the evaluation. */
 	TSharedPtr<const FCameraEvaluationContext> EvaluationContext;
 	/** An optional allocation information to optimize storage. */
-	FCameraNodeEvaluatorTreeAllocationInfo* AllocationInfo = nullptr;
+	FCameraNodeEvaluatorAllocationInfo* AllocationInfo = nullptr;
+
+	/** Whether to initialize the tree. */
+	bool bInitialize = true;
 };
 
 /**
@@ -57,14 +53,14 @@ public:
 
 public:
 
-	/** Compute allocation information for the given tree of camera nodes. */
-	static FCameraNodeEvaluatorTreeAllocationInfo ComputeTreeInfo(const UCameraRigAsset* CameraRig);
-
 	/** Build the tree of evaluators for the given tree of camera nodes. */
-	FCameraNodeEvaluatorPtr BuildEvaluatorTree(const FCameraNodeEvaluatorTreeBuilderParams& Params);
+	FCameraNodeEvaluatorPtr BuildEvaluatorTree(const FCameraNodeEvaluatorTreeBuildParams& Params);
 	/** Destroy any allocated evaluators. */
 	void DestroyEvaluatorTree(bool bFreeAllocations = false);
 
+	/** Computes the currently allocated totals. */
+	void GetAllocationInfo(FCameraNodeEvaluatorAllocationInfo& OutAllocationInfo);
+
 private:
 
 	template<typename EvaluatorType, typename ...ArgTypes>
@@ -75,6 +71,7 @@ private:
 	struct FAllocation
 	{
 		uint8* Memory = nullptr;
+		uint16 Alignment = 0;
 		uint16 Capacity = 0;
 		uint16 Used = 0;
 	};
@@ -122,6 +119,7 @@ EvaluatorType* FCameraNodeEvaluatorStorage::BuildEvaluator(ArgTypes&&... InArgs)
 
 		FAllocation& NewAllocation = Allocations.Emplace_GetRef();
 		NewAllocation.Memory = reinterpret_cast<uint8*>(FMemory::Malloc(NewCapacity, NewAlignment));
+		NewAllocation.Alignment = NewAlignment;
 		NewAllocation.Capacity = NewCapacity;
 		NewAllocation.Used = StateSizeof;
 
diff --git a/Engine/Plugins/Cameras/GameplayCameras/Source/GameplayCameras/Public/Core/CameraRigAllocationInfoBuilder.h b/Engine/Plugins/Cameras/GameplayCameras/Source/GameplayCameras/Public/Core/CameraRigAllocationInfoBuilder.h
new file mode 100644
--- /dev/null
+++ b/Engine/Plugins/Cameras/GameplayCameras/Source/GameplayCameras/Public/Core/CameraRigAllocationInfoBuilder.h
@@ -0,0 +1,50 @@
+// Copyright Epic Games, Inc. All Rights Reserved.
+
+#pragma once
+
+#include "Containers/Map.h"
+#include "Core/CameraNodeEvaluatorFwd.h"
+#include "Core/CameraRigAsset.h"
+#include "Core/CameraVariableTableFwd.h"
+#include "CoreTypes.h"
+#include "Templates/Function.h"
+#include "UObject/ObjectPtr.h"
+
+class FStructProperty;
+class UCameraNode;
+class UCameraVariableAsset;
+class UScriptStruct;
+
+namespace UE::Cameras
+{
+
+/**
+ * Class that can compute the allocation info of a camera rig.
+ */
+class FCameraRigAllocationInfoBuilder
+{
+public:
+
+	/** Builds the allocation info for the given rig. */
+	GAMEPLAYCAMERAS_API void BuildAllocationInfo(const UCameraRigAsset* CameraRigAsset, FCameraRigAllocationInfo& AllocationInfo);
+
+private:
+
+	static void EnsureParameterHandlersInitialized();
+	static void GatherVariable(const UCameraVariableAsset* InVariable, FCameraVariableTableAllocationInfo& AllocationInfo);
+
+	void BuildAllocationInfo(const UCameraNode* CameraNode);
+
+private:
+
+	using FParameterHandler = TFunction<void(UObject*, FStructProperty*, FCameraRigAllocationInfo&)>;
+	using FParameterHandlerMap = TMap<const UScriptStruct*, FParameterHandler>;
+
+	static bool bParameterHandlersInitialized;
+	static FParameterHandlerMap ParameterHandlers;
+
+	FCameraRigAllocationInfo* CurrentAllocationInfo = nullptr;
+};
+
+}  // namespace UE::Cameras
+
diff --git a/Engine/Plugins/Cameras/GameplayCameras/Source/GameplayCameras/Public/Core/CameraRigAsset.h b/Engine/Plugins/Cameras/GameplayCameras/Source/GameplayCameras/Public/Core/CameraRigAsset.h
--- a/Engine/Plugins/Cameras/GameplayCameras/Source/GameplayCameras/Public/Core/CameraRigAsset.h
+++ b/Engine/Plugins/Cameras/GameplayCameras/Source/GameplayCameras/Public/Core/CameraRigAsset.h
@@ -2,7 +2,9 @@
 
 #pragma once
 
+#include "Core/CameraNodeEvaluatorFwd.h"
 #include "Core/CameraRigTransition.h"
+#include "Core/CameraVariableTableFwd.h"
 #include "CoreTypes.h"
 #include "UObject/ObjectPtr.h"
 
@@ -22,6 +24,23 @@ enum class ECameraRigBuildStatus : uint8
 	Dirty
 };
 
+/**
+ * Structure describing various allocations needed by a camera node.
+ */
+USTRUCT()
+struct FCameraRigAllocationInfo
+{
+	GENERATED_BODY()
+
+	/** Allocation info for node evaluators. */
+	UPROPERTY()
+	FCameraNodeEvaluatorAllocationInfo EvaluatorInfo;
+
+	/** Allocation info for the camera variable. */
+	UPROPERTY()
+	FCameraVariableTableAllocationInfo VariableTableInfo;
+};
+
 /**
  * List of packages that contain the definition of a camera rig.
  * In most cases there's only one, but with nested assets there could be more.
@@ -57,6 +76,11 @@ public:
 	UPROPERTY(EditAnywhere, Category=Blending)
 	TArray<FCameraRigTransition> ExitTransitions;
 
+	UPROPERTY()
+	FCameraRigAllocationInfo AllocationInfo;
+
+public:
+
 	UPROPERTY(Transient)
 	ECameraRigBuildStatus BuildStatus = ECameraRigBuildStatus::Dirty;
 };
diff --git a/Engine/Plugins/Cameras/GameplayCameras/Source/GameplayCameras/Public/Core/DefaultRootCameraNode.h b/Engine/Plugins/Cameras/GameplayCameras/Source/GameplayCameras/Public/Core/DefaultRootCameraNode.h
--- a/Engine/Plugins/Cameras/GameplayCameras/Source/GameplayCameras/Public/Core/DefaultRootCameraNode.h
+++ b/Engine/Plugins/Cameras/GameplayCameras/Source/GameplayCameras/Public/Core/DefaultRootCameraNode.h
@@ -55,7 +55,7 @@ protected:
 
 	// FRootCameraNodeEvaluator interface.
 	virtual FCameraNodeEvaluatorChildrenView OnGetChildren() override;
-	virtual void OnInitialize(const FCameraNodeEvaluatorInitializeParams& Params) override;
+	virtual void OnBuild(const FCameraNodeEvaluatorBuildParams& Params) override;
 	virtual void OnRun(const FCameraNodeEvaluationParams& Params, FCameraNodeEvaluationResult& OutResult) override;
 	virtual void OnActivateCameraRig(const FActivateCameraRigParams& Params) override;
 
