Cameras: add Blueprint wrapper for an evaluation context's initial result

This makes it possible, for instance, to push values into the camera variable table in Blueprint.

diff --git a/Engine/Plugins/Cameras/GameplayCameras/Source/GameplayCameras/Private/GameFramework/GameplayCameraComponent.cpp b/Engine/Plugins/Cameras/GameplayCameras/Source/GameplayCameras/Private/GameFramework/GameplayCameraComponent.cpp
--- a/Engine/Plugins/Cameras/GameplayCameras/Source/GameplayCameras/Private/GameFramework/GameplayCameraComponent.cpp
+++ b/Engine/Plugins/Cameras/GameplayCameras/Source/GameplayCameras/Private/GameFramework/GameplayCameraComponent.cpp
@@ -7,6 +7,7 @@
 #include "Engine/StaticMesh.h"
 #include "Engine/World.h"
 #include "GameFramework/Controller.h"
+#include "GameFramework/GameplayCameraEvaluationResultBlueprintInterop.h"
 #include "GameFramework/GameplayCameraSystemActor.h"
 #include "GameFramework/GameplayCameraSystemComponent.h"
 #include "Logging/MessageLog.h"
@@ -22,6 +23,8 @@ UGameplayCameraComponent::UGameplayCameraComponent(const FObjectInitializer& Obj
 {
 	PrimaryComponentTick.bCanEverTick = true;
 
+	InitialResultInterop = ObjectInit.CreateDefaultSubobject<UGameplayCameraEvaluationResultBlueprintInterop>(this, "InitialResultInterop");
+
 #if WITH_EDITORONLY_DATA
 	if (GIsEditor && !IsRunningCommandlet())
 	{
@@ -103,6 +106,8 @@ void UGameplayCameraComponent::ActivateCamera(APlayerController* PlayerControlle
 	{
 		EvaluationContext = MakeShared<FGameplayCameraComponentEvaluationContext>();
 		EvaluationContext->Initialize(this, PlayerController);
+
+		InitialResultInterop->Setup(&EvaluationContext->GetInitialResult());
 	}
 
 	TSharedPtr<FCameraSystemEvaluator> Evaluator = CameraSystem->GetCameraSystemComponent()->GetCameraSystemEvaluator();
@@ -135,6 +140,11 @@ void UGameplayCameraComponent::DeactivateCamera(APlayerController* PlayerControl
 	Deactivate();
 }
 
+UGameplayCameraEvaluationResultBlueprintInterop* UGameplayCameraComponent::GetInitialResult() const
+{
+	return InitialResultInterop;
+}
+
 void UGameplayCameraComponent::OnRegister()
 {
 #if WITH_EDITORONLY_DATA
diff --git a/Engine/Plugins/Cameras/GameplayCameras/Source/GameplayCameras/Private/GameFramework/GameplayCameraEvaluationResultBlueprintInterop.cpp b/Engine/Plugins/Cameras/GameplayCameras/Source/GameplayCameras/Private/GameFramework/GameplayCameraEvaluationResultBlueprintInterop.cpp
new file mode 100644
--- /dev/null
+++ b/Engine/Plugins/Cameras/GameplayCameras/Source/GameplayCameras/Private/GameFramework/GameplayCameraEvaluationResultBlueprintInterop.cpp
@@ -0,0 +1,153 @@
+// Copyright Epic Games, Inc. All Rights Reserved.
+
+#include "GameFramework/GameplayCameraEvaluationResultBlueprintInterop.h"
+
+#include "Core/CameraNodeEvaluator.h"
+#include "Core/CameraPose.h"
+#include "Core/CameraVariableAssets.h"
+#include "Core/CameraVariableTable.h"
+
+#define UE_PRIVATE_CAMERA_EVALUATION_RESULT_BLUEPRINT_INTEROP_VALIDATE_RESULT(ErrorResult)\
+	if (!Result)\
+	{\
+		FFrame::KismetExecutionMessage(TEXT("No camera evaluation result has been set"), ELogVerbosity::Error);\
+		return ErrorResult;\
+	}
+
+#define UE_PRIVATE_CAMERA_EVALUATION_RESULT_BLUEPRINT_INTEROP_VALIDATE_VARIABLE_ASSET_PARAM(ErrorResult)\
+	if (!InVariableAsset)\
+	{\
+		FFrame::KismetExecutionMessage(TEXT("No camera variable asset was given"), ELogVerbosity::Error);\
+		return ErrorResult;\
+	}\
+
+#define UE_PRIVATE_CAMERA_EVALUATION_RESULT_BLUEPRINT_INTEROP_GET_VARIABLE(VariableType)\
+	static VariableType ErrorResult {};\
+	UE_PRIVATE_CAMERA_EVALUATION_RESULT_BLUEPRINT_INTEROP_VALIDATE_RESULT(ErrorResult)\
+	UE_PRIVATE_CAMERA_EVALUATION_RESULT_BLUEPRINT_INTEROP_VALIDATE_VARIABLE_ASSET_PARAM(ErrorResult)\
+	return Result->VariableTable.GetValue<VariableType>(InVariableAsset->GetVariableId());
+
+#define UE_PRIVATE_CAMERA_EVALUATION_RESULT_BLUEPRINT_INTEROP_SET_VARIABLE(VariableType)\
+	UE_PRIVATE_CAMERA_EVALUATION_RESULT_BLUEPRINT_INTEROP_VALIDATE_RESULT()\
+	UE_PRIVATE_CAMERA_EVALUATION_RESULT_BLUEPRINT_INTEROP_VALIDATE_VARIABLE_ASSET_PARAM()\
+	Result->VariableTable.SetValue(InVariableAsset, InValue, true);
+
+UGameplayCameraEvaluationResultBlueprintInterop::UGameplayCameraEvaluationResultBlueprintInterop(const FObjectInitializer& ObjectInit)
+	: Super(ObjectInit)
+{
+}
+
+void UGameplayCameraEvaluationResultBlueprintInterop::Setup(FCameraNodeEvaluationResult* InResult)
+{
+	Result = InResult;
+}
+
+void UGameplayCameraEvaluationResultBlueprintInterop::Teardown()
+{
+	Result = nullptr;
+}
+
+const FCameraPose& UGameplayCameraEvaluationResultBlueprintInterop::GetCameraPose() const
+{
+	static FCameraPose ErrorResult;
+	UE_PRIVATE_CAMERA_EVALUATION_RESULT_BLUEPRINT_INTEROP_VALIDATE_RESULT(ErrorResult);
+	return Result->CameraPose;
+}
+
+bool UGameplayCameraEvaluationResultBlueprintInterop::GetBooleanCameraVariable(UBooleanCameraVariable* InVariableAsset) const
+{
+	UE_PRIVATE_CAMERA_EVALUATION_RESULT_BLUEPRINT_INTEROP_GET_VARIABLE(bool);
+}
+
+int32 UGameplayCameraEvaluationResultBlueprintInterop::GetInteger32CameraVariable(UInteger32CameraVariable* InVariableAsset) const
+{
+	UE_PRIVATE_CAMERA_EVALUATION_RESULT_BLUEPRINT_INTEROP_GET_VARIABLE(int32);
+}
+
+float UGameplayCameraEvaluationResultBlueprintInterop::GetFloatCameraVariable(UFloatCameraVariable* InVariableAsset) const
+{
+	UE_PRIVATE_CAMERA_EVALUATION_RESULT_BLUEPRINT_INTEROP_GET_VARIABLE(float);
+}
+
+double UGameplayCameraEvaluationResultBlueprintInterop::GetDoubleCameraVariable(UDoubleCameraVariable* InVariableAsset) const
+{
+	UE_PRIVATE_CAMERA_EVALUATION_RESULT_BLUEPRINT_INTEROP_GET_VARIABLE(double);
+}
+
+FVector2D UGameplayCameraEvaluationResultBlueprintInterop::GetVector2CameraVariable(UVector2dCameraVariable* InVariableAsset) const
+{
+	UE_PRIVATE_CAMERA_EVALUATION_RESULT_BLUEPRINT_INTEROP_GET_VARIABLE(FVector2d);
+}
+
+FVector UGameplayCameraEvaluationResultBlueprintInterop::GetVector3CameraVariable(UVector3dCameraVariable* InVariableAsset) const
+{
+	UE_PRIVATE_CAMERA_EVALUATION_RESULT_BLUEPRINT_INTEROP_GET_VARIABLE(FVector3d);
+}
+
+FVector4 UGameplayCameraEvaluationResultBlueprintInterop::GetVector4CameraVariable(UVector4dCameraVariable* InVariableAsset) const
+{
+	UE_PRIVATE_CAMERA_EVALUATION_RESULT_BLUEPRINT_INTEROP_GET_VARIABLE(FVector4d);
+}
+
+FRotator UGameplayCameraEvaluationResultBlueprintInterop::GetRotatorCameraVariable(URotator3dCameraVariable* InVariableAsset) const
+{
+	UE_PRIVATE_CAMERA_EVALUATION_RESULT_BLUEPRINT_INTEROP_GET_VARIABLE(FRotator3d);
+}
+
+FTransform UGameplayCameraEvaluationResultBlueprintInterop::GetTransformCameraVariable(UTransform3dCameraVariable* InVariableAsset) const
+{
+	UE_PRIVATE_CAMERA_EVALUATION_RESULT_BLUEPRINT_INTEROP_GET_VARIABLE(FTransform3d);
+}
+
+void UGameplayCameraEvaluationResultBlueprintInterop::SetCameraPose(const FCameraPose& InCameraPose)
+{
+	UE_PRIVATE_CAMERA_EVALUATION_RESULT_BLUEPRINT_INTEROP_VALIDATE_RESULT();
+	Result->CameraPose = InCameraPose;
+	// TODO: auto-set flags based on differences
+}
+
+void UGameplayCameraEvaluationResultBlueprintInterop::SetBooleanCameraVariable(UBooleanCameraVariable* InVariableAsset, bool InValue)
+{
+	UE_PRIVATE_CAMERA_EVALUATION_RESULT_BLUEPRINT_INTEROP_SET_VARIABLE(bool);
+}
+
+void UGameplayCameraEvaluationResultBlueprintInterop::SetInteger32CameraVariable(UInteger32CameraVariable* InVariableAsset, int32 InValue)
+{
+	UE_PRIVATE_CAMERA_EVALUATION_RESULT_BLUEPRINT_INTEROP_SET_VARIABLE(int32);
+}
+
+void UGameplayCameraEvaluationResultBlueprintInterop::SetFloatCameraVariable(UFloatCameraVariable* InVariableAsset, float InValue)
+{
+	UE_PRIVATE_CAMERA_EVALUATION_RESULT_BLUEPRINT_INTEROP_SET_VARIABLE(float);
+}
+
+void UGameplayCameraEvaluationResultBlueprintInterop::SetDoubleCameraVariable(UDoubleCameraVariable* InVariableAsset, double InValue)
+{
+	UE_PRIVATE_CAMERA_EVALUATION_RESULT_BLUEPRINT_INTEROP_SET_VARIABLE(double);
+}
+
+void UGameplayCameraEvaluationResultBlueprintInterop::SetVector2CameraVariable(UVector2dCameraVariable* InVariableAsset, const FVector2D& InValue)
+{
+	UE_PRIVATE_CAMERA_EVALUATION_RESULT_BLUEPRINT_INTEROP_SET_VARIABLE(FVector2d);
+}
+
+void UGameplayCameraEvaluationResultBlueprintInterop::SetVector3CameraVariable(UVector3dCameraVariable* InVariableAsset, const FVector& InValue)
+{
+	UE_PRIVATE_CAMERA_EVALUATION_RESULT_BLUEPRINT_INTEROP_SET_VARIABLE(FVector3d);
+}
+
+void UGameplayCameraEvaluationResultBlueprintInterop::SetVector4CameraVariable(UVector4dCameraVariable* InVariableAsset, const FVector4& InValue)
+{
+	UE_PRIVATE_CAMERA_EVALUATION_RESULT_BLUEPRINT_INTEROP_SET_VARIABLE(FVector4d);
+}
+
+void UGameplayCameraEvaluationResultBlueprintInterop::SetRotatorCameraVariable(URotator3dCameraVariable* InVariableAsset, const FRotator& InValue)
+{
+	UE_PRIVATE_CAMERA_EVALUATION_RESULT_BLUEPRINT_INTEROP_SET_VARIABLE(FRotator3d);
+}
+
+void UGameplayCameraEvaluationResultBlueprintInterop::SetTransformCameraVariable(UTransform3dCameraVariable* InVariableAsset, const FTransform& InValue)
+{
+	UE_PRIVATE_CAMERA_EVALUATION_RESULT_BLUEPRINT_INTEROP_SET_VARIABLE(FTransform3d);
+}
+
diff --git a/Engine/Plugins/Cameras/GameplayCameras/Source/GameplayCameras/Public/Core/CameraEvaluationContext.h b/Engine/Plugins/Cameras/GameplayCameras/Source/GameplayCameras/Public/Core/CameraEvaluationContext.h
--- a/Engine/Plugins/Cameras/GameplayCameras/Source/GameplayCameras/Public/Core/CameraEvaluationContext.h
+++ b/Engine/Plugins/Cameras/GameplayCameras/Source/GameplayCameras/Public/Core/CameraEvaluationContext.h
@@ -41,6 +41,9 @@ public:
 	/** Gets the initial evaluation result for all camera rigs in this context. */
 	const FCameraNodeEvaluationResult& GetInitialResult() const { return InitialResult; }
 
+	/** Gets the initial evaluation result for all camera rigs in this context. */
+	FCameraNodeEvaluationResult& GetInitialResult() { return InitialResult; }
+
 public:
 
 	void AddReferencedObjects(FReferenceCollector& Collector);
diff --git a/Engine/Plugins/Cameras/GameplayCameras/Source/GameplayCameras/Public/Core/CameraPose.h b/Engine/Plugins/Cameras/GameplayCameras/Source/GameplayCameras/Public/Core/CameraPose.h
--- a/Engine/Plugins/Cameras/GameplayCameras/Source/GameplayCameras/Public/Core/CameraPose.h
+++ b/Engine/Plugins/Cameras/GameplayCameras/Source/GameplayCameras/Public/Core/CameraPose.h
@@ -67,7 +67,7 @@ public:
  * Fields are private and can only be accessed via the getters and setters.
  * The ChangedFlags structure keeps track of which fields were changed via the setters.
  */
-USTRUCT()
+USTRUCT(BlueprintType)
 struct FCameraPose
 {
 	GENERATED_BODY()
@@ -120,6 +120,7 @@ public:
 
 	/** Gets the transform of the camera. */
 	FTransform3d GetTransform() const;
+
 	/** Sets the transform of the camera. */
 	void SetTransform(FTransform3d Transform);
 
diff --git a/Engine/Plugins/Cameras/GameplayCameras/Source/GameplayCameras/Public/Core/CameraVariableAssets.h b/Engine/Plugins/Cameras/GameplayCameras/Source/GameplayCameras/Public/Core/CameraVariableAssets.h
--- a/Engine/Plugins/Cameras/GameplayCameras/Source/GameplayCameras/Public/Core/CameraVariableAssets.h
+++ b/Engine/Plugins/Cameras/GameplayCameras/Source/GameplayCameras/Public/Core/CameraVariableAssets.h
@@ -23,6 +23,8 @@ public:
 
 	uint32 GetVariableId() const { return VariableId; }
 
+	FCameraVariableDefinition GetVariableDefinition() const;
+
 	virtual ECameraVariableType GetVariableType() const PURE_VIRTUAL(UCameraVariableAsset::GetVariableType, return ECameraVariableType::Boolean;);
 	virtual const uint8* GetDefaultValuePtr() const PURE_VIRTUAL(UCameraVariableAsset::GetDefaultValuePtr, return nullptr;);
 
@@ -54,6 +56,8 @@ class GAMEPLAYCAMERAS_API UBooleanCameraVariable : public UCameraVariableAsset
 
 public:
 
+	using ValueType = bool;
+
 	virtual ECameraVariableType GetVariableType() const override { return ECameraVariableType::Boolean; }
 	virtual const uint8* GetDefaultValuePtr() const override { return reinterpret_cast<const uint8*>(&bDefaultValue); }
 
@@ -72,6 +76,8 @@ class GAMEPLAYCAMERAS_API UInteger32CameraVariable : public UCameraVariableAsset
 
 public:
 
+	using ValueType = int32;
+
 	virtual ECameraVariableType GetVariableType() const override { return ECameraVariableType::Integer32; }
 	virtual const uint8* GetDefaultValuePtr() const override { return reinterpret_cast<const uint8*>(&DefaultValue); }
 
@@ -90,6 +96,8 @@ class GAMEPLAYCAMERAS_API UFloatCameraVariable : public UCameraVariableAsset
 
 public:
 
+	using ValueType = float;
+
 	virtual ECameraVariableType GetVariableType() const override { return ECameraVariableType::Float; }
 	virtual const uint8* GetDefaultValuePtr() const override { return reinterpret_cast<const uint8*>(&DefaultValue); }
 
@@ -108,6 +116,8 @@ class GAMEPLAYCAMERAS_API UDoubleCameraVariable : public UCameraVariableAsset
 
 public:
 
+	using ValueType = double;
+
 	virtual ECameraVariableType GetVariableType() const override { return ECameraVariableType::Double; }
 	virtual const uint8* GetDefaultValuePtr() const override { return reinterpret_cast<const uint8*>(&DefaultValue); }
 
@@ -126,6 +136,8 @@ class GAMEPLAYCAMERAS_API UVector2fCameraVariable : public UCameraVariableAsset
 
 public:
 
+	using ValueType = FVector2f;
+
 	virtual ECameraVariableType GetVariableType() const override { return ECameraVariableType::Vector2f; }
 	virtual const uint8* GetDefaultValuePtr() const override { return reinterpret_cast<const uint8*>(&DefaultValue); }
 
@@ -144,6 +156,8 @@ class GAMEPLAYCAMERAS_API UVector2dCameraVariable : public UCameraVariableAsset
 
 public:
 
+	using ValueType = FVector2d;
+
 	virtual ECameraVariableType GetVariableType() const override { return ECameraVariableType::Vector2d; }
 	virtual const uint8* GetDefaultValuePtr() const override { return reinterpret_cast<const uint8*>(&DefaultValue); }
 
@@ -162,6 +176,8 @@ class GAMEPLAYCAMERAS_API UVector3fCameraVariable : public UCameraVariableAsset
 
 public:
 
+	using ValueType = FVector3f;
+
 	virtual ECameraVariableType GetVariableType() const override { return ECameraVariableType::Vector3f; }
 	virtual const uint8* GetDefaultValuePtr() const override { return reinterpret_cast<const uint8*>(&DefaultValue); }
 
@@ -180,6 +196,8 @@ class GAMEPLAYCAMERAS_API UVector3dCameraVariable : public UCameraVariableAsset
 
 public:
 
+	using ValueType = FVector3d;
+
 	virtual ECameraVariableType GetVariableType() const override { return ECameraVariableType::Vector3d; }
 	virtual const uint8* GetDefaultValuePtr() const override { return reinterpret_cast<const uint8*>(&DefaultValue); }
 
@@ -198,6 +216,8 @@ class GAMEPLAYCAMERAS_API UVector4fCameraVariable : public UCameraVariableAsset
 
 public:
 
+	using ValueType = FVector4f;
+
 	virtual ECameraVariableType GetVariableType() const override { return ECameraVariableType::Vector4f; }
 	virtual const uint8* GetDefaultValuePtr() const override { return reinterpret_cast<const uint8*>(&DefaultValue); }
 
@@ -216,6 +236,8 @@ class GAMEPLAYCAMERAS_API UVector4dCameraVariable : public UCameraVariableAsset
 
 public:
 
+	using ValueType = FVector4d;
+
 	virtual ECameraVariableType GetVariableType() const override { return ECameraVariableType::Vector4d; }
 	virtual const uint8* GetDefaultValuePtr() const override { return reinterpret_cast<const uint8*>(&DefaultValue); }
 
@@ -234,6 +256,8 @@ class GAMEPLAYCAMERAS_API URotator3fCameraVariable : public UCameraVariableAsset
 
 public:
 
+	using ValueType = FRotator3f;
+
 	virtual ECameraVariableType GetVariableType() const override { return ECameraVariableType::Rotator3f; }
 	virtual const uint8* GetDefaultValuePtr() const override { return reinterpret_cast<const uint8*>(&DefaultValue); }
 
@@ -252,6 +276,8 @@ class GAMEPLAYCAMERAS_API URotator3dCameraVariable : public UCameraVariableAsset
 
 public:
 
+	using ValueType = FRotator3d;
+
 	virtual ECameraVariableType GetVariableType() const override { return ECameraVariableType::Rotator3d; }
 	virtual const uint8* GetDefaultValuePtr() const override { return reinterpret_cast<const uint8*>(&DefaultValue); }
 
@@ -270,6 +296,8 @@ class GAMEPLAYCAMERAS_API UTransform3fCameraVariable : public UCameraVariableAss
 
 public:
 
+	using ValueType = FTransform3f;
+
 	virtual ECameraVariableType GetVariableType() const override { return ECameraVariableType::Transform3f; }
 	virtual const uint8* GetDefaultValuePtr() const override { return reinterpret_cast<const uint8*>(&DefaultValue); }
 
@@ -288,6 +316,8 @@ class GAMEPLAYCAMERAS_API UTransform3dCameraVariable : public UCameraVariableAss
 
 public:
 
+	using ValueType = FTransform3d;
+
 	virtual ECameraVariableType GetVariableType() const override { return ECameraVariableType::Transform3d; }
 	virtual const uint8* GetDefaultValuePtr() const override { return reinterpret_cast<const uint8*>(&DefaultValue); }
 
diff --git a/Engine/Plugins/Cameras/GameplayCameras/Source/GameplayCameras/Public/GameFramework/GameplayCameraComponent.h b/Engine/Plugins/Cameras/GameplayCameras/Source/GameplayCameras/Public/GameFramework/GameplayCameraComponent.h
--- a/Engine/Plugins/Cameras/GameplayCameras/Source/GameplayCameras/Public/GameFramework/GameplayCameraComponent.h
+++ b/Engine/Plugins/Cameras/GameplayCameras/Source/GameplayCameras/Public/GameFramework/GameplayCameraComponent.h
@@ -11,6 +11,7 @@
 #include "GameplayCameraComponent.generated.h"
 
 class UCameraAsset;
+class UGameplayCameraEvaluationResultBlueprintInterop;
 
 namespace UE::Cameras
 {
@@ -37,6 +38,9 @@ public:
 	UFUNCTION(BlueprintCallable, Category=Camera)
 	GAMEPLAYCAMERAS_API void DeactivateCamera();
 
+	UFUNCTION(BlueprintPure, Category=Camera)
+	GAMEPLAYCAMERAS_API UGameplayCameraEvaluationResultBlueprintInterop* GetInitialResult() const;
+
 public:
 
 	// UActorComponent interface
@@ -65,6 +69,9 @@ protected:
 	using FGameplayCameraComponentEvaluationContext = UE::Cameras::FGameplayCameraComponentEvaluationContext;
 
 	TSharedPtr<FGameplayCameraComponentEvaluationContext> EvaluationContext;
+
+	UPROPERTY()
+	TObjectPtr<UGameplayCameraEvaluationResultBlueprintInterop> InitialResultInterop;
 	
 #if WITH_EDITORONLY_DATA
 
diff --git a/Engine/Plugins/Cameras/GameplayCameras/Source/GameplayCameras/Public/GameFramework/GameplayCameraEvaluationResultBlueprintInterop.h b/Engine/Plugins/Cameras/GameplayCameras/Source/GameplayCameras/Public/GameFramework/GameplayCameraEvaluationResultBlueprintInterop.h
new file mode 100644
--- /dev/null
+++ b/Engine/Plugins/Cameras/GameplayCameras/Source/GameplayCameras/Public/GameFramework/GameplayCameraEvaluationResultBlueprintInterop.h
@@ -0,0 +1,110 @@
+// Copyright Epic Games, Inc. All Rights Reserved.
+
+#pragma once
+
+#include "CoreTypes.h"
+#include "UObject/ObjectMacros.h"
+
+#include "GameplayCameraEvaluationResultBlueprintInterop.generated.h"
+
+class UBooleanCameraVariable;
+class UDoubleCameraVariable;
+class UFloatCameraVariable;
+class UInteger32CameraVariable;
+class URotator3dCameraVariable;
+class UTransform3dCameraVariable;
+class UVector2dCameraVariable;
+class UVector3dCameraVariable;
+class UVector4dCameraVariable;
+struct FCameraPose;
+
+namespace UE::Cameras
+{
+
+struct FCameraNodeEvaluationResult;
+
+}  // namespace UE::Cameras
+
+UCLASS(BlueprintType, DisplayName="Gameplay Camera Evaluation Result")
+class UGameplayCameraEvaluationResultBlueprintInterop : public UObject
+{
+	GENERATED_BODY()
+
+public:
+
+	using FCameraNodeEvaluationResult = UE::Cameras::FCameraNodeEvaluationResult;
+
+	UGameplayCameraEvaluationResultBlueprintInterop(const FObjectInitializer& ObjectInit);
+
+	void Setup(FCameraNodeEvaluationResult* InResult);
+	void Teardown();
+
+public:
+
+	UFUNCTION(BlueprintPure, Category=Camera)
+	const FCameraPose& GetCameraPose() const;
+
+	UFUNCTION(BlueprintPure, Category=Camera, meta=(CompactNodeTitle="Get Bool"))
+	bool GetBooleanCameraVariable(UBooleanCameraVariable* InVariableAsset) const;
+
+	UFUNCTION(BlueprintPure, Category=Camera, meta=(CompactNodeTitle="Get Int"))
+	int32 GetInteger32CameraVariable(UInteger32CameraVariable* InVariableAsset) const;
+
+	UFUNCTION(BlueprintPure, Category=Camera, meta=(CompactNodeTitle="Get Float"))
+	float GetFloatCameraVariable(UFloatCameraVariable* InVariableAsset) const;
+
+	UFUNCTION(BlueprintPure, Category=Camera, meta=(CompactNodeTitle="Get Double"))
+	double GetDoubleCameraVariable(UDoubleCameraVariable* InVariableAsset) const;
+
+	UFUNCTION(BlueprintPure, Category=Camera, meta=(CompactNodeTitle="Get Vec2"))
+	FVector2D GetVector2CameraVariable(UVector2dCameraVariable* InVariableAsset) const;
+
+	UFUNCTION(BlueprintPure, Category=Camera, meta=(CompactNodeTitle="Get Vec3"))
+	FVector GetVector3CameraVariable(UVector3dCameraVariable* InVariableAsset) const;
+
+	UFUNCTION(BlueprintPure, Category=Camera, meta=(CompactNodeTitle="Get Vec4"))
+	FVector4 GetVector4CameraVariable(UVector4dCameraVariable* InVariableAsset) const;
+
+	UFUNCTION(BlueprintPure, Category=Camera, meta=(CompactNodeTitle="Get Rotator"))
+	FRotator GetRotatorCameraVariable(URotator3dCameraVariable* InVariableAsset) const;
+
+	UFUNCTION(BlueprintPure, Category=Camera, meta=(CompactNodeTitle="Get Transform"))
+	FTransform GetTransformCameraVariable(UTransform3dCameraVariable* InVariableAsset) const;
+
+public:
+
+	UFUNCTION(BlueprintCallable, Category=Camera)
+	void SetCameraPose(const FCameraPose& InCameraPose);
+
+	UFUNCTION(BlueprintCallable, Category=Camera)
+	void SetBooleanCameraVariable(UBooleanCameraVariable* InVariableAsset, bool InValue);
+
+	UFUNCTION(BlueprintCallable, Category=Camera)
+	void SetInteger32CameraVariable(UInteger32CameraVariable* InVariableAsset, int32 InValue);
+
+	UFUNCTION(BlueprintCallable, Category=Camera)
+	void SetFloatCameraVariable(UFloatCameraVariable* InVariableAsset, float InValue);
+
+	UFUNCTION(BlueprintCallable, Category=Camera)
+	void SetDoubleCameraVariable(UDoubleCameraVariable* InVariableAsset, double InValue);
+
+	UFUNCTION(BlueprintCallable, Category=Camera)
+	void SetVector2CameraVariable(UVector2dCameraVariable* InVariableAsset, const FVector2D& InValue);
+
+	UFUNCTION(BlueprintCallable, Category=Camera)
+	void SetVector3CameraVariable(UVector3dCameraVariable* InVariableAsset, const FVector& InValue);
+
+	UFUNCTION(BlueprintCallable, Category=Camera)
+	void SetVector4CameraVariable(UVector4dCameraVariable* InVariableAsset, const FVector4& InValue);
+
+	UFUNCTION(BlueprintCallable, Category=Camera)
+	void SetRotatorCameraVariable(URotator3dCameraVariable* InVariableAsset, const FRotator& InValue);
+
+	UFUNCTION(BlueprintCallable, Category=Camera)
+	void SetTransformCameraVariable(UTransform3dCameraVariable* InVariableAsset, const FTransform& InValue);
+
+private:
+
+	FCameraNodeEvaluationResult* Result = nullptr;
+};
+
