Cameras: add lens parameters node

diff --git a/Engine/Plugins/Cameras/GameplayCameras/Source/GameplayCameras/Private/Nodes/Common/LensParametersCameraNode.cpp b/Engine/Plugins/Cameras/GameplayCameras/Source/GameplayCameras/Private/Nodes/Common/LensParametersCameraNode.cpp
new file mode 100644
--- /dev/null
+++ b/Engine/Plugins/Cameras/GameplayCameras/Source/GameplayCameras/Private/Nodes/Common/LensParametersCameraNode.cpp
@@ -0,0 +1,82 @@
+// Copyright Epic Games, Inc. All Rights Reserved.
+
+#include "Nodes/Common/LensParametersCameraNode.h"
+
+#include "Core/CameraEvaluationContext.h"
+#include "Core/CameraParameterReader.h"
+#include "GameplayCameras.h"
+
+#include UE_INLINE_GENERATED_CPP_BY_NAME(LensParametersCameraNode)
+
+namespace UE::Cameras
+{
+
+class FLensParametersCameraNodeEvaluator : public FCameraNodeEvaluator
+{
+	UE_DECLARE_CAMERA_NODE_EVALUATOR(FLensParametersCameraNodeEvaluator)
+
+protected:
+
+	virtual void OnInitialize(const FCameraNodeEvaluatorInitializeParams& Params) override;
+	virtual void OnRun(const FCameraNodeEvaluationParams& Params, FCameraNodeEvaluationResult& OutResult) override;
+
+private:
+
+	FCameraParameterReader<float> SensorWidthReader;
+	FCameraParameterReader<float> SensorHeightReader;
+	FCameraParameterReader<float> FocusDistanceReader;
+	FCameraParameterReader<float> FocalLengthReader;
+	FCameraParameterReader<float> ApertureReader;
+};
+
+UE_DEFINE_CAMERA_NODE_EVALUATOR(FLensParametersCameraNodeEvaluator)
+
+void FLensParametersCameraNodeEvaluator::OnInitialize(const FCameraNodeEvaluatorInitializeParams& Params)
+{
+	const ULensParametersCameraNode* LensParametersNode = GetCameraNodeAs<ULensParametersCameraNode>();
+	SensorWidthReader.Initialize(LensParametersNode->SensorWidth);
+	SensorHeightReader.Initialize(LensParametersNode->SensorHeight);
+	FocusDistanceReader.Initialize(LensParametersNode->FocusDistance);
+	FocalLengthReader.Initialize(LensParametersNode->FocalLength);
+	ApertureReader.Initialize(LensParametersNode->Aperture);
+}
+
+void FLensParametersCameraNodeEvaluator::OnRun(const FCameraNodeEvaluationParams& Params, FCameraNodeEvaluationResult& OutResult)
+{
+	FCameraPose& OutPose = OutResult.CameraPose;
+
+	float SensorWidth = SensorWidthReader.Get(OutResult.VariableTable);
+	if (SensorWidth > 0)
+	{
+		OutPose.SetSensorWidth(SensorWidth);
+	}
+	float SensorHeight = SensorHeightReader.Get(OutResult.VariableTable);
+	if (SensorHeight > 0)
+	{
+		OutPose.SetSensorHeight(SensorHeight);
+	}
+	float FocusDistance = FocusDistanceReader.Get(OutResult.VariableTable);
+	if (FocusDistance > 0)
+	{
+		OutPose.SetFocusDistance(FocusDistance);
+	}
+	float FocalLength = FocalLengthReader.Get(OutResult.VariableTable);
+	if (FocalLength > 0)
+	{
+		OutPose.SetFocalLength(FocalLength);
+	}
+	float Aperture = ApertureReader.Get(OutResult.VariableTable);
+	if (Aperture > 0)
+	{
+		OutPose.SetAperture(Aperture);
+	}
+}
+
+}  // namespace UE::Cameras
+
+FCameraNodeEvaluatorPtr ULensParametersCameraNode::OnBuildEvaluator(FCameraNodeEvaluatorBuilder& Builder) const
+{
+	using namespace UE::Cameras;
+	return Builder.BuildEvaluator<FLensParametersCameraNodeEvaluator>();
+}
+
diff --git a/Engine/Plugins/Cameras/GameplayCameras/Source/GameplayCameras/Public/Nodes/Common/LensParametersCameraNode.h b/Engine/Plugins/Cameras/GameplayCameras/Source/GameplayCameras/Public/Nodes/Common/LensParametersCameraNode.h
new file mode 100644
--- /dev/null
+++ b/Engine/Plugins/Cameras/GameplayCameras/Source/GameplayCameras/Public/Nodes/Common/LensParametersCameraNode.h
@@ -0,0 +1,47 @@
+// Copyright Epic Games, Inc. All Rights Reserved.
+
+#pragma once
+
+#include "Core/CameraNode.h"
+#include "Core/CameraParameters.h"
+#include "Nodes/CameraNodeTypes.h"
+
+#include "LensParametersCameraNode.generated.h"
+
+/**
+ * A camera node that sets parameter values on the camera lens.
+ */
+UCLASS(MinimalAPI)
+class ULensParametersCameraNode : public UCameraNode
+{
+	GENERATED_BODY()
+
+protected:
+
+	// UCameraNode interface.
+	virtual FCameraNodeEvaluatorPtr OnBuildEvaluator(FCameraNodeEvaluatorBuilder& Builder) const override;
+
+public:
+
+	/** Horizontal size of filmback or digital sensor, in mm. */
+	UPROPERTY(EditAnywhere, Category="Filmback", meta=(ClampMin="0.001", ForceUnits=mm))
+	FFloatCameraParameter SensorWidth;
+
+	/** Vertical size of filmback or digital sensor, in mm. */
+	UPROPERTY(EditAnywhere, Category="Filmback", meta=(ClampMin="0.001", ForceUnits=mm))
+	FFloatCameraParameter SensorHeight;
+
+	/** Manually-controlled focus distance (manual focus mode only) */
+	UPROPERTY(EditAnywhere, Category="Lens Parameters", meta=(Units=cm))
+	FFloatCameraParameter FocusDistance;
+
+	/** Current focal length of the camera (i.e. controls FoV, zoom) */
+	UPROPERTY(EditAnywhere, Category="Lens Parameters")
+	FFloatCameraParameter FocalLength;
+
+	/** Current aperture, in terms of f-stop (e.g. 2.8 for f/2.8) */
+	UPROPERTY(EditAnywhere, Category="Lens Parameters")
+	FFloatCameraParameter Aperture;
+};
+
+
