Cameras: put all camera variables into one asset factory

This reduces a lot of clutter in the Content Browser's "Add" menu, and lets us move the camera assets back in the Gameplay menu instead of a sub-menu.

diff --git a/Engine/Plugins/Cameras/GameplayCameras/Source/GameplayCamerasEditor/Private/AssetTools/AssetDefinition_CameraAsset.cpp b/Engine/Plugins/Cameras/GameplayCameras/Source/GameplayCamerasEditor/Private/AssetTools/AssetDefinition_CameraAsset.cpp
--- a/Engine/Plugins/Cameras/GameplayCameras/Source/GameplayCamerasEditor/Private/AssetTools/AssetDefinition_CameraAsset.cpp
+++ b/Engine/Plugins/Cameras/GameplayCameras/Source/GameplayCamerasEditor/Private/AssetTools/AssetDefinition_CameraAsset.cpp
@@ -28,7 +28,7 @@ TSoftClassPtr<UObject> UAssetDefinition_CameraAsset::GetAssetClass() const
 
 TConstArrayView<FAssetCategoryPath> UAssetDefinition_CameraAsset::GetAssetCategories() const
 {
-	static const auto Categories = { FAssetCategoryPath(EAssetCategoryPaths::Gameplay, LOCTEXT("AssetSubCategory", "Cameras")) };
+	static const auto Categories = { FAssetCategoryPath(EAssetCategoryPaths::Gameplay) };
 	return Categories;
 }
 
diff --git a/Engine/Plugins/Cameras/GameplayCameras/Source/GameplayCamerasEditor/Private/AssetTools/AssetDefinition_CameraRigAsset.cpp b/Engine/Plugins/Cameras/GameplayCameras/Source/GameplayCamerasEditor/Private/AssetTools/AssetDefinition_CameraRigAsset.cpp
--- a/Engine/Plugins/Cameras/GameplayCameras/Source/GameplayCamerasEditor/Private/AssetTools/AssetDefinition_CameraRigAsset.cpp
+++ b/Engine/Plugins/Cameras/GameplayCameras/Source/GameplayCamerasEditor/Private/AssetTools/AssetDefinition_CameraRigAsset.cpp
@@ -28,7 +28,7 @@ TSoftClassPtr<UObject> UAssetDefinition_CameraRigAsset::GetAssetClass() const
 
 TConstArrayView<FAssetCategoryPath> UAssetDefinition_CameraRigAsset::GetAssetCategories() const
 {
-	static const auto Categories = { FAssetCategoryPath(EAssetCategoryPaths::Gameplay, LOCTEXT("AssetSubCategory", "Cameras")) };
+	static const auto Categories = { FAssetCategoryPath(EAssetCategoryPaths::Gameplay) };
 	return Categories;
 }
 
diff --git a/Engine/Plugins/Cameras/GameplayCameras/Source/GameplayCamerasEditor/Private/AssetTools/AssetDefinition_CameraVariableAssets.cpp b/Engine/Plugins/Cameras/GameplayCameras/Source/GameplayCamerasEditor/Private/AssetTools/AssetDefinition_CameraVariableAssets.cpp
--- a/Engine/Plugins/Cameras/GameplayCameras/Source/GameplayCamerasEditor/Private/AssetTools/AssetDefinition_CameraVariableAssets.cpp
+++ b/Engine/Plugins/Cameras/GameplayCameras/Source/GameplayCamerasEditor/Private/AssetTools/AssetDefinition_CameraVariableAssets.cpp
@@ -9,15 +9,25 @@
 
 #include UE_INLINE_GENERATED_CPP_BY_NAME(AssetDefinition_CameraVariableAssets)
 
+TConstArrayView<FAssetCategoryPath> UAssetDefinition_CameraVariableAsset::StaticMenuCategories()
+{
+	static const auto Categories = { FAssetCategoryPath(EAssetCategoryPaths::Gameplay) };
+	return Categories;
+}
+
 FLinearColor UAssetDefinition_CameraVariableAsset::GetAssetColor() const
 {
 	return FLinearColor(FColor(200, 80, 80));
 }
 
+TSoftClassPtr<UObject> UAssetDefinition_CameraVariableAsset::GetAssetClass() const
+{
+	return UCameraVariableAsset::StaticClass();
+}
+
 TConstArrayView<FAssetCategoryPath> UAssetDefinition_CameraVariableAsset::GetAssetCategories() const
 {
-	static const auto Categories = { FAssetCategoryPath(EAssetCategoryPaths::Gameplay, LOCTEXT("AssetSubCategory", "Cameras")) };
-	return Categories;
+	return StaticMenuCategories();
 }
 
 EAssetCommandResult UAssetDefinition_CameraVariableAsset::OpenAssets(const FAssetOpenArgs& OpenArgs) const
@@ -31,7 +41,7 @@ EAssetCommandResult UAssetDefinition_CameraVariableAsset::OpenAssets(const FAsse
 #define UE_CAMERA_VARIABLE_FOR_TYPE(ValueType, ValueName)\
 	FText UAssetDefinition_##ValueName##CameraVariable::GetAssetDisplayName() const\
 	{\
-		return LOCTEXT("AssetDisplayName", "#ValueName Camera Variable");\
+		return LOCTEXT("AssetDisplayName", #ValueName " Camera Variable");\
 	}\
 	TSoftClassPtr<UObject> UAssetDefinition_##ValueName##CameraVariable::GetAssetClass() const\
 	{\
diff --git a/Engine/Plugins/Cameras/GameplayCameras/Source/GameplayCamerasEditor/Private/AssetTools/AssetDefinition_CameraVariableAssets.h b/Engine/Plugins/Cameras/GameplayCameras/Source/GameplayCamerasEditor/Private/AssetTools/AssetDefinition_CameraVariableAssets.h
--- a/Engine/Plugins/Cameras/GameplayCameras/Source/GameplayCamerasEditor/Private/AssetTools/AssetDefinition_CameraVariableAssets.h
+++ b/Engine/Plugins/Cameras/GameplayCameras/Source/GameplayCamerasEditor/Private/AssetTools/AssetDefinition_CameraVariableAssets.h
@@ -13,8 +13,11 @@ class UAssetDefinition_CameraVariableAsset : public UAssetDefinitionDefault
 
 public:
 
+	static TConstArrayView<FAssetCategoryPath> StaticMenuCategories();
+
 	// UAssetDefinition interface
 	virtual FLinearColor GetAssetColor() const override;
+	virtual TSoftClassPtr<UObject> GetAssetClass() const override;
 	virtual TConstArrayView<FAssetCategoryPath> GetAssetCategories() const override;
 	virtual EAssetCommandResult OpenAssets(const FAssetOpenArgs& OpenArgs) const override;
 };
diff --git a/Engine/Plugins/Cameras/GameplayCameras/Source/GameplayCamerasEditor/Private/Factories/CameraVariableFactories.cpp b/Engine/Plugins/Cameras/GameplayCameras/Source/GameplayCamerasEditor/Private/Factories/CameraVariableFactories.cpp
--- a/Engine/Plugins/Cameras/GameplayCameras/Source/GameplayCamerasEditor/Private/Factories/CameraVariableFactories.cpp
+++ b/Engine/Plugins/Cameras/GameplayCameras/Source/GameplayCamerasEditor/Private/Factories/CameraVariableFactories.cpp
@@ -2,36 +2,131 @@
 
 #include "Factories/CameraVariableFactories.h"
 
+#include "AssetTools/AssetDefinition_CameraVariableAssets.h"
+#include "AssetToolsModule.h"
+#include "ClassViewerFilter.h"
+#include "ClassViewerModule.h"
 #include "Core/CameraVariableAssets.h"
+#include "IAssetTools.h"
+#include "Kismet2/SClassPickerDialog.h"
+#include "Modules/ModuleManager.h"
 
 #include UE_INLINE_GENERATED_CPP_BY_NAME(CameraVariableFactories)
 
 #define LOCTEXT_NAMESPACE "CameraVariableFactories"
 
+namespace UE::Cameras::Private
+{
+
+class FCameraVariableClassFilter : public IClassViewerFilter
+{
+public:
+	TSet<const UClass*> AllowedClasses;
+	EClassFlags DisallowedClassFlags;
+
+	FCameraVariableClassFilter()
+	{
+		AllowedClasses.Add(UCameraVariableAsset::StaticClass());
+		DisallowedClassFlags = CLASS_Abstract | CLASS_Deprecated;
+	}
+
+	virtual bool IsClassAllowed(const FClassViewerInitializationOptions& InInitOptions, const UClass* InClass, TSharedRef<FClassViewerFilterFuncs> InFilterFuncs) override
+	{
+		return 
+			!InClass->HasAnyClassFlags(DisallowedClassFlags) &&
+			InFilterFuncs->IfInChildOfClassesSet(AllowedClasses, InClass) != EFilterReturn::Failed;
+	}
+
+	virtual bool IsUnloadedClassAllowed(const FClassViewerInitializationOptions& InInitOptions, const TSharedRef< const IUnloadedBlueprintData > InUnloadedClassData, TSharedRef< FClassViewerFilterFuncs > InFilterFuncs) override
+	{
+		return 
+			!InUnloadedClassData->HasAnyClassFlags(DisallowedClassFlags) && 
+			InFilterFuncs->IfInChildOfClassesSet(AllowedClasses, InUnloadedClassData) != EFilterReturn::Failed;
+	}
+};
+
+}  // namespace UE::Cameras::Private
+
 UCameraVariableAssetFactory::UCameraVariableAssetFactory(const FObjectInitializer& ObjectInit)
 	: Super(ObjectInit)
 {
 	bCreateNew = true;
 	bEditAfterNew = true;
+	SupportedClass = UCameraVariableAsset::StaticClass();
+}
+
+FText UCameraVariableAssetFactory::GetDisplayName() const
+{
+	return LOCTEXT("DisplayName", "Camera Variable");
 }
 
-bool UCameraVariableAssetFactory::ShouldShowInNewMenu() const
+uint32 UCameraVariableAssetFactory::GetMenuCategories() const
 {
-	return true;
+	// By default, UFactory looks up the asset definition for our SupportedClass and asks it for the
+	// menu categories to use. However, we have many asset definitions (one per type of camera variable)
+	// for just one of us (one factory that creates any camera variable based on the type picker 
+	// dialog shown to the user). So we need to reimplement GetMenuCategories and make it return the
+	// category we want for all camera variables. It would be easy normally but we need to rewrite
+	// the code that converts a list of category paths to a mask of category bits, albeit in a slightly
+	// simpler manner.
+	static IAssetTools& AssetTools = FModuleManager::LoadModuleChecked<FAssetToolsModule>("AssetTools").Get();
+
+	uint32 CategoryBits = 0;
+	TConstArrayView<FAssetCategoryPath> Categories = UAssetDefinition_CameraVariableAsset::StaticMenuCategories();
+	for (const FAssetCategoryPath& Category : Categories)
+	{
+		const FName& CategoryName = Category.GetCategory();
+		ensure(CategoryName.IsValid());
+
+		const EAssetTypeCategories::Type AdvancedCategoryBit = AssetTools.FindAdvancedAssetCategory(CategoryName);
+		if (AdvancedCategoryBit == EAssetTypeCategories::Misc)
+		{
+			CategoryBits |= AssetTools.RegisterAdvancedAssetCategory(CategoryName, Category.GetCategoryText());
+		}
+		else
+		{
+			CategoryBits |= AdvancedCategoryBit;
+		}
+	}
+	return CategoryBits;
+}
+
+bool UCameraVariableAssetFactory::ConfigureProperties()
+{
+	FClassViewerModule& ClassViewerModule = FModuleManager::LoadModuleChecked<FClassViewerModule>("ClassViewer");
+
+	FClassViewerInitializationOptions Options;
+	Options.Mode = EClassViewerMode::ClassPicker;
+	Options.NameTypeToDisplay = EClassViewerNameTypeToDisplay::DisplayName;
+
+	using namespace UE::Cameras::Private;
+	TSharedPtr<FCameraVariableClassFilter> Filter = MakeShareable(new FCameraVariableClassFilter);
+	Options.ClassFilters.Add(Filter.ToSharedRef());
+
+	UClass* ChosenClass = nullptr;
+	CameraVariableAssetType = nullptr;
+	const FText TitleText = LOCTEXT("CameraVariablePicker", "Pick Camera Variable Type");
+	const bool bPressedOk = SClassPickerDialog::PickClass(TitleText, Options, ChosenClass, UCameraVariableAsset::StaticClass());
+	if (bPressedOk)
+	{
+		CameraVariableAssetType = ChosenClass;
+	}
+
+	return bPressedOk;
 }
 
-#define UE_CAMERA_VARIABLE_FOR_TYPE(ValueType, ValueName)\
-	U##ValueName##CameraVariableFactory::U##ValueName##CameraVariableFactory(const FObjectInitializer& ObjectInit)\
-		: Super(ObjectInit)\
-	{\
-		SupportedClass = U##ValueName##CameraVariable::StaticClass();\
-	}\
-	UObject* U##ValueName##CameraVariableFactory::FactoryCreateNew(UClass* Class, UObject* Parent, FName Name, EObjectFlags Flags, UObject* Context, FFeedbackContext* Warn)\
-	{\
-		return NewObject<U##ValueName##CameraVariable>(Parent, Class, Name, Flags | RF_Transactional);\
+UObject* UCameraVariableAssetFactory::FactoryCreateNew(UClass* Class, UObject* Parent, FName Name, EObjectFlags Flags, UObject* Context, FFeedbackContext* Warn)
+{
+	if (CameraVariableAssetType)
+	{
+		return NewObject<UCameraVariableAsset>(Parent, CameraVariableAssetType, Name, Flags | RF_Transactional);
+	}
+	else
+	{
+		check(Class->IsChildOf(UCameraVariableAsset::StaticClass()));
+		return NewObject<UCameraVariableAsset>(Parent, Class, Name, Flags);
 	}
-UE_CAMERA_VARIABLE_FOR_ALL_TYPES()
-#undef UE_CAMERA_VARIABLE_FOR_TYPE
+}
 
 #undef LOCTEXT_NAMESPACE
 
diff --git a/Engine/Plugins/Cameras/GameplayCameras/Source/GameplayCamerasEditor/Private/Factories/CameraVariableFactories.h b/Engine/Plugins/Cameras/GameplayCameras/Source/GameplayCamerasEditor/Private/Factories/CameraVariableFactories.h
--- a/Engine/Plugins/Cameras/GameplayCameras/Source/GameplayCamerasEditor/Private/Factories/CameraVariableFactories.h
+++ b/Engine/Plugins/Cameras/GameplayCameras/Source/GameplayCamerasEditor/Private/Factories/CameraVariableFactories.h
@@ -11,170 +11,23 @@ class UCameraVariableAsset;
 /**
  * Implements a factory for camera variable assets.
  */
-UCLASS(Abstract, hidecategories=Object)
+UCLASS(hidecategories=Object)
 class UCameraVariableAssetFactory : public UFactory
 {
 	GENERATED_BODY()
 
 public:
 
-	UCameraVariableAssetFactory(const FObjectInitializer& ObjectInit);
-
-	// UFactory Interface
-	virtual bool ShouldShowInNewMenu() const override;
-};
-
-UCLASS()
-class UBooleanCameraVariableFactory : public UCameraVariableAssetFactory
-{
-	GENERATED_BODY()
-
-	UBooleanCameraVariableFactory(const FObjectInitializer& ObjectInit);
-
-	// UFactory Interface
-	virtual UObject* FactoryCreateNew(UClass* Class, UObject* Parent, FName Name, EObjectFlags Flags, UObject* Context, FFeedbackContext* Warn) override;
-};
-
-UCLASS()
-class UInteger32CameraVariableFactory : public UCameraVariableAssetFactory
-{
-	GENERATED_BODY()
-
-	UInteger32CameraVariableFactory(const FObjectInitializer& ObjectInit);
-
-	// UFactory Interface
-	virtual UObject* FactoryCreateNew(UClass* Class, UObject* Parent, FName Name, EObjectFlags Flags, UObject* Context, FFeedbackContext* Warn) override;
-};
-
-UCLASS()
-class UFloatCameraVariableFactory : public UCameraVariableAssetFactory
-{
-	GENERATED_BODY()
-
-	UFloatCameraVariableFactory(const FObjectInitializer& ObjectInit);
-
-	// UFactory Interface
-	virtual UObject* FactoryCreateNew(UClass* Class, UObject* Parent, FName Name, EObjectFlags Flags, UObject* Context, FFeedbackContext* Warn) override;
-};
-
-UCLASS()
-class UDoubleCameraVariableFactory : public UCameraVariableAssetFactory
-{
-	GENERATED_BODY()
-
-	UDoubleCameraVariableFactory(const FObjectInitializer& ObjectInit);
-
-	// UFactory Interface
-	virtual UObject* FactoryCreateNew(UClass* Class, UObject* Parent, FName Name, EObjectFlags Flags, UObject* Context, FFeedbackContext* Warn) override;
-};
-
-UCLASS()
-class UVector2fCameraVariableFactory : public UCameraVariableAssetFactory
-{
-	GENERATED_BODY()
-
-	UVector2fCameraVariableFactory(const FObjectInitializer& ObjectInit);
-
-	// UFactory Interface
-	virtual UObject* FactoryCreateNew(UClass* Class, UObject* Parent, FName Name, EObjectFlags Flags, UObject* Context, FFeedbackContext* Warn) override;
-};
-
-UCLASS()
-class UVector2dCameraVariableFactory : public UCameraVariableAssetFactory
-{
-	GENERATED_BODY()
-
-	UVector2dCameraVariableFactory(const FObjectInitializer& ObjectInit);
-
-	// UFactory Interface
-	virtual UObject* FactoryCreateNew(UClass* Class, UObject* Parent, FName Name, EObjectFlags Flags, UObject* Context, FFeedbackContext* Warn) override;
-};
-
-UCLASS()
-class UVector3fCameraVariableFactory : public UCameraVariableAssetFactory
-{
-	GENERATED_BODY()
-
-	UVector3fCameraVariableFactory(const FObjectInitializer& ObjectInit);
+	// The type of camera variable asset to create.
+	UPROPERTY(EditAnywhere, Category=CameraVariable)
+	TSubclassOf<UCameraVariableAsset> CameraVariableAssetType;
 
-	// UFactory Interface
-	virtual UObject* FactoryCreateNew(UClass* Class, UObject* Parent, FName Name, EObjectFlags Flags, UObject* Context, FFeedbackContext* Warn) override;
-};
-
-UCLASS()
-class UVector3dCameraVariableFactory : public UCameraVariableAssetFactory
-{
-	GENERATED_BODY()
-
-	UVector3dCameraVariableFactory(const FObjectInitializer& ObjectInit);
-
-	// UFactory Interface
-	virtual UObject* FactoryCreateNew(UClass* Class, UObject* Parent, FName Name, EObjectFlags Flags, UObject* Context, FFeedbackContext* Warn) override;
-};
-
-UCLASS()
-class UVector4fCameraVariableFactory : public UCameraVariableAssetFactory
-{
-	GENERATED_BODY()
-
-	UVector4fCameraVariableFactory(const FObjectInitializer& ObjectInit);
-
-	// UFactory Interface
-	virtual UObject* FactoryCreateNew(UClass* Class, UObject* Parent, FName Name, EObjectFlags Flags, UObject* Context, FFeedbackContext* Warn) override;
-};
-
-UCLASS()
-class UVector4dCameraVariableFactory : public UCameraVariableAssetFactory
-{
-	GENERATED_BODY()
-
-	UVector4dCameraVariableFactory(const FObjectInitializer& ObjectInit);
-
-	// UFactory Interface
-	virtual UObject* FactoryCreateNew(UClass* Class, UObject* Parent, FName Name, EObjectFlags Flags, UObject* Context, FFeedbackContext* Warn) override;
-};
-
-UCLASS()
-class URotator3fCameraVariableFactory : public UCameraVariableAssetFactory
-{
-	GENERATED_BODY()
-
-	URotator3fCameraVariableFactory(const FObjectInitializer& ObjectInit);
-
-	// UFactory Interface
-	virtual UObject* FactoryCreateNew(UClass* Class, UObject* Parent, FName Name, EObjectFlags Flags, UObject* Context, FFeedbackContext* Warn) override;
-};
-
-UCLASS()
-class URotator3dCameraVariableFactory : public UCameraVariableAssetFactory
-{
-	GENERATED_BODY()
-
-	URotator3dCameraVariableFactory(const FObjectInitializer& ObjectInit);
-
-	// UFactory Interface
-	virtual UObject* FactoryCreateNew(UClass* Class, UObject* Parent, FName Name, EObjectFlags Flags, UObject* Context, FFeedbackContext* Warn) override;
-};
-
-UCLASS()
-class UTransform3fCameraVariableFactory : public UCameraVariableAssetFactory
-{
-	GENERATED_BODY()
-
-	UTransform3fCameraVariableFactory(const FObjectInitializer& ObjectInit);
-
-	// UFactory Interface
-	virtual UObject* FactoryCreateNew(UClass* Class, UObject* Parent, FName Name, EObjectFlags Flags, UObject* Context, FFeedbackContext* Warn) override;
-};
-
-UCLASS()
-class UTransform3dCameraVariableFactory : public UCameraVariableAssetFactory
-{
-	GENERATED_BODY()
-
-	UTransform3dCameraVariableFactory(const FObjectInitializer& ObjectInit);
+	UCameraVariableAssetFactory(const FObjectInitializer& ObjectInit);
 
 	// UFactory Interface
+	virtual FText GetDisplayName() const override;
+	virtual uint32 GetMenuCategories() const override;
+	virtual bool ConfigureProperties() override;
 	virtual UObject* FactoryCreateNew(UClass* Class, UObject* Parent, FName Name, EObjectFlags Flags, UObject* Context, FFeedbackContext* Warn) override;
 };
 
