Cameras: fix static analysis build

diff --git a/Engine/Plugins/Cameras/GameplayCameras/Source/GameplayCameras/Private/Core/CameraRigAllocationInfoBuilder.cpp b/Engine/Plugins/Cameras/GameplayCameras/Source/GameplayCameras/Private/Core/CameraRigAllocationInfoBuilder.cpp
--- a/Engine/Plugins/Cameras/GameplayCameras/Source/GameplayCameras/Private/Core/CameraRigAllocationInfoBuilder.cpp
+++ b/Engine/Plugins/Cameras/GameplayCameras/Source/GameplayCameras/Private/Core/CameraRigAllocationInfoBuilder.cpp
@@ -46,7 +46,8 @@ void FCameraRigAllocationInfoBuilder::BuildAllocationInfo(const UCameraRigAsset*
 		UCameraNode* CurrentNode = NodeStack.Pop();
 		BuildAllocationInfo(CurrentNode);
 
-		for (UCameraNode* Child : ReverseIterate(CurrentNode->GetChildren()))
+		FCameraNodeChildrenView CurrentChildren = CurrentNode->GetChildren();
+		for (UCameraNode* Child : ReverseIterate(CurrentChildren))
 		{
 			if (Child)
 			{
