Cameras: make generic object storage class out of the node evaluator storage

This will be useful for storing other things such as debug drawing blocks (in the next changelists).

diff --git a/Engine/Plugins/Cameras/GameplayCameras/Source/GameplayCameras/Private/Core/CameraNodeEvaluatorStorage.cpp b/Engine/Plugins/Cameras/GameplayCameras/Source/GameplayCameras/Private/Core/CameraNodeEvaluatorStorage.cpp
--- a/Engine/Plugins/Cameras/GameplayCameras/Source/GameplayCameras/Private/Core/CameraNodeEvaluatorStorage.cpp
+++ b/Engine/Plugins/Cameras/GameplayCameras/Source/GameplayCameras/Private/Core/CameraNodeEvaluatorStorage.cpp
@@ -10,42 +10,13 @@
 namespace UE::Cameras
 {
 
-FCameraNodeEvaluatorStorage::FCameraNodeEvaluatorStorage()
-{
-}
-
-FCameraNodeEvaluatorStorage::FCameraNodeEvaluatorStorage(FCameraNodeEvaluatorStorage&& Other)
-{
-	Allocations = MoveTemp(Other.Allocations);
-	EvaluatorInfos = MoveTemp(Other.EvaluatorInfos);
-}
-
-FCameraNodeEvaluatorStorage& FCameraNodeEvaluatorStorage::operator=(FCameraNodeEvaluatorStorage&& Other)
-{
-	if (ensure(this != &Other))
-	{
-		Allocations = MoveTemp(Other.Allocations);
-		EvaluatorInfos = MoveTemp(Other.EvaluatorInfos);
-	}
-	return *this;
-}
-
-FCameraNodeEvaluatorStorage::~FCameraNodeEvaluatorStorage()
-{
-	DestroyEvaluatorTree(true);
-}
-
 FCameraNodeEvaluatorPtr FCameraNodeEvaluatorStorage::BuildEvaluatorTree(const FCameraNodeEvaluatorTreeBuildParams& Params)
 {
 	if (Params.AllocationInfo)
 	{
 		const uint16 NewCapacity = Params.AllocationInfo->TotalSizeof;
 		const uint16 NewAlignment = Params.AllocationInfo->MaxAlignof;
-
-		FAllocation& NewAllocation = Allocations.Emplace_GetRef();
-		NewAllocation.Memory = reinterpret_cast<uint8*>(FMemory::Malloc(NewCapacity, NewAlignment));
-		NewAllocation.Capacity = NewCapacity;
-		NewAllocation.Used = 0;
+		Super::AllocatePage(NewCapacity, NewAlignment);
 	}
 
 	FCameraNodeEvaluatorBuilder Builder(*this);
@@ -67,48 +38,17 @@ FCameraNodeEvaluatorPtr FCameraNodeEvaluatorStorage::BuildEvaluatorTree(const FC
 
 void FCameraNodeEvaluatorStorage::DestroyEvaluatorTree(bool bFreeAllocations)
 {
-	for (FEvaluatorInfo& EvaluatorInfo : EvaluatorInfos)
-	{
-		FCameraNodeEvaluator* Evaluator(EvaluatorInfo.Ptr);
-		Evaluator->~FCameraNodeEvaluator();
-	}
-	EvaluatorInfos.Reset();
-
-	if (bFreeAllocations)
-	{
-		for (FAllocation& Allocation : Allocations)
-		{
-			FMemory::Free(Allocation.Memory);
-		}
-		Allocations.Reset();
-	}
-	else
-	{
-		for (FAllocation& Allocation : Allocations)
-		{
-			Allocation.Used = 0;
-		}
-	}
+	Super::DestroyObjects(bFreeAllocations);
 }
 
 void FCameraNodeEvaluatorStorage::GetAllocationInfo(FCameraNodeEvaluatorAllocationInfo& OutAllocationInfo)
 {
-	OutAllocationInfo.TotalSizeof = 0;
-	OutAllocationInfo.MaxAlignof = 0;
+	uint16 TotalUsed;
+	uint16 FirstAlignment;
+	Super::GetAllocationInfo(TotalUsed, FirstAlignment);
 
-	if (!Allocations.IsEmpty())
-	{
-		OutAllocationInfo.MaxAlignof = Allocations[0].Alignment;
-
-		for (FAllocation& Allocation : Allocations)
-		{
-			if (OutAllocationInfo.TotalSizeof > 0)
-			{
-				OutAllocationInfo.TotalSizeof = Align(OutAllocationInfo.TotalSizeof, Allocation.Alignment);
-			}
-			OutAllocationInfo.TotalSizeof += Allocation.Used;
-		}
-	}
+	OutAllocationInfo.TotalSizeof = TotalUsed;
+	OutAllocationInfo.MaxAlignof = FirstAlignment;
 }
 
 }  // namespace UE::Cameras
diff --git a/Engine/Plugins/Cameras/GameplayCameras/Source/GameplayCameras/Private/Core/CameraObjectStorage.cpp b/Engine/Plugins/Cameras/GameplayCameras/Source/GameplayCameras/Private/Core/CameraObjectStorage.cpp
new file mode 100644
--- /dev/null
+++ b/Engine/Plugins/Cameras/GameplayCameras/Source/GameplayCameras/Private/Core/CameraObjectStorage.cpp
@@ -0,0 +1,9 @@
+// Copyright Epic Games, Inc. All Rights Reserved.
+
+#include "Core/CameraObjectStorage.h"
+
+namespace UE::Cameras
+{
+
+}  // namespace UE::Cameras
+
diff --git a/Engine/Plugins/Cameras/GameplayCameras/Source/GameplayCameras/Public/Core/CameraNodeEvaluatorStorage.h b/Engine/Plugins/Cameras/GameplayCameras/Source/GameplayCameras/Public/Core/CameraNodeEvaluatorStorage.h
--- a/Engine/Plugins/Cameras/GameplayCameras/Source/GameplayCameras/Public/Core/CameraNodeEvaluatorStorage.h
+++ b/Engine/Plugins/Cameras/GameplayCameras/Source/GameplayCameras/Public/Core/CameraNodeEvaluatorStorage.h
@@ -3,9 +3,8 @@
 #pragma once
 
 #include "Core/CameraNodeEvaluatorFwd.h"
+#include "Core/CameraObjectStorage.h"
 #include "CoreTypes.h"
-#include "HAL/UnrealMemory.h"
-#include "Math/UnrealMathUtility.h"
 #include "UObject/ObjectPtr.h"
 
 class FReferenceCollector;
@@ -39,100 +38,30 @@ struct FCameraNodeEvaluatorTreeBuildParams
 /**
  * A class responsible for storing a tree of camera node evaluators.
  */
-class FCameraNodeEvaluatorStorage
+class FCameraNodeEvaluatorStorage : public TCameraObjectStorage<FCameraNodeEvaluator>
 {
 public:
 
-	FCameraNodeEvaluatorStorage();
-	FCameraNodeEvaluatorStorage(FCameraNodeEvaluatorStorage&& Other);
-	FCameraNodeEvaluatorStorage& operator=(FCameraNodeEvaluatorStorage&& Other);
-	~FCameraNodeEvaluatorStorage();
-
-	FCameraNodeEvaluatorStorage(const FCameraNodeEvaluatorStorage&) = delete;
-	FCameraNodeEvaluatorStorage& operator=(const FCameraNodeEvaluatorStorage&) = delete;
-
-public:
+	using Super = TCameraObjectStorage<FCameraNodeEvaluator>;
 
 	/** Build the tree of evaluators for the given tree of camera nodes. */
 	FCameraNodeEvaluatorPtr BuildEvaluatorTree(const FCameraNodeEvaluatorTreeBuildParams& Params);
+
 	/** Destroy any allocated evaluators. */
 	void DestroyEvaluatorTree(bool bFreeAllocations = false);
 
 	/** Computes the currently allocated totals. */
 	void GetAllocationInfo(FCameraNodeEvaluatorAllocationInfo& OutAllocationInfo);
 
-private:
-
+	/** Build a new evaluator inside this storage. */
 	template<typename EvaluatorType, typename ...ArgTypes>
 	EvaluatorType* BuildEvaluator(ArgTypes&&... InArgs);
-
-private:
-
-	struct FAllocation
-	{
-		uint8* Memory = nullptr;
-		uint16 Alignment = 0;
-		uint16 Capacity = 0;
-		uint16 Used = 0;
-	};
-	TArray<FAllocation> Allocations;
-
-	struct FEvaluatorInfo
-	{
-		FCameraNodeEvaluator* Ptr = nullptr;
-	};
-	TArray<FEvaluatorInfo> EvaluatorInfos;
-
-	friend struct FCameraNodeEvaluatorBuilder;
 };
 
 template<typename EvaluatorType, typename ...ArgTypes>
 EvaluatorType* FCameraNodeEvaluatorStorage::BuildEvaluator(ArgTypes&&... InArgs)
 {
-	static const uint16 DefaultCapacity = 128;
-	static const uint16 DefaultAlignment = 32;
-
-	uint16 StateSizeof = sizeof(EvaluatorType);
-	uint16 StateAlignof = alignof(EvaluatorType);
-
-	// Search for any allocation bucket that has enough room for the evaluator
-	// we want to build.
-	uint8* TargetPtr = nullptr;
-	FAllocation* TargetAllocation = nullptr;
-	for (FAllocation& Allocation : Allocations)
-	{
-		uint8* PossiblePtr = Align(Allocation.Memory + Allocation.Used, StateAlignof);
-		uint16 NewUsed = (PossiblePtr + StateSizeof) - Allocation.Memory;
-		if (NewUsed <= Allocation.Capacity)
-		{
-			TargetPtr = PossiblePtr;
-			TargetAllocation = &Allocation;
-			TargetAllocation->Used = NewUsed;
-			break;
-		}
-	}
-	// If we didn't find anything, we need to make a new allocation bucket.
-	if (TargetPtr == nullptr)
-	{
-		const uint16 NewCapacity = FMath::Max(DefaultCapacity, StateSizeof);
-		const uint16 NewAlignment = FMath::Max(DefaultAlignment, StateAlignof);
-
-		FAllocation& NewAllocation = Allocations.Emplace_GetRef();
-		NewAllocation.Memory = reinterpret_cast<uint8*>(FMemory::Malloc(NewCapacity, NewAlignment));
-		NewAllocation.Alignment = NewAlignment;
-		NewAllocation.Capacity = NewCapacity;
-		NewAllocation.Used = StateSizeof;
-
-		TargetPtr = NewAllocation.Memory;
-		TargetAllocation = &NewAllocation;
-	}
-	check(TargetPtr && TargetAllocation);
-
-	EvaluatorType* NewEvaluator = new(TargetPtr) EvaluatorType(Forward<ArgTypes>(InArgs)...);
-
-	EvaluatorInfos.Add({ NewEvaluator });
-
-	return NewEvaluator;
+	return Super::BuildObject<EvaluatorType>(Forward<ArgTypes>(InArgs)...);
 }
 
 }  // namespace UE::Cameras
diff --git a/Engine/Plugins/Cameras/GameplayCameras/Source/GameplayCameras/Public/Core/CameraObjectStorage.h b/Engine/Plugins/Cameras/GameplayCameras/Source/GameplayCameras/Public/Core/CameraObjectStorage.h
new file mode 100644
--- /dev/null
+++ b/Engine/Plugins/Cameras/GameplayCameras/Source/GameplayCameras/Public/Core/CameraObjectStorage.h
@@ -0,0 +1,238 @@
+// Copyright Epic Games, Inc. All Rights Reserved.
+
+#pragma once
+
+#include "CoreTypes.h"
+#include "HAL/UnrealMemory.h"
+#include "Math/UnrealMathUtility.h"
+#include "Templates/EnableIf.h"
+#include "Templates/PointerIsConvertibleFromTo.h"
+
+namespace UE::Cameras
+{
+
+/**
+ * Default traits for specific storages.
+ */
+template<typename BaseObjectType>
+struct TCameraObjectStorageTraits
+{
+	static const uint16 DefaultPageCapacity = 128;
+	static const uint16 DefaultPageAlignment = 32;
+};
+
+/**
+ * A utility class that allocates and stores objects of, or derived from, 
+ * a common base class. The storage is a paged buffer composed of one or
+ * more pages. If the needed storage size and alignment are known ahead
+ * of time, you can pre-allocate the first page appropriately and avoid
+ * any further paging.
+ */
+template<typename BaseObjectType>
+class TCameraObjectStorage
+{
+protected:
+
+	TCameraObjectStorage();
+	TCameraObjectStorage(TCameraObjectStorage&& Other);
+	TCameraObjectStorage& operator=(TCameraObjectStorage&& Other);
+	~TCameraObjectStorage();
+
+	TCameraObjectStorage(const TCameraObjectStorage&) = delete;
+	TCameraObjectStorage& operator=(const TCameraObjectStorage&) = delete;
+
+protected:
+
+	/**
+	 * Creates an object of the given type. Will allocate a new page buffer
+	 * if needed.
+	 */
+	template<typename ObjectType, typename ...ArgTypes>
+	typename TEnableIf<
+		TPointerIsConvertibleFromTo<ObjectType, BaseObjectType>::Value,
+		ObjectType*>
+		::Type
+	BuildObject(ArgTypes&&... InArgs);
+
+	/**
+	 * Destroys all objects in the storage.
+	 *
+	 * @param bFreeAllocations Whether to also free the memory buffers
+	 */
+	void DestroyObjects(bool bFreeAllocations = false);
+
+	/**
+	 * Computes information about the overall allocated memory.
+	 */
+	void GetAllocationInfo(uint16& OutTotalUsed, uint16& OutFirstAlignment) const;
+
+	/**
+	 * Allocates a new page buffer.
+	 */
+	void AllocatePage(uint16 InCapacity, uint16 InAlignment);
+
+private:
+
+	/** Allocation page */
+	struct FAllocation
+	{
+		uint8* Memory = nullptr;
+		uint16 Alignment = 0;
+		uint16 Capacity = 0;
+		uint16 Used = 0;
+	};
+	/** Allocated page buffers */
+	TArray<FAllocation> Allocations;
+
+	/** Pointer and info of objects inside the page buffers */
+	struct FObjectInfo
+	{
+		BaseObjectType* Ptr = nullptr;
+	};
+	/** List of built objects */
+	TArray<FObjectInfo> ObjectInfos;
+};
+
+template<typename BaseObjectType>
+TCameraObjectStorage<BaseObjectType>::TCameraObjectStorage()
+{
+}
+
+template<typename BaseObjectType>
+TCameraObjectStorage<BaseObjectType>::TCameraObjectStorage(TCameraObjectStorage&& Other)
+{
+	Allocations = MoveTemp(Other.Allocations);
+	ObjectInfos = MoveTemp(Other.ObjectInfos);
+}
+
+template<typename BaseObjectType>
+TCameraObjectStorage<BaseObjectType>& TCameraObjectStorage<BaseObjectType>::operator=(TCameraObjectStorage&& Other)
+{
+	if (ensure(this != &Other))
+	{
+		Allocations = MoveTemp(Other.Allocations);
+		ObjectInfos = MoveTemp(Other.ObjectInfos);
+	}
+	return *this;
+}
+
+template<typename BaseObjectType>
+TCameraObjectStorage<BaseObjectType>::~TCameraObjectStorage()
+{
+	DestroyObjects(true);
+}
+
+template<typename BaseObjectType>
+template<typename ObjectType, typename ...ArgTypes>
+typename TEnableIf<TPointerIsConvertibleFromTo<ObjectType, BaseObjectType>::Value, ObjectType*>::Type
+TCameraObjectStorage<BaseObjectType>::BuildObject(ArgTypes&&... InArgs)
+{
+	uint16 StateSizeof = sizeof(ObjectType);
+	uint16 StateAlignof = alignof(ObjectType);
+
+	// Search for any allocation bucket that has enough room for the object
+	// we want to build.
+	uint8* TargetPtr = nullptr;
+	FAllocation* TargetAllocation = nullptr;
+	for (FAllocation& Allocation : Allocations)
+	{
+		uint8* PossiblePtr = Align(Allocation.Memory + Allocation.Used, StateAlignof);
+		uint16 NewUsed = (PossiblePtr + StateSizeof) - Allocation.Memory;
+		if (NewUsed <= Allocation.Capacity)
+		{
+			TargetPtr = PossiblePtr;
+			TargetAllocation = &Allocation;
+			TargetAllocation->Used = NewUsed;
+			break;
+		}
+	}
+
+	// If we didn't find anything, we need to make a new allocation bucket.
+	if (TargetPtr == nullptr)
+	{
+		using FStorageTraits = TCameraObjectStorageTraits<BaseObjectType>;
+		const uint16 DefaultCapacity = FStorageTraits::DefaultPageCapacity;
+		const uint16 DefaultAlignment = FStorageTraits::DefaultPageAlignment;
+
+		const uint16 NewCapacity = FMath::Max(DefaultCapacity, StateSizeof);
+		const uint16 NewAlignment = FMath::Max(DefaultAlignment, StateAlignof);
+
+		FAllocation& NewAllocation = Allocations.Emplace_GetRef();
+		NewAllocation.Memory = reinterpret_cast<uint8*>(FMemory::Malloc(NewCapacity, NewAlignment));
+		NewAllocation.Alignment = NewAlignment;
+		NewAllocation.Capacity = NewCapacity;
+		NewAllocation.Used = StateSizeof;
+
+		TargetPtr = NewAllocation.Memory;
+		TargetAllocation = &NewAllocation;
+	}
+	check(TargetPtr && TargetAllocation);
+
+	ObjectType* NewObject = new(TargetPtr) ObjectType(Forward<ArgTypes>(InArgs)...);
+
+	ObjectInfos.Add({ NewObject });
+
+	return NewObject;
+}
+
+template<typename BaseObjectType>
+void TCameraObjectStorage<BaseObjectType>::DestroyObjects(bool bFreeAllocations)
+{
+	// Destroy the objects.
+	for (FObjectInfo& ObjectInfo : ObjectInfos)
+	{
+		BaseObjectType* ObjectPtr(ObjectInfo.Ptr);
+		ObjectPtr->~BaseObjectType();
+	}
+	ObjectInfos.Reset();
+
+	// Either destroy the allocations, or reset them to unused.
+	if (bFreeAllocations)
+	{
+		for (FAllocation& Allocation : Allocations)
+		{
+			FMemory::Free(Allocation.Memory);
+		}
+		Allocations.Reset();
+	}
+	else
+	{
+		for (FAllocation& Allocation : Allocations)
+		{
+			Allocation.Used = 0;
+		}
+	}
+}
+
+template<typename BaseObjectType>
+void TCameraObjectStorage<BaseObjectType>::GetAllocationInfo(uint16& OutTotalUsed, uint16& OutFirstAlignment) const
+{
+	OutTotalUsed = 0;
+	OutFirstAlignment = 0;
+
+	if (!Allocations.IsEmpty())
+	{
+		OutFirstAlignment = Allocations[0].Alignment;
+
+		for (const FAllocation& Allocation : Allocations)
+		{
+			if (OutTotalUsed > 0)
+			{
+				OutTotalUsed = Align(OutTotalUsed, Allocation.Alignment);
+			}
+			OutTotalUsed += Allocation.Used;
+		}
+	}
+}
+
+template<typename BaseObjectType>
+void TCameraObjectStorage<BaseObjectType>::AllocatePage(uint16 InCapacity, uint16 InAlignment)
+{
+	FAllocation& NewAllocation = Allocations.Emplace_GetRef();
+	NewAllocation.Memory = reinterpret_cast<uint8*>(FMemory::Malloc(InCapacity, InAlignment));
+	NewAllocation.Capacity = InCapacity;
+	NewAllocation.Used = 0;
+}
+
+}  // namespace UE::Cameras
+
