Cameras: don't initialize a camera system evaluator for the CDO of the camera system actor

diff --git a/Engine/Plugins/Cameras/GameplayCameras/Source/GameplayCameras/Private/GameFramework/GameplayCameraSystemComponent.cpp b/Engine/Plugins/Cameras/GameplayCameras/Source/GameplayCameras/Private/GameFramework/GameplayCameraSystemComponent.cpp
--- a/Engine/Plugins/Cameras/GameplayCameras/Source/GameplayCameras/Private/GameFramework/GameplayCameraSystemComponent.cpp
+++ b/Engine/Plugins/Cameras/GameplayCameras/Source/GameplayCameras/Private/GameFramework/GameplayCameraSystemComponent.cpp
@@ -34,17 +34,26 @@ void UGameplayCameraSystemComponent::GetCameraView(float DeltaTime, FMinimalView
 {
 	using namespace UE::Cameras;
 
-	FCameraSystemEvaluationUpdateParams UpdateParams;
-	UpdateParams.DeltaTime = DeltaTime;
-	Evaluator->Update(UpdateParams);
+	if (Evaluator.IsValid())
+	{
+		FCameraSystemEvaluationUpdateParams UpdateParams;
+		UpdateParams.DeltaTime = DeltaTime;
+		Evaluator->Update(UpdateParams);
 
-	Evaluator->GetEvaluatedCameraView(DesiredView);
+		Evaluator->GetEvaluatedCameraView(DesiredView);
+	}
 }
 
 void UGameplayCameraSystemComponent::OnRegister()
 {
 	using namespace UE::Cameras;
 
+	AActor* OwnerActor = GetOwner();
+	if (!OwnerActor || OwnerActor->HasAnyFlags(RF_ClassDefaultObject))
+	{
+		return;
+	}
+
 	if (!Evaluator.IsValid())
 	{
 		Evaluator = MakeShared<FCameraSystemEvaluator>();
