Cameras: add debug drawing framework

The camera system now maintains a storage for "debug blocks", which are a hierarchy of small structures that can print and display debug information on screen. These structures are meant to be standalone, i.e. they contain the values they need for the debug info. This is because we want to record gameplay and store these debug blocks in trace files, to be later replayed or scrubbed back.

Debug blocks are created and filled during evaluation, each frame, but since they are constantly recreated inside the same storage, there are few, or most often zero memory allocations done on any frame after the first few frames of gameplay.

Debug blocks can be defined "by hand", or can be defined using a set of macros that (hopefully) provide automatic member field serialization (again, for later recording/replay support).

A new debug renderer class provides helper functions for printing and drawing the debug info. In later changelists, it will gain more features, such as colored text, drawing clocks and graphs, and more.

The dampen position camera node is the first to implement debug drawing, as a test.

The camera system component registers itself with the engine's debug draw service to provide an actual hook for this debug drawing.

diff --git a/Engine/Plugins/Cameras/GameplayCameras/Source/GameplayCameras/Private/Core/BlendStackCameraNode.cpp b/Engine/Plugins/Cameras/GameplayCameras/Source/GameplayCameras/Private/Core/BlendStackCameraNode.cpp
--- a/Engine/Plugins/Cameras/GameplayCameras/Source/GameplayCameras/Private/Core/BlendStackCameraNode.cpp
+++ b/Engine/Plugins/Cameras/GameplayCameras/Source/GameplayCameras/Private/Core/BlendStackCameraNode.cpp
@@ -135,6 +135,10 @@ void FBlendStackCameraNodeEvaluator::OnRun(const FCameraNodeEvaluationParams& Pa
 
 		FCameraNodeEvaluationResult& CurResult(Entry.Result);
 
+#if UE_GAMEPLAY_CAMERAS_DEBUG
+		CurResult.DebugBlockBuilder = OutResult.DebugBlockBuilder;
+#endif  // UE_GAMEPLAY_CAMERAS_DEBUG
+
 		if (!Entry.bIsFrozen)
 		{
 			// If the context in which this camera rig runs doesn't have a valid result,
@@ -166,6 +170,10 @@ void FBlendStackCameraNodeEvaluator::OnRun(const FCameraNodeEvaluationParams& Pa
 			// Only evaluate the blend via the root node.
 			Entry.RootEvaluator->Run(CurParams, CurResult);
 		}
+
+#if UE_GAMEPLAY_CAMERAS_DEBUG
+		CurResult.DebugBlockBuilder = nullptr;
+#endif  // UE_GAMEPLAY_CAMERAS_DEBUG
 	}
 
 	// Now blend all the results, keeping track of blends that have reached 100% so
diff --git a/Engine/Plugins/Cameras/GameplayCameras/Source/GameplayCameras/Private/Core/CameraNodeEvaluator.cpp b/Engine/Plugins/Cameras/GameplayCameras/Source/GameplayCameras/Private/Core/CameraNodeEvaluator.cpp
--- a/Engine/Plugins/Cameras/GameplayCameras/Source/GameplayCameras/Private/Core/CameraNodeEvaluator.cpp
+++ b/Engine/Plugins/Cameras/GameplayCameras/Source/GameplayCameras/Private/Core/CameraNodeEvaluator.cpp
@@ -3,6 +3,10 @@
 #include "Core/CameraNodeEvaluator.h"
 
 #include "Core/CameraNode.h"
+#include "Debug/CameraDebugBlock.h"
+#include "Debug/CameraDebugBlockBuilder.h"
+#include "Debug/CameraDebugRenderer.h"
+#include "HAL/IConsoleManager.h"
 #include "UObject/UObjectGlobals.h"
 
 namespace UE::Cameras
@@ -10,6 +14,57 @@ namespace UE::Cameras
 
 UE_GAMEPLAY_CAMERAS_DEFINE_RTTI(FCameraNodeEvaluator)
 
+#if UE_GAMEPLAY_CAMERAS_DEBUG
+
+FString GGameplayCamerasDebugNodeTreeFilter;
+static FAutoConsoleVariableRef CVarGameplayCamerasDebugNodeTreeFilter(
+	TEXT("GameplayCameras.Debug.NodeTree.Filter"),
+	GGameplayCamerasDebugNodeTreeFilter,
+	TEXT("(Default: "". Filters the debug camera node tree by node name/type."));
+
+class FDefaultCameraNodeDebugBlock : public FCameraDebugBlock
+{
+public:
+
+	FDefaultCameraNodeDebugBlock() {}
+	FDefaultCameraNodeDebugBlock(const UCameraNode* InNode)
+	{
+		NodeClassName = InNode ? InNode->GetClass()->GetName() : TEXT("<null node>");
+	}
+
+protected:
+
+	virtual EDebugDrawResult OnDebugDraw(const FCameraDebugBlockDrawParams& Params, FCameraDebugRenderer& Renderer) override
+	{
+		Renderer.AddIndent();
+		
+		const bool bDoDebugDraw = (GGameplayCamerasDebugNodeTreeFilter.IsEmpty() ||
+				NodeClassName.Contains(GGameplayCamerasDebugNodeTreeFilter));
+		if (bDoDebugDraw)
+		{
+			Renderer.AddText(TEXT("[%s] "), *NodeClassName);
+		}
+		
+		return bDoDebugDraw ?  EDebugDrawResult::Default : EDebugDrawResult::SkipChildren;
+	}
+
+	virtual void OnPostDebugDraw(const FCameraDebugBlockDrawParams& Params, FCameraDebugRenderer& Renderer) override
+	{
+		Renderer.RemoveIndent();
+	}
+
+	virtual void OnSerialize(FArchive& Ar)
+	{
+		Ar << NodeClassName;
+	}
+
+private:
+
+	FString NodeClassName;
+};
+
+#endif  // UE_GAMEPLAY_CAMERAS_DEBUG
+
 void FCameraNodeEvaluationResult::Reset()
 {
 	CameraPose.Reset();
@@ -76,11 +131,39 @@ void FCameraNodeEvaluator::Initialize(const FCameraNodeEvaluatorInitializeParams
 
 void FCameraNodeEvaluator::Run(const FCameraNodeEvaluationParams& Params, FCameraNodeEvaluationResult& OutResult)
 {
-	if (PrivateCameraNode->bIsEnabled)
+	if (!PrivateCameraNode || PrivateCameraNode->bIsEnabled)
 	{
+#if UE_GAMEPLAY_CAMERAS_DEBUG
+		if (OutResult.DebugBlockBuilder)
+		{
+			CreateDebugBlock(Params, *OutResult.DebugBlockBuilder);
+		}
+#endif  // UE_GAMEPLAY_CAMERAS_DEBUG
+
 		OnRun(Params, OutResult);
+
+#if UE_GAMEPLAY_CAMERAS_DEBUG
+		if (OutResult.DebugBlockBuilder)
+		{
+			OutResult.DebugBlockBuilder->EndBlock();
+		}
+#endif  // UE_GAMEPLAY_CAMERAS_DEBUG
 	}
 }
 
+#if UE_GAMEPLAY_CAMERAS_DEBUG
+
+void FCameraNodeEvaluator::CreateDebugBlock(const FCameraNodeEvaluationParams& Params, FCameraDebugBlockBuilder& Builder)
+{
+	OnCreateDebugBlock(Params, Builder);
+}
+
+void FCameraNodeEvaluator::OnCreateDebugBlock(const FCameraNodeEvaluationParams& Params, FCameraDebugBlockBuilder& Builder)
+{
+	Builder.StartBlock<FDefaultCameraNodeDebugBlock>(PrivateCameraNode);
+}
+
+#endif  // UE_GAMEPLAY_CAMERAS_DEBUG
+
 }  // namespace UE::Cameras
 
diff --git a/Engine/Plugins/Cameras/GameplayCameras/Source/GameplayCameras/Private/Core/CameraSystemEvaluator.cpp b/Engine/Plugins/Cameras/GameplayCameras/Source/GameplayCameras/Private/Core/CameraSystemEvaluator.cpp
--- a/Engine/Plugins/Cameras/GameplayCameras/Source/GameplayCameras/Private/Core/CameraSystemEvaluator.cpp
+++ b/Engine/Plugins/Cameras/GameplayCameras/Source/GameplayCameras/Private/Core/CameraSystemEvaluator.cpp
@@ -8,16 +8,48 @@
 #include "Core/CameraEvaluationContext.h"
 #include "Core/CameraRigAsset.h"
 #include "Core/DefaultRootCameraNode.h"
-#include "GameplayCameras.h"
+#include "Debug/CameraDebugBlock.h"
+#include "Debug/CameraDebugBlockBuilder.h"
+#include "Debug/CameraDebugRenderer.h"
+#include "HAL/IConsoleManager.h"
 #include "IGameplayCamerasModule.h"
 #include "UObject/Package.h"
 #include "UObject/UObjectGlobals.h"
 
 DECLARE_CYCLE_STAT(TEXT("Camera System Eval"), CameraSystemEval_Total, STATGROUP_CameraSystem);
 
+#if UE_GAMEPLAY_CAMERAS_DEBUG
+
+bool GGameplayCamerasDebugEnable = false;
+static FAutoConsoleVariableRef CVarGameplayCamerasDebugEnable(
+	TEXT("GameplayCameras.Debug.Enable"),
+	GGameplayCamerasDebugEnable,
+	TEXT("(Default: false. Enables debug drawing for the GamplayCameras system."));
+
+#endif  // UE_GAMEPLAY_CAMERAS_DEBUG
+
 namespace UE::Cameras
 {
 
+#if UE_GAMEPLAY_CAMERAS_DEBUG
+
+class FRootCameraDebugBlock : public FCameraDebugBlock
+{
+public:
+
+	float LastDeltaTime = 0.f;
+
+protected:
+
+	virtual EDebugDrawResult OnDebugDraw(const FCameraDebugBlockDrawParams& Params, FCameraDebugRenderer& Renderer) override
+	{
+		return GGameplayCamerasDebugEnable ?
+			EDebugDrawResult::Default : EDebugDrawResult::SkipChildren;
+	}
+};
+
+#endif  // UE_GAMEPLAY_CAMERAS_DEBUG
+
 FCameraSystemEvaluator::FCameraSystemEvaluator()
 {
 }
@@ -117,18 +149,34 @@ void FCameraSystemEvaluator::Update(const FCameraSystemEvaluationUpdateParams& P
 		}
 	}
 
-	// Run the root camera node.
+	// Setup the params/result for running the root camera node.
 	FCameraNodeEvaluationParams NodeParams;
 	NodeParams.Evaluator = This;
 	NodeParams.DeltaTime = Params.DeltaTime;
 
 	RootNodeResult.Reset();
 
+#if UE_GAMEPLAY_CAMERAS_DEBUG
+	// Clear previous frame's debug info and make room for this frame's.
+	DebugBlockStorage.DestroyDebugBlocks();
+	RootDebugBlock = DebugBlockStorage.BuildDebugBlock<FRootCameraDebugBlock>();
+	RootDebugBlock->LastDeltaTime = Params.DeltaTime;
+
+	FCameraDebugBlockBuilder DebugBlockBuilder(DebugBlockStorage, *RootDebugBlock);
+	RootNodeResult.DebugBlockBuilder = &DebugBlockBuilder;
+#endif  // UE_GAMEPLAY_CAMERAS_DEBUG
+
+	// Run the root camera node.
 	RootEvaluator->Run(NodeParams, RootNodeResult);
 
+	// Harvest the result.
 	Result.CameraPose = RootNodeResult.CameraPose;
 	Result.bIsCameraCut = RootNodeResult.bIsCameraCut;
 	Result.bIsValid = true;
+
+#if UE_GAMEPLAY_CAMERAS_DEBUG
+	RootNodeResult.DebugBlockBuilder = nullptr;
+#endif  // UE_GAMEPLAY_CAMERAS_DEBUG
 }
 
 void FCameraSystemEvaluator::GetEvaluatedCameraView(FMinimalViewInfo& DesiredView)
@@ -139,5 +187,22 @@ void FCameraSystemEvaluator::GetEvaluatedCameraView(FMinimalViewInfo& DesiredVie
 	DesiredView.FOV = CameraPose.GetEffectiveFieldOfView();
 }
 
+#if UE_GAMEPLAY_CAMERAS_DEBUG
+
+void FCameraSystemEvaluator::DebugDraw(const FCameraSystemDebugDrawParams& Params)
+{
+	if (RootDebugBlock)
+	{
+		FCameraDebugBlockDrawParams BlockParams;
+		BlockParams.DeltaTime = RootDebugBlock->LastDeltaTime;
+
+		FCameraDebugRenderer Renderer(Params.Canvas);
+
+		RootDebugBlock->DebugDraw(BlockParams, Renderer);
+	}
+}
+
+#endif  // UE_GAMEPLAY_CAMERAS_DEBUG
+
 }  // namespace UE::Cameras
 
diff --git a/Engine/Plugins/Cameras/GameplayCameras/Source/GameplayCameras/Private/Debug/CameraDebugBlock.cpp b/Engine/Plugins/Cameras/GameplayCameras/Source/GameplayCameras/Private/Debug/CameraDebugBlock.cpp
new file mode 100644
--- /dev/null
+++ b/Engine/Plugins/Cameras/GameplayCameras/Source/GameplayCameras/Private/Debug/CameraDebugBlock.cpp
@@ -0,0 +1,48 @@
+// Copyright Epic Games, Inc. All Rights Reserved.
+
+#include "Debug/CameraDebugBlock.h"
+
+#if UE_GAMEPLAY_CAMERAS_DEBUG
+
+namespace UE::Cameras
+{
+
+TArray<FCameraDebugBlockField*> FCameraDebugBlock::StaticFields;
+
+void FCameraDebugBlock::AddChild(FCameraDebugBlock* InChild)
+{
+	Children.Add(InChild);
+}
+
+void FCameraDebugBlock::DebugDraw(const FCameraDebugBlockDrawParams& Params, FCameraDebugRenderer& Renderer)
+{
+	EDebugDrawResult Result = OnDebugDraw(Params, Renderer);
+	if (!EnumHasAnyFlags(Result, EDebugDrawResult::SkipChildren))
+	{
+		for (FCameraDebugBlock* Child : Children)
+		{
+			Child->DebugDraw(Params, Renderer);
+		}
+	}
+	OnPostDebugDraw(Params, Renderer);
+}
+
+EDebugDrawResult FCameraDebugBlock::OnDebugDraw(const FCameraDebugBlockDrawParams& Params, FCameraDebugRenderer& Renderer)
+{
+	return EDebugDrawResult::Default;
+}
+
+void FCameraDebugBlock::Serialize(FArchive& Ar)
+{
+	for (FCameraDebugBlockField* StaticField : StaticFields)
+	{
+		StaticField->SerializeField(this, Ar);
+	}
+
+	OnSerialize(Ar);
+}
+
+}  // namespace UE::Cameras
+
+#endif  // UE_GAMEPLAY_CAMERAS_DEBUG
+
diff --git a/Engine/Plugins/Cameras/GameplayCameras/Source/GameplayCameras/Private/Debug/CameraDebugBlockBuilder.cpp b/Engine/Plugins/Cameras/GameplayCameras/Source/GameplayCameras/Private/Debug/CameraDebugBlockBuilder.cpp
new file mode 100644
--- /dev/null
+++ b/Engine/Plugins/Cameras/GameplayCameras/Source/GameplayCameras/Private/Debug/CameraDebugBlockBuilder.cpp
@@ -0,0 +1,36 @@
+// Copyright Epic Games, Inc. All Rights Reserved.
+
+#include "Debug/CameraDebugBlockBuilder.h"
+
+#if UE_GAMEPLAY_CAMERAS_DEBUG
+
+namespace UE::Cameras
+{
+
+FCameraDebugBlockBuilder::FCameraDebugBlockBuilder(FCameraDebugBlockStorage& InStorage, FCameraDebugBlock& InRootBlock)
+	: Storage(InStorage)
+{
+	CurrentHierarchy.Add(&InRootBlock);
+}
+
+void FCameraDebugBlockBuilder::OnStartBlock(FCameraDebugBlock& InNewBlock)
+{
+	checkf(CurrentHierarchy.Num() > 0, TEXT("There should always be the root block in the current hierarchy!"));
+
+	// Add the new block under the current parent.
+	FCameraDebugBlock* ParentBlock = CurrentHierarchy.Last();
+	ParentBlock->AddChild(&InNewBlock);
+	CurrentHierarchy.Add(&InNewBlock);
+}
+
+void FCameraDebugBlockBuilder::EndBlock()
+{
+	checkf(CurrentHierarchy.Num() > 0, TEXT("Can't end block, no current block defined!"));
+
+	CurrentHierarchy.Pop();
+}
+
+}  // namespace UE::Cameras
+
+#endif  // UE_GAMEPLAY_CAMERAS_DEBUG
+
diff --git a/Engine/Plugins/Cameras/GameplayCameras/Source/GameplayCameras/Private/Debug/CameraDebugBlockStorage.cpp b/Engine/Plugins/Cameras/GameplayCameras/Source/GameplayCameras/Private/Debug/CameraDebugBlockStorage.cpp
new file mode 100644
--- /dev/null
+++ b/Engine/Plugins/Cameras/GameplayCameras/Source/GameplayCameras/Private/Debug/CameraDebugBlockStorage.cpp
@@ -0,0 +1,14 @@
+// Copyright Epic Games, Inc. All Rights Reserved.
+
+#include "Debug/CameraDebugBlockStorage.h"
+
+namespace UE::Cameras
+{
+
+void FCameraDebugBlockStorage::DestroyDebugBlocks(bool bFreeAllocations)
+{
+	Super::DestroyObjects(bFreeAllocations);
+}
+
+}  // namespace UE::Cameras
+
diff --git a/Engine/Plugins/Cameras/GameplayCameras/Source/GameplayCameras/Private/Debug/CameraDebugRenderer.cpp b/Engine/Plugins/Cameras/GameplayCameras/Source/GameplayCameras/Private/Debug/CameraDebugRenderer.cpp
new file mode 100644
--- /dev/null
+++ b/Engine/Plugins/Cameras/GameplayCameras/Source/GameplayCameras/Private/Debug/CameraDebugRenderer.cpp
@@ -0,0 +1,115 @@
+// Copyright Epic Games, Inc. All Rights Reserved.
+
+#include "Debug/CameraDebugRenderer.h"
+
+#include "Engine/Canvas.h"
+#include "Engine/Engine.h"
+#include "Engine/Font.h"
+#include "HAL/IConsoleManager.h"
+
+#if UE_GAMEPLAY_CAMERAS_DEBUG
+
+int32 GGameplayCamerasDebugMargin = 10;
+static FAutoConsoleVariableRef CVarGameplayCamerasDebugMargin(
+	TEXT("GameplayCameras.Debug.Margin"),
+	GGameplayCamerasDebugMargin,
+	TEXT("(Default: 10px. The margin for rendering Gameplay Cameras debug text."));
+
+int32 GGameplayCamerasDebugIndent = 20;
+static FAutoConsoleVariableRef CVarGameplayCamerasDebugIndent(
+	TEXT("GameplayCameras.Debug.Indent"),
+	GGameplayCamerasDebugIndent,
+	TEXT("(Default: 20px. The indent for rendering Gameplay Cameras debug text."));
+
+namespace UE::Cameras
+{
+
+FCameraDebugRenderer::FCameraDebugRenderer(UCanvas* InCanvas)
+	: Canvas(InCanvas)
+{
+	OriginalDrawColor = Canvas->DrawColor;
+	Canvas->SetDrawColor(FColor::White);
+
+	RenderFont = GEngine->GetSmallFont();
+	MaxCharHeight = RenderFont->GetMaxCharHeight();
+
+	NextBlockY = 0;
+}
+
+FCameraDebugRenderer::~FCameraDebugRenderer()
+{
+	FlushText();
+	Canvas->SetDrawColor(OriginalDrawColor);
+}
+
+void FCameraDebugRenderer::AddText(const FString& InString)
+{
+	AddTextImpl(*InString);
+}
+
+void FCameraDebugRenderer::AddText(const TCHAR* Fmt, ...)
+{
+	va_list Args;
+	va_start(Args, Fmt);
+	AddTextFmtImpl(Fmt, Args);
+	va_end(Args);
+}
+
+void FCameraDebugRenderer::AddTextFmtImpl(const TCHAR* Fmt, va_list Args)
+{
+	Formatter.Reset();
+	Formatter.AppendV(Fmt, Args);
+	const TCHAR* Message = Formatter.ToString();
+	AddTextImpl(Message);
+}
+
+void FCameraDebugRenderer::AddTextImpl(const TCHAR* Buffer)
+{
+	LineBuilder.Append(Buffer);
+}
+
+void FCameraDebugRenderer::FlushText()
+{
+	if (LineBuilder.Len() > 0)
+	{
+		if (NextBlockY < Canvas->ClipY)
+		{
+			int32 BlockHeight, BlockWidth;
+			const TCHAR* TextBlock = LineBuilder.ToString();
+			RenderFont->GetStringHeightAndWidth(TextBlock, BlockHeight, BlockWidth);
+
+			float X = GGameplayCamerasDebugMargin + IndentLevel * GGameplayCamerasDebugIndent;
+			float Y = GGameplayCamerasDebugMargin + NextBlockY;
+			Canvas->DrawText(RenderFont, TextBlock, X, Y);
+
+			NextBlockY += BlockHeight;
+		}
+		// else: text is going off-screen.
+
+		LineBuilder.Reset();
+	}
+}
+
+void FCameraDebugRenderer::AddIndent()
+{
+	// Flush any remaining text we have on the current indent level.
+	FlushText();
+
+	++IndentLevel;
+}
+
+void FCameraDebugRenderer::RemoveIndent()
+{
+	// Flush any remaining text we have on the current indent level.
+	FlushText();
+
+	if (ensureMsgf(IndentLevel > 0, TEXT("Can't go into negative indenting!")))
+	{
+		--IndentLevel;
+	}
+}
+
+}  // namespace UE::Cameras
+
+#endif  // UE_GAMEPLAY_CAMERAS_DEBUG
+
diff --git a/Engine/Plugins/Cameras/GameplayCameras/Source/GameplayCameras/Private/GameFramework/GameplayCameraSystemComponent.cpp b/Engine/Plugins/Cameras/GameplayCameras/Source/GameplayCameras/Private/GameFramework/GameplayCameraSystemComponent.cpp
--- a/Engine/Plugins/Cameras/GameplayCameras/Source/GameplayCameras/Private/GameFramework/GameplayCameraSystemComponent.cpp
+++ b/Engine/Plugins/Cameras/GameplayCameras/Source/GameplayCameras/Private/GameFramework/GameplayCameraSystemComponent.cpp
@@ -3,6 +3,7 @@
 #include "GameFramework/GameplayCameraSystemComponent.h"
 
 #include "Components/StaticMeshComponent.h"
+#include "Debug/DebugDrawService.h"
 #include "Engine/StaticMesh.h"
 #include "Engine/World.h"
 #include "GameFramework/Controller.h"
@@ -60,6 +61,11 @@ void UGameplayCameraSystemComponent::OnRegister()
 		Evaluator->Initialize(this);
 	}
 
+#if UE_GAMEPLAY_CAMERAS_DEBUG
+	DebugDrawDelegateHandle = UDebugDrawService::Register(
+			TEXT("Game"), FDebugDrawDelegate::CreateUObject(this, &UGameplayCameraSystemComponent::DebugDraw));
+#endif  // UE_GAMEPLAY_CAMERAS_DEBUG
+
 #if WITH_EDITORONLY_DATA
 	if (PreviewMesh && !PreviewMeshComponent)
 	{
@@ -87,6 +93,14 @@ void UGameplayCameraSystemComponent::OnComponentDestroyed(bool bDestroyingHierar
 {
 	Super::OnComponentDestroyed(bDestroyingHierarchy);
 
+#if UE_GAMEPLAY_CAMERAS_DEBUG
+	if (DebugDrawDelegateHandle.IsValid())
+	{
+		UDebugDrawService::Unregister(DebugDrawDelegateHandle);
+		DebugDrawDelegateHandle.Reset();
+	}
+#endif  // UE_GAMEPLAY_CAMERAS_DEBUG
+
 #if WITH_EDITORONLY_DATA
 	if (PreviewMeshComponent)
 	{
@@ -126,5 +140,19 @@ void UGameplayCameraSystemComponent::OnEndViewTarget()
 {
 }
 
+#if UE_ENABLE_DEBUG_DRAWING
+void UGameplayCameraSystemComponent::DebugDraw(UCanvas* Canvas, APlayerController* PlayController)
+{
+	using namespace UE::Cameras;
+
+	if (Evaluator.IsValid())
+	{
+		FCameraSystemDebugDrawParams DebugDrawParams;
+		DebugDrawParams.Canvas = Canvas;
+		Evaluator->DebugDraw(DebugDrawParams);
+	}
+}
+#endif  // UE_ENABLE_DEBUG_DRAWING
+
 #undef LOCTEXT_NAMESPACE
 
diff --git a/Engine/Plugins/Cameras/GameplayCameras/Source/GameplayCameras/Private/Nodes/Common/DampenPositionCameraNode.cpp b/Engine/Plugins/Cameras/GameplayCameras/Source/GameplayCameras/Private/Nodes/Common/DampenPositionCameraNode.cpp
--- a/Engine/Plugins/Cameras/GameplayCameras/Source/GameplayCameras/Private/Nodes/Common/DampenPositionCameraNode.cpp
+++ b/Engine/Plugins/Cameras/GameplayCameras/Source/GameplayCameras/Private/Nodes/Common/DampenPositionCameraNode.cpp
@@ -4,6 +4,8 @@
 
 #include "Core/CameraEvaluationContext.h"
 #include "Core/CameraNodeEvaluator.h"
+#include "Debug/CameraDebugBlockBuilder.h"
+#include "Debug/CameraDebugRenderer.h"
 #include "GameplayCameras.h"
 #include "Math/CriticalDamper.h"
 #include "Templates/Tuple.h"
@@ -33,6 +35,17 @@ private:
 
 UE_DEFINE_CAMERA_NODE_EVALUATOR(FDampenPositionCameraNodeEvaluator)
 
+UE_DEFINE_CAMERA_DEBUG_BLOCK_START(FDampenPositionCameraDebugBlock)
+	UE_DEFINE_CAMERA_DEBUG_BLOCK_FIELD(float, ForwardX0);
+	UE_DEFINE_CAMERA_DEBUG_BLOCK_FIELD(float, LateralX0);
+	UE_DEFINE_CAMERA_DEBUG_BLOCK_FIELD(float, VerticalX0);
+	UE_DEFINE_CAMERA_DEBUG_BLOCK_FIELD(float, ForwardDampingFactor);
+	UE_DEFINE_CAMERA_DEBUG_BLOCK_FIELD(float, LateralDampingFactor);
+	UE_DEFINE_CAMERA_DEBUG_BLOCK_FIELD(float, VerticalDampingFactor);
+	UE_DEFINE_CAMERA_DEBUG_BLOCK_FIELD(FVector3d, UndampedPosition);
+	UE_DEFINE_CAMERA_DEBUG_BLOCK_FIELD(FVector3d, DampedPosition);
+UE_DEFINE_CAMERA_DEBUG_BLOCK_END()
+
 void FDampenPositionCameraNodeEvaluator::OnInitialize(const FCameraNodeEvaluatorInitializeParams& Params)
 {
 	const UDampenPositionCameraNode* DampenNode = GetCameraNodeAs<UDampenPositionCameraNode>();
@@ -100,6 +113,23 @@ void FDampenPositionCameraNodeEvaluator::OnRun(const FCameraNodeEvaluationParams
 		}
 		
 		NextLocation = NewDampedLocation;
+
+#if UE_GAMEPLAY_CAMERAS_DEBUG
+		if (OutResult.DebugBlockBuilder)
+		{
+			FDampenPositionCameraDebugBlock& DebugBlock = OutResult.DebugBlockBuilder->StartBlock<FDampenPositionCameraDebugBlock>();
+			{
+				DebugBlock.ForwardX0 = ForwardDamper.GetX0();
+				DebugBlock.LateralX0 = LateralDamper.GetX0();
+				DebugBlock.VerticalX0 = VerticalDamper.GetX0();
+
+				DebugBlock.ForwardDampingFactor = ForwardDamper.GetW0();
+				DebugBlock.LateralDampingFactor = LateralDamper.GetW0();
+				DebugBlock.VerticalDampingFactor = VerticalDamper.GetW0();
+			}
+			OutResult.DebugBlockBuilder->EndBlock();
+		}
+#endif  // UE_GAMEPLAY_CAMERAS_DEBUG
 	}
 
 	PreviousLocation = NextLocation;
@@ -107,6 +137,20 @@ void FDampenPositionCameraNodeEvaluator::OnRun(const FCameraNodeEvaluationParams
 	OutResult.CameraPose.SetLocation(NextLocation);
 }
 
+#if UE_GAMEPLAY_CAMERAS_DEBUG
+
+EDebugDrawResult FDampenPositionCameraDebugBlock::OnDebugDraw(const FCameraDebugBlockDrawParams& Params, FCameraDebugRenderer& Renderer)
+{
+	Renderer.AddText(
+			TEXT("forward %.3f (factor %.3f)  lateral %.3f (factor %.3f)  vertical %.3f (factor %.3f)"),
+			ForwardX0, ForwardDampingFactor,
+			LateralX0, LateralDampingFactor,
+			VerticalX0, VerticalDampingFactor);
+	return EDebugDrawResult::Default;
+}
+
+#endif  // UE_GAMEPLAY_CAMERAS_DEBUG
+
 }  // namespace UE::Cameras
 
 UDampenPositionCameraNode::UDampenPositionCameraNode(const FObjectInitializer& ObjectInit)
diff --git a/Engine/Plugins/Cameras/GameplayCameras/Source/GameplayCameras/Public/Core/CameraNodeEvaluator.h b/Engine/Plugins/Cameras/GameplayCameras/Source/GameplayCameras/Public/Core/CameraNodeEvaluator.h
--- a/Engine/Plugins/Cameras/GameplayCameras/Source/GameplayCameras/Public/Core/CameraNodeEvaluator.h
+++ b/Engine/Plugins/Cameras/GameplayCameras/Source/GameplayCameras/Public/Core/CameraNodeEvaluator.h
@@ -7,6 +7,7 @@
 #include "Core/CameraVariableTable.h"
 #include "Core/ObjectChildrenView.h"
 #include "CoreTypes.h"
+#include "GameplayCameras.h"
 #include "UObject/ObjectPtr.h"
 
 class FReferenceCollector;
@@ -20,6 +21,11 @@ class FCameraNodeEvaluator;
 class FCameraSystemEvaluator;
 struct FCameraNodeEvaluatorBuilder;
 
+#if UE_GAMEPLAY_CAMERAS_DEBUG
+class FCameraDebugBlock;
+struct FCameraDebugBlockBuilder;
+#endif
+
 /**
  * Parameter structure for running a camera node evaluator.
  */
@@ -53,6 +59,11 @@ struct FCameraNodeEvaluationResult
 
 	/** Reset this result to its default (non-valid) state. */
 	void Reset();
+
+#if UE_GAMEPLAY_CAMERAS_DEBUG
+	/** Builder for debug drawing blocks. May be null. */
+	FCameraDebugBlockBuilder* DebugBlockBuilder = nullptr;
+#endif  // UE_GAMEPLAY_CAMERAS_DEBUG
 };
 
 /**
@@ -147,6 +158,22 @@ protected:
 	/** Collect referenced UObjects. */
 	virtual void OnAddReferencedObjects(FReferenceCollector& Collector) {}
 
+#if UE_GAMEPLAY_CAMERAS_DEBUG
+
+public:
+
+	/** Creates the debug drawing block for this camera node. */
+	void CreateDebugBlock(const FCameraNodeEvaluationParams& Params, FCameraDebugBlockBuilder& Builder);
+
+private:
+
+	/**
+	 * Creates the debug drawing block for this camera node.
+	 */
+	virtual void OnCreateDebugBlock(const FCameraNodeEvaluationParams& Params, FCameraDebugBlockBuilder& Builder);
+
+#endif  // UE_GAMEPLAY_CAMERAS_DEBUG
+
 private:
 
 	/** The camera node to run. */
diff --git a/Engine/Plugins/Cameras/GameplayCameras/Source/GameplayCameras/Public/Core/CameraSystemEvaluator.h b/Engine/Plugins/Cameras/GameplayCameras/Source/GameplayCameras/Public/Core/CameraSystemEvaluator.h
--- a/Engine/Plugins/Cameras/GameplayCameras/Source/GameplayCameras/Public/Core/CameraSystemEvaluator.h
+++ b/Engine/Plugins/Cameras/GameplayCameras/Source/GameplayCameras/Public/Core/CameraSystemEvaluator.h
@@ -9,11 +9,14 @@
 #include "Core/CameraNodeEvaluatorStorage.h"
 #include "Core/CameraPose.h"
 #include "CoreTypes.h"
+#include "Debug/CameraDebugBlockStorage.h"
+#include "GameplayCameras.h"
 #include "Templates/SharedPointer.h"
 #include "UObject/GCObject.h"
 
 class UCameraDirector;
 class UCameraRigAsset;
+class UCanvas;
 class URootCameraNode;
 struct FMinimalViewInfo;
 
@@ -23,6 +26,10 @@ namespace UE::Cameras
 class FCameraEvaluationContext;
 class FRootCameraNodeEvaluator;
 
+#if UE_GAMEPLAY_CAMERAS_DEBUG
+class FRootCameraDebugBlock;
+#endif  // UE_GAMEPLAY_CAMERAS_DEBUG
+
 /**
  * Parameter structure for initializing a new camera system evaluator.
  */
@@ -58,6 +65,13 @@ struct FCameraSystemEvaluationUpdateResult
 	bool bIsValid = false;
 };
 
+#if UE_GAMEPLAY_CAMERAS_DEBUG
+struct FCameraSystemDebugDrawParams
+{
+	UCanvas* Canvas = nullptr;
+};
+#endif  // UE_GAMEPLAY_CAMERAS_DEBUG
+
 /**
  * The main camera system evaluator class.
  */
@@ -90,6 +104,10 @@ public:
 	/** Get the last evaluated camera. */
 	void GetEvaluatedCameraView(FMinimalViewInfo& DesiredView);
 
+#if UE_GAMEPLAY_CAMERAS_DEBUG
+	void DebugDraw(const FCameraSystemDebugDrawParams& Params);
+#endif  // UE_GAMEPLAY_CAMERAS_DEBUG
+
 public:
 
 	void AddReferencedObjects(FReferenceCollector& Collector);
@@ -106,13 +124,21 @@ private:
 	FCameraNodeEvaluatorStorage RootEvaluatorStorage;
 
 	/** The root evaluator. */
-	FRootCameraNodeEvaluator* RootEvaluator;
+	FRootCameraNodeEvaluator* RootEvaluator = nullptr;
 
 	/** The current result of the root camera node. */
 	FCameraNodeEvaluationResult RootNodeResult;
 
 	/** The current overall result of the camera system. */
 	FCameraSystemEvaluationUpdateResult Result;
+
+#if UE_GAMEPLAY_CAMERAS_DEBUG
+	/** Storage for debug drawing blocks. */
+	FCameraDebugBlockStorage DebugBlockStorage;
+
+	/** The root debug drawing block. */
+	FRootCameraDebugBlock* RootDebugBlock = nullptr;
+#endif  // UE_GAMEPLAY_CAMERAS_DEBUG
 };
 
 }  // namespace UE::Cameras
diff --git a/Engine/Plugins/Cameras/GameplayCameras/Source/GameplayCameras/Public/Debug/CameraDebugBlock.h b/Engine/Plugins/Cameras/GameplayCameras/Source/GameplayCameras/Public/Debug/CameraDebugBlock.h
new file mode 100644
--- /dev/null
+++ b/Engine/Plugins/Cameras/GameplayCameras/Source/GameplayCameras/Public/Debug/CameraDebugBlock.h
@@ -0,0 +1,153 @@
+// Copyright Epic Games, Inc. All Rights Reserved.
+
+#pragma once
+
+#include "Containers/Array.h"
+#include "CoreTypes.h"
+#include "GameplayCameras.h"
+
+#if UE_GAMEPLAY_CAMERAS_DEBUG
+
+class FArchive;
+
+namespace UE::Cameras
+{
+
+class FCameraDebugBlock;
+class FCameraDebugRenderer;
+
+/**
+ * Parameter structure for debug block drawing.
+ */
+struct FCameraDebugBlockDrawParams
+{
+	float DeltaTime = 0.f;
+};
+
+/**
+ * Metadata for a debug block's member field.
+ * This is mostly used for auto-serialization of debugging info.
+ */
+struct FCameraDebugBlockField
+{
+	virtual ~FCameraDebugBlockField() {}
+	virtual void SerializeField(FCameraDebugBlock* This, FArchive& Ar) = 0;
+
+	FName FieldName;
+	uint16 FieldIndex = 0;
+	uint16 FieldOffset = 0;
+};
+
+/**
+ * Templated version of the debug block field.
+ */
+template<typename FieldType>
+struct TCameraDebugBlockField : FCameraDebugBlockField
+{
+	TCameraDebugBlockField(const FName& InName, uint16 InOffset) 
+	{
+		FieldName = InName;
+		FieldOffset = InOffset;
+	}
+
+	virtual void SerializeField(FCameraDebugBlock* This, FArchive& Ar) override
+	{
+		FieldType* Value = reinterpret_cast<FieldType*>(reinterpret_cast<uint8*>(This) + FieldOffset);
+		Ar << (*Value);
+	}
+};
+
+enum class EDebugDrawResult
+{
+	Default = 0,
+	SkipChildren = 1 << 0
+};
+
+ENUM_CLASS_FLAGS(EDebugDrawResult);
+
+/**
+ * Debug drawing structure responsible for displaying information about some aspect
+ * of the camera system evaluation. For instance, there is generally one debug block
+ * per camera node.
+ */
+class FCameraDebugBlock
+{
+public:
+
+	virtual ~FCameraDebugBlock() {}
+
+	/** Adds a child block to this block. */
+	void AddChild(FCameraDebugBlock* InChild);
+
+	/** Gets the children of this block. */
+	TArrayView<FCameraDebugBlock*> GetChildren() { return Children; }
+
+	/** Called to let this block display its information on screen. */
+	void DebugDraw(const FCameraDebugBlockDrawParams& Params, FCameraDebugRenderer& Renderer);
+
+	/** Serializes this debug block into a buffer, for recording/replaying purposes. */
+	void Serialize(FArchive& Ar);
+
+protected:
+
+	virtual EDebugDrawResult OnDebugDraw(const FCameraDebugBlockDrawParams& Params, FCameraDebugRenderer& Renderer);
+	virtual void OnPostDebugDraw(const FCameraDebugBlockDrawParams& Params, FCameraDebugRenderer& Renderer) {}
+
+	virtual void OnSerialize(FArchive& Ar) {}
+
+protected:
+
+	template<typename FieldType>
+	static TCameraDebugBlockField<FieldType> CreateField(const FName& FieldName, uint16 FieldOffset)
+	{
+		return TCameraDebugBlockField<FieldType>{ FieldName, FieldOffset };
+	}
+
+	static int32 RegisterField(FCameraDebugBlockField* InField)
+	{
+		StaticFields.Add(InField);
+		return StaticFields.Num() - 1;
+	}
+
+private:
+
+	static TArray<FCameraDebugBlockField*> StaticFields;
+
+	using FChildrenArray = TArray<FCameraDebugBlock*>;
+	FChildrenArray Children;
+};
+
+}  // namespace UE::Cameras
+
+// Macros for defining debug blocks for a camera node evaluator.
+//
+#define UE_DEFINE_CAMERA_DEBUG_BLOCK_START(ClassName)\
+	class ClassName : public ::UE::Cameras::FCameraDebugBlock\
+	{\
+	private:\
+		using Super = ::UE::Cameras::FCameraDebugBlock;\
+		using ThisClassName = ClassName;
+
+#define UE_DEFINE_CAMERA_DEBUG_BLOCK_FIELD(FieldType, FieldName)\
+	public:\
+		FieldType FieldName;\
+	private:\
+		inline static PTRINT Get##FieldName##Offset() { return (PTRINT)(&(((ThisClassName*)0)->FieldName)); }\
+		inline static TCameraDebugBlockField<FieldType>* Get##FieldName##Field() {\
+			using FField = TCameraDebugBlockField<FieldType>;\
+			static FField StaticField = CreateField<FieldType>(TEXT(#FieldName), Get##FieldName##Offset());\
+			return &StaticField; }\
+		inline static const int32 FieldName##FieldIndex = RegisterField(Get##FieldName##Field());
+
+#define UE_DEFINE_CAMERA_DEBUG_BLOCK_END()\
+		virtual EDebugDrawResult OnDebugDraw(const FCameraDebugBlockDrawParams& Params, FCameraDebugRenderer& Renderer) override;\
+	};
+
+#else
+
+#define UE_DEFINE_CAMERA_DEBUG_BLOCK_START(ClassName)
+#define UE_DEFINE_CAMERA_DEBUG_BLOCK_FIELD(FieldType, FieldName)
+#define UE_DEFINE_CAMERA_DEBUG_BLOCK_END()
+
+#endif  // UE_GAMEPLAY_CAMERAS_DEBUG
+
diff --git a/Engine/Plugins/Cameras/GameplayCameras/Source/GameplayCameras/Public/Debug/CameraDebugBlockBuilder.h b/Engine/Plugins/Cameras/GameplayCameras/Source/GameplayCameras/Public/Debug/CameraDebugBlockBuilder.h
new file mode 100644
--- /dev/null
+++ b/Engine/Plugins/Cameras/GameplayCameras/Source/GameplayCameras/Public/Debug/CameraDebugBlockBuilder.h
@@ -0,0 +1,58 @@
+// Copyright Epic Games, Inc. All Rights Reserved.
+
+#pragma once
+
+#include "Containers/Array.h"
+#include "CoreTypes.h"
+#include "Debug/CameraDebugBlock.h"
+#include "Debug/CameraDebugBlockStorage.h"
+#include "GameplayCameras.h"
+#include "Templates/SharedPointer.h"
+
+#if UE_GAMEPLAY_CAMERAS_DEBUG
+
+namespace UE::Cameras
+{
+
+class FCameraDebugBlock;
+
+/**
+ * Builder class for camera debug drawing blocks.
+ */
+struct FCameraDebugBlockBuilder
+{
+public:
+
+	/** Creates a new builder structure. */
+	FCameraDebugBlockBuilder(FCameraDebugBlockStorage& InStorage, FCameraDebugBlock& InRootBlock);
+
+	/**
+	 * Creates a new debug drawing block.
+	 * This sets the new block as the "active" block, and adds it as a child of the
+	 * previously active block.
+	 */
+	template<typename BlockType, typename ...ArgTypes>
+	BlockType& StartBlock(ArgTypes&& ...InArgs)
+	{
+		BlockType* NewBlock = Storage.BuildDebugBlock<BlockType>(Forward<ArgTypes>(InArgs)...);
+		OnStartBlock(*NewBlock);
+		return *NewBlock;
+	}
+
+	/** Ends the currently active debug drawing block. */
+	void EndBlock();
+
+private:
+
+	void OnStartBlock(FCameraDebugBlock& InNewBlock);
+
+private:
+
+	FCameraDebugBlockStorage& Storage;
+	TArray<FCameraDebugBlock*> CurrentHierarchy;
+};
+
+}  // namespace UE::Cameras
+
+#endif  // UE_GAMEPLAY_CAMERAS_DEBUG
+
diff --git a/Engine/Plugins/Cameras/GameplayCameras/Source/GameplayCameras/Public/Debug/CameraDebugBlockStorage.h b/Engine/Plugins/Cameras/GameplayCameras/Source/GameplayCameras/Public/Debug/CameraDebugBlockStorage.h
new file mode 100644
--- /dev/null
+++ b/Engine/Plugins/Cameras/GameplayCameras/Source/GameplayCameras/Public/Debug/CameraDebugBlockStorage.h
@@ -0,0 +1,36 @@
+// Copyright Epic Games, Inc. All Rights Reserved.
+
+#pragma once
+
+#include "Core/CameraObjectStorage.h"
+#include "CoreTypes.h"
+#include "Debug/CameraDebugBlock.h"
+
+namespace UE::Cameras
+{
+
+/**
+ * A class responsible for storing a tree of debug blocks
+ */
+class FCameraDebugBlockStorage : public TCameraObjectStorage<FCameraDebugBlock>
+{
+public:
+
+	using Super = TCameraObjectStorage<FCameraDebugBlock>;
+
+	/** Destroy any allocated debug blocks. */
+	void DestroyDebugBlocks(bool bFreeAllocations = false);
+
+	/** Build a new debug block inside this storage. */
+	template<typename BlockType, typename ...ArgTypes>
+	BlockType* BuildDebugBlock(ArgTypes&&... InArgs);
+};
+
+template<typename BlockType, typename ...ArgTypes>
+BlockType* FCameraDebugBlockStorage::BuildDebugBlock(ArgTypes&&... InArgs)
+{
+	return Super::BuildObject<BlockType>(Forward<ArgTypes>(InArgs)...);
+}
+
+}  // namespace UE::Cameras
+
diff --git a/Engine/Plugins/Cameras/GameplayCameras/Source/GameplayCameras/Public/Debug/CameraDebugRenderer.h b/Engine/Plugins/Cameras/GameplayCameras/Source/GameplayCameras/Public/Debug/CameraDebugRenderer.h
new file mode 100644
--- /dev/null
+++ b/Engine/Plugins/Cameras/GameplayCameras/Source/GameplayCameras/Public/Debug/CameraDebugRenderer.h
@@ -0,0 +1,56 @@
+// Copyright Epic Games, Inc. All Rights Reserved.
+
+#pragma once
+
+#include "CoreTypes.h"
+#include "GameplayCameras.h"
+#include "Math/Color.h"
+#include "Misc/StringBuilder.h"
+
+class UCanvas;
+class UFont;
+
+#if UE_GAMEPLAY_CAMERAS_DEBUG
+
+namespace UE::Cameras
+{
+
+class FCameraDebugRenderer
+{
+public:
+
+	FCameraDebugRenderer(UCanvas* InCanvas);
+	~FCameraDebugRenderer();
+
+	void AddText(const FString& InString);
+	void AddText(const TCHAR* Fmt, ...);
+
+	void AddIndent();
+	void RemoveIndent();
+
+private:
+
+	void AddTextFmtImpl(const TCHAR* Fmt, va_list Args);
+	void AddTextImpl(const TCHAR* Buffer);
+
+	void FlushText();
+
+private:
+
+	UCanvas* Canvas;
+	FColor OriginalDrawColor;
+
+	const UFont* RenderFont;
+	int32 MaxCharHeight;
+
+	TStringBuilder<512> Formatter;
+	TStringBuilder<512> LineBuilder;
+
+	int8 IndentLevel = 0;
+	float NextBlockY = 0;
+};
+
+}  // namespace UE::Cameras
+
+#endif  // UE_GAMEPLAY_CAMERAS_DEBUG
+
diff --git a/Engine/Plugins/Cameras/GameplayCameras/Source/GameplayCameras/Public/GameFramework/GameplayCameraSystemComponent.h b/Engine/Plugins/Cameras/GameplayCameras/Source/GameplayCameras/Public/GameFramework/GameplayCameraSystemComponent.h
--- a/Engine/Plugins/Cameras/GameplayCameras/Source/GameplayCameras/Public/GameFramework/GameplayCameraSystemComponent.h
+++ b/Engine/Plugins/Cameras/GameplayCameras/Source/GameplayCameras/Public/GameFramework/GameplayCameraSystemComponent.h
@@ -4,13 +4,16 @@
 
 #include "CoreMinimal.h"
 #include "Core/CameraSystemEvaluator.h"
+#include "GameplayCameras.h"
 #include "UObject/ObjectMacros.h"
 #include "UObject/ScriptInterface.h"
 #include "Components/SceneComponent.h"
 
 #include "GameplayCameraSystemComponent.generated.h"
 
+class APlayerController;
 class UCameraRigAsset;
+class UCanvas;
 struct FMinimalViewInfo;
 
 namespace UE::Cameras
@@ -59,10 +62,20 @@ public:
 	void OnBecomeViewTarget();
 	void OnEndViewTarget();
 
+private:
+
+#if UE_GAMEPLAY_CAMERAS_DEBUG
+	void DebugDraw(UCanvas* Canvas, APlayerController* PlayController);
+#endif  // UE_GAMEPLAY_CAMERAS_DEBUG
+
 private:
 	
 	TSharedPtr<FCameraSystemEvaluator> Evaluator;
 
+#if UE_GAMEPLAY_CAMERAS_DEBUG
+	FDelegateHandle DebugDrawDelegateHandle;
+#endif  // UE_GAMEPLAY_CAMERAS_DEBUG
+
 #if WITH_EDITORONLY_DATA
 
 	UPROPERTY()
diff --git a/Engine/Plugins/Cameras/GameplayCameras/Source/GameplayCameras/Public/GameplayCameras.h b/Engine/Plugins/Cameras/GameplayCameras/Source/GameplayCameras/Public/GameplayCameras.h
--- a/Engine/Plugins/Cameras/GameplayCameras/Source/GameplayCameras/Public/GameplayCameras.h
+++ b/Engine/Plugins/Cameras/GameplayCameras/Source/GameplayCameras/Public/GameplayCameras.h
@@ -3,8 +3,17 @@
 #pragma once
 
 #include "CoreTypes.h"
+#include "HAL/Platform.h"
 #include "Stats/Stats.h"
 
+#ifndef UE_GAMEPLAY_CAMERAS_DEBUG
+	#if (!UE_BUILD_SHIPPING || WITH_UNREAL_DEVELOPER_TOOLS || WITH_UNREAL_TARGET_DEVELOPER_TOOLS)
+		#define UE_GAMEPLAY_CAMERAS_DEBUG 1
+	#else
+		#define UE_GAMEPLAY_CAMERAS_DEBUG 0
+	#endif
+#endif
+
 DECLARE_STATS_GROUP(TEXT("Camera System Evaluation"), STATGROUP_CameraSystem, STATCAT_Advanced)
 DECLARE_STATS_GROUP(TEXT("Camera Animation Evaluation"), STATGROUP_CameraAnimation, STATCAT_Advanced)
 
