Cameras: fix shipping build

diff --git a/Engine/Plugins/Cameras/GameplayCameras/Source/GameplayCameras/Public/Debug/CameraDebugBlockStorage.h b/Engine/Plugins/Cameras/GameplayCameras/Source/GameplayCameras/Public/Debug/CameraDebugBlockStorage.h
--- a/Engine/Plugins/Cameras/GameplayCameras/Source/GameplayCameras/Public/Debug/CameraDebugBlockStorage.h
+++ b/Engine/Plugins/Cameras/GameplayCameras/Source/GameplayCameras/Public/Debug/CameraDebugBlockStorage.h
@@ -5,6 +5,9 @@
 #include "Core/CameraObjectStorage.h"
 #include "CoreTypes.h"
 #include "Debug/CameraDebugBlock.h"
+#include "GameplayCameras.h"
+
+#if UE_GAMEPLAY_CAMERAS_DEBUG
 
 namespace UE::Cameras
 {
@@ -34,3 +37,5 @@ BlockType* FCameraDebugBlockStorage::BuildDebugBlock(ArgTypes&&... InArgs)
 
 }  // namespace UE::Cameras
 
+#endif  // UE_GAMEPLAY_CAMERAS_DEBUG
+
