Cameras: more shipping build fixes

diff --git a/Engine/Plugins/Cameras/GameplayCameras/Source/GameplayCameras/Private/Debug/CameraDebugBlockStorage.cpp b/Engine/Plugins/Cameras/GameplayCameras/Source/GameplayCameras/Private/Debug/CameraDebugBlockStorage.cpp
--- a/Engine/Plugins/Cameras/GameplayCameras/Source/GameplayCameras/Private/Debug/CameraDebugBlockStorage.cpp
+++ b/Engine/Plugins/Cameras/GameplayCameras/Source/GameplayCameras/Private/Debug/CameraDebugBlockStorage.cpp
@@ -2,6 +2,8 @@
 
 #include "Debug/CameraDebugBlockStorage.h"
 
+#if UE_GAMEPLAY_CAMERAS_DEBUG
+
 namespace UE::Cameras
 {
 
@@ -12,3 +14,5 @@ void FCameraDebugBlockStorage::DestroyDebugBlocks(bool bFreeAllocations)
 
 }  // namespace UE::Cameras
 
+#endif  // UE_GAMEPLAY_CAMERAS_DEBUG
+
