Cameras: fix incorrect macro

diff --git a/Engine/Plugins/Cameras/GameplayCameras/Source/GameplayCameras/Private/GameFramework/GameplayCameraSystemComponent.cpp b/Engine/Plugins/Cameras/GameplayCameras/Source/GameplayCameras/Private/GameFramework/GameplayCameraSystemComponent.cpp
--- a/Engine/Plugins/Cameras/GameplayCameras/Source/GameplayCameras/Private/GameFramework/GameplayCameraSystemComponent.cpp
+++ b/Engine/Plugins/Cameras/GameplayCameras/Source/GameplayCameras/Private/GameFramework/GameplayCameraSystemComponent.cpp
@@ -140,7 +140,8 @@ void UGameplayCameraSystemComponent::OnEndViewTarget()
 {
 }
 
-#if UE_ENABLE_DEBUG_DRAWING
+#if UE_GAMEPLAY_CAMERAS_DEBUG
+
 void UGameplayCameraSystemComponent::DebugDraw(UCanvas* Canvas, APlayerController* PlayController)
 {
 	using namespace UE::Cameras;
@@ -152,7 +153,7 @@ void UGameplayCameraSystemComponent::DebugDraw(UCanvas* Canvas, APlayerControlle
 		Evaluator->DebugDraw(DebugDrawParams);
 	}
 }
-#endif  // UE_ENABLE_DEBUG_DRAWING
+#endif  // UE_GAMEPLAY_CAMERAS_DEBUG
 
 #undef LOCTEXT_NAMESPACE
 
