Fixed use of a deprecated variable by GameplayCameras

#rnx

diff --git a/Engine/Plugins/Cameras/GameplayCameras/Source/GameplayCameras/Private/Core/CameraParameters.cpp b/Engine/Plugins/Cameras/GameplayCameras/Source/GameplayCameras/Private/Core/CameraParameters.cpp
--- a/Engine/Plugins/Cameras/GameplayCameras/Source/GameplayCameras/Private/Core/CameraParameters.cpp
+++ b/Engine/Plugins/Cameras/GameplayCameras/Source/GameplayCameras/Private/Core/CameraParameters.cpp
@@ -40,7 +40,7 @@ bool FFloatCameraParameter::SerializeFromMismatchedTag(const FPropertyTag& Tag,
 
 bool FDoubleCameraParameter::SerializeFromMismatchedTag(const FPropertyTag& Tag, FStructuredArchive::FSlot Slot)
 {
-	if (Tag.Type == NAME_StructProperty && Tag.StructName == NAME_DoubleProperty)
+	if (Tag.GetType().IsStruct(NAME_DoubleProperty))
 	{
 		Slot << Value;
 		return true;
@@ -51,7 +51,7 @@ bool FDoubleCameraParameter::SerializeFromMismatchedTag(const FPropertyTag& Tag,
 
 bool FVector2fCameraParameter::SerializeFromMismatchedTag(const FPropertyTag& Tag, FStructuredArchive::FSlot Slot)
 {
-	if (Tag.Type == NAME_StructProperty && (Tag.StructName == NAME_Vector2f || Tag.StructName == NAME_Vector2D))
+	if (Tag.GetType().IsStruct(NAME_Vector2f) || Tag.GetType().IsStruct(NAME_Vector2D))
 	{
 		Slot << Value;
 		return true;
@@ -62,7 +62,7 @@ bool FVector2fCameraParameter::SerializeFromMismatchedTag(const FPropertyTag& Ta
 
 bool FVector2dCameraParameter::SerializeFromMismatchedTag(const FPropertyTag& Tag, FStructuredArchive::FSlot Slot)
 {
-	if (Tag.Type == NAME_StructProperty && Tag.StructName == NAME_Vector2d)
+	if (Tag.GetType().IsStruct(NAME_Vector2d))
 	{
 		Slot << Value;
 		return true;
@@ -73,7 +73,7 @@ bool FVector2dCameraParameter::SerializeFromMismatchedTag(const FPropertyTag& Ta
 
 bool FVector3fCameraParameter::SerializeFromMismatchedTag(const FPropertyTag& Tag, FStructuredArchive::FSlot Slot)
 {
-	if (Tag.Type == NAME_StructProperty && Tag.StructName == NAME_Vector3f)
+	if (Tag.GetType().IsStruct(NAME_Vector3f))
 	{
 		Slot << Value;
 		return true;
@@ -84,7 +84,7 @@ bool FVector3fCameraParameter::SerializeFromMismatchedTag(const FPropertyTag& Ta
 
 bool FVector3dCameraParameter::SerializeFromMismatchedTag(const FPropertyTag& Tag, FStructuredArchive::FSlot Slot)
 {
-	if (Tag.Type == NAME_StructProperty && Tag.StructName == NAME_Vector3d)
+	if (Tag.GetType().IsStruct(NAME_Vector3d))
 	{
 		Slot << Value;
 		return true;
@@ -95,7 +95,7 @@ bool FVector3dCameraParameter::SerializeFromMismatchedTag(const FPropertyTag& Ta
 
 bool FVector4fCameraParameter::SerializeFromMismatchedTag(const FPropertyTag& Tag, FStructuredArchive::FSlot Slot)
 {
-	if (Tag.Type == NAME_StructProperty && Tag.StructName == NAME_Vector4f)
+	if (Tag.GetType().IsStruct(NAME_Vector4f))
 	{
 		Slot << Value;
 		return true;
@@ -106,7 +106,7 @@ bool FVector4fCameraParameter::SerializeFromMismatchedTag(const FPropertyTag& Ta
 
 bool FVector4dCameraParameter::SerializeFromMismatchedTag(const FPropertyTag& Tag, FStructuredArchive::FSlot Slot)
 {
-	if (Tag.Type == NAME_StructProperty && Tag.StructName == NAME_Vector4d)
+	if (Tag.GetType().IsStruct(NAME_Vector4d))
 	{
 		Slot << Value;
 		return true;
@@ -117,7 +117,7 @@ bool FVector4dCameraParameter::SerializeFromMismatchedTag(const FPropertyTag& Ta
 
 bool FRotator3fCameraParameter::SerializeFromMismatchedTag(const FPropertyTag& Tag, FStructuredArchive::FSlot Slot)
 {
-	if (Tag.Type == NAME_StructProperty && Tag.StructName == NAME_Rotator3f)
+	if (Tag.GetType().IsStruct(NAME_Rotator3f))
 	{
 		Slot << Value;
 		return true;
@@ -128,7 +128,7 @@ bool FRotator3fCameraParameter::SerializeFromMismatchedTag(const FPropertyTag& T
 
 bool FRotator3dCameraParameter::SerializeFromMismatchedTag(const FPropertyTag& Tag, FStructuredArchive::FSlot Slot)
 {
-	if (Tag.Type == NAME_StructProperty && Tag.StructName == NAME_Rotator3d)
+	if (Tag.GetType().IsStruct(NAME_Rotator3d))
 	{
 		Slot << Value;
 		return true;
@@ -139,7 +139,7 @@ bool FRotator3dCameraParameter::SerializeFromMismatchedTag(const FPropertyTag& T
 
 bool FTransform3fCameraParameter::SerializeFromMismatchedTag(const FPropertyTag& Tag, FStructuredArchive::FSlot Slot)
 {
-	if (Tag.Type == NAME_StructProperty && Tag.StructName == NAME_Transform3f)
+	if (Tag.GetType().IsStruct(NAME_Transform3f))
 	{
 		Slot << Value;
 		return true;
@@ -150,7 +150,7 @@ bool FTransform3fCameraParameter::SerializeFromMismatchedTag(const FPropertyTag&
 
 bool FTransform3dCameraParameter::SerializeFromMismatchedTag(const FPropertyTag& Tag, FStructuredArchive::FSlot Slot)
 {
-	if (Tag.Type == NAME_StructProperty && Tag.StructName == NAME_Transform3d)
+	if (Tag.GetType().IsStruct(NAME_Transform3d))
 	{
 		Slot << Value;
 		return true;
