Cameras: fix gameplay components registrations

Don't early exit and miss up-calling base class implementation.
Don't register debug drawing for editor worlds.

diff --git a/Engine/Plugins/Cameras/GameplayCameras/Source/GameplayCameras/Private/GameFramework/GameplayCameraSystemComponent.cpp b/Engine/Plugins/Cameras/GameplayCameras/Source/GameplayCameras/Private/GameFramework/GameplayCameraSystemComponent.cpp
--- a/Engine/Plugins/Cameras/GameplayCameras/Source/GameplayCameras/Private/GameFramework/GameplayCameraSystemComponent.cpp
+++ b/Engine/Plugins/Cameras/GameplayCameras/Source/GameplayCameras/Private/GameFramework/GameplayCameraSystemComponent.cpp
@@ -49,6 +49,8 @@ void UGameplayCameraSystemComponent::OnRegister()
 {
 	using namespace UE::Cameras;
 
+	Super::OnRegister();
+
 	AActor* OwnerActor = GetOwner();
 	if (!OwnerActor || OwnerActor->HasAnyFlags(RF_ClassDefaultObject))
 	{
@@ -62,8 +64,12 @@ void UGameplayCameraSystemComponent::OnRegister()
 	}
 
 #if UE_GAMEPLAY_CAMERAS_DEBUG
-	DebugDrawDelegateHandle = UDebugDrawService::Register(
-			TEXT("Game"), FDebugDrawDelegate::CreateUObject(this, &UGameplayCameraSystemComponent::DebugDraw));
+	UWorld* World = GetWorld();
+	if (World && World->IsGameWorld())
+	{
+		DebugDrawDelegateHandle = UDebugDrawService::Register(
+				TEXT("Game"), FDebugDrawDelegate::CreateUObject(this, &UGameplayCameraSystemComponent::DebugDraw));
+	}
 #endif  // UE_GAMEPLAY_CAMERAS_DEBUG
 
 #if WITH_EDITORONLY_DATA
@@ -80,8 +86,6 @@ void UGameplayCameraSystemComponent::OnRegister()
 		PreviewMeshComponent->RegisterComponentWithWorld(GetWorld());
 	}
 #endif	// WITH_EDITORONLY_DATA
-
-	Super::OnRegister();
 }
 
 void UGameplayCameraSystemComponent::TickComponent(float DeltaTime, ELevelTick TickType, FActorComponentTickFunction* ThisTickFunction)
