Cameras: add type information to simple RTTI framework

This will let us re-create debug blocks during replay/rewind.

diff --git a/Engine/Plugins/Cameras/GameplayCameras/Source/GameplayCameras/Private/Core/CameraObjectRtti.cpp b/Engine/Plugins/Cameras/GameplayCameras/Source/GameplayCameras/Private/Core/CameraObjectRtti.cpp
--- a/Engine/Plugins/Cameras/GameplayCameras/Source/GameplayCameras/Private/Core/CameraObjectRtti.cpp
+++ b/Engine/Plugins/Cameras/GameplayCameras/Source/GameplayCameras/Private/Core/CameraObjectRtti.cpp
@@ -11,5 +11,29 @@ uint32 FCameraObjectTypeID::RegisterNewID()
 	return ID++;
 }
 
+FCameraObjectTypeRegistry& FCameraObjectTypeRegistry::Get()
+{
+	static FCameraObjectTypeRegistry StaticRegistry;
+	return StaticRegistry;
+}
+
+void FCameraObjectTypeRegistry::RegisterType(FCameraObjectTypeID TypeID, FCameraObjectTypeInfo&& TypeInfo)
+{
+	ensureMsgf(
+			!TypeIDsByName.Contains(TypeID.GetTypeName()), 
+			TEXT("Type '%s' has already been registered!"), *TypeID.GetTypeName().ToString());
+	TypeIDsByName.Add(TypeID.GetTypeName(), TypeID.GetTypeID());
+	TypeInfos.Insert(TypeID.GetTypeID(), MoveTemp(TypeInfo));
+}
+
+void FCameraObjectTypeRegistry::ConstructObject(FCameraObjectTypeID TypeID, void* Ptr)
+{
+	if (ensureMsgf(TypeInfos.IsValidIndex(TypeID.GetTypeID()), TEXT("Invalid camera object type ID!")))
+	{
+		const FCameraObjectTypeInfo& TypeInfo = TypeInfos[TypeID.GetTypeID()];
+		TypeInfo.Constructor(Ptr);
+	}
+}
+
 }  // namespace UE::Cameras
 
diff --git a/Engine/Plugins/Cameras/GameplayCameras/Source/GameplayCameras/Public/Core/CameraObjectRtti.h b/Engine/Plugins/Cameras/GameplayCameras/Source/GameplayCameras/Public/Core/CameraObjectRtti.h
--- a/Engine/Plugins/Cameras/GameplayCameras/Source/GameplayCameras/Public/Core/CameraObjectRtti.h
+++ b/Engine/Plugins/Cameras/GameplayCameras/Source/GameplayCameras/Public/Core/CameraObjectRtti.h
@@ -2,6 +2,8 @@
 
 #pragma once
 
+#include "Containers/Map.h"
+#include "Containers/SparseArray.h"
 #include "CoreTypes.h"
 #include "UObject/NameTypes.h"
 
@@ -33,6 +35,11 @@ struct FCameraObjectTypeID
 		return ID;
 	}
 
+	const FName& GetTypeName() const
+	{
+		return Name;
+	}
+
 	FCameraObjectTypeID(const FName& InName, uint32 InID)
 		: Name(InName)
 		, ID(InID)
@@ -48,6 +55,36 @@ protected:
 	uint32 ID;
 };
 
+/**
+ * Type information about an RTTI-enabled class.
+ */
+struct FCameraObjectTypeInfo
+{
+	using FConstructor = void(*)(void*);
+	FConstructor Constructor;
+
+	using FDestructor = void(*)(void*);
+	FDestructor Destructor;
+};
+
+/**
+ * Type registry of known RTTI-enabled classes.
+ */
+class FCameraObjectTypeRegistry
+{
+public:
+
+	static FCameraObjectTypeRegistry& Get();
+
+	void RegisterType(FCameraObjectTypeID TypeID, FCameraObjectTypeInfo&& TypeInfo);
+	void ConstructObject(FCameraObjectTypeID TypeID, void* Ptr);
+
+private:
+
+	TMap<FName, uint32> TypeIDsByName;
+	TSparseArray<FCameraObjectTypeInfo> TypeInfos;
+};
+
 /**
  * Strongly-typed ID wrapper for an RTTI-enabled class.
  */
@@ -58,9 +95,25 @@ private:
 
 	TCameraObjectTypeID(const FName& InName, uint32 InID) : FCameraObjectTypeID(InName, InID) {}
 
-	static TCameraObjectTypeID RegisterNewID(const FName& InClassName)
+	static TCameraObjectTypeID RegisterType(const FName& InClassName)
+	{
+		TCameraObjectTypeID NewTypeID(InClassName, FCameraObjectTypeID::RegisterNewID());
+		FCameraObjectTypeInfo NewTypeInfo { 
+			&TCameraObjectTypeID<T>::StaticConstructor,
+			&TCameraObjectTypeID<T>::StaticDestructor
+		};
+		FCameraObjectTypeRegistry::Get().RegisterType(NewTypeID, MoveTemp(NewTypeInfo));
+		return NewTypeID;
+	}
+
+	static void StaticConstructor(void* Ptr)
+	{
+		new(Ptr) T();
+	}
+
+	static void StaticDestructor(void* Ptr)
 	{
-		return TCameraObjectTypeID(InClassName, FCameraObjectTypeID::RegisterNewID());
+		reinterpret_cast<T*>(Ptr)->~T();
 	}
 
 	friend T;
@@ -95,5 +148,5 @@ private:
 		static const ::UE::Cameras::TCameraObjectTypeID<ClassName> PrivateTypeID;
 
 #define UE_GAMEPLAY_CAMERAS_DEFINE_RTTI(ClassName)\
-	const ::UE::Cameras::TCameraObjectTypeID<ClassName> ClassName::PrivateTypeID = ::UE::Cameras::TCameraObjectTypeID<ClassName>::RegisterNewID(#ClassName);
+	const ::UE::Cameras::TCameraObjectTypeID<ClassName> ClassName::PrivateTypeID = ::UE::Cameras::TCameraObjectTypeID<ClassName>::RegisterType(#ClassName);
 
