Cameras: move some editor classes inside a namespace, fix commands clash

diff --git a/Engine/Plugins/Cameras/GameplayCameras/Source/GameplayCamerasEditor/Private/AssetTools/CameraAssetEditor.cpp b/Engine/Plugins/Cameras/GameplayCameras/Source/GameplayCamerasEditor/Private/AssetTools/CameraAssetEditor.cpp
--- a/Engine/Plugins/Cameras/GameplayCameras/Source/GameplayCamerasEditor/Private/AssetTools/CameraAssetEditor.cpp
+++ b/Engine/Plugins/Cameras/GameplayCameras/Source/GameplayCamerasEditor/Private/AssetTools/CameraAssetEditor.cpp
@@ -22,6 +22,6 @@ void UCameraAssetEditor::GetObjectsToEdit(TArray<UObject*>& InObjectsToEdit)
 
 TSharedPtr<FBaseAssetToolkit> UCameraAssetEditor::CreateToolkit()
 {
-	return MakeShared<FCameraAssetEditorToolkit>(this);
+	return MakeShared<UE::Cameras::FCameraAssetEditorToolkit>(this);
 }
 
diff --git a/Engine/Plugins/Cameras/GameplayCameras/Source/GameplayCamerasEditor/Private/AssetTools/CameraRigAssetEditor.cpp b/Engine/Plugins/Cameras/GameplayCameras/Source/GameplayCamerasEditor/Private/AssetTools/CameraRigAssetEditor.cpp
--- a/Engine/Plugins/Cameras/GameplayCameras/Source/GameplayCamerasEditor/Private/AssetTools/CameraRigAssetEditor.cpp
+++ b/Engine/Plugins/Cameras/GameplayCameras/Source/GameplayCamerasEditor/Private/AssetTools/CameraRigAssetEditor.cpp
@@ -22,6 +22,6 @@ void UCameraRigAssetEditor::GetObjectsToEdit(TArray<UObject*>& InObjectsToEdit)
 
 TSharedPtr<FBaseAssetToolkit> UCameraRigAssetEditor::CreateToolkit()
 {
-	return MakeShared<FCameraRigAssetEditorToolkit>(this);
+	return MakeShared<UE::Cameras::FCameraRigAssetEditorToolkit>(this);
 }
 
diff --git a/Engine/Plugins/Cameras/GameplayCameras/Source/GameplayCamerasEditor/Private/Commands/CameraAssetEditorCommands.cpp b/Engine/Plugins/Cameras/GameplayCameras/Source/GameplayCamerasEditor/Private/Commands/CameraAssetEditorCommands.cpp
--- a/Engine/Plugins/Cameras/GameplayCameras/Source/GameplayCamerasEditor/Private/Commands/CameraAssetEditorCommands.cpp
+++ b/Engine/Plugins/Cameras/GameplayCameras/Source/GameplayCamerasEditor/Private/Commands/CameraAssetEditorCommands.cpp
@@ -9,10 +9,13 @@
 
 #define LOCTEXT_NAMESPACE "SequencerCommands"
 
+namespace UE::Cameras
+{
+
 FCameraAssetEditorCommands::FCameraAssetEditorCommands()
 	: TCommands<FCameraAssetEditorCommands>(
-			"GameplayCameras",
-			NSLOCTEXT("Contexts", "GameplayCameras", "Gameplay Cameras"),
+			"GameplayCameras_CameraAssetEditor",
+			NSLOCTEXT("Contexts", "GameplayCameras_CameraAssetEditor", "Camera Asset Editor"),
 			NAME_None,
 			FAppStyle::GetAppStyleSetName()
 		)
@@ -24,6 +27,8 @@ void FCameraAssetEditorCommands::RegisterCommands()
 	UI_COMMAND(Build, "Build", "Builds the asset and refreshes it in PIE",
 			EUserInterfaceActionType::Button, FInputChord(EKeys::F7));
 }
- 
+
+}  // namespace UE::Cameras
+
 #undef LOCTEXT_NAMESPACE
 
diff --git a/Engine/Plugins/Cameras/GameplayCameras/Source/GameplayCamerasEditor/Private/Commands/CameraRigAssetEditorCommands.cpp b/Engine/Plugins/Cameras/GameplayCameras/Source/GameplayCamerasEditor/Private/Commands/CameraRigAssetEditorCommands.cpp
--- a/Engine/Plugins/Cameras/GameplayCameras/Source/GameplayCamerasEditor/Private/Commands/CameraRigAssetEditorCommands.cpp
+++ b/Engine/Plugins/Cameras/GameplayCameras/Source/GameplayCamerasEditor/Private/Commands/CameraRigAssetEditorCommands.cpp
@@ -9,10 +9,13 @@
 
 #define LOCTEXT_NAMESPACE "CameraRigAssetEditorCommands"
 
+namespace UE::Cameras
+{
+
 FCameraRigAssetEditorCommands::FCameraRigAssetEditorCommands()
 	: TCommands<FCameraRigAssetEditorCommands>(
-			"GameplayCameras",
-			NSLOCTEXT("Contexts", "GameplayCameras", "Gameplay Cameras"),
+			"GameplayCameras_CameraRigAssetEditor",
+			NSLOCTEXT("Contexts", "GameplayCameras_CameraRigAssetEditor", "Camera Rig Asset Editor"),
 			NAME_None,
 			FAppStyle::GetAppStyleSetName()
 		)
@@ -24,6 +27,8 @@ void FCameraRigAssetEditorCommands::RegisterCommands()
 	UI_COMMAND(Build, "Build", "Builds the asset and refreshes it in PIE",
 			EUserInterfaceActionType::Button, FInputChord(EKeys::F7));
 }
- 
+
+}  // namespace UE::Cameras
+
 #undef LOCTEXT_NAMESPACE
 
diff --git a/Engine/Plugins/Cameras/GameplayCameras/Source/GameplayCamerasEditor/Private/GameplayCamerasEditorModule.cpp b/Engine/Plugins/Cameras/GameplayCameras/Source/GameplayCamerasEditor/Private/GameplayCamerasEditorModule.cpp
--- a/Engine/Plugins/Cameras/GameplayCameras/Source/GameplayCamerasEditor/Private/GameplayCamerasEditorModule.cpp
+++ b/Engine/Plugins/Cameras/GameplayCameras/Source/GameplayCamerasEditor/Private/GameplayCamerasEditorModule.cpp
@@ -42,6 +42,8 @@ public:
 
 	virtual void ShutdownModule() override
 	{
+		using namespace UE::Cameras;
+
 		UToolMenus::UnRegisterStartupCallback(this);
 
 		FCameraAssetEditorCommands::Unregister();
@@ -96,6 +98,8 @@ private:
 
 	void RegisterMenus()
 	{
+		using namespace UE::Cameras;
+
 		FCameraAssetEditorCommands::Register();
 		FCameraRigAssetEditorCommands::Register();	
 	}
diff --git a/Engine/Plugins/Cameras/GameplayCameras/Source/GameplayCamerasEditor/Private/Toolkits/CameraAssetEditorToolkit.cpp b/Engine/Plugins/Cameras/GameplayCameras/Source/GameplayCamerasEditor/Private/Toolkits/CameraAssetEditorToolkit.cpp
--- a/Engine/Plugins/Cameras/GameplayCameras/Source/GameplayCamerasEditor/Private/Toolkits/CameraAssetEditorToolkit.cpp
+++ b/Engine/Plugins/Cameras/GameplayCameras/Source/GameplayCamerasEditor/Private/Toolkits/CameraAssetEditorToolkit.cpp
@@ -15,6 +15,9 @@
 
 #define LOCTEXT_NAMESPACE "CameraAssetEditorToolkit"
 
+namespace UE::Cameras
+{
+
 const FName FCameraAssetEditorToolkit::DetailsViewTabId(TEXT("CameraAssetEditor_DetailsView"));
 
 FCameraAssetEditorToolkit::FCameraAssetEditorToolkit(UCameraAssetEditor* InOwningAssetEditor)
@@ -112,5 +115,7 @@ FLinearColor FCameraAssetEditorToolkit::GetWorldCentricTabColorScale() const
 	return FLinearColor(0.7, 0.0f, 0.0f, 0.5f);
 }
 
+}  // namespace UE::Cameras
+
 #undef LOCTEXT_NAMESPACE
 
diff --git a/Engine/Plugins/Cameras/GameplayCameras/Source/GameplayCamerasEditor/Private/Toolkits/CameraRigAssetEditorToolkit.cpp b/Engine/Plugins/Cameras/GameplayCameras/Source/GameplayCamerasEditor/Private/Toolkits/CameraRigAssetEditorToolkit.cpp
--- a/Engine/Plugins/Cameras/GameplayCameras/Source/GameplayCamerasEditor/Private/Toolkits/CameraRigAssetEditorToolkit.cpp
+++ b/Engine/Plugins/Cameras/GameplayCameras/Source/GameplayCamerasEditor/Private/Toolkits/CameraRigAssetEditorToolkit.cpp
@@ -4,6 +4,7 @@
 
 #include "AssetTools/CameraRigAssetEditor.h"
 #include "Commands/CameraRigAssetEditorCommands.h"
+#include "Core/CameraRigAllocationInfoBuilder.h"
 #include "Core/CameraRigAsset.h"
 #include "EditorModeManager.h"
 #include "Framework/Docking/LayoutExtender.h"
@@ -18,6 +19,9 @@
 
 #define LOCTEXT_NAMESPACE "CameraRigAssetEditorToolkit"
 
+namespace UE::Cameras
+{
+
 const FName FCameraRigAssetEditorToolkit::DetailsViewTabId(TEXT("CameraRigAssetEditor_DetailsView"));
 
 FCameraRigAssetEditorToolkit::FCameraRigAssetEditorToolkit(UCameraRigAssetEditor* InOwningAssetEditor)
@@ -204,7 +208,14 @@ void FCameraRigAssetEditorToolkit::NotifyPostChange(const FPropertyChangedEvent&
 
 void FCameraRigAssetEditorToolkit::OnBuild()
 {
+	using namespace UE::Cameras;
+
+	FCameraRigAllocationInfo AllocationInfo;
+	FCameraRigAllocationInfoBuilder CameraRigBuilder;
+	CameraRigBuilder.BuildAllocationInfo(CameraRigAsset, AllocationInfo);
+
 	CameraRigAsset->BuildStatus = ECameraRigBuildStatus::Clean;
+	CameraRigAsset->AllocationInfo = AllocationInfo;
 
 	FCameraRigPackages BuiltPackages;
 	CameraRigAsset->GatherPackages(BuiltPackages);
@@ -215,5 +226,7 @@ void FCameraRigAssetEditorToolkit::OnBuild()
 	}
 }
 
+}  // namespace UE::Cameras
+
 #undef LOCTEXT_NAMESPACE
 
diff --git a/Engine/Plugins/Cameras/GameplayCameras/Source/GameplayCamerasEditor/Public/Commands/CameraAssetEditorCommands.h b/Engine/Plugins/Cameras/GameplayCameras/Source/GameplayCamerasEditor/Public/Commands/CameraAssetEditorCommands.h
--- a/Engine/Plugins/Cameras/GameplayCameras/Source/GameplayCamerasEditor/Public/Commands/CameraAssetEditorCommands.h
+++ b/Engine/Plugins/Cameras/GameplayCameras/Source/GameplayCamerasEditor/Public/Commands/CameraAssetEditorCommands.h
@@ -9,6 +9,9 @@ class FUICommandInfo;
 class UClass;
 struct FInputChord;
 
+namespace UE::Cameras
+{
+
 class FCameraAssetEditorCommands : public TCommands<FCameraAssetEditorCommands>
 {
 public:
@@ -22,3 +25,5 @@ public:
 	TSharedPtr<FUICommandInfo> Build;
 };
 
+}  // namespace UE::Cameras
+
diff --git a/Engine/Plugins/Cameras/GameplayCameras/Source/GameplayCamerasEditor/Public/Commands/CameraRigAssetEditorCommands.h b/Engine/Plugins/Cameras/GameplayCameras/Source/GameplayCamerasEditor/Public/Commands/CameraRigAssetEditorCommands.h
--- a/Engine/Plugins/Cameras/GameplayCameras/Source/GameplayCamerasEditor/Public/Commands/CameraRigAssetEditorCommands.h
+++ b/Engine/Plugins/Cameras/GameplayCameras/Source/GameplayCamerasEditor/Public/Commands/CameraRigAssetEditorCommands.h
@@ -9,6 +9,9 @@ class FUICommandInfo;
 class UClass;
 struct FInputChord;
 
+namespace UE::Cameras
+{
+
 class FCameraRigAssetEditorCommands : public TCommands<FCameraRigAssetEditorCommands>
 {
 public:
@@ -22,3 +25,5 @@ public:
 	TSharedPtr<FUICommandInfo> Build;
 };
 
+}  // namespace UE::Cameras
+
diff --git a/Engine/Plugins/Cameras/GameplayCameras/Source/GameplayCamerasEditor/Public/Toolkits/CameraAssetEditorToolkit.h b/Engine/Plugins/Cameras/GameplayCameras/Source/GameplayCamerasEditor/Public/Toolkits/CameraAssetEditorToolkit.h
--- a/Engine/Plugins/Cameras/GameplayCameras/Source/GameplayCamerasEditor/Public/Toolkits/CameraAssetEditorToolkit.h
+++ b/Engine/Plugins/Cameras/GameplayCameras/Source/GameplayCamerasEditor/Public/Toolkits/CameraAssetEditorToolkit.h
@@ -11,6 +11,9 @@ class UCameraAsset;
 class UCameraAssetEditor;
 class SWidget;
 
+namespace UE::Cameras
+{
+
 /**
  * Editor toolkit for a camera asset.
  */
@@ -61,3 +64,5 @@ private:
 	TObjectPtr<UCameraAsset> CameraAsset;
 };
 
+}  // namespace UE::Cameras
+
