Cameras: large improvement of debug rendering

Debug blocks aren't created during evaluation anymore -- that could lead to issues when later we can do iterative evaluations for IK or input preview. Instead, debug blocks are created on demand via a new API on the node evaluator class.

The debug block builder is a bit more sophisticated, and there is more control over where a new debug block goes (under the last created debug block from a parent node evaluator, or under the root debug block for a separate debug mode, etc.)

Debug blocks also have a more formalized structure: they can be parented under each other, and they can be attached to each other. Attachment implies a close relationship, i.e. if a debug block isn't rendered, its attachments won't either. Children debug block can still in some situations be rendered while their parents aren't (e.g. when filtering a hierarchy of debug blocks).

Other changes include:

- Several debug modes implemented (director tree, blend stacks, evaluated camera pose, and viewfinder).
- Simple blend nodes now have a debug block.
- Standardized colors for debug drawing.
- Rendering a background tile for helping with readability.

diff --git a/Engine/Plugins/Cameras/GameplayCameras/Source/GameplayCameras/Private/Core/BlendStackCameraNode.cpp b/Engine/Plugins/Cameras/GameplayCameras/Source/GameplayCameras/Private/Core/BlendStackCameraNode.cpp
--- a/Engine/Plugins/Cameras/GameplayCameras/Source/GameplayCameras/Private/Core/BlendStackCameraNode.cpp
+++ b/Engine/Plugins/Cameras/GameplayCameras/Source/GameplayCameras/Private/Core/BlendStackCameraNode.cpp
@@ -8,8 +8,16 @@
 #include "Core/CameraEvaluationContext.h"
 #include "Core/CameraRigAsset.h"
 #include "Core/CameraSystemEvaluator.h"
+#include "Debug/CameraDebugBlockBuilder.h"
+#include "Debug/CameraDebugRenderer.h"
+#include "Debug/CameraNodeEvaluationResultDebugBlock.h"
+#include "Debug/CameraNodeEvaluatorDebugBlock.h"
+#include "Debug/CameraPoseDebugBlock.h"
+#include "Debug/VariableTableDebugBlock.h"
+#include "HAL/IConsoleManager.h"
 #include "IGameplayCamerasLiveEditManager.h"
 #include "IGameplayCamerasModule.h"
+#include "Math/ColorList.h"
 #include "Modules/ModuleManager.h"
 #include "Nodes/Blends/PopBlendCameraNode.h"
 
@@ -24,6 +32,18 @@ FCameraNodeEvaluatorPtr UBlendStackCameraNode::OnBuildEvaluator(FCameraNodeEvalu
 namespace UE::Cameras
 {
 
+bool GGameplayCamerasDebugBlendStackShowUnchanged = false;
+static FAutoConsoleVariableRef CVarGameplayCamerasDebugBlendStackShowUnchanged(
+	TEXT("GameplayCameras.Debug.BlendStack.ShowUnchanged"),
+	GGameplayCamerasDebugBlendStackShowUnchanged,
+	TEXT(""));
+
+bool GGameplayCamerasDebugBlendStackShowVariableIds = false;
+static FAutoConsoleVariableRef CVarGameplayCamerasDebugBlendStackShowVariableIds(
+	TEXT("GameplayCameras.Debug.BlendStack.ShowVariableIds"),
+	GGameplayCamerasDebugBlendStackShowVariableIds,
+	TEXT(""));
+
 UE_DEFINE_CAMERA_NODE_EVALUATOR(FBlendStackCameraNodeEvaluator)
 
 void FBlendStackCameraNodeEvaluator::Push(const FBlendStackCameraPushParams& Params)
@@ -105,7 +125,10 @@ FCameraNodeEvaluatorChildrenView FBlendStackCameraNodeEvaluator::OnGetChildren()
 	FCameraNodeEvaluatorChildrenView View;
 	for (FCameraRigEntry& Entry : Entries)
 	{
-		View.Add(Entry.RootEvaluator);
+		if (Entry.RootEvaluator)
+		{
+			View.Add(Entry.RootEvaluator);
+		}
 	}
 	return View;
 }
@@ -135,10 +158,6 @@ void FBlendStackCameraNodeEvaluator::OnRun(const FCameraNodeEvaluationParams& Pa
 
 		FCameraNodeEvaluationResult& CurResult(Entry.Result);
 
-#if UE_GAMEPLAY_CAMERAS_DEBUG
-		CurResult.DebugBlockBuilder = OutResult.DebugBlockBuilder;
-#endif  // UE_GAMEPLAY_CAMERAS_DEBUG
-
 		if (!Entry.bIsFrozen)
 		{
 			// If the context in which this camera rig runs doesn't have a valid result,
@@ -170,10 +189,6 @@ void FBlendStackCameraNodeEvaluator::OnRun(const FCameraNodeEvaluationParams& Pa
 			// Only evaluate the blend via the root node.
 			Entry.RootEvaluator->Run(CurParams, CurResult);
 		}
-
-#if UE_GAMEPLAY_CAMERAS_DEBUG
-		CurResult.DebugBlockBuilder = nullptr;
-#endif  // UE_GAMEPLAY_CAMERAS_DEBUG
 	}
 
 	// Now blend all the results, keeping track of blends that have reached 100% so
@@ -418,5 +433,106 @@ void FBlendStackCameraNodeEvaluator::OnPostBuildAsset(const FGameplayCameraAsset
 
 #endif  // WITH_EDITOR
 
+#if UE_GAMEPLAY_CAMERAS_DEBUG
+
+void FBlendStackCameraNodeEvaluator::OnBuildDebugBlocks(const FCameraDebugBlockBuildParams& Params, FCameraDebugBlockBuilder& Builder)
+{
+	Builder.AttachDebugBlock<FBlendStackSummaryCameraDebugBlock>(*this);
+}
+
+FBlendStackCameraDebugBlock* FBlendStackCameraNodeEvaluator::BuildDetailedDebugBlock(const FCameraDebugBlockBuildParams& Params, FCameraDebugBlockBuilder& Builder)
+{
+	FBlendStackCameraDebugBlock& StackDebugBlock = Builder.BuildDebugBlock<FBlendStackCameraDebugBlock>(*this);
+	for (const FCameraRigEntry& Entry : Entries)
+	{
+		// Each entry has a wrapper debug block with 2 children blocks:
+		// - block for the blend
+		// - block for the result
+		FCameraDebugBlock& EntryDebugBlock = Builder.BuildDebugBlock<FCameraDebugBlock>();
+		StackDebugBlock.AddChild(&EntryDebugBlock);
+		{
+			if (Entry.RootEvaluator)
+			{
+				Builder.StartParentDebugBlockOverride(EntryDebugBlock);
+				{
+					FCameraNodeEvaluator* BlendEvaluator = Entry.RootEvaluator->GetBlendEvaluator();
+					BlendEvaluator->BuildDebugBlocks(Params, Builder);
+				}
+				Builder.EndParentDebugBlockOverride();
+			}
+			else
+			{
+				// Dummy debug block.
+				EntryDebugBlock.AddChild(&Builder.BuildDebugBlock<FCameraDebugBlock>());
+			}
+
+			FCameraNodeEvaluationResultDebugBlock& ResultDebugBlock = Builder.BuildDebugBlock<FCameraNodeEvaluationResultDebugBlock>();
+			EntryDebugBlock.AddChild(&ResultDebugBlock);
+			{
+				ResultDebugBlock.Initialize(Entry.Result, Builder);
+				ResultDebugBlock.GetCameraPoseDebugBlock()->WithShowUnchangedCVar(TEXT("GameplayCameras.Debug.BlendStack.ShowUnchanged"));
+				ResultDebugBlock.GetVariableTableDebugBlock()->WithShowVariableIdsCVar(TEXT("GameplayCameras.Debug.BlendStack.ShowVariableIds"));
+			}
+		}
+	}
+	return &StackDebugBlock;
+}
+
+UE_DEFINE_CAMERA_DEBUG_BLOCK(FBlendStackSummaryCameraDebugBlock);
+
+FBlendStackSummaryCameraDebugBlock::FBlendStackSummaryCameraDebugBlock()
+{
+}
+
+FBlendStackSummaryCameraDebugBlock::FBlendStackSummaryCameraDebugBlock(const FBlendStackCameraNodeEvaluator& InEvaluator)
+{
+	NumEntries = InEvaluator.Entries.Num();
+}
+
+void FBlendStackSummaryCameraDebugBlock::OnDebugDraw(const FCameraDebugBlockDrawParams& Params, FCameraDebugRenderer& Renderer)
+{
+	Renderer.AddText(TEXT("%d entries"), NumEntries);
+}
+
+UE_DEFINE_CAMERA_DEBUG_BLOCK(FBlendStackCameraDebugBlock);
+
+FBlendStackCameraDebugBlock::FBlendStackCameraDebugBlock()
+{
+}
+
+FBlendStackCameraDebugBlock::FBlendStackCameraDebugBlock(const FBlendStackCameraNodeEvaluator& InEvaluator)
+{
+	for (const FBlendStackCameraNodeEvaluator::FCameraRigEntry& Entry : InEvaluator.Entries)
+	{
+		FEntryDebugInfo EntryDebugInfo;
+		EntryDebugInfo.CameraRigName = Entry.CameraRig ? Entry.CameraRig->GetName() : TEXT("<no camera rig>");
+		Entries.Add(EntryDebugInfo);
+	}
+}
+
+void FBlendStackCameraDebugBlock::OnDebugDraw(const FCameraDebugBlockDrawParams& Params, FCameraDebugRenderer& Renderer)
+{
+	TArrayView<FCameraDebugBlock*> ChildrenView(GetChildren());
+
+	for (int32 Index = 0; Index < Entries.Num(); ++Index)
+	{
+		const FEntryDebugInfo& Entry(Entries[Index]);
+
+		Renderer.AddText(TEXT("{springgreen}[%d]{white} %s\n"), Index + 1, *Entry.CameraRigName);
+
+		if (ChildrenView.IsValidIndex(Index))
+		{
+			Renderer.AddIndent();
+			ChildrenView[Index]->DebugDraw(Params, Renderer);
+			Renderer.RemoveIndent();
+		}
+	}
+
+	// We've already manually renderered our children blocks.
+	Renderer.SkipAllBlocks();
+}
+
+#endif  // UE_GAMEPLAY_CAMERAS_DEBUG
+
 }  // namespace UE::Cameras
 
diff --git a/Engine/Plugins/Cameras/GameplayCameras/Source/GameplayCameras/Private/Core/BlendStackRootCameraNode.cpp b/Engine/Plugins/Cameras/GameplayCameras/Source/GameplayCameras/Private/Core/BlendStackRootCameraNode.cpp
--- a/Engine/Plugins/Cameras/GameplayCameras/Source/GameplayCameras/Private/Core/BlendStackRootCameraNode.cpp
+++ b/Engine/Plugins/Cameras/GameplayCameras/Source/GameplayCameras/Private/Core/BlendStackRootCameraNode.cpp
@@ -4,6 +4,9 @@
 
 #include "Core/BlendCameraNode.h"
 #include "Core/CameraRigAsset.h"
+#include "Debug/CameraDebugBlock.h"
+#include "Debug/CameraDebugBlockBuilder.h"
+#include "Debug/CameraDebugRenderer.h"
 
 #include UE_INLINE_GENERATED_CPP_BY_NAME(BlendStackRootCameraNode)
 
@@ -32,6 +35,12 @@ namespace UE::Cameras
 
 UE_DEFINE_CAMERA_NODE_EVALUATOR(FBlendStackRootCameraNodeEvaluator)
 
+UE_DECLARE_CAMERA_DEBUG_BLOCK_START(FBlendStackRootCameraDebugBlock)
+	UE_DECLARE_CAMERA_DEBUG_BLOCK_FIELD(FString, CameraRigAssetName)
+UE_DECLARE_CAMERA_DEBUG_BLOCK_END()
+
+UE_DEFINE_CAMERA_DEBUG_BLOCK(FBlendStackRootCameraDebugBlock)
+
 FCameraNodeEvaluatorChildrenView FBlendStackRootCameraNodeEvaluator::OnGetChildren()
 {
 	FCameraNodeEvaluatorChildrenView Children;
@@ -65,5 +74,53 @@ void FBlendStackRootCameraNodeEvaluator::OnRun(const FCameraNodeEvaluationParams
 	}
 }
 
+#if UE_GAMEPLAY_CAMERAS_DEBUG
+
+void FBlendStackRootCameraNodeEvaluator::OnBuildDebugBlocks(const FCameraDebugBlockBuildParams& Params, FCameraDebugBlockBuilder& Builder)
+{
+	if (BlendEvaluator)
+	{
+		BlendEvaluator->BuildDebugBlocks(Params, Builder);
+	}
+	else
+	{
+		// Dummy block.
+		Builder.StartChildDebugBlock<FCameraDebugBlock>();
+		Builder.EndChildDebugBlock();
+	}
+
+	if (RootEvaluator)
+	{
+		RootEvaluator->BuildDebugBlocks(Params, Builder);
+	}
+	else
+	{
+		// Dummy block.
+		Builder.StartChildDebugBlock<FCameraDebugBlock>();
+		Builder.EndChildDebugBlock();
+	}
+
+	Builder.SkipChildren();
+}
+
+void FBlendStackRootCameraDebugBlock::OnDebugDraw(const FCameraDebugBlockDrawParams& Params, FCameraDebugRenderer& Renderer)
+{
+	TArrayView<FCameraDebugBlock*> ChildrenView(GetChildren());
+
+	Renderer.AddText("{yellowgreen}[Blend]\n");
+	Renderer.AddIndent();
+	ChildrenView[0]->DebugDraw(Params, Renderer);
+	Renderer.RemoveIndent();
+
+	Renderer.AddText("{yellowgreen}[Camera Rig]\n");
+	Renderer.AddIndent();
+	ChildrenView[1]->DebugDraw(Params, Renderer);
+	Renderer.RemoveIndent();
+
+	Renderer.SkipAllBlocks();
+}
+
+#endif  // UE_GAMEPLAY_CAMERAS_DEBUG
+
 }  // namespace UE::Cameras
 
diff --git a/Engine/Plugins/Cameras/GameplayCameras/Source/GameplayCameras/Private/Core/BlendStackRootCameraNode.h b/Engine/Plugins/Cameras/GameplayCameras/Source/GameplayCameras/Private/Core/BlendStackRootCameraNode.h
--- a/Engine/Plugins/Cameras/GameplayCameras/Source/GameplayCameras/Private/Core/BlendStackRootCameraNode.h
+++ b/Engine/Plugins/Cameras/GameplayCameras/Source/GameplayCameras/Private/Core/BlendStackRootCameraNode.h
@@ -61,6 +61,10 @@ protected:
 	virtual void OnBuild(const FCameraNodeEvaluatorBuildParams& Params) override;
 	virtual void OnRun(const FCameraNodeEvaluationParams& Params, FCameraNodeEvaluationResult& OutResult) override;
 
+#if UE_GAMEPLAY_CAMERAS_DEBUG
+	virtual void OnBuildDebugBlocks(const FCameraDebugBlockBuildParams& Params, FCameraDebugBlockBuilder& Builder) override;
+#endif  // UE_GAMEPLAY_CAMERAS_DEBUG
+
 private:
 
 	FBlendCameraNodeEvaluator* BlendEvaluator = nullptr;
diff --git a/Engine/Plugins/Cameras/GameplayCameras/Source/GameplayCameras/Private/Core/CameraNodeEvaluator.cpp b/Engine/Plugins/Cameras/GameplayCameras/Source/GameplayCameras/Private/Core/CameraNodeEvaluator.cpp
--- a/Engine/Plugins/Cameras/GameplayCameras/Source/GameplayCameras/Private/Core/CameraNodeEvaluator.cpp
+++ b/Engine/Plugins/Cameras/GameplayCameras/Source/GameplayCameras/Private/Core/CameraNodeEvaluator.cpp
@@ -6,7 +6,7 @@
 #include "Debug/CameraDebugBlock.h"
 #include "Debug/CameraDebugBlockBuilder.h"
 #include "Debug/CameraDebugRenderer.h"
-#include "HAL/IConsoleManager.h"
+#include "Debug/CameraNodeEvaluatorDebugBlock.h"
 #include "UObject/UObjectGlobals.h"
 
 namespace UE::Cameras
@@ -14,57 +14,6 @@ namespace UE::Cameras
 
 UE_GAMEPLAY_CAMERAS_DEFINE_RTTI(FCameraNodeEvaluator)
 
-#if UE_GAMEPLAY_CAMERAS_DEBUG
-
-FString GGameplayCamerasDebugNodeTreeFilter;
-static FAutoConsoleVariableRef CVarGameplayCamerasDebugNodeTreeFilter(
-	TEXT("GameplayCameras.Debug.NodeTree.Filter"),
-	GGameplayCamerasDebugNodeTreeFilter,
-	TEXT("(Default: "". Filters the debug camera node tree by node name/type."));
-
-class FDefaultCameraNodeDebugBlock : public FCameraDebugBlock
-{
-public:
-
-	FDefaultCameraNodeDebugBlock() {}
-	FDefaultCameraNodeDebugBlock(const UCameraNode* InNode)
-	{
-		NodeClassName = InNode ? InNode->GetClass()->GetName() : TEXT("<null node>");
-	}
-
-protected:
-
-	virtual EDebugDrawResult OnDebugDraw(const FCameraDebugBlockDrawParams& Params, FCameraDebugRenderer& Renderer) override
-	{
-		Renderer.AddIndent();
-		
-		const bool bDoDebugDraw = (GGameplayCamerasDebugNodeTreeFilter.IsEmpty() ||
-				NodeClassName.Contains(GGameplayCamerasDebugNodeTreeFilter));
-		if (bDoDebugDraw)
-		{
-			Renderer.AddText(TEXT("[%s] "), *NodeClassName);
-		}
-		
-		return bDoDebugDraw ?  EDebugDrawResult::Default : EDebugDrawResult::SkipChildren;
-	}
-
-	virtual void OnPostDebugDraw(const FCameraDebugBlockDrawParams& Params, FCameraDebugRenderer& Renderer) override
-	{
-		Renderer.RemoveIndent();
-	}
-
-	virtual void OnSerialize(FArchive& Ar)
-	{
-		Ar << NodeClassName;
-	}
-
-private:
-
-	FString NodeClassName;
-};
-
-#endif  // UE_GAMEPLAY_CAMERAS_DEBUG
-
 void FCameraNodeEvaluationResult::Reset()
 {
 	CameraPose.Reset();
@@ -133,34 +82,47 @@ void FCameraNodeEvaluator::Run(const FCameraNodeEvaluationParams& Params, FCamer
 {
 	if (!PrivateCameraNode || PrivateCameraNode->bIsEnabled)
 	{
-#if UE_GAMEPLAY_CAMERAS_DEBUG
-		if (OutResult.DebugBlockBuilder)
-		{
-			CreateDebugBlock(Params, *OutResult.DebugBlockBuilder);
-		}
-#endif  // UE_GAMEPLAY_CAMERAS_DEBUG
-
 		OnRun(Params, OutResult);
-
-#if UE_GAMEPLAY_CAMERAS_DEBUG
-		if (OutResult.DebugBlockBuilder)
-		{
-			OutResult.DebugBlockBuilder->EndBlock();
-		}
-#endif  // UE_GAMEPLAY_CAMERAS_DEBUG
 	}
 }
 
 #if UE_GAMEPLAY_CAMERAS_DEBUG
 
-void FCameraNodeEvaluator::CreateDebugBlock(const FCameraNodeEvaluationParams& Params, FCameraDebugBlockBuilder& Builder)
+void FCameraNodeEvaluator::BuildDebugBlocks(const FCameraDebugBlockBuildParams& Params, FCameraDebugBlockBuilder& Builder)
 {
-	OnCreateDebugBlock(Params, Builder);
+	// Let's start by adding the default debug block for a node evaluator.
+	Builder.StartChildDebugBlock<FCameraNodeEvaluatorDebugBlock>(PrivateCameraNode);
+	{
+		// Then let the node evaluator attach or add other custom debug blocks.
+		const int32 PreviousLevel = Builder.GetHierarchyLevel();
+		OnBuildDebugBlocks(Params, Builder);
+		if (!ensureMsgf(
+				PreviousLevel == Builder.GetHierarchyLevel(), 
+				TEXT("Node evaluator added new children debug blocks but forgot to end them!")))
+		{
+			const int32 LevelsToEnd = Builder.GetHierarchyLevel() - PreviousLevel;
+			for (int32 Index = 0; Index < LevelsToEnd; ++Index)
+			{
+				Builder.EndChildDebugBlock();
+			}
+		}
+
+		// Build debug blocks for children node evaluators.
+		ECameraDebugBlockBuildVisitFlags VisitFlags = Builder.GetVisitFlags();
+		if (!EnumHasAnyFlags(VisitFlags, ECameraDebugBlockBuildVisitFlags::SkipChildren))
+		{
+			FCameraNodeEvaluatorChildrenView ChildrenView(GetChildren());
+			for (FCameraNodeEvaluator* Child : ChildrenView)
+			{
+				Child->BuildDebugBlocks(Params, Builder);
+			}
+		}
+	}
+	Builder.EndChildDebugBlock();
 }
 
-void FCameraNodeEvaluator::OnCreateDebugBlock(const FCameraNodeEvaluationParams& Params, FCameraDebugBlockBuilder& Builder)
+void FCameraNodeEvaluator::OnBuildDebugBlocks(const FCameraDebugBlockBuildParams& Params, FCameraDebugBlockBuilder& Builder)
 {
-	Builder.StartBlock<FDefaultCameraNodeDebugBlock>(PrivateCameraNode);
 }
 
 #endif  // UE_GAMEPLAY_CAMERAS_DEBUG
diff --git a/Engine/Plugins/Cameras/GameplayCameras/Source/GameplayCameras/Private/Core/CameraSystemEvaluator.cpp b/Engine/Plugins/Cameras/GameplayCameras/Source/GameplayCameras/Private/Core/CameraSystemEvaluator.cpp
--- a/Engine/Plugins/Cameras/GameplayCameras/Source/GameplayCameras/Private/Core/CameraSystemEvaluator.cpp
+++ b/Engine/Plugins/Cameras/GameplayCameras/Source/GameplayCameras/Private/Core/CameraSystemEvaluator.cpp
@@ -11,6 +11,8 @@
 #include "Debug/CameraDebugBlock.h"
 #include "Debug/CameraDebugBlockBuilder.h"
 #include "Debug/CameraDebugRenderer.h"
+#include "Debug/CategoryTitleDebugBlock.h"
+#include "Debug/RootCameraDebugBlock.h"
 #include "HAL/IConsoleManager.h"
 #include "IGameplayCamerasModule.h"
 #include "UObject/Package.h"
@@ -18,37 +20,11 @@
 
 DECLARE_CYCLE_STAT(TEXT("Camera System Eval"), CameraSystemEval_Total, STATGROUP_CameraSystem);
 
-#if UE_GAMEPLAY_CAMERAS_DEBUG
-
-bool GGameplayCamerasDebugEnable = false;
-static FAutoConsoleVariableRef CVarGameplayCamerasDebugEnable(
-	TEXT("GameplayCameras.Debug.Enable"),
-	GGameplayCamerasDebugEnable,
-	TEXT("(Default: false. Enables debug drawing for the GamplayCameras system."));
-
-#endif  // UE_GAMEPLAY_CAMERAS_DEBUG
-
 namespace UE::Cameras
 {
 
-#if UE_GAMEPLAY_CAMERAS_DEBUG
-
-class FRootCameraDebugBlock : public FCameraDebugBlock
-{
-public:
-
-	float LastDeltaTime = 0.f;
-
-protected:
-
-	virtual EDebugDrawResult OnDebugDraw(const FCameraDebugBlockDrawParams& Params, FCameraDebugRenderer& Renderer) override
-	{
-		return GGameplayCamerasDebugEnable ?
-			EDebugDrawResult::Default : EDebugDrawResult::SkipChildren;
-	}
-};
-
-#endif  // UE_GAMEPLAY_CAMERAS_DEBUG
+extern bool GGameplayCamerasDebugEnable;
+extern bool GGameplayCamerasDebugTrace;
 
 FCameraSystemEvaluator::FCameraSystemEvaluator()
 {
@@ -156,16 +132,6 @@ void FCameraSystemEvaluator::Update(const FCameraSystemEvaluationUpdateParams& P
 
 	RootNodeResult.Reset();
 
-#if UE_GAMEPLAY_CAMERAS_DEBUG
-	// Clear previous frame's debug info and make room for this frame's.
-	DebugBlockStorage.DestroyDebugBlocks();
-	RootDebugBlock = DebugBlockStorage.BuildDebugBlock<FRootCameraDebugBlock>();
-	RootDebugBlock->LastDeltaTime = Params.DeltaTime;
-
-	FCameraDebugBlockBuilder DebugBlockBuilder(DebugBlockStorage, *RootDebugBlock);
-	RootNodeResult.DebugBlockBuilder = &DebugBlockBuilder;
-#endif  // UE_GAMEPLAY_CAMERAS_DEBUG
-
 	// Run the root camera node.
 	RootEvaluator->Run(NodeParams, RootNodeResult);
 
@@ -173,10 +139,6 @@ void FCameraSystemEvaluator::Update(const FCameraSystemEvaluationUpdateParams& P
 	Result.CameraPose = RootNodeResult.CameraPose;
 	Result.bIsCameraCut = RootNodeResult.bIsCameraCut;
 	Result.bIsValid = true;
-
-#if UE_GAMEPLAY_CAMERAS_DEBUG
-	RootNodeResult.DebugBlockBuilder = nullptr;
-#endif  // UE_GAMEPLAY_CAMERAS_DEBUG
 }
 
 void FCameraSystemEvaluator::GetEvaluatedCameraView(FMinimalViewInfo& DesiredView)
@@ -189,17 +151,25 @@ void FCameraSystemEvaluator::GetEvaluatedCameraView(FMinimalViewInfo& DesiredVie
 
 #if UE_GAMEPLAY_CAMERAS_DEBUG
 
-void FCameraSystemEvaluator::DebugDraw(const FCameraSystemDebugDrawParams& Params)
+void FCameraSystemEvaluator::DebugUpdate(const FCameraSystemDebugUpdateParams& Params)
 {
-	if (RootDebugBlock)
+	if (!GGameplayCamerasDebugTrace && !GGameplayCamerasDebugEnable)
 	{
-		FCameraDebugBlockDrawParams BlockParams;
-		BlockParams.DeltaTime = RootDebugBlock->LastDeltaTime;
+		return;
+	}
+
+	// Clear previous frame's debug info and make room for this frame's.
+	DebugBlockStorage.DestroyDebugBlocks();
 
-		FCameraDebugRenderer Renderer(Params.Canvas);
+	// Create the root debug block and start building more.
+	RootDebugBlock = DebugBlockStorage.BuildDebugBlock<FRootCameraDebugBlock>();
 
-		RootDebugBlock->DebugDraw(BlockParams, Renderer);
-	}
+	FCameraDebugBlockBuildParams BuildParams;
+	FCameraDebugBlockBuilder DebugBlockBuilder(DebugBlockStorage, *RootDebugBlock);
+	RootDebugBlock->BuildDebugBlocks(*this, BuildParams, DebugBlockBuilder);
+	
+	FCameraDebugRenderer Renderer(Params.Canvas);
+	RootDebugBlock->RootDebugDraw(Renderer);
 }
 
 #endif  // UE_GAMEPLAY_CAMERAS_DEBUG
diff --git a/Engine/Plugins/Cameras/GameplayCameras/Source/GameplayCameras/Private/Core/DefaultRootCameraNode.cpp b/Engine/Plugins/Cameras/GameplayCameras/Source/GameplayCameras/Private/Core/DefaultRootCameraNode.cpp
--- a/Engine/Plugins/Cameras/GameplayCameras/Source/GameplayCameras/Private/Core/DefaultRootCameraNode.cpp
+++ b/Engine/Plugins/Cameras/GameplayCameras/Source/GameplayCameras/Private/Core/DefaultRootCameraNode.cpp
@@ -4,6 +4,9 @@
 
 #include "Core/BlendStackCameraNode.h"
 #include "Core/CameraRigAsset.h"
+#include "Debug/BlendStacksCameraDebugBlock.h"
+#include "Debug/CameraDebugBlockBuilder.h"
+#include "Debug/RootCameraDebugBlock.h"
 
 #include UE_INLINE_GENERATED_CPP_BY_NAME(DefaultRootCameraNode)
 
@@ -59,7 +62,6 @@ FCameraNodeEvaluatorChildrenView FDefaultRootCameraNodeEvaluator::OnGetChildren(
 	return FCameraNodeEvaluatorChildrenView({ BaseLayer, MainLayer, GlobalLayer, VisualLayer });
 }
 
-
 void FDefaultRootCameraNodeEvaluator::OnRun(const FCameraNodeEvaluationParams& Params, FCameraNodeEvaluationResult& OutResult)
 {
 	BaseLayer->Run(Params, OutResult);
@@ -97,5 +99,32 @@ void FDefaultRootCameraNodeEvaluator::OnActivateCameraRig(const FActivateCameraR
 	}
 }
 
+#if UE_GAMEPLAY_CAMERAS_DEBUG
+
+UE_DECLARE_CAMERA_DEBUG_BLOCK_START(FDefaultRootCameraNodeEvaluatorDebugBlock)
+UE_DECLARE_CAMERA_DEBUG_BLOCK_END()
+
+UE_DEFINE_CAMERA_DEBUG_BLOCK(FDefaultRootCameraNodeEvaluatorDebugBlock)
+
+void FDefaultRootCameraNodeEvaluator::OnBuildDebugBlocks(const FCameraDebugBlockBuildParams& Params, FCameraDebugBlockBuilder& Builder)
+{
+	// Create the debug block that shows the overall blend stack layers.
+	FBlendStacksCameraDebugBlock& DebugBlock = Builder.BuildDebugBlock<FBlendStacksCameraDebugBlock>();
+	{
+		DebugBlock.AddBlendStack(TEXT("Base Layer"), BaseLayer->BuildDetailedDebugBlock(Params, Builder));
+		DebugBlock.AddBlendStack(TEXT("Main Layer"), MainLayer->BuildDetailedDebugBlock(Params, Builder));
+		DebugBlock.AddBlendStack(TEXT("Global Layer"), GlobalLayer->BuildDetailedDebugBlock(Params, Builder));
+		DebugBlock.AddBlendStack(TEXT("Visual Layer"), VisualLayer->BuildDetailedDebugBlock(Params, Builder));
+	}
+
+	Builder.GetRootDebugBlock().AddChild(&DebugBlock);
+}
+
+void FDefaultRootCameraNodeEvaluatorDebugBlock::OnDebugDraw(const FCameraDebugBlockDrawParams& Params, FCameraDebugRenderer& Renderer)
+{
+}
+
+#endif  // UE_GAMEPLAY_CAMERAS_DEBUG
+
 }  // namespace UE::Cameras
 
diff --git a/Engine/Plugins/Cameras/GameplayCameras/Source/GameplayCameras/Private/Debug/BlendStacksCameraDebugBlock.cpp b/Engine/Plugins/Cameras/GameplayCameras/Source/GameplayCameras/Private/Debug/BlendStacksCameraDebugBlock.cpp
new file mode 100644
--- /dev/null
+++ b/Engine/Plugins/Cameras/GameplayCameras/Source/GameplayCameras/Private/Debug/BlendStacksCameraDebugBlock.cpp
@@ -0,0 +1,72 @@
+// Copyright Epic Games, Inc. All Rights Reserved.
+
+#include "Debug/BlendStacksCameraDebugBlock.h"
+
+#include "Core/BlendStackCameraNode.h"
+#include "Debug/CameraDebugCategories.h"
+#include "Debug/CameraDebugColors.h"
+#include "Debug/CameraDebugRenderer.h"
+#include "Debug/CameraNodeEvaluatorDebugBlock.h"
+#include "Math/ColorList.h"
+
+#if UE_GAMEPLAY_CAMERAS_DEBUG
+
+namespace UE::Cameras
+{
+
+UE_DEFINE_CAMERA_DEBUG_BLOCK(FBlendStacksCameraDebugBlock)
+
+FBlendStacksCameraDebugBlock::FBlendStacksCameraDebugBlock()
+{
+}
+
+void FBlendStacksCameraDebugBlock::AddBlendStack(const FString& InBlendStackName, FBlendStackCameraDebugBlock* InDebugBlock)
+{
+	BlendStackNames.Add(InBlendStackName);
+	AddChild(InDebugBlock);
+}
+
+void FBlendStacksCameraDebugBlock::OnDebugDraw(const FCameraDebugBlockDrawParams& Params, FCameraDebugRenderer& Renderer)
+{
+	if (!Params.IsCategoryActive(FCameraDebugCategories::BlendStacks))
+	{
+		Renderer.SkipAllBlocks();
+		return;
+	}
+	
+	const FCameraDebugColors& Colors = FCameraDebugColors::Get();
+
+	Renderer.SetTextColor(Colors.Title);
+	Renderer.AddText(TEXT("Blend Stacks\n\n"));
+	Renderer.SetTextColor(Colors.Default);
+
+	TArrayView<FCameraDebugBlock*> BlendStackBlocks(GetChildren());
+
+	const int32 MaxIndex = FMath::Max(BlendStackBlocks.Num(), BlendStackNames.Num());
+
+	for (int32 Index = 0; Index < MaxIndex; ++Index)
+	{
+		Renderer.SetTextColor(Colors.Notice);
+		Renderer.AddText(BlendStackNames.IsValidIndex(Index) ? BlendStackNames[Index] : TEXT("<unnamed blend stack>"));
+		Renderer.NewLine();
+		Renderer.SetTextColor(Colors.Default);
+
+		if (BlendStackBlocks.IsValidIndex(Index))
+		{
+			FCameraDebugBlock* BlendStackBlock = BlendStackBlocks[Index];
+			BlendStackBlock->DebugDraw(Params, Renderer);
+		}
+		else
+		{
+			Renderer.AddText(TEXT("<missing blend stack>"));
+		}
+	}
+
+	// We've already manually renderered our children blocks.
+	Renderer.SkipAllBlocks();
+}
+
+}  // namespace UE::Cameras
+
+#endif  // UE_GAMEPLAY_CAMERAS_DEBUG
+
diff --git a/Engine/Plugins/Cameras/GameplayCameras/Source/GameplayCameras/Private/Debug/CameraDebugBlock.cpp b/Engine/Plugins/Cameras/GameplayCameras/Source/GameplayCameras/Private/Debug/CameraDebugBlock.cpp
--- a/Engine/Plugins/Cameras/GameplayCameras/Source/GameplayCameras/Private/Debug/CameraDebugBlock.cpp
+++ b/Engine/Plugins/Cameras/GameplayCameras/Source/GameplayCameras/Private/Debug/CameraDebugBlock.cpp
@@ -2,44 +2,72 @@
 
 #include "Debug/CameraDebugBlock.h"
 
+#include "Debug/CameraDebugRenderer.h"
+
 #if UE_GAMEPLAY_CAMERAS_DEBUG
 
 namespace UE::Cameras
 {
 
+bool FCameraDebugBlockDrawParams::IsCategoryActive(const FString& InCategory) const
+{
+	return ActiveCategories.Contains(InCategory);
+}
+
+UE_GAMEPLAY_CAMERAS_DEFINE_RTTI(FCameraDebugBlock)
+
 TArray<FCameraDebugBlockField*> FCameraDebugBlock::StaticFields;
 
+void FCameraDebugBlock::Attach(FCameraDebugBlock* InAttachment)
+{
+	Attachments.Add(InAttachment);
+}
+
 void FCameraDebugBlock::AddChild(FCameraDebugBlock* InChild)
 {
 	Children.Add(InChild);
 }
 
+
 void FCameraDebugBlock::DebugDraw(const FCameraDebugBlockDrawParams& Params, FCameraDebugRenderer& Renderer)
 {
-	EDebugDrawResult Result = OnDebugDraw(Params, Renderer);
-	if (!EnumHasAnyFlags(Result, EDebugDrawResult::SkipChildren))
+	Renderer.ResetVisitFlags();
+
+	OnDebugDraw(Params, Renderer);
+	
+	ECameraDebugDrawVisitFlags VisitFlags = Renderer.GetVisitFlags();
+
+	if (!EnumHasAnyFlags(VisitFlags, ECameraDebugDrawVisitFlags::SkipAttachedBlocks) && !Attachments.IsEmpty())
+	{
+		for (FCameraDebugBlock* Attachment : Attachments)
+		{
+			Attachment->DebugDraw(Params, Renderer);
+		}
+	}
+
+	if (!EnumHasAnyFlags(VisitFlags, ECameraDebugDrawVisitFlags::SkipChildrenBlocks) && !Children.IsEmpty())
 	{
+		Renderer.AddIndent();
+
 		for (FCameraDebugBlock* Child : Children)
 		{
 			Child->DebugDraw(Params, Renderer);
 		}
+
+		Renderer.RemoveIndent();
 	}
-	OnPostDebugDraw(Params, Renderer);
-}
 
-EDebugDrawResult FCameraDebugBlock::OnDebugDraw(const FCameraDebugBlockDrawParams& Params, FCameraDebugRenderer& Renderer)
-{
-	return EDebugDrawResult::Default;
+	OnPostDebugDraw(Params, Renderer);
 }
 
 void FCameraDebugBlock::Serialize(FArchive& Ar)
 {
+	OnSerialize(Ar);
+
 	for (FCameraDebugBlockField* StaticField : StaticFields)
 	{
 		StaticField->SerializeField(this, Ar);
 	}
-
-	OnSerialize(Ar);
 }
 
 }  // namespace UE::Cameras
diff --git a/Engine/Plugins/Cameras/GameplayCameras/Source/GameplayCameras/Private/Debug/CameraDebugBlockBuilder.cpp b/Engine/Plugins/Cameras/GameplayCameras/Source/GameplayCameras/Private/Debug/CameraDebugBlockBuilder.cpp
--- a/Engine/Plugins/Cameras/GameplayCameras/Source/GameplayCameras/Private/Debug/CameraDebugBlockBuilder.cpp
+++ b/Engine/Plugins/Cameras/GameplayCameras/Source/GameplayCameras/Private/Debug/CameraDebugBlockBuilder.cpp
@@ -2,32 +2,80 @@
 
 #include "Debug/CameraDebugBlockBuilder.h"
 
+#include "Debug/CameraDebugBlock.h"
+#include "Debug/RootCameraDebugBlock.h"
+
 #if UE_GAMEPLAY_CAMERAS_DEBUG
 
 namespace UE::Cameras
 {
 
-FCameraDebugBlockBuilder::FCameraDebugBlockBuilder(FCameraDebugBlockStorage& InStorage, FCameraDebugBlock& InRootBlock)
+FCameraDebugBlockBuilder::FCameraDebugBlockBuilder(FCameraDebugBlockStorage& InStorage, FRootCameraDebugBlock& InRootBlock)
 	: Storage(InStorage)
+	, RootBlock(InRootBlock)
 {
+	// We should always have the root block in the working hierarchy.
 	CurrentHierarchy.Add(&InRootBlock);
 }
 
-void FCameraDebugBlockBuilder::OnStartBlock(FCameraDebugBlock& InNewBlock)
+void FCameraDebugBlockBuilder::OnAttachDebugBlock(FCameraDebugBlock* InNewBlock)
+{
+	if (ensureMsgf(CurrentHierarchy.Num() > 0, TEXT("Can't attach block, no current block defined!")))
+	{
+		CurrentHierarchy.Last()->Attach(InNewBlock);
+	}
+}
+
+void FCameraDebugBlockBuilder::OnStartChildDebugBlock(FCameraDebugBlock* InNewBlock)
+{
+	if (ensureMsgf(CurrentHierarchy.Num() > 0, TEXT("Can't add child block, no current block defined!")))
+	{
+		CurrentHierarchy.Last()->AddChild(InNewBlock);
+	}
+	CurrentHierarchy.Add(InNewBlock);
+}
+
+void FCameraDebugBlockBuilder::EndChildDebugBlock()
 {
-	checkf(CurrentHierarchy.Num() > 0, TEXT("There should always be the root block in the current hierarchy!"));
+	if (ensureMsgf(CurrentHierarchy.Num() > 0, TEXT("Can't end block, no current block defined!")))
+	{
+		CurrentHierarchy.Pop();
+	}
+}
 
-	// Add the new block under the current parent.
-	FCameraDebugBlock* ParentBlock = CurrentHierarchy.Last();
-	ParentBlock->AddChild(&InNewBlock);
-	CurrentHierarchy.Add(&InNewBlock);
+void FCameraDebugBlockBuilder::SkipChildren()
+{
+	VisitFlags |= ECameraDebugBlockBuildVisitFlags::SkipChildren;
 }
 
-void FCameraDebugBlockBuilder::EndBlock()
+void FCameraDebugBlockBuilder::ResetVisitFlags()
 {
-	checkf(CurrentHierarchy.Num() > 0, TEXT("Can't end block, no current block defined!"));
+	VisitFlags = ECameraDebugBlockBuildVisitFlags::None;
+}
 
-	CurrentHierarchy.Pop();
+void FCameraDebugBlockBuilder::StartParentDebugBlockOverride(FCameraDebugBlock& InNewParent)
+{
+	if (ensureMsgf(HierarchyOverrideStart == INDEX_NONE, TEXT("Can't override parenting, an override is already in progress.")))
+	{
+		HierarchyOverrideStart = CurrentHierarchy.Num();
+		CurrentHierarchy.Add(&InNewParent);
+	}
+}
+
+void FCameraDebugBlockBuilder::EndParentDebugBlockOverride()
+{
+	if (ensureMsgf(HierarchyOverrideStart != INDEX_NONE, TEXT("No parenting override active.")))
+	{
+		ensureMsgf(
+				HierarchyOverrideStart == CurrentHierarchy.Num() - 1,
+				TEXT("Mismatch between expected parenting override and actual hierarchy level. "
+					 "Did you forget to end some child blocks?"));
+		while (CurrentHierarchy.Num() > HierarchyOverrideStart)
+		{
+			CurrentHierarchy.Pop();
+		}
+		HierarchyOverrideStart = INDEX_NONE;
+	}
 }
 
 }  // namespace UE::Cameras
diff --git a/Engine/Plugins/Cameras/GameplayCameras/Source/GameplayCameras/Private/Debug/CameraDebugCategories.cpp b/Engine/Plugins/Cameras/GameplayCameras/Source/GameplayCameras/Private/Debug/CameraDebugCategories.cpp
new file mode 100644
--- /dev/null
+++ b/Engine/Plugins/Cameras/GameplayCameras/Source/GameplayCameras/Private/Debug/CameraDebugCategories.cpp
@@ -0,0 +1,21 @@
+// Copyright Epic Games, Inc. All Rights Reserved.
+
+#include "Debug/CameraDebugCategories.h"
+
+#include "Containers/UnrealString.h"
+
+#if UE_GAMEPLAY_CAMERAS_DEBUG
+
+namespace UE::Cameras
+{
+
+const FString FCameraDebugCategories::NodeTree("nodetree");
+const FString FCameraDebugCategories::DirectorTree("directortree");
+const FString FCameraDebugCategories::BlendStacks("blendstacks");
+const FString FCameraDebugCategories::PoseStats("posestats");
+const FString FCameraDebugCategories::Viewfinder("viewfinder");
+
+}  // namespace UE::Cameras
+
+#endif  // UE_GAMEPLAY_CAMERAS_DEBUG
+
diff --git a/Engine/Plugins/Cameras/GameplayCameras/Source/GameplayCameras/Private/Debug/CameraDebugColors.cpp b/Engine/Plugins/Cameras/GameplayCameras/Source/GameplayCameras/Private/Debug/CameraDebugColors.cpp
new file mode 100644
--- /dev/null
+++ b/Engine/Plugins/Cameras/GameplayCameras/Source/GameplayCameras/Private/Debug/CameraDebugColors.cpp
@@ -0,0 +1,89 @@
+// Copyright Epic Games, Inc. All Rights Reserved.
+
+#include "Debug/CameraDebugColors.h"
+
+#include "Misc/Optional.h"
+
+#if UE_GAMEPLAY_CAMERAS_DEBUG
+
+namespace UE::Cameras
+{
+
+TMap<FString, FColor> FCameraDebugColors::ColorMap;
+
+const FCameraDebugColors& FCameraDebugColors::Get()
+{
+	static FCameraDebugColors StaticInstance;
+	return StaticInstance;
+}
+
+void FCameraDebugColors::UpdateColorMap(const FCameraDebugColors& Instance)
+{
+	ColorMap.Reset();
+	ColorMap.Add(TEXT("cam_title"), Instance.Title);
+	ColorMap.Add(TEXT("cam_default"), Instance.Default);
+	ColorMap.Add(TEXT("cam_passive"), Instance.Passive);
+	ColorMap.Add(TEXT("cam_verypassive"), Instance.VeryPassive);
+	ColorMap.Add(TEXT("cam_highlighted"), Instance.Hightlighted);
+	ColorMap.Add(TEXT("cam_notice"), Instance.Notice);
+	ColorMap.Add(TEXT("cam_notice2"), Instance.Notice2);
+	ColorMap.Add(TEXT("cam_good"), Instance.Good);
+	ColorMap.Add(TEXT("cam_warning"), Instance.Warning);
+	ColorMap.Add(TEXT("cam_error"), Instance.Error);
+	ColorMap.Add(TEXT("cam_background"), Instance.Background);
+}
+
+TOptional<FColor> FCameraDebugColors::GetFColorByName(const FString& InColorName)
+{
+	if (ColorMap.IsEmpty())
+	{
+		UpdateColorMap(FCameraDebugColors::Get());
+	}
+	if (const FColor* Color = ColorMap.Find(InColorName))
+	{
+		return TOptional<FColor>(*Color);
+	}
+	return TOptional<FColor>();
+}
+
+FCameraDebugColors::FCameraDebugColors()
+{
+	// Default to colors inspired by the Solarized palette.
+	// Other palettes could be implemented.
+	//
+	//    SOLARIZED HEX     RGB        
+    //    --------- ------- -----------
+    //    base03    #002b36   0  43  54
+    //    base02    #073642   7  54  66
+    //    base01    #586e75  88 110 117
+    //    base00    #657b83 101 123 131
+    //    base0     #839496 131 148 150
+    //    base1     #93a1a1 147 161 161
+    //    base2     #eee8d5 238 232 213
+    //    base3     #fdf6e3 253 246 227
+    //    yellow    #b58900 181 137   0
+    //    orange    #cb4b16 203  75  22
+    //    red       #dc322f 220  50  47
+    //    magenta   #d33682 211  54 130
+    //    violet    #6c71c4 108 113 196
+    //    blue      #268bd2  38 139 210
+    //    cyan      #2aa198  42 161 152
+    //    green     #859900 133 153   0
+	//
+	Title = FColor(38, 139, 210); // blue
+	Default = FColor(238, 232, 213); // base2
+	Passive = FColor(147, 161, 161); // base1
+	VeryPassive = FColor(101, 123, 131); // base00
+	Hightlighted = FColor(253, 246, 227); // base3
+	Notice = FColor(42, 161, 152); // cyan
+	Notice2 = FColor(211, 54, 130); // magenta
+	Good = FColor(133, 153, 0); // green
+	Warning = FColor(181, 137, 0); // yellow
+	Error = FColor(220, 50, 47); // red
+	Background = FColor(7, 54, 66); // base02
+}
+
+}  // namespace UE::Cameras
+
+#endif  // UE_GAMEPLAY_CAMERAS_DEBUG
+
diff --git a/Engine/Plugins/Cameras/GameplayCameras/Source/GameplayCameras/Private/Debug/CameraDebugRenderer.cpp b/Engine/Plugins/Cameras/GameplayCameras/Source/GameplayCameras/Private/Debug/CameraDebugRenderer.cpp
--- a/Engine/Plugins/Cameras/GameplayCameras/Source/GameplayCameras/Private/Debug/CameraDebugRenderer.cpp
+++ b/Engine/Plugins/Cameras/GameplayCameras/Source/GameplayCameras/Private/Debug/CameraDebugRenderer.cpp
@@ -2,13 +2,20 @@
 
 #include "Debug/CameraDebugRenderer.h"
 
+#include "Algo/Find.h"
+#include "Debug/CameraDebugColors.h"
+#include "Debug/DebugTextRenderer.h"
 #include "Engine/Canvas.h"
 #include "Engine/Engine.h"
 #include "Engine/Font.h"
 #include "HAL/IConsoleManager.h"
+#include "Misc/TVariant.h"
 
 #if UE_GAMEPLAY_CAMERAS_DEBUG
 
+namespace UE::Cameras
+{
+
 int32 GGameplayCamerasDebugMargin = 10;
 static FAutoConsoleVariableRef CVarGameplayCamerasDebugMargin(
 	TEXT("GameplayCameras.Debug.Margin"),
@@ -21,8 +28,11 @@ static FAutoConsoleVariableRef CVarGameplayCamerasDebugIndent(
 	GGameplayCamerasDebugIndent,
 	TEXT("(Default: 20px. The indent for rendering Gameplay Cameras debug text."));
 
-namespace UE::Cameras
-{
+int32 GGameplayCamerasDebugBackgroundDepthSortKey = 1;
+static FAutoConsoleVariableRef CVarGameplayCamerasDebugBackgroundDepthSortKey(
+	TEXT("GameplayCameras.Debug.BackgroundDepthSortKey"),
+	GGameplayCamerasDebugBackgroundDepthSortKey,
+	TEXT(""));
 
 FCameraDebugRenderer::FCameraDebugRenderer(UCanvas* InCanvas)
 	: Canvas(InCanvas)
@@ -33,7 +43,7 @@ FCameraDebugRenderer::FCameraDebugRenderer(UCanvas* InCanvas)
 	RenderFont = GEngine->GetSmallFont();
 	MaxCharHeight = RenderFont->GetMaxCharHeight();
 
-	NextBlockY = 0;
+	NextDrawPosition = FVector2f{ (float)GGameplayCamerasDebugMargin, (float)GGameplayCamerasDebugMargin };
 }
 
 FCameraDebugRenderer::~FCameraDebugRenderer()
@@ -68,21 +78,51 @@ void FCameraDebugRenderer::AddTextImpl(const TCHAR* Buffer)
 	LineBuilder.Append(Buffer);
 }
 
+bool FCameraDebugRenderer::NewLine(bool bSkipIfEmptyLine)
+{
+	FlushText();
+
+	const float IndentMargin = GetIndentMargin();
+	const bool bIsLineEmpty = FMath::IsNearlyEqual(NextDrawPosition.X, IndentMargin);
+	if (!bIsLineEmpty || !bSkipIfEmptyLine)
+	{
+		NextDrawPosition.X = IndentMargin;
+		NextDrawPosition.Y += MaxCharHeight;
+		return true;
+	}
+	return false;
+}
+
+FColor FCameraDebugRenderer::GetTextColor() const
+{
+	return Canvas->DrawColor;
+}
+
+FColor FCameraDebugRenderer::SetTextColor(const FColor& Color)
+{
+	FlushText();
+	FColor ReturnColor = Canvas->DrawColor;
+	Canvas->SetDrawColor(Color);
+	return ReturnColor;
+}
+
+float FCameraDebugRenderer::GetIndentMargin() const
+{
+	return (float)(GGameplayCamerasDebugMargin + IndentLevel * GGameplayCamerasDebugIndent);
+}
+
 void FCameraDebugRenderer::FlushText()
 {
 	if (LineBuilder.Len() > 0)
 	{
-		if (NextBlockY < Canvas->ClipY)
+		if (NextDrawPosition.Y < Canvas->ClipY)
 		{
-			int32 BlockHeight, BlockWidth;
-			const TCHAR* TextBlock = LineBuilder.ToString();
-			RenderFont->GetStringHeightAndWidth(TextBlock, BlockHeight, BlockWidth);
+			FDebugTextRenderer TextRenderer(Canvas, RenderFont);
+			TextRenderer.LeftMargin = GetIndentMargin();
+			TextRenderer.RenderText(NextDrawPosition, LineBuilder.ToView());
 
-			float X = GGameplayCamerasDebugMargin + IndentLevel * GGameplayCamerasDebugIndent;
-			float Y = GGameplayCamerasDebugMargin + NextBlockY;
-			Canvas->DrawText(RenderFont, TextBlock, X, Y);
-
-			NextBlockY += BlockHeight;
+			NextDrawPosition = TextRenderer.GetEndDrawPosition();
+			RightMargin = FMath::Max(RightMargin, TextRenderer.GetRightMargin());
 		}
 		// else: text is going off-screen.
 
@@ -92,23 +132,80 @@ void FCameraDebugRenderer::FlushText()
 
 void FCameraDebugRenderer::AddIndent()
 {
-	// Flush any remaining text we have on the current indent level.
-	FlushText();
+	// Flush any remaining text we have on the current indent level and move
+	// to a new line, unless the current line was empty.
+	NewLine(true);
 
 	++IndentLevel;
+
+	// The next draw position is at the beginning of a new line (or the beginning
+	// of an old line if it was empty). Either way, it's left at the previous
+	// indent level, so we need to bump it to the right.
+	NextDrawPosition.X = GetIndentMargin();
 }
 
 void FCameraDebugRenderer::RemoveIndent()
 {
-	// Flush any remaining text we have on the current indent level.
-	FlushText();
+	// Flush any remaining text we have on the current indent level and move
+	// to a new line, unless the current line was empty.
+	NewLine(true);
 
 	if (ensureMsgf(IndentLevel > 0, TEXT("Can't go into negative indenting!")))
 	{
 		--IndentLevel;
+
+		// See comment in AddIndent().
+		NextDrawPosition.X = GetIndentMargin();
 	}
 }
 
+void FCameraDebugRenderer::DrawTextBackgroundTile(float Opacity)
+{
+	const float IndentMargin = GetIndentMargin();
+	const bool bIsLineEmpty = FMath::IsNearlyEqual(NextDrawPosition.X, IndentMargin);
+	const float TextBottom = bIsLineEmpty ? NextDrawPosition.Y : NextDrawPosition.Y + MaxCharHeight;
+
+	const float OuterMargin = GGameplayCamerasDebugMargin / 2.f;
+	FVector2D TopLeft(OuterMargin, OuterMargin);
+	FVector2D TileSize(RightMargin + OuterMargin, TextBottom + OuterMargin);
+
+	const FColor BackgroundColor = FCameraDebugColors::Get().Background.WithAlpha((uint8)(Opacity * 255));
+
+	// Draw the background behind the text.
+	Canvas->Canvas->PushDepthSortKey(GGameplayCamerasDebugBackgroundDepthSortKey);
+	{
+		FCanvasTileItem BackgroundTile(TopLeft, TileSize, BackgroundColor);
+		BackgroundTile.BlendMode = SE_BLEND_Translucent;
+		Canvas->DrawItem(BackgroundTile);
+	}
+	Canvas->Canvas->PopDepthSortKey();
+}
+
+void FCameraDebugRenderer::SkipAttachedBlocks()
+{
+	VisitFlags |= ECameraDebugDrawVisitFlags::SkipAttachedBlocks;
+}
+
+void FCameraDebugRenderer::SkipChildrenBlocks()
+{
+	VisitFlags |= ECameraDebugDrawVisitFlags::SkipChildrenBlocks;
+}
+
+void FCameraDebugRenderer::SkipAllBlocks()
+{
+	VisitFlags |= (ECameraDebugDrawVisitFlags::SkipAttachedBlocks | ECameraDebugDrawVisitFlags::SkipChildrenBlocks);
+}
+
+ECameraDebugDrawVisitFlags FCameraDebugRenderer::GetVisitFlags() const
+{
+	return VisitFlags;
+}
+
+void FCameraDebugRenderer::ResetVisitFlags()
+{
+	VisitFlags = ECameraDebugDrawVisitFlags::None;
+}
+
 }  // namespace UE::Cameras
 
 #endif  // UE_GAMEPLAY_CAMERAS_DEBUG
diff --git a/Engine/Plugins/Cameras/GameplayCameras/Source/GameplayCameras/Private/Debug/CameraDirectorTreeDebugBlock.cpp b/Engine/Plugins/Cameras/GameplayCameras/Source/GameplayCameras/Private/Debug/CameraDirectorTreeDebugBlock.cpp
new file mode 100644
--- /dev/null
+++ b/Engine/Plugins/Cameras/GameplayCameras/Source/GameplayCameras/Private/Debug/CameraDirectorTreeDebugBlock.cpp
@@ -0,0 +1,94 @@
+// Copyright Epic Games, Inc. All Rights Reserved.
+
+#include "Debug/CameraDirectorTreeDebugBlock.h"
+
+#include "Core/CameraAsset.h"
+#include "Core/CameraEvaluationContext.h"
+#include "Core/CameraEvaluationContextStack.h"
+#include "Debug/CameraDebugBlockBuilder.h"
+#include "Debug/CameraDebugColors.h"
+#include "Debug/CameraDebugRenderer.h"
+#include "Debug/CameraPoseDebugBlock.h"
+#include "HAL/IConsoleManager.h"
+
+#if UE_GAMEPLAY_CAMERAS_DEBUG
+
+namespace UE::Cameras
+{
+
+bool GGameplayCamerasDebugContextInitialResultShowUnchanged = false;
+static FAutoConsoleVariableRef CVarGameplayCamerasDebugContextInitialResultShowUnchanged(
+	TEXT("GameplayCameras.Debug.ContextInitialResult.ShowUnchanged"),
+	GGameplayCamerasDebugContextInitialResultShowUnchanged,
+	TEXT(""));
+
+UE_DEFINE_CAMERA_DEBUG_BLOCK(FCameraDirectorTreeDebugBlock)
+
+FCameraDirectorTreeDebugBlock::FCameraDirectorTreeDebugBlock()
+{
+}
+
+void FCameraDirectorTreeDebugBlock::Initialize(const FCameraEvaluationContextStack& ContextStack, FCameraDebugBlockBuilder& Builder)
+{
+	const int32 NumContexts = ContextStack.NumContexts();
+	for (int32 Index = 0; Index < NumContexts; ++Index)
+	{
+		const FCameraEvaluationContextStack::FContextEntry& Entry(ContextStack.Entries[Index]);
+		TSharedPtr<FCameraEvaluationContext> Context = Entry.WeakContext.Pin();
+
+		FDirectorDebugInfo EntryDebugInfo;
+		EntryDebugInfo.CameraAssetName = Context ? Context->GetCameraAsset()->GetName() : TEXT("<no camera asset>");
+		CameraDirectors.Add(EntryDebugInfo);
+
+		const FCameraNodeEvaluationResult& InitialResult = Context->GetInitialResult();
+		AddChild(&Builder.BuildDebugBlock<FCameraPoseDebugBlock>(InitialResult.CameraPose)
+				.WithShowUnchangedCVar(TEXT("GameplayCameras.Debug.ContextInitialResult.ShowUnchanged")));
+	}
+}
+
+void FCameraDirectorTreeDebugBlock::OnDebugDraw(const FCameraDebugBlockDrawParams& Params, FCameraDebugRenderer& Renderer)
+{
+	const FCameraDebugColors& Colors = FCameraDebugColors::Get();
+
+	TArrayView<FCameraDebugBlock*> ChildrenView(GetChildren());
+
+	const int32 MinNum = FMath::Min(ChildrenView.Num(), CameraDirectors.Num());
+
+	Renderer.SetTextColor(Colors.Hightlighted);
+	Renderer.AddText("Inactive Directors\n\n");
+	Renderer.SetTextColor(Colors.Default);
+	Renderer.AddIndent();
+
+	for (int32 Index = 0; Index < MinNum; ++Index)
+	{
+		if (Index == CameraDirectors.Num() - 1)
+		{
+			Renderer.RemoveIndent();
+			Renderer.NewLine();
+
+			Renderer.SetTextColor(Colors.Notice);
+			Renderer.AddText("Active Director\n\n");
+			Renderer.SetTextColor(Colors.Default);
+			Renderer.AddIndent();
+		}
+
+		const FDirectorDebugInfo& EntryDebugInfo(CameraDirectors[Index]);
+		Renderer.AddText(TEXT("Camera asset: {cam_notice}%s{cam_default}\n"), *EntryDebugInfo.CameraAssetName);
+
+		Renderer.AddIndent();
+		ChildrenView[Index]->DebugDraw(Params, Renderer);
+		Renderer.RemoveIndent();
+
+		Renderer.NewLine();
+	}
+
+	Renderer.RemoveIndent();
+	Renderer.SetTextColor(Colors.Default);
+
+	Renderer.SkipAllBlocks();
+}
+
+}  // namespace UE::Cameras
+
+#endif  // UE_GAMEPLAY_CAMERAS_DEBUG
+
diff --git a/Engine/Plugins/Cameras/GameplayCameras/Source/GameplayCameras/Private/Debug/CameraNodeEvaluationResultDebugBlock.cpp b/Engine/Plugins/Cameras/GameplayCameras/Source/GameplayCameras/Private/Debug/CameraNodeEvaluationResultDebugBlock.cpp
new file mode 100644
--- /dev/null
+++ b/Engine/Plugins/Cameras/GameplayCameras/Source/GameplayCameras/Private/Debug/CameraNodeEvaluationResultDebugBlock.cpp
@@ -0,0 +1,100 @@
+// Copyright Epic Games, Inc. All Rights Reserved.
+
+#include "Debug/CameraNodeEvaluationResultDebugBlock.h"
+
+#include "Core/CameraNodeEvaluator.h"
+#include "Debug/CameraDebugBlockBuilder.h"
+#include "Debug/CameraDebugColors.h"
+#include "Debug/CameraDebugRenderer.h"
+#include "Debug/CameraPoseDebugBlock.h"
+#include "Debug/DebugTextRenderer.h"
+#include "Debug/VariableTableDebugBlock.h"
+#include "Math/ColorList.h"
+
+#if UE_GAMEPLAY_CAMERAS_DEBUG
+
+namespace UE::Cameras
+{
+
+UE_DEFINE_CAMERA_DEBUG_BLOCK(FCameraNodeEvaluationResultDebugBlock)
+
+FCameraNodeEvaluationResultDebugBlock::FCameraNodeEvaluationResultDebugBlock()
+{
+}
+
+void FCameraNodeEvaluationResultDebugBlock::Initialize(const FCameraNodeEvaluationResult& InResult, FCameraDebugBlockBuilder& Builder)
+{
+	bIsCameraCut = InResult.bIsCameraCut;
+	bIsValid = InResult.bIsValid;
+
+	AddChild(&Builder.BuildDebugBlock<FCameraPoseDebugBlock>(InResult.CameraPose));
+	AddChild(&Builder.BuildDebugBlock<FVariableTableDebugBlock>(InResult.VariableTable));
+}
+
+FCameraPoseDebugBlock* FCameraNodeEvaluationResultDebugBlock::GetCameraPoseDebugBlock()
+{
+	TArrayView<FCameraDebugBlock*> ChildrenView(GetChildren());
+	if (ChildrenView.IsValidIndex(0))
+	{
+		return ChildrenView[0]->CastThis<FCameraPoseDebugBlock>();
+	}
+	return nullptr;
+}
+
+FVariableTableDebugBlock* FCameraNodeEvaluationResultDebugBlock::GetVariableTableDebugBlock()
+{
+	TArrayView<FCameraDebugBlock*> ChildrenView(GetChildren());
+	if (ChildrenView.IsValidIndex(1))
+	{
+		return ChildrenView[1]->CastThis<FVariableTableDebugBlock>();
+	}
+	return nullptr;
+}
+
+void FCameraNodeEvaluationResultDebugBlock::OnDebugDraw(const FCameraDebugBlockDrawParams& Params, FCameraDebugRenderer& Renderer)
+{
+	const FCameraDebugColors& Colors = FCameraDebugColors::Get();
+
+	if (bIsValid)
+	{
+		Renderer.AddText(TEXT("Valid: {cam_good}YES"));
+	}
+	else
+	{
+		Renderer.AddText(TEXT("Valid: {cam_error}NO"));
+	}
+
+	if (bIsCameraCut)
+	{
+		Renderer.AddText(TEXT("  {cam_warning}IsCameraCut"));
+	}
+
+	Renderer.NewLine();
+	Renderer.SetTextColor(Colors.Default);
+
+	TArrayView<FCameraDebugBlock*> ChildrenView(GetChildren());
+	if (ChildrenView.IsValidIndex(0))
+	{
+		Renderer.AddText(TEXT("{cam_title}Camera Pose:"));
+		Renderer.AddIndent();
+		Renderer.SetTextColor(Colors.Default);
+		ChildrenView[0]->DebugDraw(Params, Renderer);
+		Renderer.RemoveIndent();
+	}
+	if (ChildrenView.IsValidIndex(1))
+	{
+		Renderer.AddText(TEXT("{cam_title}Variable Table:"));
+		Renderer.AddIndent();
+		Renderer.SetTextColor(Colors.Default);
+		ChildrenView[1]->DebugDraw(Params, Renderer);
+		Renderer.RemoveIndent();
+	}
+
+	// We already manually rendered our children.
+	Renderer.SkipAllBlocks();
+}
+
+}  // namespace UE::Cameras
+
+#endif  // UE_GAMEPLAY_CAMERAS_DEBUG
+
diff --git a/Engine/Plugins/Cameras/GameplayCameras/Source/GameplayCameras/Private/Debug/CameraNodeEvaluatorDebugBlock.cpp b/Engine/Plugins/Cameras/GameplayCameras/Source/GameplayCameras/Private/Debug/CameraNodeEvaluatorDebugBlock.cpp
new file mode 100644
--- /dev/null
+++ b/Engine/Plugins/Cameras/GameplayCameras/Source/GameplayCameras/Private/Debug/CameraNodeEvaluatorDebugBlock.cpp
@@ -0,0 +1,56 @@
+// Copyright Epic Games, Inc. All Rights Reserved.
+
+#include "Debug/CameraNodeEvaluatorDebugBlock.h"
+
+#include "Core/CameraNode.h"
+#include "Debug/CameraDebugRenderer.h"
+#include "HAL/IConsoleManager.h"
+
+#if UE_GAMEPLAY_CAMERAS_DEBUG
+
+namespace UE::Cameras
+{
+
+FString GGameplayCamerasDebugNodeTreeFilter;
+static FAutoConsoleVariableRef CVarGameplayCamerasDebugNodeTreeFilter(
+	TEXT("GameplayCameras.Debug.NodeTree.Filter"),
+	GGameplayCamerasDebugNodeTreeFilter,
+	TEXT("(Default: "". Filters the debug camera node tree by node name/type."));
+
+UE_DEFINE_CAMERA_DEBUG_BLOCK(FCameraNodeEvaluatorDebugBlock)
+
+FCameraNodeEvaluatorDebugBlock::FCameraNodeEvaluatorDebugBlock()
+{
+}
+
+FCameraNodeEvaluatorDebugBlock::FCameraNodeEvaluatorDebugBlock(TObjectPtr<const UCameraNode> InCameraNode)
+{
+	NodeClassName = (InCameraNode ? InCameraNode->GetClass()->GetName() : TEXT("<no node>"));
+}
+
+void FCameraNodeEvaluatorDebugBlock::OnDebugDraw(const FCameraDebugBlockDrawParams& Params, FCameraDebugRenderer& Renderer)
+{
+	// Only show the name/type of this node, along with extra info (in the attached blocks) if we
+	// are in the "nodetree" debug mode, and we're not filtering this node out.
+	const bool bDoDebugDraw = (
+			GGameplayCamerasDebugNodeTreeFilter.IsEmpty() ||
+			NodeClassName.Contains(GGameplayCamerasDebugNodeTreeFilter));
+	if (bDoDebugDraw)
+	{
+		Renderer.AddText(TEXT("{cam_passive}[%s]{cam_default} "), *NodeClassName);
+	}
+	else
+	{
+		Renderer.SkipAttachedBlocks();
+	}
+}
+
+void FCameraNodeEvaluatorDebugBlock::OnSerialize(FArchive& Ar)
+{
+	Ar << NodeClassName;
+}
+
+}  // namespace UE::Cameras
+
+#endif  // UE_GAMEPLAY_CAMERAS_DEBUG
+
diff --git a/Engine/Plugins/Cameras/GameplayCameras/Source/GameplayCameras/Private/Debug/CameraPoseDebugBlock.cpp b/Engine/Plugins/Cameras/GameplayCameras/Source/GameplayCameras/Private/Debug/CameraPoseDebugBlock.cpp
new file mode 100644
--- /dev/null
+++ b/Engine/Plugins/Cameras/GameplayCameras/Source/GameplayCameras/Private/Debug/CameraPoseDebugBlock.cpp
@@ -0,0 +1,71 @@
+// Copyright Epic Games, Inc. All Rights Reserved.
+
+#include "Debug/CameraPoseDebugBlock.h"
+
+#include "Containers/UnrealString.h"
+#include "Debug/CameraDebugColors.h"
+#include "Debug/CameraDebugRenderer.h"
+#include "Debug/DebugTextRenderer.h"
+#include "HAL/IConsoleManager.h"
+#include "Math/ColorList.h"
+
+#if UE_GAMEPLAY_CAMERAS_DEBUG
+
+namespace UE::Cameras
+{
+
+namespace Private
+{
+
+template<typename FieldType>
+void DebugDrawCameraPoseField(FCameraDebugRenderer& Renderer, const TCHAR* FieldName, typename TCallTraits<FieldType>::ParamType FieldValue, const FColor& Color)
+{
+	Renderer.SetTextColor(Color);
+	Renderer.AddText(TEXT("%s  : %s\n"), FieldName, *ToDebugString(FieldValue));
+}
+
+}  // namespace Private
+
+UE_DEFINE_CAMERA_DEBUG_BLOCK(FCameraPoseDebugBlock)
+
+FCameraPoseDebugBlock::FCameraPoseDebugBlock()
+{
+}
+
+FCameraPoseDebugBlock::FCameraPoseDebugBlock(const FCameraPose& InCameraPose)
+	: CameraPose(InCameraPose)
+{
+}
+
+void FCameraPoseDebugBlock::OnDebugDraw(const FCameraDebugBlockDrawParams& Params, FCameraDebugRenderer& Renderer)
+{
+	bool bShowUnchanged = false;
+	if (!ShowUnchangedCVarName.IsEmpty())
+	{
+		IConsoleVariable* ShowUnchangedCVar = IConsoleManager::Get().FindConsoleVariable(*ShowUnchangedCVarName, false);
+		if (ensureMsgf(ShowUnchangedCVar, TEXT("No such console variable: %s"), *ShowUnchangedCVarName))
+		{
+			bShowUnchanged = ShowUnchangedCVar->GetBool();
+		}
+	}
+
+	const FCameraDebugColors& Colors = FCameraDebugColors::Get();
+	const FColor ChangedColor = Colors.Default;
+	const FColor UnchangedColor = Colors.Passive;
+
+	const FCameraPoseFlags& ChangedFlags = CameraPose.GetChangedFlags();
+
+#define UE_CAMERA_POSE_FOR_PROPERTY(PropType, PropName)\
+	if (bShowUnchanged || ChangedFlags.PropName)\
+	{\
+		const FColor& PropColor = ChangedFlags.PropName ? ChangedColor : UnchangedColor;\
+		Private::DebugDrawCameraPoseField<PropType>(Renderer, TEXT(#PropName), CameraPose.Get##PropName(), PropColor);\
+	}
+	UE_CAMERA_POSE_FOR_ALL_PROPERTIES()
+#undef UE_CAMERA_POSE_FOR_PROPERTY
+}
+
+}  // namespace UE::Cameras
+
+#endif  // UE_GAMEPLAY_CAMERAS_DEBUG
+
diff --git a/Engine/Plugins/Cameras/GameplayCameras/Source/GameplayCameras/Private/Debug/CategoryTitleDebugBlock.cpp b/Engine/Plugins/Cameras/GameplayCameras/Source/GameplayCameras/Private/Debug/CategoryTitleDebugBlock.cpp
new file mode 100644
--- /dev/null
+++ b/Engine/Plugins/Cameras/GameplayCameras/Source/GameplayCameras/Private/Debug/CategoryTitleDebugBlock.cpp
@@ -0,0 +1,51 @@
+// Copyright Epic Games, Inc. All Rights Reserved.
+
+#include "Debug/CategoryTitleDebugBlock.h"
+
+#include "Debug/CameraDebugRenderer.h"
+
+#if UE_GAMEPLAY_CAMERAS_DEBUG
+
+namespace UE::Cameras
+{
+
+UE_DEFINE_CAMERA_DEBUG_BLOCK(FCategoryTitleDebugBlock)
+
+FCategoryTitleDebugBlock::FCategoryTitleDebugBlock()
+	: TitleColor(FColor::Green)
+{
+}
+
+FCategoryTitleDebugBlock::FCategoryTitleDebugBlock(const FString& InCategory, const FString& InTitle)
+	: Category(InCategory)
+	, Title(InTitle)
+	, TitleColor(FColor::Green)
+{
+}
+
+void FCategoryTitleDebugBlock::OnDebugDraw(const FCameraDebugBlockDrawParams& Params, FCameraDebugRenderer& Renderer)
+{
+	if (Params.IsCategoryActive(Category))
+	{
+		if (!Title.IsEmpty())
+		{
+			Renderer.AddText(TEXT("{cam_title}%s{cam_default}\n"), *Title);
+		}
+	}
+	else
+	{
+		if (bSkipAttachedBlocksIfInactive)
+		{
+			Renderer.SkipAttachedBlocks();
+		}
+		if (bSkipChildrenBlocksIfInactive)
+		{
+			Renderer.SkipChildrenBlocks();
+		}
+	}
+}
+
+}  // namespace UE::Cameras
+
+#endif  // UE_GAMEPLAY_CAMERAS_DEBUG
+
diff --git a/Engine/Plugins/Cameras/GameplayCameras/Source/GameplayCameras/Private/Debug/DebugTextRenderer.cpp b/Engine/Plugins/Cameras/GameplayCameras/Source/GameplayCameras/Private/Debug/DebugTextRenderer.cpp
new file mode 100644
--- /dev/null
+++ b/Engine/Plugins/Cameras/GameplayCameras/Source/GameplayCameras/Private/Debug/DebugTextRenderer.cpp
@@ -0,0 +1,245 @@
+// Copyright Epic Games, Inc. All Rights Reserved.
+
+#include "Debug/DebugTextRenderer.h"
+
+#include "CanvasItem.h"
+#include "Debug/CameraDebugColors.h"
+#include "Engine/Canvas.h"
+#include "Engine/Font.h"
+
+#if UE_GAMEPLAY_CAMERAS_DEBUG
+
+namespace UE::Cameras
+{
+
+void FDebugTextDrawCommand::Execute(UCanvas* Canvas, const UFont* Font, FVector2f& InOutDrawPosition) const
+{
+	// Sadly we need to allocate a string here...
+	const FText Text(FText::FromStringView(TextView));
+
+	FCanvasTextItem TextItem(FVector2D(InOutDrawPosition), Text, Font, Canvas->DrawColor);
+	TextItem.BlendMode = SE_BLEND_Translucent;
+	Canvas->DrawItem(TextItem);
+	
+	const float TextWidth = FDebugTextRenderer::GetStringViewSize(Font, TextView);
+	InOutDrawPosition.X += TextWidth;
+}
+
+void FDebugTextNewLineCommand::Execute(FVector2f& InOutDrawPosition) const
+{
+	InOutDrawPosition.X = LeftMargin;
+	InOutDrawPosition.Y += LineSpacing;
+}
+
+void FDebugTextSetColorCommand::Execute(UCanvas* Canvas) const
+{
+	Canvas->SetDrawColor(DrawColor);
+}
+
+FDebugTextRenderer::FDebugTextRenderer(UCanvas* InCanvas, const UFont* InFont)
+	: Canvas(InCanvas)
+	, Font(InFont)
+{
+	LineSpacing = InFont->GetMaxCharHeight();
+	LeftMargin = 0;
+}
+
+void FDebugTextRenderer::RenderText(float StartingDrawY, const TStringView<TCHAR> TextView)
+{
+	RenderText(FVector2f(LeftMargin, StartingDrawY), TextView);
+}
+
+void FDebugTextRenderer::RenderText(FVector2f StartingDrawPosition, const TStringView<TCHAR> TextView)
+{
+	FDebugTextCommandArray Commands;
+	ParseText(TextView, Commands);
+
+	NextDrawPosition = StartingDrawPosition;
+	ExecuteCommands(Commands);
+}
+
+void FDebugTextRenderer::ParseText(const TStringView<TCHAR> TextView, FDebugTextCommandArray& OutCommands)
+{
+	const TCHAR* const Begin = TextView.GetData();
+	const TCHAR* const End = Begin + TextView.Len();
+
+	const TCHAR* RangeStart = TextView.GetData();
+	const TCHAR* RangeCur = RangeStart;
+
+	// Let's start parsing the given text and look for tokens of the form `{token}`.
+	bool bIsInToken = false;
+	while (true)
+	{
+		const bool bIsEOF = (RangeCur == End || *RangeCur == '\0');
+		const bool bIsNL = !bIsEOF && (*RangeCur == '\n');
+		if (!bIsInToken)
+		{
+			// Not in a token... let's see if we find an EOF or NL, in which case we need to
+			// draw the text and either bail out or move to a new line. Otherwise, look for
+			// the start of a token with `{`.
+			if (bIsEOF || bIsNL)
+			{
+				// EOF or NL found... render the text so far and move to a new line.
+				const bool bAddNewLine = bIsNL || (bIsEOF && bEndWithNewLine);
+				AddDrawCommand(RangeStart, RangeCur, bAddNewLine, OutCommands);
+				if (bIsEOF)
+				{
+					break;
+				}
+				else
+				{
+					++RangeCur;
+					RangeStart = RangeCur;
+				}
+			}
+			else if (*RangeCur == '{')
+			{
+				// We have a string to render up until the start of the new token.
+				AddDrawCommand(RangeStart, RangeCur, false, OutCommands);
+				++RangeCur;
+				RangeStart = RangeCur;
+				bIsInToken = true;
+			}
+			else
+			{
+				++RangeCur;
+			}
+		}
+		else
+		{
+			if (*RangeCur == '}')
+			{
+				// We have a token!
+				AddTokenCommand(RangeStart, RangeCur, OutCommands);
+				++RangeCur;
+				RangeStart = RangeCur;
+				bIsInToken = false;
+			}
+			else if (bIsEOF || bIsNL)
+			{
+				// Unclosed token... just treat the whole thing as a string.
+				const bool bAddNewLine = bIsNL || (bIsEOF && bEndWithNewLine);
+				AddDrawCommand(RangeStart, RangeCur, bAddNewLine, OutCommands);
+				break;
+			}
+			else
+			{
+				++RangeCur;
+			}
+		}
+	}
+}
+
+void FDebugTextRenderer::ExecuteCommands(FDebugTextCommandArray& Commands)
+{
+	FColor OriginalDrawColor = Canvas->DrawColor;
+
+	for (FDebugTextCommand& Command : Commands)
+	{
+		switch (Command.GetIndex())
+		{
+			case FDebugTextCommand::IndexOfType<FDebugTextDrawCommand>():
+				Command.Get<FDebugTextDrawCommand>().Execute(Canvas, Font, NextDrawPosition);
+				UpdateRightMargin();
+				break;
+			case FDebugTextCommand::IndexOfType<FDebugTextNewLineCommand>():
+				Command.Get<FDebugTextNewLineCommand>().Execute(NextDrawPosition);
+				break;
+			case FDebugTextCommand::IndexOfType<FDebugTextSetColorCommand>():
+				Command.Get<FDebugTextSetColorCommand>().Execute(Canvas);
+				break;
+			default:
+				ensureMsgf(false, TEXT("Unsupported command type!"));
+				break;
+		}
+	}
+
+	Canvas->SetDrawColor(OriginalDrawColor);
+}
+
+void FDebugTextRenderer::UpdateRightMargin()
+{
+	RightMargin = FMath::Max(RightMargin, NextDrawPosition.X);
+}
+
+void FDebugTextRenderer::ExecuteCommands(float StartingDrawY, FDebugTextCommandArray& Commands)
+{
+	ExecuteCommands(FVector2f(LeftMargin, StartingDrawY), Commands);
+}
+
+void FDebugTextRenderer::ExecuteCommands(FVector2f StartingDrawPosition, FDebugTextCommandArray& Commands)
+{
+	NextDrawPosition = StartingDrawPosition;
+	ExecuteCommands(Commands);
+}
+
+void FDebugTextRenderer::AddDrawCommand(const TCHAR* RangeStart, const TCHAR* RangeEnd, bool bNewLine, FDebugTextCommandArray& OutCommands)
+{
+	TStringView<TCHAR> CmdTextView(RangeStart, RangeEnd - RangeStart);
+	if (!CmdTextView.IsEmpty())
+	{
+		OutCommands.Emplace(TInPlaceType<FDebugTextDrawCommand>(), FDebugTextDrawCommand{ CmdTextView });
+	}
+	// else: empty text, skip command entirely.
+
+	if (bNewLine)
+	{
+		OutCommands.Emplace(TInPlaceType<FDebugTextNewLineCommand>(), FDebugTextNewLineCommand{ LineSpacing, LeftMargin });
+	}
+}
+
+void FDebugTextRenderer::AddTokenCommand(const TCHAR* RangeStart, const TCHAR* RangeEnd, FDebugTextCommandArray& OutCommands)
+{
+	TStringView<TCHAR> TokenView(RangeStart, RangeEnd - RangeStart);
+	// Sadly we need to allocate a string here too...
+	FColor DrawColor = InterpretColor(FString(TokenView));
+	OutCommands.Emplace(TInPlaceType<FDebugTextSetColorCommand>(), FDebugTextSetColorCommand{ DrawColor });
+}
+
+FColor FDebugTextRenderer::InterpretColor(const FString& ColorName)
+{
+	TOptional<FColor> CameraDebugColor = FCameraDebugColors::GetFColorByName(ColorName);
+	if (CameraDebugColor.IsSet())
+	{
+		return CameraDebugColor.GetValue();
+	}
+
+	const bool bIsColorName = GColorList.IsValidColorName(*ColorName);
+	if (bIsColorName)
+	{
+		return GColorList.GetFColorByName(*ColorName);
+	}
+	else
+	{
+		FColor OutColor;
+		OutColor.InitFromString(ColorName);
+		return OutColor;
+	}
+}
+
+float FDebugTextRenderer::GetStringViewSize(const UFont* Font, TStringView<TCHAR> TextView)
+{
+	float TotalWidth = 0.0f;
+	TCHAR PrevChar = '\0';
+	for (TCHAR Char : TextView)
+	{
+		float TmpWidth, TmpHeight;
+		Font->GetCharSize(Char, TmpWidth, TmpHeight);
+
+		int8 CharKerning = 0;
+		if (PrevChar != '\0')
+		{
+			CharKerning = Font->GetCharKerning(PrevChar, Char);
+		}
+
+		TotalWidth += TmpWidth + CharKerning;
+		PrevChar = Char;
+	}
+
+	return FMath::CeilToInt(TotalWidth);
+}
+
+}  // namespace UE::Cameras
+
+#endif  // UE_GAMEPLAY_CAMERAS_DEBUG
+
diff --git a/Engine/Plugins/Cameras/GameplayCameras/Source/GameplayCameras/Private/Debug/RootCameraDebugBlock.cpp b/Engine/Plugins/Cameras/GameplayCameras/Source/GameplayCameras/Private/Debug/RootCameraDebugBlock.cpp
new file mode 100644
--- /dev/null
+++ b/Engine/Plugins/Cameras/GameplayCameras/Source/GameplayCameras/Private/Debug/RootCameraDebugBlock.cpp
@@ -0,0 +1,138 @@
+// Copyright Epic Games, Inc. All Rights Reserved.
+
+#include "Debug/RootCameraDebugBlock.h"
+
+#include "Core/CameraEvaluationContextStack.h"
+#include "Core/CameraSystemEvaluator.h"
+#include "Core/RootCameraNode.h"
+#include "Debug/CameraDebugBlockBuilder.h"
+#include "Debug/CameraDebugCategories.h"
+#include "Debug/CameraDebugRenderer.h"
+#include "Debug/CameraDirectorTreeDebugBlock.h"
+#include "Debug/CameraNodeEvaluatorDebugBlock.h"
+#include "Debug/CameraPoseDebugBlock.h"
+#include "Debug/CategoryTitleDebugBlock.h"
+#include "Debug/ViewfinderDebugBlock.h"
+#include "HAL/IConsoleManager.h"
+#include "String/ParseTokens.h"
+
+#if UE_GAMEPLAY_CAMERAS_DEBUG
+
+namespace UE::Cameras
+{
+
+bool GGameplayCamerasDebugEnable = false;
+static FAutoConsoleVariableRef CVarGameplayCamerasDebugEnable(
+	TEXT("GameplayCameras.Debug.Enable"),
+	GGameplayCamerasDebugEnable,
+	TEXT("(Default: false. Enables debug drawing for the GamplayCameras system."));
+
+bool GGameplayCamerasDebugTrace = true;
+static FAutoConsoleVariableRef CVarGameplayCamerasDebugTrace(
+	TEXT("GameplayCameras.Debug.Trace"),
+	GGameplayCamerasDebugTrace,
+	TEXT("(Default: false. Enables background tracing of GamplayCameras system debug info."));
+
+FString GGameplayCamerasDebugCategories = "nodetree";
+static FAutoConsoleVariableRef CVarGameplayCamerasDebugCategories(
+	TEXT("GameplayCameras.Debug.Categories"),
+	GGameplayCamerasDebugCategories,
+	TEXT("(Default: nodes. Specifies which debug categories to display the GamplayCameras system."));
+
+bool GGameplayCamerasDebugPoseStatsShowUnchanged = false;
+static FAutoConsoleVariableRef CVarGameplayCamerasDebugPoseStatsShowUnchanged(
+	TEXT("GameplayCameras.Debug.PoseStats.ShowUnchanged"),
+	GGameplayCamerasDebugPoseStatsShowUnchanged,
+	TEXT(""));
+
+bool GGameplayCamerasDebugDrawBackground = true;
+static FAutoConsoleVariableRef CVarGameplayCamerasDebugDrawBackground(
+	TEXT("GameplayCameras.Debug.DrawBackground"),
+	GGameplayCamerasDebugDrawBackground,
+	TEXT(""));
+
+float GGameplayCamerasDebugBackgroundOpacity = 0.6f;
+static FAutoConsoleVariableRef CVarGameplayCamerasDebugBackgroundOpacity(
+	TEXT("GameplayCameras.Debug.BackgroundOpacity"),
+	GGameplayCamerasDebugBackgroundOpacity,
+	TEXT(""));
+
+UE_DEFINE_CAMERA_DEBUG_BLOCK(FRootCameraDebugBlock)
+
+void FRootCameraDebugBlock::BuildDebugBlocks(const FCameraSystemEvaluator& CameraSystem, const FCameraDebugBlockBuildParams& Params, FCameraDebugBlockBuilder& Builder)
+{
+	ensureMsgf(GetChildren().IsEmpty() && GetAttachments().IsEmpty(), TEXT("This root debug block has already been initialized!"));
+
+	// Debug block for showing the directors and context stack.
+	FCategoryTitleDebugBlock& DirectorTreeCategory = Builder.StartChildDebugBlock<FCategoryTitleDebugBlock>();
+	{
+		DirectorTreeCategory.Title = TEXT("Camera Directors");
+		DirectorTreeCategory.Category = FCameraDebugCategories::DirectorTree;
+
+		const FCameraEvaluationContextStack& ContextStack = CameraSystem.GetEvaluationContextStack();
+		FCameraDirectorTreeDebugBlock& DirectorTreeDebugBlock = Builder.StartChildDebugBlock<FCameraDirectorTreeDebugBlock>();
+		DirectorTreeDebugBlock.Initialize(ContextStack, Builder);
+		Builder.EndChildDebugBlock();
+	}
+	Builder.EndChildDebugBlock();
+
+	// Debug block for showing the tree of camera nodes.
+	FCategoryTitleDebugBlock& NodeTreeCategory = Builder.StartChildDebugBlock<FCategoryTitleDebugBlock>();
+	{
+		NodeTreeCategory.Title = TEXT("Camera Nodes");
+		NodeTreeCategory.Category = FCameraDebugCategories::NodeTree;
+
+		if (FRootCameraNodeEvaluator* RootNodeEvaluator = CameraSystem.GetRootNodeEvaluator())
+		{
+			RootNodeEvaluator->BuildDebugBlocks(Params, Builder);
+		}
+	}
+	Builder.EndChildDebugBlock();
+
+	// Debug block for showing the final evaluated camera.
+	FCategoryTitleDebugBlock& PoseStatsCategory = Builder.StartChildDebugBlock<FCategoryTitleDebugBlock>();
+	{
+		PoseStatsCategory.Title = TEXT("Evaluated Camera");
+		PoseStatsCategory.Category = FCameraDebugCategories::PoseStats;
+
+		const FCameraSystemEvaluationUpdateResult& Result = CameraSystem.GetEvaluatedResult();
+		Builder.AttachDebugBlock<FCameraPoseDebugBlock>(Result.CameraPose)
+			.WithShowUnchangedCVar(TEXT("GameplayCameras.Debug.PoseStats.ShowUnchanged"));
+	}
+	Builder.EndChildDebugBlock();
+	
+	// Debug block for rendering a viewfinder.
+	AddChild(&Builder.BuildDebugBlock<FViewfinderDebugBlock>());
+}
+
+void FRootCameraDebugBlock::RootDebugDraw(FCameraDebugRenderer& Renderer)
+{
+	if (!GGameplayCamerasDebugEnable)
+	{
+		return;
+	}
+
+	// Figure out what debug categories are active.
+	FCameraDebugBlockDrawParams Params;
+
+	TArray<FStringView, TInlineAllocator<4>> ActiveCategories;
+	UE::String::ParseTokens(GGameplayCamerasDebugCategories, ',', ActiveCategories);
+	for (FStringView CategoryView : ActiveCategories)
+	{
+		Params.ActiveCategories.Add(FString(CategoryView));
+	}
+
+	// Do the drawing!
+	FCameraDebugBlock::DebugDraw(Params, Renderer);
+
+	// Render a translucent background to help readability.
+	if (GGameplayCamerasDebugDrawBackground)
+	{
+		Renderer.DrawTextBackgroundTile(GGameplayCamerasDebugBackgroundOpacity);
+	}
+}
+
+}  // namespace UE::Cameras
+
+#endif  // UE_GAMEPLAY_CAMERAS_DEBUG
+
diff --git a/Engine/Plugins/Cameras/GameplayCameras/Source/GameplayCameras/Private/Debug/VariableTableDebugBlock.cpp b/Engine/Plugins/Cameras/GameplayCameras/Source/GameplayCameras/Private/Debug/VariableTableDebugBlock.cpp
new file mode 100644
--- /dev/null
+++ b/Engine/Plugins/Cameras/GameplayCameras/Source/GameplayCameras/Private/Debug/VariableTableDebugBlock.cpp
@@ -0,0 +1,113 @@
+// Copyright Epic Games, Inc. All Rights Reserved.
+
+#include "Debug/VariableTableDebugBlock.h"
+
+#include "Core/CameraVariableTable.h"
+#include "Debug/CameraDebugColors.h"
+#include "Debug/CameraDebugRenderer.h"
+#include "Debug/DebugTextRenderer.h"
+#include "HAL/IConsoleManager.h"
+
+#if UE_GAMEPLAY_CAMERAS_DEBUG
+
+namespace UE::Cameras
+{
+
+UE_DEFINE_CAMERA_DEBUG_BLOCK(FVariableTableDebugBlock)
+
+FVariableTableDebugBlock::FVariableTableDebugBlock()
+{
+}
+
+FVariableTableDebugBlock::FVariableTableDebugBlock(const FCameraVariableTable& InVariableTable)
+{
+	Initialize(InVariableTable);
+}
+
+void FVariableTableDebugBlock::Initialize(const FCameraVariableTable& InVariableTable)
+{
+	for (const TPair<uint32, FCameraVariableTable::FEntry>& Pair : InVariableTable.Entries)
+	{
+		const int32 EntryId = Pair.Key;
+		const FCameraVariableTable::FEntry& Entry = Pair.Value;
+
+		FString EntryName;
+#if WITH_EDITORONLY_DATA
+		EntryName = Entry.DebugName;
+#endif
+
+		FString EntryValueStr;
+#define UE_CAMERA_VARIABLE_FOR_TYPE(ValueType, ValueName)\
+			case ECameraVariableType::ValueName:\
+				if (EnumHasAnyFlags(Entry.Flags, FCameraVariableTable::EEntryFlags::Written))\
+				{\
+					const ValueType EntryValue = InVariableTable.GetValue<ValueType>(EntryId);\
+					EntryValueStr = ToDebugString(EntryValue);\
+				}\
+				break;
+		switch (Entry.Type)
+		{
+			UE_CAMERA_VARIABLE_FOR_ALL_TYPES()
+		}
+#undef UE_CAMERA_VARIABLE_FOR_TYPE
+
+		FEntryDebugInfo EntryDebugInfo{ EntryId, EntryName, EntryValueStr };
+		EntryDebugInfo.bWritten = EnumHasAnyFlags(Entry.Flags, FCameraVariableTable::EEntryFlags::Written);
+		EntryDebugInfo.bWrittenThisFrame = EnumHasAnyFlags(Entry.Flags, FCameraVariableTable::EEntryFlags::WrittenThisFrame);
+		Entries.Add(EntryDebugInfo);
+	}
+}
+
+void FVariableTableDebugBlock::OnDebugDraw(const FCameraDebugBlockDrawParams& Params, FCameraDebugRenderer& Renderer)
+{
+#if WITH_EDITORONLY_DATA
+	bool bShowVariableIds = false;
+	if (!ShowVariableIdsCVarName.IsEmpty())
+	{
+		IConsoleVariable* ShowVariableIdsCVar = IConsoleManager::Get().FindConsoleVariable(*ShowVariableIdsCVarName, false);
+		if (ensureMsgf(ShowVariableIdsCVar, TEXT("No such console variable: %s"), *ShowVariableIdsCVarName))
+		{
+			bShowVariableIds = ShowVariableIdsCVar->GetBool();
+		}
+	}
+#endif
+
+	const FCameraDebugColors& Colors = FCameraDebugColors::Get();
+
+	Renderer.AddText(TEXT("{cam_title}Variable Table{cam_default}\n"));
+	Renderer.AddIndent();
+	for (const FEntryDebugInfo& Entry : Entries)
+	{
+#if WITH_EDITORONLY_DATA
+		if (bShowVariableIds)
+		{
+			Renderer.AddText(TEXT("{cam_passive}[%d]{cam_default} "), Entry.Id);
+		}
+		Renderer.AddText(TEXT("%s : "), *Entry.Name);
+#else
+		Renderer.AddText(TEXT("[%d] <no name data> : "), Entry.Id);
+#endif
+
+		if (Entry.bWritten)
+		{
+			Renderer.AddText(Entry.Value);
+			if (Entry.bWrittenThisFrame)
+			{
+				Renderer.AddText(TEXT(" {cam_passive}[WrittenThisFrame]"));
+			}
+		}
+		else
+		{
+			Renderer.AddText("{cam_warning}[Uninitialized]");
+		}
+
+		Renderer.NewLine();
+		Renderer.SetTextColor(Colors.Default);
+	}
+	Renderer.RemoveIndent();
+}
+
+}  // namespace UE::Cameras
+
+#endif  // UE_GAMEPLAY_CAMERAS_DEBUG
+
diff --git a/Engine/Plugins/Cameras/GameplayCameras/Source/GameplayCameras/Private/Debug/ViewfinderDebugBlock.cpp b/Engine/Plugins/Cameras/GameplayCameras/Source/GameplayCameras/Private/Debug/ViewfinderDebugBlock.cpp
new file mode 100644
--- /dev/null
+++ b/Engine/Plugins/Cameras/GameplayCameras/Source/GameplayCameras/Private/Debug/ViewfinderDebugBlock.cpp
@@ -0,0 +1,122 @@
+// Copyright Epic Games, Inc. All Rights Reserved.
+
+#include "Debug/ViewfinderDebugBlock.h"
+
+#include "Debug/CameraDebugCategories.h"
+#include "Debug/CameraDebugColors.h"
+#include "Debug/CameraDebugRenderer.h"
+#include "DrawDebugHelpers.h"
+#include "Engine/Canvas.h"
+#include "HAL/IConsoleManager.h"
+
+#if UE_GAMEPLAY_CAMERAS_DEBUG
+
+namespace UE::Cameras
+{
+
+float GGameplayCamerasDebugViewfinderReticleSizeFactor = 0.27f;
+static FAutoConsoleVariableRef CVarGameplayCamerasDebugViewfinderReticleSizeFactor(
+	TEXT("GameplayCameras.Debug.Viewfinder.ReticleSizeFactor"),
+	GGameplayCamerasDebugViewfinderReticleSizeFactor,
+	TEXT("Default: 0.1. The size of the viewfinder reticle, as a factor of the screen's vertical size."));
+
+float GGameplayCamerasDebugViewfinderReticleInnerSizeFactor = 0.7f;
+static FAutoConsoleVariableRef CVarGameplayCamerasDebugViewfinderReticleInnerSizeFactor(
+	TEXT("GameplayCameras.Debug.Viewfinder.ReticleInnerSizeFactor"),
+	GGameplayCamerasDebugViewfinderReticleInnerSizeFactor,
+	TEXT(""));
+
+int32 GGameplayCamerasDebugViewfinderReticleNumSides = 60;
+static FAutoConsoleVariableRef CVarGameplayCamerasDebugViewfinderReticleNumSides(
+	TEXT("GameplayCameras.Debug.Viewfinder.ReticleNumSides"),
+	GGameplayCamerasDebugViewfinderReticleNumSides,
+	TEXT(""));
+
+float GGameplayCamerasDebugViewfinderGuidesGapFactor = 0.02f;
+static FAutoConsoleVariableRef CVarGameplayCamerasDebugViewfinderGuidesGapFactor(
+	TEXT("GameplayCameras.Debug.Viewfinder.GuidesGapFactor"),
+	GGameplayCamerasDebugViewfinderGuidesGapFactor,
+	TEXT(""));
+
+UE_DEFINE_CAMERA_DEBUG_BLOCK(FViewfinderDebugBlock)
+
+FViewfinderDebugBlock::FViewfinderDebugBlock()
+{
+}
+
+void FViewfinderDebugBlock::OnDebugDraw(const FCameraDebugBlockDrawParams& Params, FCameraDebugRenderer& Renderer)
+{
+	if (!Params.IsCategoryActive(FCameraDebugCategories::Viewfinder))
+	{
+		return;
+	}
+
+	UCanvas* Canvas = Renderer.GetCanvas();
+
+	const FVector2d CanvasCenter(Canvas->SizeX / 2.f, Canvas->SizeY / 2.f);
+
+	// Draw the reticle.
+	const float ReticleRadius = Canvas->SizeY * GGameplayCamerasDebugViewfinderReticleSizeFactor / 2.f;
+	const float ReticleInnerRadiusFactor = GGameplayCamerasDebugViewfinderReticleInnerSizeFactor;
+	const int32 ReticleNumSides = GGameplayCamerasDebugViewfinderReticleNumSides;
+	const FColor ReticleColor = FCameraDebugColors::Get().Passive;
+
+	// ...outer reticle circle.
+	DrawDebugCanvas2DCircle(Canvas, CanvasCenter, ReticleRadius, ReticleNumSides, ReticleColor);
+	// ...inner reticle circle.
+	const float ReticleInnerRadius = ReticleRadius * ReticleInnerRadiusFactor;
+	const int32 ReticleInnerNumSides = (int32)(ReticleNumSides * ReticleInnerRadiusFactor);
+	DrawDebugCanvas2DCircle(Canvas, CanvasCenter, ReticleInnerRadius, ReticleInnerNumSides, ReticleColor);
+	// ...horizontal line inside reticle.
+	DrawDebugCanvas2DLine(
+			Canvas, 
+			CanvasCenter - FVector2D(ReticleInnerRadius, 0), CanvasCenter + FVector2D(ReticleInnerRadius, 0), 
+			ReticleColor);
+
+	// Draw the rule-of-thirds guides.
+	const FColor GuideColor = FCameraDebugColors::Get().VeryPassive;
+	const float GuidesGap = Canvas->SizeY * GGameplayCamerasDebugViewfinderGuidesGapFactor;
+	const FVector2D OneThird(Canvas->SizeX / 3.f, Canvas->SizeY / 3.f);
+	const FVector2D TwoThirds(Canvas->SizeX / 1.5f, Canvas->SizeY / 1.5f);
+	// ...top vertical guides
+	DrawDebugCanvas2DLine(
+			Canvas,
+			FVector2D(OneThird.X, 0), FVector2D(OneThird.X, OneThird.Y - GuidesGap),
+			ReticleColor, 2.f);
+	DrawDebugCanvas2DLine(
+			Canvas,
+			FVector2D(TwoThirds.X, 0), FVector2D(TwoThirds.X, OneThird.Y - GuidesGap),
+			ReticleColor, 2.f);
+	// ...bottom vertical guides
+	DrawDebugCanvas2DLine(
+			Canvas,
+			FVector2D(OneThird.X, TwoThirds.Y + GuidesGap), FVector2D(OneThird.X, Canvas->SizeY),
+			ReticleColor, 2.f);
+	DrawDebugCanvas2DLine(
+			Canvas,
+			FVector2D(TwoThirds.X, TwoThirds.Y + GuidesGap), FVector2D(TwoThirds.X, Canvas->SizeY),
+			ReticleColor, 2.f);
+	// ...left horizontal guides
+	DrawDebugCanvas2DLine(
+			Canvas,
+			FVector2D(0, OneThird.Y), FVector2D(OneThird.X - GuidesGap, OneThird.Y),
+			ReticleColor, 2.f);
+	DrawDebugCanvas2DLine(
+			Canvas,
+			FVector2D(0, TwoThirds.Y), FVector2D(OneThird.X - GuidesGap, TwoThirds.Y),
+			ReticleColor, 2.f);
+	// ...right horizontal guides
+	DrawDebugCanvas2DLine(
+			Canvas,
+			FVector2D(TwoThirds.X + GuidesGap, OneThird.Y), FVector2D(Canvas->SizeX, OneThird.Y),
+			ReticleColor, 2.f);
+	DrawDebugCanvas2DLine(
+			Canvas,
+			FVector2D(TwoThirds.X + GuidesGap, TwoThirds.Y), FVector2D(Canvas->SizeX, TwoThirds.Y),
+			ReticleColor, 2.f);
+}
+
+}  // namespace UE::Cameras
+
+#endif  // UE_GAMEPLAY_CAMERAS_DEBUG
+
diff --git a/Engine/Plugins/Cameras/GameplayCameras/Source/GameplayCameras/Private/GameFramework/GameplayCameraSystemComponent.cpp b/Engine/Plugins/Cameras/GameplayCameras/Source/GameplayCameras/Private/GameFramework/GameplayCameraSystemComponent.cpp
--- a/Engine/Plugins/Cameras/GameplayCameras/Source/GameplayCameras/Private/GameFramework/GameplayCameraSystemComponent.cpp
+++ b/Engine/Plugins/Cameras/GameplayCameras/Source/GameplayCameras/Private/GameFramework/GameplayCameraSystemComponent.cpp
@@ -152,9 +152,9 @@ void UGameplayCameraSystemComponent::DebugDraw(UCanvas* Canvas, APlayerControlle
 
 	if (Evaluator.IsValid())
 	{
-		FCameraSystemDebugDrawParams DebugDrawParams;
-		DebugDrawParams.Canvas = Canvas;
-		Evaluator->DebugDraw(DebugDrawParams);
+		FCameraSystemDebugUpdateParams DebugUpdateParams;
+		DebugUpdateParams.Canvas = Canvas;
+		Evaluator->DebugUpdate(DebugUpdateParams);
 	}
 }
 #endif  // UE_GAMEPLAY_CAMERAS_DEBUG
diff --git a/Engine/Plugins/Cameras/GameplayCameras/Source/GameplayCameras/Private/Nodes/Blends/SimpleBlendCameraNode.cpp b/Engine/Plugins/Cameras/GameplayCameras/Source/GameplayCameras/Private/Nodes/Blends/SimpleBlendCameraNode.cpp
--- a/Engine/Plugins/Cameras/GameplayCameras/Source/GameplayCameras/Private/Nodes/Blends/SimpleBlendCameraNode.cpp
+++ b/Engine/Plugins/Cameras/GameplayCameras/Source/GameplayCameras/Private/Nodes/Blends/SimpleBlendCameraNode.cpp
@@ -2,6 +2,10 @@
 
 #include "Nodes/Blends/SimpleBlendCameraNode.h"
 
+#include "Debug/CameraDebugBlock.h"
+#include "Debug/CameraDebugBlockBuilder.h"
+#include "Debug/CameraDebugRenderer.h"
+
 #include UE_INLINE_GENERATED_CPP_BY_NAME(SimpleBlendCameraNode)
 
 namespace UE::Cameras
@@ -9,6 +13,12 @@ namespace UE::Cameras
 
 UE_DEFINE_BLEND_CAMERA_NODE_EVALUATOR(FSimpleBlendCameraNodeEvaluator)
 
+UE_DECLARE_CAMERA_DEBUG_BLOCK_START(FSimpleBlendCameraDebugBlock)
+	UE_DECLARE_CAMERA_DEBUG_BLOCK_FIELD(float, BlendFactor)
+UE_DECLARE_CAMERA_DEBUG_BLOCK_END()
+
+UE_DEFINE_CAMERA_DEBUG_BLOCK(FSimpleBlendCameraDebugBlock)
+
 void FSimpleBlendCameraNodeEvaluator::OnRun(const FCameraNodeEvaluationParams& Params, FCameraNodeEvaluationResult& OutResult)
 {
 	FSimpleBlendCameraNodeEvaluationResult FactorResult;
@@ -41,6 +51,20 @@ void FSimpleBlendCameraNodeEvaluator::OnBlendResults(const FCameraNodeBlendParam
 	OutResult.bIsBlendFinished = bIsBlendFinished;
 }
 
+#if UE_GAMEPLAY_CAMERAS_DEBUG
+
+void FSimpleBlendCameraNodeEvaluator::OnBuildDebugBlocks(const FCameraDebugBlockBuildParams& Params, FCameraDebugBlockBuilder& Builder)
+{
+	FSimpleBlendCameraDebugBlock& DebugBlock = Builder.AttachDebugBlock<FSimpleBlendCameraDebugBlock>();
+	DebugBlock.BlendFactor = BlendFactor;
+}
+
+void FSimpleBlendCameraDebugBlock::OnDebugDraw(const FCameraDebugBlockDrawParams& Params, FCameraDebugRenderer& Renderer)
+{
+	Renderer.AddText(TEXT("blend %.2f%%"), BlendFactor * 100.f);
+}
+#endif  // UE_GAMEPLAY_CAMERAS_DEBUG
+
 UE_DEFINE_BLEND_CAMERA_NODE_EVALUATOR(FSimpleFixedTimeBlendCameraNodeEvaluator)
 
 void FSimpleFixedTimeBlendCameraNodeEvaluator::OnRun(const FCameraNodeEvaluationParams& Params, FCameraNodeEvaluationResult& OutResult)
diff --git a/Engine/Plugins/Cameras/GameplayCameras/Source/GameplayCameras/Private/Nodes/Common/DampenPositionCameraNode.cpp b/Engine/Plugins/Cameras/GameplayCameras/Source/GameplayCameras/Private/Nodes/Common/DampenPositionCameraNode.cpp
--- a/Engine/Plugins/Cameras/GameplayCameras/Source/GameplayCameras/Private/Nodes/Common/DampenPositionCameraNode.cpp
+++ b/Engine/Plugins/Cameras/GameplayCameras/Source/GameplayCameras/Private/Nodes/Common/DampenPositionCameraNode.cpp
@@ -4,6 +4,7 @@
 
 #include "Core/CameraEvaluationContext.h"
 #include "Core/CameraNodeEvaluator.h"
+#include "Debug/CameraDebugBlock.h"
 #include "Debug/CameraDebugBlockBuilder.h"
 #include "Debug/CameraDebugRenderer.h"
 #include "GameplayCameras.h"
@@ -24,6 +25,10 @@ protected:
 	virtual void OnInitialize(const FCameraNodeEvaluatorInitializeParams& Params) override;
 	virtual void OnRun(const FCameraNodeEvaluationParams& Params, FCameraNodeEvaluationResult& OutResult) override;
 
+#if UE_GAMEPLAY_CAMERAS_DEBUG
+	virtual void OnBuildDebugBlocks(const FCameraDebugBlockBuildParams& Params, FCameraDebugBlockBuilder& Builder) override;
+#endif  // UE_GAMEPLAY_CAMERAS_DEBUG
+
 private:
 
 	FCriticalDamper ForwardDamper;
@@ -31,20 +36,27 @@ private:
 	FCriticalDamper VerticalDamper;
 
 	FVector3d PreviousLocation;
+
+#if UE_GAMEPLAY_CAMERAS_DEBUG
+	FVector3d LastUndampedPosition;
+	FVector3d LastDampedPosition;
+#endif  // UE_GAMEPLAY_CAMERAS_DEBUG
 };
 
 UE_DEFINE_CAMERA_NODE_EVALUATOR(FDampenPositionCameraNodeEvaluator)
 
-UE_DEFINE_CAMERA_DEBUG_BLOCK_START(FDampenPositionCameraDebugBlock)
-	UE_DEFINE_CAMERA_DEBUG_BLOCK_FIELD(float, ForwardX0);
-	UE_DEFINE_CAMERA_DEBUG_BLOCK_FIELD(float, LateralX0);
-	UE_DEFINE_CAMERA_DEBUG_BLOCK_FIELD(float, VerticalX0);
-	UE_DEFINE_CAMERA_DEBUG_BLOCK_FIELD(float, ForwardDampingFactor);
-	UE_DEFINE_CAMERA_DEBUG_BLOCK_FIELD(float, LateralDampingFactor);
-	UE_DEFINE_CAMERA_DEBUG_BLOCK_FIELD(float, VerticalDampingFactor);
-	UE_DEFINE_CAMERA_DEBUG_BLOCK_FIELD(FVector3d, UndampedPosition);
-	UE_DEFINE_CAMERA_DEBUG_BLOCK_FIELD(FVector3d, DampedPosition);
-UE_DEFINE_CAMERA_DEBUG_BLOCK_END()
+UE_DECLARE_CAMERA_DEBUG_BLOCK_START(FDampenPositionCameraDebugBlock)
+	UE_DECLARE_CAMERA_DEBUG_BLOCK_FIELD(float, ForwardX0);
+	UE_DECLARE_CAMERA_DEBUG_BLOCK_FIELD(float, LateralX0);
+	UE_DECLARE_CAMERA_DEBUG_BLOCK_FIELD(float, VerticalX0);
+	UE_DECLARE_CAMERA_DEBUG_BLOCK_FIELD(float, ForwardDampingFactor);
+	UE_DECLARE_CAMERA_DEBUG_BLOCK_FIELD(float, LateralDampingFactor);
+	UE_DECLARE_CAMERA_DEBUG_BLOCK_FIELD(float, VerticalDampingFactor);
+	UE_DECLARE_CAMERA_DEBUG_BLOCK_FIELD(FVector3d, UndampedPosition);
+	UE_DECLARE_CAMERA_DEBUG_BLOCK_FIELD(FVector3d, DampedPosition);
+UE_DECLARE_CAMERA_DEBUG_BLOCK_END()
+
+UE_DEFINE_CAMERA_DEBUG_BLOCK(FDampenPositionCameraDebugBlock)
 
 void FDampenPositionCameraNodeEvaluator::OnInitialize(const FCameraNodeEvaluatorInitializeParams& Params)
 {
@@ -113,24 +125,12 @@ void FDampenPositionCameraNodeEvaluator::OnRun(const FCameraNodeEvaluationParams
 		}
 		
 		NextLocation = NewDampedLocation;
+	}
 
 #if UE_GAMEPLAY_CAMERAS_DEBUG
-		if (OutResult.DebugBlockBuilder)
-		{
-			FDampenPositionCameraDebugBlock& DebugBlock = OutResult.DebugBlockBuilder->StartBlock<FDampenPositionCameraDebugBlock>();
-			{
-				DebugBlock.ForwardX0 = ForwardDamper.GetX0();
-				DebugBlock.LateralX0 = LateralDamper.GetX0();
-				DebugBlock.VerticalX0 = VerticalDamper.GetX0();
-
-				DebugBlock.ForwardDampingFactor = ForwardDamper.GetW0();
-				DebugBlock.LateralDampingFactor = LateralDamper.GetW0();
-				DebugBlock.VerticalDampingFactor = VerticalDamper.GetW0();
-			}
-			OutResult.DebugBlockBuilder->EndBlock();
-		}
+	LastUndampedPosition = NextTarget;
+	LastDampedPosition = NextLocation;
 #endif  // UE_GAMEPLAY_CAMERAS_DEBUG
-	}
 
 	PreviousLocation = NextLocation;
 
@@ -139,25 +139,35 @@ void FDampenPositionCameraNodeEvaluator::OnRun(const FCameraNodeEvaluationParams
 
 #if UE_GAMEPLAY_CAMERAS_DEBUG
 
-EDebugDrawResult FDampenPositionCameraDebugBlock::OnDebugDraw(const FCameraDebugBlockDrawParams& Params, FCameraDebugRenderer& Renderer)
+void FDampenPositionCameraNodeEvaluator::OnBuildDebugBlocks(const FCameraDebugBlockBuildParams& Params, FCameraDebugBlockBuilder& Builder)
+{
+	FDampenPositionCameraDebugBlock& DebugBlock = Builder.AttachDebugBlock<FDampenPositionCameraDebugBlock>();
+
+	DebugBlock.ForwardX0 = ForwardDamper.GetX0();
+	DebugBlock.LateralX0 = LateralDamper.GetX0();
+	DebugBlock.VerticalX0 = VerticalDamper.GetX0();
+
+	DebugBlock.ForwardDampingFactor = ForwardDamper.GetW0();
+	DebugBlock.LateralDampingFactor = LateralDamper.GetW0();
+	DebugBlock.VerticalDampingFactor = VerticalDamper.GetW0();
+
+	DebugBlock.UndampedPosition = LastUndampedPosition;
+	DebugBlock.DampedPosition = LastDampedPosition;
+}
+
+void FDampenPositionCameraDebugBlock::OnDebugDraw(const FCameraDebugBlockDrawParams& Params, FCameraDebugRenderer& Renderer)
 {
 	Renderer.AddText(
 			TEXT("forward %.3f (factor %.3f)  lateral %.3f (factor %.3f)  vertical %.3f (factor %.3f)"),
 			ForwardX0, ForwardDampingFactor,
 			LateralX0, LateralDampingFactor,
 			VerticalX0, VerticalDampingFactor);
-	return EDebugDrawResult::Default;
 }
 
 #endif  // UE_GAMEPLAY_CAMERAS_DEBUG
 
 }  // namespace UE::Cameras
 
-UDampenPositionCameraNode::UDampenPositionCameraNode(const FObjectInitializer& ObjectInit)
-	: Super(ObjectInit)
-{
-}
-
 FCameraNodeEvaluatorPtr UDampenPositionCameraNode::OnBuildEvaluator(FCameraNodeEvaluatorBuilder& Builder) const
 {
 	using namespace UE::Cameras;
diff --git a/Engine/Plugins/Cameras/GameplayCameras/Source/GameplayCameras/Public/Core/BlendStackCameraNode.h b/Engine/Plugins/Cameras/GameplayCameras/Source/GameplayCameras/Public/Core/BlendStackCameraNode.h
--- a/Engine/Plugins/Cameras/GameplayCameras/Source/GameplayCameras/Public/Core/BlendStackCameraNode.h
+++ b/Engine/Plugins/Cameras/GameplayCameras/Source/GameplayCameras/Public/Core/BlendStackCameraNode.h
@@ -6,6 +6,7 @@
 #include "Core/CameraNodeEvaluator.h"
 #include "Core/CameraNodeEvaluatorStorage.h"
 #include "Core/CameraRigAsset.h"
+#include "Debug/CameraDebugBlock.h"
 #include "IGameplayCamerasLiveEditListener.h"
 
 #include "BlendStackCameraNode.generated.h"
@@ -22,6 +23,10 @@ class FBlendStackRootCameraNodeEvaluator;
 class FCameraEvaluationContext;
 class FCameraSystemEvaluator;
 
+#if UE_GAMEPLAY_CAMERAS_DEBUG
+class FBlendStackCameraDebugBlock;
+#endif  // UE_GAMEPLAY_CAMERAS_DEBUG
+
 }  // namespace UE::Cameras
 
 /**
@@ -87,6 +92,10 @@ public:
 	/** Push a new camera rig onto the blend stack. */
 	void Push(const FBlendStackCameraPushParams& Params);
 
+#if UE_GAMEPLAY_CAMERAS_DEBUG
+	FBlendStackCameraDebugBlock* BuildDetailedDebugBlock(const FCameraDebugBlockBuildParams& Params, FCameraDebugBlockBuilder& Builder);
+#endif  // UE_GAMEPLAY_CAMERAS_DEBUG
+
 protected:
 
 	// FCameraNodeEvaluator interface
@@ -95,6 +104,10 @@ protected:
 	virtual void OnRun(const FCameraNodeEvaluationParams& Params, FCameraNodeEvaluationResult& OutResult) override;
 	virtual void OnAddReferencedObjects(FReferenceCollector& Collector) override;
 
+#if UE_GAMEPLAY_CAMERAS_DEBUG
+	virtual void OnBuildDebugBlocks(const FCameraDebugBlockBuildParams& Params, FCameraDebugBlockBuilder& Builder) override;
+#endif  // UE_GAMEPLAY_CAMERAS_DEBUG
+
 #if WITH_EDITOR
 	// IGameplayCamerasLiveEditListener interface
 	virtual void OnPostBuildAsset(const FGameplayCameraAssetBuildEvent& BuildEvent) override;
@@ -143,7 +156,57 @@ protected:
 #if WITH_EDITOR
 	TMap<const UPackage*, int32> AllListenedPackages;
 #endif  // WITH_EDITOR
+
+#if UE_GAMEPLAY_CAMERAS_DEBUG
+	friend class FBlendStackSummaryCameraDebugBlock;
+	friend class FBlendStackCameraDebugBlock;
+#endif  // UE_GAMEPLAY_CAMERAS_DEBUG
+};
+
+#if UE_GAMEPLAY_CAMERAS_DEBUG
+
+class FBlendStackSummaryCameraDebugBlock : public FCameraDebugBlock
+{
+	UE_DECLARE_CAMERA_DEBUG_BLOCK(FBlendStackSummaryCameraDebugBlock)
+
+public:
+
+	FBlendStackSummaryCameraDebugBlock();
+	FBlendStackSummaryCameraDebugBlock(const FBlendStackCameraNodeEvaluator& InEvaluator);
+
+protected:
+
+	virtual void OnDebugDraw(const FCameraDebugBlockDrawParams& Params, FCameraDebugRenderer& Renderer) override;
+
+private:
+
+	int32 NumEntries;
 };
 
+class FBlendStackCameraDebugBlock : public FCameraDebugBlock
+{
+	UE_DECLARE_CAMERA_DEBUG_BLOCK(FBlendStackCameraDebugBlock)
+
+public:
+
+	FBlendStackCameraDebugBlock();
+	FBlendStackCameraDebugBlock(const FBlendStackCameraNodeEvaluator& InEvaluator);
+	
+protected:
+
+	virtual void OnDebugDraw(const FCameraDebugBlockDrawParams& Params, FCameraDebugRenderer& Renderer) override;
+
+private:
+
+	struct FEntryDebugInfo
+	{
+		FString CameraRigName;
+	};
+
+	TArray<FEntryDebugInfo> Entries;
+};
+
+#endif  // UE_GAMEPLAY_CAMERAS_DEBUG
+
 }  // namespace UE::Cameras
 
diff --git a/Engine/Plugins/Cameras/GameplayCameras/Source/GameplayCameras/Public/Core/CameraEvaluationContextStack.h b/Engine/Plugins/Cameras/GameplayCameras/Source/GameplayCameras/Public/Core/CameraEvaluationContextStack.h
--- a/Engine/Plugins/Cameras/GameplayCameras/Source/GameplayCameras/Public/Core/CameraEvaluationContextStack.h
+++ b/Engine/Plugins/Cameras/GameplayCameras/Source/GameplayCameras/Public/Core/CameraEvaluationContextStack.h
@@ -4,6 +4,7 @@
 
 #include "Containers/Array.h"
 #include "Core/CameraDirectorEvaluator.h"
+#include "GameplayCameras.h"
 #include "UObject/WeakObjectPtr.h"
 
 class UCameraAsset;
@@ -54,6 +55,9 @@ public:
 	/** Pop the active (top) context. */
 	void PopContext();
 
+	/** The number of contexts on the stack. */
+	int32 NumContexts() const { return Entries.Num(); }
+
 public:
 
 	// Internal API
@@ -75,6 +79,10 @@ private:
 
 	/** The owner evaluator. */
 	TSharedPtr<FCameraSystemEvaluator> Evaluator;
+
+#if UE_GAMEPLAY_CAMERAS_DEBUG
+	friend class FCameraDirectorTreeDebugBlock;
+#endif  // UE_GAMEPLAY_CAMERAS_DEBUG
 };
 
 }  // namespace UE::Cameras
diff --git a/Engine/Plugins/Cameras/GameplayCameras/Source/GameplayCameras/Public/Core/CameraNodeEvaluator.h b/Engine/Plugins/Cameras/GameplayCameras/Source/GameplayCameras/Public/Core/CameraNodeEvaluator.h
--- a/Engine/Plugins/Cameras/GameplayCameras/Source/GameplayCameras/Public/Core/CameraNodeEvaluator.h
+++ b/Engine/Plugins/Cameras/GameplayCameras/Source/GameplayCameras/Public/Core/CameraNodeEvaluator.h
@@ -7,6 +7,7 @@
 #include "Core/CameraVariableTable.h"
 #include "Core/ObjectChildrenView.h"
 #include "CoreTypes.h"
+#include "Debug/RootCameraDebugBlock.h"
 #include "GameplayCameras.h"
 #include "UObject/ObjectPtr.h"
 
@@ -22,9 +23,8 @@ class FCameraSystemEvaluator;
 struct FCameraNodeEvaluatorBuilder;
 
 #if UE_GAMEPLAY_CAMERAS_DEBUG
-class FCameraDebugBlock;
 struct FCameraDebugBlockBuilder;
-#endif
+#endif  // UE_GAMEPLAY_CAMERAS_DEBUG
 
 /**
  * Parameter structure for running a camera node evaluator.
@@ -59,11 +59,6 @@ struct FCameraNodeEvaluationResult
 
 	/** Reset this result to its default (non-valid) state. */
 	void Reset();
-
-#if UE_GAMEPLAY_CAMERAS_DEBUG
-	/** Builder for debug drawing blocks. May be null. */
-	FCameraDebugBlockBuilder* DebugBlockBuilder = nullptr;
-#endif  // UE_GAMEPLAY_CAMERAS_DEBUG
 };
 
 /**
@@ -102,6 +97,18 @@ struct FCameraNodeEvaluatorInitializeParams
 /** View on a camera node evaluator's children. */
 using FCameraNodeEvaluatorChildrenView = TObjectChildrenView<FCameraNodeEvaluator*>;
 
+#if UE_GAMEPLAY_CAMERAS_DEBUG
+
+/**
+ * Structure for creating the node evaluator's debug blocks.
+ */
+struct FCameraDebugBlockBuildParams
+{
+	// Empty for now, but defined for later API changes.
+};
+
+#endif  // UE_GAMEPLAY_CAMERAS_DEBUG
+
 /**
  * Base class for objects responsible for running a camera node.
  */
@@ -136,6 +143,11 @@ public:
 		return Cast<CameraNodeType>(PrivateCameraNode);
 	}
 
+#if UE_GAMEPLAY_CAMERAS_DEBUG
+	/** Called to create debug blocks for this node evaluator. */
+	void BuildDebugBlocks(const FCameraDebugBlockBuildParams& Params, FCameraDebugBlockBuilder& Builder);
+#endif  // UE_GAMEPLAY_CAMERAS_DEBUG
+
 public:
 
 	// Internal API.
@@ -159,19 +171,8 @@ protected:
 	virtual void OnAddReferencedObjects(FReferenceCollector& Collector) {}
 
 #if UE_GAMEPLAY_CAMERAS_DEBUG
-
-public:
-
-	/** Creates the debug drawing block for this camera node. */
-	void CreateDebugBlock(const FCameraNodeEvaluationParams& Params, FCameraDebugBlockBuilder& Builder);
-
-private:
-
-	/**
-	 * Creates the debug drawing block for this camera node.
-	 */
-	virtual void OnCreateDebugBlock(const FCameraNodeEvaluationParams& Params, FCameraDebugBlockBuilder& Builder);
-
+	/** Called to create debug blocks for this node evaluator. */
+	virtual void OnBuildDebugBlocks(const FCameraDebugBlockBuildParams& Params, FCameraDebugBlockBuilder& Builder);
 #endif  // UE_GAMEPLAY_CAMERAS_DEBUG
 
 private:
diff --git a/Engine/Plugins/Cameras/GameplayCameras/Source/GameplayCameras/Public/Core/CameraSystemEvaluator.h b/Engine/Plugins/Cameras/GameplayCameras/Source/GameplayCameras/Public/Core/CameraSystemEvaluator.h
--- a/Engine/Plugins/Cameras/GameplayCameras/Source/GameplayCameras/Public/Core/CameraSystemEvaluator.h
+++ b/Engine/Plugins/Cameras/GameplayCameras/Source/GameplayCameras/Public/Core/CameraSystemEvaluator.h
@@ -66,7 +66,7 @@ struct FCameraSystemEvaluationUpdateResult
 };
 
 #if UE_GAMEPLAY_CAMERAS_DEBUG
-struct FCameraSystemDebugDrawParams
+struct FCameraSystemDebugUpdateParams
 {
 	UCanvas* Canvas = nullptr;
 };
@@ -96,16 +96,27 @@ public:
 	/** Pop the active (top) evaluation context from the stack. */
 	void PopEvaluationContext();
 
+	/** Gets the context stack. */
+	FCameraEvaluationContextStack& GetEvaluationContextStack() { return ContextStack; }
+	/** Gets the context stack. */
+	const FCameraEvaluationContextStack& GetEvaluationContextStack() const { return ContextStack; }
+
 public:
 
 	/** Run an update of the camera system. */
 	void Update(const FCameraSystemEvaluationUpdateParams& Params);
 
+	/** Returns the root node evaluator. */
+	FRootCameraNodeEvaluator* GetRootNodeEvaluator() const { return RootEvaluator; }
+
+	/** Gets the evaluated result. */
+	const FCameraSystemEvaluationUpdateResult& GetEvaluatedResult() const { return Result; }
+
 	/** Get the last evaluated camera. */
 	void GetEvaluatedCameraView(FMinimalViewInfo& DesiredView);
 
 #if UE_GAMEPLAY_CAMERAS_DEBUG
-	void DebugDraw(const FCameraSystemDebugDrawParams& Params);
+	void DebugUpdate(const FCameraSystemDebugUpdateParams& Params);
 #endif  // UE_GAMEPLAY_CAMERAS_DEBUG
 
 public:
diff --git a/Engine/Plugins/Cameras/GameplayCameras/Source/GameplayCameras/Public/Core/CameraVariableTable.h b/Engine/Plugins/Cameras/GameplayCameras/Source/GameplayCameras/Public/Core/CameraVariableTable.h
--- a/Engine/Plugins/Cameras/GameplayCameras/Source/GameplayCameras/Public/Core/CameraVariableTable.h
+++ b/Engine/Plugins/Cameras/GameplayCameras/Source/GameplayCameras/Public/Core/CameraVariableTable.h
@@ -6,6 +6,7 @@
 #include "Containers/Set.h"
 #include "Core/CameraVariableTableFwd.h"
 #include "CoreTypes.h"
+#include "GameplayCameras.h"
 #include "Misc/EnumClassFlags.h"
 #include "Templates/UnrealTypeTraits.h"
 #include "UObject/NameTypes.h"
@@ -168,6 +169,10 @@ private:
 
 	template<typename T>
 	friend struct TCameraVariableInterpolation;
+
+#if UE_GAMEPLAY_CAMERAS_DEBUG
+	friend class FVariableTableDebugBlock;
+#endif  // UE_GAMEPLAY_CAMERAS_DEBUG
 };
 
 ENUM_CLASS_FLAGS(FCameraVariableTable::EEntryFlags)
@@ -178,7 +183,10 @@ const ValueType* FCameraVariableTable::FindValue(uint32 VariableId) const
 	if (const FEntry* Entry = Entries.Find(VariableId))
 	{
 		CheckVariableType<ValueType>(Entry->Type);
-		return reinterpret_cast<ValueType*>(Memory + Entry->Offset);
+		if (EnumHasAnyFlags(Entry->Flags, EEntryFlags::Written))
+		{
+			return reinterpret_cast<ValueType*>(Memory + Entry->Offset);
+		}
 	}
 	return nullptr;
 }
@@ -188,6 +196,17 @@ const ValueType& FCameraVariableTable::GetValue(uint32 VariableId) const
 {
 	const FEntry& Entry = Entries.FindChecked(VariableId);
 	CheckVariableType<ValueType>(Entry.Type);
+#if WITH_EDITORONLY_DATA
+	checkf(
+			EnumHasAnyFlags(Entry.Flags, EEntryFlags::Written),
+			TEXT("Variable '%s' has never been written to. GetValue() will return uninitialized memory!"),
+			*Entry.DebugName);
+#else
+	checkf(
+			EnumHasAnyFlags(Entry.Flags, EEntryFlags::Written),
+			TEXT("Variable '%s' has never been written to. GetValue() will return uninitialized memory!"),
+			*LexToString(VariableId));
+#endif
 	return *reinterpret_cast<ValueType*>(Memory + Entry.Offset);
 }
 
diff --git a/Engine/Plugins/Cameras/GameplayCameras/Source/GameplayCameras/Public/Core/DefaultRootCameraNode.h b/Engine/Plugins/Cameras/GameplayCameras/Source/GameplayCameras/Public/Core/DefaultRootCameraNode.h
--- a/Engine/Plugins/Cameras/GameplayCameras/Source/GameplayCameras/Public/Core/DefaultRootCameraNode.h
+++ b/Engine/Plugins/Cameras/GameplayCameras/Source/GameplayCameras/Public/Core/DefaultRootCameraNode.h
@@ -59,6 +59,10 @@ protected:
 	virtual void OnRun(const FCameraNodeEvaluationParams& Params, FCameraNodeEvaluationResult& OutResult) override;
 	virtual void OnActivateCameraRig(const FActivateCameraRigParams& Params) override;
 
+#if UE_GAMEPLAY_CAMERAS_DEBUG
+	virtual void OnBuildDebugBlocks(const FCameraDebugBlockBuildParams& Params, FCameraDebugBlockBuilder& Builder) override;
+#endif  // UE_GAMEPLAY_CAMERAS_DEBUG
+
 private:
 
 	FBlendStackCameraNodeEvaluator* BaseLayer;
diff --git a/Engine/Plugins/Cameras/GameplayCameras/Source/GameplayCameras/Public/Debug/BlendStacksCameraDebugBlock.h b/Engine/Plugins/Cameras/GameplayCameras/Source/GameplayCameras/Public/Debug/BlendStacksCameraDebugBlock.h
new file mode 100644
--- /dev/null
+++ b/Engine/Plugins/Cameras/GameplayCameras/Source/GameplayCameras/Public/Debug/BlendStacksCameraDebugBlock.h
@@ -0,0 +1,43 @@
+// Copyright Epic Games, Inc. All Rights Reserved.
+
+#pragma once
+
+#include "Debug/CameraDebugBlock.h"
+
+#if UE_GAMEPLAY_CAMERAS_DEBUG
+
+namespace UE::Cameras
+{
+
+class FBlendStackCameraDebugBlock;
+
+/**
+ * Debug block that shows multiple blend stacks.
+ */
+class FBlendStacksCameraDebugBlock : public FCameraDebugBlock
+{
+	UE_DECLARE_CAMERA_DEBUG_BLOCK(FBlendStacksCameraDebugBlock)
+
+public:
+
+	FBlendStacksCameraDebugBlock();
+
+	/**
+	 * Adds the given blend stack block as a child, which will be rendered with the given title.
+	 * The blend stack block can be obtained from the blend stack evaluator.
+	 */
+	void AddBlendStack(const FString& InBlendStackName, FBlendStackCameraDebugBlock* InDebugBlock);
+
+protected:
+
+	virtual void OnDebugDraw(const FCameraDebugBlockDrawParams& Params, FCameraDebugRenderer& Renderer) override;
+
+private:
+
+	TArray<FString> BlendStackNames;
+};
+
+}  // namespace UE::Cameras
+
+#endif  // UE_GAMEPLAY_CAMERAS_DEBUG
+
diff --git a/Engine/Plugins/Cameras/GameplayCameras/Source/GameplayCameras/Public/Debug/CameraDebugBlock.h b/Engine/Plugins/Cameras/GameplayCameras/Source/GameplayCameras/Public/Debug/CameraDebugBlock.h
--- a/Engine/Plugins/Cameras/GameplayCameras/Source/GameplayCameras/Public/Debug/CameraDebugBlock.h
+++ b/Engine/Plugins/Cameras/GameplayCameras/Source/GameplayCameras/Public/Debug/CameraDebugBlock.h
@@ -2,7 +2,7 @@
 
 #pragma once
 
-#include "Containers/Array.h"
+#include "Core/CameraObjectRtti.h"
 #include "CoreTypes.h"
 #include "GameplayCameras.h"
 
@@ -13,15 +13,18 @@ class FArchive;
 namespace UE::Cameras
 {
 
-class FCameraDebugBlock;
 class FCameraDebugRenderer;
+class FCameraDebugBlock;
 
 /**
  * Parameter structure for debug block drawing.
  */
 struct FCameraDebugBlockDrawParams
 {
-	float DeltaTime = 0.f;
+	/** The list of active debug categories. */
+	TSet<FString> ActiveCategories;
+
+	bool IsCategoryActive(const FString& InCategory) const;
 };
 
 /**
@@ -57,14 +60,6 @@ struct TCameraDebugBlockField : FCameraDebugBlockField
 	}
 };
 
-enum class EDebugDrawResult
-{
-	Default = 0,
-	SkipChildren = 1 << 0
-};
-
-ENUM_CLASS_FLAGS(EDebugDrawResult);
-
 /**
  * Debug drawing structure responsible for displaying information about some aspect
  * of the camera system evaluation. For instance, there is generally one debug block
@@ -72,13 +67,19 @@ ENUM_CLASS_FLAGS(EDebugDrawResult);
  */
 class FCameraDebugBlock
 {
+	UE_GAMEPLAY_CAMERAS_DECLARE_RTTI_BASE(FCameraDebugBlock)
+
 public:
 
 	virtual ~FCameraDebugBlock() {}
 
+	/** Attaches a block to this block. */
+	void Attach(FCameraDebugBlock* InAttachment);
+	/** Gets the list of blocks attached to this block. */
+	TArrayView<FCameraDebugBlock*> GetAttachments() { return Attachments; }
+
 	/** Adds a child block to this block. */
 	void AddChild(FCameraDebugBlock* InChild);
-
 	/** Gets the children of this block. */
 	TArrayView<FCameraDebugBlock*> GetChildren() { return Children; }
 
@@ -90,9 +91,12 @@ public:
 
 protected:
 
-	virtual EDebugDrawResult OnDebugDraw(const FCameraDebugBlockDrawParams& Params, FCameraDebugRenderer& Renderer);
+	/** Called to let this block display its information on screen. */
+	virtual void OnDebugDraw(const FCameraDebugBlockDrawParams& Params, FCameraDebugRenderer& Renderer) {}
+	/** Called after attached and children blocks' debug draw. */
 	virtual void OnPostDebugDraw(const FCameraDebugBlockDrawParams& Params, FCameraDebugRenderer& Renderer) {}
 
+	/** Serializes this debug block into a buffer, for recording/replaying purposes. */
 	virtual void OnSerialize(FArchive& Ar) {}
 
 protected:
@@ -113,22 +117,50 @@ private:
 
 	static TArray<FCameraDebugBlockField*> StaticFields;
 
-	using FChildrenArray = TArray<FCameraDebugBlock*>;
-	FChildrenArray Children;
+	using FRelatedDebugBlockArray = TArray<FCameraDebugBlock*>;
+	FRelatedDebugBlockArray Attachments;
+	FRelatedDebugBlockArray Children;
 };
 
 }  // namespace UE::Cameras
 
-// Macros for defining debug blocks for a camera node evaluator.
+// Macros for defining debug blocks.
 //
-#define UE_DEFINE_CAMERA_DEBUG_BLOCK_START(ClassName)\
+#define UE_DECLARE_CAMERA_DEBUG_BLOCK(ClassName)\
+	UE_GAMEPLAY_CAMERAS_DECLARE_RTTI(ClassName, ::UE::Cameras::FCameraDebugBlock)
+
+#define UE_DECLARE_CAMERA_DEBUG_BLOCK_EX(ClassName, BaseClassName)\
+	UE_GAMEPLAY_CAMERAS_DECLARE_RTTI(ClassName, BaseClassName)
+
+#define UE_DEFINE_CAMERA_DEBUG_BLOCK(ClassName)\
+	UE_GAMEPLAY_CAMERAS_DEFINE_RTTI(ClassName)
+
+#endif  // UE_GAMEPLAY_CAMERAS_DEBUG
+
+#if UE_GAMEPLAY_CAMERAS_DEBUG
+
+// Macros for defining "simple" debug blocks with primitive types as fields.
+// This will setup debug block serialization automatically, compared to the macros
+// above where you have to write your own OnSerialize().
+//
+// Example:
+//
+//     UE_DECLARE_CAMERA_DEBUG_BLOCK_START(FMyDebugBlock)
+//			UE_DECLARE_CAMERA_DEBUG_BLOCK_FIELD(float, Something)
+//			UE_DECLARE_CAMERA_DEBUG_BLOCK_FIELD(bool, bSomethingElse)
+//	   UE_DECLARE_CAMERA_DEBUG_BLOCK_END()
+//
+//	   UE_DEFINE_CAMERA_DEBUG_BLOCK(FMyDebugBlock)
+//
+#define UE_DECLARE_CAMERA_DEBUG_BLOCK_START(ClassName)\
 	class ClassName : public ::UE::Cameras::FCameraDebugBlock\
 	{\
+		UE_DECLARE_CAMERA_DEBUG_BLOCK(ClassName)\
 	private:\
 		using Super = ::UE::Cameras::FCameraDebugBlock;\
 		using ThisClassName = ClassName;
 
-#define UE_DEFINE_CAMERA_DEBUG_BLOCK_FIELD(FieldType, FieldName)\
+#define UE_DECLARE_CAMERA_DEBUG_BLOCK_FIELD(FieldType, FieldName)\
 	public:\
 		FieldType FieldName;\
 	private:\
@@ -139,15 +171,19 @@ private:
 			return &StaticField; }\
 		inline static const int32 FieldName##FieldIndex = RegisterField(Get##FieldName##Field());
 
-#define UE_DEFINE_CAMERA_DEBUG_BLOCK_END()\
-		virtual EDebugDrawResult OnDebugDraw(const FCameraDebugBlockDrawParams& Params, FCameraDebugRenderer& Renderer) override;\
+#define UE_DECLARE_CAMERA_DEBUG_BLOCK_END()\
+		virtual void OnDebugDraw(const FCameraDebugBlockDrawParams& Params, FCameraDebugRenderer& Renderer) override;\
 	};
 
-#else
+#else  // UE_GAMEPLAY_CAMERAS_DEBUG
 
-#define UE_DEFINE_CAMERA_DEBUG_BLOCK_START(ClassName)
-#define UE_DEFINE_CAMERA_DEBUG_BLOCK_FIELD(FieldType, FieldName)
-#define UE_DEFINE_CAMERA_DEBUG_BLOCK_END()
+// Empty macros for shipping builds.
+//
+#define UE_DECLARE_CAMERA_DEBUG_BLOCK_START(ClassName)
+#define UE_DECLARE_CAMERA_DEBUG_BLOCK_FIELD(FieldType, FieldName)
+#define UE_DECLARE_CAMERA_DEBUG_BLOCK_END()
+// This one is included here too since it's required for "simple" debug blocks.
+#define UE_DEFINE_CAMERA_DEBUG_BLOCK(ClassName)
 
 #endif  // UE_GAMEPLAY_CAMERAS_DEBUG
 
diff --git a/Engine/Plugins/Cameras/GameplayCameras/Source/GameplayCameras/Public/Debug/CameraDebugBlockBuilder.h b/Engine/Plugins/Cameras/GameplayCameras/Source/GameplayCameras/Public/Debug/CameraDebugBlockBuilder.h
--- a/Engine/Plugins/Cameras/GameplayCameras/Source/GameplayCameras/Public/Debug/CameraDebugBlockBuilder.h
+++ b/Engine/Plugins/Cameras/GameplayCameras/Source/GameplayCameras/Public/Debug/CameraDebugBlockBuilder.h
@@ -15,6 +15,16 @@ namespace UE::Cameras
 {
 
 class FCameraDebugBlock;
+class FCameraNodeEvaluatorDebugBlock;
+class FRootCameraDebugBlock;
+struct FCameraDebugBlockBuilder;
+
+enum class ECameraDebugBlockBuildVisitFlags
+{
+	None,
+	SkipChildren = 1 << 0
+};
+ENUM_CLASS_FLAGS(ECameraDebugBlockBuildVisitFlags);
 
 /**
  * Builder class for camera debug drawing blocks.
@@ -24,32 +34,91 @@ struct FCameraDebugBlockBuilder
 public:
 
 	/** Creates a new builder structure. */
-	FCameraDebugBlockBuilder(FCameraDebugBlockStorage& InStorage, FCameraDebugBlock& InRootBlock);
+	FCameraDebugBlockBuilder(FCameraDebugBlockStorage& InStorage, FRootCameraDebugBlock& InRootBlock);
+
+	/** Gets the storage used by this builder. */
+	FCameraDebugBlockStorage& GetStorage() const { return Storage; }
+
+	/** Gets the root debug block. */
+	FRootCameraDebugBlock& GetRootDebugBlock() const { return RootBlock; }
+
+	/** Gets the current parent debug block. */
+	FCameraDebugBlock& GetParentDebugBlock() const { return *CurrentHierarchy.Last(); }
+
+	/**
+	 * Creates a new unassociated debug block.
+	 * It won't render unless it's referenced by or parented under another debug block!
+	 */
+	template<typename BlockType, typename ...ArgTypes>
+	BlockType& BuildDebugBlock(ArgTypes&&... InArgs)
+	{
+		return *Storage.BuildDebugBlock<BlockType>(Forward<ArgTypes>(InArgs)...);
+	}
+
+	/**
+	 * Creates a new debug block and attaches it to the current active block.
+	 *
+	 * Attached debug blocks are rendered/not rendered along with their "anchor" debug block.
+	 * This differs from children debug blocks which may be rendered/not rendered independently
+	 * of their parent.
+	 */
+	template<typename BlockType, typename ...ArgTypes>
+	BlockType& AttachDebugBlock(ArgTypes&&... InArgs)
+	{
+		BlockType* NewBlock = Storage.BuildDebugBlock<BlockType>(Forward<ArgTypes>(InArgs)...);
+		OnAttachDebugBlock(NewBlock);
+		return *NewBlock;
+	}
 
 	/**
-	 * Creates a new debug drawing block.
+	 * Creates a new debug block and adds it to the current hierarhcy.
 	 * This sets the new block as the "active" block, and adds it as a child of the
 	 * previously active block.
 	 */
 	template<typename BlockType, typename ...ArgTypes>
-	BlockType& StartBlock(ArgTypes&& ...InArgs)
+	BlockType& StartChildDebugBlock(ArgTypes&& ...InArgs)
 	{
 		BlockType* NewBlock = Storage.BuildDebugBlock<BlockType>(Forward<ArgTypes>(InArgs)...);
-		OnStartBlock(*NewBlock);
+		OnStartChildDebugBlock(NewBlock);
 		return *NewBlock;
 	}
 
 	/** Ends the currently active debug drawing block. */
-	void EndBlock();
+	void EndChildDebugBlock();
+
+	/** Gets current hierarchy level. */
+	int32 GetHierarchyLevel() const { return CurrentHierarchy.Num(); }
+
+	/**
+	 * Don't visit children node evaluators when building the hierarchy of debug blocks.
+	 *
+	 * This implies that a node evaluator will visit its children "manually", otherwise these
+	 * children node evaluators won't have any debugging information available.
+	 */
+	void SkipChildren();
+	/** Gets visiting flags. */
+	ECameraDebugBlockBuildVisitFlags GetVisitFlags() const { return VisitFlags; }
+	/** Resets visiting flags. */
+	void ResetVisitFlags();
+
+	/** Temporarily overrides the attachment/children to work on the new parent. */
+	void StartParentDebugBlockOverride(FCameraDebugBlock& InNewParent);
+	/** Ends a temporary attachment/children override. */
+	void EndParentDebugBlockOverride();
 
 private:
 
-	void OnStartBlock(FCameraDebugBlock& InNewBlock);
+	void OnAttachDebugBlock(FCameraDebugBlock* InNewBlock);
+	void OnStartChildDebugBlock(FCameraDebugBlock* InNewBlock);
 
 private:
 
 	FCameraDebugBlockStorage& Storage;
+	FRootCameraDebugBlock& RootBlock;
 	TArray<FCameraDebugBlock*> CurrentHierarchy;
+
+	ECameraDebugBlockBuildVisitFlags VisitFlags;
+	int32 HierarchyOverrideStart = INDEX_NONE;
 };
 
 }  // namespace UE::Cameras
diff --git a/Engine/Plugins/Cameras/GameplayCameras/Source/GameplayCameras/Public/Debug/CameraDebugCategories.h b/Engine/Plugins/Cameras/GameplayCameras/Source/GameplayCameras/Public/Debug/CameraDebugCategories.h
new file mode 100644
--- /dev/null
+++ b/Engine/Plugins/Cameras/GameplayCameras/Source/GameplayCameras/Public/Debug/CameraDebugCategories.h
@@ -0,0 +1,28 @@
+// Copyright Epic Games, Inc. All Rights Reserved.
+
+#pragma once
+
+#include "Containers/StringFwd.h"
+#include "GameplayCameras.h"
+
+#if UE_GAMEPLAY_CAMERAS_DEBUG
+
+namespace UE::Cameras
+{
+
+/**
+ * Standard debug categories for showing different "debug modes".
+ */
+struct GAMEPLAYCAMERAS_API FCameraDebugCategories
+{
+	static const FString NodeTree;
+	static const FString DirectorTree;
+	static const FString BlendStacks;
+	static const FString PoseStats;
+	static const FString Viewfinder;
+};
+
+}  // namespace UE::Cameras
+
+#endif  // UE_GAMEPLAY_CAMERAS_DEBUG
+
diff --git a/Engine/Plugins/Cameras/GameplayCameras/Source/GameplayCameras/Public/Debug/CameraDebugColors.h b/Engine/Plugins/Cameras/GameplayCameras/Source/GameplayCameras/Public/Debug/CameraDebugColors.h
new file mode 100644
--- /dev/null
+++ b/Engine/Plugins/Cameras/GameplayCameras/Source/GameplayCameras/Public/Debug/CameraDebugColors.h
@@ -0,0 +1,60 @@
+// Copyright Epic Games, Inc. All Rights Reserved.
+
+#pragma once
+
+#include "GameplayCameras.h"
+#include "Math/Color.h"
+#include "Misc/OptionalFwd.h"
+
+#if UE_GAMEPLAY_CAMERAS_DEBUG
+
+namespace UE::Cameras
+{
+
+/**
+ * Standard debug drawing colors for the camera system.
+ */
+class FCameraDebugColors
+{
+public:
+
+	static const FCameraDebugColors& Get();
+
+	/** Debug category titles and section titles. */
+	FColor Title;
+	/** Normal text. */
+	FColor Default;
+	/** Unimportant text. */
+	FColor Passive;
+	FColor VeryPassive;
+	/** Important or notable text. */
+	FColor Hightlighted;
+	/** Text with no specific importance, but that needs to be separate from the rest. */
+	FColor Notice;
+	FColor Notice2;
+	/** Positive text. */
+	FColor Good;
+	/** Warning text. */
+	FColor Warning;
+	/** Error text. */
+	FColor Error;
+	/** Background. Should only be used for the background tile. */
+	FColor Background;
+
+public:
+
+	static TOptional<FColor> GetFColorByName(const FString& InColorName);
+
+private:
+
+	FCameraDebugColors();
+
+	static void UpdateColorMap(const FCameraDebugColors& Instance);
+
+	static TMap<FString, FColor> ColorMap;
+};
+
+}  // namespace UE::Cameras
+
+#endif  // UE_GAMEPLAY_CAMERAS_DEBUG
+
diff --git a/Engine/Plugins/Cameras/GameplayCameras/Source/GameplayCameras/Public/Debug/CameraDebugRenderer.h b/Engine/Plugins/Cameras/GameplayCameras/Source/GameplayCameras/Public/Debug/CameraDebugRenderer.h
--- a/Engine/Plugins/Cameras/GameplayCameras/Source/GameplayCameras/Public/Debug/CameraDebugRenderer.h
+++ b/Engine/Plugins/Cameras/GameplayCameras/Source/GameplayCameras/Public/Debug/CameraDebugRenderer.h
@@ -5,6 +5,7 @@
 #include "CoreTypes.h"
 #include "GameplayCameras.h"
 #include "Math/Color.h"
+#include "Math/Vector.h"
 #include "Misc/StringBuilder.h"
 
 class UCanvas;
@@ -15,39 +16,104 @@ class UFont;
 namespace UE::Cameras
 {
 
+class FDebugTextRenderer;
+
+enum class ECameraDebugDrawVisitFlags
+{
+	None = 0,
+	SkipAttachedBlocks = 1 << 0,
+	SkipChildrenBlocks = 1 << 1
+};
+ENUM_CLASS_FLAGS(ECameraDebugDrawVisitFlags)
+
+/**
+ * Utility class for camera-related debug drawing.
+ */
 class FCameraDebugRenderer
 {
 public:
 
+	/** Creates a new debug renderer. */
 	FCameraDebugRenderer(UCanvas* InCanvas);
+	/** Destroys the debug renderer. */
 	~FCameraDebugRenderer();
 
+	/** Adds text to the text wall. */
 	void AddText(const FString& InString);
 	void AddText(const TCHAR* Fmt, ...);
 
+	/** 
+	 * Move to a new line on the text wall.
+	 *
+	 * @return Whether a new line was added.
+	 */
+	bool NewLine(bool bSkipIfEmptyLine = false);
+
+	/** Gets the current text color. */
+	FColor GetTextColor() const;
+	/** Sets the text color for further calls. Returns the previous color. */
+	FColor SetTextColor(const FColor& Color);
+
+	/** Increases the indent of the next text wall entry. This will make a new line. */
 	void AddIndent();
+	/** Decreases the indent of the next text wall entry. This will make a new line. */
 	void RemoveIndent();
 
+	/** Draws a translucent background behind the text. */
+	void DrawTextBackgroundTile(float Opacity);
+
+public:
+
+	/** Request skipping drawing any blocks attached to the current block. */
+	void SkipAttachedBlocks();
+	/** Request skipping drawing any children blocks of the current block. */
+	void SkipChildrenBlocks();
+	/** Skip all related blocks (attached, children, etc.) */
+	void SkipAllBlocks();
+	/**Gets block visiting flags. */
+	ECameraDebugDrawVisitFlags GetVisitFlags() const;
+	/** Resets block visiting flags. */
+	void ResetVisitFlags();
+
+public:
+
+	/** Gets the drawing canvas. */
+	UCanvas* GetCanvas() const { return Canvas; }
+
 private:
 
 	void AddTextFmtImpl(const TCHAR* Fmt, va_list Args);
 	void AddTextImpl(const TCHAR* Buffer);
 
+	float GetIndentMargin() const;
 	void FlushText();
 
 private:
 
+	/** The canvas used to draw the text wall. */
 	UCanvas* Canvas;
+	/** The original draw color of the canvas. */
 	FColor OriginalDrawColor;
 
+	/** The font used to render the text wall. */
 	const UFont* RenderFont;
+	/** The height of one line of the text wall. */
 	int32 MaxCharHeight;
 
+	/** Temporary string formatter for variadic methods. */
 	TStringBuilder<512> Formatter;
+	/** String formatter for building a line up until the point it needs to be rendered. */
 	TStringBuilder<512> LineBuilder;
 
+	/** Current indent level. */
 	int8 IndentLevel = 0;
-	float NextBlockY = 0;
+	/** The screenspace coordinates for the next block of text on the wall. */
+	FVector2f NextDrawPosition;
+	/** The maximum horizontal extent of the text rendered so far. */
+	float RightMargin = 0;
+
+	/** How to visit the next debug blocks. */
+	ECameraDebugDrawVisitFlags VisitFlags;
 };
 
 }  // namespace UE::Cameras
diff --git a/Engine/Plugins/Cameras/GameplayCameras/Source/GameplayCameras/Public/Debug/CameraDirectorTreeDebugBlock.h b/Engine/Plugins/Cameras/GameplayCameras/Source/GameplayCameras/Public/Debug/CameraDirectorTreeDebugBlock.h
new file mode 100644
--- /dev/null
+++ b/Engine/Plugins/Cameras/GameplayCameras/Source/GameplayCameras/Public/Debug/CameraDirectorTreeDebugBlock.h
@@ -0,0 +1,45 @@
+// Copyright Epic Games, Inc. All Rights Reserved.
+
+#pragma once
+
+#include "Debug/CameraDebugBlock.h"
+
+#if UE_GAMEPLAY_CAMERAS_DEBUG
+
+namespace UE::Cameras
+{
+
+struct FCameraDebugBlockBuilder;
+struct FCameraEvaluationContextStack;
+
+/**
+ * A debug block for showing the list of camera directors in the camera system's context stack.
+ */
+class FCameraDirectorTreeDebugBlock : public FCameraDebugBlock
+{
+	UE_DECLARE_CAMERA_DEBUG_BLOCK(FCameraDirectorTreeDebugBlock)
+
+public:
+
+	FCameraDirectorTreeDebugBlock();
+
+	void Initialize(const FCameraEvaluationContextStack& ContextStack, FCameraDebugBlockBuilder& Builder);
+
+protected:
+
+	virtual void OnDebugDraw(const FCameraDebugBlockDrawParams& Params, FCameraDebugRenderer& Renderer) override;
+
+private:
+
+	struct FDirectorDebugInfo
+	{
+		FString CameraAssetName;
+		FCameraDebugBlock* DebugBlock = nullptr;
+	};
+	TArray<FDirectorDebugInfo> CameraDirectors;
+};
+
+}  // namespace UE::Cameras
+
+#endif  // UE_GAMEPLAY_CAMERAS_DEBUG
+
diff --git a/Engine/Plugins/Cameras/GameplayCameras/Source/GameplayCameras/Public/Debug/CameraNodeEvaluationResultDebugBlock.h b/Engine/Plugins/Cameras/GameplayCameras/Source/GameplayCameras/Public/Debug/CameraNodeEvaluationResultDebugBlock.h
new file mode 100644
--- /dev/null
+++ b/Engine/Plugins/Cameras/GameplayCameras/Source/GameplayCameras/Public/Debug/CameraNodeEvaluationResultDebugBlock.h
@@ -0,0 +1,57 @@
+// Copyright Epic Games, Inc. All Rights Reserved.
+
+#pragma once
+
+#include "Debug/CameraDebugBlock.h"
+
+#if UE_GAMEPLAY_CAMERAS_DEBUG
+
+namespace UE::Cameras
+{
+
+class FCameraPoseDebugBlock;
+class FVariableTableDebugBlock;
+struct FCameraDebugBlockBuilder;
+struct FCameraNodeEvaluationResult;
+
+/**
+ * A debug block that prints an evaluation result.
+ * 
+ * It will print a few basic bits of information, followed by the result camera pose (using 
+ * FCameraPoseDebugBlock) and the result variable table (using FVariableTableDebugBlock).
+ */
+class FCameraNodeEvaluationResultDebugBlock : public FCameraDebugBlock
+{
+	UE_DECLARE_CAMERA_DEBUG_BLOCK(FCameraNodeEvaluationResultDebugBlock)
+
+public:
+
+	FCameraNodeEvaluationResultDebugBlock();
+
+	/**
+	 * Initializes this block with the given result.
+	 *
+	 * @param InResult   The result to display
+	 * @param Builder    The builder to use to make the children debug blocks
+	 */
+	void Initialize(const FCameraNodeEvaluationResult& InResult, FCameraDebugBlockBuilder& Builder);
+
+	/** Gets the pose stats debug block. */
+	FCameraPoseDebugBlock* GetCameraPoseDebugBlock();
+	/** Gets the variable table debug block. */
+	FVariableTableDebugBlock* GetVariableTableDebugBlock();
+
+protected:
+
+	virtual void OnDebugDraw(const FCameraDebugBlockDrawParams& Params, FCameraDebugRenderer& Renderer) override;
+
+private:
+
+	bool bIsCameraCut;
+	bool bIsValid;
+};
+
+}  // namespace UE::Cameras
+
+#endif  // UE_GAMEPLAY_CAMERAS_DEBUG
+
diff --git a/Engine/Plugins/Cameras/GameplayCameras/Source/GameplayCameras/Public/Debug/CameraNodeEvaluatorDebugBlock.h b/Engine/Plugins/Cameras/GameplayCameras/Source/GameplayCameras/Public/Debug/CameraNodeEvaluatorDebugBlock.h
new file mode 100644
--- /dev/null
+++ b/Engine/Plugins/Cameras/GameplayCameras/Source/GameplayCameras/Public/Debug/CameraNodeEvaluatorDebugBlock.h
@@ -0,0 +1,45 @@
+// Copyright Epic Games, Inc. All Rights Reserved.
+
+#pragma once
+
+#include "Containers/Array.h"
+#include "CoreTypes.h"
+#include "Debug/CameraDebugBlock.h"
+#include "GameplayCameras.h"
+
+#if UE_GAMEPLAY_CAMERAS_DEBUG
+
+class UCameraNode;
+
+namespace UE::Cameras
+{
+
+/**
+ * Basic debug block for a camera node evaluator.
+ */
+class FCameraNodeEvaluatorDebugBlock : public FCameraDebugBlock
+{
+	UE_DECLARE_CAMERA_DEBUG_BLOCK(FCameraNodeEvaluatorDebugBlock)
+
+public:
+
+	/** Constructs a new node evaluator debug block. */
+	FCameraNodeEvaluatorDebugBlock();
+	/** Constructs a new node evaluator debug block. */
+	FCameraNodeEvaluatorDebugBlock(TObjectPtr<const UCameraNode> InCameraNode);
+
+protected:
+
+	// FCameraDebugBlock interface.
+	virtual void OnDebugDraw(const FCameraDebugBlockDrawParams& Params, FCameraDebugRenderer& Renderer) override;
+	virtual void OnSerialize(FArchive& Ar) override;
+
+private:
+
+	FString NodeClassName;
+};
+
+}  // namespace UE::Cameras
+
+#endif  // UE_GAMEPLAY_CAMERAS_DEBUG
+
diff --git a/Engine/Plugins/Cameras/GameplayCameras/Source/GameplayCameras/Public/Debug/CameraPoseDebugBlock.h b/Engine/Plugins/Cameras/GameplayCameras/Source/GameplayCameras/Public/Debug/CameraPoseDebugBlock.h
new file mode 100644
--- /dev/null
+++ b/Engine/Plugins/Cameras/GameplayCameras/Source/GameplayCameras/Public/Debug/CameraPoseDebugBlock.h
@@ -0,0 +1,50 @@
+// Copyright Epic Games, Inc. All Rights Reserved.
+
+#pragma once
+
+#include "Core/CameraPose.h"
+#include "Debug/CameraDebugBlock.h"
+
+#if UE_GAMEPLAY_CAMERAS_DEBUG
+
+namespace UE::Cameras
+{
+
+/**
+ * A debug block that displays information about a camera pose.
+ */
+class FCameraPoseDebugBlock : public FCameraDebugBlock
+{
+	UE_DECLARE_CAMERA_DEBUG_BLOCK(FCameraPoseDebugBlock)
+
+public:
+
+	/** Creates a new camera pose debug block. */
+	FCameraPoseDebugBlock();
+	/** Creates a new camera pose debug block. */
+	FCameraPoseDebugBlock(const FCameraPose& InCameraPose);
+
+	/** 
+	 * Specifies the console variable to use to toggle between only showing camera pose
+	 * properties that were written to, or showing all camera pose properties.
+	 */
+	FCameraPoseDebugBlock& WithShowUnchangedCVar(const TCHAR* InShowUnchangedCVarName)
+	{
+		ShowUnchangedCVarName = InShowUnchangedCVarName;
+		return *this;
+	}
+
+protected:
+
+	virtual void OnDebugDraw(const FCameraDebugBlockDrawParams& Params, FCameraDebugRenderer& Renderer) override;
+
+private:
+
+	FCameraPose CameraPose;
+	FString ShowUnchangedCVarName;
+};
+
+}  // namespace UE::Cameras
+
+#endif  // UE_GAMEPLAY_CAMERAS_DEBUG
+
diff --git a/Engine/Plugins/Cameras/GameplayCameras/Source/GameplayCameras/Public/Debug/CategoryTitleDebugBlock.h b/Engine/Plugins/Cameras/GameplayCameras/Source/GameplayCameras/Public/Debug/CategoryTitleDebugBlock.h
new file mode 100644
--- /dev/null
+++ b/Engine/Plugins/Cameras/GameplayCameras/Source/GameplayCameras/Public/Debug/CategoryTitleDebugBlock.h
@@ -0,0 +1,44 @@
+// Copyright Epic Games, Inc. All Rights Reserved.
+
+#pragma once
+
+#include "Debug/CameraDebugBlock.h"
+
+#if UE_GAMEPLAY_CAMERAS_DEBUG
+
+namespace UE::Cameras
+{
+
+class FCameraNodeEvaluatorDebugBlock;
+
+/**
+ * A utility debug block that prints a title for a debug category, if that debug category is active.
+ * By default, if the debug category is inactive, it will skip any attachments and children blocks.
+ */
+class FCategoryTitleDebugBlock : public FCameraDebugBlock
+{
+	UE_DECLARE_CAMERA_DEBUG_BLOCK(FCategoryTitleDebugBlock)
+
+public:
+
+	FCategoryTitleDebugBlock();
+	FCategoryTitleDebugBlock(const FString& InCategory, const FString& InTitle);
+
+protected:
+
+	virtual void OnDebugDraw(const FCameraDebugBlockDrawParams& Params, FCameraDebugRenderer& Renderer) override;
+
+public:
+
+	FString Category;
+	FString Title;
+	FColor TitleColor;
+
+	bool bSkipAttachedBlocksIfInactive = true;
+	bool bSkipChildrenBlocksIfInactive = true;
+};
+
+}  // namespace UE::Cameras
+
+#endif  // UE_GAMEPLAY_CAMERAS_DEBUG
+
diff --git a/Engine/Plugins/Cameras/GameplayCameras/Source/GameplayCameras/Public/Debug/DebugTextRenderer.h b/Engine/Plugins/Cameras/GameplayCameras/Source/GameplayCameras/Public/Debug/DebugTextRenderer.h
new file mode 100644
--- /dev/null
+++ b/Engine/Plugins/Cameras/GameplayCameras/Source/GameplayCameras/Public/Debug/DebugTextRenderer.h
@@ -0,0 +1,142 @@
+// Copyright Epic Games, Inc. All Rights Reserved.
+
+#pragma once
+
+#include "CoreTypes.h"
+#include "Containers/StringView.h"
+#include "GameplayCameras.h"
+#include "Math/UnrealMath.h"
+#include "Misc/TVariant.h"
+
+class UCanvas;
+class UFont;
+
+#if UE_GAMEPLAY_CAMERAS_DEBUG
+
+namespace UE::Cameras
+{
+
+// Utility function to print something. Most types have a LexToString implementation, but not all.
+template<typename FieldType>
+FString ToDebugString(const FieldType& FieldValue)
+{
+	return LexToString(FieldValue);
+}
+template<typename T>
+FString ToDebugString(const UE::Math::TVector<T>& FieldValue)
+{
+	return FieldValue.ToString();
+}
+template<typename T>
+FString ToDebugString(const UE::Math::TVector2<T>& FieldValue)
+{
+	return FieldValue.ToString();
+}
+template<typename T>
+FString ToDebugString(const UE::Math::TVector4<T>& FieldValue)
+{
+	return FieldValue.ToString();
+}
+template<typename T>
+FString ToDebugString(const UE::Math::TRotator<T>& FieldValue)
+{
+	return FieldValue.ToString();
+}
+template<typename T>
+FString ToDebugString(const UE::Math::TTransform<T>& FieldValue)
+{
+	return FieldValue.ToString();
+}
+
+/** Command for drawing text on a canvas. */
+struct FDebugTextDrawCommand
+{
+	TStringView<TCHAR> TextView;
+
+	void Execute(UCanvas* Canvas, const UFont* Font, FVector2f& InOutDrawPosition) const;
+};
+
+/** Command for moving the drawing position to a new line. */
+struct FDebugTextNewLineCommand
+{
+	float LineSpacing = 0.f;
+	float LeftMargin = 0.f;
+
+	void Execute(FVector2f& InOutDrawPosition) const;
+};
+
+/** Command for setting the text color on a canvas. */
+struct FDebugTextSetColorCommand
+{
+	FColor DrawColor;
+
+	void Execute(UCanvas* Canvas) const;
+};
+
+/** A debug text drawing command, which can be of multiple types. */
+using FDebugTextCommand = TVariant<
+	FDebugTextDrawCommand, 
+	FDebugTextNewLineCommand,
+	FDebugTextSetColorCommand>;
+
+/** A command queue for drawing text on a canvas. */
+using FDebugTextCommandArray = TArray<FDebugTextCommand>;
+
+/**
+ * Rendering utility for colored debug text.
+ */
+class FDebugTextRenderer
+{
+public:
+
+	/** Space between the lines, defaults to max font character height. */
+	float LineSpacing;
+	/** The X coordinate for where text drawing starts, and where new lines start from. */
+	float LeftMargin;
+	/** Moves the next draw position to a new line at the end of the text render. */
+	bool bEndWithNewLine = false;
+
+	/** Creates a new debug text renderer. */
+	FDebugTextRenderer(UCanvas* InCanvas, const UFont* InFont);
+
+	/** Renders the given text to the canvas. */
+	void RenderText(float StartingDrawY, const TStringView<TCHAR> TextView);
+	void RenderText(FVector2f StartingDrawPosition, const TStringView<TCHAR> TextView);
+
+	/** Executes the given command queue. */
+	void ExecuteCommands(float StartingDrawY, FDebugTextCommandArray& Commands);
+	void ExecuteCommands(FVector2f StartingDrawPosition, FDebugTextCommandArray& Commands);
+
+	/** Gets the coordinate of where any new text would go, just after the last render. */
+	FVector2f GetEndDrawPosition() const { return NextDrawPosition; }
+
+	/** Gets the maximum horizontal extent of the rendered text. */
+	float GetRightMargin() const { return RightMargin; }
+
+public:
+
+	static float GetStringViewSize(const UFont* Font, TStringView<TCHAR> TextView);
+
+private:
+
+	void ParseText(const TStringView<TCHAR> TextView, FDebugTextCommandArray& OutCommands);
+	void ExecuteCommands(FDebugTextCommandArray& Commands);
+	void UpdateRightMargin();
+
+	void AddDrawCommand(const TCHAR* RangeStart, const TCHAR* RangeEnd, bool bNewLine, FDebugTextCommandArray& OutCommands);
+	void AddTokenCommand(const TCHAR* RangeStart, const TCHAR* RangeEnd, FDebugTextCommandArray& OutCommands);
+
+	FColor InterpretColor(const FString& ColorName);
+
+private:
+
+	UCanvas* Canvas;
+	const UFont* Font;
+	FVector2f NextDrawPosition;
+	float RightMargin = 0;
+};
+
+}  // namespace UE::Cameras
+
+#endif  // UE_GAMEPLAY_CAMERAS_DEBUG
+
diff --git a/Engine/Plugins/Cameras/GameplayCameras/Source/GameplayCameras/Public/Debug/RootCameraDebugBlock.h b/Engine/Plugins/Cameras/GameplayCameras/Source/GameplayCameras/Public/Debug/RootCameraDebugBlock.h
new file mode 100644
--- /dev/null
+++ b/Engine/Plugins/Cameras/GameplayCameras/Source/GameplayCameras/Public/Debug/RootCameraDebugBlock.h
@@ -0,0 +1,39 @@
+// Copyright Epic Games, Inc. All Rights Reserved.
+
+#pragma once
+
+#include "Debug/CameraDebugBlock.h"
+#include "Debug/CameraDebugBlockStorage.h"
+
+#if UE_GAMEPLAY_CAMERAS_DEBUG
+
+namespace UE::Cameras
+{
+
+class FCameraSystemEvaluator;
+struct FCameraDebugBlockBuildParams;
+struct FCameraDebugBlockBuilder;
+
+GAMEPLAYCAMERAS_API extern bool GGameplayCamerasDebugEnable;
+GAMEPLAYCAMERAS_API extern FString GGameplayCamerasDebugCategories;
+
+/**
+ * The root debug block for the camera system.
+ */
+class FRootCameraDebugBlock : public FCameraDebugBlock
+{
+	UE_DECLARE_CAMERA_DEBUG_BLOCK(FRootCameraDebugBlock)
+
+public:
+
+	/** Build all debug blocks for the last evaluation frame. */
+	void BuildDebugBlocks(const FCameraSystemEvaluator& CameraSystem, const FCameraDebugBlockBuildParams& Params, FCameraDebugBlockBuilder& Builder);
+
+	/** Initiate the debug drawing. */
+	void RootDebugDraw(FCameraDebugRenderer& Renderer);
+};
+
+}  // namespace UE::Cameras
+
+#endif  // UE_GAMEPLAY_CAMERAS_DEBUG
+
diff --git a/Engine/Plugins/Cameras/GameplayCameras/Source/GameplayCameras/Public/Debug/VariableTableDebugBlock.h b/Engine/Plugins/Cameras/GameplayCameras/Source/GameplayCameras/Public/Debug/VariableTableDebugBlock.h
new file mode 100644
--- /dev/null
+++ b/Engine/Plugins/Cameras/GameplayCameras/Source/GameplayCameras/Public/Debug/VariableTableDebugBlock.h
@@ -0,0 +1,59 @@
+// Copyright Epic Games, Inc. All Rights Reserved.
+
+#pragma once
+
+#include "Debug/CameraDebugBlock.h"
+
+#if UE_GAMEPLAY_CAMERAS_DEBUG
+
+namespace UE::Cameras
+{
+
+class FCameraVariableTable;
+
+/**
+ * A debug block that prints the contents of a variable table.
+ */
+class FVariableTableDebugBlock : public FCameraDebugBlock
+{
+	UE_DECLARE_CAMERA_DEBUG_BLOCK(FVariableTableDebugBlock)
+
+public:
+
+	/** Creates a new variable table debug block. */
+	FVariableTableDebugBlock();
+	/** Creates a new variable table debug block. */
+	FVariableTableDebugBlock(const FCameraVariableTable& InVariableTable);
+
+	/** Specifies the console variable to use to toggle the printing of variable IDs. */
+	FVariableTableDebugBlock& WithShowVariableIdsCVar(const TCHAR* InShowVariableIdsCVarName)
+	{
+		ShowVariableIdsCVarName = InShowVariableIdsCVarName;
+		return *this;
+	}
+
+protected:
+
+	virtual void OnDebugDraw(const FCameraDebugBlockDrawParams& Params, FCameraDebugRenderer& Renderer) override;
+
+	void Initialize(const FCameraVariableTable& InVariableTable);
+
+private:
+
+	struct FEntryDebugInfo
+	{
+		uint32 Id;
+		FString Name;
+		FString Value;
+		bool bWritten : 1;
+		bool bWrittenThisFrame : 1;
+	};
+	TArray<FEntryDebugInfo> Entries;
+
+	FString ShowVariableIdsCVarName;
+};
+
+}  // namespace UE::Cameras
+
+#endif  // UE_GAMEPLAY_CAMERAS_DEBUG
+
diff --git a/Engine/Plugins/Cameras/GameplayCameras/Source/GameplayCameras/Public/Debug/ViewfinderDebugBlock.h b/Engine/Plugins/Cameras/GameplayCameras/Source/GameplayCameras/Public/Debug/ViewfinderDebugBlock.h
new file mode 100644
--- /dev/null
+++ b/Engine/Plugins/Cameras/GameplayCameras/Source/GameplayCameras/Public/Debug/ViewfinderDebugBlock.h
@@ -0,0 +1,28 @@
+// Copyright Epic Games, Inc. All Rights Reserved.
+
+#pragma once
+
+#include "Debug/CameraDebugBlock.h"
+
+#if UE_GAMEPLAY_CAMERAS_DEBUG
+
+namespace UE::Cameras
+{
+
+class FViewfinderDebugBlock : public FCameraDebugBlock
+{
+	UE_DECLARE_CAMERA_DEBUG_BLOCK(FViewfinderDebugBlock)
+
+public:
+
+	FViewfinderDebugBlock();
+
+protected:
+
+	virtual void OnDebugDraw(const FCameraDebugBlockDrawParams& Params, FCameraDebugRenderer& Renderer) override;
+};
+
+}  // namespace UE::Cameras
+
+#endif  // UE_GAMEPLAY_CAMERAS_DEBUG
+
diff --git a/Engine/Plugins/Cameras/GameplayCameras/Source/GameplayCameras/Public/GameplayCameras.h b/Engine/Plugins/Cameras/GameplayCameras/Source/GameplayCameras/Public/GameplayCameras.h
--- a/Engine/Plugins/Cameras/GameplayCameras/Source/GameplayCameras/Public/GameplayCameras.h
+++ b/Engine/Plugins/Cameras/GameplayCameras/Source/GameplayCameras/Public/GameplayCameras.h
@@ -3,11 +3,12 @@
 #pragma once
 
 #include "CoreTypes.h"
+#include "EngineDefines.h"
 #include "HAL/Platform.h"
 #include "Stats/Stats.h"
 
 #ifndef UE_GAMEPLAY_CAMERAS_DEBUG
-	#if (!UE_BUILD_SHIPPING || WITH_UNREAL_DEVELOPER_TOOLS || WITH_UNREAL_TARGET_DEVELOPER_TOOLS)
+	#if UE_ENABLE_DEBUG_DRAWING
 		#define UE_GAMEPLAY_CAMERAS_DEBUG 1
 	#else
 		#define UE_GAMEPLAY_CAMERAS_DEBUG 0
diff --git a/Engine/Plugins/Cameras/GameplayCameras/Source/GameplayCameras/Public/Nodes/Blends/SimpleBlendCameraNode.h b/Engine/Plugins/Cameras/GameplayCameras/Source/GameplayCameras/Public/Nodes/Blends/SimpleBlendCameraNode.h
--- a/Engine/Plugins/Cameras/GameplayCameras/Source/GameplayCameras/Public/Nodes/Blends/SimpleBlendCameraNode.h
+++ b/Engine/Plugins/Cameras/GameplayCameras/Source/GameplayCameras/Public/Nodes/Blends/SimpleBlendCameraNode.h
@@ -59,6 +59,10 @@ protected:
 
 	void SetBlendFinished() { bIsBlendFinished = true; }
 
+#if UE_GAMEPLAY_CAMERAS_DEBUG
+	virtual void OnBuildDebugBlocks(const FCameraDebugBlockBuildParams& Params, FCameraDebugBlockBuilder& Builder) override;
+#endif  // UE_GAMEPLAY_CAMERAS_DEBUG
+
 private:
 
 	float BlendFactor = 0.f;
diff --git a/Engine/Plugins/Cameras/GameplayCameras/Source/GameplayCameras/Public/Nodes/Common/DampenPositionCameraNode.h b/Engine/Plugins/Cameras/GameplayCameras/Source/GameplayCameras/Public/Nodes/Common/DampenPositionCameraNode.h
--- a/Engine/Plugins/Cameras/GameplayCameras/Source/GameplayCameras/Public/Nodes/Common/DampenPositionCameraNode.h
+++ b/Engine/Plugins/Cameras/GameplayCameras/Source/GameplayCameras/Public/Nodes/Common/DampenPositionCameraNode.h
@@ -15,10 +15,6 @@ class UDampenPositionCameraNode : public UCameraNode
 {
 	GENERATED_BODY()
 
-public:	
-
-	UDampenPositionCameraNode(const FObjectInitializer& ObjectInit);
-
 protected:
 
 	// UCameraNode interface.
