Cameras: the camera system debugger UI is here!

This adds a simple little UI that manipulates all the various cvars that drive the camera system's debug drawing.

All icons are temporary.

diff --git a/Engine/Plugins/Cameras/GameplayCameras/Source/GameplayCamerasEditor/GameplayCamerasEditor.Build.cs b/Engine/Plugins/Cameras/GameplayCameras/Source/GameplayCamerasEditor/GameplayCamerasEditor.Build.cs
--- a/Engine/Plugins/Cameras/GameplayCameras/Source/GameplayCamerasEditor/GameplayCamerasEditor.Build.cs
+++ b/Engine/Plugins/Cameras/GameplayCameras/Source/GameplayCamerasEditor/GameplayCamerasEditor.Build.cs
@@ -17,30 +17,33 @@ public class GameplayCamerasEditor : ModuleRules
 		);
 
 		PrivateDependencyModuleNames.AddRange(
-			new string[] {
+			new string[]
+			{
 				"AssetDefinition",
 				"AssetRegistry",
-				"GameplayCameras",
 				"BlueprintGraph",
-				"EditorSubsystem",
 				"Core",
 				"CoreUObject",
 				"EditorFramework",
-				"InteractiveToolsFramework",
+				"EditorSubsystem",
 				"Engine",
+				"GameplayCameras",
 				"InputCore",
+				"InteractiveToolsFramework",
 				"Kismet",
+				"Projects",
 				"Slate",
 				"SlateCore",
+				"TimeManagement",
 				"ToolMenus",
 				"UnrealEd",
-				"TimeManagement"
 			}
 		);
 
 		var DynamicModuleNames = new string[] {
 			"LevelEditor",
 			"PropertyEditor",
+			"WorkspaceMenuStructure",
 		};
 
 		foreach (var Name in DynamicModuleNames)
diff --git a/Engine/Plugins/Cameras/GameplayCameras/Source/GameplayCamerasEditor/Private/Commands/GameplayCamerasDebuggerCommand.cpp b/Engine/Plugins/Cameras/GameplayCameras/Source/GameplayCamerasEditor/Private/Commands/GameplayCamerasDebuggerCommand.cpp
new file mode 100644
--- /dev/null
+++ b/Engine/Plugins/Cameras/GameplayCameras/Source/GameplayCamerasEditor/Private/Commands/GameplayCamerasDebuggerCommand.cpp
@@ -0,0 +1,37 @@
+// Copyright Epic Games, Inc. All Rights Reserved.
+
+#include "Commands/GameplayCamerasDebuggerCommands.h"
+
+#include "Framework/Commands/Commands.h"
+#include "Framework/Commands/InputChord.h"
+#include "Framework/Commands/UICommandInfo.h"
+#include "Styling/AppStyle.h"
+
+#define LOCTEXT_NAMESPACE "GameplayCamerasDebuggerCommands"
+
+namespace UE::Cameras
+{
+
+FGameplayCamerasDebuggerCommands::FGameplayCamerasDebuggerCommands()
+	: TCommands<FGameplayCamerasDebuggerCommands>(
+			"GameplayCameras_Debugger",
+			NSLOCTEXT("Contexts", "GameplayCameras_Debugger", "Gameplay Cameras Debugger"),
+			NAME_None,
+			FAppStyle::GetAppStyleSetName()
+		)
+{
+}
+
+void FGameplayCamerasDebuggerCommands::RegisterCommands()
+{
+	UI_COMMAND(
+			EnableDebugInfo, "Enable Debug Info", 
+			"Shows Gameplay Cameras debugging information in the game",
+			EUserInterfaceActionType::ToggleButton,
+			FInputChord());
+}
+
+}  // namespace UE::Cameras
+
+#undef LOCTEXT_NAMESPACE
+
diff --git a/Engine/Plugins/Cameras/GameplayCameras/Source/GameplayCamerasEditor/Private/Debugger/GameplayCamerasDebugger.cpp b/Engine/Plugins/Cameras/GameplayCameras/Source/GameplayCamerasEditor/Private/Debugger/GameplayCamerasDebugger.cpp
new file mode 100644
--- /dev/null
+++ b/Engine/Plugins/Cameras/GameplayCameras/Source/GameplayCamerasEditor/Private/Debugger/GameplayCamerasDebugger.cpp
@@ -0,0 +1,9 @@
+// Copyright Epic Games, Inc. All Rights Reserved.
+
+#include "Debugger/GameplayCamerasDebugger.h"
+
+namespace UE::Cameras
+{
+
+}  // namespace UE::Cameras
+
diff --git a/Engine/Plugins/Cameras/GameplayCameras/Source/GameplayCamerasEditor/Private/Debugger/GameplayCamerasDebugger.h b/Engine/Plugins/Cameras/GameplayCameras/Source/GameplayCamerasEditor/Private/Debugger/GameplayCamerasDebugger.h
new file mode 100644
--- /dev/null
+++ b/Engine/Plugins/Cameras/GameplayCameras/Source/GameplayCamerasEditor/Private/Debugger/GameplayCamerasDebugger.h
@@ -0,0 +1,15 @@
+// Copyright Epic Games, Inc. All Rights Reserved.
+
+#pragma once
+
+#include "CoreTypes.h"
+
+namespace UE::Cameras
+{
+
+class FGameplayCamerasDebugger
+{
+};
+
+}  // namespace UE::Cameras
+
diff --git a/Engine/Plugins/Cameras/GameplayCameras/Source/GameplayCamerasEditor/Private/Debugger/SBlendStacksDebugPanel.cpp b/Engine/Plugins/Cameras/GameplayCameras/Source/GameplayCamerasEditor/Private/Debugger/SBlendStacksDebugPanel.cpp
new file mode 100644
--- /dev/null
+++ b/Engine/Plugins/Cameras/GameplayCameras/Source/GameplayCamerasEditor/Private/Debugger/SBlendStacksDebugPanel.cpp
@@ -0,0 +1,39 @@
+// Copyright Epic Games, Inc. All Rights Reserved.
+
+#include "Debugger/SBlendStacksDebugPanel.h"
+
+#include "Debugger/SDebugWidgetUtils.h"
+#include "Widgets/Input/SCheckBox.h"
+#include "Widgets/SBoxPanel.h"
+
+#define LOCTEXT_NAMESPACE "SBlendStacksDebugPanel"
+
+namespace UE::Cameras
+{
+
+void SBlendStacksDebugPanel::Construct(const FArguments& InArgs)
+{
+	ChildSlot
+	[
+		SNew(SVerticalBox)
+		+ SVerticalBox::Slot()
+			.AutoHeight()
+			[
+				SDebugWidgetUtils::CreateConsoleVariableCheckBox(
+						LOCTEXT("ShowUnchanged", "Show unchanged properties"),
+						TEXT("GameplayCameras.Debug.BlendStack.ShowUnchanged"))
+			]
+		+ SVerticalBox::Slot()
+			.AutoHeight()
+			[
+				SDebugWidgetUtils::CreateConsoleVariableCheckBox(
+						LOCTEXT("ShowVariableIds", "Show variable IDs"),
+						TEXT("GameplayCameras.Debug.BlendStack.ShowVariableIds"))
+			]
+	];
+}
+
+}  // namespace UE::Cameras
+
+#undef LOCTEXT_NAMESPACE
+
diff --git a/Engine/Plugins/Cameras/GameplayCameras/Source/GameplayCamerasEditor/Private/Debugger/SBlendStacksDebugPanel.h b/Engine/Plugins/Cameras/GameplayCameras/Source/GameplayCamerasEditor/Private/Debugger/SBlendStacksDebugPanel.h
new file mode 100644
--- /dev/null
+++ b/Engine/Plugins/Cameras/GameplayCameras/Source/GameplayCamerasEditor/Private/Debugger/SBlendStacksDebugPanel.h
@@ -0,0 +1,24 @@
+// Copyright Epic Games, Inc. All Rights Reserved.
+
+#pragma once
+
+#include "CoreTypes.h"
+#include "Widgets/SCompoundWidget.h"
+
+class FTabManager;
+
+namespace UE::Cameras
+{
+
+class SBlendStacksDebugPanel : public SCompoundWidget
+{
+public:
+
+	SLATE_BEGIN_ARGS(SBlendStacksDebugPanel) {}
+	SLATE_END_ARGS();
+
+	void Construct(const FArguments& InArgs);
+};
+
+}  // namespace UE::Cameras
+
diff --git a/Engine/Plugins/Cameras/GameplayCameras/Source/GameplayCamerasEditor/Private/Debugger/SCameraDirectorTreeDebugPanel.cpp b/Engine/Plugins/Cameras/GameplayCameras/Source/GameplayCamerasEditor/Private/Debugger/SCameraDirectorTreeDebugPanel.cpp
new file mode 100644
--- /dev/null
+++ b/Engine/Plugins/Cameras/GameplayCameras/Source/GameplayCamerasEditor/Private/Debugger/SCameraDirectorTreeDebugPanel.cpp
@@ -0,0 +1,32 @@
+// Copyright Epic Games, Inc. All Rights Reserved.
+
+#include "Debugger/SCameraDirectorTreeDebugPanel.h"
+
+#include "Debugger/SDebugWidgetUtils.h"
+#include "Widgets/Input/SCheckBox.h"
+#include "Widgets/SBoxPanel.h"
+
+#define LOCTEXT_NAMESPACE "SCameraDirectorTreeDebugPanel"
+
+namespace UE::Cameras
+{
+
+void SCameraDirectorTreeDebugPanel::Construct(const FArguments& InArgs)
+{
+	ChildSlot
+	[
+		SNew(SVerticalBox)
+		+ SVerticalBox::Slot()
+			.AutoHeight()
+			[
+				SDebugWidgetUtils::CreateConsoleVariableCheckBox(
+						LOCTEXT("ShowUnchanged", "Show unchanged properties"),
+						TEXT("GameplayCameras.Debug.ContextInitialResult.ShowUnchanged"))
+			]
+	];
+}
+
+}  // namespace UE::Cameras
+
+#undef LOCTEXT_NAMESPACE
+
diff --git a/Engine/Plugins/Cameras/GameplayCameras/Source/GameplayCamerasEditor/Private/Debugger/SCameraDirectorTreeDebugPanel.h b/Engine/Plugins/Cameras/GameplayCameras/Source/GameplayCamerasEditor/Private/Debugger/SCameraDirectorTreeDebugPanel.h
new file mode 100644
--- /dev/null
+++ b/Engine/Plugins/Cameras/GameplayCameras/Source/GameplayCamerasEditor/Private/Debugger/SCameraDirectorTreeDebugPanel.h
@@ -0,0 +1,24 @@
+// Copyright Epic Games, Inc. All Rights Reserved.
+
+#pragma once
+
+#include "CoreTypes.h"
+#include "Widgets/SCompoundWidget.h"
+
+class FTabManager;
+
+namespace UE::Cameras
+{
+
+class SCameraDirectorTreeDebugPanel : public SCompoundWidget
+{
+public:
+
+	SLATE_BEGIN_ARGS(SCameraDirectorTreeDebugPanel) {}
+	SLATE_END_ARGS();
+
+	void Construct(const FArguments& InArgs);
+};
+
+}  // namespace UE::Cameras
+
diff --git a/Engine/Plugins/Cameras/GameplayCameras/Source/GameplayCamerasEditor/Private/Debugger/SCameraNodeTreeDebugPanel.cpp b/Engine/Plugins/Cameras/GameplayCameras/Source/GameplayCamerasEditor/Private/Debugger/SCameraNodeTreeDebugPanel.cpp
new file mode 100644
--- /dev/null
+++ b/Engine/Plugins/Cameras/GameplayCameras/Source/GameplayCamerasEditor/Private/Debugger/SCameraNodeTreeDebugPanel.cpp
@@ -0,0 +1,38 @@
+// Copyright Epic Games, Inc. All Rights Reserved.
+
+#include "Debugger/SCameraNodeTreeDebugPanel.h"
+
+#include "Debugger/SDebugWidgetUtils.h"
+#include "Widgets/SBoxPanel.h"
+#include "Widgets/Input/SEditableTextBox.h"
+#include "Widgets/Text/STextBlock.h"
+
+#define LOCTEXT_NAMESPACE "SCameraNodeTreeDebugPanel"
+
+namespace UE::Cameras
+{
+
+void SCameraNodeTreeDebugPanel::Construct(const FArguments& InArgs)
+{
+	ChildSlot
+	[
+		SNew(SVerticalBox)
+		+ SVerticalBox::Slot()
+			.AutoHeight()
+			[
+				SNew(STextBlock)
+					.Text(LOCTEXT("FilterNodeNames", "Filter node names:"))
+			]
+		+ SVerticalBox::Slot()
+			.AutoHeight()
+			[
+				SDebugWidgetUtils::CreateConsoleVariableTextBox(
+						TEXT("GameplayCameras.Debug.NodeTree.Filter"))
+			]
+	];
+}
+
+}  // namespace UE::Cameras
+
+#undef LOCTEXT_NAMESPACE
+
diff --git a/Engine/Plugins/Cameras/GameplayCameras/Source/GameplayCamerasEditor/Private/Debugger/SCameraNodeTreeDebugPanel.h b/Engine/Plugins/Cameras/GameplayCameras/Source/GameplayCamerasEditor/Private/Debugger/SCameraNodeTreeDebugPanel.h
new file mode 100644
--- /dev/null
+++ b/Engine/Plugins/Cameras/GameplayCameras/Source/GameplayCamerasEditor/Private/Debugger/SCameraNodeTreeDebugPanel.h
@@ -0,0 +1,24 @@
+// Copyright Epic Games, Inc. All Rights Reserved.
+
+#pragma once
+
+#include "CoreTypes.h"
+#include "Widgets/SCompoundWidget.h"
+
+class FTabManager;
+
+namespace UE::Cameras
+{
+
+class SCameraNodeTreeDebugPanel : public SCompoundWidget
+{
+public:
+
+	SLATE_BEGIN_ARGS(SCameraNodeTreeDebugPanel) {}
+	SLATE_END_ARGS();
+
+	void Construct(const FArguments& InArgs);
+};
+
+}  // namespace UE::Cameras
+
diff --git a/Engine/Plugins/Cameras/GameplayCameras/Source/GameplayCamerasEditor/Private/Debugger/SCameraPoseStatsDebugPanel.cpp b/Engine/Plugins/Cameras/GameplayCameras/Source/GameplayCamerasEditor/Private/Debugger/SCameraPoseStatsDebugPanel.cpp
new file mode 100644
--- /dev/null
+++ b/Engine/Plugins/Cameras/GameplayCameras/Source/GameplayCamerasEditor/Private/Debugger/SCameraPoseStatsDebugPanel.cpp
@@ -0,0 +1,39 @@
+// Copyright Epic Games, Inc. All Rights Reserved.
+
+#include "Debugger/SCameraPoseStatsDebugPanel.h"
+
+#include "Debugger/SDebugWidgetUtils.h"
+#include "Widgets/SBoxPanel.h"
+#include "Widgets/Input/SCheckBox.h"
+
+#define LOCTEXT_NAMESPACE "SCameraPoseStatsDebugPanel"
+
+namespace UE::Cameras
+{
+
+void SCameraPoseStatsDebugPanel::Construct(const FArguments& InArgs)
+{
+	ChildSlot
+	[
+		SNew(SVerticalBox)
+		+ SVerticalBox::Slot()
+			.AutoHeight()
+			[
+				SDebugWidgetUtils::CreateConsoleVariableCheckBox(
+						LOCTEXT("ShowUnchanged", "Show unchanged properties"),
+						TEXT("GameplayCameras.Debug.PoseStats.ShowUnchanged"))
+			]
+		+ SVerticalBox::Slot()
+			.AutoHeight()
+			[
+				SDebugWidgetUtils::CreateConsoleVariableCheckBox(
+						LOCTEXT("ShowVariableIds", "Show variable IDs"),
+						TEXT("GameplayCameras.Debug.PoseStats.ShowVariableIds"))
+			]
+	];
+}
+
+}  // namespace UE::Cameras
+
+#undef LOCTEXT_NAMESPACE
+
diff --git a/Engine/Plugins/Cameras/GameplayCameras/Source/GameplayCamerasEditor/Private/Debugger/SCameraPoseStatsDebugPanel.h b/Engine/Plugins/Cameras/GameplayCameras/Source/GameplayCamerasEditor/Private/Debugger/SCameraPoseStatsDebugPanel.h
new file mode 100644
--- /dev/null
+++ b/Engine/Plugins/Cameras/GameplayCameras/Source/GameplayCamerasEditor/Private/Debugger/SCameraPoseStatsDebugPanel.h
@@ -0,0 +1,24 @@
+// Copyright Epic Games, Inc. All Rights Reserved.
+
+#pragma once
+
+#include "CoreTypes.h"
+#include "Widgets/SCompoundWidget.h"
+
+class FTabManager;
+
+namespace UE::Cameras
+{
+
+class SCameraPoseStatsDebugPanel : public SCompoundWidget
+{
+public:
+
+	SLATE_BEGIN_ARGS(SCameraPoseStatsDebugPanel) {}
+	SLATE_END_ARGS();
+
+	void Construct(const FArguments& InArgs);
+};
+
+}  // namespace UE::Cameras
+
diff --git a/Engine/Plugins/Cameras/GameplayCameras/Source/GameplayCamerasEditor/Private/Debugger/SDebugCategoryButton.cpp b/Engine/Plugins/Cameras/GameplayCameras/Source/GameplayCamerasEditor/Private/Debugger/SDebugCategoryButton.cpp
new file mode 100644
--- /dev/null
+++ b/Engine/Plugins/Cameras/GameplayCameras/Source/GameplayCamerasEditor/Private/Debugger/SDebugCategoryButton.cpp
@@ -0,0 +1,79 @@
+// Copyright Epic Games, Inc. All Rights Reserved.
+
+#include "Debugger/SDebugCategoryButton.h"
+
+#include "Styling/AppStyle.h"
+#include "Styling/SlateColor.h"
+#include "Templates/SharedPointer.h"
+#include "Widgets/Images/SImage.h"
+#include "Widgets/Input/SCheckBox.h"
+#include "Widgets/SBoxPanel.h"
+#include "Widgets/Text/STextBlock.h"
+
+namespace UE::Cameras
+{
+
+void SDebugCategoryButton::Construct(const FArguments& InArgs)
+{
+	DebugCategoryName = InArgs._DebugCategoryName.Get();
+	IsDebugCategoryActive = InArgs._IsDebugCategoryActive;
+	RequestDebugCategoryChange = InArgs._RequestDebugCategoryChange;
+
+	InactiveModeBorderImage = FAppStyle::GetBrush("ModeSelector.ToggleButton.Normal");
+	ActiveModeBorderImage = FAppStyle::GetBrush("ModeSelector.ToggleButton.Pressed");
+	HoverBorderImage = FAppStyle::GetBrush("ModeSelector.ToggleButton.Hovered");
+
+	const FMargin IconPadding(4.0f, 0.0f, 4.0f, 0.0f);
+	const FMargin TextPadding(0.0f, 0.0f, 0.0f, 0.0f);
+
+	TSharedRef<SHorizontalBox> ButtonContents = SNew(SHorizontalBox);
+
+	if (InArgs._IconImage.IsSet())
+	{
+		ButtonContents->AddSlot()
+			.AutoWidth()
+			.VAlign(VAlign_Center)
+			.Padding(IconPadding)
+			[
+				SNew(SImage)
+					.Image(InArgs._IconImage)
+					.ColorAndOpacity(FSlateColor::UseForeground())
+			];
+	}
+
+	ButtonContents->AddSlot()
+		.AutoWidth()
+		.VAlign(VAlign_Center)
+		.Padding(TextPadding)
+		[
+			SNew(STextBlock)
+				.Text(InArgs._DisplayText)
+		];
+
+	ChildSlot
+	.VAlign(VAlign_Fill)
+	[
+		SNew(SCheckBox)
+			.ToolTipText(InArgs._ToolTipText)
+			.Style(FAppStyle::Get(), "ToggleButtonCheckbox")
+			.IsChecked(this, &SDebugCategoryButton::GetDebugCategoryCheckState)
+			.OnCheckStateChanged(this, &SDebugCategoryButton::OnDebugCategoryCheckStateChanged)
+			[
+				ButtonContents
+			]
+	];
+}
+
+ECheckBoxState SDebugCategoryButton::GetDebugCategoryCheckState() const
+{
+	return IsDebugCategoryActive.Get() ?
+		ECheckBoxState::Checked : ECheckBoxState::Unchecked;
+}
+
+void SDebugCategoryButton::OnDebugCategoryCheckStateChanged(ECheckBoxState CheckBoxState)
+{
+	RequestDebugCategoryChange.ExecuteIfBound(DebugCategoryName);
+}
+
+}  // namespace UE::Cameras
+
diff --git a/Engine/Plugins/Cameras/GameplayCameras/Source/GameplayCamerasEditor/Private/Debugger/SDebugCategoryButton.h b/Engine/Plugins/Cameras/GameplayCameras/Source/GameplayCamerasEditor/Private/Debugger/SDebugCategoryButton.h
new file mode 100644
--- /dev/null
+++ b/Engine/Plugins/Cameras/GameplayCameras/Source/GameplayCamerasEditor/Private/Debugger/SDebugCategoryButton.h
@@ -0,0 +1,47 @@
+// Copyright Epic Games, Inc. All Rights Reserved.
+
+#pragma once
+
+#include "CoreMinimal.h"
+#include "Styling/SlateTypes.h"
+#include "Widgets/SCompoundWidget.h"
+
+struct FSlateBrush;
+
+namespace UE::Cameras
+{
+
+DECLARE_DELEGATE_OneParam(FOnDebugCategoryChangeRequested, const FString&);
+
+class SDebugCategoryButton : public SCompoundWidget
+{
+public:
+	SLATE_BEGIN_ARGS(SDebugCategoryButton) {}
+		SLATE_ATTRIBUTE(const FSlateBrush*, IconImage)
+		SLATE_ATTRIBUTE(FText, DisplayText)
+		SLATE_ATTRIBUTE(FText, ToolTipText)
+		SLATE_ATTRIBUTE(FString, DebugCategoryName)
+		SLATE_ATTRIBUTE(bool, IsDebugCategoryActive)
+		SLATE_EVENT(FOnDebugCategoryChangeRequested, RequestDebugCategoryChange)
+	SLATE_END_ARGS()
+
+	void Construct(const FArguments& InArgs);
+
+private:
+
+	ECheckBoxState GetDebugCategoryCheckState() const;
+	void OnDebugCategoryCheckStateChanged(ECheckBoxState CheckBoxState);
+
+private:
+
+	FString DebugCategoryName;
+	TAttribute<bool> IsDebugCategoryActive;
+	FOnDebugCategoryChangeRequested RequestDebugCategoryChange;
+
+	const FSlateBrush* ActiveModeBorderImage;
+	const FSlateBrush* InactiveModeBorderImage;
+	const FSlateBrush* HoverBorderImage;
+};
+
+}  // namespace UE::Cameras
+
diff --git a/Engine/Plugins/Cameras/GameplayCameras/Source/GameplayCamerasEditor/Private/Debugger/SDebugWidgetUtils.cpp b/Engine/Plugins/Cameras/GameplayCameras/Source/GameplayCamerasEditor/Private/Debugger/SDebugWidgetUtils.cpp
new file mode 100644
--- /dev/null
+++ b/Engine/Plugins/Cameras/GameplayCameras/Source/GameplayCamerasEditor/Private/Debugger/SDebugWidgetUtils.cpp
@@ -0,0 +1,78 @@
+// Copyright Epic Games, Inc. All Rights Reserved.
+
+#include "Debugger/SDebugWidgetUtils.h"
+
+#include "HAL/IConsoleManager.h"
+#include "Widgets/Input/SCheckBox.h"
+#include "Widgets/Input/SEditableTextBox.h"
+#include "Widgets/Text/STextBlock.h"
+
+namespace UE::Cameras
+{
+
+TMap<FString, IConsoleVariable*> SDebugWidgetUtils::CachedConsoleVariables;
+
+TSharedRef<SCheckBox> SDebugWidgetUtils::CreateConsoleVariableCheckBox(const FText& Text, const FString& ConsoleVariableName)
+{
+	return SNew(SCheckBox)
+		.IsChecked_Lambda([ConsoleVariableName]() 
+				{
+					if (IConsoleVariable* ConsoleVariable = SDebugWidgetUtils::GetConsoleVariable(ConsoleVariableName))
+					{
+						return ConsoleVariable->GetBool() ?
+							ECheckBoxState::Checked : ECheckBoxState::Unchecked;
+					}
+					return ECheckBoxState::Undetermined;
+				})
+		.OnCheckStateChanged_Lambda([ConsoleVariableName](ECheckBoxState NewState)
+				{
+					if (IConsoleVariable* ConsoleVariable = SDebugWidgetUtils::GetConsoleVariable(ConsoleVariableName))
+					{
+						ConsoleVariable->Set(NewState == ECheckBoxState::Checked);
+					}
+				})
+		[
+			SNew(STextBlock)
+				.Text(Text)
+		];
+}
+
+TSharedRef<SEditableTextBox> SDebugWidgetUtils::CreateConsoleVariableTextBox(const FString& ConsoleVariableName)
+{
+	return SNew(SEditableTextBox)
+		.Text_Lambda([ConsoleVariableName]()
+				{
+					if (IConsoleVariable* ConsoleVariable = SDebugWidgetUtils::GetConsoleVariable(ConsoleVariableName))
+					{
+						return FText::FromString(ConsoleVariable->GetString());
+					}
+					return FText();
+				})
+		.OnTextChanged_Lambda([ConsoleVariableName](const FText& NewText)
+				{
+					if (IConsoleVariable* ConsoleVariable = SDebugWidgetUtils::GetConsoleVariable(ConsoleVariableName))
+					{
+						ConsoleVariable->Set(*NewText.ToString());
+					}
+				});
+}
+
+IConsoleVariable* SDebugWidgetUtils::GetConsoleVariable(const FString& ConsoleVariableName)
+{
+	if (IConsoleVariable** ConsoleVariablePtr = CachedConsoleVariables.Find(ConsoleVariableName))
+	{
+		return (*ConsoleVariablePtr);
+	}
+
+	IConsoleVariable* ConsoleVariable = IConsoleManager::Get().FindConsoleVariable(*ConsoleVariableName);
+	if (ensureMsgf(ConsoleVariable, TEXT("No such console variable: %s"), *ConsoleVariableName))
+	{
+		CachedConsoleVariables.Add(ConsoleVariableName, ConsoleVariable);
+		return ConsoleVariable;
+	}
+
+	return nullptr;
+}
+
+}  // namespace UE::Cameras
+
diff --git a/Engine/Plugins/Cameras/GameplayCameras/Source/GameplayCamerasEditor/Private/Debugger/SDebugWidgetUtils.h b/Engine/Plugins/Cameras/GameplayCameras/Source/GameplayCamerasEditor/Private/Debugger/SDebugWidgetUtils.h
new file mode 100644
--- /dev/null
+++ b/Engine/Plugins/Cameras/GameplayCameras/Source/GameplayCamerasEditor/Private/Debugger/SDebugWidgetUtils.h
@@ -0,0 +1,33 @@
+// Copyright Epic Games, Inc. All Rights Reserved.
+
+#pragma once
+
+#include "Containers/Map.h"
+#include "CoreMinimal.h"
+#include "Templates/SharedPointerFwd.h"
+
+class IConsoleVariable;
+class SCheckBox;
+class SEditableTextBox;
+
+namespace UE::Cameras
+{
+
+class SDebugWidgetUtils
+{
+public:
+
+	static TSharedRef<SCheckBox> CreateConsoleVariableCheckBox(const FText& Text, const FString& ConsoleVariableName);
+	static TSharedRef<SEditableTextBox> CreateConsoleVariableTextBox(const FString& ConsoleVariableName);
+
+private:
+
+	static IConsoleVariable* GetConsoleVariable(const FString& ConsoleVariableName);
+
+private:
+
+	static TMap<FString, IConsoleVariable*> CachedConsoleVariables;
+};
+
+}  // namespace UE::Cameras
+
diff --git a/Engine/Plugins/Cameras/GameplayCameras/Source/GameplayCamerasEditor/Private/Debugger/SGameplayCamerasDebugger.cpp b/Engine/Plugins/Cameras/GameplayCameras/Source/GameplayCamerasEditor/Private/Debugger/SGameplayCamerasDebugger.cpp
new file mode 100644
--- /dev/null
+++ b/Engine/Plugins/Cameras/GameplayCameras/Source/GameplayCamerasEditor/Private/Debugger/SGameplayCamerasDebugger.cpp
@@ -0,0 +1,203 @@
+// Copyright Epic Games, Inc. All Rights Reserved.
+
+#include "Debugger/SGameplayCamerasDebugger.h"
+
+#include "Commands/GameplayCamerasDebuggerCommands.h"
+#include "Debug/RootCameraDebugBlock.h"
+#include "Debugger/SDebugCategoryButton.h"
+#include "Framework/Application/SlateApplication.h"
+#include "Framework/Docking/LayoutService.h"
+#include "Framework/Docking/TabManager.h"
+#include "Framework/MultiBox/MultiBoxBuilder.h"
+#include "IGameplayCamerasEditorModule.h"
+#include "Modules/ModuleManager.h"
+#include "String/ParseTokens.h"
+#include "Styles/GameplayCamerasEditorStyle.h"
+#include "Styling/SlateTypes.h"
+#include "ToolMenus.h"
+#include "Widgets/Docking/SDockTab.h"
+#include "Widgets/Layout/SBox.h"
+#include "Widgets/SNullWidget.h"
+#include "WorkspaceMenuStructure.h"
+#include "WorkspaceMenuStructureModule.h"
+
+#define LOCTEXT_NAMESPACE "GameplayCamerasDebugger"
+
+namespace UE::Cameras
+{
+
+const FName SGameplayCamerasDebugger::WindowName(TEXT("GameplayCamerasDebugger"));
+const FName SGameplayCamerasDebugger::ToolbarName(TEXT("GameplayCamerasDebugger.Toolbar"));
+
+void SGameplayCamerasDebugger::RegisterTabSpawners()
+{
+	FGlobalTabmanager::Get()->RegisterNomadTabSpawner(
+		SGameplayCamerasDebugger::WindowName,
+		FOnSpawnTab::CreateStatic(&SGameplayCamerasDebugger::SpawnGameplayCamerasDebugger)
+	)
+	.SetDisplayName(LOCTEXT("TabDisplayName", "Cameras Debugger"))
+	.SetTooltipText(LOCTEXT("TabTooltipText", "Open the Cameras Debugger tab."))
+	.SetGroup(WorkspaceMenu::GetMenuStructure().GetDeveloperToolsDebugCategory())
+	.SetCanSidebarTab(false);
+}
+
+void SGameplayCamerasDebugger::UnregisterTabSpawners()
+{
+	if (FSlateApplication::IsInitialized())
+	{
+		FGlobalTabmanager::Get()->UnregisterNomadTabSpawner(SGameplayCamerasDebugger::WindowName);
+	}
+}
+
+TSharedRef<SDockTab> SGameplayCamerasDebugger::SpawnGameplayCamerasDebugger(const FSpawnTabArgs& Args)
+{
+	auto NomadTab = SNew(SDockTab)
+		.TabRole(ETabRole::NomadTab)
+		.Label(LOCTEXT("TabTitle", "Cameras Debugger"));
+
+	TSharedRef<SWidget> MainWidget = SNew(SGameplayCamerasDebugger);
+	NomadTab->SetContent(MainWidget);
+	return NomadTab;
+}
+
+SGameplayCamerasDebugger::SGameplayCamerasDebugger()
+{
+}
+
+SGameplayCamerasDebugger::~SGameplayCamerasDebugger()
+{
+}
+
+void SGameplayCamerasDebugger::Construct(const FArguments& InArgs)
+{
+	IGameplayCamerasEditorModule& GameplayCamerasEditorModule = FModuleManager::GetModuleChecked<IGameplayCamerasEditorModule>(TEXT("GameplayCamerasEditor"));
+	TSharedRef<FGameplayCamerasEditorStyle> GameplayCamerasEditorStyle = FGameplayCamerasEditorStyle::Get();
+	const FName& GameplayCamerasEditorStyleName = GameplayCamerasEditorStyle->GetStyleSetName();
+
+	// Setup commands.
+	const FGameplayCamerasDebuggerCommands& Commands = FGameplayCamerasDebuggerCommands::Get();
+	TSharedRef<FUICommandList> CommandList = MakeShareable(new FUICommandList);
+	CommandList->MapAction(
+			Commands.EnableDebugInfo,
+			FExecuteAction::CreateLambda([]() { GGameplayCamerasDebugEnable = !GGameplayCamerasDebugEnable; }),
+			FCanExecuteAction(),
+			FIsActionChecked::CreateLambda([]() { return GGameplayCamerasDebugEnable; }));
+
+	// Main menu bar.
+	UToolMenus* ToolMenus = UToolMenus::Get();
+	if (!ToolMenus->IsMenuRegistered(SGameplayCamerasDebugger::ToolbarName))
+	{
+		FToolMenuOwnerScoped ToolMenuOwnerScope(this);
+
+		UToolMenu* ToolBarMenu = ToolMenus->RegisterMenu(
+				SGameplayCamerasDebugger::ToolbarName, NAME_None, EMultiBoxType::ToolBar);
+	}
+	FToolMenuContext ToolbarContext;
+	TSharedRef<SWidget> ToolbarContents = ToolMenus->GenerateWidget(
+			SGameplayCamerasDebugger::ToolbarName, ToolbarContext);
+
+	// Empty panel.
+	EmptyPanel = SNew(SBox)
+		.HAlign(HAlign_Center)
+		.VAlign(VAlign_Center)
+		[
+			SNew(STextBlock)
+				.ColorAndOpacity(FSlateColor::UseSubduedForeground())
+				.Text(LOCTEXT("EmptyPanelWarning", "No custom controls for this debug category."))
+		];
+
+	// Toolbar with debug category buttons.
+	FToolBarBuilder DebugCategoriesToolbarBuilder(CommandList.ToSharedPtr(), FMultiBoxCustomization::None);
+	DebugCategoriesToolbarBuilder.BeginSection(TEXT("Main"));
+	{
+		DebugCategoriesToolbarBuilder.AddToolBarButton(
+				Commands.EnableDebugInfo,
+				NAME_None,
+				TAttribute<FText>(FText::FromString(TEXT(""))),
+				TAttribute<FText>(),
+				FSlateIcon(GameplayCamerasEditorStyleName, "Debugger.EnableDebugInfo.Icon"));
+	}
+	DebugCategoriesToolbarBuilder.EndSection();
+	DebugCategoriesToolbarBuilder.BeginSection(TEXT("DebugCategories"));
+	{
+		TArray<FCameraDebugCategoryInfo> RegisteredDebugCategories;
+		GameplayCamerasEditorModule.GetRegisteredDebugCategories(RegisteredDebugCategories);
+		for (const FCameraDebugCategoryInfo& DebugCategory : RegisteredDebugCategories)
+		{
+			TSharedPtr<SWidget> DebugCategoryPanel = GameplayCamerasEditorModule.CreateDebugCategoryPanel(DebugCategory.Name);
+			if (DebugCategoryPanel.IsValid())
+			{
+				DebugPanels.Add(DebugCategory.Name, DebugCategoryPanel);
+			}
+			else
+			{
+				// If there aren't any special UI controls for this category, use an empty panel.
+				DebugPanels.Add(DebugCategory.Name, EmptyPanel);
+			}
+
+			TSharedRef<SDebugCategoryButton> DebugCategoryButton = SNew(SDebugCategoryButton)
+				.DebugCategoryName(DebugCategory.Name)
+				.DisplayText(DebugCategory.DisplayText)
+				.ToolTipText(DebugCategory.ToolTipText)
+				.IconImage(DebugCategory.IconImage.GetIcon())
+				.IsDebugCategoryActive_Lambda([DebugCategory](){ return SGameplayCamerasDebugger::IsDebugCategoryActive(DebugCategory.Name); })
+				.RequestDebugCategoryChange(this, &SGameplayCamerasDebugger::SetActiveDebugCategoryPanel);
+			DebugCategoriesToolbarBuilder.AddToolBarWidget(DebugCategoryButton);
+		}
+	}
+	DebugCategoriesToolbarBuilder.EndSection();
+
+	// Main layout.
+	ChildSlot
+	[
+		SNew(SVerticalBox)
+		+ SVerticalBox::Slot()
+			.AutoHeight()
+			[
+				ToolbarContents
+			]
+		+ SVerticalBox::Slot()
+			.AutoHeight()
+			[
+				DebugCategoriesToolbarBuilder.MakeWidget()
+			]
+		+ SVerticalBox::Slot()
+		.Padding(2.0)
+			[
+				SAssignNew(PanelHost, SBox)
+				[
+					SNullWidget::NullWidget
+				]
+			]
+	];
+}
+
+bool SGameplayCamerasDebugger::IsDebugCategoryActive(const FString& InCategoryName)
+{
+	TArray<FStringView, TInlineAllocator<4>> ActiveCategories;
+	UE::String::ParseTokens(GGameplayCamerasDebugCategories, ',', ActiveCategories);
+	return ActiveCategories.Contains(InCategoryName);
+}
+
+void SGameplayCamerasDebugger::SetActiveDebugCategoryPanel(const FString& InCategoryName)
+{
+	if (ensureMsgf(
+				DebugPanels.Contains(InCategoryName), 
+				TEXT("Debug category was not registered with IGameplayCamerasEditorModule: %s"), *InCategoryName))
+	{
+		TSharedPtr<SWidget> DebugPanel = DebugPanels.FindChecked(InCategoryName);
+		check(DebugPanel.IsValid());
+		PanelHost->SetContent(DebugPanel.ToSharedRef());
+
+		GGameplayCamerasDebugCategories = InCategoryName;
+	}
+	else
+	{
+		PanelHost->SetContent(SNullWidget::NullWidget);
+	}
+}
+
+}  // namespace UE::Cameras
+
+#undef LOCTEXT_NAMESPACE
+
diff --git a/Engine/Plugins/Cameras/GameplayCameras/Source/GameplayCamerasEditor/Private/Debugger/SGameplayCamerasDebugger.h b/Engine/Plugins/Cameras/GameplayCameras/Source/GameplayCamerasEditor/Private/Debugger/SGameplayCamerasDebugger.h
new file mode 100644
--- /dev/null
+++ b/Engine/Plugins/Cameras/GameplayCameras/Source/GameplayCamerasEditor/Private/Debugger/SGameplayCamerasDebugger.h
@@ -0,0 +1,50 @@
+// Copyright Epic Games, Inc. All Rights Reserved.
+
+#pragma once
+
+#include "CoreTypes.h"
+#include "Widgets/SCompoundWidget.h"
+
+class FTabManager;
+class FSpawnTabArgs;
+class SBox;
+class SDockTab;
+
+namespace UE::Cameras
+{
+
+class SGameplayCamerasDebugger : public SCompoundWidget
+{
+public:
+
+	static const FName WindowName;
+	static const FName ToolbarName;
+
+	static void RegisterTabSpawners();
+	static TSharedRef<SDockTab> SpawnGameplayCamerasDebugger(const FSpawnTabArgs& Args);
+	static void UnregisterTabSpawners();
+
+public:
+
+	SLATE_BEGIN_ARGS(SGameplayCamerasDebugger) {}
+	SLATE_END_ARGS();
+
+	SGameplayCamerasDebugger();
+	virtual ~SGameplayCamerasDebugger();
+
+	void Construct(const FArguments& InArgs);
+
+protected:
+
+	static bool IsDebugCategoryActive(const FString& InCategoryName);
+	void SetActiveDebugCategoryPanel(const FString& InCategoryName);
+
+private:
+
+	TSharedPtr<SBox> PanelHost;
+	TSharedPtr<SWidget> EmptyPanel;
+	TMap<FString, TSharedPtr<SWidget>> DebugPanels;
+};
+
+}  // namespace UE::Cameras
+
diff --git a/Engine/Plugins/Cameras/GameplayCameras/Source/GameplayCamerasEditor/Private/GameplayCamerasEditorModule.cpp b/Engine/Plugins/Cameras/GameplayCameras/Source/GameplayCamerasEditor/Private/GameplayCamerasEditorModule.cpp
--- a/Engine/Plugins/Cameras/GameplayCameras/Source/GameplayCamerasEditor/Private/GameplayCamerasEditorModule.cpp
+++ b/Engine/Plugins/Cameras/GameplayCameras/Source/GameplayCamerasEditor/Private/GameplayCamerasEditorModule.cpp
@@ -6,10 +6,16 @@
 #include "AssetTools/CameraRigAssetEditor.h"
 #include "Commands/CameraAssetEditorCommands.h"
 #include "Commands/CameraRigAssetEditorCommands.h"
+#include "Commands/GameplayCamerasDebuggerCommands.h"
+#include "Debug/CameraDebugCategories.h"
+#include "Debugger/SBlendStacksDebugPanel.h"
+#include "Debugger/SCameraNodeTreeDebugPanel.h"
+#include "Debugger/SGameplayCamerasDebugger.h"
 #include "GameplayCamerasLiveEditManager.h"
 #include "IGameplayCamerasEditorModule.h"
 #include "IGameplayCamerasModule.h"
 #include "ISettingsModule.h"
+#include "Misc/CoreDelegates.h"
 #include "Modules/ModuleManager.h"
 #include "Styles/GameplayCamerasEditorStyle.h"
 #include "ToolMenus.h"
@@ -33,7 +39,19 @@ public:
 
 	virtual void StartupModule() override
 	{
+		if (GEditor)
+		{
+			OnPostEngineInit();
+		}
+		else
+		{
+			FCoreDelegates::OnPostEngineInit.AddRaw(this, &FGameplayCamerasEditorModule::OnPostEngineInit);
+		}
+
+		FCoreDelegates::OnEnginePreExit.AddRaw(this, &FGameplayCamerasEditorModule::OnPreExit);
+
 		RegisterSettings();
+		RegisterCoreDebugCategories();
 		InitializeLiveEditManager();
 
 		UToolMenus::RegisterStartupCallback(FSimpleMulticastDelegate::FDelegate::CreateRaw(
@@ -48,9 +66,14 @@ public:
 
 		FCameraAssetEditorCommands::Unregister();
 		FCameraRigAssetEditorCommands::Unregister();
+		FGameplayCamerasDebuggerCommands::Unregister();
 
 		UnregisterSettings();
+		UnregisterCoreDebugCategories();
 		TeardownLiveEditManager();
+
+		FCoreDelegates::OnPostEngineInit.RemoveAll(this);
+		FCoreDelegates::OnEnginePreExit.RemoveAll(this);
 	}
 
 	virtual UCameraAssetEditor* CreateCameraAssetEditor(const EToolkitMode::Type Mode, const TSharedPtr<IToolkitHost>& InitToolkitHost, UCameraAsset* CameraAsset) override
@@ -69,8 +92,70 @@ public:
 		return AssetEditor;
 	}
 
+	virtual void RegisterDebugCategory(const UE::Cameras::FCameraDebugCategoryInfo& InCategoryInfo) override
+	{
+		if (!ensureMsgf(!InCategoryInfo.Name.IsEmpty(), TEXT("A debug category must at least specify a name!")))
+		{
+			return;
+		}
+
+		DebugCategoryInfos.Add(InCategoryInfo.Name, InCategoryInfo);
+	}
+
+	virtual void GetRegisteredDebugCategories(TArray<UE::Cameras::FCameraDebugCategoryInfo>& OutCategoryInfos) override
+	{
+		DebugCategoryInfos.GenerateValueArray(OutCategoryInfos);
+	}
+
+	virtual void UnregisterDebugCategory(const FString& InCategoryName)
+	{
+		DebugCategoryInfos.Remove(InCategoryName);
+
+	}
+
+	virtual void RegisterDebugCategoryPanel(const FString& InDebugCategory, FOnCreateDebugCategoryPanel OnCreatePanel) override
+	{
+		if (!DebugCategoryPanelCreators.Contains(InDebugCategory))
+		{
+			DebugCategoryPanelCreators.Add(InDebugCategory, OnCreatePanel);
+		}
+		else
+		{
+			// Override existing creator... for games and projects that want to extend a panel with extra controls.
+			DebugCategoryPanelCreators[InDebugCategory] = OnCreatePanel;
+		}
+	}
+
+	virtual TSharedPtr<SWidget> CreateDebugCategoryPanel(const FString& InDebugCategory) override
+	{
+		if (FOnCreateDebugCategoryPanel* PanelCreator = DebugCategoryPanelCreators.Find(InDebugCategory))
+		{
+			return PanelCreator->Execute(InDebugCategory).ToSharedPtr();
+		}
+		return nullptr;
+	}
+
+	virtual void UnregisterDebugCategoryPanel(const FString& InDebugCategory) override
+	{
+		DebugCategoryPanelCreators.Remove(InDebugCategory);
+	}
+
 private:
 
+	void OnPostEngineInit()
+	{
+		using namespace UE::Cameras;
+
+		SGameplayCamerasDebugger::RegisterTabSpawners();
+	}
+
+	void OnPreExit()
+	{
+		using namespace UE::Cameras;
+
+		SGameplayCamerasDebugger::UnregisterTabSpawners();
+	}
+
 	void RegisterSettings()
 	{
 		ISettingsModule* SettingsModule = FModuleManager::GetModulePtr<ISettingsModule>("Settings");
@@ -96,12 +181,69 @@ private:
 		}
 	}
 
+	void RegisterCoreDebugCategories()
+	{
+		using namespace UE::Cameras;
+
+		TSharedRef<FGameplayCamerasEditorStyle> GameplayCamerasEditorStyle = FGameplayCamerasEditorStyle::Get();
+		const FName& GameplayCamerasEditorStyleName = GameplayCamerasEditorStyle->GetStyleSetName();
+
+		RegisterDebugCategory(FCameraDebugCategoryInfo{
+				FCameraDebugCategories::NodeTree,
+				LOCTEXT("NodeTreeDebugCategory", "Node Tree"),
+				LOCTEXT("NodeTreeDebugCategoryToolTip", "Shows the entire camrera node evaluator tree"),
+				FSlateIcon(GameplayCamerasEditorStyleName, "DebugCategory.NodeTree.Icon")
+			});
+		RegisterDebugCategory(FCameraDebugCategoryInfo{
+				FCameraDebugCategories::DirectorTree,
+				LOCTEXT("DirectorTreeDebugCategory", "Node Tree"),
+				LOCTEXT("DirectorTreeDebugCategoryToolTip", "Shows the active/inactive directors, and their evaluation context"),
+				FSlateIcon(GameplayCamerasEditorStyleName, "DebugCategory.DirectorTree.Icon")
+			});
+		RegisterDebugCategory(FCameraDebugCategoryInfo{
+				FCameraDebugCategories::BlendStacks,
+				LOCTEXT("BlendStacksDebugCategory", "Blend Stacks"),
+				LOCTEXT("BlendStacksDebugCategoryToolTip", "Shows a summary of the blend stacks"),
+				FSlateIcon(GameplayCamerasEditorStyleName, "DebugCategory.BlendStacks.Icon")
+			});
+		RegisterDebugCategory(FCameraDebugCategoryInfo{
+				FCameraDebugCategories::PoseStats,
+				LOCTEXT("PoseStatsDebugCategory", "Pose Stats"),
+				LOCTEXT("PoseStatsDebugCategoryToolTip", "Shows the evaluated camera pose"),
+				FSlateIcon(GameplayCamerasEditorStyleName, "DebugCategory.PoseStats.Icon")
+			});
+		RegisterDebugCategory(FCameraDebugCategoryInfo{
+				FCameraDebugCategories::Viewfinder,
+				LOCTEXT("ViewfinderDebugCategory", "Viewfinder"),
+				LOCTEXT("ViewfinderDebugCategoryToolTip", "Shows an old-school viewfinder on screen"),
+				FSlateIcon(GameplayCamerasEditorStyleName, "DebugCategory.Viewfinder.Icon")
+			});
+
+		RegisterDebugCategoryPanel(FCameraDebugCategories::NodeTree, FOnCreateDebugCategoryPanel::CreateLambda([](const FString&)
+					{
+						return SNew(SCameraNodeTreeDebugPanel);
+					}));
+		RegisterDebugCategoryPanel(FCameraDebugCategories::BlendStacks, FOnCreateDebugCategoryPanel::CreateLambda([](const FString&)
+					{
+						return SNew(SBlendStacksDebugPanel);
+					}));
+	}
+
+	void UnregisterCoreDebugCategories()
+	{
+		using namespace UE::Cameras;
+
+		UnregisterDebugCategoryPanel(FCameraDebugCategories::BlendStacks);
+		UnregisterDebugCategoryPanel(FCameraDebugCategories::NodeTree);
+	}
+
 	void RegisterMenus()
 	{
 		using namespace UE::Cameras;
 
 		FCameraAssetEditorCommands::Register();
 		FCameraRigAssetEditorCommands::Register();	
+		FGameplayCamerasDebuggerCommands::Register();
 	}
 
 	void InitializeLiveEditManager()
@@ -133,6 +275,9 @@ private:
 private:
 
 	TSharedPtr<UE::Cameras::FGameplayCamerasLiveEditManager> LiveEditManager;
+
+	TMap<FString, UE::Cameras::FCameraDebugCategoryInfo> DebugCategoryInfos;
+	TMap<FString, FOnCreateDebugCategoryPanel> DebugCategoryPanelCreators;
 };
 
 IMPLEMENT_MODULE(FGameplayCamerasEditorModule, GameplayCamerasEditor);
diff --git a/Engine/Plugins/Cameras/GameplayCameras/Source/GameplayCamerasEditor/Private/Styles/GameplayCamerasEditorStyle.cpp b/Engine/Plugins/Cameras/GameplayCameras/Source/GameplayCamerasEditor/Private/Styles/GameplayCamerasEditorStyle.cpp
--- a/Engine/Plugins/Cameras/GameplayCameras/Source/GameplayCamerasEditor/Private/Styles/GameplayCamerasEditorStyle.cpp
+++ b/Engine/Plugins/Cameras/GameplayCameras/Source/GameplayCamerasEditor/Private/Styles/GameplayCamerasEditorStyle.cpp
@@ -2,5 +2,56 @@
 
 #include "Styles/GameplayCamerasEditorStyle.h"
 
+#include "Brushes/SlateImageBrush.h"
+#include "Interfaces/IPluginManager.h"
+#include "Misc/Paths.h"
+#include "Styling/SlateStyleRegistry.h"
+
+namespace UE::Cameras
+{
+
+#define IMAGE_BRUSH(RelativePath, ...) FSlateImageBrush(RootToContentDir(RelativePath, TEXT(".png")), __VA_ARGS__)
+
 TSharedPtr<FGameplayCamerasEditorStyle> FGameplayCamerasEditorStyle::Singleton;
 
+FGameplayCamerasEditorStyle::FGameplayCamerasEditorStyle()
+	: FSlateStyleSet("GameplayCamerasEditorStyle")
+{
+	const FVector2D Icon16x16(16.0f, 16.0f);
+	const FVector2D Icon24x24(24.0f, 24.0f);
+	const FVector2D Icon48x48(48.0f, 48.0f);
+	const FVector2D Icon64x64(64.0f, 64.0f);
+
+	const FString ContentDir = IPluginManager::Get().FindPlugin(TEXT("GameplayCameras"))->GetContentDir();
+	SetContentRoot(ContentDir);
+	SetCoreContentRoot(FPaths::EngineContentDir() / TEXT("Slate"));
+
+	Set("Debugger.EnableDebugInfo.Icon", new IMAGE_BRUSH("/Icons/EnableDebugInfo", Icon16x16));
+
+	Set("DebugCategory.NodeTree.Icon", new IMAGE_BRUSH("/Icons/DebugCategory-NodeTree", Icon16x16));
+	Set("DebugCategory.DirectorTree.Icon", new IMAGE_BRUSH("/Icons/DebugCategory-DirectorTree", Icon16x16));
+	Set("DebugCategory.BlendStacks.Icon", new IMAGE_BRUSH("/Icons/DebugCategory-BlendStacks", Icon16x16));
+	Set("DebugCategory.PoseStats.Icon", new IMAGE_BRUSH("/Icons/DebugCategory-PoseStats", Icon16x16));
+	Set("DebugCategory.Viewfinder.Icon", new IMAGE_BRUSH("/Icons/DebugCategory-Viewfinder", Icon16x16));
+
+	FSlateStyleRegistry::RegisterSlateStyle(*this);
+}
+
+FGameplayCamerasEditorStyle::~FGameplayCamerasEditorStyle()
+{
+	FSlateStyleRegistry::UnRegisterSlateStyle(*this);
+}
+
+TSharedRef<FGameplayCamerasEditorStyle> FGameplayCamerasEditorStyle::Get()
+{
+	if (!Singleton.IsValid())
+	{
+		Singleton = MakeShareable(new FGameplayCamerasEditorStyle);
+	}
+	return Singleton.ToSharedRef();
+}
+
+#undef IMAGE_BRUSH
+
+}  // namespace UE::Cameras
+
diff --git a/Engine/Plugins/Cameras/GameplayCameras/Source/GameplayCamerasEditor/Private/Styles/GameplayCamerasEditorStyle.h b/Engine/Plugins/Cameras/GameplayCameras/Source/GameplayCamerasEditor/Private/Styles/GameplayCamerasEditorStyle.h
--- a/Engine/Plugins/Cameras/GameplayCameras/Source/GameplayCamerasEditor/Private/Styles/GameplayCamerasEditorStyle.h
+++ b/Engine/Plugins/Cameras/GameplayCameras/Source/GameplayCamerasEditor/Private/Styles/GameplayCamerasEditorStyle.h
@@ -2,12 +2,10 @@
 
 #pragma once
 
-#include "Misc/Paths.h"
-#include "Brushes/SlateImageBrush.h"
-#include "Styling/SlateStyleRegistry.h"
 #include "Styling/SlateStyle.h"
 
-#define IMAGE_BRUSH(RelativePath, ...) FSlateImageBrush(RootToContentDir(RelativePath, TEXT(".png")), __VA_ARGS__)
+namespace UE::Cameras
+{
 
 /**
  * Implements the visual style of the gameplay cameras editors.
@@ -16,37 +14,15 @@ class FGameplayCamerasEditorStyle final : public FSlateStyleSet
 {
 public:
 	
-	FGameplayCamerasEditorStyle()
-		: FSlateStyleSet("GameplayCamerasEditorStyle")
-	{
-		const FVector2D Icon16x16(16.0f, 16.0f);
-		const FVector2D Icon24x24(24.0f, 24.0f);
-		const FVector2D Icon48x48(48.0f, 48.0f);
-		const FVector2D Icon64x64(64.0f, 64.0f);
-
-		//SetContentRoot(FPaths::EnginePluginsDir() / TEXT("Cameras/GameplayCameras/Content"));
-
-		FSlateStyleRegistry::RegisterSlateStyle(*this);
-	}
-
-	virtual ~FGameplayCamerasEditorStyle()
-	{
-		FSlateStyleRegistry::UnRegisterSlateStyle(*this);
-	}
-
-	static TSharedRef<FGameplayCamerasEditorStyle> Get()
-	{
-		if (!Singleton.IsValid())
-		{
-			Singleton = MakeShareable(new FGameplayCamerasEditorStyle);
-		}
-		return Singleton.ToSharedRef();
-	}
+	FGameplayCamerasEditorStyle();
+	virtual ~FGameplayCamerasEditorStyle();
+
+	static TSharedRef<FGameplayCamerasEditorStyle> Get();
 
 private:
 
 	static TSharedPtr<FGameplayCamerasEditorStyle> Singleton;
 };
 
-#undef IMAGE_BRUSH
+}  // namespace UE::Cameras
 
diff --git a/Engine/Plugins/Cameras/GameplayCameras/Source/GameplayCamerasEditor/Public/Commands/GameplayCamerasDebuggerCommands.h b/Engine/Plugins/Cameras/GameplayCameras/Source/GameplayCamerasEditor/Public/Commands/GameplayCamerasDebuggerCommands.h
new file mode 100644
--- /dev/null
+++ b/Engine/Plugins/Cameras/GameplayCameras/Source/GameplayCamerasEditor/Public/Commands/GameplayCamerasDebuggerCommands.h
@@ -0,0 +1,29 @@
+// Copyright Epic Games, Inc. All Rights Reserved.
+
+#pragma once
+
+#include "CoreTypes.h"
+#include "Framework/Commands/Commands.h"
+
+class FUICommandInfo;
+class UClass;
+struct FInputChord;
+
+namespace UE::Cameras
+{
+
+class FGameplayCamerasDebuggerCommands : public TCommands<FGameplayCamerasDebuggerCommands>
+{
+public:
+
+	FGameplayCamerasDebuggerCommands();
+
+	virtual void RegisterCommands() override;
+	
+public:
+
+	TSharedPtr<FUICommandInfo> EnableDebugInfo;
+};
+
+}  // namespace UE::Cameras
+
diff --git a/Engine/Plugins/Cameras/GameplayCameras/Source/GameplayCamerasEditor/Public/IGameplayCamerasEditorModule.h b/Engine/Plugins/Cameras/GameplayCameras/Source/GameplayCamerasEditor/Public/IGameplayCamerasEditorModule.h
--- a/Engine/Plugins/Cameras/GameplayCameras/Source/GameplayCamerasEditor/Public/IGameplayCamerasEditorModule.h
+++ b/Engine/Plugins/Cameras/GameplayCameras/Source/GameplayCamerasEditor/Public/IGameplayCamerasEditorModule.h
@@ -3,6 +3,7 @@
 #pragma once
 
 #include "CoreMinimal.h"
+#include "Delegates/Delegate.h"
 #include "Modules/ModuleInterface.h"
 #include "Toolkits/AssetEditorToolkit.h"
 
@@ -11,6 +12,21 @@ class UCameraAssetEditor;
 class UCameraRigAsset;
 class UCameraRigAssetEditor;
 
+DECLARE_DELEGATE_RetVal_OneParam(TSharedRef<SWidget>, FOnCreateDebugCategoryPanel, const FString&);
+
+namespace UE::Cameras
+{
+
+struct FCameraDebugCategoryInfo
+{
+	FString Name;
+	FText DisplayText;
+	FText ToolTipText;
+	FSlateIcon IconImage;
+};
+
+}  // namespace UE::Cameras
+
 /**
  * The gameplay cameras editor module.
  */
@@ -19,7 +35,6 @@ class IGameplayCamerasEditorModule : public IModuleInterface
 public:
 
 	static const FName GameplayCamerasEditorAppIdentifier;
-
 	static const FName CameraRigAssetEditorToolBarName;
 
 	virtual ~IGameplayCamerasEditorModule() = default;
@@ -31,5 +46,23 @@ public:
 
 	/** Creates an editor for the given camera rig asset */
 	virtual UCameraRigAssetEditor* CreateCameraRigEditor(const EToolkitMode::Type Mode, const TSharedPtr<IToolkitHost>& InitToolkitHost, UCameraRigAsset* CameraRig) = 0;
+
+public:
+
+	using FCameraDebugCategoryInfo = UE::Cameras::FCameraDebugCategoryInfo;
+
+	/** Registers a new debug category, to be displayed in the camera debugger tool. */
+	virtual void RegisterDebugCategory(const UE::Cameras::FCameraDebugCategoryInfo& InCategoryInfo) = 0;
+	/** Gets all registered debug categories. */
+	virtual void GetRegisteredDebugCategories(TArray<UE::Cameras::FCameraDebugCategoryInfo>& OutCategoryInfos) = 0;
+	/** Unregisters a debug category. */
+	virtual void UnregisterDebugCategory(const FString& InCategoryName) = 0;
+
+	/** Registers a custom UI panel for a given debug category. */
+	virtual void RegisterDebugCategoryPanel(const FString& InCategoryName, FOnCreateDebugCategoryPanel OnCreatePanel) = 0;
+	/** Creates the custom UI panel (if any) for a given debug category. */
+	virtual TSharedPtr<SWidget> CreateDebugCategoryPanel(const FString& InCategoryName) = 0;
+	/** Unregisters a debug category's custom UI panel. */
+	virtual void UnregisterDebugCategoryPanel(const FString& InCategoryName) = 0;
 };
 
