Cameras: fix non-unity build

diff --git a/Engine/Plugins/Cameras/GameplayCameras/Source/GameplayCameras/Public/GameFramework/GameplayCameraSystemActor.h b/Engine/Plugins/Cameras/GameplayCameras/Source/GameplayCameras/Public/GameFramework/GameplayCameraSystemActor.h
--- a/Engine/Plugins/Cameras/GameplayCameras/Source/GameplayCameras/Public/GameFramework/GameplayCameraSystemActor.h
+++ b/Engine/Plugins/Cameras/GameplayCameras/Source/GameplayCameras/Public/GameFramework/GameplayCameraSystemActor.h
@@ -2,9 +2,11 @@
 
 #pragma once
 
+#include "Core/CameraPose.h"
 #include "CoreMinimal.h"
 #include "Engine/EngineTypes.h"
 #include "GameFramework/Actor.h"
+#include "GameplayCameras.h"
 #include "UObject/ObjectMacros.h"
 
 #include "GameplayCameraSystemActor.generated.h"
