Cameras: RTTI framework handles name lookups and type size/align info.

diff --git a/Engine/Plugins/Cameras/GameplayCameras/Source/GameplayCameras/Private/Core/CameraObjectRtti.cpp b/Engine/Plugins/Cameras/GameplayCameras/Source/GameplayCameras/Private/Core/CameraObjectRtti.cpp
--- a/Engine/Plugins/Cameras/GameplayCameras/Source/GameplayCameras/Private/Core/CameraObjectRtti.cpp
+++ b/Engine/Plugins/Cameras/GameplayCameras/Source/GameplayCameras/Private/Core/CameraObjectRtti.cpp
@@ -20,12 +20,33 @@ FCameraObjectTypeRegistry& FCameraObjectTypeRegistry::Get()
 void FCameraObjectTypeRegistry::RegisterType(FCameraObjectTypeID TypeID, FCameraObjectTypeInfo&& TypeInfo)
 {
 	ensureMsgf(
-			!TypeIDsByName.Contains(TypeID.GetTypeName()), 
-			TEXT("Type '%s' has already been registered!"), *TypeID.GetTypeName().ToString());
-	TypeIDsByName.Add(TypeID.GetTypeName(), TypeID.GetTypeID());
+			!TypeIDsByName.Contains(TypeInfo.TypeName), 
+			TEXT("Type '%s' has already been registered!"), *TypeInfo.TypeName.ToString());
+	TypeIDsByName.Add(TypeInfo.TypeName, TypeID.GetTypeID());
 	TypeInfos.Insert(TypeID.GetTypeID(), MoveTemp(TypeInfo));
 }
 
+FCameraObjectTypeID FCameraObjectTypeRegistry::FindTypeByName(const FName& TypeName)
+{
+	if (uint32* RegisteredTypeID = TypeIDsByName.Find(TypeName))
+	{
+		return FCameraObjectTypeID{ *RegisteredTypeID };
+	}
+	return FCameraObjectTypeID::Invalid();
+}
+
+const FCameraObjectTypeInfo* FCameraObjectTypeRegistry::GetTypeInfo(FCameraObjectTypeID TypeID)
+{
+	if (ensureMsgf(
+				TypeID.IsValid() && TypeInfos.IsValidIndex(TypeID.GetTypeID()),
+				TEXT("Given type ID is not valid, or not registered.")))
+	{
+		const FCameraObjectTypeInfo& TypeInfo = TypeInfos[TypeID.GetTypeID()];
+		return &TypeInfo;
+	}
+	return nullptr;
+}
+
 void FCameraObjectTypeRegistry::ConstructObject(FCameraObjectTypeID TypeID, void* Ptr)
 {
 	if (ensureMsgf(TypeInfos.IsValidIndex(TypeID.GetTypeID()), TEXT("Invalid camera object type ID!")))
diff --git a/Engine/Plugins/Cameras/GameplayCameras/Source/GameplayCameras/Public/Core/CameraObjectRtti.h b/Engine/Plugins/Cameras/GameplayCameras/Source/GameplayCameras/Public/Core/CameraObjectRtti.h
--- a/Engine/Plugins/Cameras/GameplayCameras/Source/GameplayCameras/Public/Core/CameraObjectRtti.h
+++ b/Engine/Plugins/Cameras/GameplayCameras/Source/GameplayCameras/Public/Core/CameraObjectRtti.h
@@ -34,24 +34,27 @@ struct FCameraObjectTypeID
 	{
 		return ID;
 	}
-
-	const FName& GetTypeName() const
+	
+	bool IsValid() const
 	{
-		return Name;
+		return ID != MAX_uint32;
 	}
 
-	FCameraObjectTypeID(const FName& InName, uint32 InID)
-		: Name(InName)
-		, ID(InID)
+	FCameraObjectTypeID(uint32 InID)
+		: ID(InID)
 	{}
 
+	static FCameraObjectTypeID Invalid()
+	{
+		return FCameraObjectTypeID(MAX_uint32);
+	}
+
 protected:
 
-	static uint32 RegisterNewID();
+	GAMEPLAYCAMERAS_API static uint32 RegisterNewID();
 
 protected:
 
-	FName Name;
 	uint32 ID;
 };
 
@@ -60,6 +63,11 @@ protected:
  */
 struct FCameraObjectTypeInfo
 {
+	FName TypeName;
+
+	uint32 Sizeof = 0;
+	uint32 Alignof = 0;
+
 	using FConstructor = void(*)(void*);
 	FConstructor Constructor;
 
@@ -76,8 +84,11 @@ public:
 
 	static FCameraObjectTypeRegistry& Get();
 
-	void RegisterType(FCameraObjectTypeID TypeID, FCameraObjectTypeInfo&& TypeInfo);
-	void ConstructObject(FCameraObjectTypeID TypeID, void* Ptr);
+	GAMEPLAYCAMERAS_API void RegisterType(FCameraObjectTypeID TypeID, FCameraObjectTypeInfo&& TypeInfo);
+	GAMEPLAYCAMERAS_API FCameraObjectTypeID FindTypeByName(const FName& TypeName);
+	GAMEPLAYCAMERAS_API const FCameraObjectTypeInfo* GetTypeInfo(FCameraObjectTypeID TypeID);
+
+	GAMEPLAYCAMERAS_API void ConstructObject(FCameraObjectTypeID TypeID, void* Ptr);
 
 private:
 
@@ -93,12 +104,14 @@ struct TCameraObjectTypeID : FCameraObjectTypeID
 {
 private:
 
-	TCameraObjectTypeID(const FName& InName, uint32 InID) : FCameraObjectTypeID(InName, InID) {}
+	TCameraObjectTypeID(uint32 InID) : FCameraObjectTypeID(InID) {}
 
 	static TCameraObjectTypeID RegisterType(const FName& InClassName)
 	{
-		TCameraObjectTypeID NewTypeID(InClassName, FCameraObjectTypeID::RegisterNewID());
-		FCameraObjectTypeInfo NewTypeInfo { 
+		TCameraObjectTypeID NewTypeID(FCameraObjectTypeID::RegisterNewID());
+		FCameraObjectTypeInfo NewTypeInfo {
+			InClassName,
+			sizeof(T), alignof(T),
 			&TCameraObjectTypeID<T>::StaticConstructor,
 			&TCameraObjectTypeID<T>::StaticDestructor
 		};
@@ -137,7 +150,7 @@ private:
 		template<typename Type> Type* CastThisChecked() { check(IsKindOf<Type>()); return static_cast<Type*>(this); }\
 		template<typename Type> const Type* CastThisChecked() const { check(IsKindOf<Type>()); return static_cast<Type*>(this); }\
 	private:\
-		static const ::UE::Cameras::TCameraObjectTypeID<ClassName> PrivateTypeID;
+		GAMEPLAYCAMERAS_API static const ::UE::Cameras::TCameraObjectTypeID<ClassName> PrivateTypeID;
 
 #define UE_GAMEPLAY_CAMERAS_DECLARE_RTTI(ClassName, BaseClassName)\
 	public:\
@@ -145,7 +158,7 @@ private:
 		virtual const FCameraObjectTypeID& GetTypeID() const override { return ClassName::PrivateTypeID; }\
 		virtual bool IsKindOf(const FCameraObjectTypeID& InTypeID) const override { return (InTypeID == ClassName::PrivateTypeID) || BaseClassName::IsKindOf(InTypeID); }\
 	private:\
-		static const ::UE::Cameras::TCameraObjectTypeID<ClassName> PrivateTypeID;
+		GAMEPLAYCAMERAS_API static const ::UE::Cameras::TCameraObjectTypeID<ClassName> PrivateTypeID;
 
 #define UE_GAMEPLAY_CAMERAS_DEFINE_RTTI(ClassName)\
 	const ::UE::Cameras::TCameraObjectTypeID<ClassName> ClassName::PrivateTypeID = ::UE::Cameras::TCameraObjectTypeID<ClassName>::RegisterType(#ClassName);
