Cameras: fix bug where the last frame of blending between rigs resulted in invalid FocalLength/FieldOfView.

This change also makes blending FocalLength and FieldOfView work between themselves first, and only fallback to effective FOV if both rigs are using a mix of the two properties.

diff --git a/Engine/Plugins/Cameras/GameplayCameras/Source/GameplayCameras/Private/Core/BlendStackCameraNode.cpp b/Engine/Plugins/Cameras/GameplayCameras/Source/GameplayCameras/Private/Core/BlendStackCameraNode.cpp
--- a/Engine/Plugins/Cameras/GameplayCameras/Source/GameplayCameras/Private/Core/BlendStackCameraNode.cpp
+++ b/Engine/Plugins/Cameras/GameplayCameras/Source/GameplayCameras/Private/Core/BlendStackCameraNode.cpp
@@ -204,8 +204,6 @@ void FBlendStackCameraNodeEvaluator::OnRun(const FCameraNodeEvaluationParams& Pa
 			continue;
 		}
 
-		const FCameraPoseFlags ChangedFlags(CurResult.CameraPose.GetChangedFlags());
-
 		FCameraNodeEvaluationParams CurParams(Params);
 		CurParams.EvaluationContext = Entry.EvaluationContext.Pin();
 		CurParams.bIsFirstFrame = Entry.bIsFirstFrame;
@@ -229,8 +227,6 @@ void FBlendStackCameraNodeEvaluator::OnRun(const FCameraNodeEvaluationParams& Pa
 
 			PopEntriesBelow = EntryIndex;
 		}
-
-		CurResult.CameraPose.SetChangedFlags(ChangedFlags);
 		
 		++EntryIndex;
 	}
diff --git a/Engine/Plugins/Cameras/GameplayCameras/Source/GameplayCameras/Private/Core/CameraPose.cpp b/Engine/Plugins/Cameras/GameplayCameras/Source/GameplayCameras/Private/Core/CameraPose.cpp
--- a/Engine/Plugins/Cameras/GameplayCameras/Source/GameplayCameras/Private/Core/CameraPose.cpp
+++ b/Engine/Plugins/Cameras/GameplayCameras/Source/GameplayCameras/Private/Core/CameraPose.cpp
@@ -7,6 +7,21 @@
 
 #include UE_INLINE_GENERATED_CPP_BY_NAME(CameraPose)
 
+const FCameraPoseFlags& FCameraPoseFlags::All()
+{
+	static FCameraPoseFlags Instance(true);
+	return Instance;
+}
+
+FCameraPoseFlags::FCameraPoseFlags()
+{
+}
+
+FCameraPoseFlags::FCameraPoseFlags(bool bInValue)
+{
+	SetAllFlags(bInValue);
+}
+
 FCameraPoseFlags& FCameraPoseFlags::SetAllFlags(bool bInValue)
 {
 #define UE_CAMERA_POSE_FOR_PROPERTY(PropType, PropName)\
@@ -132,12 +147,22 @@ FRay3d FCameraPose::GetAimRay() const
 	return FRay3d(Location, Rotation.RotateVector(TargetDir), bDirectionIsNormalized);
 }
 
+void FCameraPose::OverrideAll(const FCameraPose& OtherPose)
+{
+	InternalOverrideChanged(OtherPose, false);
+}
+
 void FCameraPose::OverrideChanged(const FCameraPose& OtherPose)
+{
+	InternalOverrideChanged(OtherPose, true);
+}
+
+void FCameraPose::InternalOverrideChanged(const FCameraPose& OtherPose, bool bChangedOnly)
 {
 	const FCameraPoseFlags& OtherPoseChangedFlags = OtherPose.GetChangedFlags();
 
 #define UE_CAMERA_POSE_FOR_PROPERTY(PropType, PropName)\
-	if (OtherPoseChangedFlags.PropName)\
+	if (!bChangedOnly || OtherPoseChangedFlags.PropName)\
 	{\
 		Set##PropName(OtherPose.Get##PropName());\
 	}
@@ -147,19 +172,24 @@ UE_CAMERA_POSE_FOR_ALL_PROPERTIES()
 #undef UE_CAMERA_POSE_FOR_PROPERTY
 }
 
+void FCameraPose::LerpAll(const FCameraPose& ToPose, float Factor)
+{
+	FCameraPoseFlags DummyFlags(true);
+	InternalLerpChanged(ToPose, Factor, DummyFlags, false, DummyFlags, false);
+}
+
 void FCameraPose::LerpChanged(const FCameraPose& ToPose, float Factor)
 {
-	FCameraPoseFlags DummyFlags;
-	DummyFlags.SetAllFlags(true);
-	InternalLerpChanged(ToPose, Factor, DummyFlags, false, DummyFlags);
+	FCameraPoseFlags DummyFlags(true);
+	InternalLerpChanged(ToPose, Factor, DummyFlags, false, DummyFlags, true);
 }
 
 void FCameraPose::LerpChanged(const FCameraPose& ToPose, float Factor, const FCameraPoseFlags& InMask, bool bInvertMask, FCameraPoseFlags& OutMask)
 {
-	InternalLerpChanged(ToPose, Factor, InMask, bInvertMask, OutMask);
+	InternalLerpChanged(ToPose, Factor, InMask, bInvertMask, OutMask, true);
 }
 
-void FCameraPose::InternalLerpChanged(const FCameraPose& ToPose, float Factor, const FCameraPoseFlags& InMask, bool bInvertMask, FCameraPoseFlags& OutMask)
+void FCameraPose::InternalLerpChanged(const FCameraPose& ToPose, float Factor, const FCameraPoseFlags& InMask, bool bInvertMask, FCameraPoseFlags& OutMask, bool bChangedOnly)
 {
 	if (UNLIKELY(Factor == 0.f))
 	{
@@ -172,21 +202,34 @@ void FCameraPose::InternalLerpChanged(const FCameraPose& ToPose, float Factor, c
 	if (bIsOverride)
 	{
 		// The interpolation factor is 1 so we just override the properties.
+		// We do all of them except the FOV/Focal Length, which is done in a special way.
 	
 #define UE_CAMERA_POSE_FOR_PROPERTY(PropType, PropName)\
 		if ((!bInvertMask && InMask.PropName) || (bInvertMask && !InMask.PropName))\
 		{\
-			if (ToPoseChangedFlags.PropName)\
+			if (!bChangedOnly || ToPoseChangedFlags.PropName)\
 			{\
 				Set##PropName(ToPose.Get##PropName());\
-				OutMask.PropName = true;\
 			}\
+			OutMask.PropName = true;\
 		}
 
-		UE_CAMERA_POSE_FOR_ALL_PROPERTIES()
+		UE_CAMERA_POSE_FOR_TRANSFORM_PROPERTIES()
+		UE_CAMERA_POSE_FOR_INTERPOLABLE_PROPERTIES()
+		UE_CAMERA_POSE_FOR_BOOL_PROPERTIES()
 
 #undef UE_CAMERA_POSE_FOR_PROPERTY
 
+		if (
+				(!bInvertMask && (InMask.FieldOfView || InMask.FocalLength)) ||
+				(bInvertMask && (!InMask.FieldOfView || !InMask.FocalLength)))
+		{
+			if (!bChangedOnly || (ToPoseChangedFlags.FieldOfView || ToPoseChangedFlags.FocalLength))
+			{
+				SetFocalLength(ToPose.GetFocalLength());
+				SetFieldOfView(ToPose.GetFieldOfView());
+			}
+		}
 	}
 	else
 	{
@@ -200,7 +243,7 @@ void FCameraPose::InternalLerpChanged(const FCameraPose& ToPose, float Factor, c
 #define UE_CAMERA_POSE_FOR_PROPERTY(PropType, PropName)\
 		if ((!bInvertMask && InMask.PropName) || (bInvertMask && !InMask.PropName))\
 		{\
-			if (ToPoseChangedFlags.PropName)\
+			if (!bChangedOnly || ToPoseChangedFlags.PropName)\
 			{\
 				ensureMsgf(ChangedFlags.PropName, TEXT("Interpolating " #PropName " from default value!"));\
 				Set##PropName(FMath::Lerp(Get##PropName(), ToPose.Get##PropName(), Factor));\
@@ -220,14 +263,40 @@ void FCameraPose::InternalLerpChanged(const FCameraPose& ToPose, float Factor, c
 				(bInvertMask && (!InMask.FieldOfView || !InMask.FocalLength)))
 		{
 			ensureMsgf(
-					ChangedFlags.FieldOfView || ChangedFlags.FocalLength, 
-					TEXT("Interpolating FieldOfView or FocalLength from default value!"));
-			const float FromFieldOfView = GetEffectiveFieldOfView();
-			const float ToFieldOfView = ToPose.GetEffectiveFieldOfView();
-			SetFieldOfView(FMath::Lerp(FromFieldOfView, ToFieldOfView, Factor));
-			SetFocalLength(-1);
-			OutMask.FieldOfView = true;
-			OutMask.FocalLength = true;
+					(FocalLength <= 0 || FieldOfView <= 0) &&
+					(ToPose.FocalLength <= 0 || ToPose.FieldOfView <= 0),
+					TEXT("Can't specify both FocalLength and FieldOfView on a camera pose!"));
+
+			if (!bChangedOnly || (ToPoseChangedFlags.FocalLength || ToPoseChangedFlags.FieldOfView))
+			{
+				ensureMsgf(
+						ChangedFlags.FieldOfView || ChangedFlags.FocalLength, 
+						TEXT("Interpolating FieldOfView or FocalLength from default value!"));
+
+				// Interpolate FocalLength, or FieldOfView, if both poses use the same.
+				// If there's a mix, interpolate the effective FieldOfView.
+				//
+				// We realize that linearly interpolating FocalLength won't linearly interpolate
+				// the effective FOV, so this will actually behave differently between the two
+				// code branches, but this also means that an "all proper" camera setup will
+				// enjoy more realistic camera behavior.
+				//
+				const float FromFocalLength = GetFocalLength();
+				const float ToFocalLength = ToPose.GetFocalLength();
+				if (FromFocalLength > 0 && ToFocalLength > 0)
+				{
+					SetFocalLength(FMath::Lerp(FromFocalLength, ToFocalLength, Factor));
+				}
+				else // only FieldOfView is specified on both, or we have a mix.
+				{
+					const float FromFieldOfView = GetEffectiveFieldOfView();
+					const float ToFieldOfView = ToPose.GetEffectiveFieldOfView();
+					SetFieldOfView(FMath::Lerp(FromFieldOfView, ToFieldOfView, Factor));
+					SetFocalLength(-1);
+				}
+				OutMask.FieldOfView = true;
+				OutMask.FocalLength = true;
+			}
 		}
 
 		// Last, do booleans, which just flip their value once we reach 50% interpolation.
@@ -235,7 +304,7 @@ void FCameraPose::InternalLerpChanged(const FCameraPose& ToPose, float Factor, c
 #define UE_CAMERA_POSE_FOR_PROPERTY(PropType, PropName)\
 		if ((!bInvertMask && InMask.PropName) || (bInvertMask && !InMask.PropName))\
 		{\
-			if (ToPoseChangedFlags.PropName && Factor >= 0.5f)\
+			if (!bChangedOnly || ToPoseChangedFlags.PropName && Factor >= 0.5f)\
 			{\
 				ensureMsgf(ChangedFlags.PropName, TEXT("Interpolating " #PropName " from default value!"));\
 				Set##PropName(ToPose.Get##PropName());\
diff --git a/Engine/Plugins/Cameras/GameplayCameras/Source/GameplayCameras/Private/Core/CameraVariableTable.cpp b/Engine/Plugins/Cameras/GameplayCameras/Source/GameplayCameras/Private/Core/CameraVariableTable.cpp
--- a/Engine/Plugins/Cameras/GameplayCameras/Source/GameplayCameras/Private/Core/CameraVariableTable.cpp
+++ b/Engine/Plugins/Cameras/GameplayCameras/Source/GameplayCameras/Private/Core/CameraVariableTable.cpp
@@ -303,6 +303,11 @@ void FCameraVariableTable::InternalOverride(const FCameraVariableTable& OtherTab
 	}
 }
 
+void FCameraVariableTable::LerpAll(const FCameraVariableTable& ToTable, float Factor)
+{
+	InternalLerp(ToTable, Factor, nullptr, false, nullptr, false);
+}
+
 void FCameraVariableTable::LerpChanged(const FCameraVariableTable& ToTable, float Factor)
 {
 	InternalLerp(ToTable, Factor, nullptr, false, nullptr, true);
diff --git a/Engine/Plugins/Cameras/GameplayCameras/Source/GameplayCameras/Private/Nodes/Blends/SimpleBlendCameraNode.cpp b/Engine/Plugins/Cameras/GameplayCameras/Source/GameplayCameras/Private/Nodes/Blends/SimpleBlendCameraNode.cpp
--- a/Engine/Plugins/Cameras/GameplayCameras/Source/GameplayCameras/Private/Nodes/Blends/SimpleBlendCameraNode.cpp
+++ b/Engine/Plugins/Cameras/GameplayCameras/Source/GameplayCameras/Private/Nodes/Blends/SimpleBlendCameraNode.cpp
@@ -31,14 +31,9 @@ void FSimpleBlendCameraNodeEvaluator::OnBlendResults(const FCameraNodeBlendParam
 	const FCameraNodeEvaluationResult& ChildResult(Params.ChildResult);
 	FCameraNodeEvaluationResult& BlendedResult(OutResult.BlendedResult);
 
-	BlendedResult.CameraPose.LerpChanged(
-			ChildResult.CameraPose, 
-			BlendFactor,
-			ChildResult.CameraPose.GetChangedFlags(), // or is it BlendedResult flags?
-			false,
-			BlendedResult.CameraPose.GetChangedFlags());
-
-	BlendedResult.VariableTable.LerpChanged(ChildResult.VariableTable, BlendFactor);
+	// Blend all properties.
+	BlendedResult.CameraPose.LerpAll(ChildResult.CameraPose, BlendFactor);
+	BlendedResult.VariableTable.LerpAll(ChildResult.VariableTable, BlendFactor);
 
 	// If we have even a fraction of a camera cut, we need to make the
 	// whole result into a camera cut.
@@ -63,6 +58,7 @@ void FSimpleBlendCameraDebugBlock::OnDebugDraw(const FCameraDebugBlockDrawParams
 {
 	Renderer.AddText(TEXT("blend %.2f%%"), BlendFactor * 100.f);
 }
+
 #endif  // UE_GAMEPLAY_CAMERAS_DEBUG
 
 UE_DEFINE_BLEND_CAMERA_NODE_EVALUATOR(FSimpleFixedTimeBlendCameraNodeEvaluator)
diff --git a/Engine/Plugins/Cameras/GameplayCameras/Source/GameplayCameras/Public/Core/CameraPose.h b/Engine/Plugins/Cameras/GameplayCameras/Source/GameplayCameras/Public/Core/CameraPose.h
--- a/Engine/Plugins/Cameras/GameplayCameras/Source/GameplayCameras/Public/Core/CameraPose.h
+++ b/Engine/Plugins/Cameras/GameplayCameras/Source/GameplayCameras/Public/Core/CameraPose.h
@@ -50,6 +50,14 @@ UE_CAMERA_POSE_FOR_ALL_PROPERTIES()
 
 public:
 
+	/** Returns a structure where all flags are set. */
+	static const FCameraPoseFlags& All();
+
+	/** Creates a new flags structure. */
+	FCameraPoseFlags();
+	/** Creates a new flags structure with all flags set to the given value. */
+	FCameraPoseFlags(bool bInValue);
+
 	/** Sets all flags to the given value. */
 	FCameraPoseFlags& SetAllFlags(bool bInValue);
 	/** Sets the flags that are set in OtherFlags, but checks that no flag is set on both structures. */
@@ -140,8 +148,12 @@ public:
 
 	// Interpolation
 	
+	/** Takes all properties from OtherPose and sets them on this camera pose. */
+	void OverrideAll(const FCameraPose& OtherPose);
 	/** Takes all changed properties from OtherPose and sets them on this camera pose. */
 	void OverrideChanged(const FCameraPose& OtherPose);
+	/** Interpolates all properties from ToPose using the given factor. */
+	void LerpAll(const FCameraPose& ToPose, float Factor);
 	/** Interpolates all changed properties from ToPose using the given factor. */
 	void LerpChanged(const FCameraPose& ToPose, float Factor);
 	/** Interpolates changed properties from ToPose using the given factor. Only properties defined by InMask are taken into account. */
@@ -150,6 +162,8 @@ public:
 private:
 
 	void InternalLerpChanged(const FCameraPose& ToPose, float Factor, const FCameraPoseFlags& InMask, bool bInvertMask, FCameraPoseFlags& OutMask);
+	void InternalOverrideChanged(const FCameraPose& OtherPose, bool bChangedOnly);
+	void InternalLerpChanged(const FCameraPose& ToPose, float Factor, const FCameraPoseFlags& InMask, bool bInvertMask, FCameraPoseFlags& OutMask, bool bChangedOnly);
 
 private:
 
diff --git a/Engine/Plugins/Cameras/GameplayCameras/Source/GameplayCameras/Public/Core/CameraVariableTable.h b/Engine/Plugins/Cameras/GameplayCameras/Source/GameplayCameras/Public/Core/CameraVariableTable.h
--- a/Engine/Plugins/Cameras/GameplayCameras/Source/GameplayCameras/Public/Core/CameraVariableTable.h
+++ b/Engine/Plugins/Cameras/GameplayCameras/Source/GameplayCameras/Public/Core/CameraVariableTable.h
@@ -110,6 +110,7 @@ public:
 	void OverrideChanged(const FCameraVariableTable& OtherTable);
 	void OverrideChanged(const FCameraVariableTable& OtherTable, const FCameraVariableTableFlags& InMask, bool bInvertMask, FCameraVariableTableFlags& OutMask);
 
+	void LerpAll(const FCameraVariableTable& ToTable, float Factor);
 	void LerpChanged(const FCameraVariableTable& ToTable, float Factor);
 	void LerpChanged(const FCameraVariableTable& ToTable, float Factor, const FCameraVariableTableFlags& InMask, bool bInvertMask, FCameraVariableTableFlags& OutMask);
 
