Cameras: add support for debug info replay via Rewind Debugger

The main changes include the following:

- Add the Unreal Insights, Trace Services, and Rewind Debugger infrastructure needed to record and replay camera system evaluation data.
- This data is the serialized bytes of all the debug blocks generated for a given frame.
- This requires all debug blocks to be serializable. Unused support for this already half-existed before, but this change finishes its implementation.
- Showing debug info in an editor viewport requires drawing to an FCanvas, and not a UCanvas, so some of the rendering code has been changed accordingly.

Miscellaneous changes:

- Tweak the main toolbar of the debugger panel, including a better button for enabling/disabling debug info.

diff --git a/Engine/Plugins/Cameras/GameplayCameras/Source/GameplayCameras/GameplayCameras.Build.cs b/Engine/Plugins/Cameras/GameplayCameras/Source/GameplayCameras/GameplayCameras.Build.cs
--- a/Engine/Plugins/Cameras/GameplayCameras/Source/GameplayCameras/GameplayCameras.Build.cs
+++ b/Engine/Plugins/Cameras/GameplayCameras/Source/GameplayCameras/GameplayCameras.Build.cs
@@ -17,7 +17,8 @@ public class GameplayCameras : ModuleRules
 				"HeadMountedDisplay",
 				"MovieScene",
 				"MovieSceneTracks",
-				"TemplateSequence"
+				"TemplateSequence",
+				"TraceLog"
 			}
 		);
 	}
diff --git a/Engine/Plugins/Cameras/GameplayCameras/Source/GameplayCameras/Private/Core/BlendStackCameraNode.cpp b/Engine/Plugins/Cameras/GameplayCameras/Source/GameplayCameras/Private/Core/BlendStackCameraNode.cpp
--- a/Engine/Plugins/Cameras/GameplayCameras/Source/GameplayCameras/Private/Core/BlendStackCameraNode.cpp
+++ b/Engine/Plugins/Cameras/GameplayCameras/Source/GameplayCameras/Private/Core/BlendStackCameraNode.cpp
@@ -490,6 +490,11 @@ void FBlendStackSummaryCameraDebugBlock::OnDebugDraw(const FCameraDebugBlockDraw
 	Renderer.AddText(TEXT("%d entries"), NumEntries);
 }
 
+void FBlendStackSummaryCameraDebugBlock::OnSerialize(FArchive& Ar)
+{
+	Ar << NumEntries;
+}
+
 UE_DEFINE_CAMERA_DEBUG_BLOCK(FBlendStackCameraDebugBlock);
 
 FBlendStackCameraDebugBlock::FBlendStackCameraDebugBlock()
@@ -528,6 +533,17 @@ void FBlendStackCameraDebugBlock::OnDebugDraw(const FCameraDebugBlockDrawParams&
 	Renderer.SkipAllBlocks();
 }
 
+void FBlendStackCameraDebugBlock::OnSerialize(FArchive& Ar)
+{
+	Ar << Entries;
+}
+
+FArchive& operator<< (FArchive& Ar, FBlendStackCameraDebugBlock::FEntryDebugInfo& EntryDebugInfo)
+{
+	Ar << EntryDebugInfo.CameraRigName;
+	return Ar;
+}
+
 #endif  // UE_GAMEPLAY_CAMERAS_DEBUG
 
 }  // namespace UE::Cameras
diff --git a/Engine/Plugins/Cameras/GameplayCameras/Source/GameplayCameras/Private/Core/BlendStackRootCameraNode.cpp b/Engine/Plugins/Cameras/GameplayCameras/Source/GameplayCameras/Private/Core/BlendStackRootCameraNode.cpp
--- a/Engine/Plugins/Cameras/GameplayCameras/Source/GameplayCameras/Private/Core/BlendStackRootCameraNode.cpp
+++ b/Engine/Plugins/Cameras/GameplayCameras/Source/GameplayCameras/Private/Core/BlendStackRootCameraNode.cpp
@@ -39,7 +39,7 @@ UE_DECLARE_CAMERA_DEBUG_BLOCK_START(FBlendStackRootCameraDebugBlock)
 	UE_DECLARE_CAMERA_DEBUG_BLOCK_FIELD(FString, CameraRigAssetName)
 UE_DECLARE_CAMERA_DEBUG_BLOCK_END()
 
-UE_DEFINE_CAMERA_DEBUG_BLOCK(FBlendStackRootCameraDebugBlock)
+UE_DEFINE_CAMERA_DEBUG_BLOCK_WITH_FIELDS(FBlendStackRootCameraDebugBlock)
 
 FCameraNodeEvaluatorChildrenView FBlendStackRootCameraNodeEvaluator::OnGetChildren()
 {
diff --git a/Engine/Plugins/Cameras/GameplayCameras/Source/GameplayCameras/Private/Core/CameraNodeEvaluatorStorage.cpp b/Engine/Plugins/Cameras/GameplayCameras/Source/GameplayCameras/Private/Core/CameraNodeEvaluatorStorage.cpp
--- a/Engine/Plugins/Cameras/GameplayCameras/Source/GameplayCameras/Private/Core/CameraNodeEvaluatorStorage.cpp
+++ b/Engine/Plugins/Cameras/GameplayCameras/Source/GameplayCameras/Private/Core/CameraNodeEvaluatorStorage.cpp
@@ -43,8 +43,8 @@ void FCameraNodeEvaluatorStorage::DestroyEvaluatorTree(bool bFreeAllocations)
 
 void FCameraNodeEvaluatorStorage::GetAllocationInfo(FCameraNodeEvaluatorAllocationInfo& OutAllocationInfo)
 {
-	uint16 TotalUsed;
-	uint16 FirstAlignment;
+	uint32 TotalUsed;
+	uint32 FirstAlignment;
 	Super::GetAllocationInfo(TotalUsed, FirstAlignment);
 
 	OutAllocationInfo.TotalSizeof = TotalUsed;
diff --git a/Engine/Plugins/Cameras/GameplayCameras/Source/GameplayCameras/Private/Core/CameraPose.cpp b/Engine/Plugins/Cameras/GameplayCameras/Source/GameplayCameras/Private/Core/CameraPose.cpp
--- a/Engine/Plugins/Cameras/GameplayCameras/Source/GameplayCameras/Private/Core/CameraPose.cpp
+++ b/Engine/Plugins/Cameras/GameplayCameras/Source/GameplayCameras/Private/Core/CameraPose.cpp
@@ -319,3 +319,18 @@ void FCameraPose::InternalLerpChanged(const FCameraPose& ToPose, float Factor, c
 	}
 }
 
+void FCameraPose::SerializeWithFlags(FArchive& Ar, FCameraPose& CameraPose)
+{
+	static FCameraPose DefaultCameraPose;
+
+	UScriptStruct* ClassStruct = FCameraPose::StaticStruct();
+	ClassStruct->SerializeItem(Ar, &CameraPose, &DefaultCameraPose);
+
+#define UE_CAMERA_POSE_FOR_PROPERTY(PropType, PropName)\
+	Ar << CameraPose.ChangedFlags.PropName;
+
+UE_CAMERA_POSE_FOR_ALL_PROPERTIES()
+
+#undef UE_CAMERA_POSE_FOR_PROPERTY
+}
+
diff --git a/Engine/Plugins/Cameras/GameplayCameras/Source/GameplayCameras/Private/Core/CameraSystemEvaluator.cpp b/Engine/Plugins/Cameras/GameplayCameras/Source/GameplayCameras/Private/Core/CameraSystemEvaluator.cpp
--- a/Engine/Plugins/Cameras/GameplayCameras/Source/GameplayCameras/Private/Core/CameraSystemEvaluator.cpp
+++ b/Engine/Plugins/Cameras/GameplayCameras/Source/GameplayCameras/Private/Core/CameraSystemEvaluator.cpp
@@ -11,6 +11,7 @@
 #include "Debug/CameraDebugBlock.h"
 #include "Debug/CameraDebugBlockBuilder.h"
 #include "Debug/CameraDebugRenderer.h"
+#include "Debug/CameraSystemTrace.h"
 #include "Debug/CategoryTitleDebugBlock.h"
 #include "Debug/RootCameraDebugBlock.h"
 #include "HAL/IConsoleManager.h"
@@ -24,7 +25,6 @@ namespace UE::Cameras
 {
 
 extern bool GGameplayCamerasDebugEnable;
-extern bool GGameplayCamerasDebugTrace;
 
 FCameraSystemEvaluator::FCameraSystemEvaluator()
 {
@@ -44,6 +44,7 @@ void FCameraSystemEvaluator::Initialize(const FCameraSystemEvaluatorCreateParams
 	{
 		Owner = GetTransientPackage();
 	}
+	WeakOwner = Owner;
 
 	if (Params.RootNodeFactory)
 	{
@@ -153,11 +154,23 @@ void FCameraSystemEvaluator::GetEvaluatedCameraView(FMinimalViewInfo& DesiredVie
 
 void FCameraSystemEvaluator::DebugUpdate(const FCameraSystemDebugUpdateParams& Params)
 {
-	if (!GGameplayCamerasDebugTrace && !GGameplayCamerasDebugEnable)
+#if UE_GAMEPLAY_CAMERAS_TRACE
+	const bool bTraceEnabled = FCameraSystemTrace::IsTraceEnabled();
+#else
+	const bool bTraceEnabled = false;
+#endif  // UE_GAMEPLAY_CAMERAS_TRACE
+	if (!bTraceEnabled && !GGameplayCamerasDebugEnable)
 	{
 		return;
 	}
 
+#if UE_GAMEPLAY_CAMERAS_TRACE
+	if (FCameraSystemTrace::IsTraceReplay())
+	{
+		return;
+	}
+#endif  // UE_GAMEPLAY_CAMERAS_TRACE
+
 	// Clear previous frame's debug info and make room for this frame's.
 	DebugBlockStorage.DestroyDebugBlocks();
 
@@ -167,6 +180,15 @@ void FCameraSystemEvaluator::DebugUpdate(const FCameraSystemDebugUpdateParams& P
 	FCameraDebugBlockBuildParams BuildParams;
 	FCameraDebugBlockBuilder DebugBlockBuilder(DebugBlockStorage, *RootDebugBlock);
 	RootDebugBlock->BuildDebugBlocks(*this, BuildParams, DebugBlockBuilder);
+
+#if UE_GAMEPLAY_CAMERAS_TRACE
+	if (bTraceEnabled)
+	{
+		UObject* Owner = WeakOwner.Get();
+		UWorld* OwnerWorld = Owner ? Owner->GetWorld() : nullptr;
+		FCameraSystemTrace::TraceEvaluation(OwnerWorld, Result, *RootDebugBlock);
+	}
+#endif
 	
 	FCameraDebugRenderer Renderer(Params.Canvas);
 	RootDebugBlock->RootDebugDraw(Renderer);
diff --git a/Engine/Plugins/Cameras/GameplayCameras/Source/GameplayCameras/Private/Core/DefaultRootCameraNode.cpp b/Engine/Plugins/Cameras/GameplayCameras/Source/GameplayCameras/Private/Core/DefaultRootCameraNode.cpp
--- a/Engine/Plugins/Cameras/GameplayCameras/Source/GameplayCameras/Private/Core/DefaultRootCameraNode.cpp
+++ b/Engine/Plugins/Cameras/GameplayCameras/Source/GameplayCameras/Private/Core/DefaultRootCameraNode.cpp
@@ -104,7 +104,7 @@ void FDefaultRootCameraNodeEvaluator::OnActivateCameraRig(const FActivateCameraR
 UE_DECLARE_CAMERA_DEBUG_BLOCK_START(FDefaultRootCameraNodeEvaluatorDebugBlock)
 UE_DECLARE_CAMERA_DEBUG_BLOCK_END()
 
-UE_DEFINE_CAMERA_DEBUG_BLOCK(FDefaultRootCameraNodeEvaluatorDebugBlock)
+UE_DEFINE_CAMERA_DEBUG_BLOCK_WITH_FIELDS(FDefaultRootCameraNodeEvaluatorDebugBlock)
 
 void FDefaultRootCameraNodeEvaluator::OnBuildDebugBlocks(const FCameraDebugBlockBuildParams& Params, FCameraDebugBlockBuilder& Builder)
 {
diff --git a/Engine/Plugins/Cameras/GameplayCameras/Source/GameplayCameras/Private/Debug/BlendStacksCameraDebugBlock.cpp b/Engine/Plugins/Cameras/GameplayCameras/Source/GameplayCameras/Private/Debug/BlendStacksCameraDebugBlock.cpp
--- a/Engine/Plugins/Cameras/GameplayCameras/Source/GameplayCameras/Private/Debug/BlendStacksCameraDebugBlock.cpp
+++ b/Engine/Plugins/Cameras/GameplayCameras/Source/GameplayCameras/Private/Debug/BlendStacksCameraDebugBlock.cpp
@@ -66,6 +66,11 @@ void FBlendStacksCameraDebugBlock::OnDebugDraw(const FCameraDebugBlockDrawParams
 	Renderer.SkipAllBlocks();
 }
 
+void FBlendStacksCameraDebugBlock::OnSerialize(FArchive& Ar)
+{
+	Ar << BlendStackNames;
+}
+
 }  // namespace UE::Cameras
 
 #endif  // UE_GAMEPLAY_CAMERAS_DEBUG
diff --git a/Engine/Plugins/Cameras/GameplayCameras/Source/GameplayCameras/Private/Debug/CameraDebugBlock.cpp b/Engine/Plugins/Cameras/GameplayCameras/Source/GameplayCameras/Private/Debug/CameraDebugBlock.cpp
--- a/Engine/Plugins/Cameras/GameplayCameras/Source/GameplayCameras/Private/Debug/CameraDebugBlock.cpp
+++ b/Engine/Plugins/Cameras/GameplayCameras/Source/GameplayCameras/Private/Debug/CameraDebugBlock.cpp
@@ -16,8 +16,6 @@ bool FCameraDebugBlockDrawParams::IsCategoryActive(const FString& InCategory) co
 
 UE_GAMEPLAY_CAMERAS_DEFINE_RTTI(FCameraDebugBlock)
 
-TArray<FCameraDebugBlockField*> FCameraDebugBlock::StaticFields;
-
 void FCameraDebugBlock::Attach(FCameraDebugBlock* InAttachment)
 {
 	Attachments.Add(InAttachment);
@@ -37,6 +35,7 @@ void FCameraDebugBlock::DebugDraw(const FCameraDebugBlockDrawParams& Params, FCa
 	
 	ECameraDebugDrawVisitFlags VisitFlags = Renderer.GetVisitFlags();
 
+	// Attachments can render on the same line as this debug block so we call DebugDraw directly on them.
 	if (!EnumHasAnyFlags(VisitFlags, ECameraDebugDrawVisitFlags::SkipAttachedBlocks) && !Attachments.IsEmpty())
 	{
 		for (FCameraDebugBlock* Attachment : Attachments)
@@ -45,6 +44,11 @@ void FCameraDebugBlock::DebugDraw(const FCameraDebugBlockDrawParams& Params, FCa
 		}
 	}
 
+	// Children should always render on lines below so we need to make sure we're on a new line for
+	// the remainder of this function. The call to AddIndent() below will flush any pending text and 
+	// add this new line automatically, but we need to add the optional new line ourselves if we
+	// happen to skip this next section.
+	//
 	if (!EnumHasAnyFlags(VisitFlags, ECameraDebugDrawVisitFlags::SkipChildrenBlocks) && !Children.IsEmpty())
 	{
 		Renderer.AddIndent();
@@ -56,6 +60,10 @@ void FCameraDebugBlock::DebugDraw(const FCameraDebugBlockDrawParams& Params, FCa
 
 		Renderer.RemoveIndent();
 	}
+	else
+	{
+		Renderer.NewLine(true);
+	}
 
 	OnPostDebugDraw(Params, Renderer);
 }
@@ -63,11 +71,6 @@ void FCameraDebugBlock::DebugDraw(const FCameraDebugBlockDrawParams& Params, FCa
 void FCameraDebugBlock::Serialize(FArchive& Ar)
 {
 	OnSerialize(Ar);
-
-	for (FCameraDebugBlockField* StaticField : StaticFields)
-	{
-		StaticField->SerializeField(this, Ar);
-	}
 }
 
 }  // namespace UE::Cameras
diff --git a/Engine/Plugins/Cameras/GameplayCameras/Source/GameplayCameras/Private/Debug/CameraDebugRenderer.cpp b/Engine/Plugins/Cameras/GameplayCameras/Source/GameplayCameras/Private/Debug/CameraDebugRenderer.cpp
--- a/Engine/Plugins/Cameras/GameplayCameras/Source/GameplayCameras/Private/Debug/CameraDebugRenderer.cpp
+++ b/Engine/Plugins/Cameras/GameplayCameras/Source/GameplayCameras/Private/Debug/CameraDebugRenderer.cpp
@@ -3,9 +3,10 @@
 #include "Debug/CameraDebugRenderer.h"
 
 #include "Algo/Find.h"
+#include "CanvasItem.h"
+#include "CanvasTypes.h"
 #include "Debug/CameraDebugColors.h"
 #include "Debug/DebugTextRenderer.h"
-#include "Engine/Canvas.h"
 #include "Engine/Engine.h"
 #include "Engine/Font.h"
 #include "HAL/IConsoleManager.h"
@@ -34,12 +35,10 @@ static FAutoConsoleVariableRef CVarGameplayCamerasDebugBackgroundDepthSortKey(
 	GGameplayCamerasDebugBackgroundDepthSortKey,
 	TEXT(""));
 
-FCameraDebugRenderer::FCameraDebugRenderer(UCanvas* InCanvas)
+FCameraDebugRenderer::FCameraDebugRenderer(FCanvas* InCanvas)
 	: Canvas(InCanvas)
+	, DrawColor(FColor::White)
 {
-	OriginalDrawColor = Canvas->DrawColor;
-	Canvas->SetDrawColor(FColor::White);
-
 	RenderFont = GEngine->GetSmallFont();
 	MaxCharHeight = RenderFont->GetMaxCharHeight();
 
@@ -49,7 +48,12 @@ FCameraDebugRenderer::FCameraDebugRenderer(UCanvas* InCanvas)
 FCameraDebugRenderer::~FCameraDebugRenderer()
 {
 	FlushText();
-	Canvas->SetDrawColor(OriginalDrawColor);
+}
+
+FVector2D FCameraDebugRenderer::GetCanvasSize() const
+{
+	FIntPoint ParentSize = Canvas->GetParentCanvasSize();
+	return FVector2D(ParentSize.X, ParentSize.Y);
 }
 
 void FCameraDebugRenderer::AddText(const FString& InString)
@@ -95,14 +99,14 @@ bool FCameraDebugRenderer::NewLine(bool bSkipIfEmptyLine)
 
 FColor FCameraDebugRenderer::GetTextColor() const
 {
-	return Canvas->DrawColor;
+	return DrawColor;
 }
 
 FColor FCameraDebugRenderer::SetTextColor(const FColor& Color)
 {
 	FlushText();
-	FColor ReturnColor = Canvas->DrawColor;
-	Canvas->SetDrawColor(Color);
+	FColor ReturnColor = DrawColor;
+	DrawColor = Color;
 	return ReturnColor;
 }
 
@@ -115,9 +119,10 @@ void FCameraDebugRenderer::FlushText()
 {
 	if (LineBuilder.Len() > 0)
 	{
-		if (NextDrawPosition.Y < Canvas->ClipY)
+		int32 ViewHeight = GetCanvasSize().Y;
+		if (NextDrawPosition.Y < ViewHeight)
 		{
-			FDebugTextRenderer TextRenderer(Canvas, RenderFont);
+			FDebugTextRenderer TextRenderer(Canvas, DrawColor, RenderFont);
 			TextRenderer.LeftMargin = GetIndentMargin();
 			TextRenderer.RenderText(NextDrawPosition, LineBuilder.ToView());
 
@@ -172,13 +177,16 @@ void FCameraDebugRenderer::DrawTextBackgroundTile(float Opacity)
 	const FColor BackgroundColor = FCameraDebugColors::Get().Background.WithAlpha((uint8)(Opacity * 255));
 
 	// Draw the background behind the text.
-	Canvas->Canvas->PushDepthSortKey(GGameplayCamerasDebugBackgroundDepthSortKey);
+	if (Canvas)
 	{
-		FCanvasTileItem BackgroundTile(TopLeft, TileSize, BackgroundColor);
-		BackgroundTile.BlendMode = SE_BLEND_Translucent;
-		Canvas->DrawItem(BackgroundTile);
+		Canvas->PushDepthSortKey(GGameplayCamerasDebugBackgroundDepthSortKey);
+		{
+			FCanvasTileItem BackgroundTile(TopLeft, TileSize, BackgroundColor);
+			BackgroundTile.BlendMode = SE_BLEND_Translucent;
+			Canvas->DrawItem(BackgroundTile);
+		}
+		Canvas->PopDepthSortKey();
 	}
-	Canvas->Canvas->PopDepthSortKey();
 }
 
 void FCameraDebugRenderer::SkipAttachedBlocks()
diff --git a/Engine/Plugins/Cameras/GameplayCameras/Source/GameplayCameras/Private/Debug/CameraDirectorTreeDebugBlock.cpp b/Engine/Plugins/Cameras/GameplayCameras/Source/GameplayCameras/Private/Debug/CameraDirectorTreeDebugBlock.cpp
--- a/Engine/Plugins/Cameras/GameplayCameras/Source/GameplayCameras/Private/Debug/CameraDirectorTreeDebugBlock.cpp
+++ b/Engine/Plugins/Cameras/GameplayCameras/Source/GameplayCameras/Private/Debug/CameraDirectorTreeDebugBlock.cpp
@@ -88,6 +88,17 @@ void FCameraDirectorTreeDebugBlock::OnDebugDraw(const FCameraDebugBlockDrawParam
 	Renderer.SkipAllBlocks();
 }
 
+void FCameraDirectorTreeDebugBlock::OnSerialize(FArchive& Ar)
+{
+	 Ar << CameraDirectors;
+}
+
+FArchive& operator<< (FArchive& Ar, FCameraDirectorTreeDebugBlock::FDirectorDebugInfo& DirectorDebugInfo)
+{
+	Ar << DirectorDebugInfo.CameraAssetName;
+	return Ar;
+}
+
 }  // namespace UE::Cameras
 
 #endif  // UE_GAMEPLAY_CAMERAS_DEBUG
diff --git a/Engine/Plugins/Cameras/GameplayCameras/Source/GameplayCameras/Private/Debug/CameraNodeEvaluationResultDebugBlock.cpp b/Engine/Plugins/Cameras/GameplayCameras/Source/GameplayCameras/Private/Debug/CameraNodeEvaluationResultDebugBlock.cpp
--- a/Engine/Plugins/Cameras/GameplayCameras/Source/GameplayCameras/Private/Debug/CameraNodeEvaluationResultDebugBlock.cpp
+++ b/Engine/Plugins/Cameras/GameplayCameras/Source/GameplayCameras/Private/Debug/CameraNodeEvaluationResultDebugBlock.cpp
@@ -94,6 +94,12 @@ void FCameraNodeEvaluationResultDebugBlock::OnDebugDraw(const FCameraDebugBlockD
 	Renderer.SkipAllBlocks();
 }
 
+void FCameraNodeEvaluationResultDebugBlock::OnSerialize(FArchive& Ar)
+{
+	Ar << bIsCameraCut;
+	Ar << bIsValid;
+}
+
 }  // namespace UE::Cameras
 
 #endif  // UE_GAMEPLAY_CAMERAS_DEBUG
diff --git a/Engine/Plugins/Cameras/GameplayCameras/Source/GameplayCameras/Private/Debug/CameraPoseDebugBlock.cpp b/Engine/Plugins/Cameras/GameplayCameras/Source/GameplayCameras/Private/Debug/CameraPoseDebugBlock.cpp
--- a/Engine/Plugins/Cameras/GameplayCameras/Source/GameplayCameras/Private/Debug/CameraPoseDebugBlock.cpp
+++ b/Engine/Plugins/Cameras/GameplayCameras/Source/GameplayCameras/Private/Debug/CameraPoseDebugBlock.cpp
@@ -65,6 +65,12 @@ void FCameraPoseDebugBlock::OnDebugDraw(const FCameraDebugBlockDrawParams& Param
 #undef UE_CAMERA_POSE_FOR_PROPERTY
 }
 
+void FCameraPoseDebugBlock::OnSerialize(FArchive& Ar)
+{
+	FCameraPose::SerializeWithFlags(Ar, CameraPose);
+	Ar << ShowUnchangedCVarName;
+}
+
 }  // namespace UE::Cameras
 
 #endif  // UE_GAMEPLAY_CAMERAS_DEBUG
diff --git a/Engine/Plugins/Cameras/GameplayCameras/Source/GameplayCameras/Private/Debug/CameraSystemTrace.cpp b/Engine/Plugins/Cameras/GameplayCameras/Source/GameplayCameras/Private/Debug/CameraSystemTrace.cpp
new file mode 100644
--- /dev/null
+++ b/Engine/Plugins/Cameras/GameplayCameras/Source/GameplayCameras/Private/Debug/CameraSystemTrace.cpp
@@ -0,0 +1,330 @@
+// Copyright Epic Games, Inc. All Rights Reserved.
+
+#include "Debug/CameraSystemTrace.h"
+
+#include "Containers/UnrealString.h"
+#include "Core/CameraSystemEvaluator.h"
+#include "Debug/RootCameraDebugBlock.h"
+#include "Math/Rotator.h"
+#include "Math/Vector.h"
+#include "ObjectTrace.h"
+#include "Serialization/BufferArchive.h"
+#include "Serialization/MemoryReader.h"
+
+#if UE_GAMEPLAY_CAMERAS_TRACE
+
+namespace UE::Cameras
+{
+
+bool GGameplayCamerasDebugTrace = false;
+static FAutoConsoleVariableRef CVarGameplayCamerasDebugTrace(
+	TEXT("GameplayCameras.Debug.Trace"),
+	GGameplayCamerasDebugTrace,
+	TEXT("(Default: false. Enables background tracing of GamplayCameras system debug info."));
+
+// Channel "CameraSystemChannel".
+UE_TRACE_CHANNEL(CameraSystemChannel)
+
+// Log name "CameraSystem", event name "CameraSystemEvaluation".
+UE_TRACE_EVENT_BEGIN(CameraSystem, CameraSystemEvaluation)
+	UE_TRACE_EVENT_FIELD(uint64, Cycle)
+	UE_TRACE_EVENT_FIELD(double, RecordingTime)
+	UE_TRACE_EVENT_FIELD(double, EvaluatedLocationX)
+	UE_TRACE_EVENT_FIELD(double, EvaluatedLocationY)
+	UE_TRACE_EVENT_FIELD(double, EvaluatedLocationZ)
+	UE_TRACE_EVENT_FIELD(double, EvaluatedRotationPitch)
+	UE_TRACE_EVENT_FIELD(double, EvaluatedRotationYaw)
+	UE_TRACE_EVENT_FIELD(double, EvaluatedRotationRoll)
+	UE_TRACE_EVENT_FIELD(float, EvaluatedFieldOfView)
+	UE_TRACE_EVENT_FIELD(uint8[], SerializedBlocks)
+UE_TRACE_EVENT_END()
+
+class FCameraDebugBlockSerializer
+{
+protected:
+
+	static uint8 TokenBufferStart;
+	static uint8 TokenSerializerVersion;
+	static uint8 TokenBlockStart;
+	static uint8 TokenRelatedIndices;
+	static uint8 TokenBlockEnd;
+	static uint8 TokenBufferEnd;
+};
+
+uint8 FCameraDebugBlockSerializer::TokenBufferStart = 0x11;
+uint8 FCameraDebugBlockSerializer::TokenSerializerVersion = 0x01;
+uint8 FCameraDebugBlockSerializer::TokenBlockStart = 0x22;
+uint8 FCameraDebugBlockSerializer::TokenRelatedIndices = 0x33;
+uint8 FCameraDebugBlockSerializer::TokenBlockEnd = 0x44;
+uint8 FCameraDebugBlockSerializer::TokenBufferEnd = 0x44;
+
+class FCameraDebugBlockWriter : public FCameraDebugBlockSerializer
+{
+public:
+
+	FCameraDebugBlockWriter(FBufferArchive& InArchive)
+		: Archive(InArchive)
+		, TypeRegistry(FCameraObjectTypeRegistry::Get())
+	{}
+	
+	void Write(FCameraDebugBlock& InRootDebugBlock)
+	{
+		check(Archive.IsSaving());
+
+		Archive << TokenBufferStart;
+		Archive << TokenSerializerVersion;
+
+		NextBlockIndex = 0;
+		BlockIndices.Add(&InRootDebugBlock, NextBlockIndex++);
+
+		TArray<FCameraDebugBlock*> WriteStack;
+		WriteStack.Add(&InRootDebugBlock);
+		while (!WriteStack.IsEmpty())
+		{
+			FCameraDebugBlock* CurBlock(WriteStack.Pop());
+			WriteImpl(CurBlock, WriteStack);
+		}
+
+		Archive << TokenBufferEnd;
+
+		ensure(BlockIndices.IsEmpty());
+	}
+
+private:
+	
+	void WriteImpl(FCameraDebugBlock* InBlock, TArray<FCameraDebugBlock*>& InWriteStack)
+	{
+		Archive << TokenBlockStart;
+
+		int32 BlockIndex = BlockIndices.FindChecked(InBlock);
+		Archive << BlockIndex;
+		BlockIndices.Remove(InBlock);
+
+		// We can't serialize type IDs because they're not stable, so we serialize
+		// the type name directly. This isn't very optimized, as it writes lots of
+		// strings in the buffer, but it's good enough as a first implementation.
+		FCameraObjectTypeID BlockTypeID = InBlock->GetTypeID();
+		const FCameraObjectTypeInfo* BlockTypeInfo = TypeRegistry.GetTypeInfo(BlockTypeID);
+		check(BlockTypeInfo);
+		FName BlockTypeName = BlockTypeInfo->TypeName;
+		Archive << BlockTypeName;
+
+		InBlock->Serialize(Archive);
+
+		Archive << TokenRelatedIndices;
+
+		TArray<int32> AttachmentIndices;
+		for (FCameraDebugBlock* Attachment : InBlock->GetAttachments())
+		{
+			int32 AttachmentIndex = BlockIndices.Add(Attachment, NextBlockIndex++);
+			AttachmentIndices.Add(AttachmentIndex);
+			InWriteStack.Add(Attachment);
+		}
+		Archive << AttachmentIndices;
+
+		TArray<int32> ChildrenIndices;
+		for (FCameraDebugBlock* Child : InBlock->GetChildren())
+		{
+			int32 ChildIndex = BlockIndices.Add(Child, NextBlockIndex++);
+			ChildrenIndices.Add(ChildIndex);
+			InWriteStack.Add(Child);
+		}
+		Archive << ChildrenIndices;
+
+		Archive << TokenBlockEnd;
+	}
+
+private:
+
+	FBufferArchive& Archive;
+	FCameraObjectTypeRegistry& TypeRegistry;
+
+	TMap<FCameraDebugBlock*, int32> BlockIndices;
+	int32 NextBlockIndex;
+};
+
+class FCameraDebugBlockReader : public FCameraDebugBlockSerializer
+{
+public:
+
+	FCameraDebugBlockReader(FArchive& InArchive, FCameraDebugBlockStorage& InStorage)
+		: Archive(InArchive)
+		, Storage(InStorage)
+		, TypeRegistry(FCameraObjectTypeRegistry::Get())
+	{
+	}
+
+	FCameraDebugBlock* Read()
+	{
+		check(Archive.IsLoading());
+
+		uint8 Token;
+
+		Archive << Token;
+		check(Token == TokenBufferStart);
+		Archive << Token;
+		check(Token == TokenSerializerVersion);
+		
+		bool bIsRootBlock = true;
+		while (true)
+		{
+			Archive << Token;
+			if (Token == TokenBufferEnd)
+			{
+				break;
+			}
+			else if (Token == TokenBlockStart)
+			{
+				ReadImpl();
+			}
+			else
+			{
+				ensure(false);
+				break;
+			}
+		}
+
+		SetupRelatedBlocks();
+
+		ensure(BlocksByIndex.IsValidIndex(0));
+		return BlocksByIndex[0];
+	}
+
+private:
+
+	void ReadImpl()
+	{
+		int32 BlockIndex;
+		Archive << BlockIndex;
+
+		FName BlockTypeName;
+		Archive << BlockTypeName;
+		check(BlockTypeName != NAME_None);
+
+		FCameraObjectTypeID BlockTypeID = TypeRegistry.FindTypeByName(BlockTypeName);
+		ensure(BlockTypeID.IsValid());
+		const FCameraObjectTypeInfo* BlockTypeInfo = TypeRegistry.GetTypeInfo(BlockTypeID);
+		ensure(BlockTypeInfo);
+
+		void* NewBlockPtr = Storage.BuildDebugBlockUninitialized(BlockTypeInfo->Sizeof, BlockTypeInfo->Alignof);
+		BlockTypeInfo->Constructor(NewBlockPtr);
+
+		// This isn't quite correct for complicated inheritance configurations, but we don't
+		// expect those sorts of setups for debug blocks... hopefully.
+		FCameraDebugBlock* NewBlock = reinterpret_cast<FCameraDebugBlock*>(NewBlockPtr);
+
+		NewBlock->Serialize(Archive);
+
+		uint8 Token;
+
+		Archive << Token;
+		check(Token == TokenRelatedIndices);
+
+		FRelatedIndices CurRelatedIndices;
+		Archive << CurRelatedIndices.AttachmentIndices;
+		Archive << CurRelatedIndices.ChildrenIndices;
+		RelatedIndices.Add(NewBlock, MoveTemp(CurRelatedIndices));
+
+		Archive << Token;
+		check(Token == TokenBlockEnd);
+
+		BlocksByIndex.Insert(BlockIndex, NewBlock);
+	}
+
+	void SetupRelatedBlocks()
+	{
+		for (const TPair<FCameraDebugBlock*, FRelatedIndices>& Pair : RelatedIndices)
+		{
+			FCameraDebugBlock* CurBlock = Pair.Key;
+
+			for (uint32 AttachmentIndex : Pair.Value.AttachmentIndices)
+			{
+				check(BlocksByIndex.IsValidIndex(AttachmentIndex));
+				FCameraDebugBlock* AttachedBlock = BlocksByIndex[AttachmentIndex];
+				CurBlock->Attach(AttachedBlock);
+			}
+
+			for (uint32 ChildIndex : Pair.Value.ChildrenIndices)
+			{
+				check(BlocksByIndex.IsValidIndex(ChildIndex));
+				FCameraDebugBlock* ChildBlock = BlocksByIndex[ChildIndex];
+				CurBlock->AddChild(ChildBlock);
+			}
+		}
+	}
+
+private:
+
+	struct FRelatedIndices
+	{
+		TArray<int32> AttachmentIndices;
+		TArray<int32> ChildrenIndices;
+	};
+
+	FArchive& Archive;
+	FCameraDebugBlockStorage& Storage;
+	FCameraObjectTypeRegistry& TypeRegistry;
+
+	TSparseArray<FCameraDebugBlock*> BlocksByIndex;
+	TMap<FCameraDebugBlock*, FRelatedIndices> RelatedIndices;
+};
+
+// This name must match the one passed to UE_TRACE_CHANNEL above.
+FString FCameraSystemTrace::ChannelName("CameraSystemChannel");
+// These two names must match the names passed to UE_TRACE_EVENT_BEGIN above.
+FString FCameraSystemTrace::LoggerName("CameraSystem");
+FString FCameraSystemTrace::EvaluationEventName("CameraSystemEvaluation");
+
+bool FCameraSystemTrace::bIsReplaying(false);
+
+bool FCameraSystemTrace::IsTraceReplay()
+{
+	return bIsReplaying;
+}
+
+void FCameraSystemTrace::SetTraceReplay(bool bInIsReplaying)
+{
+	bIsReplaying = bInIsReplaying;
+}
+
+bool FCameraSystemTrace::IsTraceEnabled()
+{
+	return GGameplayCamerasDebugTrace || UE_TRACE_CHANNELEXPR_IS_ENABLED(CameraSystemChannel);
+}
+
+void FCameraSystemTrace::TraceEvaluation(UWorld* InWorld, const FCameraSystemEvaluationUpdateResult& InResult, FCameraDebugBlock& InRootDebugBlock)
+{
+	if (!IsTraceEnabled())
+	{
+		return;
+	}
+
+	FBufferArchive BufferArchive;
+	FCameraDebugBlockWriter Writer(BufferArchive);
+	Writer.Write(InRootDebugBlock);
+
+	// Names must match LoggerName, EvaluationEventName, ChannelName.
+	UE_TRACE_LOG(CameraSystem, CameraSystemEvaluation, CameraSystemChannel)
+		<< CameraSystemEvaluation.Cycle(FPlatformTime::Cycles64())
+		<< CameraSystemEvaluation.RecordingTime(FObjectTrace::GetWorldElapsedTime(InWorld))
+		<< CameraSystemEvaluation.EvaluatedLocationX(InResult.CameraPose.GetLocation().X)
+		<< CameraSystemEvaluation.EvaluatedLocationY(InResult.CameraPose.GetLocation().Y)
+		<< CameraSystemEvaluation.EvaluatedLocationZ(InResult.CameraPose.GetLocation().Z)
+		<< CameraSystemEvaluation.EvaluatedRotationYaw(InResult.CameraPose.GetRotation().Yaw)
+		<< CameraSystemEvaluation.EvaluatedRotationPitch(InResult.CameraPose.GetRotation().Pitch)
+		<< CameraSystemEvaluation.EvaluatedRotationRoll(InResult.CameraPose.GetRotation().Roll)
+		<< CameraSystemEvaluation.EvaluatedFieldOfView(InResult.CameraPose.GetEffectiveFieldOfView())
+    	<< CameraSystemEvaluation.SerializedBlocks(BufferArchive.GetData(), BufferArchive.Num());
+}
+
+FCameraDebugBlock* FCameraSystemTrace::ReadEvaluationTrace(TArray<uint8> InSerializedBlocks, FCameraDebugBlockStorage& InStorage)
+{
+	FMemoryReader MemoryArchive(InSerializedBlocks);
+	FCameraDebugBlockReader Reader(MemoryArchive, InStorage);
+	return Reader.Read();
+}
+
+}  // namespace UE::Cameras
+
+#endif  // UE_GAMEPLAY_CAMERAS_TRACE
+
diff --git a/Engine/Plugins/Cameras/GameplayCameras/Source/GameplayCameras/Private/Debug/CategoryTitleDebugBlock.cpp b/Engine/Plugins/Cameras/GameplayCameras/Source/GameplayCameras/Private/Debug/CategoryTitleDebugBlock.cpp
--- a/Engine/Plugins/Cameras/GameplayCameras/Source/GameplayCameras/Private/Debug/CategoryTitleDebugBlock.cpp
+++ b/Engine/Plugins/Cameras/GameplayCameras/Source/GameplayCameras/Private/Debug/CategoryTitleDebugBlock.cpp
@@ -12,14 +12,12 @@ namespace UE::Cameras
 UE_DEFINE_CAMERA_DEBUG_BLOCK(FCategoryTitleDebugBlock)
 
 FCategoryTitleDebugBlock::FCategoryTitleDebugBlock()
-	: TitleColor(FColor::Green)
 {
 }
 
 FCategoryTitleDebugBlock::FCategoryTitleDebugBlock(const FString& InCategory, const FString& InTitle)
 	: Category(InCategory)
 	, Title(InTitle)
-	, TitleColor(FColor::Green)
 {
 }
 
@@ -45,6 +43,15 @@ void FCategoryTitleDebugBlock::OnDebugDraw(const FCameraDebugBlockDrawParams& Pa
 	}
 }
 
+void FCategoryTitleDebugBlock::OnSerialize(FArchive& Ar)
+{
+	Ar << Category;
+	Ar << Title;
+
+	Ar << bSkipAttachedBlocksIfInactive;
+	Ar << bSkipChildrenBlocksIfInactive;
+}
+
 }  // namespace UE::Cameras
 
 #endif  // UE_GAMEPLAY_CAMERAS_DEBUG
diff --git a/Engine/Plugins/Cameras/GameplayCameras/Source/GameplayCameras/Private/Debug/DebugTextRenderer.cpp b/Engine/Plugins/Cameras/GameplayCameras/Source/GameplayCameras/Private/Debug/DebugTextRenderer.cpp
--- a/Engine/Plugins/Cameras/GameplayCameras/Source/GameplayCameras/Private/Debug/DebugTextRenderer.cpp
+++ b/Engine/Plugins/Cameras/GameplayCameras/Source/GameplayCameras/Private/Debug/DebugTextRenderer.cpp
@@ -3,8 +3,8 @@
 #include "Debug/DebugTextRenderer.h"
 
 #include "CanvasItem.h"
+#include "CanvasTypes.h"
 #include "Debug/CameraDebugColors.h"
-#include "Engine/Canvas.h"
 #include "Engine/Font.h"
 
 #if UE_GAMEPLAY_CAMERAS_DEBUG
@@ -12,12 +12,12 @@
 namespace UE::Cameras
 {
 
-void FDebugTextDrawCommand::Execute(UCanvas* Canvas, const UFont* Font, FVector2f& InOutDrawPosition) const
+void FDebugTextDrawCommand::Execute(FCanvas* Canvas, const FColor& DrawColor, const UFont* Font, FVector2f& InOutDrawPosition) const
 {
 	// Sadly we need to allocate a string here...
 	const FText Text(FText::FromStringView(TextView));
 
-	FCanvasTextItem TextItem(FVector2D(InOutDrawPosition), Text, Font, Canvas->DrawColor);
+	FCanvasTextItem TextItem(FVector2D(InOutDrawPosition), Text, Font, DrawColor);
 	TextItem.BlendMode = SE_BLEND_Translucent;
 	Canvas->DrawItem(TextItem);
 	
@@ -31,13 +31,14 @@ void FDebugTextNewLineCommand::Execute(FVector2f& InOutDrawPosition) const
 	InOutDrawPosition.Y += LineSpacing;
 }
 
-void FDebugTextSetColorCommand::Execute(UCanvas* Canvas) const
+void FDebugTextSetColorCommand::Execute(FColor* OutDrawColor) const
 {
-	Canvas->SetDrawColor(DrawColor);
+	*OutDrawColor = DrawColor;
 }
 
-FDebugTextRenderer::FDebugTextRenderer(UCanvas* InCanvas, const UFont* InFont)
+FDebugTextRenderer::FDebugTextRenderer(FCanvas* InCanvas, const FColor& InDrawColor, const UFont* InFont)
 	: Canvas(InCanvas)
+	, DrawColor(InDrawColor)
 	, Font(InFont)
 {
 	LineSpacing = InFont->GetMaxCharHeight();
@@ -132,21 +133,21 @@ void FDebugTextRenderer::ParseText(const TStringView<TCHAR> TextView, FDebugText
 
 void FDebugTextRenderer::ExecuteCommands(FDebugTextCommandArray& Commands)
 {
-	FColor OriginalDrawColor = Canvas->DrawColor;
+	FColor OriginalDrawColor = DrawColor;
 
 	for (FDebugTextCommand& Command : Commands)
 	{
 		switch (Command.GetIndex())
 		{
 			case FDebugTextCommand::IndexOfType<FDebugTextDrawCommand>():
-				Command.Get<FDebugTextDrawCommand>().Execute(Canvas, Font, NextDrawPosition);
+				Command.Get<FDebugTextDrawCommand>().Execute(Canvas, DrawColor, Font, NextDrawPosition);
 				UpdateRightMargin();
 				break;
 			case FDebugTextCommand::IndexOfType<FDebugTextNewLineCommand>():
 				Command.Get<FDebugTextNewLineCommand>().Execute(NextDrawPosition);
 				break;
 			case FDebugTextCommand::IndexOfType<FDebugTextSetColorCommand>():
-				Command.Get<FDebugTextSetColorCommand>().Execute(Canvas);
+				Command.Get<FDebugTextSetColorCommand>().Execute(&DrawColor);
 				break;
 			default:
 				ensureMsgf(false, TEXT("Unsupported command type!"));
@@ -154,7 +155,7 @@ void FDebugTextRenderer::ExecuteCommands(FDebugTextCommandArray& Commands)
 		}
 	}
 
-	Canvas->SetDrawColor(OriginalDrawColor);
+	DrawColor = OriginalDrawColor;
 }
 
 void FDebugTextRenderer::UpdateRightMargin()
@@ -192,8 +193,8 @@ void FDebugTextRenderer::AddTokenCommand(const TCHAR* RangeStart, const TCHAR* R
 {
 	TStringView<TCHAR> TokenView(RangeStart, RangeEnd - RangeStart);
 	// Sadly we need to allocate a string here too...
-	FColor DrawColor = InterpretColor(FString(TokenView));
-	OutCommands.Emplace(TInPlaceType<FDebugTextSetColorCommand>(), FDebugTextSetColorCommand{ DrawColor });
+	FColor NewDrawColor = InterpretColor(FString(TokenView));
+	OutCommands.Emplace(TInPlaceType<FDebugTextSetColorCommand>(), FDebugTextSetColorCommand{ NewDrawColor });
 }
 
 FColor FDebugTextRenderer::InterpretColor(const FString& ColorName)
diff --git a/Engine/Plugins/Cameras/GameplayCameras/Source/GameplayCameras/Private/Debug/RootCameraDebugBlock.cpp b/Engine/Plugins/Cameras/GameplayCameras/Source/GameplayCameras/Private/Debug/RootCameraDebugBlock.cpp
--- a/Engine/Plugins/Cameras/GameplayCameras/Source/GameplayCameras/Private/Debug/RootCameraDebugBlock.cpp
+++ b/Engine/Plugins/Cameras/GameplayCameras/Source/GameplayCameras/Private/Debug/RootCameraDebugBlock.cpp
@@ -27,12 +27,6 @@ static FAutoConsoleVariableRef CVarGameplayCamerasDebugEnable(
 	GGameplayCamerasDebugEnable,
 	TEXT("(Default: false. Enables debug drawing for the GamplayCameras system."));
 
-bool GGameplayCamerasDebugTrace = true;
-static FAutoConsoleVariableRef CVarGameplayCamerasDebugTrace(
-	TEXT("GameplayCameras.Debug.Trace"),
-	GGameplayCamerasDebugTrace,
-	TEXT("(Default: false. Enables background tracing of GamplayCameras system debug info."));
-
 FString GGameplayCamerasDebugCategories = "nodetree";
 static FAutoConsoleVariableRef CVarGameplayCamerasDebugCategories(
 	TEXT("GameplayCameras.Debug.Categories"),
diff --git a/Engine/Plugins/Cameras/GameplayCameras/Source/GameplayCameras/Private/Debug/VariableTableDebugBlock.cpp b/Engine/Plugins/Cameras/GameplayCameras/Source/GameplayCameras/Private/Debug/VariableTableDebugBlock.cpp
--- a/Engine/Plugins/Cameras/GameplayCameras/Source/GameplayCameras/Private/Debug/VariableTableDebugBlock.cpp
+++ b/Engine/Plugins/Cameras/GameplayCameras/Source/GameplayCameras/Private/Debug/VariableTableDebugBlock.cpp
@@ -74,8 +74,6 @@ void FVariableTableDebugBlock::OnDebugDraw(const FCameraDebugBlockDrawParams& Pa
 
 	const FCameraDebugColors& Colors = FCameraDebugColors::Get();
 
-	Renderer.AddText(TEXT("{cam_title}Variable Table{cam_default}\n"));
-	Renderer.AddIndent();
 	for (const FEntryDebugInfo& Entry : Entries)
 	{
 #if WITH_EDITORONLY_DATA
@@ -83,7 +81,14 @@ void FVariableTableDebugBlock::OnDebugDraw(const FCameraDebugBlockDrawParams& Pa
 		{
 			Renderer.AddText(TEXT("{cam_passive}[%d]{cam_default} "), Entry.Id);
 		}
-		Renderer.AddText(TEXT("%s : "), *Entry.Name);
+		if (!Entry.Name.IsEmpty())
+		{
+			Renderer.AddText(TEXT("%s : "), *Entry.Name);
+		}
+		else
+		{
+			Renderer.AddText(TEXT("<no name data> : "), Entry.Id);
+		}
 #else
 		Renderer.AddText(TEXT("[%d] <no name data> : "), Entry.Id);
 #endif
@@ -104,7 +109,22 @@ void FVariableTableDebugBlock::OnDebugDraw(const FCameraDebugBlockDrawParams& Pa
 		Renderer.NewLine();
 		Renderer.SetTextColor(Colors.Default);
 	}
-	Renderer.RemoveIndent();
+}
+
+void FVariableTableDebugBlock::OnSerialize(FArchive& Ar)
+{
+	Ar << Entries;
+	Ar << ShowVariableIdsCVarName;
+}
+
+FArchive& operator<< (FArchive& Ar, FVariableTableDebugBlock::FEntryDebugInfo& EntryDebugInfo)
+{
+	Ar << EntryDebugInfo.Id;
+	Ar << EntryDebugInfo.Name;
+	Ar << EntryDebugInfo.Value;
+	Ar << EntryDebugInfo.bWritten;
+	Ar << EntryDebugInfo.bWrittenThisFrame;
+	return Ar;
 }
 
 }  // namespace UE::Cameras
diff --git a/Engine/Plugins/Cameras/GameplayCameras/Source/GameplayCameras/Private/Debug/ViewfinderDebugBlock.cpp b/Engine/Plugins/Cameras/GameplayCameras/Source/GameplayCameras/Private/Debug/ViewfinderDebugBlock.cpp
--- a/Engine/Plugins/Cameras/GameplayCameras/Source/GameplayCameras/Private/Debug/ViewfinderDebugBlock.cpp
+++ b/Engine/Plugins/Cameras/GameplayCameras/Source/GameplayCameras/Private/Debug/ViewfinderDebugBlock.cpp
@@ -2,11 +2,12 @@
 
 #include "Debug/ViewfinderDebugBlock.h"
 
+#include "CanvasItem.h"
+#include "CanvasTypes.h"
 #include "Debug/CameraDebugCategories.h"
 #include "Debug/CameraDebugColors.h"
 #include "Debug/CameraDebugRenderer.h"
 #include "DrawDebugHelpers.h"
-#include "Engine/Canvas.h"
 #include "HAL/IConsoleManager.h"
 
 #if UE_GAMEPLAY_CAMERAS_DEBUG
@@ -38,6 +39,34 @@ static FAutoConsoleVariableRef CVarGameplayCamerasDebugViewfinderGuidesGapFactor
 	GGameplayCamerasDebugViewfinderGuidesGapFactor,
 	TEXT(""));
 
+namespace Private
+{
+
+void DrawCanvasLine(FCanvas* Canvas, const FVector2D& Start, const FVector2D& End, const FLinearColor& LineColor, float LineThickness = 1.f)
+{
+	FCanvasLineItem LineItem(Start, End);
+	LineItem.SetColor(LineColor);
+	LineItem.LineThickness = LineThickness;
+	LineItem.Draw(Canvas);
+}
+
+void DrawCanvasCircle(FCanvas* Canvas, const FVector2D& Center, float Radius, int32 NumSides, const FLinearColor& LineColor, float LineThickness = 1.f)
+{
+	const float	AngleDelta = 2.0f * UE_PI / NumSides;
+	FVector2D AxisX(1.f, 0.f);
+	FVector2D AxisY(0.f, -1.f);
+	FVector2D LastVertex = Center + AxisX * Radius;
+
+	for (int32 SideIndex = 0; SideIndex < NumSides; SideIndex++)
+	{
+		const FVector2D Vertex = Center + (AxisX * FMath::Cos(AngleDelta * (SideIndex + 1)) + AxisY * FMath::Sin(AngleDelta * (SideIndex + 1))) * Radius;
+		DrawCanvasLine(Canvas, LastVertex, Vertex, LineColor, LineThickness);
+		LastVertex = Vertex;
+	}
+}
+
+}  //  namespace Private 
+
 UE_DEFINE_CAMERA_DEBUG_BLOCK(FViewfinderDebugBlock)
 
 FViewfinderDebugBlock::FViewfinderDebugBlock()
@@ -51,68 +80,72 @@ void FViewfinderDebugBlock::OnDebugDraw(const FCameraDebugBlockDrawParams& Param
 		return;
 	}
 
-	UCanvas* Canvas = Renderer.GetCanvas();
+	FCanvas* Canvas = Renderer.GetCanvas();
+
+	const FVector2D CanvasSize = Renderer.GetCanvasSize();
+	const float CanvasSizeX = CanvasSize.X;
+	const float CanvasSizeY = CanvasSize.Y;
 
-	const FVector2d CanvasCenter(Canvas->SizeX / 2.f, Canvas->SizeY / 2.f);
+	const FVector2d CanvasCenter(CanvasSizeX / 2.f, CanvasSizeY / 2.f);
 
 	// Draw the reticle.
-	const float ReticleRadius = Canvas->SizeY * GGameplayCamerasDebugViewfinderReticleSizeFactor / 2.f;
+	const float ReticleRadius = CanvasSizeY * GGameplayCamerasDebugViewfinderReticleSizeFactor / 2.f;
 	const float ReticleInnerRadiusFactor = GGameplayCamerasDebugViewfinderReticleInnerSizeFactor;
 	const int32 ReticleNumSides = GGameplayCamerasDebugViewfinderReticleNumSides;
 	const FColor ReticleColor = FCameraDebugColors::Get().Passive;
 
 	// ...outer reticle circle.
-	DrawDebugCanvas2DCircle(Canvas, CanvasCenter, ReticleRadius, ReticleNumSides, ReticleColor);
+	Private::DrawCanvasCircle(Canvas, CanvasCenter, ReticleRadius, ReticleNumSides, ReticleColor);
 	// ...inner reticle circle.
 	const float ReticleInnerRadius = ReticleRadius * ReticleInnerRadiusFactor;
 	const int32 ReticleInnerNumSides = (int32)(ReticleNumSides * ReticleInnerRadiusFactor);
-	DrawDebugCanvas2DCircle(Canvas, CanvasCenter, ReticleInnerRadius, ReticleInnerNumSides, ReticleColor);
+	Private::DrawCanvasCircle(Canvas, CanvasCenter, ReticleInnerRadius, ReticleInnerNumSides, ReticleColor);
 	// ...horizontal line inside reticle.
-	DrawDebugCanvas2DLine(
+	Private::DrawCanvasLine(
 			Canvas, 
 			CanvasCenter - FVector2D(ReticleInnerRadius, 0), CanvasCenter + FVector2D(ReticleInnerRadius, 0), 
 			ReticleColor);
 
 	// Draw the rule-of-thirds guides.
 	const FColor GuideColor = FCameraDebugColors::Get().VeryPassive;
-	const float GuidesGap = Canvas->SizeY * GGameplayCamerasDebugViewfinderGuidesGapFactor;
-	const FVector2D OneThird(Canvas->SizeX / 3.f, Canvas->SizeY / 3.f);
-	const FVector2D TwoThirds(Canvas->SizeX / 1.5f, Canvas->SizeY / 1.5f);
+	const float GuidesGap = CanvasSizeY * GGameplayCamerasDebugViewfinderGuidesGapFactor;
+	const FVector2D OneThird(CanvasSizeX / 3.f, CanvasSizeY / 3.f);
+	const FVector2D TwoThirds(CanvasSizeX / 1.5f, CanvasSizeY / 1.5f);
 	// ...top vertical guides
-	DrawDebugCanvas2DLine(
+	Private::DrawCanvasLine(
 			Canvas,
 			FVector2D(OneThird.X, 0), FVector2D(OneThird.X, OneThird.Y - GuidesGap),
 			ReticleColor, 2.f);
-	DrawDebugCanvas2DLine(
+	Private::DrawCanvasLine(
 			Canvas,
 			FVector2D(TwoThirds.X, 0), FVector2D(TwoThirds.X, OneThird.Y - GuidesGap),
 			ReticleColor, 2.f);
 	// ...bottom vertical guides
-	DrawDebugCanvas2DLine(
+	Private::DrawCanvasLine(
 			Canvas,
-			FVector2D(OneThird.X, TwoThirds.Y + GuidesGap), FVector2D(OneThird.X, Canvas->SizeY),
+			FVector2D(OneThird.X, TwoThirds.Y + GuidesGap), FVector2D(OneThird.X, CanvasSizeY),
 			ReticleColor, 2.f);
-	DrawDebugCanvas2DLine(
+	Private::DrawCanvasLine(
 			Canvas,
-			FVector2D(TwoThirds.X, TwoThirds.Y + GuidesGap), FVector2D(TwoThirds.X, Canvas->SizeY),
+			FVector2D(TwoThirds.X, TwoThirds.Y + GuidesGap), FVector2D(TwoThirds.X, CanvasSizeY),
 			ReticleColor, 2.f);
 	// ...left horizontal guides
-	DrawDebugCanvas2DLine(
+	Private::DrawCanvasLine(
 			Canvas,
 			FVector2D(0, OneThird.Y), FVector2D(OneThird.X - GuidesGap, OneThird.Y),
 			ReticleColor, 2.f);
-	DrawDebugCanvas2DLine(
+	Private::DrawCanvasLine(
 			Canvas,
 			FVector2D(0, TwoThirds.Y), FVector2D(OneThird.X - GuidesGap, TwoThirds.Y),
 			ReticleColor, 2.f);
 	// ...right horizontal guides
-	DrawDebugCanvas2DLine(
+	Private::DrawCanvasLine(
 			Canvas,
-			FVector2D(TwoThirds.X + GuidesGap, OneThird.Y), FVector2D(Canvas->SizeX, OneThird.Y),
+			FVector2D(TwoThirds.X + GuidesGap, OneThird.Y), FVector2D(CanvasSizeX, OneThird.Y),
 			ReticleColor, 2.f);
-	DrawDebugCanvas2DLine(
+	Private::DrawCanvasLine(
 			Canvas,
-			FVector2D(TwoThirds.X + GuidesGap, TwoThirds.Y), FVector2D(Canvas->SizeX, TwoThirds.Y),
+			FVector2D(TwoThirds.X + GuidesGap, TwoThirds.Y), FVector2D(CanvasSizeX, TwoThirds.Y),
 			ReticleColor, 2.f);
 }
 
diff --git a/Engine/Plugins/Cameras/GameplayCameras/Source/GameplayCameras/Private/GameFramework/GameplayCameraSystemComponent.cpp b/Engine/Plugins/Cameras/GameplayCameras/Source/GameplayCameras/Private/GameFramework/GameplayCameraSystemComponent.cpp
--- a/Engine/Plugins/Cameras/GameplayCameras/Source/GameplayCameras/Private/GameFramework/GameplayCameraSystemComponent.cpp
+++ b/Engine/Plugins/Cameras/GameplayCameras/Source/GameplayCameras/Private/GameFramework/GameplayCameraSystemComponent.cpp
@@ -4,6 +4,7 @@
 
 #include "Components/StaticMeshComponent.h"
 #include "Debug/DebugDrawService.h"
+#include "Engine/Canvas.h"
 #include "Engine/StaticMesh.h"
 #include "Engine/World.h"
 #include "GameFramework/Controller.h"
@@ -153,7 +154,7 @@ void UGameplayCameraSystemComponent::DebugDraw(UCanvas* Canvas, APlayerControlle
 	if (Evaluator.IsValid())
 	{
 		FCameraSystemDebugUpdateParams DebugUpdateParams;
-		DebugUpdateParams.Canvas = Canvas;
+		DebugUpdateParams.Canvas = Canvas->Canvas;
 		Evaluator->DebugUpdate(DebugUpdateParams);
 	}
 }
diff --git a/Engine/Plugins/Cameras/GameplayCameras/Source/GameplayCameras/Private/Nodes/Blends/SimpleBlendCameraNode.cpp b/Engine/Plugins/Cameras/GameplayCameras/Source/GameplayCameras/Private/Nodes/Blends/SimpleBlendCameraNode.cpp
--- a/Engine/Plugins/Cameras/GameplayCameras/Source/GameplayCameras/Private/Nodes/Blends/SimpleBlendCameraNode.cpp
+++ b/Engine/Plugins/Cameras/GameplayCameras/Source/GameplayCameras/Private/Nodes/Blends/SimpleBlendCameraNode.cpp
@@ -17,7 +17,7 @@ UE_DECLARE_CAMERA_DEBUG_BLOCK_START(FSimpleBlendCameraDebugBlock)
 	UE_DECLARE_CAMERA_DEBUG_BLOCK_FIELD(float, BlendFactor)
 UE_DECLARE_CAMERA_DEBUG_BLOCK_END()
 
-UE_DEFINE_CAMERA_DEBUG_BLOCK(FSimpleBlendCameraDebugBlock)
+UE_DEFINE_CAMERA_DEBUG_BLOCK_WITH_FIELDS(FSimpleBlendCameraDebugBlock)
 
 void FSimpleBlendCameraNodeEvaluator::OnRun(const FCameraNodeEvaluationParams& Params, FCameraNodeEvaluationResult& OutResult)
 {
diff --git a/Engine/Plugins/Cameras/GameplayCameras/Source/GameplayCameras/Private/Nodes/Common/DampenPositionCameraNode.cpp b/Engine/Plugins/Cameras/GameplayCameras/Source/GameplayCameras/Private/Nodes/Common/DampenPositionCameraNode.cpp
--- a/Engine/Plugins/Cameras/GameplayCameras/Source/GameplayCameras/Private/Nodes/Common/DampenPositionCameraNode.cpp
+++ b/Engine/Plugins/Cameras/GameplayCameras/Source/GameplayCameras/Private/Nodes/Common/DampenPositionCameraNode.cpp
@@ -56,7 +56,7 @@ UE_DECLARE_CAMERA_DEBUG_BLOCK_START(FDampenPositionCameraDebugBlock)
 	UE_DECLARE_CAMERA_DEBUG_BLOCK_FIELD(FVector3d, DampedPosition);
 UE_DECLARE_CAMERA_DEBUG_BLOCK_END()
 
-UE_DEFINE_CAMERA_DEBUG_BLOCK(FDampenPositionCameraDebugBlock)
+UE_DEFINE_CAMERA_DEBUG_BLOCK_WITH_FIELDS(FDampenPositionCameraDebugBlock)
 
 void FDampenPositionCameraNodeEvaluator::OnInitialize(const FCameraNodeEvaluatorInitializeParams& Params)
 {
diff --git a/Engine/Plugins/Cameras/GameplayCameras/Source/GameplayCameras/Public/Core/BlendStackCameraNode.h b/Engine/Plugins/Cameras/GameplayCameras/Source/GameplayCameras/Public/Core/BlendStackCameraNode.h
--- a/Engine/Plugins/Cameras/GameplayCameras/Source/GameplayCameras/Public/Core/BlendStackCameraNode.h
+++ b/Engine/Plugins/Cameras/GameplayCameras/Source/GameplayCameras/Public/Core/BlendStackCameraNode.h
@@ -177,6 +177,7 @@ public:
 protected:
 
 	virtual void OnDebugDraw(const FCameraDebugBlockDrawParams& Params, FCameraDebugRenderer& Renderer) override;
+	virtual void OnSerialize(FArchive& Ar) override;
 
 private:
 
@@ -195,6 +196,7 @@ public:
 protected:
 
 	virtual void OnDebugDraw(const FCameraDebugBlockDrawParams& Params, FCameraDebugRenderer& Renderer) override;
+	virtual void OnSerialize(FArchive& Ar) override;
 
 private:
 
@@ -204,6 +206,8 @@ private:
 	};
 
 	TArray<FEntryDebugInfo> Entries;
+
+	friend FArchive& operator<< (FArchive&, FEntryDebugInfo&);
 };
 
 #endif  // UE_GAMEPLAY_CAMERAS_DEBUG
diff --git a/Engine/Plugins/Cameras/GameplayCameras/Source/GameplayCameras/Public/Core/CameraObjectStorage.h b/Engine/Plugins/Cameras/GameplayCameras/Source/GameplayCameras/Public/Core/CameraObjectStorage.h
--- a/Engine/Plugins/Cameras/GameplayCameras/Source/GameplayCameras/Public/Core/CameraObjectStorage.h
+++ b/Engine/Plugins/Cameras/GameplayCameras/Source/GameplayCameras/Public/Core/CameraObjectStorage.h
@@ -17,8 +17,8 @@ namespace UE::Cameras
 template<typename BaseObjectType>
 struct TCameraObjectStorageTraits
 {
-	static const uint16 DefaultPageCapacity = 128;
-	static const uint16 DefaultPageAlignment = 32;
+	static const uint32 DefaultPageCapacity = 128;
+	static const uint32 DefaultPageAlignment = 32;
 };
 
 /**
@@ -54,6 +54,21 @@ protected:
 		::Type
 	BuildObject(ArgTypes&&... InArgs);
 
+	/**
+	 * Allocates memory for an object of the given size and alignment, but
+	 * doesn't initialize anything in that memory block.
+	 *
+	 * After the caller has constructed the object in-place, they MUST call
+	 * RegisterInitializedObject() with the actual object pointer, otherwise 
+	 * the storage won't call its destructor.
+	 */
+	void* BuildObjectUninitialized(uint32 Sizeof, uint32 Alignof);
+
+	/**
+	 * Called to registered an object after it has been constructed.
+	 */
+	void RegisterInitializedObject(BaseObjectType* BaseObjectPtr);
+
 	/**
 	 * Destroys all objects in the storage.
 	 *
@@ -64,12 +79,12 @@ protected:
 	/**
 	 * Computes information about the overall allocated memory.
 	 */
-	void GetAllocationInfo(uint16& OutTotalUsed, uint16& OutFirstAlignment) const;
+	void GetAllocationInfo(uint32& OutTotalUsed, uint32& OutFirstAlignment) const;
 
 	/**
 	 * Allocates a new page buffer.
 	 */
-	void AllocatePage(uint16 InCapacity, uint16 InAlignment);
+	void AllocatePage(uint32 InCapacity, uint32 InAlignment);
 
 private:
 
@@ -77,9 +92,9 @@ private:
 	struct FAllocation
 	{
 		uint8* Memory = nullptr;
-		uint16 Alignment = 0;
-		uint16 Capacity = 0;
-		uint16 Used = 0;
+		uint32 Alignment = 0;
+		uint32 Capacity = 0;
+		uint32 Used = 0;
 	};
 	/** Allocated page buffers */
 	TArray<FAllocation> Allocations;
@@ -127,17 +142,28 @@ template<typename ObjectType, typename ...ArgTypes>
 typename TEnableIf<TPointerIsConvertibleFromTo<ObjectType, BaseObjectType>::Value, ObjectType*>::Type
 TCameraObjectStorage<BaseObjectType>::BuildObject(ArgTypes&&... InArgs)
 {
-	uint16 StateSizeof = sizeof(ObjectType);
-	uint16 StateAlignof = alignof(ObjectType);
+	const uint32 Sizeof = sizeof(ObjectType);
+	const uint32 Alignof = alignof(ObjectType);
+	void* TargetPtr = BuildObjectUninitialized(Sizeof, Alignof);
+
+	ObjectType* NewObject = new(TargetPtr) ObjectType(Forward<ArgTypes>(InArgs)...);
 
+	RegisterInitializedObject(NewObject);
+
+	return NewObject;
+}
+
+template<typename BaseObjectType>
+void* TCameraObjectStorage<BaseObjectType>::BuildObjectUninitialized(uint32 Sizeof, uint32 Alignof)
+{
 	// Search for any allocation bucket that has enough room for the object
 	// we want to build.
 	uint8* TargetPtr = nullptr;
 	FAllocation* TargetAllocation = nullptr;
 	for (FAllocation& Allocation : Allocations)
 	{
-		uint8* PossiblePtr = Align(Allocation.Memory + Allocation.Used, StateAlignof);
-		uint16 NewUsed = (PossiblePtr + StateSizeof) - Allocation.Memory;
+		uint8* PossiblePtr = Align(Allocation.Memory + Allocation.Used, Alignof);
+		uint32 NewUsed = (PossiblePtr + Sizeof) - Allocation.Memory;
 		if (NewUsed <= Allocation.Capacity)
 		{
 			TargetPtr = PossiblePtr;
@@ -151,28 +177,30 @@ TCameraObjectStorage<BaseObjectType>::BuildObject(ArgTypes&&... InArgs)
 	if (TargetPtr == nullptr)
 	{
 		using FStorageTraits = TCameraObjectStorageTraits<BaseObjectType>;
-		const uint16 DefaultCapacity = FStorageTraits::DefaultPageCapacity;
-		const uint16 DefaultAlignment = FStorageTraits::DefaultPageAlignment;
+		const uint32 DefaultCapacity = FStorageTraits::DefaultPageCapacity;
+		const uint32 DefaultAlignment = FStorageTraits::DefaultPageAlignment;
 
-		const uint16 NewCapacity = FMath::Max(DefaultCapacity, StateSizeof);
-		const uint16 NewAlignment = FMath::Max(DefaultAlignment, StateAlignof);
+		const uint32 NewCapacity = FMath::Max(DefaultCapacity, Sizeof);
+		const uint32 NewAlignment = FMath::Max(DefaultAlignment, Alignof);
 
 		FAllocation& NewAllocation = Allocations.Emplace_GetRef();
 		NewAllocation.Memory = reinterpret_cast<uint8*>(FMemory::Malloc(NewCapacity, NewAlignment));
 		NewAllocation.Alignment = NewAlignment;
 		NewAllocation.Capacity = NewCapacity;
-		NewAllocation.Used = StateSizeof;
+		NewAllocation.Used = Sizeof;
 
 		TargetPtr = NewAllocation.Memory;
 		TargetAllocation = &NewAllocation;
 	}
 	check(TargetPtr && TargetAllocation);
 
-	ObjectType* NewObject = new(TargetPtr) ObjectType(Forward<ArgTypes>(InArgs)...);
-
-	ObjectInfos.Add({ NewObject });
+	return TargetPtr;
+}
 
-	return NewObject;
+template<typename BaseObjectType>
+void TCameraObjectStorage<BaseObjectType>::RegisterInitializedObject(BaseObjectType* BaseObjectPtr)
+{
+	ObjectInfos.Add({ BaseObjectPtr });
 }
 
 template<typename BaseObjectType>
@@ -205,7 +233,7 @@ void TCameraObjectStorage<BaseObjectType>::DestroyObjects(bool bFreeAllocations)
 }
 
 template<typename BaseObjectType>
-void TCameraObjectStorage<BaseObjectType>::GetAllocationInfo(uint16& OutTotalUsed, uint16& OutFirstAlignment) const
+void TCameraObjectStorage<BaseObjectType>::GetAllocationInfo(uint32& OutTotalUsed, uint32& OutFirstAlignment) const
 {
 	OutTotalUsed = 0;
 	OutFirstAlignment = 0;
@@ -226,7 +254,7 @@ void TCameraObjectStorage<BaseObjectType>::GetAllocationInfo(uint16& OutTotalUse
 }
 
 template<typename BaseObjectType>
-void TCameraObjectStorage<BaseObjectType>::AllocatePage(uint16 InCapacity, uint16 InAlignment)
+void TCameraObjectStorage<BaseObjectType>::AllocatePage(uint32 InCapacity, uint32 InAlignment)
 {
 	FAllocation& NewAllocation = Allocations.Emplace_GetRef();
 	NewAllocation.Memory = reinterpret_cast<uint8*>(FMemory::Malloc(InCapacity, InAlignment));
diff --git a/Engine/Plugins/Cameras/GameplayCameras/Source/GameplayCameras/Public/Core/CameraPose.h b/Engine/Plugins/Cameras/GameplayCameras/Source/GameplayCameras/Public/Core/CameraPose.h
--- a/Engine/Plugins/Cameras/GameplayCameras/Source/GameplayCameras/Public/Core/CameraPose.h
+++ b/Engine/Plugins/Cameras/GameplayCameras/Source/GameplayCameras/Public/Core/CameraPose.h
@@ -8,6 +8,8 @@
 
 #include "CameraPose.generated.h"
 
+class FArchive;
+
 #define UE_CAMERA_POSE_FOR_TRANSFORM_PROPERTIES()\
 	UE_CAMERA_POSE_FOR_PROPERTY(FVector, Location)\
 	UE_CAMERA_POSE_FOR_PROPERTY(FRotator3d, Rotation)
@@ -39,7 +41,7 @@
 /**
  * Boolean flags for each of the properties inside FCameraPose.
  */
-struct FCameraPoseFlags
+struct GAMEPLAYCAMERAS_API FCameraPoseFlags
 {
 #define UE_CAMERA_POSE_FOR_PROPERTY(PropType, PropName)\
 	bool PropName = false;
@@ -76,7 +78,7 @@ public:
  * The ChangedFlags structure keeps track of which fields were changed via the setters.
  */
 USTRUCT(BlueprintType)
-struct FCameraPose
+struct GAMEPLAYCAMERAS_API FCameraPose
 {
 	GENERATED_BODY()
 
@@ -159,9 +161,12 @@ public:
 	/** Interpolates changed properties from ToPose using the given factor. Only properties defined by InMask are taken into account. */
 	void LerpChanged(const FCameraPose& ToPose, float Factor, const FCameraPoseFlags& InMask, bool bInvertMask, FCameraPoseFlags& OutMask);
 
+public:
+
+	static void SerializeWithFlags(FArchive& Ar, FCameraPose& CameraPose);
+
 private:
 
-	void InternalLerpChanged(const FCameraPose& ToPose, float Factor, const FCameraPoseFlags& InMask, bool bInvertMask, FCameraPoseFlags& OutMask);
 	void InternalOverrideChanged(const FCameraPose& OtherPose, bool bChangedOnly);
 	void InternalLerpChanged(const FCameraPose& ToPose, float Factor, const FCameraPoseFlags& InMask, bool bInvertMask, FCameraPoseFlags& OutMask, bool bChangedOnly);
 
diff --git a/Engine/Plugins/Cameras/GameplayCameras/Source/GameplayCameras/Public/Core/CameraSystemEvaluator.h b/Engine/Plugins/Cameras/GameplayCameras/Source/GameplayCameras/Public/Core/CameraSystemEvaluator.h
--- a/Engine/Plugins/Cameras/GameplayCameras/Source/GameplayCameras/Public/Core/CameraSystemEvaluator.h
+++ b/Engine/Plugins/Cameras/GameplayCameras/Source/GameplayCameras/Public/Core/CameraSystemEvaluator.h
@@ -14,9 +14,9 @@
 #include "Templates/SharedPointer.h"
 #include "UObject/GCObject.h"
 
+class FCanvas;
 class UCameraDirector;
 class UCameraRigAsset;
-class UCanvas;
 class URootCameraNode;
 struct FMinimalViewInfo;
 
@@ -68,7 +68,7 @@ struct FCameraSystemEvaluationUpdateResult
 #if UE_GAMEPLAY_CAMERAS_DEBUG
 struct FCameraSystemDebugUpdateParams
 {
-	UCanvas* Canvas = nullptr;
+	FCanvas* Canvas = nullptr;
 };
 #endif  // UE_GAMEPLAY_CAMERAS_DEBUG
 
@@ -125,6 +125,9 @@ public:
 
 private:
 
+	/** The owner (if any) of this camera system evaluator. */
+	TWeakObjectPtr<> WeakOwner;
+
 	/** The root camera node. */
 	TObjectPtr<URootCameraNode> RootNode;
 
diff --git a/Engine/Plugins/Cameras/GameplayCameras/Source/GameplayCameras/Public/Debug/BlendStacksCameraDebugBlock.h b/Engine/Plugins/Cameras/GameplayCameras/Source/GameplayCameras/Public/Debug/BlendStacksCameraDebugBlock.h
--- a/Engine/Plugins/Cameras/GameplayCameras/Source/GameplayCameras/Public/Debug/BlendStacksCameraDebugBlock.h
+++ b/Engine/Plugins/Cameras/GameplayCameras/Source/GameplayCameras/Public/Debug/BlendStacksCameraDebugBlock.h
@@ -31,6 +31,7 @@ public:
 protected:
 
 	virtual void OnDebugDraw(const FCameraDebugBlockDrawParams& Params, FCameraDebugRenderer& Renderer) override;
+	virtual void OnSerialize(FArchive& Ar) override;
 
 private:
 
diff --git a/Engine/Plugins/Cameras/GameplayCameras/Source/GameplayCameras/Public/Debug/CameraDebugBlock.h b/Engine/Plugins/Cameras/GameplayCameras/Source/GameplayCameras/Public/Debug/CameraDebugBlock.h
--- a/Engine/Plugins/Cameras/GameplayCameras/Source/GameplayCameras/Public/Debug/CameraDebugBlock.h
+++ b/Engine/Plugins/Cameras/GameplayCameras/Source/GameplayCameras/Public/Debug/CameraDebugBlock.h
@@ -24,7 +24,7 @@ struct FCameraDebugBlockDrawParams
 	/** The list of active debug categories. */
 	TSet<FString> ActiveCategories;
 
-	bool IsCategoryActive(const FString& InCategory) const;
+	GAMEPLAYCAMERAS_API bool IsCategoryActive(const FString& InCategory) const;
 };
 
 /**
@@ -33,8 +33,8 @@ struct FCameraDebugBlockDrawParams
  */
 struct FCameraDebugBlockField
 {
-	virtual ~FCameraDebugBlockField() {}
-	virtual void SerializeField(FCameraDebugBlock* This, FArchive& Ar) = 0;
+	GAMEPLAYCAMERAS_API virtual ~FCameraDebugBlockField() {}
+	GAMEPLAYCAMERAS_API virtual void SerializeField(FCameraDebugBlock* This, FArchive& Ar) = 0;
 
 	FName FieldName;
 	uint16 FieldIndex = 0;
@@ -47,13 +47,13 @@ struct FCameraDebugBlockField
 template<typename FieldType>
 struct TCameraDebugBlockField : FCameraDebugBlockField
 {
-	TCameraDebugBlockField(const FName& InName, uint16 InOffset) 
+	GAMEPLAYCAMERAS_API TCameraDebugBlockField(const FName& InName, uint16 InOffset) 
 	{
 		FieldName = InName;
 		FieldOffset = InOffset;
 	}
 
-	virtual void SerializeField(FCameraDebugBlock* This, FArchive& Ar) override
+	GAMEPLAYCAMERAS_API virtual void SerializeField(FCameraDebugBlock* This, FArchive& Ar) override
 	{
 		FieldType* Value = reinterpret_cast<FieldType*>(reinterpret_cast<uint8*>(This) + FieldOffset);
 		Ar << (*Value);
@@ -71,52 +71,52 @@ class FCameraDebugBlock
 
 public:
 
-	virtual ~FCameraDebugBlock() {}
+	GAMEPLAYCAMERAS_API virtual ~FCameraDebugBlock() {}
 
 	/** Attaches a block to this block. */
-	void Attach(FCameraDebugBlock* InAttachment);
+	GAMEPLAYCAMERAS_API void Attach(FCameraDebugBlock* InAttachment);
 	/** Gets the list of blocks attached to this block. */
 	TArrayView<FCameraDebugBlock*> GetAttachments() { return Attachments; }
 
 	/** Adds a child block to this block. */
-	void AddChild(FCameraDebugBlock* InChild);
+	GAMEPLAYCAMERAS_API void AddChild(FCameraDebugBlock* InChild);
 	/** Gets the children of this block. */
 	TArrayView<FCameraDebugBlock*> GetChildren() { return Children; }
 
 	/** Called to let this block display its information on screen. */
-	void DebugDraw(const FCameraDebugBlockDrawParams& Params, FCameraDebugRenderer& Renderer);
+	GAMEPLAYCAMERAS_API void DebugDraw(const FCameraDebugBlockDrawParams& Params, FCameraDebugRenderer& Renderer);
 
 	/** Serializes this debug block into a buffer, for recording/replaying purposes. */
-	void Serialize(FArchive& Ar);
+	GAMEPLAYCAMERAS_API void Serialize(FArchive& Ar);
 
 protected:
 
 	/** Called to let this block display its information on screen. */
-	virtual void OnDebugDraw(const FCameraDebugBlockDrawParams& Params, FCameraDebugRenderer& Renderer) {}
+	GAMEPLAYCAMERAS_API virtual void OnDebugDraw(const FCameraDebugBlockDrawParams& Params, FCameraDebugRenderer& Renderer) {}
 	/** Called after attached and children blocks' debug draw. */
-	virtual void OnPostDebugDraw(const FCameraDebugBlockDrawParams& Params, FCameraDebugRenderer& Renderer) {}
+	GAMEPLAYCAMERAS_API virtual void OnPostDebugDraw(const FCameraDebugBlockDrawParams& Params, FCameraDebugRenderer& Renderer) {}
 
 	/** Serializes this debug block into a buffer, for recording/replaying purposes. */
-	virtual void OnSerialize(FArchive& Ar) {}
+	GAMEPLAYCAMERAS_API virtual void OnSerialize(FArchive& Ar) {}
 
 protected:
 
+	using FStaticFieldArray = TArray<FCameraDebugBlockField*>;
+
 	template<typename FieldType>
 	static TCameraDebugBlockField<FieldType> CreateField(const FName& FieldName, uint16 FieldOffset)
 	{
 		return TCameraDebugBlockField<FieldType>{ FieldName, FieldOffset };
 	}
 
-	static int32 RegisterField(FCameraDebugBlockField* InField)
+	static int32 RegisterField(FCameraDebugBlockField* InField, FStaticFieldArray& InStaticFields)
 	{
-		StaticFields.Add(InField);
-		return StaticFields.Num() - 1;
+		InStaticFields.Add(InField);
+		return InStaticFields.Num() - 1;
 	}
 
 private:
 
-	static TArray<FCameraDebugBlockField*> StaticFields;
-
 	using FRelatedDebugBlockArray = TArray<FCameraDebugBlock*>;
 	FRelatedDebugBlockArray Attachments;
 	FRelatedDebugBlockArray Children;
@@ -158,23 +158,40 @@ private:
 		UE_DECLARE_CAMERA_DEBUG_BLOCK(ClassName)\
 	private:\
 		using Super = ::UE::Cameras::FCameraDebugBlock;\
-		using ThisClassName = ClassName;
+		using ThisClassName = ClassName;\
+		static FStaticFieldArray StaticFields;\
+	protected:\
+		virtual void OnSerialize(FArchive& Ar) override\
+		{\
+			Super::OnSerialize(Ar);\
+			for (FCameraDebugBlockField* StaticField : ClassName::StaticFields)\
+			{\
+				StaticField->SerializeField(this, Ar);\
+			}\
+		}\
+	private:
 
 #define UE_DECLARE_CAMERA_DEBUG_BLOCK_FIELD(FieldType, FieldName)\
 	public:\
 		FieldType FieldName;\
 	private:\
 		inline static PTRINT Get##FieldName##Offset() { return (PTRINT)(&(((ThisClassName*)0)->FieldName)); }\
-		inline static TCameraDebugBlockField<FieldType>* Get##FieldName##Field() {\
+		inline static TCameraDebugBlockField<FieldType>* Get##FieldName##Field()\
+		{\
 			using FField = TCameraDebugBlockField<FieldType>;\
 			static FField StaticField = CreateField<FieldType>(TEXT(#FieldName), Get##FieldName##Offset());\
-			return &StaticField; }\
-		inline static const int32 FieldName##FieldIndex = RegisterField(Get##FieldName##Field());
+			return &StaticField;\
+		}\
+		inline static const int32 FieldName##FieldIndex = RegisterField(Get##FieldName##Field(), StaticFields);
 
 #define UE_DECLARE_CAMERA_DEBUG_BLOCK_END()\
 		virtual void OnDebugDraw(const FCameraDebugBlockDrawParams& Params, FCameraDebugRenderer& Renderer) override;\
 	};
 
+#define UE_DEFINE_CAMERA_DEBUG_BLOCK_WITH_FIELDS(ClassName)\
+	UE_DEFINE_CAMERA_DEBUG_BLOCK(ClassName)\
+	TArray<FCameraDebugBlockField*> ClassName::StaticFields;
+
 #else  // UE_GAMEPLAY_CAMERAS_DEBUG
 
 // Empty macros for shipping builds.
@@ -182,8 +199,7 @@ private:
 #define UE_DECLARE_CAMERA_DEBUG_BLOCK_START(ClassName)
 #define UE_DECLARE_CAMERA_DEBUG_BLOCK_FIELD(FieldType, FieldName)
 #define UE_DECLARE_CAMERA_DEBUG_BLOCK_END()
-// This one is included here too since it's required for "simple" debug blocks.
-#define UE_DEFINE_CAMERA_DEBUG_BLOCK(ClassName)
+#define UE_DEFINE_CAMERA_DEBUG_BLOCK_WITH_FIELDS(ClassName)
 
 #endif  // UE_GAMEPLAY_CAMERAS_DEBUG
 
diff --git a/Engine/Plugins/Cameras/GameplayCameras/Source/GameplayCameras/Public/Debug/CameraDebugBlockStorage.h b/Engine/Plugins/Cameras/GameplayCameras/Source/GameplayCameras/Public/Debug/CameraDebugBlockStorage.h
--- a/Engine/Plugins/Cameras/GameplayCameras/Source/GameplayCameras/Public/Debug/CameraDebugBlockStorage.h
+++ b/Engine/Plugins/Cameras/GameplayCameras/Source/GameplayCameras/Public/Debug/CameraDebugBlockStorage.h
@@ -22,11 +22,19 @@ public:
 	using Super = TCameraObjectStorage<FCameraDebugBlock>;
 
 	/** Destroy any allocated debug blocks. */
-	void DestroyDebugBlocks(bool bFreeAllocations = false);
+	GAMEPLAYCAMERAS_API void DestroyDebugBlocks(bool bFreeAllocations = false);
 
 	/** Build a new debug block inside this storage. */
 	template<typename BlockType, typename ...ArgTypes>
 	BlockType* BuildDebugBlock(ArgTypes&&... InArgs);
+
+public:
+
+	// Internal API.
+	void* BuildDebugBlockUninitialized(uint32 Sizeof, uint32 Alignof)
+	{
+		return BuildObjectUninitialized(Sizeof, Alignof);
+	}
 };
 
 template<typename BlockType, typename ...ArgTypes>
diff --git a/Engine/Plugins/Cameras/GameplayCameras/Source/GameplayCameras/Public/Debug/CameraDebugRenderer.h b/Engine/Plugins/Cameras/GameplayCameras/Source/GameplayCameras/Public/Debug/CameraDebugRenderer.h
--- a/Engine/Plugins/Cameras/GameplayCameras/Source/GameplayCameras/Public/Debug/CameraDebugRenderer.h
+++ b/Engine/Plugins/Cameras/GameplayCameras/Source/GameplayCameras/Public/Debug/CameraDebugRenderer.h
@@ -8,7 +8,7 @@
 #include "Math/Vector.h"
 #include "Misc/StringBuilder.h"
 
-class UCanvas;
+class FCanvas;
 class UFont;
 
 #if UE_GAMEPLAY_CAMERAS_DEBUG
@@ -34,51 +34,54 @@ class FCameraDebugRenderer
 public:
 
 	/** Creates a new debug renderer. */
-	FCameraDebugRenderer(UCanvas* InCanvas);
+	GAMEPLAYCAMERAS_API FCameraDebugRenderer(FCanvas* InCanvas);
 	/** Destroys the debug renderer. */
-	~FCameraDebugRenderer();
+	GAMEPLAYCAMERAS_API ~FCameraDebugRenderer();
 
 	/** Adds text to the text wall. */
-	void AddText(const FString& InString);
-	void AddText(const TCHAR* Fmt, ...);
+	GAMEPLAYCAMERAS_API void AddText(const FString& InString);
+	GAMEPLAYCAMERAS_API void AddText(const TCHAR* Fmt, ...);
 
 	/** 
 	 * Move to a new line on the text wall.
 	 *
 	 * @return Whether a new line was added.
 	 */
-	bool NewLine(bool bSkipIfEmptyLine = false);
+	GAMEPLAYCAMERAS_API bool NewLine(bool bSkipIfEmptyLine = false);
 
 	/** Gets the current text color. */
-	FColor GetTextColor() const;
+	GAMEPLAYCAMERAS_API FColor GetTextColor() const;
 	/** Sets the text color for further calls. Returns the previous color. */
-	FColor SetTextColor(const FColor& Color);
+	GAMEPLAYCAMERAS_API FColor SetTextColor(const FColor& Color);
 
 	/** Increases the indent of the next text wall entry. This will make a new line. */
-	void AddIndent();
+	GAMEPLAYCAMERAS_API void AddIndent();
 	/** Decreases the indent of the next text wall entry. This will make a new line. */
-	void RemoveIndent();
+	GAMEPLAYCAMERAS_API void RemoveIndent();
 
 	/** Draws a translucent background behind the text. */
-	void DrawTextBackgroundTile(float Opacity);
+	GAMEPLAYCAMERAS_API void DrawTextBackgroundTile(float Opacity);
 
 public:
 
 	/** Request skipping drawing any blocks attached to the current block. */
-	void SkipAttachedBlocks();
+	GAMEPLAYCAMERAS_API void SkipAttachedBlocks();
 	/** Request skipping drawing any children blocks of the current block. */
-	void SkipChildrenBlocks();
+	GAMEPLAYCAMERAS_API void SkipChildrenBlocks();
 	/** Skip all related blocks (attached, children, etc.) */
-	void SkipAllBlocks();
+	GAMEPLAYCAMERAS_API void SkipAllBlocks();
 	/**Gets block visiting flags. */
-	ECameraDebugDrawVisitFlags GetVisitFlags() const;
+	GAMEPLAYCAMERAS_API ECameraDebugDrawVisitFlags GetVisitFlags() const;
 	/** Resets block visiting flags. */
-	void ResetVisitFlags();
+	GAMEPLAYCAMERAS_API void ResetVisitFlags();
 
 public:
 
 	/** Gets the drawing canvas. */
-	UCanvas* GetCanvas() const { return Canvas; }
+	FCanvas* GetCanvas() const { return Canvas; }
+
+	/** Gets the size of the canvas. */
+	FVector2D GetCanvasSize() const;
 
 private:
 
@@ -91,9 +94,9 @@ private:
 private:
 
 	/** The canvas used to draw the text wall. */
-	UCanvas* Canvas;
-	/** The original draw color of the canvas. */
-	FColor OriginalDrawColor;
+	FCanvas* Canvas;
+	/** The draw color of the canvas. */
+	FColor DrawColor;
 
 	/** The font used to render the text wall. */
 	const UFont* RenderFont;
diff --git a/Engine/Plugins/Cameras/GameplayCameras/Source/GameplayCameras/Public/Debug/CameraDirectorTreeDebugBlock.h b/Engine/Plugins/Cameras/GameplayCameras/Source/GameplayCameras/Public/Debug/CameraDirectorTreeDebugBlock.h
--- a/Engine/Plugins/Cameras/GameplayCameras/Source/GameplayCameras/Public/Debug/CameraDirectorTreeDebugBlock.h
+++ b/Engine/Plugins/Cameras/GameplayCameras/Source/GameplayCameras/Public/Debug/CameraDirectorTreeDebugBlock.h
@@ -28,15 +28,17 @@ public:
 protected:
 
 	virtual void OnDebugDraw(const FCameraDebugBlockDrawParams& Params, FCameraDebugRenderer& Renderer) override;
+	virtual void OnSerialize(FArchive& Ar) override;
 
 private:
 
 	struct FDirectorDebugInfo
 	{
 		FString CameraAssetName;
-		FCameraDebugBlock* DebugBlock = nullptr;
 	};
 	TArray<FDirectorDebugInfo> CameraDirectors;
+
+	friend FArchive& operator<< (FArchive&, FDirectorDebugInfo&);
 };
 
 }  // namespace UE::Cameras
diff --git a/Engine/Plugins/Cameras/GameplayCameras/Source/GameplayCameras/Public/Debug/CameraNodeEvaluationResultDebugBlock.h b/Engine/Plugins/Cameras/GameplayCameras/Source/GameplayCameras/Public/Debug/CameraNodeEvaluationResultDebugBlock.h
--- a/Engine/Plugins/Cameras/GameplayCameras/Source/GameplayCameras/Public/Debug/CameraNodeEvaluationResultDebugBlock.h
+++ b/Engine/Plugins/Cameras/GameplayCameras/Source/GameplayCameras/Public/Debug/CameraNodeEvaluationResultDebugBlock.h
@@ -44,6 +44,7 @@ public:
 protected:
 
 	virtual void OnDebugDraw(const FCameraDebugBlockDrawParams& Params, FCameraDebugRenderer& Renderer) override;
+	virtual void OnSerialize(FArchive& Ar) override;
 
 private:
 
diff --git a/Engine/Plugins/Cameras/GameplayCameras/Source/GameplayCameras/Public/Debug/CameraPoseDebugBlock.h b/Engine/Plugins/Cameras/GameplayCameras/Source/GameplayCameras/Public/Debug/CameraPoseDebugBlock.h
--- a/Engine/Plugins/Cameras/GameplayCameras/Source/GameplayCameras/Public/Debug/CameraPoseDebugBlock.h
+++ b/Engine/Plugins/Cameras/GameplayCameras/Source/GameplayCameras/Public/Debug/CameraPoseDebugBlock.h
@@ -37,6 +37,7 @@ public:
 protected:
 
 	virtual void OnDebugDraw(const FCameraDebugBlockDrawParams& Params, FCameraDebugRenderer& Renderer) override;
+	virtual void OnSerialize(FArchive& Ar) override;
 
 private:
 
diff --git a/Engine/Plugins/Cameras/GameplayCameras/Source/GameplayCameras/Public/Debug/CameraSystemTrace.h b/Engine/Plugins/Cameras/GameplayCameras/Source/GameplayCameras/Public/Debug/CameraSystemTrace.h
new file mode 100644
--- /dev/null
+++ b/Engine/Plugins/Cameras/GameplayCameras/Source/GameplayCameras/Public/Debug/CameraSystemTrace.h
@@ -0,0 +1,54 @@
+// Copyright Epic Games, Inc. All Rights Reserved.
+
+#pragma once
+
+#include "CoreTypes.h"
+#include "GameplayCameras.h"
+#include "Trace/Config.h"
+#include "TraceFilter.h"
+
+#if UE_GAMEPLAY_CAMERAS_TRACE
+
+class UWorld;
+
+namespace UE::Cameras
+{
+
+class FCameraDebugBlock;
+class FCameraDebugBlockStorage;
+struct FCameraSystemEvaluationUpdateResult;
+
+/**
+ * Trace utility class for the camera system.
+ */
+class FCameraSystemTrace
+{
+public:
+
+	GAMEPLAYCAMERAS_API static FString ChannelName;
+	GAMEPLAYCAMERAS_API static FString LoggerName;
+	GAMEPLAYCAMERAS_API static FString EvaluationEventName;
+
+public:
+
+	/** Gets whether we are currently replaying traced information (such as with rewind debugger). */
+	GAMEPLAYCAMERAS_API static bool IsTraceReplay();
+	/** Sets whether we are currently replaying traced information (such as with rewind debugger). */
+	GAMEPLAYCAMERAS_API static void SetTraceReplay(bool bInIsReplaying);
+
+	/** Returns whether tracing of camera system evaluation is enabled. */
+	GAMEPLAYCAMERAS_API static bool IsTraceEnabled();
+	/** Records one frame of camera system evaluation. */
+	GAMEPLAYCAMERAS_API static void TraceEvaluation(UWorld* InWorld, const FCameraSystemEvaluationUpdateResult& InResult, FCameraDebugBlock& InRootDebugBlock);
+	/** Reads back one frame of camera system evaluation. */
+	GAMEPLAYCAMERAS_API static FCameraDebugBlock* ReadEvaluationTrace(TArray<uint8> InSerializedBlocks, FCameraDebugBlockStorage& InStorage);
+
+private:
+
+	static bool bIsReplaying;
+};
+
+}  // namespace UE::Cameras
+
+#endif  // UE_GAMEPLAY_CAMERAS_TRACE
+
diff --git a/Engine/Plugins/Cameras/GameplayCameras/Source/GameplayCameras/Public/Debug/CategoryTitleDebugBlock.h b/Engine/Plugins/Cameras/GameplayCameras/Source/GameplayCameras/Public/Debug/CategoryTitleDebugBlock.h
--- a/Engine/Plugins/Cameras/GameplayCameras/Source/GameplayCameras/Public/Debug/CategoryTitleDebugBlock.h
+++ b/Engine/Plugins/Cameras/GameplayCameras/Source/GameplayCameras/Public/Debug/CategoryTitleDebugBlock.h
@@ -21,18 +21,18 @@ class FCategoryTitleDebugBlock : public FCameraDebugBlock
 
 public:
 
-	FCategoryTitleDebugBlock();
-	FCategoryTitleDebugBlock(const FString& InCategory, const FString& InTitle);
+	GAMEPLAYCAMERAS_API FCategoryTitleDebugBlock();
+	GAMEPLAYCAMERAS_API FCategoryTitleDebugBlock(const FString& InCategory, const FString& InTitle);
 
 protected:
 
 	virtual void OnDebugDraw(const FCameraDebugBlockDrawParams& Params, FCameraDebugRenderer& Renderer) override;
+	virtual void OnSerialize(FArchive& Ar) override;
 
 public:
 
 	FString Category;
 	FString Title;
-	FColor TitleColor;
 
 	bool bSkipAttachedBlocksIfInactive = true;
 	bool bSkipChildrenBlocksIfInactive = true;
diff --git a/Engine/Plugins/Cameras/GameplayCameras/Source/GameplayCameras/Public/Debug/DebugTextRenderer.h b/Engine/Plugins/Cameras/GameplayCameras/Source/GameplayCameras/Public/Debug/DebugTextRenderer.h
--- a/Engine/Plugins/Cameras/GameplayCameras/Source/GameplayCameras/Public/Debug/DebugTextRenderer.h
+++ b/Engine/Plugins/Cameras/GameplayCameras/Source/GameplayCameras/Public/Debug/DebugTextRenderer.h
@@ -8,7 +8,7 @@
 #include "Math/UnrealMath.h"
 #include "Misc/TVariant.h"
 
-class UCanvas;
+class FCanvas;
 class UFont;
 
 #if UE_GAMEPLAY_CAMERAS_DEBUG
@@ -53,7 +53,7 @@ struct FDebugTextDrawCommand
 {
 	TStringView<TCHAR> TextView;
 
-	void Execute(UCanvas* Canvas, const UFont* Font, FVector2f& InOutDrawPosition) const;
+	void Execute(FCanvas* Canvas, const FColor& DrawColor, const UFont* Font, FVector2f& InOutDrawPosition) const;
 };
 
 /** Command for moving the drawing position to a new line. */
@@ -70,7 +70,7 @@ struct FDebugTextSetColorCommand
 {
 	FColor DrawColor;
 
-	void Execute(UCanvas* Canvas) const;
+	void Execute(FColor* OutDrawColor) const;
 };
 
 /** A debug text drawing command, which can be of multiple types. */
@@ -97,7 +97,7 @@ public:
 	bool bEndWithNewLine = false;
 
 	/** Creates a new debug text renderer. */
-	FDebugTextRenderer(UCanvas* InCanvas, const UFont* InFont);
+	FDebugTextRenderer(FCanvas* InCanvas, const FColor& InDrawColor, const UFont* InFont);
 
 	/** Renders the given text to the canvas. */
 	void RenderText(float StartingDrawY, const TStringView<TCHAR> TextView);
@@ -130,7 +130,8 @@ private:
 
 private:
 
-	UCanvas* Canvas;
+	FCanvas* Canvas;
+	FColor DrawColor;
 	const UFont* Font;
 	FVector2f NextDrawPosition;
 	float RightMargin = 0;
diff --git a/Engine/Plugins/Cameras/GameplayCameras/Source/GameplayCameras/Public/Debug/RootCameraDebugBlock.h b/Engine/Plugins/Cameras/GameplayCameras/Source/GameplayCameras/Public/Debug/RootCameraDebugBlock.h
--- a/Engine/Plugins/Cameras/GameplayCameras/Source/GameplayCameras/Public/Debug/RootCameraDebugBlock.h
+++ b/Engine/Plugins/Cameras/GameplayCameras/Source/GameplayCameras/Public/Debug/RootCameraDebugBlock.h
@@ -27,10 +27,10 @@ class FRootCameraDebugBlock : public FCameraDebugBlock
 public:
 
 	/** Build all debug blocks for the last evaluation frame. */
-	void BuildDebugBlocks(const FCameraSystemEvaluator& CameraSystem, const FCameraDebugBlockBuildParams& Params, FCameraDebugBlockBuilder& Builder);
+	GAMEPLAYCAMERAS_API void BuildDebugBlocks(const FCameraSystemEvaluator& CameraSystem, const FCameraDebugBlockBuildParams& Params, FCameraDebugBlockBuilder& Builder);
 
 	/** Initiate the debug drawing. */
-	void RootDebugDraw(FCameraDebugRenderer& Renderer);
+	GAMEPLAYCAMERAS_API void RootDebugDraw(FCameraDebugRenderer& Renderer);
 };
 
 }  // namespace UE::Cameras
diff --git a/Engine/Plugins/Cameras/GameplayCameras/Source/GameplayCameras/Public/Debug/VariableTableDebugBlock.h b/Engine/Plugins/Cameras/GameplayCameras/Source/GameplayCameras/Public/Debug/VariableTableDebugBlock.h
--- a/Engine/Plugins/Cameras/GameplayCameras/Source/GameplayCameras/Public/Debug/VariableTableDebugBlock.h
+++ b/Engine/Plugins/Cameras/GameplayCameras/Source/GameplayCameras/Public/Debug/VariableTableDebugBlock.h
@@ -21,9 +21,9 @@ class FVariableTableDebugBlock : public FCameraDebugBlock
 public:
 
 	/** Creates a new variable table debug block. */
-	FVariableTableDebugBlock();
+	GAMEPLAYCAMERAS_API FVariableTableDebugBlock();
 	/** Creates a new variable table debug block. */
-	FVariableTableDebugBlock(const FCameraVariableTable& InVariableTable);
+	GAMEPLAYCAMERAS_API FVariableTableDebugBlock(const FCameraVariableTable& InVariableTable);
 
 	/** Specifies the console variable to use to toggle the printing of variable IDs. */
 	FVariableTableDebugBlock& WithShowVariableIdsCVar(const TCHAR* InShowVariableIdsCVarName)
@@ -35,6 +35,7 @@ public:
 protected:
 
 	virtual void OnDebugDraw(const FCameraDebugBlockDrawParams& Params, FCameraDebugRenderer& Renderer) override;
+	virtual void OnSerialize(FArchive& Ar) override;
 
 	void Initialize(const FCameraVariableTable& InVariableTable);
 
@@ -45,12 +46,14 @@ private:
 		uint32 Id;
 		FString Name;
 		FString Value;
-		bool bWritten : 1;
-		bool bWrittenThisFrame : 1;
+		bool bWritten;
+		bool bWrittenThisFrame;
 	};
 	TArray<FEntryDebugInfo> Entries;
 
 	FString ShowVariableIdsCVarName;
+
+	friend FArchive& operator<< (FArchive&, FEntryDebugInfo&);
 };
 
 }  // namespace UE::Cameras
diff --git a/Engine/Plugins/Cameras/GameplayCameras/Source/GameplayCameras/Public/Debug/ViewfinderDebugBlock.h b/Engine/Plugins/Cameras/GameplayCameras/Source/GameplayCameras/Public/Debug/ViewfinderDebugBlock.h
--- a/Engine/Plugins/Cameras/GameplayCameras/Source/GameplayCameras/Public/Debug/ViewfinderDebugBlock.h
+++ b/Engine/Plugins/Cameras/GameplayCameras/Source/GameplayCameras/Public/Debug/ViewfinderDebugBlock.h
@@ -15,7 +15,7 @@ class FViewfinderDebugBlock : public FCameraDebugBlock
 
 public:
 
-	FViewfinderDebugBlock();
+	GAMEPLAYCAMERAS_API FViewfinderDebugBlock();
 
 protected:
 
diff --git a/Engine/Plugins/Cameras/GameplayCameras/Source/GameplayCameras/Public/GameplayCameras.h b/Engine/Plugins/Cameras/GameplayCameras/Source/GameplayCameras/Public/GameplayCameras.h
--- a/Engine/Plugins/Cameras/GameplayCameras/Source/GameplayCameras/Public/GameplayCameras.h
+++ b/Engine/Plugins/Cameras/GameplayCameras/Source/GameplayCameras/Public/GameplayCameras.h
@@ -15,6 +15,12 @@
 	#endif
 #endif
 
+#if UE_TRACE_ENABLED && !IS_PROGRAM && !UE_BUILD_SHIPPING && !UE_BUILD_TEST
+#define UE_GAMEPLAY_CAMERAS_TRACE 1
+#else
+#define UE_GAMEPLAY_CAMERAS_TRACE 0
+#endif
+
 DECLARE_STATS_GROUP(TEXT("Camera System Evaluation"), STATGROUP_CameraSystem, STATCAT_Advanced)
 DECLARE_STATS_GROUP(TEXT("Camera Animation Evaluation"), STATGROUP_CameraAnimation, STATCAT_Advanced)
 
diff --git a/Engine/Plugins/Cameras/GameplayCameras/Source/GameplayCamerasEditor/GameplayCamerasEditor.Build.cs b/Engine/Plugins/Cameras/GameplayCameras/Source/GameplayCamerasEditor/GameplayCamerasEditor.Build.cs
--- a/Engine/Plugins/Cameras/GameplayCameras/Source/GameplayCamerasEditor/GameplayCamerasEditor.Build.cs
+++ b/Engine/Plugins/Cameras/GameplayCameras/Source/GameplayCamerasEditor/GameplayCamerasEditor.Build.cs
@@ -31,17 +31,22 @@ public class GameplayCamerasEditor : ModuleRules
 				"InputCore",
 				"InteractiveToolsFramework",
 				"Kismet",
+				"LevelEditor",
 				"Projects",
+				"RewindDebuggerInterface",
 				"Slate",
 				"SlateCore",
 				"TimeManagement",
 				"ToolMenus",
+				"TraceAnalysis",
+				"TraceLog",
+				"TraceServices",
+				"TraceInsights",
 				"UnrealEd",
 			}
 		);
 
 		var DynamicModuleNames = new string[] {
-			"LevelEditor",
 			"PropertyEditor",
 			"WorkspaceMenuStructure",
 		};
diff --git a/Engine/Plugins/Cameras/GameplayCameras/Source/GameplayCamerasEditor/Private/Debugger/SBlendStacksDebugPanel.cpp b/Engine/Plugins/Cameras/GameplayCameras/Source/GameplayCamerasEditor/Private/Debugger/SBlendStacksDebugPanel.cpp
--- a/Engine/Plugins/Cameras/GameplayCameras/Source/GameplayCamerasEditor/Private/Debugger/SBlendStacksDebugPanel.cpp
+++ b/Engine/Plugins/Cameras/GameplayCameras/Source/GameplayCamerasEditor/Private/Debugger/SBlendStacksDebugPanel.cpp
@@ -4,6 +4,7 @@
 
 #include "Debugger/SDebugWidgetUtils.h"
 #include "Widgets/Input/SCheckBox.h"
+#include "Widgets/Layout/SBox.h"
 #include "Widgets/SBoxPanel.h"
 
 #define LOCTEXT_NAMESPACE "SBlendStacksDebugPanel"
diff --git a/Engine/Plugins/Cameras/GameplayCameras/Source/GameplayCamerasEditor/Private/Debugger/SCameraNodeTreeDebugPanel.cpp b/Engine/Plugins/Cameras/GameplayCameras/Source/GameplayCamerasEditor/Private/Debugger/SCameraNodeTreeDebugPanel.cpp
--- a/Engine/Plugins/Cameras/GameplayCameras/Source/GameplayCamerasEditor/Private/Debugger/SCameraNodeTreeDebugPanel.cpp
+++ b/Engine/Plugins/Cameras/GameplayCameras/Source/GameplayCamerasEditor/Private/Debugger/SCameraNodeTreeDebugPanel.cpp
@@ -21,6 +21,7 @@ void SCameraNodeTreeDebugPanel::Construct(const FArguments& InArgs)
 			.AutoHeight()
 			[
 				SNew(STextBlock)
+					.Margin(4.f)
 					.Text(LOCTEXT("FilterNodeNames", "Filter node names:"))
 			]
 		+ SVerticalBox::Slot()
diff --git a/Engine/Plugins/Cameras/GameplayCameras/Source/GameplayCamerasEditor/Private/Debugger/SDebugCategoryButton.cpp b/Engine/Plugins/Cameras/GameplayCameras/Source/GameplayCamerasEditor/Private/Debugger/SDebugCategoryButton.cpp
--- a/Engine/Plugins/Cameras/GameplayCameras/Source/GameplayCamerasEditor/Private/Debugger/SDebugCategoryButton.cpp
+++ b/Engine/Plugins/Cameras/GameplayCameras/Source/GameplayCamerasEditor/Private/Debugger/SDebugCategoryButton.cpp
@@ -23,8 +23,8 @@ void SDebugCategoryButton::Construct(const FArguments& InArgs)
 	ActiveModeBorderImage = FAppStyle::GetBrush("ModeSelector.ToggleButton.Pressed");
 	HoverBorderImage = FAppStyle::GetBrush("ModeSelector.ToggleButton.Hovered");
 
-	const FMargin IconPadding(4.0f, 0.0f, 4.0f, 0.0f);
-	const FMargin TextPadding(0.0f, 0.0f, 0.0f, 0.0f);
+	const FMargin IconPadding(8.0f, 8.0f, 4.0f, 8.0f);
+	const FMargin TextPadding(4.0f, 8.0f, 8.0f, 8.0f);
 
 	TSharedRef<SHorizontalBox> ButtonContents = SNew(SHorizontalBox);
 
diff --git a/Engine/Plugins/Cameras/GameplayCameras/Source/GameplayCamerasEditor/Private/Debugger/SDebugWidgetUtils.cpp b/Engine/Plugins/Cameras/GameplayCameras/Source/GameplayCamerasEditor/Private/Debugger/SDebugWidgetUtils.cpp
--- a/Engine/Plugins/Cameras/GameplayCameras/Source/GameplayCamerasEditor/Private/Debugger/SDebugWidgetUtils.cpp
+++ b/Engine/Plugins/Cameras/GameplayCameras/Source/GameplayCamerasEditor/Private/Debugger/SDebugWidgetUtils.cpp
@@ -15,6 +15,7 @@ TMap<FString, IConsoleVariable*> SDebugWidgetUtils::CachedConsoleVariables;
 TSharedRef<SCheckBox> SDebugWidgetUtils::CreateConsoleVariableCheckBox(const FText& Text, const FString& ConsoleVariableName)
 {
 	return SNew(SCheckBox)
+		.Padding(4.f)
 		.IsChecked_Lambda([ConsoleVariableName]() 
 				{
 					if (IConsoleVariable* ConsoleVariable = SDebugWidgetUtils::GetConsoleVariable(ConsoleVariableName))
@@ -40,6 +41,7 @@ TSharedRef<SCheckBox> SDebugWidgetUtils::CreateConsoleVariableCheckBox(const FTe
 TSharedRef<SEditableTextBox> SDebugWidgetUtils::CreateConsoleVariableTextBox(const FString& ConsoleVariableName)
 {
 	return SNew(SEditableTextBox)
+		.Padding(4.f)
 		.Text_Lambda([ConsoleVariableName]()
 				{
 					if (IConsoleVariable* ConsoleVariable = SDebugWidgetUtils::GetConsoleVariable(ConsoleVariableName))
diff --git a/Engine/Plugins/Cameras/GameplayCameras/Source/GameplayCamerasEditor/Private/Debugger/SGameplayCamerasDebugger.cpp b/Engine/Plugins/Cameras/GameplayCameras/Source/GameplayCamerasEditor/Private/Debugger/SGameplayCamerasDebugger.cpp
--- a/Engine/Plugins/Cameras/GameplayCameras/Source/GameplayCamerasEditor/Private/Debugger/SGameplayCamerasDebugger.cpp
+++ b/Engine/Plugins/Cameras/GameplayCameras/Source/GameplayCamerasEditor/Private/Debugger/SGameplayCamerasDebugger.cpp
@@ -16,6 +16,7 @@
 #include "Styling/SlateTypes.h"
 #include "ToolMenus.h"
 #include "Widgets/Docking/SDockTab.h"
+#include "Widgets/Input/SButton.h"
 #include "Widgets/Layout/SBox.h"
 #include "Widgets/SNullWidget.h"
 #include "WorkspaceMenuStructure.h"
@@ -27,6 +28,7 @@ namespace UE::Cameras
 {
 
 const FName SGameplayCamerasDebugger::WindowName(TEXT("GameplayCamerasDebugger"));
+const FName SGameplayCamerasDebugger::MenubarName(TEXT("GameplayCamerasDebugger.Menubar"));
 const FName SGameplayCamerasDebugger::ToolbarName(TEXT("GameplayCamerasDebugger.Toolbar"));
 
 void SGameplayCamerasDebugger::RegisterTabSpawners()
@@ -72,7 +74,7 @@ void SGameplayCamerasDebugger::Construct(const FArguments& InArgs)
 {
 	IGameplayCamerasEditorModule& GameplayCamerasEditorModule = FModuleManager::GetModuleChecked<IGameplayCamerasEditorModule>(TEXT("GameplayCamerasEditor"));
 	TSharedRef<FGameplayCamerasEditorStyle> GameplayCamerasEditorStyle = FGameplayCamerasEditorStyle::Get();
-	const FName& GameplayCamerasEditorStyleName = GameplayCamerasEditorStyle->GetStyleSetName();
+	GameplayCamerasEditorStyleName = GameplayCamerasEditorStyle->GetStyleSetName();
 
 	// Setup commands.
 	const FGameplayCamerasDebuggerCommands& Commands = FGameplayCamerasDebuggerCommands::Get();
@@ -85,16 +87,16 @@ void SGameplayCamerasDebugger::Construct(const FArguments& InArgs)
 
 	// Main menu bar.
 	UToolMenus* ToolMenus = UToolMenus::Get();
-	if (!ToolMenus->IsMenuRegistered(SGameplayCamerasDebugger::ToolbarName))
+	if (!ToolMenus->IsMenuRegistered(SGameplayCamerasDebugger::MenubarName))
 	{
 		FToolMenuOwnerScoped ToolMenuOwnerScope(this);
 
-		UToolMenu* ToolBarMenu = ToolMenus->RegisterMenu(
-				SGameplayCamerasDebugger::ToolbarName, NAME_None, EMultiBoxType::ToolBar);
+		UToolMenu* Menubar = ToolMenus->RegisterMenu(
+				SGameplayCamerasDebugger::MenubarName, NAME_None, EMultiBoxType::MenuBar);
 	}
-	FToolMenuContext ToolbarContext;
-	TSharedRef<SWidget> ToolbarContents = ToolMenus->GenerateWidget(
-			SGameplayCamerasDebugger::ToolbarName, ToolbarContext);
+	FToolMenuContext MenubarContext;
+	TSharedRef<SWidget> MenubarContents = ToolMenus->GenerateWidget(
+			SGameplayCamerasDebugger::MenubarName, MenubarContext);
 
 	// Empty panel.
 	EmptyPanel = SNew(SBox)
@@ -107,45 +109,58 @@ void SGameplayCamerasDebugger::Construct(const FArguments& InArgs)
 		];
 
 	// Toolbar with debug category buttons.
-	FToolBarBuilder DebugCategoriesToolbarBuilder(CommandList.ToSharedPtr(), FMultiBoxCustomization::None);
-	DebugCategoriesToolbarBuilder.BeginSection(TEXT("Main"));
-	{
-		DebugCategoriesToolbarBuilder.AddToolBarButton(
-				Commands.EnableDebugInfo,
-				NAME_None,
-				TAttribute<FText>(FText::FromString(TEXT(""))),
-				TAttribute<FText>(),
-				FSlateIcon(GameplayCamerasEditorStyleName, "Debugger.EnableDebugInfo.Icon"));
-	}
-	DebugCategoriesToolbarBuilder.EndSection();
-	DebugCategoriesToolbarBuilder.BeginSection(TEXT("DebugCategories"));
+	if (!ToolMenus->IsMenuRegistered(SGameplayCamerasDebugger::ToolbarName))
 	{
-		TArray<FCameraDebugCategoryInfo> RegisteredDebugCategories;
-		GameplayCamerasEditorModule.GetRegisteredDebugCategories(RegisteredDebugCategories);
-		for (const FCameraDebugCategoryInfo& DebugCategory : RegisteredDebugCategories)
+		FToolMenuOwnerScoped ToolMenuOwnerScope(this);
+
+		UToolMenu* Toolbar = ToolMenus->RegisterMenu(
+				SGameplayCamerasDebugger::ToolbarName, NAME_None, EMultiBoxType::SlimHorizontalToolBar);
+
+		FToolMenuSection& MainSection = Toolbar->AddSection(TEXT("Main"));
 		{
-			TSharedPtr<SWidget> DebugCategoryPanel = GameplayCamerasEditorModule.CreateDebugCategoryPanel(DebugCategory.Name);
-			if (DebugCategoryPanel.IsValid())
-			{
-				DebugPanels.Add(DebugCategory.Name, DebugCategoryPanel);
-			}
-			else
+			FToolMenuEntry ToggleDebugInfo = FToolMenuEntry::InitToolBarButton(
+					Commands.EnableDebugInfo,
+					TAttribute<FText>::CreateSP(this, &SGameplayCamerasDebugger::GetToggleDebugDrawText),
+					TAttribute<FText>(),
+					TAttribute<FSlateIcon>::CreateSP(this, &SGameplayCamerasDebugger::GetToggleDebugDrawIcon));
+			ToggleDebugInfo.SetCommandList(CommandList);
+			MainSection.AddEntry(ToggleDebugInfo);
+		}
+	
+		FToolMenuSection& DebugCategoriesSection = Toolbar->AddSection(TEXT("DebugCategories"));
+		{
+			TArray<FCameraDebugCategoryInfo> RegisteredDebugCategories;
+			GameplayCamerasEditorModule.GetRegisteredDebugCategories(RegisteredDebugCategories);
+			for (const FCameraDebugCategoryInfo& DebugCategory : RegisteredDebugCategories)
 			{
-				// If there aren't any special UI controls for this category, use an empty panel.
-				DebugPanels.Add(DebugCategory.Name, EmptyPanel);
+				TSharedPtr<SWidget> DebugCategoryPanel = GameplayCamerasEditorModule.CreateDebugCategoryPanel(DebugCategory.Name);
+				if (DebugCategoryPanel.IsValid())
+				{
+					DebugPanels.Add(DebugCategory.Name, DebugCategoryPanel);
+				}
+				else
+				{
+					// If there aren't any special UI controls for this category, use an empty panel.
+					DebugPanels.Add(DebugCategory.Name, EmptyPanel);
+				}
+
+				FToolMenuEntry ToggleDebugCategory = FToolMenuEntry::InitToolBarButton(
+					FName(DebugCategory.Name),
+					FUIAction(
+						FExecuteAction::CreateSP(this, &SGameplayCamerasDebugger::SetActiveDebugCategoryPanel, DebugCategory.Name),
+						FCanExecuteAction(),
+						FIsActionChecked::CreateStatic(&SGameplayCamerasDebugger::IsDebugCategoryActive, DebugCategory.Name)),
+					DebugCategory.DisplayText,
+					DebugCategory.ToolTipText,
+					DebugCategory.IconImage,
+					EUserInterfaceActionType::ToggleButton);
+				DebugCategoriesSection.AddEntry(ToggleDebugCategory);
 			}
-
-			TSharedRef<SDebugCategoryButton> DebugCategoryButton = SNew(SDebugCategoryButton)
-				.DebugCategoryName(DebugCategory.Name)
-				.DisplayText(DebugCategory.DisplayText)
-				.ToolTipText(DebugCategory.ToolTipText)
-				.IconImage(DebugCategory.IconImage.GetIcon())
-				.IsDebugCategoryActive_Lambda([DebugCategory](){ return SGameplayCamerasDebugger::IsDebugCategoryActive(DebugCategory.Name); })
-				.RequestDebugCategoryChange(this, &SGameplayCamerasDebugger::SetActiveDebugCategoryPanel);
-			DebugCategoriesToolbarBuilder.AddToolBarWidget(DebugCategoryButton);
 		}
 	}
-	DebugCategoriesToolbarBuilder.EndSection();
+	FToolMenuContext ToolbarContext;
+	TSharedRef<SWidget> ToolbarContents = ToolMenus->GenerateWidget(
+			SGameplayCamerasDebugger::ToolbarName, ToolbarContext);
 
 	// Main layout.
 	ChildSlot
@@ -154,32 +169,57 @@ void SGameplayCamerasDebugger::Construct(const FArguments& InArgs)
 		+ SVerticalBox::Slot()
 			.AutoHeight()
 			[
-				ToolbarContents
+				MenubarContents
 			]
 		+ SVerticalBox::Slot()
 			.AutoHeight()
 			[
-				DebugCategoriesToolbarBuilder.MakeWidget()
+				ToolbarContents
 			]
 		+ SVerticalBox::Slot()
 		.Padding(2.0)
 			[
 				SAssignNew(PanelHost, SBox)
-				[
-					SNullWidget::NullWidget
-				]
+					.Padding(8.f)
+					[
+						EmptyPanel.ToSharedRef()
+					]
 			]
 	];
 }
 
-bool SGameplayCamerasDebugger::IsDebugCategoryActive(const FString& InCategoryName)
+FText SGameplayCamerasDebugger::GetToggleDebugDrawText() const
+{
+	if (GGameplayCamerasDebugEnable)
+	{
+		return LOCTEXT("DebugInfoEnabled", "Debug Info Enabled");
+	}
+	else
+	{
+		return LOCTEXT("DebugInfoDisabled", "Debug Info Disabled");
+	}
+}
+
+FSlateIcon SGameplayCamerasDebugger::GetToggleDebugDrawIcon() const
+{
+	if (GGameplayCamerasDebugEnable)
+	{
+		return FSlateIcon(GameplayCamerasEditorStyleName, "Debugger.EnableDebugInfo.Icon");
+	}
+	else
+	{
+		return FSlateIcon(GameplayCamerasEditorStyleName, "Debugger.DisableDebugInfo.Icon");
+	}
+}
+
+bool SGameplayCamerasDebugger::IsDebugCategoryActive(FString InCategoryName)
 {
 	TArray<FStringView, TInlineAllocator<4>> ActiveCategories;
 	UE::String::ParseTokens(GGameplayCamerasDebugCategories, ',', ActiveCategories);
 	return ActiveCategories.Contains(InCategoryName);
 }
 
-void SGameplayCamerasDebugger::SetActiveDebugCategoryPanel(const FString& InCategoryName)
+void SGameplayCamerasDebugger::SetActiveDebugCategoryPanel(FString InCategoryName)
 {
 	if (ensureMsgf(
 				DebugPanels.Contains(InCategoryName), 
diff --git a/Engine/Plugins/Cameras/GameplayCameras/Source/GameplayCamerasEditor/Private/Debugger/SGameplayCamerasDebugger.h b/Engine/Plugins/Cameras/GameplayCameras/Source/GameplayCamerasEditor/Private/Debugger/SGameplayCamerasDebugger.h
--- a/Engine/Plugins/Cameras/GameplayCameras/Source/GameplayCamerasEditor/Private/Debugger/SGameplayCamerasDebugger.h
+++ b/Engine/Plugins/Cameras/GameplayCameras/Source/GameplayCamerasEditor/Private/Debugger/SGameplayCamerasDebugger.h
@@ -9,6 +9,7 @@ class FTabManager;
 class FSpawnTabArgs;
 class SBox;
 class SDockTab;
+struct FSlateIcon;
 
 namespace UE::Cameras
 {
@@ -18,6 +19,7 @@ class SGameplayCamerasDebugger : public SCompoundWidget
 public:
 
 	static const FName WindowName;
+	static const FName MenubarName;
 	static const FName ToolbarName;
 
 	static void RegisterTabSpawners();
@@ -36,11 +38,16 @@ public:
 
 protected:
 
-	static bool IsDebugCategoryActive(const FString& InCategoryName);
-	void SetActiveDebugCategoryPanel(const FString& InCategoryName);
+	static bool IsDebugCategoryActive(FString InCategoryName);
+	void SetActiveDebugCategoryPanel(FString InCategoryName);
+
+	FText GetToggleDebugDrawText() const;
+	FSlateIcon GetToggleDebugDrawIcon() const;
 
 private:
 
+	FName GameplayCamerasEditorStyleName;
+
 	TSharedPtr<SBox> PanelHost;
 	TSharedPtr<SWidget> EmptyPanel;
 	TMap<FString, TSharedPtr<SWidget>> DebugPanels;
diff --git a/Engine/Plugins/Cameras/GameplayCameras/Source/GameplayCamerasEditor/Private/GameplayCamerasEditorModule.cpp b/Engine/Plugins/Cameras/GameplayCameras/Source/GameplayCamerasEditor/Private/GameplayCamerasEditorModule.cpp
--- a/Engine/Plugins/Cameras/GameplayCameras/Source/GameplayCamerasEditor/Private/GameplayCamerasEditorModule.cpp
+++ b/Engine/Plugins/Cameras/GameplayCameras/Source/GameplayCamerasEditor/Private/GameplayCamerasEditorModule.cpp
@@ -11,9 +11,11 @@
 #include "Debugger/SBlendStacksDebugPanel.h"
 #include "Debugger/SCameraNodeTreeDebugPanel.h"
 #include "Debugger/SGameplayCamerasDebugger.h"
+#include "Features/IModularFeatures.h"
 #include "GameplayCamerasLiveEditManager.h"
 #include "IGameplayCamerasEditorModule.h"
 #include "IGameplayCamerasModule.h"
+#include "IRewindDebuggerExtension.h"
 #include "ISettingsModule.h"
 #include "Misc/CoreDelegates.h"
 #include "Modules/ModuleManager.h"
@@ -21,6 +23,9 @@
 #include "ToolMenus.h"
 #include "Toolkits/CameraAssetEditorToolkit.h"
 #include "Toolkits/CameraRigAssetEditorToolkit.h"
+#include "Trace/CameraSystemRewindDebuggerExtension.h"
+#include "Trace/CameraSystemRewindDebuggerTrack.h"
+#include "Trace/CameraSystemTraceModule.h"
 
 #define LOCTEXT_NAMESPACE "GameplayCamerasEditor"
 
@@ -52,6 +57,7 @@ public:
 
 		RegisterSettings();
 		RegisterCoreDebugCategories();
+		RegisterRewindDebuggerFeatures();
 		InitializeLiveEditManager();
 
 		UToolMenus::RegisterStartupCallback(FSimpleMulticastDelegate::FDelegate::CreateRaw(
@@ -70,6 +76,7 @@ public:
 
 		UnregisterSettings();
 		UnregisterCoreDebugCategories();
+		UnregisterRewindDebuggerFeatures();
 		TeardownLiveEditManager();
 
 		FCoreDelegates::OnPostEngineInit.RemoveAll(this);
@@ -196,7 +203,7 @@ private:
 			});
 		RegisterDebugCategory(FCameraDebugCategoryInfo{
 				FCameraDebugCategories::DirectorTree,
-				LOCTEXT("DirectorTreeDebugCategory", "Node Tree"),
+				LOCTEXT("DirectorTreeDebugCategory", "Director Tree"),
 				LOCTEXT("DirectorTreeDebugCategoryToolTip", "Shows the active/inactive directors, and their evaluation context"),
 				FSlateIcon(GameplayCamerasEditorStyleName, "DebugCategory.DirectorTree.Icon")
 			});
@@ -246,6 +253,28 @@ private:
 		FGameplayCamerasDebuggerCommands::Register();
 	}
 
+	void RegisterRewindDebuggerFeatures()
+	{
+		using namespace UE::Cameras;
+
+		TraceModule = MakeShared<UE::Cameras::FCameraSystemTraceModule>();
+		RewindDebuggerExtension = MakeShared<FCameraSystemRewindDebuggerExtension>();
+		RewindDebuggerTrackCreator = MakeShared<FCameraSystemRewindDebuggerTrackCreator>();
+
+		IModularFeatures& ModularFeatures = IModularFeatures::Get();
+		ModularFeatures.RegisterModularFeature(IRewindDebuggerExtension::ModularFeatureName, RewindDebuggerExtension.Get());
+		ModularFeatures.RegisterModularFeature(RewindDebugger::IRewindDebuggerTrackCreator::ModularFeatureName, RewindDebuggerTrackCreator.Get());
+		ModularFeatures.RegisterModularFeature(TraceServices::ModuleFeatureName, TraceModule.Get());
+	}
+
+	void UnregisterRewindDebuggerFeatures()
+	{
+		IModularFeatures& ModularFeatures = IModularFeatures::Get();
+		ModularFeatures.UnregisterModularFeature(IRewindDebuggerExtension::ModularFeatureName, RewindDebuggerExtension.Get());
+		ModularFeatures.UnregisterModularFeature(RewindDebugger::IRewindDebuggerTrackCreator::ModularFeatureName, RewindDebuggerTrackCreator.Get());
+		ModularFeatures.UnregisterModularFeature(TraceServices::ModuleFeatureName, TraceModule.Get());
+	}
+
 	void InitializeLiveEditManager()
 	{
 		using namespace UE::Cameras;
@@ -278,6 +307,10 @@ private:
 
 	TMap<FString, UE::Cameras::FCameraDebugCategoryInfo> DebugCategoryInfos;
 	TMap<FString, FOnCreateDebugCategoryPanel> DebugCategoryPanelCreators;
+
+	TSharedPtr<UE::Cameras::FCameraSystemTraceModule> TraceModule;
+	TSharedPtr<UE::Cameras::FCameraSystemRewindDebuggerExtension> RewindDebuggerExtension;
+	TSharedPtr<UE::Cameras::FCameraSystemRewindDebuggerTrackCreator> RewindDebuggerTrackCreator;
 };
 
 IMPLEMENT_MODULE(FGameplayCamerasEditorModule, GameplayCamerasEditor);
diff --git a/Engine/Plugins/Cameras/GameplayCameras/Source/GameplayCamerasEditor/Private/Styles/GameplayCamerasEditorStyle.cpp b/Engine/Plugins/Cameras/GameplayCameras/Source/GameplayCamerasEditor/Private/Styles/GameplayCamerasEditorStyle.cpp
--- a/Engine/Plugins/Cameras/GameplayCameras/Source/GameplayCamerasEditor/Private/Styles/GameplayCamerasEditorStyle.cpp
+++ b/Engine/Plugins/Cameras/GameplayCameras/Source/GameplayCamerasEditor/Private/Styles/GameplayCamerasEditorStyle.cpp
@@ -27,6 +27,7 @@ FGameplayCamerasEditorStyle::FGameplayCamerasEditorStyle()
 	SetCoreContentRoot(FPaths::EngineContentDir() / TEXT("Slate"));
 
 	Set("Debugger.EnableDebugInfo.Icon", new IMAGE_BRUSH("/Icons/EnableDebugInfo", Icon16x16));
+	Set("Debugger.DisableDebugInfo.Icon", new IMAGE_BRUSH("/Icons/DisableDebugInfo", Icon16x16));
 
 	Set("DebugCategory.NodeTree.Icon", new IMAGE_BRUSH("/Icons/DebugCategory-NodeTree", Icon16x16));
 	Set("DebugCategory.DirectorTree.Icon", new IMAGE_BRUSH("/Icons/DebugCategory-DirectorTree", Icon16x16));
diff --git a/Engine/Plugins/Cameras/GameplayCameras/Source/GameplayCamerasEditor/Private/Trace/CameraSystemRewindDebuggerExtension.cpp b/Engine/Plugins/Cameras/GameplayCameras/Source/GameplayCamerasEditor/Private/Trace/CameraSystemRewindDebuggerExtension.cpp
new file mode 100644
--- /dev/null
+++ b/Engine/Plugins/Cameras/GameplayCameras/Source/GameplayCamerasEditor/Private/Trace/CameraSystemRewindDebuggerExtension.cpp
@@ -0,0 +1,108 @@
+// Copyright Epic Games, Inc. All Rights Reserved.
+
+#include "Trace/CameraSystemRewindDebuggerExtension.h"
+
+#include "Core/CameraPose.h"
+#include "Debug/CameraDebugRenderer.h"
+#include "Debug/CameraSystemTrace.h"
+#include "Debug/CategoryTitleDebugBlock.h"
+#include "Debug/DebugDrawService.h"
+#include "Debug/RootCameraDebugBlock.h"
+#include "Engine/Canvas.h"
+#include "Engine/World.h"
+#include "GameFramework/GameplayCameraSystemActor.h"
+#include "Trace/CameraSystemTraceProvider.h"
+
+namespace UE::Cameras
+{
+
+FCameraSystemRewindDebuggerExtension::FCameraSystemRewindDebuggerExtension()
+{
+}
+
+FCameraSystemRewindDebuggerExtension::~FCameraSystemRewindDebuggerExtension()
+{
+	EnsureDebugDrawDelegate(false);
+}
+
+void FCameraSystemRewindDebuggerExtension::RecordingStarted(IRewindDebugger* RewindDebugger)
+{
+	UE::Trace::ToggleChannel(*FCameraSystemTrace::ChannelName, true);
+}
+
+void FCameraSystemRewindDebuggerExtension::Update(float DeltaTime, IRewindDebugger* RewindDebugger)
+{
+	if (RewindDebugger->IsPIESimulating() || RewindDebugger->GetRecordingDuration() == 0.0)
+	{
+		EnsureDebugDrawDelegate(false);
+		return;
+	}
+
+	EnsureDebugDrawDelegate(true);
+
+	const TraceServices::IAnalysisSession* AnalysisSession = RewindDebugger->GetAnalysisSession();
+	TraceServices::FAnalysisSessionReadScope SessionReadScope(*AnalysisSession);
+	
+	const FCameraSystemTraceProvider* CameraSystemProvider = AnalysisSession->ReadProvider<FCameraSystemTraceProvider>(FCameraSystemTraceProvider::ProviderName);
+	const TraceServices::IFrameProvider& FrameProvider = TraceServices::ReadFrameProvider(*RewindDebugger->GetAnalysisSession());
+
+	double CurrentTraceTime = RewindDebugger->CurrentTraceTime();
+
+	TraceServices::FFrame Frame;
+	if (FrameProvider.GetFrameFromTime(ETraceFrameType::TraceFrameType_Game, CurrentTraceTime, Frame))
+	{
+		const FCameraSystemTraceFrameData* FoundFrameData = nullptr;
+
+		const FCameraSystemTraceTimeline* CameraSystemTimeline = CameraSystemProvider->GetTimeline();
+		CameraSystemTimeline->EnumerateEvents(
+				Frame.StartTime, Frame.EndTime, 
+				[&FoundFrameData](double InStartTime, double InEndTime, uint32 InDepth, const FCameraSystemTraceFrameData& FrameData)  
+				{
+					FoundFrameData = &FrameData;
+					return TraceServices::EEventEnumerate::Continue;
+				});
+
+		if (FoundFrameData && CurrentTraceTime != LastTraceTime)
+		{
+			LastTraceTime = CurrentTraceTime;
+
+			DebugBlockStorage.DestroyDebugBlocks();
+
+			FCameraDebugBlock* ReadBlock = FCameraSystemTrace::ReadEvaluationTrace(
+					FoundFrameData->SerializedBlocks, DebugBlockStorage);
+			ensure(ReadBlock && ReadBlock->GetTypeID() == FRootCameraDebugBlock::StaticTypeID());
+			RootDebugBlock = ReadBlock->CastThisChecked<FRootCameraDebugBlock>();
+		}
+	}
+}
+
+void FCameraSystemRewindDebuggerExtension::RecordingStopped(IRewindDebugger* RewindDebugger)
+{
+	UE::Trace::ToggleChannel(*FCameraSystemTrace::ChannelName, false);
+}
+
+void FCameraSystemRewindDebuggerExtension::EnsureDebugDrawDelegate(bool bIsRegistered)
+{
+	if (bIsRegistered && !DebugDrawDelegateHandle.IsValid())
+	{
+		DebugDrawDelegateHandle = UDebugDrawService::Register(
+				TEXT("GameplayDebug"), FDebugDrawDelegate::CreateRaw(this, &FCameraSystemRewindDebuggerExtension::DebugDraw));
+	}
+	else if (!bIsRegistered && DebugDrawDelegateHandle.IsValid())
+	{
+		UDebugDrawService::Unregister(DebugDrawDelegateHandle);
+		DebugDrawDelegateHandle .Reset();
+	}
+}
+
+void FCameraSystemRewindDebuggerExtension::DebugDraw(UCanvas* Canvas, APlayerController* PlayController)
+{
+	if (RootDebugBlock)
+	{
+		FCameraDebugRenderer CameraDebugRenderer(Canvas->Canvas);
+		RootDebugBlock->RootDebugDraw(CameraDebugRenderer);
+	}
+}
+
+}  // namespace UE::Cameras
+
diff --git a/Engine/Plugins/Cameras/GameplayCameras/Source/GameplayCamerasEditor/Private/Trace/CameraSystemRewindDebuggerExtension.h b/Engine/Plugins/Cameras/GameplayCameras/Source/GameplayCamerasEditor/Private/Trace/CameraSystemRewindDebuggerExtension.h
new file mode 100644
--- /dev/null
+++ b/Engine/Plugins/Cameras/GameplayCameras/Source/GameplayCamerasEditor/Private/Trace/CameraSystemRewindDebuggerExtension.h
@@ -0,0 +1,51 @@
+// Copyright Epic Games, Inc. All Rights Reserved.
+
+#pragma once
+
+#include "CoreMinimal.h"
+#include "Debug/CameraDebugBlockStorage.h"
+#include "IRewindDebugger.h"
+#include "IRewindDebuggerExtension.h"
+#include "Templates/SharedPointerFwd.h"
+
+class APlayerController;
+class UCanvas;
+
+namespace UE::Cameras
+{
+
+class FDebugCanvasDeferredRenderer;
+class FRootCameraDebugBlock;
+
+/**
+ * Rewind debugger extension for the camera system evaluation trace.
+ */
+class FCameraSystemRewindDebuggerExtension : public IRewindDebuggerExtension
+{
+public:
+
+	FCameraSystemRewindDebuggerExtension();
+	virtual ~FCameraSystemRewindDebuggerExtension();
+
+	virtual void RecordingStarted(IRewindDebugger* RewindDebugger) override;
+	virtual void Update(float DeltaTime, IRewindDebugger* RewindDebugger) override;
+	virtual void RecordingStopped(IRewindDebugger* RewindDebugger) override;
+
+private:
+
+	void EnsureDebugDrawDelegate(bool bIsRegistered);
+	void DebugDraw(UCanvas* Canvas, APlayerController* PlayController);
+
+private:
+
+	FDelegateHandle DebugDrawDelegateHandle;
+	double LastTraceTime = 0.f;
+
+	FCameraDebugBlockStorage DebugBlockStorage;
+	FRootCameraDebugBlock* RootDebugBlock = nullptr;
+
+	TSharedPtr<FDebugCanvasDeferredRenderer> DeferredRenderer;
+};
+
+}  // namespace UE::Cameras
+
diff --git a/Engine/Plugins/Cameras/GameplayCameras/Source/GameplayCamerasEditor/Private/Trace/CameraSystemRewindDebuggerTrack.cpp b/Engine/Plugins/Cameras/GameplayCameras/Source/GameplayCamerasEditor/Private/Trace/CameraSystemRewindDebuggerTrack.cpp
new file mode 100644
--- /dev/null
+++ b/Engine/Plugins/Cameras/GameplayCameras/Source/GameplayCamerasEditor/Private/Trace/CameraSystemRewindDebuggerTrack.cpp
@@ -0,0 +1,156 @@
+// Copyright Epic Games, Inc. All Rights Reserved.
+
+#include "Trace/CameraSystemRewindDebuggerTrack.h"
+
+#include "CoreTypes.h"
+#include "GameFramework/Pawn.h"
+#include "IRewindDebugger.h"
+#include "Math/Range.h"
+#include "SSimpleTimeSlider.h"
+#include "Styles/GameplayCamerasEditorStyle.h"
+#include "Widgets/SCompoundWidget.h"
+#include "Widgets/DeclarativeSyntaxSupport.h"
+
+#define LOCTEXT_NAMESPACE "CameraSystemRewindDebuggerTrack"
+
+namespace UE::Cameras
+{
+
+class SCameraSystemRewindDebuggerTrackTimeline : public SCompoundWidget
+{
+public:
+
+	SLATE_BEGIN_ARGS(SCameraSystemRewindDebuggerTrackTimeline)
+		: _ViewRange(TRange<double>(0,10))
+		, _DesiredSize(FVector2D(100.f, 20.f))
+	{}
+    	SLATE_ATTRIBUTE(TRange<double>, ViewRange);
+		SLATE_ATTRIBUTE(FVector2D, DesiredSize);
+		SLATE_ATTRIBUTE(TSharedPtr<FCameraSystemRewindDebuggerTrackTimelineData>, TimelineData);
+	SLATE_END_ARGS()
+
+	void Construct(const FArguments& InArgs)
+	{
+		ViewRange = InArgs._ViewRange;
+		DesiredSize = InArgs._DesiredSize;
+		TimelineData = InArgs._TimelineData;
+
+		BackgroundBrush = FAppStyle::GetBrush("Sequencer.SectionArea.Background");
+	}
+
+protected:
+
+	virtual FVector2D ComputeDesiredSize(float) const override
+	{
+		return DesiredSize.Get();
+	}
+
+	virtual int32 OnPaint(const FPaintArgs& Args, const FGeometry& AllottedGeometry, const FSlateRect& MyCullingRect, FSlateWindowElementList& OutDrawElements, int32 LayerId, const FWidgetStyle& InWidgetStyle, bool bParentEnabled) const override
+	{
+		const int32 NewLayer = PaintWindows(AllottedGeometry, MyCullingRect, OutDrawElements, LayerId, InWidgetStyle, bParentEnabled);
+		const int32 BaseLayer = SCompoundWidget::OnPaint(Args, AllottedGeometry, MyCullingRect, OutDrawElements, NewLayer, InWidgetStyle, ShouldBeEnabled(bParentEnabled));
+		return FMath::Max(NewLayer, BaseLayer);
+	}
+
+private:
+
+	int32 PaintWindows(const FGeometry& AllottedGeometry, const FSlateRect& MyCullingRect, FSlateWindowElementList& OutDrawElements, int32 LayerId, const FWidgetStyle& InWidgetStyle, bool bParentEnabled) const
+	{
+		TSharedPtr<FCameraSystemRewindDebuggerTrackTimelineData> ActualTimelineData = TimelineData.Get();
+		if (!ActualTimelineData)
+		{
+			return LayerId;
+		}
+
+		TRange<double> DebugTimeRange = ViewRange.Get();
+
+		SSimpleTimeSlider::FScrubRangeToScreen RangeToScreen( DebugTimeRange, AllottedGeometry.GetLocalSize() );
+		FVector2D Size = AllottedGeometry.GetLocalSize();
+
+		const TArray<FCameraSystemRewindDebuggerTrackTimelineData::DataWindow>& Windows(ActualTimelineData->Windows);
+		const int NumWindows = Windows.Num();
+		if (NumWindows > 0)
+		{
+			for (int32 Index = 0; Index < NumWindows; ++Index)
+			{
+				const FCameraSystemRewindDebuggerTrackTimelineData::DataWindow& Window(Windows[Index]);
+			
+				const float StartX = RangeToScreen.InputToLocalX(Window.TimeStart);
+				const float EndX = RangeToScreen.InputToLocalX(Window.TimeEnd);
+				const float SizeY = AllottedGeometry.Size.Y;
+
+				FSlateDrawElement::MakeBox(
+						OutDrawElements, 
+						LayerId++,
+						AllottedGeometry.ToPaintGeometry(FVector2D(EndX - StartX, SizeY - 2), FSlateLayoutTransform(FVector2D(StartX, 1))),
+						BackgroundBrush, 
+						ESlateDrawEffect::None, 
+						Window.Color);
+			}
+		}
+		
+		LayerId++;
+		return LayerId;
+	}
+	
+private:
+
+	TAttribute<TRange<double>> ViewRange;
+	TAttribute<FVector2D> DesiredSize;
+	TAttribute<TSharedPtr<FCameraSystemRewindDebuggerTrackTimelineData>> TimelineData;
+
+	const FSlateBrush* BackgroundBrush;
+};
+
+FCameraSystemRewindDebuggerTrack::FCameraSystemRewindDebuggerTrack()
+{
+	const FName& GameplayCamerasEditorStyleName = FGameplayCamerasEditorStyle::Get()->GetStyleSetName();
+	Icon = FSlateIcon(GameplayCamerasEditorStyleName,"DebugCategory.PoseStats.Icon");
+
+	TimelineData = MakeShared<FCameraSystemRewindDebuggerTrackTimelineData>();
+}
+
+FSlateIcon FCameraSystemRewindDebuggerTrack::GetIconInternal()
+{
+	return Icon;
+}
+
+FText FCameraSystemRewindDebuggerTrack::GetDisplayNameInternal() const
+{
+	return LOCTEXT("DisplayName", "Gameplay Camera System");
+}
+
+TSharedPtr<SWidget> FCameraSystemRewindDebuggerTrack::GetTimelineViewInternal()
+{
+	return SNew(SCameraSystemRewindDebuggerTrackTimeline)
+		.ViewRange_Lambda([]() { return IRewindDebugger::Instance()->GetCurrentViewRange(); })
+		.TimelineData(TimelineData);
+}
+
+FName FCameraSystemRewindDebuggerTrackCreator::GetTargetTypeNameInternal() const
+{
+	// Only show this track once at the top level.
+ 	static const FName TargetTypeName = APawn::StaticClass()->GetFName();
+	return TargetTypeName;
+}
+
+FName FCameraSystemRewindDebuggerTrackCreator::GetNameInternal() const
+{
+	static const FName TrackName("Gameplay Camera System");
+	return TrackName;
+}
+
+void FCameraSystemRewindDebuggerTrackCreator::GetTrackTypesInternal(TArray<RewindDebugger::FRewindDebuggerTrackType>& Types) const
+{
+	Types.Add({GetNameInternal(), LOCTEXT("DisplayName", "GameplayCamera System")});
+}
+
+TSharedPtr<RewindDebugger::FRewindDebuggerTrack> FCameraSystemRewindDebuggerTrackCreator::CreateTrackInternal(uint64 ObjectId) const
+{
+	return MakeShared<FCameraSystemRewindDebuggerTrack>();
+}
+
+}  // namespace UE::Cameras
+
+#undef LOCTEXT_NAMESPACE
+
diff --git a/Engine/Plugins/Cameras/GameplayCameras/Source/GameplayCamerasEditor/Private/Trace/CameraSystemRewindDebuggerTrack.h b/Engine/Plugins/Cameras/GameplayCameras/Source/GameplayCamerasEditor/Private/Trace/CameraSystemRewindDebuggerTrack.h
new file mode 100644
--- /dev/null
+++ b/Engine/Plugins/Cameras/GameplayCameras/Source/GameplayCamerasEditor/Private/Trace/CameraSystemRewindDebuggerTrack.h
@@ -0,0 +1,65 @@
+// Copyright Epic Games, Inc. All Rights Reserved.
+
+#pragma once
+
+#include "IRewindDebuggerTrackCreator.h"
+#include "Math/Color.h"
+#include "RewindDebuggerTrack.h"
+
+namespace UE::Cameras
+{
+
+/**
+ * Data for drawing the rewind debugger track that corresponds to the 
+ * camera system evaluation trace.
+ */
+struct FCameraSystemRewindDebuggerTrackTimelineData
+{
+public:
+
+	struct DataWindow
+	{
+		double TimeStart;
+		double TimeEnd;
+		FLinearColor Color;
+	};
+
+	TArray<DataWindow> Windows;
+};
+
+/**
+ * Rewing debugger track for the camera system evaluation trace.
+ */
+class FCameraSystemRewindDebuggerTrack : public RewindDebugger::FRewindDebuggerTrack
+{
+public:
+
+	FCameraSystemRewindDebuggerTrack();
+
+protected:
+
+	virtual FSlateIcon GetIconInternal() override;
+	virtual FText GetDisplayNameInternal() const override;
+	virtual TSharedPtr<SWidget> GetTimelineViewInternal() override;
+
+private:
+
+	FSlateIcon Icon;
+	TSharedPtr<FCameraSystemRewindDebuggerTrackTimelineData> TimelineData;
+};
+
+/**
+ * Factory class that creates the rewind debugger track for the camera system evaluation trace.
+ */
+class FCameraSystemRewindDebuggerTrackCreator : public RewindDebugger::IRewindDebuggerTrackCreator
+{
+protected:
+
+	virtual FName GetTargetTypeNameInternal() const override;
+	virtual FName GetNameInternal() const override;
+	virtual void GetTrackTypesInternal(TArray<RewindDebugger::FRewindDebuggerTrackType>& Types) const override;
+	virtual TSharedPtr<RewindDebugger::FRewindDebuggerTrack> CreateTrackInternal(uint64 ObjectId) const override;
+};
+
+}  // namespace UE::Cameras
+
diff --git a/Engine/Plugins/Cameras/GameplayCameras/Source/GameplayCamerasEditor/Private/Trace/CameraSystemTraceAnalyzer.cpp b/Engine/Plugins/Cameras/GameplayCameras/Source/GameplayCamerasEditor/Private/Trace/CameraSystemTraceAnalyzer.cpp
new file mode 100644
--- /dev/null
+++ b/Engine/Plugins/Cameras/GameplayCameras/Source/GameplayCamerasEditor/Private/Trace/CameraSystemTraceAnalyzer.cpp
@@ -0,0 +1,48 @@
+// Copyright Epic Games, Inc. All Rights Reserved.
+
+#include "Trace/CameraSystemTraceAnalyzer.h"
+
+#include "Debug/CameraSystemTrace.h"
+#include "Trace/CameraSystemTraceProvider.h"
+
+namespace UE::Cameras
+{
+
+FCameraSystemTraceAnalyzer::FCameraSystemTraceAnalyzer(TraceServices::IAnalysisSession& InSession, FCameraSystemTraceProvider& InProvider)
+	: Session(InSession)
+	, Provider(InProvider)
+{
+}
+
+void FCameraSystemTraceAnalyzer::OnAnalysisBegin(const FOnAnalysisContext& Context)
+{
+	auto& Builder = Context.InterfaceBuilder;
+
+	Builder.RouteEvent(0, "CameraSystem", "CameraSystemEvaluation");
+}
+
+bool FCameraSystemTraceAnalyzer::OnEvent(uint16 RouteId, EStyle Style, const FOnEventContext& Context)
+{
+	TraceServices::FAnalysisSessionEditScope _(Session);
+
+	FCameraSystemTraceFrameData FrameData;
+
+	uint64 Cycle = Context.EventData.GetValue<uint64>("Cycle");
+	double RecordingTime = Context.EventData.GetValue<double>("RecordingTime");
+	FrameData.EvaluatedLocation.X = Context.EventData.GetValue<double>("EvaluatedLocationX");
+	FrameData.EvaluatedLocation.Y = Context.EventData.GetValue<double>("EvaluatedLocationY");
+	FrameData.EvaluatedLocation.Z = Context.EventData.GetValue<double>("EvaluatedLocationZ");
+	FrameData.EvaluatedRotation.Yaw = Context.EventData.GetValue<double>("EvaluatedRotationYaw");
+	FrameData.EvaluatedRotation.Pitch = Context.EventData.GetValue<double>("EvaluatedRotationPitch");
+	FrameData.EvaluatedRotation.Roll = Context.EventData.GetValue<double>("EvaluatedRotationRoll");
+	FrameData.EvaluatedFieldOfView = Context.EventData.GetValue<float>("EvaluatedFieldOfView");
+	FrameData.SerializedBlocks = Context.EventData.GetArrayView<uint8>("SerializedBlocks");
+	
+	double EventTime = Context.EventTime.AsSeconds(Cycle);
+	Provider.AppendFrameData(EventTime, MoveTemp(FrameData));
+
+	return true;
+}
+
+}  // namespace UE::Cameras
+
diff --git a/Engine/Plugins/Cameras/GameplayCameras/Source/GameplayCamerasEditor/Private/Trace/CameraSystemTraceAnalyzer.h b/Engine/Plugins/Cameras/GameplayCameras/Source/GameplayCamerasEditor/Private/Trace/CameraSystemTraceAnalyzer.h
new file mode 100644
--- /dev/null
+++ b/Engine/Plugins/Cameras/GameplayCameras/Source/GameplayCamerasEditor/Private/Trace/CameraSystemTraceAnalyzer.h
@@ -0,0 +1,35 @@
+// Copyright Epic Games, Inc. All Rights Reserved.
+
+#pragma once
+
+#include "CoreTypes.h"
+#include "Trace/Analyzer.h"
+
+namespace TraceServices { class IAnalysisSession; }
+
+namespace UE::Cameras
+{
+
+class FCameraSystemTraceProvider;
+
+/**
+ * Trace analyzer for the camera system evaluation.
+ */
+class FCameraSystemTraceAnalyzer : public UE::Trace::IAnalyzer
+{
+public:
+
+	FCameraSystemTraceAnalyzer(TraceServices::IAnalysisSession& InSession, FCameraSystemTraceProvider& InProvider);
+
+	// IAnalyzer interface
+	virtual void OnAnalysisBegin(const FOnAnalysisContext& Context) override;
+	virtual bool OnEvent(uint16 RouteId, EStyle Style, const FOnEventContext& Context) override;
+
+private:
+
+	TraceServices::IAnalysisSession& Session;
+	FCameraSystemTraceProvider& Provider;
+};
+
+}  // namespace UE::Cameras
+
diff --git a/Engine/Plugins/Cameras/GameplayCameras/Source/GameplayCamerasEditor/Private/Trace/CameraSystemTraceModule.cpp b/Engine/Plugins/Cameras/GameplayCameras/Source/GameplayCamerasEditor/Private/Trace/CameraSystemTraceModule.cpp
new file mode 100644
--- /dev/null
+++ b/Engine/Plugins/Cameras/GameplayCameras/Source/GameplayCamerasEditor/Private/Trace/CameraSystemTraceModule.cpp
@@ -0,0 +1,38 @@
+// Copyright Epic Games, Inc. All Rights Reserved.
+
+#include "Trace/CameraSystemTraceModule.h"
+
+#include "Debug/CameraSystemTrace.h"
+#include "Trace/CameraSystemTraceModule.h"
+#include "Trace/CameraSystemTraceProvider.h"
+#include "Trace/CameraSystemTraceAnalyzer.h"
+
+namespace UE::Cameras
+{
+
+FName FCameraSystemTraceModule::ModuleName("CameraSystem");
+
+void FCameraSystemTraceModule::GetModuleInfo(TraceServices::FModuleInfo& OutModuleInfo)
+{
+	OutModuleInfo.Name = ModuleName;
+	OutModuleInfo.DisplayName = TEXT("Camera System");
+}
+
+void FCameraSystemTraceModule::OnAnalysisBegin(TraceServices::IAnalysisSession& InSession)
+{
+	TSharedPtr<FCameraSystemTraceProvider> CameraSystemProvider = MakeShared<FCameraSystemTraceProvider>(InSession);
+	InSession.AddProvider(FCameraSystemTraceProvider::ProviderName, CameraSystemProvider);
+	InSession.AddAnalyzer(new FCameraSystemTraceAnalyzer(InSession, *CameraSystemProvider));
+}
+
+void FCameraSystemTraceModule::GetLoggers(TArray<const TCHAR *>& OutLoggers)
+{
+	OutLoggers.Add(*FCameraSystemTrace::LoggerName);
+}
+
+void FCameraSystemTraceModule::GenerateReports(const TraceServices::IAnalysisSession& Session, const TCHAR* CmdLine, const TCHAR* OutputDirectory)
+{
+}
+
+}  // namespace UE::Cameras
+
diff --git a/Engine/Plugins/Cameras/GameplayCameras/Source/GameplayCamerasEditor/Private/Trace/CameraSystemTraceModule.h b/Engine/Plugins/Cameras/GameplayCameras/Source/GameplayCamerasEditor/Private/Trace/CameraSystemTraceModule.h
new file mode 100644
--- /dev/null
+++ b/Engine/Plugins/Cameras/GameplayCameras/Source/GameplayCamerasEditor/Private/Trace/CameraSystemTraceModule.h
@@ -0,0 +1,34 @@
+// Copyright Epic Games, Inc. All Rights Reserved.
+
+#pragma once
+
+#include "CoreMinimal.h"
+#include "TraceServices/ModuleService.h"
+
+namespace Insights { class ITimingViewSession; }
+
+namespace UE::Cameras
+{
+
+/**
+ * Trace module for camera system evaluation.
+ */
+class FCameraSystemTraceModule : public TraceServices::IModule
+{
+public:
+
+	// TraceServices::IModule interface
+	virtual void GetModuleInfo(TraceServices::FModuleInfo& OutModuleInfo) override;
+	virtual void OnAnalysisBegin(TraceServices::IAnalysisSession& Session) override;
+	virtual void GetLoggers(TArray<const TCHAR *>& OutLoggers) override;
+	virtual void GenerateReports(const TraceServices::IAnalysisSession& Session, const TCHAR* CmdLine, const TCHAR* OutputDirectory) override;
+	virtual const TCHAR* GetCommandLineArgument() override { return TEXT("gameplaycameras"); }
+
+private:
+
+	static FName ModuleName;
+};
+
+
+}  // namespace UE::Cameras
+
diff --git a/Engine/Plugins/Cameras/GameplayCameras/Source/GameplayCamerasEditor/Private/Trace/CameraSystemTraceProvider.cpp b/Engine/Plugins/Cameras/GameplayCameras/Source/GameplayCamerasEditor/Private/Trace/CameraSystemTraceProvider.cpp
new file mode 100644
--- /dev/null
+++ b/Engine/Plugins/Cameras/GameplayCameras/Source/GameplayCamerasEditor/Private/Trace/CameraSystemTraceProvider.cpp
@@ -0,0 +1,24 @@
+// Copyright Epic Games, Inc. All Rights Reserved.
+
+#include "Trace/CameraSystemTraceProvider.h"
+
+namespace UE::Cameras
+{
+
+FName FCameraSystemTraceProvider::ProviderName("CameraSystemTraceProvider");
+
+FCameraSystemTraceProvider::FCameraSystemTraceProvider(TraceServices::IAnalysisSession& InSession)
+	: Session(InSession)
+{
+	Timeline = MakeShared<FFrameDataPointTimeline>(Session.GetLinearAllocator());
+}
+
+void FCameraSystemTraceProvider::AppendFrameData(double InRecordingTime, FCameraSystemTraceFrameData&& InFrameData)
+{
+	Session.WriteAccessCheck();
+
+	Timeline->AppendEvent(InRecordingTime, MoveTemp(InFrameData));
+}
+
+}  // namespace UE::Cameras
+
diff --git a/Engine/Plugins/Cameras/GameplayCameras/Source/GameplayCamerasEditor/Private/Trace/CameraSystemTraceProvider.h b/Engine/Plugins/Cameras/GameplayCameras/Source/GameplayCamerasEditor/Private/Trace/CameraSystemTraceProvider.h
new file mode 100644
--- /dev/null
+++ b/Engine/Plugins/Cameras/GameplayCameras/Source/GameplayCamerasEditor/Private/Trace/CameraSystemTraceProvider.h
@@ -0,0 +1,51 @@
+// Copyright Epic Games, Inc. All Rights Reserved.
+
+#pragma once
+
+#include "CoreTypes.h"
+#include "Math/Rotator.h"
+#include "Math/Vector.h"
+#include "Model/PointTimeline.h"
+#include "TraceServices/Model/AnalysisSession.h"
+
+namespace TraceServices { class IAnalysisSession; }
+
+namespace UE::Cameras
+{
+
+/** Data for one frame of tracing. */
+struct FCameraSystemTraceFrameData
+{
+	FVector3d EvaluatedLocation;
+	FRotator3d EvaluatedRotation;
+	float EvaluatedFieldOfView;
+	TArray<uint8> SerializedBlocks;
+};
+
+using FCameraSystemTraceTimeline = TraceServices::ITimeline<FCameraSystemTraceFrameData>;
+
+/**
+ * Trace provider for the camera system evaluation.
+ */
+class FCameraSystemTraceProvider : public TraceServices::IProvider
+{
+public:
+
+	static FName ProviderName;
+
+	FCameraSystemTraceProvider(TraceServices::IAnalysisSession& InSession);
+
+	void AppendFrameData(double InRecordingTime, FCameraSystemTraceFrameData&& InFrameData);
+
+	const FCameraSystemTraceTimeline* GetTimeline() const { return Timeline.Get(); }	
+
+private:
+
+	TraceServices::IAnalysisSession& Session;
+
+	using FFrameDataPointTimeline = TraceServices::TPointTimeline<FCameraSystemTraceFrameData>;
+	TSharedPtr<FFrameDataPointTimeline> Timeline;
+};
+
+}  // namespace UE::Cameras
+
