Cameras: fix variable table debugging info

diff --git a/Engine/Plugins/Cameras/GameplayCameras/Source/GameplayCameras/Private/Core/CameraVariableTable.cpp b/Engine/Plugins/Cameras/GameplayCameras/Source/GameplayCameras/Private/Core/CameraVariableTable.cpp
--- a/Engine/Plugins/Cameras/GameplayCameras/Source/GameplayCameras/Private/Core/CameraVariableTable.cpp
+++ b/Engine/Plugins/Cameras/GameplayCameras/Source/GameplayCameras/Private/Core/CameraVariableTable.cpp
@@ -142,6 +142,9 @@ void FCameraVariableTable::AddVariable(const FCameraVariableDefinition& Variable
 	NewEntry.Type = VariableDefinition.VariableType;
 	NewEntry.Offset = VariablePtr - Memory;
 	NewEntry.Flags = VariableDefinition.bIsPrivate ? EEntryFlags::Private : EEntryFlags::None;
+#if WITH_EDITORONLY_DATA
+	NewEntry.DebugName = VariableDefinition.VariableName;
+#endif
 
 	Entries.Add(VariableDefinition.VariableId, NewEntry);
 }
