Cameras: fix cooked editor build

diff --git a/Engine/Plugins/Cameras/GameplayCameras/Source/GameplayCamerasEditor/Private/GameplayCamerasEditorModule.cpp b/Engine/Plugins/Cameras/GameplayCameras/Source/GameplayCamerasEditor/Private/GameplayCamerasEditorModule.cpp
--- a/Engine/Plugins/Cameras/GameplayCameras/Source/GameplayCamerasEditor/Private/GameplayCamerasEditorModule.cpp
+++ b/Engine/Plugins/Cameras/GameplayCameras/Source/GameplayCamerasEditor/Private/GameplayCamerasEditorModule.cpp
@@ -12,6 +12,7 @@
 #include "Debugger/SCameraNodeTreeDebugPanel.h"
 #include "Debugger/SGameplayCamerasDebugger.h"
 #include "Features/IModularFeatures.h"
+#include "GameplayCameras.h"
 #include "GameplayCamerasLiveEditManager.h"
 #include "IGameplayCamerasEditorModule.h"
 #include "IGameplayCamerasModule.h"
@@ -255,6 +256,7 @@ private:
 
 	void RegisterRewindDebuggerFeatures()
 	{
+#if UE_GAMEPLAY_CAMERAS_TRACE
 		using namespace UE::Cameras;
 
 		TraceModule = MakeShared<UE::Cameras::FCameraSystemTraceModule>();
@@ -265,14 +267,17 @@ private:
 		ModularFeatures.RegisterModularFeature(IRewindDebuggerExtension::ModularFeatureName, RewindDebuggerExtension.Get());
 		ModularFeatures.RegisterModularFeature(RewindDebugger::IRewindDebuggerTrackCreator::ModularFeatureName, RewindDebuggerTrackCreator.Get());
 		ModularFeatures.RegisterModularFeature(TraceServices::ModuleFeatureName, TraceModule.Get());
+#endif  // UE_GAMEPLAY_CAMERAS_TRACE
 	}
 
 	void UnregisterRewindDebuggerFeatures()
 	{
+#if UE_GAMEPLAY_CAMERAS_TRACE
 		IModularFeatures& ModularFeatures = IModularFeatures::Get();
 		ModularFeatures.UnregisterModularFeature(IRewindDebuggerExtension::ModularFeatureName, RewindDebuggerExtension.Get());
 		ModularFeatures.UnregisterModularFeature(RewindDebugger::IRewindDebuggerTrackCreator::ModularFeatureName, RewindDebuggerTrackCreator.Get());
 		ModularFeatures.UnregisterModularFeature(TraceServices::ModuleFeatureName, TraceModule.Get());
+#endif  // UE_GAMEPLAY_CAMERAS_TRACE
 	}
 
 	void InitializeLiveEditManager()
@@ -308,9 +313,11 @@ private:
 	TMap<FString, UE::Cameras::FCameraDebugCategoryInfo> DebugCategoryInfos;
 	TMap<FString, FOnCreateDebugCategoryPanel> DebugCategoryPanelCreators;
 
+#if UE_GAMEPLAY_CAMERAS_TRACE
 	TSharedPtr<UE::Cameras::FCameraSystemTraceModule> TraceModule;
 	TSharedPtr<UE::Cameras::FCameraSystemRewindDebuggerExtension> RewindDebuggerExtension;
 	TSharedPtr<UE::Cameras::FCameraSystemRewindDebuggerTrackCreator> RewindDebuggerTrackCreator;
+#endif  // UE_GAMEPLAY_CAMERAS_TRACE
 };
 
 IMPLEMENT_MODULE(FGameplayCamerasEditorModule, GameplayCamerasEditor);
diff --git a/Engine/Plugins/Cameras/GameplayCameras/Source/GameplayCamerasEditor/Private/Trace/CameraSystemRewindDebuggerExtension.cpp b/Engine/Plugins/Cameras/GameplayCameras/Source/GameplayCamerasEditor/Private/Trace/CameraSystemRewindDebuggerExtension.cpp
--- a/Engine/Plugins/Cameras/GameplayCameras/Source/GameplayCamerasEditor/Private/Trace/CameraSystemRewindDebuggerExtension.cpp
+++ b/Engine/Plugins/Cameras/GameplayCameras/Source/GameplayCamerasEditor/Private/Trace/CameraSystemRewindDebuggerExtension.cpp
@@ -13,6 +13,8 @@
 #include "GameFramework/GameplayCameraSystemActor.h"
 #include "Trace/CameraSystemTraceProvider.h"
 
+#if UE_GAMEPLAY_CAMERAS_TRACE
+
 namespace UE::Cameras
 {
 
@@ -106,3 +108,5 @@ void FCameraSystemRewindDebuggerExtension::DebugDraw(UCanvas* Canvas, APlayerCon
 
 }  // namespace UE::Cameras
 
+#endif  // UE_GAMEPLAY_CAMERAS_TRACE
+
diff --git a/Engine/Plugins/Cameras/GameplayCameras/Source/GameplayCamerasEditor/Private/Trace/CameraSystemRewindDebuggerExtension.h b/Engine/Plugins/Cameras/GameplayCameras/Source/GameplayCamerasEditor/Private/Trace/CameraSystemRewindDebuggerExtension.h
--- a/Engine/Plugins/Cameras/GameplayCameras/Source/GameplayCamerasEditor/Private/Trace/CameraSystemRewindDebuggerExtension.h
+++ b/Engine/Plugins/Cameras/GameplayCameras/Source/GameplayCamerasEditor/Private/Trace/CameraSystemRewindDebuggerExtension.h
@@ -4,10 +4,13 @@
 
 #include "CoreMinimal.h"
 #include "Debug/CameraDebugBlockStorage.h"
+#include "GameplayCameras.h"
 #include "IRewindDebugger.h"
 #include "IRewindDebuggerExtension.h"
 #include "Templates/SharedPointerFwd.h"
 
+#if UE_GAMEPLAY_CAMERAS_TRACE
+
 class APlayerController;
 class UCanvas;
 
@@ -49,3 +52,5 @@ private:
 
 }  // namespace UE::Cameras
 
+#endif  // UE_GAMEPLAY_CAMERAS_TRACE
+
diff --git a/Engine/Plugins/Cameras/GameplayCameras/Source/GameplayCamerasEditor/Private/Trace/CameraSystemRewindDebuggerTrack.cpp b/Engine/Plugins/Cameras/GameplayCameras/Source/GameplayCamerasEditor/Private/Trace/CameraSystemRewindDebuggerTrack.cpp
--- a/Engine/Plugins/Cameras/GameplayCameras/Source/GameplayCamerasEditor/Private/Trace/CameraSystemRewindDebuggerTrack.cpp
+++ b/Engine/Plugins/Cameras/GameplayCameras/Source/GameplayCamerasEditor/Private/Trace/CameraSystemRewindDebuggerTrack.cpp
@@ -13,6 +13,8 @@
 
 #define LOCTEXT_NAMESPACE "CameraSystemRewindDebuggerTrack"
 
+#if UE_GAMEPLAY_CAMERAS_TRACE
+
 namespace UE::Cameras
 {
 
@@ -152,5 +154,7 @@ TSharedPtr<RewindDebugger::FRewindDebuggerTrack> FCameraSystemRewindDebuggerTrac
 
 }  // namespace UE::Cameras
 
+#endif  // UE_GAMEPLAY_CAMERAS_TRACE
+
 #undef LOCTEXT_NAMESPACE
 
diff --git a/Engine/Plugins/Cameras/GameplayCameras/Source/GameplayCamerasEditor/Private/Trace/CameraSystemRewindDebuggerTrack.h b/Engine/Plugins/Cameras/GameplayCameras/Source/GameplayCamerasEditor/Private/Trace/CameraSystemRewindDebuggerTrack.h
--- a/Engine/Plugins/Cameras/GameplayCameras/Source/GameplayCamerasEditor/Private/Trace/CameraSystemRewindDebuggerTrack.h
+++ b/Engine/Plugins/Cameras/GameplayCameras/Source/GameplayCamerasEditor/Private/Trace/CameraSystemRewindDebuggerTrack.h
@@ -2,10 +2,13 @@
 
 #pragma once
 
+#include "GameplayCameras.h"
 #include "IRewindDebuggerTrackCreator.h"
 #include "Math/Color.h"
 #include "RewindDebuggerTrack.h"
 
+#if UE_GAMEPLAY_CAMERAS_TRACE
+
 namespace UE::Cameras
 {
 
@@ -63,3 +66,5 @@ protected:
 
 }  // namespace UE::Cameras
 
+#endif  // UE_GAMEPLAY_CAMERAS_TRACE
+
diff --git a/Engine/Plugins/Cameras/GameplayCameras/Source/GameplayCamerasEditor/Private/Trace/CameraSystemTraceAnalyzer.cpp b/Engine/Plugins/Cameras/GameplayCameras/Source/GameplayCamerasEditor/Private/Trace/CameraSystemTraceAnalyzer.cpp
--- a/Engine/Plugins/Cameras/GameplayCameras/Source/GameplayCamerasEditor/Private/Trace/CameraSystemTraceAnalyzer.cpp
+++ b/Engine/Plugins/Cameras/GameplayCameras/Source/GameplayCamerasEditor/Private/Trace/CameraSystemTraceAnalyzer.cpp
@@ -5,6 +5,8 @@
 #include "Debug/CameraSystemTrace.h"
 #include "Trace/CameraSystemTraceProvider.h"
 
+#if UE_GAMEPLAY_CAMERAS_TRACE
+
 namespace UE::Cameras
 {
 
@@ -46,3 +48,5 @@ bool FCameraSystemTraceAnalyzer::OnEvent(uint16 RouteId, EStyle Style, const FOn
 
 }  // namespace UE::Cameras
 
+#endif  // UE_GAMEPLAY_CAMERAS_TRACE
+
diff --git a/Engine/Plugins/Cameras/GameplayCameras/Source/GameplayCamerasEditor/Private/Trace/CameraSystemTraceAnalyzer.h b/Engine/Plugins/Cameras/GameplayCameras/Source/GameplayCamerasEditor/Private/Trace/CameraSystemTraceAnalyzer.h
--- a/Engine/Plugins/Cameras/GameplayCameras/Source/GameplayCamerasEditor/Private/Trace/CameraSystemTraceAnalyzer.h
+++ b/Engine/Plugins/Cameras/GameplayCameras/Source/GameplayCamerasEditor/Private/Trace/CameraSystemTraceAnalyzer.h
@@ -3,8 +3,11 @@
 #pragma once
 
 #include "CoreTypes.h"
+#include "GameplayCameras.h"
 #include "Trace/Analyzer.h"
 
+#if UE_GAMEPLAY_CAMERAS_TRACE
+
 namespace TraceServices { class IAnalysisSession; }
 
 namespace UE::Cameras
@@ -33,3 +36,5 @@ private:
 
 }  // namespace UE::Cameras
 
+#endif  // UE_GAMEPLAY_CAMERAS_TRACE
+
diff --git a/Engine/Plugins/Cameras/GameplayCameras/Source/GameplayCamerasEditor/Private/Trace/CameraSystemTraceModule.cpp b/Engine/Plugins/Cameras/GameplayCameras/Source/GameplayCamerasEditor/Private/Trace/CameraSystemTraceModule.cpp
--- a/Engine/Plugins/Cameras/GameplayCameras/Source/GameplayCamerasEditor/Private/Trace/CameraSystemTraceModule.cpp
+++ b/Engine/Plugins/Cameras/GameplayCameras/Source/GameplayCamerasEditor/Private/Trace/CameraSystemTraceModule.cpp
@@ -7,6 +7,8 @@
 #include "Trace/CameraSystemTraceProvider.h"
 #include "Trace/CameraSystemTraceAnalyzer.h"
 
+#if UE_GAMEPLAY_CAMERAS_TRACE
+
 namespace UE::Cameras
 {
 
@@ -36,3 +38,5 @@ void FCameraSystemTraceModule::GenerateReports(const TraceServices::IAnalysisSes
 
 }  // namespace UE::Cameras
 
+#endif  // UE_GAMEPLAY_CAMERAS_TRACE
+
diff --git a/Engine/Plugins/Cameras/GameplayCameras/Source/GameplayCamerasEditor/Private/Trace/CameraSystemTraceModule.h b/Engine/Plugins/Cameras/GameplayCameras/Source/GameplayCamerasEditor/Private/Trace/CameraSystemTraceModule.h
--- a/Engine/Plugins/Cameras/GameplayCameras/Source/GameplayCamerasEditor/Private/Trace/CameraSystemTraceModule.h
+++ b/Engine/Plugins/Cameras/GameplayCameras/Source/GameplayCamerasEditor/Private/Trace/CameraSystemTraceModule.h
@@ -3,8 +3,11 @@
 #pragma once
 
 #include "CoreMinimal.h"
+#include "GameplayCameras.h"
 #include "TraceServices/ModuleService.h"
 
+#if UE_GAMEPLAY_CAMERAS_TRACE
+
 namespace Insights { class ITimingViewSession; }
 
 namespace UE::Cameras
@@ -29,6 +32,7 @@ private:
 	static FName ModuleName;
 };
 
-
 }  // namespace UE::Cameras
 
+#endif  // UE_GAMEPLAY_CAMERAS_TRACE
+
diff --git a/Engine/Plugins/Cameras/GameplayCameras/Source/GameplayCamerasEditor/Private/Trace/CameraSystemTraceProvider.cpp b/Engine/Plugins/Cameras/GameplayCameras/Source/GameplayCamerasEditor/Private/Trace/CameraSystemTraceProvider.cpp
--- a/Engine/Plugins/Cameras/GameplayCameras/Source/GameplayCamerasEditor/Private/Trace/CameraSystemTraceProvider.cpp
+++ b/Engine/Plugins/Cameras/GameplayCameras/Source/GameplayCamerasEditor/Private/Trace/CameraSystemTraceProvider.cpp
@@ -2,6 +2,8 @@
 
 #include "Trace/CameraSystemTraceProvider.h"
 
+#if UE_GAMEPLAY_CAMERAS_TRACE
+
 namespace UE::Cameras
 {
 
@@ -22,3 +24,5 @@ void FCameraSystemTraceProvider::AppendFrameData(double InRecordingTime, FCamera
 
 }  // namespace UE::Cameras
 
+#endif  // UE_GAMEPLAY_CAMERAS_TRACE
+
diff --git a/Engine/Plugins/Cameras/GameplayCameras/Source/GameplayCamerasEditor/Private/Trace/CameraSystemTraceProvider.h b/Engine/Plugins/Cameras/GameplayCameras/Source/GameplayCamerasEditor/Private/Trace/CameraSystemTraceProvider.h
--- a/Engine/Plugins/Cameras/GameplayCameras/Source/GameplayCamerasEditor/Private/Trace/CameraSystemTraceProvider.h
+++ b/Engine/Plugins/Cameras/GameplayCameras/Source/GameplayCamerasEditor/Private/Trace/CameraSystemTraceProvider.h
@@ -3,11 +3,14 @@
 #pragma once
 
 #include "CoreTypes.h"
+#include "GameplayCameras.h"
 #include "Math/Rotator.h"
 #include "Math/Vector.h"
 #include "Model/PointTimeline.h"
 #include "TraceServices/Model/AnalysisSession.h"
 
+#if UE_GAMEPLAY_CAMERAS_TRACE
+
 namespace TraceServices { class IAnalysisSession; }
 
 namespace UE::Cameras
@@ -49,3 +52,5 @@ private:
 
 }  // namespace UE::Cameras
 
+#endif  // UE_GAMEPLAY_CAMERAS_TRACE
+
