Cameras: fix non-unity build

diff --git a/Engine/Plugins/Cameras/GameplayCameras/Source/GameplayCameras/Public/Core/CameraObjectStorage.h b/Engine/Plugins/Cameras/GameplayCameras/Source/GameplayCameras/Public/Core/CameraObjectStorage.h
--- a/Engine/Plugins/Cameras/GameplayCameras/Source/GameplayCameras/Public/Core/CameraObjectStorage.h
+++ b/Engine/Plugins/Cameras/GameplayCameras/Source/GameplayCameras/Public/Core/CameraObjectStorage.h
@@ -5,6 +5,7 @@
 #include "CoreTypes.h"
 #include "HAL/UnrealMemory.h"
 #include "Math/UnrealMathUtility.h"
+#include "Templates/AlignmentTemplates.h"
 #include "Templates/EnableIf.h"
 #include "Templates/PointerIsConvertibleFromTo.h"
 
