Cameras: fix debug rendering bug caused by uninitialized field

diff --git a/Engine/Plugins/Cameras/GameplayCameras/Source/GameplayCameras/Public/Debug/CameraDebugBlockBuilder.h b/Engine/Plugins/Cameras/GameplayCameras/Source/GameplayCameras/Public/Debug/CameraDebugBlockBuilder.h
--- a/Engine/Plugins/Cameras/GameplayCameras/Source/GameplayCameras/Public/Debug/CameraDebugBlockBuilder.h
+++ b/Engine/Plugins/Cameras/GameplayCameras/Source/GameplayCameras/Public/Debug/CameraDebugBlockBuilder.h
@@ -117,7 +117,7 @@ private:
 	FRootCameraDebugBlock& RootBlock;
 	TArray<FCameraDebugBlock*> CurrentHierarchy;
 
-	ECameraDebugBlockBuildVisitFlags VisitFlags;
+	ECameraDebugBlockBuildVisitFlags VisitFlags = ECameraDebugBlockBuildVisitFlags::None;
 	int32 HierarchyOverrideStart = INDEX_NONE;
 };
 
