Cameras: fix LOCTEXT issues

diff --git a/Engine/Plugins/Cameras/GameplayCameras/Source/GameplayCamerasEditor/Private/Commands/CameraAssetEditorCommands.cpp b/Engine/Plugins/Cameras/GameplayCameras/Source/GameplayCamerasEditor/Private/Commands/CameraAssetEditorCommands.cpp
--- a/Engine/Plugins/Cameras/GameplayCameras/Source/GameplayCamerasEditor/Private/Commands/CameraAssetEditorCommands.cpp
+++ b/Engine/Plugins/Cameras/GameplayCameras/Source/GameplayCamerasEditor/Private/Commands/CameraAssetEditorCommands.cpp
@@ -7,7 +7,7 @@
 #include "Framework/Commands/UICommandInfo.h"
 #include "Styling/AppStyle.h"
 
-#define LOCTEXT_NAMESPACE "SequencerCommands"
+#define LOCTEXT_NAMESPACE "CameraAssetEditorCommands"
 
 namespace UE::Cameras
 {
