Cameras: fix static analysis issues

diff --git a/Engine/Plugins/Cameras/GameplayCameras/Source/GameplayCameras/Private/Debug/CameraSystemTrace.cpp b/Engine/Plugins/Cameras/GameplayCameras/Source/GameplayCameras/Private/Debug/CameraSystemTrace.cpp
--- a/Engine/Plugins/Cameras/GameplayCameras/Source/GameplayCameras/Private/Debug/CameraSystemTrace.cpp
+++ b/Engine/Plugins/Cameras/GameplayCameras/Source/GameplayCameras/Private/Debug/CameraSystemTrace.cpp
@@ -202,11 +202,12 @@ private:
 		check(BlockTypeName != NAME_None);
 
 		FCameraObjectTypeID BlockTypeID = TypeRegistry.FindTypeByName(BlockTypeName);
-		ensure(BlockTypeID.IsValid());
+		check(BlockTypeID.IsValid());
 		const FCameraObjectTypeInfo* BlockTypeInfo = TypeRegistry.GetTypeInfo(BlockTypeID);
-		ensure(BlockTypeInfo);
+		check(BlockTypeInfo);
 
 		void* NewBlockPtr = Storage.BuildDebugBlockUninitialized(BlockTypeInfo->Sizeof, BlockTypeInfo->Alignof);
+		check(NewBlockPtr);
 		BlockTypeInfo->Constructor(NewBlockPtr);
 
 		// This isn't quite correct for complicated inheritance configurations, but we don't
