Cameras: correctly build the menubar and toolbar for the debugger tool.

The previous setup was incorrect and resulted in the menubar buttons only working the first time the dialog was opened.

diff --git a/Engine/Plugins/Cameras/GameplayCameras/Source/GameplayCamerasEditor/Private/Debugger/SGameplayCamerasDebugger.cpp b/Engine/Plugins/Cameras/GameplayCameras/Source/GameplayCamerasEditor/Private/Debugger/SGameplayCamerasDebugger.cpp
--- a/Engine/Plugins/Cameras/GameplayCameras/Source/GameplayCamerasEditor/Private/Debugger/SGameplayCamerasDebugger.cpp
+++ b/Engine/Plugins/Cameras/GameplayCameras/Source/GameplayCamerasEditor/Private/Debugger/SGameplayCamerasDebugger.cpp
@@ -15,9 +15,11 @@
 #include "Styles/GameplayCamerasEditorStyle.h"
 #include "Styling/SlateTypes.h"
 #include "ToolMenus.h"
+#include "ToolMenuDelegates.h"
 #include "Widgets/Docking/SDockTab.h"
 #include "Widgets/Input/SButton.h"
 #include "Widgets/Layout/SBox.h"
+#include "Widgets/Text/STextBlock.h"
 #include "Widgets/SNullWidget.h"
 #include "WorkspaceMenuStructure.h"
 #include "WorkspaceMenuStructureModule.h"
@@ -72,7 +74,6 @@ SGameplayCamerasDebugger::~SGameplayCamerasDebugger()
 
 void SGameplayCamerasDebugger::Construct(const FArguments& InArgs)
 {
-	IGameplayCamerasEditorModule& GameplayCamerasEditorModule = FModuleManager::GetModuleChecked<IGameplayCamerasEditorModule>(TEXT("GameplayCamerasEditor"));
 	TSharedRef<FGameplayCamerasEditorStyle> GameplayCamerasEditorStyle = FGameplayCamerasEditorStyle::Get();
 	GameplayCamerasEditorStyleName = GameplayCamerasEditorStyle->GetStyleSetName();
 
@@ -85,7 +86,47 @@ void SGameplayCamerasDebugger::Construct(const FArguments& InArgs)
 			FCanExecuteAction(),
 			FIsActionChecked::CreateLambda([]() { return GGameplayCamerasDebugEnable; }));
 
-	// Main menu bar.
+	// Build all UI elements.
+	TSharedRef<SWidget> MenubarContents = ConstructMenubar();
+	TSharedRef<SWidget> ToolbarContents = ConstructToolbar(CommandList);
+	ConstructDebugPanels();
+
+	// Main layout.
+	ChildSlot
+	[
+		SNew(SVerticalBox)
+		+ SVerticalBox::Slot()
+			.AutoHeight()
+			[
+				MenubarContents
+			]
+		+ SVerticalBox::Slot()
+			.AutoHeight()
+			[
+				ToolbarContents
+			]
+		+ SVerticalBox::Slot()
+		.Padding(2.0)
+			[
+				SAssignNew(PanelHost, SBox)
+					.Padding(8.f)
+					[
+						EmptyPanel.ToSharedRef()
+					]
+			]
+	];
+}
+
+SGameplayCamerasDebugger* SGameplayCamerasDebugger::FromContext(UToolMenu* InMenu)
+{
+	UGameplayCamerasDebuggerMenuContext* Context = InMenu->FindContext<UGameplayCamerasDebuggerMenuContext>();
+	ensure(Context);
+	TSharedPtr<SGameplayCamerasDebugger> This = Context->CamerasDebugger.Pin();
+	return This.Get();
+}
+
+TSharedRef<SWidget> SGameplayCamerasDebugger::ConstructMenubar()
+{
 	UToolMenus* ToolMenus = UToolMenus::Get();
 	if (!ToolMenus->IsMenuRegistered(SGameplayCamerasDebugger::MenubarName))
 	{
@@ -94,10 +135,76 @@ void SGameplayCamerasDebugger::Construct(const FArguments& InArgs)
 		UToolMenu* Menubar = ToolMenus->RegisterMenu(
 				SGameplayCamerasDebugger::MenubarName, NAME_None, EMultiBoxType::MenuBar);
 	}
+
 	FToolMenuContext MenubarContext;
-	TSharedRef<SWidget> MenubarContents = ToolMenus->GenerateWidget(
-			SGameplayCamerasDebugger::MenubarName, MenubarContext);
+	return ToolMenus->GenerateWidget(SGameplayCamerasDebugger::MenubarName, MenubarContext);
+}
+
+TSharedRef<SWidget> SGameplayCamerasDebugger::ConstructToolbar(TSharedRef<FUICommandList> InCommandList)
+{
+	UToolMenus* ToolMenus = UToolMenus::Get();
+	if (!ToolMenus->IsMenuRegistered(SGameplayCamerasDebugger::ToolbarName))
+	{
+		FToolMenuOwnerScoped ToolMenuOwnerScope(this);
+
+		UToolMenu* Toolbar = ToolMenus->RegisterMenu(
+				SGameplayCamerasDebugger::ToolbarName, NAME_None, EMultiBoxType::SlimHorizontalToolBar);
+
+		Toolbar->AddDynamicSection(TEXT("Main"), FNewToolMenuDelegate::CreateLambda(
+			[](UToolMenu* InMenu)
+			{
+				const FGameplayCamerasDebuggerCommands& Commands = FGameplayCamerasDebuggerCommands::Get();
+				SGameplayCamerasDebugger* This = SGameplayCamerasDebugger::FromContext(InMenu);
+
+				FToolMenuSection& MainSection = InMenu->AddSection(TEXT("Main"));
+
+				FToolMenuEntry ToggleDebugInfo = FToolMenuEntry::InitToolBarButton(
+						Commands.EnableDebugInfo,
+						TAttribute<FText>::CreateSP(This, &SGameplayCamerasDebugger::GetToggleDebugDrawText),
+						TAttribute<FText>(),
+						TAttribute<FSlateIcon>::CreateSP(This, &SGameplayCamerasDebugger::GetToggleDebugDrawIcon));
+				MainSection.AddEntry(ToggleDebugInfo);
+			}));
+	
+		Toolbar->AddDynamicSection(TEXT("DebugCategories"), FNewToolMenuDelegate::CreateLambda(
+			[](UToolMenu* InMenu)
+			{
+				IGameplayCamerasEditorModule& ThisModule = FModuleManager::GetModuleChecked<IGameplayCamerasEditorModule>(
+						TEXT("GameplayCamerasEditor"));
+				TArray<FCameraDebugCategoryInfo> RegisteredDebugCategories;
+				ThisModule.GetRegisteredDebugCategories(RegisteredDebugCategories);
+
+				SGameplayCamerasDebugger* This = SGameplayCamerasDebugger::FromContext(InMenu);
 
+				FToolMenuSection& DebugCategoriesSection = InMenu->AddSection(TEXT("DebugCategories"));
+
+				for (const FCameraDebugCategoryInfo& DebugCategory : RegisteredDebugCategories)
+				{
+					FToolMenuEntry ToggleDebugCategory = FToolMenuEntry::InitToolBarButton(
+						FName(DebugCategory.Name),
+						FUIAction(
+							FExecuteAction::CreateSP(This, &SGameplayCamerasDebugger::SetActiveDebugCategoryPanel, DebugCategory.Name),
+							FCanExecuteAction(),
+							FIsActionChecked::CreateStatic(&SGameplayCamerasDebugger::IsDebugCategoryActive, DebugCategory.Name)),
+						DebugCategory.DisplayText,
+						DebugCategory.ToolTipText,
+						DebugCategory.IconImage,
+						EUserInterfaceActionType::ToggleButton);
+					DebugCategoriesSection.AddEntry(ToggleDebugCategory);
+				}
+			}));
+	}
+
+	UGameplayCamerasDebuggerMenuContext* ThisContextWrapper = NewObject<UGameplayCamerasDebuggerMenuContext>();
+	ThisContextWrapper->CamerasDebugger = SharedThis(this);
+	FToolMenuContext ToolbarContext(InCommandList, TSharedPtr<FExtender>());
+	ToolbarContext.AddObject(ThisContextWrapper);
+
+	return ToolMenus->GenerateWidget(SGameplayCamerasDebugger::ToolbarName, ToolbarContext);
+}
+
+void SGameplayCamerasDebugger::ConstructDebugPanels()
+{
 	// Empty panel.
 	EmptyPanel = SNew(SBox)
 		.HAlign(HAlign_Center)
@@ -108,84 +215,25 @@ void SGameplayCamerasDebugger::Construct(const FArguments& InArgs)
 				.Text(LOCTEXT("EmptyPanelWarning", "No custom controls for this debug category."))
 		];
 
-	// Toolbar with debug category buttons.
-	if (!ToolMenus->IsMenuRegistered(SGameplayCamerasDebugger::ToolbarName))
-	{
-		FToolMenuOwnerScoped ToolMenuOwnerScope(this);
-
-		UToolMenu* Toolbar = ToolMenus->RegisterMenu(
-				SGameplayCamerasDebugger::ToolbarName, NAME_None, EMultiBoxType::SlimHorizontalToolBar);
+	// Register custom panels.
+	IGameplayCamerasEditorModule& ThisModule = FModuleManager::GetModuleChecked<IGameplayCamerasEditorModule>(
+			TEXT("GameplayCamerasEditor"));
+	TArray<FCameraDebugCategoryInfo> RegisteredDebugCategories;
+	ThisModule.GetRegisteredDebugCategories(RegisteredDebugCategories);
 
-		FToolMenuSection& MainSection = Toolbar->AddSection(TEXT("Main"));
+	for (const FCameraDebugCategoryInfo& DebugCategory : RegisteredDebugCategories)
+	{
+		TSharedPtr<SWidget> DebugCategoryPanel = ThisModule.CreateDebugCategoryPanel(DebugCategory.Name);
+		if (DebugCategoryPanel.IsValid())
 		{
-			FToolMenuEntry ToggleDebugInfo = FToolMenuEntry::InitToolBarButton(
-					Commands.EnableDebugInfo,
-					TAttribute<FText>::CreateSP(this, &SGameplayCamerasDebugger::GetToggleDebugDrawText),
-					TAttribute<FText>(),
-					TAttribute<FSlateIcon>::CreateSP(this, &SGameplayCamerasDebugger::GetToggleDebugDrawIcon));
-			ToggleDebugInfo.SetCommandList(CommandList);
-			MainSection.AddEntry(ToggleDebugInfo);
+			DebugPanels.Add(DebugCategory.Name, DebugCategoryPanel);
 		}
-	
-		FToolMenuSection& DebugCategoriesSection = Toolbar->AddSection(TEXT("DebugCategories"));
+		else
 		{
-			TArray<FCameraDebugCategoryInfo> RegisteredDebugCategories;
-			GameplayCamerasEditorModule.GetRegisteredDebugCategories(RegisteredDebugCategories);
-			for (const FCameraDebugCategoryInfo& DebugCategory : RegisteredDebugCategories)
-			{
-				TSharedPtr<SWidget> DebugCategoryPanel = GameplayCamerasEditorModule.CreateDebugCategoryPanel(DebugCategory.Name);
-				if (DebugCategoryPanel.IsValid())
-				{
-					DebugPanels.Add(DebugCategory.Name, DebugCategoryPanel);
-				}
-				else
-				{
-					// If there aren't any special UI controls for this category, use an empty panel.
-					DebugPanels.Add(DebugCategory.Name, EmptyPanel);
-				}
-
-				FToolMenuEntry ToggleDebugCategory = FToolMenuEntry::InitToolBarButton(
-					FName(DebugCategory.Name),
-					FUIAction(
-						FExecuteAction::CreateSP(this, &SGameplayCamerasDebugger::SetActiveDebugCategoryPanel, DebugCategory.Name),
-						FCanExecuteAction(),
-						FIsActionChecked::CreateStatic(&SGameplayCamerasDebugger::IsDebugCategoryActive, DebugCategory.Name)),
-					DebugCategory.DisplayText,
-					DebugCategory.ToolTipText,
-					DebugCategory.IconImage,
-					EUserInterfaceActionType::ToggleButton);
-				DebugCategoriesSection.AddEntry(ToggleDebugCategory);
-			}
+			// If there aren't any special UI controls for this category, use an empty panel.
+			DebugPanels.Add(DebugCategory.Name, EmptyPanel);
 		}
 	}
-	FToolMenuContext ToolbarContext;
-	TSharedRef<SWidget> ToolbarContents = ToolMenus->GenerateWidget(
-			SGameplayCamerasDebugger::ToolbarName, ToolbarContext);
-
-	// Main layout.
-	ChildSlot
-	[
-		SNew(SVerticalBox)
-		+ SVerticalBox::Slot()
-			.AutoHeight()
-			[
-				MenubarContents
-			]
-		+ SVerticalBox::Slot()
-			.AutoHeight()
-			[
-				ToolbarContents
-			]
-		+ SVerticalBox::Slot()
-		.Padding(2.0)
-			[
-				SAssignNew(PanelHost, SBox)
-					.Padding(8.f)
-					[
-						EmptyPanel.ToSharedRef()
-					]
-			]
-	];
 }
 
 FText SGameplayCamerasDebugger::GetToggleDebugDrawText() const
diff --git a/Engine/Plugins/Cameras/GameplayCameras/Source/GameplayCamerasEditor/Private/Debugger/SGameplayCamerasDebugger.h b/Engine/Plugins/Cameras/GameplayCameras/Source/GameplayCamerasEditor/Private/Debugger/SGameplayCamerasDebugger.h
--- a/Engine/Plugins/Cameras/GameplayCameras/Source/GameplayCamerasEditor/Private/Debugger/SGameplayCamerasDebugger.h
+++ b/Engine/Plugins/Cameras/GameplayCameras/Source/GameplayCamerasEditor/Private/Debugger/SGameplayCamerasDebugger.h
@@ -5,10 +5,15 @@
 #include "CoreTypes.h"
 #include "Widgets/SCompoundWidget.h"
 
-class FTabManager;
+#include "SGameplayCamerasDebugger.generated.h"
+
 class FSpawnTabArgs;
+class FTabManager;
+class FUICommandList;
 class SBox;
 class SDockTab;
+class SWidget;
+class UToolMenu;
 struct FSlateIcon;
 
 namespace UE::Cameras
@@ -38,6 +43,11 @@ public:
 
 protected:
 
+	static SGameplayCamerasDebugger* FromContext(UToolMenu* InMenu);
+	TSharedRef<SWidget> ConstructMenubar();
+	TSharedRef<SWidget> ConstructToolbar(TSharedRef<FUICommandList> InCommandList);
+	void ConstructDebugPanels();
+
 	static bool IsDebugCategoryActive(FString InCategoryName);
 	void SetActiveDebugCategoryPanel(FString InCategoryName);
 
@@ -55,3 +65,13 @@ private:
 
 }  // namespace UE::Cameras
 
+UCLASS()
+class UGameplayCamerasDebuggerMenuContext : public UObject
+{
+	GENERATED_BODY()
+
+public:
+
+	TWeakPtr<UE::Cameras::SGameplayCamerasDebugger> CamerasDebugger;
+};
+
