Cameras: fix colors in blend stack debug render

diff --git a/Engine/Plugins/Cameras/GameplayCameras/Source/GameplayCameras/Private/Core/BlendStackCameraNode.cpp b/Engine/Plugins/Cameras/GameplayCameras/Source/GameplayCameras/Private/Core/BlendStackCameraNode.cpp
--- a/Engine/Plugins/Cameras/GameplayCameras/Source/GameplayCameras/Private/Core/BlendStackCameraNode.cpp
+++ b/Engine/Plugins/Cameras/GameplayCameras/Source/GameplayCameras/Private/Core/BlendStackCameraNode.cpp
@@ -519,7 +519,7 @@ void FBlendStackCameraDebugBlock::OnDebugDraw(const FCameraDebugBlockDrawParams&
 	{
 		const FEntryDebugInfo& Entry(Entries[Index]);
 
-		Renderer.AddText(TEXT("{springgreen}[%d]{white} %s\n"), Index + 1, *Entry.CameraRigName);
+		Renderer.AddText(TEXT("{cam_passive}[%d]{cam_highlighted} %s\n"), Index + 1, *Entry.CameraRigName);
 
 		if (ChildrenView.IsValidIndex(Index))
 		{
