Cameras: customizable color schemes for debug rendering

diff --git a/Engine/Plugins/Cameras/GameplayCameras/Source/GameplayCameras/Private/Debug/CameraDebugColors.cpp b/Engine/Plugins/Cameras/GameplayCameras/Source/GameplayCameras/Private/Debug/CameraDebugColors.cpp
--- a/Engine/Plugins/Cameras/GameplayCameras/Source/GameplayCameras/Private/Debug/CameraDebugColors.cpp
+++ b/Engine/Plugins/Cameras/GameplayCameras/Source/GameplayCameras/Private/Debug/CameraDebugColors.cpp
@@ -9,12 +9,50 @@
 namespace UE::Cameras
 {
 
+FString FCameraDebugColors::CurrentColorSchemeName;
+FCameraDebugColors FCameraDebugColors::CurrentColorScheme;
 TMap<FString, FColor> FCameraDebugColors::ColorMap;
+TMap<FString, FCameraDebugColors> FCameraDebugColors::ColorSchemes;
 
 const FCameraDebugColors& FCameraDebugColors::Get()
 {
-	static FCameraDebugColors StaticInstance;
-	return StaticInstance;
+	return CurrentColorScheme;
+}
+
+const FString& FCameraDebugColors::GetName()
+{
+	return CurrentColorSchemeName;
+}
+
+void FCameraDebugColors::Set(const FString& InColorSchemeName)
+{
+	if (InColorSchemeName != CurrentColorSchemeName)
+	{
+		const FCameraDebugColors* ColorScheme = ColorSchemes.Find(InColorSchemeName);
+		if (ensureMsgf(ColorScheme, TEXT("No such color scheme: %s"), *InColorSchemeName))
+		{
+			CurrentColorSchemeName = InColorSchemeName;
+			CurrentColorScheme = *ColorScheme;
+			UpdateColorMap(CurrentColorScheme);
+		}
+	}
+}
+
+void FCameraDebugColors::Set(const FCameraDebugColors& InColorScheme)
+{
+	CurrentColorSchemeName = TEXT("<Custom>");
+	CurrentColorScheme = InColorScheme;
+	UpdateColorMap(CurrentColorScheme);
+}
+
+void FCameraDebugColors::RegisterColorScheme(const FString& InColorSchemeName, const FCameraDebugColors& InColorScheme)
+{
+	ColorSchemes.Add(InColorSchemeName, InColorScheme);
+}
+
+void FCameraDebugColors::GetColorSchemeNames(TArray<FString>& OutColorSchemeNames)
+{
+	ColorSchemes.GetKeys(OutColorSchemeNames);
 }
 
 void FCameraDebugColors::UpdateColorMap(const FCameraDebugColors& Instance)
@@ -46,10 +84,9 @@ TOptional<FColor> FCameraDebugColors::GetFColorByName(const FString& InColorName
 	return TOptional<FColor>();
 }
 
-FCameraDebugColors::FCameraDebugColors()
+void FCameraDebugColors::RegisterBuiltinColorSchemes()
 {
-	// Default to colors inspired by the Solarized palette.
-	// Other palettes could be implemented.
+	// Colors inspired by the Solarized palette.
 	//
 	//    SOLARIZED HEX     RGB        
     //    --------- ------- -----------
@@ -70,17 +107,56 @@ FCameraDebugColors::FCameraDebugColors()
     //    cyan      #2aa198  42 161 152
     //    green     #859900 133 153   0
 	//
-	Title = FColor(38, 139, 210); // blue
-	Default = FColor(238, 232, 213); // base2
-	Passive = FColor(147, 161, 161); // base1
-	VeryPassive = FColor(101, 123, 131); // base00
-	Hightlighted = FColor(253, 246, 227); // base3
-	Notice = FColor(42, 161, 152); // cyan
-	Notice2 = FColor(211, 54, 130); // magenta
-	Good = FColor(133, 153, 0); // green
-	Warning = FColor(181, 137, 0); // yellow
-	Error = FColor(220, 50, 47); // red
-	Background = FColor(7, 54, 66); // base02
+    const FColor Base03(   0,  43,  54);
+    const FColor Base02(   7,  54,  66);
+    const FColor Base01(  88, 110, 117);
+    const FColor Base00( 101, 123, 131);
+    const FColor Base0(  131, 148, 150);
+    const FColor Base1(  147, 161, 161);
+    const FColor Base2(  238, 232, 213);
+    const FColor Base3(  253, 246, 227);
+    const FColor Yellow( 181, 137,   0);
+    const FColor Orange( 203,  75,  22);
+    const FColor Red(    220,  50,  47);
+    const FColor Magenta(211,  54, 130);
+    const FColor Violet( 108, 113, 196);
+    const FColor Blue(    38, 139, 210);
+    const FColor Cyan(    42, 161, 152);
+    const FColor Green(  133, 153,   0);
+	{
+		FCameraDebugColors SolarizedDark;
+		SolarizedDark.Title = Blue;
+		SolarizedDark.Default = Base2;
+		SolarizedDark.Passive = Base1;
+		SolarizedDark.VeryPassive = Base0;
+		SolarizedDark.Hightlighted = Base3;
+		SolarizedDark.Notice = Cyan;
+		SolarizedDark.Notice2 = Magenta;
+		SolarizedDark.Good = Green;
+		SolarizedDark.Warning = Yellow;
+		SolarizedDark.Error = Red;
+		SolarizedDark.Background = Base03;
+
+		RegisterColorScheme(TEXT("SolarizedDark"), SolarizedDark);
+	}
+	{
+		FCameraDebugColors SolarizedLight;
+		SolarizedLight.Title = Blue;
+		SolarizedLight.Default = Base01;
+		SolarizedLight.Passive = Base00;
+		SolarizedLight.VeryPassive = Base0;
+		SolarizedLight.Hightlighted = Base03;
+		SolarizedLight.Notice = Cyan;
+		SolarizedLight.Notice2 = Magenta;
+		SolarizedLight.Good = Green;
+		SolarizedLight.Warning = Yellow;
+		SolarizedLight.Error = Red;
+		SolarizedLight.Background = Base3;
+
+		RegisterColorScheme(TEXT("SolarizedLight"), SolarizedLight);
+	}
+
+	Set(TEXT("SolarizedDark"));
 }
 
 }  // namespace UE::Cameras
diff --git a/Engine/Plugins/Cameras/GameplayCameras/Source/GameplayCameras/Private/Debug/RootCameraDebugBlock.cpp b/Engine/Plugins/Cameras/GameplayCameras/Source/GameplayCameras/Private/Debug/RootCameraDebugBlock.cpp
--- a/Engine/Plugins/Cameras/GameplayCameras/Source/GameplayCameras/Private/Debug/RootCameraDebugBlock.cpp
+++ b/Engine/Plugins/Cameras/GameplayCameras/Source/GameplayCameras/Private/Debug/RootCameraDebugBlock.cpp
@@ -7,6 +7,7 @@
 #include "Core/RootCameraNode.h"
 #include "Debug/CameraDebugBlockBuilder.h"
 #include "Debug/CameraDebugCategories.h"
+#include "Debug/CameraDebugColors.h"
 #include "Debug/CameraDebugRenderer.h"
 #include "Debug/CameraDirectorTreeDebugBlock.h"
 #include "Debug/CameraNodeEvaluatorDebugBlock.h"
@@ -51,6 +52,12 @@ static FAutoConsoleVariableRef CVarGameplayCamerasDebugBackgroundOpacity(
 	GGameplayCamerasDebugBackgroundOpacity,
 	TEXT(""));
 
+FString GGameplayCamerasDebugColorScheme = TEXT("SolarizedDark");
+static FAutoConsoleVariableRef CVarGameplayCamerasDebugColorScheme(
+	TEXT("GameplayCameras.Debug.ColorScheme"),
+	GGameplayCamerasDebugColorScheme,
+	TEXT(""));
+
 UE_DEFINE_CAMERA_DEBUG_BLOCK(FRootCameraDebugBlock)
 
 void FRootCameraDebugBlock::BuildDebugBlocks(const FCameraSystemEvaluator& CameraSystem, const FCameraDebugBlockBuildParams& Params, FCameraDebugBlockBuilder& Builder)
@@ -116,6 +123,9 @@ void FRootCameraDebugBlock::RootDebugDraw(FCameraDebugRenderer& Renderer)
 		Params.ActiveCategories.Add(FString(CategoryView));
 	}
 
+	// Update the color scheme in case it changed.
+	FCameraDebugColors::Set(GGameplayCamerasDebugColorScheme);
+
 	// Do the drawing!
 	FCameraDebugBlock::DebugDraw(Params, Renderer);
 
diff --git a/Engine/Plugins/Cameras/GameplayCameras/Source/GameplayCameras/Private/GameplayCamerasModule.cpp b/Engine/Plugins/Cameras/GameplayCameras/Source/GameplayCameras/Private/GameplayCamerasModule.cpp
--- a/Engine/Plugins/Cameras/GameplayCameras/Source/GameplayCameras/Private/GameplayCamerasModule.cpp
+++ b/Engine/Plugins/Cameras/GameplayCameras/Source/GameplayCameras/Private/GameplayCamerasModule.cpp
@@ -5,6 +5,7 @@
 #include "Camera/CameraModularFeature.h"
 #include "CameraAnimationCameraModifier.h"
 #include "CameraAnimationSequencePlayer.h"
+#include "Debug/CameraDebugColors.h"
 #include "Features/IModularFeatures.h"
 #include "GameplayCameras.h"
 #include "Logging/MessageLog.h"
@@ -29,6 +30,10 @@ public:
 		{
 			IModularFeatures::Get().RegisterModularFeature(ICameraModularFeature::GetModularFeatureName(), CameraModularFeature.Get());
 		}
+
+#if UE_GAMEPLAY_CAMERAS_DEBUG
+		UE::Cameras::FCameraDebugColors::RegisterBuiltinColorSchemes();
+#endif  // UE_GAMEPLAY_CAMERAS_DEBUG
 	}
 
 	virtual void ShutdownModule() override
diff --git a/Engine/Plugins/Cameras/GameplayCameras/Source/GameplayCameras/Public/Debug/CameraDebugColors.h b/Engine/Plugins/Cameras/GameplayCameras/Source/GameplayCameras/Public/Debug/CameraDebugColors.h
--- a/Engine/Plugins/Cameras/GameplayCameras/Source/GameplayCameras/Public/Debug/CameraDebugColors.h
+++ b/Engine/Plugins/Cameras/GameplayCameras/Source/GameplayCameras/Public/Debug/CameraDebugColors.h
@@ -18,8 +18,6 @@ class FCameraDebugColors
 {
 public:
 
-	static const FCameraDebugColors& Get();
-
 	/** Debug category titles and section titles. */
 	FColor Title;
 	/** Normal text. */
@@ -43,15 +41,41 @@ public:
 
 public:
 
+	/** Gets the current color scheme. */
+	static const FCameraDebugColors& Get();
+
+	/** Gets a specific color by name on the current scheme. */
 	static TOptional<FColor> GetFColorByName(const FString& InColorName);
 
-private:
+	/** Gets the current color scheme name. */
+	GAMEPLAYCAMERAS_API static const FString& GetName();
+	/** Sets the current color scheme. */
+	GAMEPLAYCAMERAS_API static void Set(const FString& InColorSchemeName);
+	/** Sets the current color scheme. */
+	GAMEPLAYCAMERAS_API static void Set(const FCameraDebugColors& InColorScheme);
+
+	/** Registers a new color scheme. */
+	GAMEPLAYCAMERAS_API static void RegisterColorScheme(const FString& InColorSchemeName, const FCameraDebugColors& InColorScheme);
+	/** Gets all registered color scheme names. */
+	GAMEPLAYCAMERAS_API static void GetColorSchemeNames(TArray<FString>& OutColorSchemeNames);
+
 
-	FCameraDebugColors();
+public:
+
+	// Internal API.
+	static void RegisterBuiltinColorSchemes();
+
+private:
 
 	static void UpdateColorMap(const FCameraDebugColors& Instance);
 
+private:
+
+	static FString CurrentColorSchemeName;
+	static FCameraDebugColors CurrentColorScheme;
 	static TMap<FString, FColor> ColorMap;
+
+	static TMap<FString, FCameraDebugColors> ColorSchemes;
 };
 
 }  // namespace UE::Cameras
