Cameras: add more options for debug rendering position

diff --git a/Engine/Plugins/Cameras/GameplayCameras/Source/GameplayCameras/Private/Debug/CameraDebugRenderer.cpp b/Engine/Plugins/Cameras/GameplayCameras/Source/GameplayCameras/Private/Debug/CameraDebugRenderer.cpp
--- a/Engine/Plugins/Cameras/GameplayCameras/Source/GameplayCameras/Private/Debug/CameraDebugRenderer.cpp
+++ b/Engine/Plugins/Cameras/GameplayCameras/Source/GameplayCameras/Private/Debug/CameraDebugRenderer.cpp
@@ -17,11 +17,23 @@
 namespace UE::Cameras
 {
 
-int32 GGameplayCamerasDebugMargin = 10;
-static FAutoConsoleVariableRef CVarGameplayCamerasDebugMargin(
-	TEXT("GameplayCameras.Debug.Margin"),
-	GGameplayCamerasDebugMargin,
-	TEXT("(Default: 10px. The margin for rendering Gameplay Cameras debug text."));
+int32 GGameplayCamerasDebugLeftMargin = 10;
+static FAutoConsoleVariableRef CVarGameplayCamerasDebugLeftMargin(
+	TEXT("GameplayCameras.Debug.LeftMargin"),
+	GGameplayCamerasDebugLeftMargin,
+	TEXT("(Default: 10px. The left margin for rendering Gameplay Cameras debug text."));
+
+int32 GGameplayCamerasDebugTopMargin = 10;
+static FAutoConsoleVariableRef CVarGameplayCamerasDebugTopMargin(
+	TEXT("GameplayCameras.Debug.TopMargin"),
+	GGameplayCamerasDebugTopMargin,
+	TEXT("(Default: 10px. The top margin for rendering Gameplay Cameras debug text."));
+
+int32 GGameplayCamerasDebugInnerMargin = 5;
+static FAutoConsoleVariableRef CVarGameplayCamerasDebugInnerMargin(
+	TEXT("GameplayCameras.Debug.InnerMargin"),
+	GGameplayCamerasDebugInnerMargin,
+	TEXT("(Default: 10px. The inner margin for rendering Gameplay Cameras debug text."));
 
 int32 GGameplayCamerasDebugIndent = 20;
 static FAutoConsoleVariableRef CVarGameplayCamerasDebugIndent(
@@ -42,7 +54,7 @@ FCameraDebugRenderer::FCameraDebugRenderer(FCanvas* InCanvas)
 	RenderFont = GEngine->GetSmallFont();
 	MaxCharHeight = RenderFont->GetMaxCharHeight();
 
-	NextDrawPosition = FVector2f{ (float)GGameplayCamerasDebugMargin, (float)GGameplayCamerasDebugMargin };
+	NextDrawPosition = FVector2f{ (float)GGameplayCamerasDebugLeftMargin, (float)GGameplayCamerasDebugTopMargin };
 }
 
 FCameraDebugRenderer::~FCameraDebugRenderer()
@@ -112,7 +124,7 @@ FColor FCameraDebugRenderer::SetTextColor(const FColor& Color)
 
 float FCameraDebugRenderer::GetIndentMargin() const
 {
-	return (float)(GGameplayCamerasDebugMargin + IndentLevel * GGameplayCamerasDebugIndent);
+	return (float)(GGameplayCamerasDebugLeftMargin + IndentLevel * GGameplayCamerasDebugIndent);
 }
 
 void FCameraDebugRenderer::FlushText()
@@ -170,9 +182,10 @@ void FCameraDebugRenderer::DrawTextBackgroundTile(float Opacity)
 	const bool bIsLineEmpty = FMath::IsNearlyEqual(NextDrawPosition.X, IndentMargin);
 	const float TextBottom = bIsLineEmpty ? NextDrawPosition.Y : NextDrawPosition.Y + MaxCharHeight;
 
-	const float OuterMargin = GGameplayCamerasDebugMargin / 2.f;
-	FVector2D TopLeft(OuterMargin, OuterMargin);
-	FVector2D TileSize(RightMargin + OuterMargin, TextBottom + OuterMargin);
+	const float InnerMargin = GGameplayCamerasDebugInnerMargin;
+	const FVector2D TopLeft(GGameplayCamerasDebugLeftMargin - InnerMargin, GGameplayCamerasDebugTopMargin - InnerMargin);
+	const FVector2D BottomRight(RightMargin + InnerMargin, TextBottom + InnerMargin);
+	const FVector2D TileSize(BottomRight.X - TopLeft.X, BottomRight.Y - TopLeft.Y);
 
 	const FColor BackgroundColor = FCameraDebugColors::Get().Background.WithAlpha((uint8)(Opacity * 255));
 
