Cameras: add global options panel in the debugger tool

diff --git a/Engine/Plugins/Cameras/GameplayCameras/Source/GameplayCamerasEditor/Private/Debugger/SDebugWidgetUtils.cpp b/Engine/Plugins/Cameras/GameplayCameras/Source/GameplayCamerasEditor/Private/Debugger/SDebugWidgetUtils.cpp
--- a/Engine/Plugins/Cameras/GameplayCameras/Source/GameplayCamerasEditor/Private/Debugger/SDebugWidgetUtils.cpp
+++ b/Engine/Plugins/Cameras/GameplayCameras/Source/GameplayCamerasEditor/Private/Debugger/SDebugWidgetUtils.cpp
@@ -4,30 +4,35 @@
 
 #include "HAL/IConsoleManager.h"
 #include "Widgets/Input/SCheckBox.h"
+#include "Widgets/Input/SComboBox.h"
 #include "Widgets/Input/SEditableTextBox.h"
+#include "Widgets/Input/SSpinBox.h"
+#include "Widgets/SBoxPanel.h"
 #include "Widgets/Text/STextBlock.h"
 
+#define LOCTEXT_NAMESPACE "SDebugWidgetUtils"
+
 namespace UE::Cameras
 {
 
-TMap<FString, IConsoleVariable*> SDebugWidgetUtils::CachedConsoleVariables;
-
-TSharedRef<SCheckBox> SDebugWidgetUtils::CreateConsoleVariableCheckBox(const FText& Text, const FString& ConsoleVariableName)
+TSharedRef<SWidget> SDebugWidgetUtils::CreateConsoleVariableCheckBox(const FText& Text, const FString& ConsoleVariableName)
 {
+	IConsoleVariable* ConsoleVariable = SDebugWidgetUtils::GetConsoleVariable(ConsoleVariableName);
+
 	return SNew(SCheckBox)
 		.Padding(4.f)
-		.IsChecked_Lambda([ConsoleVariableName]() 
+		.IsChecked_Lambda([ConsoleVariable]() 
 				{
-					if (IConsoleVariable* ConsoleVariable = SDebugWidgetUtils::GetConsoleVariable(ConsoleVariableName))
+					if (ConsoleVariable)
 					{
 						return ConsoleVariable->GetBool() ?
 							ECheckBoxState::Checked : ECheckBoxState::Unchecked;
 					}
 					return ECheckBoxState::Undetermined;
 				})
-		.OnCheckStateChanged_Lambda([ConsoleVariableName](ECheckBoxState NewState)
+		.OnCheckStateChanged_Lambda([ConsoleVariable](ECheckBoxState NewState)
 				{
-					if (IConsoleVariable* ConsoleVariable = SDebugWidgetUtils::GetConsoleVariable(ConsoleVariableName))
+					if (ConsoleVariable)
 					{
 						ConsoleVariable->Set(NewState == ECheckBoxState::Checked);
 					}
@@ -38,21 +43,112 @@ TSharedRef<SCheckBox> SDebugWidgetUtils::CreateConsoleVariableCheckBox(const FTe
 		];
 }
 
-TSharedRef<SEditableTextBox> SDebugWidgetUtils::CreateConsoleVariableTextBox(const FString& ConsoleVariableName)
+TSharedRef<SWidget> SDebugWidgetUtils::CreateConsoleVariableSpinBox(const FString& ConsoleVariableName)
+{
+	IConsoleVariable* ConsoleVariable = SDebugWidgetUtils::GetConsoleVariable(ConsoleVariableName);
+
+	return SNew(SSpinBox<float>)
+		.Value_Lambda([ConsoleVariable]()
+				{
+					return ConsoleVariable ? ConsoleVariable->GetFloat() : 0.f;
+				})
+		.OnValueChanged_Lambda([ConsoleVariable](float NewValue)
+				{
+					if (ConsoleVariable)
+					{
+						ConsoleVariable->Set(NewValue);
+					}
+				})
+		.OnValueCommitted_Lambda([ConsoleVariable](float NewValue, ETextCommit::Type CommitType)
+				{
+					if (ConsoleVariable)
+					{
+						ConsoleVariable->Set(NewValue);
+					}
+				});
+}
+
+TSharedRef<SWidget> SDebugWidgetUtils::CreateConsoleVariableSpinBox(const FText& Text, const FString& ConsoleVariableName)
+{
+	return SNew(SHorizontalBox)
+		+ SHorizontalBox::Slot()
+			.Padding(4.f)
+			[
+				SNew(STextBlock)
+					.Text(Text)
+			]
+		+ SHorizontalBox::Slot()
+			.Padding(4.f)
+			[
+				CreateConsoleVariableSpinBox(ConsoleVariableName)
+			];
+}
+
+TSharedRef<SWidget> SDebugWidgetUtils::CreateConsoleVariableComboBox(const FString& ConsoleVariableName, TArray<TSharedPtr<FString>>* OptionsSource)
+{
+	IConsoleVariable* ConsoleVariable = SDebugWidgetUtils::GetConsoleVariable(ConsoleVariableName);
+
+	TSharedPtr<FString> InitialSelection;
+	if (ConsoleVariable)
+	{
+		const FString CurrentValue = ConsoleVariable->GetString();
+		TSharedPtr<FString>* FoundItem = OptionsSource->FindByPredicate([&CurrentValue](TSharedPtr<FString> Item)
+				{
+					return Item && CurrentValue == *Item.Get();
+				});
+		if (FoundItem)
+		{
+			InitialSelection = *FoundItem;
+		}
+	}
+
+	return SNew(SComboBox<TSharedPtr<FString>>)
+		.OptionsSource(OptionsSource)
+		.OnGenerateWidget_Lambda([](TSharedPtr<FString> InOption)
+				{
+					return SNew(STextBlock)
+						.Text(FText::FromString(*InOption));
+				})
+		.OnSelectionChanged_Lambda([ConsoleVariable](TSharedPtr<FString> NewSelection, ESelectInfo::Type SelectInfo)
+				{
+					if (ConsoleVariable)
+					{
+						ConsoleVariable->Set(**NewSelection.Get());
+					}
+				})
+		.InitiallySelectedItem(InitialSelection)
+		.ContentPadding(2.f)
+		.Content()
+		[
+			SNew(STextBlock)
+				.Text_Lambda([ConsoleVariable]() -> FText
+						{
+							if (ConsoleVariable)
+							{
+								return FText::FromString(ConsoleVariable->GetString());
+							}
+							return LOCTEXT("NoSuchVariable", "<no such variable>");
+						})
+		];
+}
+
+TSharedRef<SWidget> SDebugWidgetUtils::CreateConsoleVariableTextBox(const FString& ConsoleVariableName)
 {
+	IConsoleVariable* ConsoleVariable = SDebugWidgetUtils::GetConsoleVariable(ConsoleVariableName);
+
 	return SNew(SEditableTextBox)
 		.Padding(4.f)
-		.Text_Lambda([ConsoleVariableName]()
+		.Text_Lambda([ConsoleVariable]()
 				{
-					if (IConsoleVariable* ConsoleVariable = SDebugWidgetUtils::GetConsoleVariable(ConsoleVariableName))
+					if (ConsoleVariable)
 					{
 						return FText::FromString(ConsoleVariable->GetString());
 					}
 					return FText();
 				})
-		.OnTextChanged_Lambda([ConsoleVariableName](const FText& NewText)
+		.OnTextChanged_Lambda([ConsoleVariable](const FText& NewText)
 				{
-					if (IConsoleVariable* ConsoleVariable = SDebugWidgetUtils::GetConsoleVariable(ConsoleVariableName))
+					if (ConsoleVariable)
 					{
 						ConsoleVariable->Set(*NewText.ToString());
 					}
@@ -61,20 +157,12 @@ TSharedRef<SEditableTextBox> SDebugWidgetUtils::CreateConsoleVariableTextBox(con
 
 IConsoleVariable* SDebugWidgetUtils::GetConsoleVariable(const FString& ConsoleVariableName)
 {
-	if (IConsoleVariable** ConsoleVariablePtr = CachedConsoleVariables.Find(ConsoleVariableName))
-	{
-		return (*ConsoleVariablePtr);
-	}
-
 	IConsoleVariable* ConsoleVariable = IConsoleManager::Get().FindConsoleVariable(*ConsoleVariableName);
-	if (ensureMsgf(ConsoleVariable, TEXT("No such console variable: %s"), *ConsoleVariableName))
-	{
-		CachedConsoleVariables.Add(ConsoleVariableName, ConsoleVariable);
-		return ConsoleVariable;
-	}
-
-	return nullptr;
+	ensureMsgf(ConsoleVariable, TEXT("No such console variable: %s"), *ConsoleVariableName);
+	return ConsoleVariable;
 }
 
 }  // namespace UE::Cameras
 
+#undef LOCTEXT_NAMESPACE
+
diff --git a/Engine/Plugins/Cameras/GameplayCameras/Source/GameplayCamerasEditor/Private/Debugger/SDebugWidgetUtils.h b/Engine/Plugins/Cameras/GameplayCameras/Source/GameplayCamerasEditor/Private/Debugger/SDebugWidgetUtils.h
--- a/Engine/Plugins/Cameras/GameplayCameras/Source/GameplayCamerasEditor/Private/Debugger/SDebugWidgetUtils.h
+++ b/Engine/Plugins/Cameras/GameplayCameras/Source/GameplayCamerasEditor/Private/Debugger/SDebugWidgetUtils.h
@@ -2,31 +2,38 @@
 
 #pragma once
 
-#include "Containers/Map.h"
+#include "Containers/Array.h"
 #include "CoreMinimal.h"
 #include "Templates/SharedPointerFwd.h"
 
+class FString;
 class IConsoleVariable;
-class SCheckBox;
-class SEditableTextBox;
+class SWidget;
 
 namespace UE::Cameras
 {
 
+/**
+ * Utility class for creating Slate widgets that drive console variables.
+ */
 class SDebugWidgetUtils
 {
 public:
 
-	static TSharedRef<SCheckBox> CreateConsoleVariableCheckBox(const FText& Text, const FString& ConsoleVariableName);
-	static TSharedRef<SEditableTextBox> CreateConsoleVariableTextBox(const FString& ConsoleVariableName);
+	/** Creates a checkbox that sets a boolean console variable. */
+	static TSharedRef<SWidget> CreateConsoleVariableCheckBox(const FText& Text, const FString& ConsoleVariableName);
+	/** Creates a spinbox that sets a float console variable. */
+	static TSharedRef<SWidget> CreateConsoleVariableSpinBox(const FString& ConsoleVariableName);
+	/** Creates a spinbox that sets a float console variable, with a text label. */
+	static TSharedRef<SWidget> CreateConsoleVariableSpinBox(const FText& Text, const FString& ConsoleVariableName);
+	/** Creates a combox box that sets a string console variable. */
+	static TSharedRef<SWidget> CreateConsoleVariableComboBox(const FString& ConsoleVariableName, TArray<TSharedPtr<FString>>* OptionsSource);
+	/** Creates a text box that sets a string console variable. */
+	static TSharedRef<SWidget> CreateConsoleVariableTextBox(const FString& ConsoleVariableName);
 
 private:
 
 	static IConsoleVariable* GetConsoleVariable(const FString& ConsoleVariableName);
-
-private:
-
-	static TMap<FString, IConsoleVariable*> CachedConsoleVariables;
 };
 
 }  // namespace UE::Cameras
diff --git a/Engine/Plugins/Cameras/GameplayCameras/Source/GameplayCamerasEditor/Private/Debugger/SGameplayCamerasDebugger.cpp b/Engine/Plugins/Cameras/GameplayCameras/Source/GameplayCamerasEditor/Private/Debugger/SGameplayCamerasDebugger.cpp
--- a/Engine/Plugins/Cameras/GameplayCameras/Source/GameplayCamerasEditor/Private/Debugger/SGameplayCamerasDebugger.cpp
+++ b/Engine/Plugins/Cameras/GameplayCameras/Source/GameplayCamerasEditor/Private/Debugger/SGameplayCamerasDebugger.cpp
@@ -3,8 +3,10 @@
 #include "Debugger/SGameplayCamerasDebugger.h"
 
 #include "Commands/GameplayCamerasDebuggerCommands.h"
+#include "Debug/CameraDebugColors.h"
 #include "Debug/RootCameraDebugBlock.h"
 #include "Debugger/SDebugCategoryButton.h"
+#include "Debugger/SDebugWidgetUtils.h"
 #include "Framework/Application/SlateApplication.h"
 #include "Framework/Docking/LayoutService.h"
 #include "Framework/Docking/TabManager.h"
@@ -17,8 +19,9 @@
 #include "ToolMenus.h"
 #include "ToolMenuDelegates.h"
 #include "Widgets/Docking/SDockTab.h"
-#include "Widgets/Input/SButton.h"
 #include "Widgets/Layout/SBox.h"
+#include "Widgets/Layout/SExpandableArea.h"
+#include "Widgets/Layout/SGridPanel.h"
 #include "Widgets/Text/STextBlock.h"
 #include "Widgets/SNullWidget.h"
 #include "WorkspaceMenuStructure.h"
@@ -77,6 +80,8 @@ void SGameplayCamerasDebugger::Construct(const FArguments& InArgs)
 	TSharedRef<FGameplayCamerasEditorStyle> GameplayCamerasEditorStyle = FGameplayCamerasEditorStyle::Get();
 	GameplayCamerasEditorStyleName = GameplayCamerasEditorStyle->GetStyleSetName();
 
+	InitializeColorSchemeNames();
+
 	// Setup commands.
 	const FGameplayCamerasDebuggerCommands& Commands = FGameplayCamerasDebuggerCommands::Get();
 	TSharedRef<FUICommandList> CommandList = MakeShareable(new FUICommandList);
@@ -89,6 +94,7 @@ void SGameplayCamerasDebugger::Construct(const FArguments& InArgs)
 	// Build all UI elements.
 	TSharedRef<SWidget> MenubarContents = ConstructMenubar();
 	TSharedRef<SWidget> ToolbarContents = ConstructToolbar(CommandList);
+	TSharedRef<SWidget> GeneralOptionsContents = ConstructGeneralOptions(CommandList);
 	ConstructDebugPanels();
 
 	// Main layout.
@@ -106,7 +112,7 @@ void SGameplayCamerasDebugger::Construct(const FArguments& InArgs)
 				ToolbarContents
 			]
 		+ SVerticalBox::Slot()
-		.Padding(2.0)
+			.Padding(2.0)
 			[
 				SAssignNew(PanelHost, SBox)
 					.Padding(8.f)
@@ -114,7 +120,31 @@ void SGameplayCamerasDebugger::Construct(const FArguments& InArgs)
 						EmptyPanel.ToSharedRef()
 					]
 			]
+		+ SVerticalBox::Slot()
+			.AutoHeight()
+			.Padding(2.0)
+			[
+				GeneralOptionsContents
+			]
 	];
+
+	// Set initial debug panel.
+	TArray<FStringView, TInlineAllocator<4>> ActiveCategories;
+	UE::String::ParseTokens(GGameplayCamerasDebugCategories, ',', ActiveCategories);
+	if (!ActiveCategories.IsEmpty())
+	{
+		SetActiveDebugCategoryPanel(FString(ActiveCategories[0]));
+	}
+}
+
+void SGameplayCamerasDebugger::InitializeColorSchemeNames()
+{
+	TArray<FString> RawNames;
+	FCameraDebugColors::GetColorSchemeNames(RawNames);
+	for (const FString& RawName :RawNames)
+	{
+		ColorSchemeNames.Add(MakeShared<FString>(RawName));
+	}
 }
 
 SGameplayCamerasDebugger* SGameplayCamerasDebugger::FromContext(UToolMenu* InMenu)
@@ -203,6 +233,108 @@ TSharedRef<SWidget> SGameplayCamerasDebugger::ConstructToolbar(TSharedRef<FUICom
 	return ToolMenus->GenerateWidget(SGameplayCamerasDebugger::ToolbarName, ToolbarContext);
 }
 
+TSharedRef<SWidget> SGameplayCamerasDebugger::ConstructGeneralOptions(TSharedRef<FUICommandList> InCommandList)
+{
+	const ISlateStyle& AppStyle = FAppStyle::Get();
+	const FMargin GridCellPadding(4.f);
+
+	return SNew(SExpandableArea)
+		.BorderImage(AppStyle.GetBrush("Brushes.Header"))
+		.BodyBorderImage(AppStyle.GetBrush("Brushes.Recessed"))
+		.HeaderPadding(FMargin(4.0f))
+		.Padding(FMargin(0, 1, 0, 0))
+		.InitiallyCollapsed(true)
+		.AllowAnimatedTransition(false)
+		.HeaderContent()
+		[
+			SNew(SHorizontalBox)
+			+ SHorizontalBox::Slot()
+				.VAlign(VAlign_Center)
+				[
+					SNew(STextBlock)
+						.Text(LOCTEXT("GeneralOptions", "General Options"))
+						.TextStyle(AppStyle, "ButtonText")
+						.Font(AppStyle.GetFontStyle("NormalFontBold"))
+				]
+		]
+		.BodyContent()
+		[
+			SNew(SBorder)
+				.BorderImage(AppStyle.GetBrush("Brushes.Header"))
+				.Padding(2.f)
+				[
+					SNew(SGridPanel)
+						.FillColumn(0, 1.f)
+						.FillColumn(2, 1.f)
+					+ SGridPanel::Slot(0, 0)
+					.Padding(GridCellPadding)
+					.VAlign(VAlign_Center)
+					[
+						SNew(STextBlock)
+							.Text(LOCTEXT("TopMargin", "Top margin"))
+					]
+					+ SGridPanel::Slot(1, 0)
+					.Padding(GridCellPadding)
+					.VAlign(VAlign_Center)
+					[
+						SDebugWidgetUtils::CreateConsoleVariableSpinBox(TEXT("GameplayCameras.Debug.TopMargin"))
+					]
+					+ SGridPanel::Slot(0, 1)
+					.Padding(GridCellPadding)
+					.VAlign(VAlign_Center)
+					[
+						SNew(STextBlock)
+							.Text(LOCTEXT("LeftMargin", "Left margin"))
+					]
+					+ SGridPanel::Slot(1, 1)
+					.Padding(GridCellPadding)
+					.VAlign(VAlign_Center)
+					[
+						SDebugWidgetUtils::CreateConsoleVariableSpinBox(TEXT("GameplayCameras.Debug.LeftMargin"))
+					]
+					+ SGridPanel::Slot(0, 2)
+					.Padding(GridCellPadding)
+					.VAlign(VAlign_Center)
+					[
+						SNew(STextBlock)
+							.Text(LOCTEXT("InnerMargin", "Inner margin"))
+					]
+					+ SGridPanel::Slot(1, 2)
+					.Padding(GridCellPadding)
+					.VAlign(VAlign_Center)
+					[
+						SDebugWidgetUtils::CreateConsoleVariableSpinBox(TEXT("GameplayCameras.Debug.InnerMargin"))
+					]
+					+ SGridPanel::Slot(0, 3)
+					.Padding(GridCellPadding)
+					.VAlign(VAlign_Center)
+					[
+						SNew(STextBlock)
+							.Text(LOCTEXT("IndentSize", "Indent size"))
+					]
+					+ SGridPanel::Slot(1, 3)
+					.Padding(GridCellPadding)
+					.VAlign(VAlign_Center)
+					[
+						SDebugWidgetUtils::CreateConsoleVariableSpinBox(TEXT("GameplayCameras.Debug.Indent"))
+					]
+					+ SGridPanel::Slot(2, 0)
+					.Padding(GridCellPadding)
+					.VAlign(VAlign_Center)
+					[
+						SNew(STextBlock)
+							.Text(LOCTEXT("ColorScheme", "Color scheme"))
+					]
+					+ SGridPanel::Slot(3, 0)
+					.Padding(GridCellPadding)
+					.VAlign(VAlign_Center)
+					[
+						SDebugWidgetUtils::CreateConsoleVariableComboBox(TEXT("GameplayCameras.Debug.ColorScheme"), &ColorSchemeNames)
+					]
+				]
+		];
+}
+
 void SGameplayCamerasDebugger::ConstructDebugPanels()
 {
 	// Empty panel.
diff --git a/Engine/Plugins/Cameras/GameplayCameras/Source/GameplayCamerasEditor/Private/Debugger/SGameplayCamerasDebugger.h b/Engine/Plugins/Cameras/GameplayCameras/Source/GameplayCamerasEditor/Private/Debugger/SGameplayCamerasDebugger.h
--- a/Engine/Plugins/Cameras/GameplayCameras/Source/GameplayCamerasEditor/Private/Debugger/SGameplayCamerasDebugger.h
+++ b/Engine/Plugins/Cameras/GameplayCameras/Source/GameplayCamerasEditor/Private/Debugger/SGameplayCamerasDebugger.h
@@ -46,8 +46,11 @@ protected:
 	static SGameplayCamerasDebugger* FromContext(UToolMenu* InMenu);
 	TSharedRef<SWidget> ConstructMenubar();
 	TSharedRef<SWidget> ConstructToolbar(TSharedRef<FUICommandList> InCommandList);
+	TSharedRef<SWidget> ConstructGeneralOptions(TSharedRef<FUICommandList> InCommandList);
 	void ConstructDebugPanels();
 
+	void InitializeColorSchemeNames();
+
 	static bool IsDebugCategoryActive(FString InCategoryName);
 	void SetActiveDebugCategoryPanel(FString InCategoryName);
 
@@ -59,8 +62,11 @@ private:
 	FName GameplayCamerasEditorStyleName;
 
 	TSharedPtr<SBox> PanelHost;
+
 	TSharedPtr<SWidget> EmptyPanel;
 	TMap<FString, TSharedPtr<SWidget>> DebugPanels;
+
+	TArray<TSharedPtr<FString>> ColorSchemeNames;
 };
 
 }  // namespace UE::Cameras
