Cameras: fix static analysis build

diff --git a/Engine/Plugins/Cameras/GameplayCameras/Source/GameplayCamerasEditor/Private/Debugger/SGameplayCamerasDebugger.cpp b/Engine/Plugins/Cameras/GameplayCameras/Source/GameplayCamerasEditor/Private/Debugger/SGameplayCamerasDebugger.cpp
--- a/Engine/Plugins/Cameras/GameplayCameras/Source/GameplayCamerasEditor/Private/Debugger/SGameplayCamerasDebugger.cpp
+++ b/Engine/Plugins/Cameras/GameplayCameras/Source/GameplayCamerasEditor/Private/Debugger/SGameplayCamerasDebugger.cpp
@@ -150,9 +150,12 @@ void SGameplayCamerasDebugger::InitializeColorSchemeNames()
 SGameplayCamerasDebugger* SGameplayCamerasDebugger::FromContext(UToolMenu* InMenu)
 {
 	UGameplayCamerasDebuggerMenuContext* Context = InMenu->FindContext<UGameplayCamerasDebuggerMenuContext>();
-	ensure(Context);
-	TSharedPtr<SGameplayCamerasDebugger> This = Context->CamerasDebugger.Pin();
-	return This.Get();
+	if (ensure(Context))
+	{
+		TSharedPtr<SGameplayCamerasDebugger> This = Context->CamerasDebugger.Pin();
+		return This.Get();
+	}
+	return nullptr;
 }
 
 TSharedRef<SWidget> SGameplayCamerasDebugger::ConstructMenubar()
