Cameras: fix initializing camera parameter values

diff --git a/Engine/Plugins/Cameras/GameplayCameras/Source/GameplayCameras/Private/Core/CameraParameters.cpp b/Engine/Plugins/Cameras/GameplayCameras/Source/GameplayCameras/Private/Core/CameraParameters.cpp
--- a/Engine/Plugins/Cameras/GameplayCameras/Source/GameplayCameras/Private/Core/CameraParameters.cpp
+++ b/Engine/Plugins/Cameras/GameplayCameras/Source/GameplayCameras/Private/Core/CameraParameters.cpp
@@ -2,6 +2,11 @@
 
 #include "Core/CameraParameters.h"
 
+#include "Math/Rotator.h"
+#include "Math/Transform.h"
+#include "Math/Vector.h"
+#include "Math/Vector2D.h"
+#include "Math/Vector4.h"
 #include "UObject/UnrealNames.h"
 
 bool FBooleanCameraParameter::SerializeFromMismatchedTag(const FPropertyTag& Tag, FStructuredArchive::FSlot Slot)
@@ -49,6 +54,11 @@ bool FDoubleCameraParameter::SerializeFromMismatchedTag(const FPropertyTag& Tag,
 	return false;
 }
 
+FVector2fCameraParameter::FVector2fCameraParameter()
+	: Value(EForceInit::ForceInit)
+{
+}
+
 bool FVector2fCameraParameter::SerializeFromMismatchedTag(const FPropertyTag& Tag, FStructuredArchive::FSlot Slot)
 {
 	if (Tag.GetType().IsStruct(NAME_Vector2f) || Tag.GetType().IsStruct(NAME_Vector2D))
@@ -60,6 +70,11 @@ bool FVector2fCameraParameter::SerializeFromMismatchedTag(const FPropertyTag& Ta
 	return false;
 }
 
+FVector2dCameraParameter::FVector2dCameraParameter()
+	: Value(EForceInit::ForceInit)
+{
+}
+
 bool FVector2dCameraParameter::SerializeFromMismatchedTag(const FPropertyTag& Tag, FStructuredArchive::FSlot Slot)
 {
 	if (Tag.GetType().IsStruct(NAME_Vector2d))
@@ -71,6 +86,11 @@ bool FVector2dCameraParameter::SerializeFromMismatchedTag(const FPropertyTag& Ta
 	return false;
 }
 
+FVector3fCameraParameter::FVector3fCameraParameter()
+	: Value(EForceInit::ForceInit)
+{
+}
+
 bool FVector3fCameraParameter::SerializeFromMismatchedTag(const FPropertyTag& Tag, FStructuredArchive::FSlot Slot)
 {
 	if (Tag.GetType().IsStruct(NAME_Vector3f))
@@ -82,6 +102,11 @@ bool FVector3fCameraParameter::SerializeFromMismatchedTag(const FPropertyTag& Ta
 	return false;
 }
 
+FVector3dCameraParameter::FVector3dCameraParameter()
+	: Value(EForceInit::ForceInit)
+{
+}
+
 bool FVector3dCameraParameter::SerializeFromMismatchedTag(const FPropertyTag& Tag, FStructuredArchive::FSlot Slot)
 {
 	if (Tag.GetType().IsStruct(NAME_Vector3d))
@@ -93,6 +118,11 @@ bool FVector3dCameraParameter::SerializeFromMismatchedTag(const FPropertyTag& Ta
 	return false;
 }
 
+FVector4fCameraParameter::FVector4fCameraParameter()
+	: Value(EForceInit::ForceInit)
+{
+}
+
 bool FVector4fCameraParameter::SerializeFromMismatchedTag(const FPropertyTag& Tag, FStructuredArchive::FSlot Slot)
 {
 	if (Tag.GetType().IsStruct(NAME_Vector4f))
@@ -104,6 +134,11 @@ bool FVector4fCameraParameter::SerializeFromMismatchedTag(const FPropertyTag& Ta
 	return false;
 }
 
+FVector4dCameraParameter::FVector4dCameraParameter()
+	: Value(EForceInit::ForceInit)
+{
+}
+
 bool FVector4dCameraParameter::SerializeFromMismatchedTag(const FPropertyTag& Tag, FStructuredArchive::FSlot Slot)
 {
 	if (Tag.GetType().IsStruct(NAME_Vector4d))
@@ -115,6 +150,11 @@ bool FVector4dCameraParameter::SerializeFromMismatchedTag(const FPropertyTag& Ta
 	return false;
 }
 
+FRotator3fCameraParameter::FRotator3fCameraParameter()
+	: Value(EForceInit::ForceInit)
+{
+}
+
 bool FRotator3fCameraParameter::SerializeFromMismatchedTag(const FPropertyTag& Tag, FStructuredArchive::FSlot Slot)
 {
 	if (Tag.GetType().IsStruct(NAME_Rotator3f))
@@ -126,6 +166,11 @@ bool FRotator3fCameraParameter::SerializeFromMismatchedTag(const FPropertyTag& T
 	return false;
 }
 
+FRotator3dCameraParameter::FRotator3dCameraParameter()
+	: Value(EForceInit::ForceInit)
+{
+}
+
 bool FRotator3dCameraParameter::SerializeFromMismatchedTag(const FPropertyTag& Tag, FStructuredArchive::FSlot Slot)
 {
 	if (Tag.GetType().IsStruct(NAME_Rotator3d))
diff --git a/Engine/Plugins/Cameras/GameplayCameras/Source/GameplayCameras/Public/Core/CameraParameters.h b/Engine/Plugins/Cameras/GameplayCameras/Source/GameplayCameras/Public/Core/CameraParameters.h
--- a/Engine/Plugins/Cameras/GameplayCameras/Source/GameplayCameras/Public/Core/CameraParameters.h
+++ b/Engine/Plugins/Cameras/GameplayCameras/Source/GameplayCameras/Public/Core/CameraParameters.h
@@ -90,6 +90,7 @@ struct FVector2fCameraParameter
 	UPROPERTY(EditAnywhere, Category=Common)
 	TObjectPtr<UVector2fCameraVariable> Variable;
 
+	FVector2fCameraParameter();
 	bool SerializeFromMismatchedTag(const FPropertyTag& Tag, FStructuredArchive::FSlot Slot);
 };
 
@@ -107,6 +108,7 @@ struct FVector2dCameraParameter
 	UPROPERTY(EditAnywhere, Category=Common)
 	TObjectPtr<UVector2dCameraVariable> Variable;
 
+	FVector2dCameraParameter();
 	bool SerializeFromMismatchedTag(const FPropertyTag& Tag, FStructuredArchive::FSlot Slot);
 };
 
@@ -124,6 +126,7 @@ struct FVector3fCameraParameter
 	UPROPERTY(EditAnywhere, Category=Common)
 	TObjectPtr<UVector3fCameraVariable> Variable;
 
+	FVector3fCameraParameter();
 	bool SerializeFromMismatchedTag(const FPropertyTag& Tag, FStructuredArchive::FSlot Slot);
 };
 
@@ -141,6 +144,7 @@ struct FVector3dCameraParameter
 	UPROPERTY(EditAnywhere, Category=Common)
 	TObjectPtr<UVector3dCameraVariable> Variable;
 
+	FVector3dCameraParameter();
 	bool SerializeFromMismatchedTag(const FPropertyTag& Tag, FStructuredArchive::FSlot Slot);
 };
 
@@ -158,6 +162,7 @@ struct FVector4fCameraParameter
 	UPROPERTY(EditAnywhere, Category=Common)
 	TObjectPtr<UVector4fCameraVariable> Variable;
 
+	FVector4fCameraParameter();
 	bool SerializeFromMismatchedTag(const FPropertyTag& Tag, FStructuredArchive::FSlot Slot);
 };
 
@@ -175,6 +180,7 @@ struct FVector4dCameraParameter
 	UPROPERTY(EditAnywhere, Category=Common)
 	TObjectPtr<UVector4dCameraVariable> Variable;
 
+	FVector4dCameraParameter();
 	bool SerializeFromMismatchedTag(const FPropertyTag& Tag, FStructuredArchive::FSlot Slot);
 };
 
@@ -192,6 +198,7 @@ struct FRotator3fCameraParameter
 	UPROPERTY(EditAnywhere, Category=Common)
 	TObjectPtr<URotator3fCameraVariable> Variable;
 
+	FRotator3fCameraParameter();
 	bool SerializeFromMismatchedTag(const FPropertyTag& Tag, FStructuredArchive::FSlot Slot);
 };
 
@@ -209,6 +216,7 @@ struct FRotator3dCameraParameter
 	UPROPERTY(EditAnywhere, Category=Common)
 	TObjectPtr<URotator3dCameraVariable> Variable;
 
+	FRotator3dCameraParameter();
 	bool SerializeFromMismatchedTag(const FPropertyTag& Tag, FStructuredArchive::FSlot Slot);
 };
 
