Cameras: fix incorrect usage of shared pointers, and correctly tear down the camera system on end play

The blend stack node evaluator was not clearing its live-edit callbacks when destroyed, so if the user stopped and restarted PIE, any live-edit builds would also notify the old instances from the previous PIE session. This included GC'ed camera nodes like blend stack root nodes.

However, the blend stack wasn't even destroyed anyway because the camera system component wasn't resetting its shared pointer on destruction. It would have been destroyed eventually when the component gets GC'ed, but many things inside the camera system were holding shared pointers to the system itself. This created a cycle of references that kept the system alive forever, so this change also fixes that.

diff --git a/Engine/Plugins/Cameras/GameplayCameras/Source/GameplayCameras/Private/Core/BlendStackCameraNode.cpp b/Engine/Plugins/Cameras/GameplayCameras/Source/GameplayCameras/Private/Core/BlendStackCameraNode.cpp
--- a/Engine/Plugins/Cameras/GameplayCameras/Source/GameplayCameras/Private/Core/BlendStackCameraNode.cpp
+++ b/Engine/Plugins/Cameras/GameplayCameras/Source/GameplayCameras/Private/Core/BlendStackCameraNode.cpp
@@ -46,6 +46,12 @@ static FAutoConsoleVariableRef CVarGameplayCamerasDebugBlendStackShowVariableIds
 
 UE_DEFINE_CAMERA_NODE_EVALUATOR(FBlendStackCameraNodeEvaluator)
 
+FBlendStackCameraNodeEvaluator::~FBlendStackCameraNodeEvaluator()
+{
+	// Pop all our entries to unregister the live-edit callbacks.
+	PopEntries(Entries.Num());
+}
+
 void FBlendStackCameraNodeEvaluator::Push(const FBlendStackCameraPushParams& Params)
 {
 	if (!Entries.IsEmpty())
@@ -68,8 +74,7 @@ void FBlendStackCameraNodeEvaluator::Push(const FBlendStackCameraPushParams& Par
 	UObject* Outer = const_cast<UObject*>((UObject*)GetCameraNode());
 	UBlendStackRootCameraNode* EntryRootNode = NewObject<UBlendStackRootCameraNode>(Outer, NAME_None);
 	{
-		UCameraNode* ModeRootNode = Params.CameraRig->RootNode;
-		EntryRootNode->RootNode = ModeRootNode;
+		EntryRootNode->RootNode = Params.CameraRig->RootNode;
 
 		// Find a transition and use its blend. If no transition is found,
 		// make a camera cut transition.
@@ -234,38 +239,48 @@ void FBlendStackCameraNodeEvaluator::OnRun(const FCameraNodeEvaluationParams& Pa
 	// Pop out camera rigs that have been blended out.
 	if (BlendStackNode->bAutoPop && PopEntriesBelow != INDEX_NONE)
 	{
+		PopEntries(PopEntriesBelow);
+	}
+
+	// Reset first frame flags.
+	for (FCameraRigEntry& Entry : Entries)
+	{
+		Entry.bIsFirstFrame = false;
+	}
+}
+
+void FBlendStackCameraNodeEvaluator::PopEntries(int32 FirstIndexToKeep)
+{
+	if (UNLIKELY(Entries.IsEmpty()))
+	{
+		return;
+	}
+
 #if WITH_EDITOR
-		IGameplayCamerasModule& GameplayCamerasModule = FModuleManager::GetModuleChecked<IGameplayCamerasModule>("GameplayCameras");
-		TSharedPtr<IGameplayCamerasLiveEditManager> LiveEditManager = GameplayCamerasModule.GetLiveEditManager();
+	IGameplayCamerasModule& GameplayCamerasModule = FModuleManager::GetModuleChecked<IGameplayCamerasModule>("GameplayCameras");
+	TSharedPtr<IGameplayCamerasLiveEditManager> LiveEditManager = GameplayCamerasModule.GetLiveEditManager();
 #endif  // WITH_EDITOR
 
-		for (int32 Index = 0; Index < PopEntriesBelow; ++Index)
-		{
+	for (int32 Index = 0; Index < FirstIndexToKeep; ++Index)
+	{
 #if WITH_EDITOR
-			const FCameraRigEntry& FirstEntry = Entries[0];
-			for (const UPackage* ListenPackage : FirstEntry.ListenedPackages)
+		const FCameraRigEntry& FirstEntry = Entries[0];
+		for (const UPackage* ListenPackage : FirstEntry.ListenedPackages)
+		{
+			int32* NumListens = AllListenedPackages.Find(ListenPackage);
+			if (ensure(NumListens))
 			{
-				int32* NumListens = AllListenedPackages.Find(ListenPackage);
-				if (ensure(NumListens))
+				--(*NumListens);
+				if (*NumListens == 0)
 				{
-					--(*NumListens);
-					if (*NumListens == 0)
-					{
-						LiveEditManager->RemoveListener(ListenPackage, this);
-						AllListenedPackages.Remove(ListenPackage);
-					}
+					LiveEditManager->RemoveListener(ListenPackage, this);
+					AllListenedPackages.Remove(ListenPackage);
 				}
 			}
-#endif  // WITH_EDITOR
-
-			Entries.RemoveAt(0);
 		}
-	}
+#endif  // WITH_EDITOR
 
-	// Reset first frame flags.
-	for (FCameraRigEntry& Entry : Entries)
-	{
-		Entry.bIsFirstFrame = false;
+		Entries.RemoveAt(0);
 	}
 }
 
diff --git a/Engine/Plugins/Cameras/GameplayCameras/Source/GameplayCameras/Private/Core/CameraEvaluationContextStack.cpp b/Engine/Plugins/Cameras/GameplayCameras/Source/GameplayCameras/Private/Core/CameraEvaluationContextStack.cpp
--- a/Engine/Plugins/Cameras/GameplayCameras/Source/GameplayCameras/Private/Core/CameraEvaluationContextStack.cpp
+++ b/Engine/Plugins/Cameras/GameplayCameras/Source/GameplayCameras/Private/Core/CameraEvaluationContextStack.cpp
@@ -79,9 +79,9 @@ void FCameraEvaluationContextStack::PopContext()
 	Entries.Pop();
 }
 
-void FCameraEvaluationContextStack::Initialize(TSharedRef<FCameraSystemEvaluator> InEvaluator)
+void FCameraEvaluationContextStack::Initialize(FCameraSystemEvaluator& InEvaluator)
 {
-	Evaluator = InEvaluator;
+	Evaluator = &InEvaluator;
 }
 
 void FCameraEvaluationContextStack::AddReferencedObjects(FReferenceCollector& Collector)
diff --git a/Engine/Plugins/Cameras/GameplayCameras/Source/GameplayCameras/Private/Core/CameraSystemEvaluator.cpp b/Engine/Plugins/Cameras/GameplayCameras/Source/GameplayCameras/Private/Core/CameraSystemEvaluator.cpp
--- a/Engine/Plugins/Cameras/GameplayCameras/Source/GameplayCameras/Private/Core/CameraSystemEvaluator.cpp
+++ b/Engine/Plugins/Cameras/GameplayCameras/Source/GameplayCameras/Private/Core/CameraSystemEvaluator.cpp
@@ -55,12 +55,10 @@ void FCameraSystemEvaluator::Initialize(const FCameraSystemEvaluatorCreateParams
 		RootNode = NewObject<UDefaultRootCameraNode>(Owner, TEXT("RootNode"));
 	}
 
-	TSharedRef<FCameraSystemEvaluator> This(SharedThis(this));
-
-	ContextStack.Initialize(This);
+	ContextStack.Initialize(*this);
 
 	FCameraNodeEvaluatorTreeBuildParams BuildParams;
-	BuildParams.Evaluator = This;
+	BuildParams.Evaluator = this;
 	BuildParams.RootCameraNode = RootNode;
 	RootEvaluator = static_cast<FRootCameraNodeEvaluator*>(RootEvaluatorStorage.BuildEvaluatorTree(BuildParams));
 }
@@ -102,8 +100,6 @@ void FCameraSystemEvaluator::Update(const FCameraSystemEvaluationUpdateParams& P
 		return;
 	}
 
-	TSharedPtr<FCameraSystemEvaluator> This(SharedThis(this));
-
 	// Run the camera director, and activate any camera rig(s) it returns to us.
 	FCameraDirectorEvaluator* ActiveDirectorEvaluator = ActiveContextInfo.Evaluator;
 	if (ActiveDirectorEvaluator)
@@ -119,7 +115,7 @@ void FCameraSystemEvaluator::Update(const FCameraSystemEvaluationUpdateParams& P
 		if (DirectorResult.ActiveCameraRigs.Num() == 1)
 		{
 			FActivateCameraRigParams CameraRigParams;
-			CameraRigParams.Evaluator = This;
+			CameraRigParams.Evaluator = this;
 			CameraRigParams.EvaluationContext = ActiveContextInfo.EvaluationContext;
 			CameraRigParams.CameraRig = DirectorResult.ActiveCameraRigs[0];
 			RootEvaluator->ActivateCameraRig(CameraRigParams);
@@ -128,7 +124,7 @@ void FCameraSystemEvaluator::Update(const FCameraSystemEvaluationUpdateParams& P
 
 	// Setup the params/result for running the root camera node.
 	FCameraNodeEvaluationParams NodeParams;
-	NodeParams.Evaluator = This;
+	NodeParams.Evaluator = this;
 	NodeParams.DeltaTime = Params.DeltaTime;
 
 	RootNodeResult.Reset();
diff --git a/Engine/Plugins/Cameras/GameplayCameras/Source/GameplayCameras/Private/GameFramework/GameplayCameraSystemComponent.cpp b/Engine/Plugins/Cameras/GameplayCameras/Source/GameplayCameras/Private/GameFramework/GameplayCameraSystemComponent.cpp
--- a/Engine/Plugins/Cameras/GameplayCameras/Source/GameplayCameras/Private/GameFramework/GameplayCameraSystemComponent.cpp
+++ b/Engine/Plugins/Cameras/GameplayCameras/Source/GameplayCameras/Private/GameFramework/GameplayCameraSystemComponent.cpp
@@ -98,6 +98,13 @@ void UGameplayCameraSystemComponent::OnComponentDestroyed(bool bDestroyingHierar
 {
 	Super::OnComponentDestroyed(bDestroyingHierarchy);
 
+#if WITH_EDITORONLY_DATA
+	if (PreviewMeshComponent)
+	{
+		PreviewMeshComponent->DestroyComponent();
+	}
+#endif  // WITH_EDITORONLY_DATA
+
 #if UE_GAMEPLAY_CAMERAS_DEBUG
 	if (DebugDrawDelegateHandle.IsValid())
 	{
@@ -106,12 +113,10 @@ void UGameplayCameraSystemComponent::OnComponentDestroyed(bool bDestroyingHierar
 	}
 #endif  // UE_GAMEPLAY_CAMERAS_DEBUG
 
-#if WITH_EDITORONLY_DATA
-	if (PreviewMeshComponent)
+	if (Evaluator.IsValid())
 	{
-		PreviewMeshComponent->DestroyComponent();
+		Evaluator.Reset();
 	}
-#endif  // WITH_EDITORONLY_DATA
 }
 
 #if WITH_EDITOR
diff --git a/Engine/Plugins/Cameras/GameplayCameras/Source/GameplayCameras/Public/Core/BlendStackCameraNode.h b/Engine/Plugins/Cameras/GameplayCameras/Source/GameplayCameras/Public/Core/BlendStackCameraNode.h
--- a/Engine/Plugins/Cameras/GameplayCameras/Source/GameplayCameras/Public/Core/BlendStackCameraNode.h
+++ b/Engine/Plugins/Cameras/GameplayCameras/Source/GameplayCameras/Public/Core/BlendStackCameraNode.h
@@ -67,7 +67,7 @@ namespace UE::Cameras
 struct FBlendStackCameraPushParams
 {
 	/** The evaluator currently running.*/
-	TSharedPtr<FCameraSystemEvaluator> Evaluator;
+	FCameraSystemEvaluator* Evaluator = nullptr;
 
 	/** The evaluation context within which a camera rig's node tree should run. */
 	TSharedPtr<const FCameraEvaluationContext> EvaluationContext;
@@ -89,6 +89,8 @@ class FBlendStackCameraNodeEvaluator
 
 public:
 
+	~FBlendStackCameraNodeEvaluator();
+
 	/** Push a new camera rig onto the blend stack. */
 	void Push(const FBlendStackCameraPushParams& Params);
 
@@ -122,6 +124,8 @@ protected:
 			const UCameraRigAsset* FromCameraRig, const UCameraAsset* FromCameraAsset, bool bFromFrozen,
 			const UCameraRigAsset* ToCameraRig, const UCameraAsset* ToCameraAsset) const;
 
+	void PopEntries(int32 FirstIndexToKeep);
+
 protected:
 
 	struct FCameraRigEntry
@@ -148,7 +152,7 @@ protected:
 	};
 
 	/** The camera system evaluator running this node. */
-	TSharedPtr<FCameraSystemEvaluator> OwningEvaluator;
+	FCameraSystemEvaluator* OwningEvaluator = nullptr;
 
 	/** Entries in the blend stack. */
 	TArray<FCameraRigEntry> Entries;
diff --git a/Engine/Plugins/Cameras/GameplayCameras/Source/GameplayCameras/Public/Core/CameraEvaluationContextStack.h b/Engine/Plugins/Cameras/GameplayCameras/Source/GameplayCameras/Public/Core/CameraEvaluationContextStack.h
--- a/Engine/Plugins/Cameras/GameplayCameras/Source/GameplayCameras/Public/Core/CameraEvaluationContextStack.h
+++ b/Engine/Plugins/Cameras/GameplayCameras/Source/GameplayCameras/Public/Core/CameraEvaluationContextStack.h
@@ -61,7 +61,7 @@ public:
 public:
 
 	// Internal API
-	void Initialize(TSharedRef<FCameraSystemEvaluator> InEvaluator);
+	void Initialize(FCameraSystemEvaluator& InEvaluator);
 	void AddReferencedObjects(FReferenceCollector& Collector);
 
 private:
@@ -78,7 +78,7 @@ private:
 	TArray<FContextEntry> Entries;
 
 	/** The owner evaluator. */
-	TSharedPtr<FCameraSystemEvaluator> Evaluator;
+	FCameraSystemEvaluator* Evaluator = nullptr;
 
 #if UE_GAMEPLAY_CAMERAS_DEBUG
 	friend class FCameraDirectorTreeDebugBlock;
diff --git a/Engine/Plugins/Cameras/GameplayCameras/Source/GameplayCameras/Public/Core/CameraNodeEvaluator.h b/Engine/Plugins/Cameras/GameplayCameras/Source/GameplayCameras/Public/Core/CameraNodeEvaluator.h
--- a/Engine/Plugins/Cameras/GameplayCameras/Source/GameplayCameras/Public/Core/CameraNodeEvaluator.h
+++ b/Engine/Plugins/Cameras/GameplayCameras/Source/GameplayCameras/Public/Core/CameraNodeEvaluator.h
@@ -32,7 +32,7 @@ struct FCameraDebugBlockBuilder;
 struct FCameraNodeEvaluationParams
 {
 	/** The evaluation running this evaluation. */
-	TSharedPtr<FCameraSystemEvaluator> Evaluator;
+	FCameraSystemEvaluator* Evaluator = nullptr;
 	/** The evaluation context (if any) responsible for this branch of the evaluation. */
 	TSharedPtr<const FCameraEvaluationContext> EvaluationContext;
 	/** The time interval for the evaluation. */
@@ -89,7 +89,7 @@ private:
 struct FCameraNodeEvaluatorInitializeParams
 {
 	/** The evaluation running this evaluation. */
-	TSharedPtr<FCameraSystemEvaluator> Evaluator;
+	FCameraSystemEvaluator* Evaluator = nullptr;
 	/** The evaluation context (if any) responsible for this branch of the evaluation. */
 	TSharedPtr<const FCameraEvaluationContext> EvaluationContext;
 };
diff --git a/Engine/Plugins/Cameras/GameplayCameras/Source/GameplayCameras/Public/Core/CameraNodeEvaluatorStorage.h b/Engine/Plugins/Cameras/GameplayCameras/Source/GameplayCameras/Public/Core/CameraNodeEvaluatorStorage.h
--- a/Engine/Plugins/Cameras/GameplayCameras/Source/GameplayCameras/Public/Core/CameraNodeEvaluatorStorage.h
+++ b/Engine/Plugins/Cameras/GameplayCameras/Source/GameplayCameras/Public/Core/CameraNodeEvaluatorStorage.h
@@ -26,7 +26,7 @@ struct FCameraNodeEvaluatorTreeBuildParams
 	/** The root node of the tree. */
 	TObjectPtr<const UCameraNode> RootCameraNode;
 	/** The evaluator running this evaluation. */
-	TSharedPtr<FCameraSystemEvaluator> Evaluator;
+	FCameraSystemEvaluator* Evaluator = nullptr;
 	/** The evaluation context (if any) responsible for this branch of the evaluation. */
 	TSharedPtr<const FCameraEvaluationContext> EvaluationContext;
 	/** An optional allocation information to optimize storage. */
diff --git a/Engine/Plugins/Cameras/GameplayCameras/Source/GameplayCameras/Public/Core/RootCameraNode.h b/Engine/Plugins/Cameras/GameplayCameras/Source/GameplayCameras/Public/Core/RootCameraNode.h
--- a/Engine/Plugins/Cameras/GameplayCameras/Source/GameplayCameras/Public/Core/RootCameraNode.h
+++ b/Engine/Plugins/Cameras/GameplayCameras/Source/GameplayCameras/Public/Core/RootCameraNode.h
@@ -47,7 +47,7 @@ class FCameraSystemEvaluator;
 struct FActivateCameraRigParams
 {
 	/** The evaluator currently running.*/
-	TSharedPtr<FCameraSystemEvaluator> Evaluator;
+	FCameraSystemEvaluator* Evaluator = nullptr;
 
 	/** The evaluation context in which the camera rig runs. */
 	TSharedPtr<const FCameraEvaluationContext> EvaluationContext;
