Cameras: fix rig asset editor toolbar icons only working once

diff --git a/Engine/Plugins/Cameras/GameplayCameras/Source/GameplayCamerasEditor/Private/Toolkits/CameraRigAssetEditorToolkit.cpp b/Engine/Plugins/Cameras/GameplayCameras/Source/GameplayCamerasEditor/Private/Toolkits/CameraRigAssetEditorToolkit.cpp
--- a/Engine/Plugins/Cameras/GameplayCameras/Source/GameplayCamerasEditor/Private/Toolkits/CameraRigAssetEditorToolkit.cpp
+++ b/Engine/Plugins/Cameras/GameplayCameras/Source/GameplayCamerasEditor/Private/Toolkits/CameraRigAssetEditorToolkit.cpp
@@ -17,6 +17,8 @@
 #include "PropertyEditorModule.h"
 #include "Widgets/Docking/SDockTab.h"
 
+#include UE_INLINE_GENERATED_CPP_BY_NAME(CameraRigAssetEditorToolkit)
+
 #define LOCTEXT_NAMESPACE "CameraRigAssetEditorToolkit"
 
 namespace UE::Cameras
@@ -105,23 +107,37 @@ void FCameraRigAssetEditorToolkit::RegisterToolbar()
 	UToolMenus* ToolMenus = UToolMenus::Get();
 	if (!ToolMenus->IsMenuRegistered(MenuName))
 	{
-		const FCameraRigAssetEditorCommands& Commands = FCameraRigAssetEditorCommands::Get();
-
 		FToolMenuOwnerScoped ToolMenuOwnerScope(this);
 
-		UToolMenu* ToolBarMenu = UToolMenus::Get()->RegisterMenu(
+		UToolMenu* ToolbarMenu = UToolMenus::Get()->RegisterMenu(
 				MenuName, ParentName, EMultiBoxType::ToolBar);
 
-		FToolMenuSection& ToolBarSection = ToolBarMenu->FindOrAddSection("Build");
-		FToolMenuEntry BuildButton = FToolMenuEntry::InitToolBarButton(Commands.Build);
-		BuildButton.Icon = TAttribute<FSlateIcon>(
-				SharedThis(this), &FCameraRigAssetEditorToolkit::GetBuildButtonIcon);
-		BuildButton.ToolTip = TAttribute<FText>(
-				SharedThis(this), &FCameraRigAssetEditorToolkit::GetBuildButtonTooltip);
-		ToolBarSection.AddEntry(BuildButton);
+		ToolbarMenu->AddDynamicSection(TEXT("Build"), FNewToolMenuDelegate::CreateLambda(
+				[](UToolMenu* InMenu)
+				{
+					const FCameraRigAssetEditorCommands& Commands = FCameraRigAssetEditorCommands::Get();
+					UCameraRigAssetEditorMenuContext* Context = InMenu->FindContext<UCameraRigAssetEditorMenuContext>();
+					FCameraRigAssetEditorToolkit* This = Context->CameraRigAssetEditorToolkit.Pin().Get();
+
+					FToolMenuSection& BuildSection = InMenu->AddSection(TEXT("Build"));
+
+					FToolMenuEntry BuildButton = FToolMenuEntry::InitToolBarButton(Commands.Build);
+					BuildButton.Icon = TAttribute<FSlateIcon>(This, &FCameraRigAssetEditorToolkit::GetBuildButtonIcon);
+					BuildButton.ToolTip = TAttribute<FText>(This, &FCameraRigAssetEditorToolkit::GetBuildButtonTooltip);
+					BuildSection.AddEntry(BuildButton);
+				}));
 	}
 }
 
+void FCameraRigAssetEditorToolkit::InitToolMenuContext(FToolMenuContext& MenuContext)
+{
+	FBaseAssetToolkit::InitToolMenuContext(MenuContext);
+
+	UCameraRigAssetEditorMenuContext* Context = NewObject<UCameraRigAssetEditorMenuContext>();
+	Context->CameraRigAssetEditorToolkit = SharedThis(this);
+	MenuContext.AddObject(Context);
+}
+
 FSlateIcon FCameraRigAssetEditorToolkit::GetBuildButtonIcon() const
 {
 	static const FName BuildStatusBackground("Blueprint.CompileStatus.Background");
@@ -130,30 +146,34 @@ FSlateIcon FCameraRigAssetEditorToolkit::GetBuildButtonIcon() const
 	static const FName BuildStatusGood("Blueprint.CompileStatus.Overlay.Good");
 	static const FName BuildStatusWarning("Blueprint.CompileStatus.Overlay.Warning");
 
+	const FName AppStyle = FAppStyle::GetAppStyleSetName();
+
+	if (!CameraRigAsset)
+	{
+		return FSlateIcon(AppStyle, BuildStatusBackground, NAME_None, BuildStatusError);
+	}
+
 	switch (CameraRigAsset->BuildStatus)
 	{
 		default:
 		case ECameraRigBuildStatus::Dirty:
-			return FSlateIcon(
-					FAppStyle::GetAppStyleSetName(), BuildStatusBackground, 
-					NAME_None, BuildStatusUnknown);
+			return FSlateIcon(AppStyle, BuildStatusBackground, NAME_None, BuildStatusUnknown);
 		case ECameraRigBuildStatus::WithErrors:
-			return FSlateIcon(
-					FAppStyle::GetAppStyleSetName(), BuildStatusBackground, 
-					NAME_None, BuildStatusError);
+			return FSlateIcon(AppStyle, BuildStatusBackground, NAME_None, BuildStatusError);
 		case ECameraRigBuildStatus::Clean:
-			return FSlateIcon(
-					FAppStyle::GetAppStyleSetName(), BuildStatusBackground, 
-					NAME_None, BuildStatusGood);
+			return FSlateIcon(AppStyle, BuildStatusBackground, NAME_None, BuildStatusGood);
 		case ECameraRigBuildStatus::CleanWithWarnings:
-			return FSlateIcon(
-					FAppStyle::GetAppStyleSetName(), BuildStatusBackground, 
-					NAME_None, BuildStatusWarning);
+			return FSlateIcon(AppStyle, BuildStatusBackground, NAME_None, BuildStatusWarning);
 	}
 }
 
 FText FCameraRigAssetEditorToolkit::GetBuildButtonTooltip() const
 {
+	if (!CameraRigAsset)
+	{
+		return LOCTEXT("BuildButtonStatusNoAsset", "No asset is open");
+	}
+
 	switch (CameraRigAsset->BuildStatus)
 	{
 		default:
diff --git a/Engine/Plugins/Cameras/GameplayCameras/Source/GameplayCamerasEditor/Public/Toolkits/CameraRigAssetEditorToolkit.h b/Engine/Plugins/Cameras/GameplayCameras/Source/GameplayCamerasEditor/Public/Toolkits/CameraRigAssetEditorToolkit.h
--- a/Engine/Plugins/Cameras/GameplayCameras/Source/GameplayCamerasEditor/Public/Toolkits/CameraRigAssetEditorToolkit.h
+++ b/Engine/Plugins/Cameras/GameplayCameras/Source/GameplayCamerasEditor/Public/Toolkits/CameraRigAssetEditorToolkit.h
@@ -7,6 +7,8 @@
 #include "Tools/BaseAssetToolkit.h"
 #include "UObject/GCObject.h"
 
+#include "CameraRigAssetEditorToolkit.generated.h"
+
 class IMessageLogListing;
 class UCameraRigAsset;
 class UCameraRigAssetEditor;
@@ -37,6 +39,7 @@ protected:
 	virtual void UnregisterTabSpawners(const TSharedRef<FTabManager>& InTabManager) override;
 	virtual void CreateWidgets() override;
 	virtual void RegisterToolbar() override;
+	virtual void InitToolMenuContext(FToolMenuContext& MenuContext) override;
 	virtual void PostInitAssetEditor() override;
 
 	// IToolkit interface
@@ -87,3 +90,13 @@ private:
 
 }  // namespace UE::Cameras
 
+UCLASS()
+class UCameraRigAssetEditorMenuContext : public UObject
+{
+	GENERATED_BODY()
+
+public:
+
+	TWeakPtr<UE::Cameras::FCameraRigAssetEditorToolkit> CameraRigAssetEditorToolkit;
+};
+
