Cameras: camera rig asset editor

This big change implements a graph editor for camera rig assets. There are actually two editors, one for the node hierarchy and one for the transitions.

A common chunk of code implements a generic-ish "object tree" graph editor, which is used by both those graph editors. Given a few pieces of configuration, this mini-library provides a graph editor where nodes are objects, and links are references between these objects inside the same package.

An IObjectTreeGraphObject interface is provided to integrate between the editor and model layers, such as saving node positions and comments.

diff --git a/Engine/Plugins/Cameras/GameplayCameras/Source/GameplayCameras/Private/Core/BlendStackCameraNode.cpp b/Engine/Plugins/Cameras/GameplayCameras/Source/GameplayCameras/Private/Core/BlendStackCameraNode.cpp
--- a/Engine/Plugins/Cameras/GameplayCameras/Source/GameplayCameras/Private/Core/BlendStackCameraNode.cpp
+++ b/Engine/Plugins/Cameras/GameplayCameras/Source/GameplayCameras/Private/Core/BlendStackCameraNode.cpp
@@ -79,7 +79,7 @@ void FBlendStackCameraNodeEvaluator::Push(const FBlendStackCameraPushParams& Par
 		// Find a transition and use its blend. If no transition is found,
 		// make a camera cut transition.
 		UBlendCameraNode* ModeBlend = nullptr;
-		if (const FCameraRigTransition* Transition = FindTransition(Params))
+		if (const UCameraRigTransition* Transition = FindTransition(Params))
 		{
 			ModeBlend = Transition->Blend;
 		}
@@ -284,7 +284,7 @@ void FBlendStackCameraNodeEvaluator::PopEntries(int32 FirstIndexToKeep)
 	}
 }
 
-const FCameraRigTransition* FBlendStackCameraNodeEvaluator::FindTransition(const FBlendStackCameraPushParams& Params) const
+const UCameraRigTransition* FBlendStackCameraNodeEvaluator::FindTransition(const FBlendStackCameraPushParams& Params) const
 {
 	const UBlendStackCameraNode* BlendStackNode = GetCameraNodeAs<UBlendStackCameraNode>();
 
@@ -298,7 +298,7 @@ const FCameraRigTransition* FBlendStackCameraNodeEvaluator::FindTransition(const
 	// appropriate.
 	if (!Entries.IsEmpty())
 	{
-		const FCameraRigTransition* TransitionToUse = nullptr;
+		const UCameraRigTransition* TransitionToUse = nullptr;
 
 		// Start by looking at exit transitions on the last active (top) camera rig.
 		const FCameraRigEntry& TopEntry = Entries.Top();
@@ -367,8 +367,8 @@ const FCameraRigTransition* FBlendStackCameraNodeEvaluator::FindTransition(const
 	return nullptr;
 }
 
-const FCameraRigTransition* FBlendStackCameraNodeEvaluator::FindTransition(
-			TArrayView<const FCameraRigTransition> Transitions, 
+const UCameraRigTransition* FBlendStackCameraNodeEvaluator::FindTransition(
+			TArrayView<const TObjectPtr<UCameraRigTransition>> Transitions, 
 			const UCameraRigAsset* FromCameraRig, const UCameraAsset* FromCameraAsset, bool bFromFrozen,
 			const UCameraRigAsset* ToCameraRig, const UCameraAsset* ToCameraAsset) const
 {
@@ -379,10 +379,10 @@ const FCameraRigTransition* FBlendStackCameraNodeEvaluator::FindTransition(
 	MatchParams.ToCameraAsset = ToCameraAsset;
 
 	// The transition should be used if all its conditions pass.
-	for (const FCameraRigTransition& Transition : Transitions)
+	for (TObjectPtr<const UCameraRigTransition> Transition : Transitions)
 	{
 		bool bConditionsPass = true;
-		for (const UCameraRigTransitionCondition* Condition : Transition.Conditions)
+		for (const UCameraRigTransitionCondition* Condition : Transition->Conditions)
 		{
 			if (!Condition->TransitionMatches(MatchParams))
 			{
@@ -393,7 +393,7 @@ const FCameraRigTransition* FBlendStackCameraNodeEvaluator::FindTransition(
 
 		if (bConditionsPass)
 		{
-			return &Transition;
+			return Transition;
 		}
 	}
 
diff --git a/Engine/Plugins/Cameras/GameplayCameras/Source/GameplayCameras/Private/Core/BlendStackRootCameraNode.h b/Engine/Plugins/Cameras/GameplayCameras/Source/GameplayCameras/Private/Core/BlendStackRootCameraNode.h
--- a/Engine/Plugins/Cameras/GameplayCameras/Source/GameplayCameras/Private/Core/BlendStackRootCameraNode.h
+++ b/Engine/Plugins/Cameras/GameplayCameras/Source/GameplayCameras/Private/Core/BlendStackRootCameraNode.h
@@ -15,7 +15,7 @@ class UCameraRigAsset;
  * This camera node wraps both the camera rig's root node, and the
  * blend node used to blend it.
  */
-UCLASS(MinimalAPI)
+UCLASS(MinimalAPI, Hidden)
 class UBlendStackRootCameraNode : public UCameraNode
 {
 	GENERATED_BODY()
diff --git a/Engine/Plugins/Cameras/GameplayCameras/Source/GameplayCameras/Private/Core/CameraNode.cpp b/Engine/Plugins/Cameras/GameplayCameras/Source/GameplayCameras/Private/Core/CameraNode.cpp
--- a/Engine/Plugins/Cameras/GameplayCameras/Source/GameplayCameras/Private/Core/CameraNode.cpp
+++ b/Engine/Plugins/Cameras/GameplayCameras/Source/GameplayCameras/Private/Core/CameraNode.cpp
@@ -24,3 +24,29 @@ FCameraNodeEvaluatorPtr UCameraNode::BuildEvaluator(FCameraNodeEvaluatorBuilder&
 	return NewEvaluator;
 }
 
+#if WITH_EDITOR
+
+void UCameraNode::GetGraphNodePosition(int32& NodePosX, int32& NodePosY) const
+{
+	NodePosX = GraphNodePosX;
+	NodePosY = GraphNodePosY;
+}
+
+void UCameraNode::OnGraphNodeMoved(int32 NodePosX, int32 NodePosY)
+{
+	GraphNodePosX = NodePosX;
+	GraphNodePosY = NodePosY;
+}
+
+const FString& UCameraNode::GetGraphNodeCommentText() const
+{
+	return GraphNodeComment;
+}
+
+void UCameraNode::OnUpdateGraphNodeCommentText(const FString& NewComment)
+{
+	GraphNodeComment = NewComment;
+}
+
+#endif  // WITH_EDITOR
+
diff --git a/Engine/Plugins/Cameras/GameplayCameras/Source/GameplayCameras/Private/Core/CameraRigAsset.cpp b/Engine/Plugins/Cameras/GameplayCameras/Source/GameplayCameras/Private/Core/CameraRigAsset.cpp
--- a/Engine/Plugins/Cameras/GameplayCameras/Source/GameplayCameras/Private/Core/CameraRigAsset.cpp
+++ b/Engine/Plugins/Cameras/GameplayCameras/Source/GameplayCameras/Private/Core/CameraRigAsset.cpp
@@ -29,4 +29,26 @@ void UCameraRigAsset::GatherPackages(FCameraRigPackages& OutPackages) const
 	}
 }
 
+void UCameraRigAsset::GetGraphNodePosition(int32& NodePosX, int32& NodePosY) const
+{
+	NodePosX = GraphNodePosX;
+	NodePosY = GraphNodePosY;
+}
+
+void UCameraRigAsset::OnGraphNodeMoved(int32 NodePosX, int32 NodePosY)
+{
+	GraphNodePosX = NodePosX;
+	GraphNodePosY = NodePosY;
+}
+
+const FString& UCameraRigAsset::GetGraphNodeCommentText() const
+{
+	return GraphNodeComment;
+}
+
+void UCameraRigAsset::OnUpdateGraphNodeCommentText(const FString& NewComment)
+{
+	GraphNodeComment = NewComment;
+}
+
 #endif  // WITH_EDITOR
diff --git a/Engine/Plugins/Cameras/GameplayCameras/Source/GameplayCameras/Private/Core/CameraRigTransition.cpp b/Engine/Plugins/Cameras/GameplayCameras/Source/GameplayCameras/Private/Core/CameraRigTransition.cpp
--- a/Engine/Plugins/Cameras/GameplayCameras/Source/GameplayCameras/Private/Core/CameraRigTransition.cpp
+++ b/Engine/Plugins/Cameras/GameplayCameras/Source/GameplayCameras/Private/Core/CameraRigTransition.cpp
@@ -9,3 +9,50 @@ bool UCameraRigTransitionCondition::TransitionMatches(const FCameraRigTransition
 	return OnTransitionMatches(Params);
 }
 
+#if WITH_EDITOR
+
+void UCameraRigTransitionCondition::GetGraphNodePosition(int32& NodePosX, int32& NodePosY) const
+{
+	NodePosX = GraphNodePosX;
+	NodePosY = GraphNodePosY;
+}
+
+void UCameraRigTransitionCondition::OnGraphNodeMoved(int32 NodePosX, int32 NodePosY)
+{
+	GraphNodePosX = NodePosX;
+	GraphNodePosY = NodePosY;
+}
+
+const FString& UCameraRigTransitionCondition::GetGraphNodeCommentText() const
+{
+	return GraphNodeComment;
+}
+
+void UCameraRigTransitionCondition::OnUpdateGraphNodeCommentText(const FString& NewComment)
+{
+	GraphNodeComment = NewComment;
+}
+
+void UCameraRigTransition::GetGraphNodePosition(int32& NodePosX, int32& NodePosY) const
+{
+	NodePosX = GraphNodePosX;
+	NodePosY = GraphNodePosY;
+}
+
+void UCameraRigTransition::OnGraphNodeMoved(int32 NodePosX, int32 NodePosY)
+{
+	GraphNodePosX = NodePosX;
+	GraphNodePosY = NodePosY;
+}
+
+const FString& UCameraRigTransition::GetGraphNodeCommentText() const
+{
+	return GraphNodeComment;
+}
+
+void UCameraRigTransition::OnUpdateGraphNodeCommentText(const FString& NewComment)
+{
+	GraphNodeComment = NewComment;
+}
+
+#endif  // WITH_EDITOR
diff --git a/Engine/Plugins/Cameras/GameplayCameras/Source/GameplayCameras/Public/Core/BlendCameraNode.h b/Engine/Plugins/Cameras/GameplayCameras/Source/GameplayCameras/Public/Core/BlendCameraNode.h
--- a/Engine/Plugins/Cameras/GameplayCameras/Source/GameplayCameras/Public/Core/BlendCameraNode.h
+++ b/Engine/Plugins/Cameras/GameplayCameras/Source/GameplayCameras/Public/Core/BlendCameraNode.h
@@ -10,7 +10,7 @@
 /**
  * Base class for blend camera nodes.
  */
-UCLASS(MinimalAPI)
+UCLASS(MinimalAPI, Hidden, meta=(ObjectTreeGraphCategory="Camera Blends"))
 class UBlendCameraNode : public UCameraNode
 {
 	GENERATED_BODY()
diff --git a/Engine/Plugins/Cameras/GameplayCameras/Source/GameplayCameras/Public/Core/BlendStackCameraNode.h b/Engine/Plugins/Cameras/GameplayCameras/Source/GameplayCameras/Public/Core/BlendStackCameraNode.h
--- a/Engine/Plugins/Cameras/GameplayCameras/Source/GameplayCameras/Public/Core/BlendStackCameraNode.h
+++ b/Engine/Plugins/Cameras/GameplayCameras/Source/GameplayCameras/Public/Core/BlendStackCameraNode.h
@@ -11,10 +11,10 @@
 
 #include "BlendStackCameraNode.generated.h"
 
-struct FCameraRigTransition;
 class UBlendStackRootCameraNode;
 class UCameraAsset;
 class UCameraRigAsset;
+class UCameraRigTransition;
 
 namespace UE::Cameras
 {
@@ -32,7 +32,7 @@ class FBlendStackCameraDebugBlock;
 /**
  * A blend stack implemented as a camera node.
  */
-UCLASS(MinimalAPI)
+UCLASS(MinimalAPI, Hidden)
 class UBlendStackCameraNode : public UCameraNode
 {
 	GENERATED_BODY()
@@ -118,9 +118,9 @@ protected:
 protected:
 
 	// Utility functions for finding an appropriate transition.
-	const FCameraRigTransition* FindTransition(const FBlendStackCameraPushParams& Params) const;
-	const FCameraRigTransition* FindTransition(
-			TArrayView<const FCameraRigTransition> Transitions, 
+	const UCameraRigTransition* FindTransition(const FBlendStackCameraPushParams& Params) const;
+	const UCameraRigTransition* FindTransition(
+			TArrayView<const TObjectPtr<UCameraRigTransition>> Transitions, 
 			const UCameraRigAsset* FromCameraRig, const UCameraAsset* FromCameraAsset, bool bFromFrozen,
 			const UCameraRigAsset* ToCameraRig, const UCameraAsset* ToCameraAsset) const;
 
diff --git a/Engine/Plugins/Cameras/GameplayCameras/Source/GameplayCameras/Public/Core/CameraAsset.h b/Engine/Plugins/Cameras/GameplayCameras/Source/GameplayCameras/Public/Core/CameraAsset.h
--- a/Engine/Plugins/Cameras/GameplayCameras/Source/GameplayCameras/Public/Core/CameraAsset.h
+++ b/Engine/Plugins/Cameras/GameplayCameras/Source/GameplayCameras/Public/Core/CameraAsset.h
@@ -26,11 +26,11 @@ public:
 
 	/** A list of default enter transitions for all the camera rigs in this asset. */
 	UPROPERTY(EditAnywhere, Category=Blending)
-	TArray<FCameraRigTransition> EnterTransitions;
+	TArray<TObjectPtr<UCameraRigTransition>> EnterTransitions;
 
 	/** A list of default exit transitions for all the camera rigs in this asset. */
 	UPROPERTY(EditAnywhere, Category=Blending)
-	TArray<FCameraRigTransition> ExitTransitions;
+	TArray<TObjectPtr<UCameraRigTransition>> ExitTransitions;
 };
 
 
diff --git a/Engine/Plugins/Cameras/GameplayCameras/Source/GameplayCameras/Public/Core/CameraNode.h b/Engine/Plugins/Cameras/GameplayCameras/Source/GameplayCameras/Public/Core/CameraNode.h
--- a/Engine/Plugins/Cameras/GameplayCameras/Source/GameplayCameras/Public/Core/CameraNode.h
+++ b/Engine/Plugins/Cameras/GameplayCameras/Source/GameplayCameras/Public/Core/CameraNode.h
@@ -5,6 +5,7 @@
 #include "Core/CameraNodeEvaluatorBuilder.h"
 #include "Core/CameraVariableTableFwd.h"
 #include "Core/ObjectChildrenView.h"
+#include "Core/ObjectTreeGraphObject.h"
 #include "CoreTypes.h"
 #include "UObject/Object.h"
 
@@ -18,8 +19,10 @@ using FCameraNodeChildrenView = UE::Cameras::TObjectChildrenView<TObjectPtr<UCam
 /**
  * The base class for a camera node.
  */
-UCLASS(Abstract, DefaultToInstanced, EditInlineNew, MinimalAPI)
-class UCameraNode : public UObject
+UCLASS(Abstract, DefaultToInstanced, EditInlineNew, MinimalAPI, meta=(ObjectTreeGraphCategory="Camera Nodes"))
+class UCameraNode 
+	: public UObject
+	, public IObjectTreeGraphObject
 {
 	GENERATED_BODY()
 
@@ -47,10 +50,34 @@ protected:
 	/** Builds the evaluator for this node. */
 	virtual FCameraNodeEvaluatorPtr OnBuildEvaluator(FCameraNodeEvaluatorBuilder& Builder) const { return nullptr; }
 
+protected:
+
+	// IObjectTreeGraphObject interface.
+#if WITH_EDITOR
+	virtual void GetGraphNodePosition(int32& NodePosX, int32& NodePosY) const override;
+	virtual void OnGraphNodeMoved(int32 NodePosX, int32 NodePosY) override;
+	virtual EObjectTreeGraphObjectSupportFlags GetSupportFlags() const override { return EObjectTreeGraphObjectSupportFlags::CommentText; }
+	virtual const FString& GetGraphNodeCommentText() const override;
+	virtual void OnUpdateGraphNodeCommentText(const FString& NewComment) override;
+#endif
+
 public:
 
 	/** Specifies whether this node is enabled. */
 	UPROPERTY(EditAnywhere, Category=Common)
 	bool bIsEnabled = true;
+
+#if WITH_EDITORONLY_DATA
+
+	UPROPERTY()
+	int32 GraphNodePosX = 0;
+
+	UPROPERTY()
+	int32 GraphNodePosY = 0;
+
+	UPROPERTY()
+	FString GraphNodeComment;
+
+#endif  // WITH_EDITORONLY_DATA
 };
 
diff --git a/Engine/Plugins/Cameras/GameplayCameras/Source/GameplayCameras/Public/Core/CameraRigAsset.h b/Engine/Plugins/Cameras/GameplayCameras/Source/GameplayCameras/Public/Core/CameraRigAsset.h
--- a/Engine/Plugins/Cameras/GameplayCameras/Source/GameplayCameras/Public/Core/CameraRigAsset.h
+++ b/Engine/Plugins/Cameras/GameplayCameras/Source/GameplayCameras/Public/Core/CameraRigAsset.h
@@ -5,6 +5,7 @@
 #include "Core/CameraNodeEvaluatorFwd.h"
 #include "Core/CameraRigTransition.h"
 #include "Core/CameraVariableTableFwd.h"
+#include "Core/ObjectTreeGraphObject.h"
 #include "CoreTypes.h"
 #include "UObject/ObjectPtr.h"
 
@@ -52,7 +53,9 @@ using FCameraRigPackages = TArray<const UPackage*, TInlineAllocator<4>>;
  * the behavior of a camera.
  */
 UCLASS(MinimalAPI)
-class UCameraRigAsset : public UObject
+class UCameraRigAsset
+	: public UObject
+	, public IObjectTreeGraphObject
 {
 	GENERATED_BODY()
 
@@ -70,11 +73,11 @@ public:
 
 	/** List of enter transitions for this camera rig. */
 	UPROPERTY(EditAnywhere, Category=Blending)
-	TArray<FCameraRigTransition> EnterTransitions;
+	TArray<TObjectPtr<UCameraRigTransition>> EnterTransitions;
 
 	/** List of exist transitions for this camera rig. */
 	UPROPERTY(EditAnywhere, Category=Blending)
-	TArray<FCameraRigTransition> ExitTransitions;
+	TArray<TObjectPtr<UCameraRigTransition>> ExitTransitions;
 
 	UPROPERTY()
 	FCameraRigAllocationInfo AllocationInfo;
@@ -83,5 +86,31 @@ public:
 
 	UPROPERTY(Transient)
 	ECameraRigBuildStatus BuildStatus = ECameraRigBuildStatus::Dirty;
+
+protected:
+
+	// IObjectTreeGraphObject interface.
+#if WITH_EDITOR
+	virtual void GetGraphNodePosition(int32& NodePosX, int32& NodePosY) const override;
+	virtual void OnGraphNodeMoved(int32 NodePosX, int32 NodePosY) override;
+	virtual EObjectTreeGraphObjectSupportFlags GetSupportFlags() const override { return EObjectTreeGraphObjectSupportFlags::CommentText; }
+	virtual const FString& GetGraphNodeCommentText() const override;
+	virtual void OnUpdateGraphNodeCommentText(const FString& NewComment) override;
+#endif
+
+private:
+
+#if WITH_EDITORONLY_DATA
+
+	UPROPERTY()
+	int32 GraphNodePosX = 0;
+
+	UPROPERTY()
+	int32 GraphNodePosY = 0;
+
+	UPROPERTY()
+	FString GraphNodeComment;
+
+#endif  // WITH_EDITORONLY_DATA
 };
 
diff --git a/Engine/Plugins/Cameras/GameplayCameras/Source/GameplayCameras/Public/Core/CameraRigTransition.h b/Engine/Plugins/Cameras/GameplayCameras/Source/GameplayCameras/Public/Core/CameraRigTransition.h
--- a/Engine/Plugins/Cameras/GameplayCameras/Source/GameplayCameras/Public/Core/CameraRigTransition.h
+++ b/Engine/Plugins/Cameras/GameplayCameras/Source/GameplayCameras/Public/Core/CameraRigTransition.h
@@ -3,6 +3,7 @@
 #pragma once
 
 #include "Containers/Array.h"
+#include "Core/ObjectTreeGraphObject.h"
 #include "CoreTypes.h"
 #include "UObject/ObjectPtr.h"
 
@@ -31,8 +32,10 @@ struct FCameraRigTransitionConditionMatchParams
 /**
  * Base class for a camera transition condition.
  */
-UCLASS(Abstract, DefaultToInstanced, MinimalAPI)
-class UCameraRigTransitionCondition : public UObject
+UCLASS(Abstract, DefaultToInstanced, MinimalAPI, meta=(ObjectTreeGraphCategory="Transition Conditions"))
+class UCameraRigTransitionCondition 
+	: public UObject
+	, public IObjectTreeGraphObject
 {
 	GENERATED_BODY()
 
@@ -45,16 +48,46 @@ protected:
 
 	/** Evaluates whether this transition should be used. */
 	virtual bool OnTransitionMatches(const FCameraRigTransitionConditionMatchParams& Params) const { return false; }
+
+protected:
+
+	// IObjectTreeGraphObject interface.
+#if WITH_EDITOR
+	virtual void GetGraphNodePosition(int32& NodePosX, int32& NodePosY) const override;
+	virtual void OnGraphNodeMoved(int32 NodePosX, int32 NodePosY) override;
+	virtual EObjectTreeGraphObjectSupportFlags GetSupportFlags() const override { return EObjectTreeGraphObjectSupportFlags::CommentText; }
+	virtual const FString& GetGraphNodeCommentText() const override;
+	virtual void OnUpdateGraphNodeCommentText(const FString& NewComment) override;
+#endif
+
+private:
+
+#if WITH_EDITORONLY_DATA
+
+	UPROPERTY()
+		int32 GraphNodePosX = 0;
+
+	UPROPERTY()
+		int32 GraphNodePosY = 0;
+
+	UPROPERTY()
+		FString GraphNodeComment;
+
+#endif  // WITH_EDITORONLY_DATA
 };
 
 /**
  * A camera transition.
  */
-USTRUCT()
-struct FCameraRigTransition
+UCLASS(MinimalAPI)
+class UCameraRigTransition 
+	: public UObject
+	, public IObjectTreeGraphObject
 {
 	GENERATED_BODY()
 
+public:
+
 	/** The list of conditions that must pass for this transition to be used. */
 	UPROPERTY(EditAnywhere, Instanced, Category=Common)
 	TArray<TObjectPtr<UCameraRigTransitionCondition>> Conditions;
@@ -62,5 +95,31 @@ struct FCameraRigTransition
 	/** The blend to use to blend a given camera rig in or out. */
 	UPROPERTY(EditAnywhere, Instanced, Category=Common)
 	TObjectPtr<UBlendCameraNode> Blend;
+
+protected:
+
+	// IObjectTreeGraphObject interface.
+#if WITH_EDITOR
+	virtual void GetGraphNodePosition(int32& NodePosX, int32& NodePosY) const override;
+	virtual void OnGraphNodeMoved(int32 NodePosX, int32 NodePosY) override;
+	virtual EObjectTreeGraphObjectSupportFlags GetSupportFlags() const override { return EObjectTreeGraphObjectSupportFlags::CommentText; }
+	virtual const FString& GetGraphNodeCommentText() const override;
+	virtual void OnUpdateGraphNodeCommentText(const FString& NewComment) override;
+#endif
+
+private:
+
+#if WITH_EDITORONLY_DATA
+
+	UPROPERTY()
+		int32 GraphNodePosX = 0;
+
+	UPROPERTY()
+		int32 GraphNodePosY = 0;
+
+	UPROPERTY()
+		FString GraphNodeComment;
+
+#endif  // WITH_EDITORONLY_DATA
 };
 
diff --git a/Engine/Plugins/Cameras/GameplayCameras/Source/GameplayCameras/Public/Core/DefaultRootCameraNode.h b/Engine/Plugins/Cameras/GameplayCameras/Source/GameplayCameras/Public/Core/DefaultRootCameraNode.h
--- a/Engine/Plugins/Cameras/GameplayCameras/Source/GameplayCameras/Public/Core/DefaultRootCameraNode.h
+++ b/Engine/Plugins/Cameras/GameplayCameras/Source/GameplayCameras/Public/Core/DefaultRootCameraNode.h
@@ -11,7 +11,7 @@ class UBlendStackCameraNode;
 /**
  * The default implementation of a root camera node.
  */
-UCLASS(MinimalAPI)
+UCLASS(MinimalAPI, Hidden)
 class UDefaultRootCameraNode : public URootCameraNode
 {
 	GENERATED_BODY()
diff --git a/Engine/Plugins/Cameras/GameplayCameras/Source/GameplayCameras/Public/Core/ObjectTreeGraphObject.h b/Engine/Plugins/Cameras/GameplayCameras/Source/GameplayCameras/Public/Core/ObjectTreeGraphObject.h
new file mode 100644
--- /dev/null
+++ b/Engine/Plugins/Cameras/GameplayCameras/Source/GameplayCameras/Public/Core/ObjectTreeGraphObject.h
@@ -0,0 +1,52 @@
+// Copyright Epic Games, Inc. All Rights Reserved.
+
+#pragma once
+
+#include "CoreMinimal.h"
+#include "Misc/EnumClassFlags.h"
+#include "UObject/Interface.h"
+
+#include "ObjectTreeGraphObject.generated.h"
+
+UINTERFACE(MinimalAPI)
+class UObjectTreeGraphObject : public UInterface
+{
+	GENERATED_BODY()
+};
+
+#if WITH_EDITOR
+
+enum class EObjectTreeGraphObjectSupportFlags
+{
+	None,
+	CommentText = 1 << 0,
+	CustomRename = 1 << 1
+};
+ENUM_CLASS_FLAGS(EObjectTreeGraphObjectSupportFlags);
+
+#endif  // WITH_EDITOR
+
+
+class IObjectTreeGraphObject
+{
+	GENERATED_BODY()
+
+#if WITH_EDITOR
+
+public:
+
+	virtual void GetGraphNodePosition(int32& NodePosX, int32& NodePosY) const {}
+	virtual void OnGraphNodeMoved(int32 NodePosX, int32 NodePosY) {}
+
+	virtual EObjectTreeGraphObjectSupportFlags GetSupportFlags() const { return EObjectTreeGraphObjectSupportFlags::None; }
+	bool HasSupportFlags(EObjectTreeGraphObjectSupportFlags InFlags) const { return EnumHasAllFlags(GetSupportFlags(), InFlags); }
+
+	virtual const FString& GetGraphNodeCommentText() const { static const FString EmptyString; return EmptyString; }
+	virtual void OnUpdateGraphNodeCommentText(const FString& NewComment) {}
+
+	virtual const FString& GetGraphNodeName() const { static const FString EmptyString; return EmptyString; }
+	virtual void OnRenameGraphNode(const FString& NewName) {}
+
+#endif  // WITH_EDITOR
+};
+
diff --git a/Engine/Plugins/Cameras/GameplayCameras/Source/GameplayCamerasEditor/GameplayCamerasEditor.Build.cs b/Engine/Plugins/Cameras/GameplayCameras/Source/GameplayCamerasEditor/GameplayCamerasEditor.Build.cs
--- a/Engine/Plugins/Cameras/GameplayCameras/Source/GameplayCamerasEditor/GameplayCamerasEditor.Build.cs
+++ b/Engine/Plugins/Cameras/GameplayCameras/Source/GameplayCamerasEditor/GameplayCamerasEditor.Build.cs
@@ -19,6 +19,7 @@ public class GameplayCamerasEditor : ModuleRules
 		PrivateDependencyModuleNames.AddRange(
 			new string[]
 			{
+				"ApplicationCore",
 				"AssetDefinition",
 				"AssetRegistry",
 				"BlueprintGraph",
@@ -28,6 +29,7 @@ public class GameplayCamerasEditor : ModuleRules
 				"EditorSubsystem",
 				"Engine",
 				"GameplayCameras",
+				"GraphEditor",
 				"InputCore",
 				"InteractiveToolsFramework",
 				"Kismet",
diff --git a/Engine/Plugins/Cameras/GameplayCameras/Source/GameplayCamerasEditor/Private/Commands/CameraRigAssetEditorCommands.cpp b/Engine/Plugins/Cameras/GameplayCameras/Source/GameplayCamerasEditor/Private/Commands/CameraRigAssetEditorCommands.cpp
--- a/Engine/Plugins/Cameras/GameplayCameras/Source/GameplayCamerasEditor/Private/Commands/CameraRigAssetEditorCommands.cpp
+++ b/Engine/Plugins/Cameras/GameplayCameras/Source/GameplayCamerasEditor/Private/Commands/CameraRigAssetEditorCommands.cpp
@@ -5,6 +5,8 @@
 #include "Framework/Commands/Commands.h"
 #include "Framework/Commands/InputChord.h"
 #include "Framework/Commands/UICommandInfo.h"
+#include "GenericPlatform/GenericApplication.h"
+#include "Styles/GameplayCamerasEditorStyle.h"
 #include "Styling/AppStyle.h"
 
 #define LOCTEXT_NAMESPACE "CameraRigAssetEditorCommands"
@@ -14,10 +16,10 @@ namespace UE::Cameras
 
 FCameraRigAssetEditorCommands::FCameraRigAssetEditorCommands()
 	: TCommands<FCameraRigAssetEditorCommands>(
-			"GameplayCameras_CameraRigAssetEditor",
-			NSLOCTEXT("Contexts", "GameplayCameras_CameraRigAssetEditor", "Camera Rig Asset Editor"),
+			"CameraRigAssetEditor",
+			LOCTEXT("CameraRigAssetEditor", "Camera Rig Asset Editor"),
 			NAME_None,
-			FAppStyle::GetAppStyleSetName()
+			FGameplayCamerasEditorStyle::Get()->GetStyleSetName()
 		)
 {
 }
@@ -26,6 +28,18 @@ void FCameraRigAssetEditorCommands::RegisterCommands()
 {
 	UI_COMMAND(Build, "Build", "Builds the asset and refreshes it in PIE",
 			EUserInterfaceActionType::Button, FInputChord(EKeys::F7));
+
+	UI_COMMAND(ShowNodeHierarchy, "Node Hierarchy", "Shows the node hierarchy of this camera rig",
+			EUserInterfaceActionType::ToggleButton, FInputChord());
+	UI_COMMAND(ShowTransitions, "Transitions", "Shows the transitions of this camera rig",
+			EUserInterfaceActionType::ToggleButton, FInputChord());
+
+	UI_COMMAND(ShowMessages, "Messages", "Shows the message log for this camera rig",
+			EUserInterfaceActionType::ToggleButton, FInputChord());
+	UI_COMMAND(FindInCameraRig, "Search", "Searches for nodes in this camera rig",
+			EUserInterfaceActionType::Button, FInputChord(EKeys::F, EModifierKey::Control));
+	UI_COMMAND(FocusHome, "Home", "Focuses the graph canvas back on the root node",
+			EUserInterfaceActionType::Button, FInputChord());
 }
 
 }  // namespace UE::Cameras
diff --git a/Engine/Plugins/Cameras/GameplayCameras/Source/GameplayCamerasEditor/Private/Editors/CameraTransitionGraphSchema.cpp b/Engine/Plugins/Cameras/GameplayCameras/Source/GameplayCamerasEditor/Private/Editors/CameraTransitionGraphSchema.cpp
new file mode 100644
--- /dev/null
+++ b/Engine/Plugins/Cameras/GameplayCameras/Source/GameplayCamerasEditor/Private/Editors/CameraTransitionGraphSchema.cpp
@@ -0,0 +1,87 @@
+// Copyright Epic Games, Inc. All Rights Reserved.
+
+#include "Editors/CameraTransitionGraphSchema.h"
+
+#include "Core/CameraRigAsset.h"
+#include "EdGraph/EdGraphPin.h"
+#include "Editors/ObjectTreeGraphNode.h"
+
+#include "ScopedTransaction.h"
+
+#include UE_INLINE_GENERATED_CPP_BY_NAME(CameraTransitionGraphSchema)
+
+#define LOCTEXT_NAMESPACE "CameraTransitionGraphSchema"
+
+void UCameraTransitionGraphSchema::GetGraphContextActions(FGraphContextMenuBuilder& ContextMenuBuilder) const
+{
+	bool bCanCreateEnterTransition = true;
+	bool bCanCreateExitTransition = false;
+	if (const UEdGraphPin* DraggedPin = ContextMenuBuilder.FromPin)
+	{
+		bCanCreateEnterTransition = false;
+		bCanCreateExitTransition = false;
+
+		UObjectTreeGraphNode* OwningNode = Cast<UObjectTreeGraphNode>(DraggedPin->GetOwningNode());
+		if (OwningNode)
+		{
+			FProperty* DraggedPinProperty = OwningNode->GetPropertyForPin(DraggedPin);
+			if (DraggedPinProperty && DraggedPinProperty->GetOwnerClass()->IsChildOf<UCameraRigAsset>())
+			{
+				bCanCreateEnterTransition = (DraggedPinProperty->GetFName() == GET_MEMBER_NAME_CHECKED(UCameraRigAsset, EnterTransitions));
+				bCanCreateExitTransition = (DraggedPinProperty->GetFName() == GET_MEMBER_NAME_CHECKED(UCameraRigAsset, ExitTransitions));
+			}
+		}
+	}
+
+	if (bCanCreateEnterTransition)
+	{
+		TSharedRef<FCameraTransitionGraphSchemaAction_NewTransitionNode> EnterAction = MakeShared<FCameraTransitionGraphSchemaAction_NewTransitionNode>(
+				LOCTEXT("TransitionsCategory", "Transitions"),
+				LOCTEXT("EnterTransition", "Enter Transition"),
+				LOCTEXT("EnterTransitionToolTip", "Creates a new enter transition"));
+		EnterAction->TransitionType = FCameraTransitionGraphSchemaAction_NewTransitionNode::ETransitionType::Enter;
+		ContextMenuBuilder.AddAction(StaticCastSharedPtr<FEdGraphSchemaAction>(EnterAction.ToSharedPtr()));
+	}
+	
+	if (bCanCreateExitTransition)
+	{
+		TSharedRef<FCameraTransitionGraphSchemaAction_NewTransitionNode> ExitAction = MakeShared<FCameraTransitionGraphSchemaAction_NewTransitionNode>(
+				LOCTEXT("TransitionsCategory", "Transitions"),
+				LOCTEXT("ExitTransition", "Exit Transition"),
+				LOCTEXT("ExitTransitionToolTip", "Creates a new exit transition"),
+				0);
+		ExitAction->TransitionType = FCameraTransitionGraphSchemaAction_NewTransitionNode::ETransitionType::Exit;
+		ContextMenuBuilder.AddAction(StaticCastSharedPtr<FEdGraphSchemaAction>(ExitAction.ToSharedPtr()));
+	}
+
+	Super::GetGraphContextActions(ContextMenuBuilder);
+}
+
+void UCameraTransitionGraphSchema::FilterGraphContextPlaceableClasses(TArray<UClass*>& InOutClasses) const
+{
+	InOutClasses.Remove(UCameraRigTransition::StaticClass());
+}
+
+FCameraTransitionGraphSchemaAction_NewTransitionNode::FCameraTransitionGraphSchemaAction_NewTransitionNode()
+{
+	ObjectClass = UCameraRigTransition::StaticClass();
+}
+
+FCameraTransitionGraphSchemaAction_NewTransitionNode::FCameraTransitionGraphSchemaAction_NewTransitionNode(FText InNodeCategory, FText InMenuDesc, FText InToolTip, const int32 InGrouping, FText InKeywords)
+	: FObjectGraphSchemaAction_NewNode(InNodeCategory, InMenuDesc, InToolTip, InGrouping, InKeywords)
+{
+	ObjectClass = UCameraRigTransition::StaticClass();
+}
+
+void FCameraTransitionGraphSchemaAction_NewTransitionNode::AutoSetupNewNode(UObjectTreeGraphNode* NewNode, UEdGraphPin* FromPin)
+{
+	if (TransitionType == ETransitionType::Enter)
+	{
+		NewNode->OverrideSelfPinDirection(EGPD_Output);
+	}
+
+	FObjectGraphSchemaAction_NewNode::AutoSetupNewNode(NewNode, FromPin);
+}
+
+#undef LOCTEXT_NAMESPACE
+
diff --git a/Engine/Plugins/Cameras/GameplayCameras/Source/GameplayCamerasEditor/Private/Editors/ObjectTreeConnectionDrawingPolicy.cpp b/Engine/Plugins/Cameras/GameplayCameras/Source/GameplayCamerasEditor/Private/Editors/ObjectTreeConnectionDrawingPolicy.cpp
new file mode 100644
--- /dev/null
+++ b/Engine/Plugins/Cameras/GameplayCameras/Source/GameplayCamerasEditor/Private/Editors/ObjectTreeConnectionDrawingPolicy.cpp
@@ -0,0 +1,12 @@
+// Copyright Epic Games, Inc. All Rights Reserved.
+
+#include "Editors/ObjectTreeConnectionDrawingPolicy.h"
+
+FObjectTreeConnectionDrawingPolicy::FObjectTreeConnectionDrawingPolicy(int32 InBackLayerID, int32 InFrontLayerID, float InZoomFactor, const FSlateRect& InClippingRect, FSlateWindowElementList& InDrawElements)
+	: FConnectionDrawingPolicy(InBackLayerID, InFrontLayerID, InZoomFactor, InClippingRect, InDrawElements)
+{
+	// Don't draw arrowheads.
+	ArrowImage = nullptr;
+	ArrowRadius = FVector2D::ZeroVector;
+}
+
diff --git a/Engine/Plugins/Cameras/GameplayCameras/Source/GameplayCamerasEditor/Private/Editors/ObjectTreeGraph.cpp b/Engine/Plugins/Cameras/GameplayCameras/Source/GameplayCamerasEditor/Private/Editors/ObjectTreeGraph.cpp
new file mode 100644
--- /dev/null
+++ b/Engine/Plugins/Cameras/GameplayCameras/Source/GameplayCamerasEditor/Private/Editors/ObjectTreeGraph.cpp
@@ -0,0 +1,152 @@
+// Copyright Epic Games, Inc. All Rights Reserved.
+
+#include "Editors/ObjectTreeGraph.h"
+
+#include "Editors/ObjectTreeGraphNode.h"
+#include "Editors/ObjectTreeGraphSchema.h"
+#include "Misc/CoreMiscDefines.h"
+#include "UObject/Package.h"
+#include "UObject/UObjectHash.h"
+
+#include UE_INLINE_GENERATED_CPP_BY_NAME(ObjectTreeGraph)
+
+UObjectTreeGraph::UObjectTreeGraph(const FObjectInitializer& ObjInit)
+	: Super(ObjInit)
+{
+	Schema = UObjectTreeGraphSchema::StaticClass();
+}
+
+void UObjectTreeGraph::Initialize(TObjectPtr<UObject> InRootObject, const FObjectTreeGraphConfig& InConfig)
+{
+	WeakRootObject = InRootObject;
+
+	Config = InConfig;
+	ensureMsgf(Config.ConnectableObjectClasses.Num() > 0, TEXT("No connectable object classes specified... no graph can be created!"));
+	if (!Config.DefaultGraphNodeClass)
+	{
+		Config.DefaultGraphNodeClass = UObjectTreeGraphNode::StaticClass();
+	}
+}
+
+const FObjectTreeGraphConfig& UObjectTreeGraph::GetConfig() const
+{
+	return Config;
+}
+
+void UObjectTreeGraph::RebuildGraph(EObjectTreeGraphBuildSource InSource)
+{
+	RemoveAllNodes();
+	CreateAllNodes(InSource);
+	NotifyGraphChanged();
+}
+
+void UObjectTreeGraph::RemoveAllNodes()
+{
+	TArray<UEdGraphNode*> NodesToRemove(Nodes);  // Copy all nodes to remove them.
+	for (UEdGraphNode* NodeToRemove : NodesToRemove)
+	{
+		RemoveNode(NodeToRemove);
+	}
+}
+
+void UObjectTreeGraph::CreateAllNodes(EObjectTreeGraphBuildSource InSource)
+{
+	const UObjectTreeGraphSchema* GraphSchema = Cast<UObjectTreeGraphSchema>(GetSchema());
+
+	UObject* RootObject = WeakRootObject.Get();
+	if (!ensure(RootObject))
+	{
+		return;
+	}
+
+	TArray<UObject*> AllObjects;
+	UPackage* Package = RootObject->GetOutermost();
+	GetObjectsWithPackage(Package, AllObjects);
+
+	ensure(AllObjects.Contains(RootObject));
+
+	FCreatedNodes CreatedNodes;
+	for (UObject* Object : AllObjects)
+	{
+		if (UObjectTreeGraphNode* GraphNode = GraphSchema->CreateObjectNode(this, Object))
+		{
+			CreatedNodes.CreatedNodes.Add(Object, GraphNode);
+		}
+	}
+
+	RootObjectNode = nullptr;
+	UObjectTreeGraphNode** CreatedRootObjectNode = CreatedNodes.CreatedNodes.Find(RootObject);
+	if (ensure(CreatedRootObjectNode))
+	{
+		RootObjectNode = *CreatedRootObjectNode;
+	}
+
+	for (TPair<UObject*, UObjectTreeGraphNode*> Pair : CreatedNodes.CreatedNodes)
+	{
+		CreateConnections(Pair.Value, CreatedNodes);
+	}
+}
+
+void UObjectTreeGraph::CreateConnections(UObjectTreeGraphNode* InGraphNode, const FCreatedNodes& InCreatedNodes)
+{
+	UObject* Object = InGraphNode->GetObject();
+	UClass* ObjectClass = Object->GetClass();
+
+	TArray<FProperty*> ConnectableProperties;
+	InGraphNode->GetAllConnectableProperties(ConnectableProperties);
+
+	for (FProperty* ConnectableProperty : ConnectableProperties)
+	{
+		if (FObjectProperty* ObjectProperty = CastField<FObjectProperty>(ConnectableProperty))
+		{
+			UEdGraphPin* Pin = InGraphNode->GetPinForProperty(ObjectProperty);
+			if (!ensure(Pin))
+			{
+				continue;
+			}
+
+			TObjectPtr<UObject> OutConnectedObject;
+			ObjectProperty->GetValue_InContainer(Object, &OutConnectedObject);
+			if (!OutConnectedObject)
+			{
+				continue;
+			}
+
+			UObjectTreeGraphNode* const* ConnectedNode = InCreatedNodes.CreatedNodes.Find(OutConnectedObject);
+			if (ensure(ConnectedNode))
+			{
+				UEdGraphPin* ConnectedPin = (*ConnectedNode)->GetSelfPin();
+				Pin->MakeLinkTo(ConnectedPin);
+			}
+		}
+		else if (FArrayProperty* ArrayProperty = CastField<FArrayProperty>(ConnectableProperty))
+		{
+			FObjectProperty* InnerProperty = CastFieldChecked<FObjectProperty>(ArrayProperty->Inner);
+			FScriptArrayHelper ArrayHelper(ArrayProperty, ArrayProperty->ContainerPtrToValuePtr<void>(Object));
+
+			const int32 ArrayNum = ArrayHelper.Num();
+			for (int32 Index = 0; Index < ArrayNum; ++Index)
+			{
+				UEdGraphPin* Pin = InGraphNode->GetPinForPropertyNewItem(ArrayProperty, true);
+				if (!ensure(Pin))
+				{
+					continue;
+				}
+
+				UObject* ConnectedObject = InnerProperty->GetObjectPropertyValue(ArrayHelper.GetRawPtr(Index));
+				if (!ConnectedObject)
+				{
+					continue;
+				}
+
+				UObjectTreeGraphNode* const* ConnectedNode = InCreatedNodes.CreatedNodes.Find(ConnectedObject);
+				if (ensure(ConnectedNode))
+				{
+					UEdGraphPin* ConnectedPin = (*ConnectedNode)->GetSelfPin();
+					Pin->MakeLinkTo(ConnectedPin);
+				}
+			}
+		}
+	}
+}
+
diff --git a/Engine/Plugins/Cameras/GameplayCameras/Source/GameplayCamerasEditor/Private/Editors/ObjectTreeGraphConfig.cpp b/Engine/Plugins/Cameras/GameplayCameras/Source/GameplayCamerasEditor/Private/Editors/ObjectTreeGraphConfig.cpp
new file mode 100644
--- /dev/null
+++ b/Engine/Plugins/Cameras/GameplayCameras/Source/GameplayCamerasEditor/Private/Editors/ObjectTreeGraphConfig.cpp
@@ -0,0 +1,175 @@
+// Copyright Epic Games, Inc. All Rights Reserved.
+
+#include "Editors/ObjectTreeGraphConfig.h"
+
+#include "Algo/AnyOf.h"
+#include "Core/ObjectTreeGraphObject.h"
+#include "UObject/Class.h"
+#include "UObject/Object.h"
+#include "UObject/UObjectIterator.h"
+
+#define LOCTEXT_NAMESPACE "ObjectTreeGraphConfig"
+
+FObjectTreeGraphClassConfig::FObjectTreeGraphClassConfig()
+	: _SelfPinFriendlyName(LOCTEXT("SelfPinFriendlyName", "self"))
+	, _SelfPinDirection(EGPD_Input)
+	, _HasSelfPin(true)
+	, _DefaultPropertyPinDirection(EGPD_Output)
+	, _NodeTitleUsesObjectName(false)
+	, _CanCreateNew(true)
+	, _CanDuplicate(true)
+	, _CanDelete(true)
+{
+}
+
+FObjectTreeGraphClassConfig& FObjectTreeGraphClassConfig::OnlyAsRoot()
+{
+	_CanCreateNew = false;
+	_CanDuplicate = false;
+	_CanDelete = false;
+	return *this;
+}
+
+FObjectTreeGraphConfig::FObjectTreeGraphConfig()
+	: DefaultGraphNodeTitleColor(FLinearColor(0.549f, 0.745f, 0.698f))
+	, DefaultGraphNodeBodyTintColor(FLinearColor::White)
+{
+}
+
+bool FObjectTreeGraphConfig::IsConnectable(UClass* InObjectClass) const
+{
+	const bool bIsConnectable = Algo::AnyOf(ConnectableObjectClasses, [InObjectClass](UClass* Item)
+			{
+				return InObjectClass->IsChildOf(Item);
+			});
+	if (!bIsConnectable)
+	{
+		return false;
+	}
+
+	const bool bIsExcluded = Algo::AnyOf(NonConnectableObjectClasses, [InObjectClass](UClass* Item)
+			{
+				return InObjectClass->IsChildOf(Item);
+			});
+	if (bIsExcluded)
+	{
+		return false;
+	}
+
+	return true;
+}
+
+void FObjectTreeGraphConfig::GetConnectableClasses(TArray<UClass*>& OutClasses, bool bPlaceableOnly)
+{
+	for (TObjectIterator<UClass> ClassIt; ClassIt; ++ClassIt)
+	{
+		if (!IsConnectable(*ClassIt))
+		{
+			continue;
+		}
+
+		if (bPlaceableOnly)
+		{
+			if (ClassIt->HasAnyClassFlags(CLASS_Abstract | CLASS_Deprecated | CLASS_NewerVersionExists | CLASS_Hidden))
+			{
+				continue;
+			}
+
+			const FObjectTreeGraphClassConfig& ClassConfig = GetObjectClassConfig(*ClassIt);
+			if (!ClassConfig.CanCreateNew())
+			{
+				continue;
+			}
+		}
+
+		OutClasses.Add(*ClassIt);
+	}
+}
+
+const FObjectTreeGraphClassConfig& FObjectTreeGraphConfig::GetObjectClassConfig(const UClass* InObjectClass) const
+{
+	static const FObjectTreeGraphClassConfig DefaultClassConfig;
+	
+	while (InObjectClass)
+	{
+		const FObjectTreeGraphClassConfig* ClassConfig = ObjectClassConfigs.Find(InObjectClass);
+		if (ClassConfig)
+		{
+			return *ClassConfig;
+		}
+
+		InObjectClass = InObjectClass->GetSuperClass();
+	}
+
+	return DefaultClassConfig;
+}
+
+FText FObjectTreeGraphConfig::GetDisplayNameText(const UObject* InObject) const
+{
+	if (InObject)
+	{
+		FText DisplayNameText;
+
+		const IObjectTreeGraphObject* GraphObject = Cast<IObjectTreeGraphObject>(InObject);
+		if (GraphObject && GraphObject->HasSupportFlags(EObjectTreeGraphObjectSupportFlags::CustomRename))
+		{
+			DisplayNameText = FText::FromString(GraphObject->GetGraphNodeName());
+		}
+		if (!DisplayNameText.IsEmpty())
+		{
+			FormatDisplayNameText(InObject, DisplayNameText);
+			return DisplayNameText;
+		}
+
+		const FObjectTreeGraphClassConfig& ClassConfig = GetObjectClassConfig(InObject->GetClass());
+		if (ClassConfig.NodeTitleUsesObjectName())
+		{
+			DisplayNameText = FText::FromString(InObject->GetName());
+			FormatDisplayNameText(InObject, DisplayNameText);
+			return DisplayNameText;
+		}
+
+		return GetDisplayNameText(InObject->GetClass());
+	}
+	return FText::GetEmpty();
+}
+
+FText FObjectTreeGraphConfig::GetDisplayNameText(const UClass* InClass) const
+{
+	if (InClass)
+	{
+		const FObjectTreeGraphClassConfig& ClassConfig = GetObjectClassConfig(InClass);
+		if (ClassConfig.OnGetObjectClassDisplayName().IsBound())
+		{
+			return ClassConfig.OnGetObjectClassDisplayName().Execute(InClass);
+		}
+
+		FText DisplayNameText = InClass->GetDisplayNameText();
+		FormatDisplayNameText(InClass, DisplayNameText);
+		return DisplayNameText;
+	}
+	return FText::GetEmpty();
+}
+
+void FObjectTreeGraphConfig::FormatDisplayNameText(const UObject* InObject, FText& InOutDisplayNameText) const
+{
+	if (StripDisplayNameSuffixes.Num() > 0)
+	{
+		FString DisplayName = InOutDisplayNameText.ToString();
+		for (const FString& StripSuffix : StripDisplayNameSuffixes)
+		{
+			if (DisplayName.RemoveFromEnd(StripSuffix))
+			{
+				DisplayName.TrimEndInline();
+				break;
+			}
+		}
+
+		InOutDisplayNameText = FText::FromString(DisplayName);
+	}
+
+	OnFormatObjectDisplayName.ExecuteIfBound(InObject, InOutDisplayNameText);
+}
+
+#undef LOCTEXT_NAMESPACE
+
diff --git a/Engine/Plugins/Cameras/GameplayCameras/Source/GameplayCamerasEditor/Private/Editors/ObjectTreeGraphNode.cpp b/Engine/Plugins/Cameras/GameplayCameras/Source/GameplayCamerasEditor/Private/Editors/ObjectTreeGraphNode.cpp
new file mode 100644
--- /dev/null
+++ b/Engine/Plugins/Cameras/GameplayCameras/Source/GameplayCamerasEditor/Private/Editors/ObjectTreeGraphNode.cpp
@@ -0,0 +1,650 @@
+// Copyright Epic Games, Inc. All Rights Reserved.
+
+#include "Editors/ObjectTreeGraphNode.h"
+
+#include "Core/ObjectTreeGraphObject.h"
+#include "EdGraph/EdGraphPin.h"
+#include "Editors/ObjectTreeGraph.h"
+#include "Editors/ObjectTreeGraphSchema.h"
+#include "Editors/SObjectTreeGraphNode.h"
+#include "Framework/Commands/GenericCommands.h"
+#include "GraphEditorActions.h"
+#include "ScopedTransaction.h"
+#include "ToolMenus.h"
+
+#define LOCTEXT_NAMESPACE "ObjectTreeGraphNode"
+
+UObjectTreeGraphNode::UObjectTreeGraphNode(const FObjectInitializer& ObjInit)
+	: Super(ObjInit)
+{
+	bCanRenameNode = true;
+}
+
+void UObjectTreeGraphNode::Initialize(UObject* InObject)
+{
+	ensure(InObject);
+
+	Object = InObject;
+
+	IObjectTreeGraphObject* GraphObject = Cast<IObjectTreeGraphObject>(Object);
+	if (GraphObject && GraphObject->HasSupportFlags(EObjectTreeGraphObjectSupportFlags::CommentText))
+	{
+		NodeComment = GraphObject->GetGraphNodeCommentText();
+	}
+}
+
+void UObjectTreeGraphNode::PostDuplicateObject(const TMap<UEdGraphNode*, UEdGraphNode*>& NodeMap)
+{
+	if (!ensure(Object))
+	{
+		return;
+	}
+
+	// Set the correct values on our object.
+	UClass* ObjectClass = Object->GetClass();
+	for (TArray<UEdGraphPin*>::TIterator It(Pins.CreateIterator()); It; ++It)
+	{
+		UEdGraphPin* Pin = *It;
+		if (!ensure(Pin))
+		{
+			continue;
+		}
+		if (Pin->PinType.PinCategory != UObjectTreeGraphSchema::PC_Property)
+		{
+			continue;
+		}
+
+		if (Pin->PinType.PinSubCategory == UObjectTreeGraphSchema::PSC_ObjectProperty)
+		{
+			// Set the underlying object property to the duplicated object pointed to by the connected
+			// duplicated node.
+			FObjectProperty* ObjectProperty = CastFieldChecked<FObjectProperty>(ObjectClass->FindPropertyByName(Pin->GetFName()));
+			if (!Pin->LinkedTo.IsEmpty())
+			{
+				ensure(Pin->LinkedTo.Num() == 1);
+				UObjectTreeGraphNode* LinkedNode = CastChecked<UObjectTreeGraphNode>(Pin->LinkedTo[0]->GetOwningNode());
+				UObject* LinkedObject = LinkedNode->GetObject();
+
+				ObjectProperty->SetValue_InContainer(Object, TObjectPtr<UObject>(LinkedObject));
+			}
+			else
+			{
+				ObjectProperty->ClearValue_InContainer(Object);
+			}
+		}
+		else if (Pin->PinType.PinSubCategory == UObjectTreeGraphSchema::PSC_ArrayPropertyItem)
+		{
+			// Do as above, but in the appropriate place in the value array.
+			// Note that we can have unconnected pins here if we didn't duplicate some of the nodes.
+			// In that case, we need to delete the pin, and remove the null item from the value array.
+			UEdGraphPin* ParentPin = Pin->ParentPin;
+			check(ParentPin);
+			FArrayProperty* ArrayProperty = CastFieldChecked<FArrayProperty>(ObjectClass->FindPropertyByName(ParentPin->GetFName()));
+			FObjectProperty* InnerProperty = CastFieldChecked<FObjectProperty>(ArrayProperty->Inner);
+			FScriptArrayHelper ArrayHelper(ArrayProperty, ArrayProperty->ContainerPtrToValuePtr<void>(Object));
+
+			const int32 PinIndex = ParentPin->SubPins.Find(Pin);
+			check(PinIndex != INDEX_NONE);
+
+			if (!Pin->LinkedTo.IsEmpty())
+			{
+				ensure(Pin->LinkedTo.Num() == 1);
+				UObjectTreeGraphNode* LinkedNode = CastChecked<UObjectTreeGraphNode>(Pin->LinkedTo[0]->GetOwningNode());
+				UObject* LinkedObject = LinkedNode->GetObject();
+
+				InnerProperty->SetObjectPropertyValue(ArrayHelper.GetRawPtr(PinIndex), LinkedObject);
+			}
+			else
+			{
+				It.RemoveCurrent();
+				ParentPin->SubPins.Remove(Pin);
+				Pin->MarkAsGarbage();
+
+				ArrayHelper.RemoveValues(PinIndex, 1);
+			}
+		}
+	}
+
+	RefreshArrayPropertyPinNames();
+}
+
+FText UObjectTreeGraphNode::GetNodeTitle(ENodeTitleType::Type TitleType) const
+{
+	if (Object)
+	{
+		const FNodeContext NodeContext = GetNodeContext();
+		return NodeContext.GraphConfig.GetDisplayNameText(Object);
+	}
+	return FText::GetEmpty();
+}
+
+TSharedPtr<SGraphNode> UObjectTreeGraphNode::CreateVisualWidget()
+{
+	return SNew(SObjectTreeGraphNode).GraphNode(this);
+}
+
+FLinearColor UObjectTreeGraphNode::GetNodeTitleColor() const
+{
+	const FNodeContext NodeContext = GetNodeContext();
+	return NodeContext.ObjectClassConfig.NodeTitleColor().Get(NodeContext.GraphConfig.DefaultGraphNodeTitleColor);
+}
+
+FLinearColor UObjectTreeGraphNode::GetNodeBodyTintColor() const
+{
+	const FNodeContext NodeContext = GetNodeContext();
+	return NodeContext.ObjectClassConfig.NodeBodyTintColor().Get(NodeContext.GraphConfig.DefaultGraphNodeBodyTintColor);
+}
+
+void UObjectTreeGraphNode::AllocateDefaultPins()
+{
+	static const FName SelfPinName("self");
+
+	if (!ensure(Object))
+	{
+		return;
+	}
+
+	const FNodeContext NodeContext = GetNodeContext();
+	const FObjectTreeGraphConfig& OuterGraphConfig = NodeContext.GraphConfig;
+	const FObjectTreeGraphClassConfig& ObjectClassConfig = NodeContext.ObjectClassConfig;
+
+	if (ObjectClassConfig.HasSelfPin())
+	{
+		FEdGraphPinType SelfPinType;
+		SelfPinType.PinCategory = UObjectTreeGraphSchema::PC_Self;
+		UEdGraphPin* SelfPin = CreatePin(ObjectClassConfig.SelfPinDirection(), SelfPinType, SelfPinName);
+		SelfPin->PinFriendlyName = ObjectClassConfig.SelfPinFriendlyName();
+	}
+
+	for (TFieldIterator<FProperty> PropertyIt(NodeContext.ObjectClass); PropertyIt; ++PropertyIt)
+	{
+		const FName PropertyName = PropertyIt->GetFName();
+
+		const EEdGraphPinDirection PinDirection = ObjectClassConfig.GetPropertyPinDirection(PropertyName);
+
+		FEdGraphPinType ChildPinType;
+		ChildPinType.PinCategory = UObjectTreeGraphSchema::PC_Property;
+
+		if (FObjectProperty* ObjectProperty = CastField<FObjectProperty>(*PropertyIt))
+		{
+			if (!ObjectProperty->PropertyClass || !OuterGraphConfig.IsConnectable(ObjectProperty->PropertyClass))
+			{
+				continue;
+			}
+
+			ChildPinType.PinSubCategory = UObjectTreeGraphSchema::PSC_ObjectProperty;
+			UEdGraphPin* PropertyPin = CreatePin(PinDirection, ChildPinType, PropertyName);
+
+			PropertyPin->PinFriendlyName = FText::FromName(PropertyName);
+			PropertyPin->PinToolTip = ObjectProperty->PropertyClass->GetName();
+		}
+		else if (FArrayProperty* ArrayProperty = CastField<FArrayProperty>(*PropertyIt))
+		{
+			FObjectProperty* InnerProperty = CastField<FObjectProperty>(ArrayProperty->Inner);
+			if (!InnerProperty || !InnerProperty->PropertyClass || !OuterGraphConfig.IsConnectable(InnerProperty->PropertyClass))
+			{
+				continue;
+			}
+
+			ChildPinType.PinSubCategory = UObjectTreeGraphSchema::PSC_ArrayProperty;
+			ChildPinType.ContainerType = EPinContainerType::Array;
+			UEdGraphPin* ArrayPin = CreatePin(PinDirection, ChildPinType, PropertyName);
+
+			ArrayPin->PinFriendlyName = FText::FromName(PropertyName);
+			ArrayPin->PinToolTip = InnerProperty->PropertyClass->GetName();
+			ArrayPin->bHidden = true;  // Always hidden, we only ever show the sub-pins.
+
+			CreateNewItemPin(ArrayPin);
+		}
+	}
+}
+
+void UObjectTreeGraphNode::AutowireNewNode(UEdGraphPin* FromPin)
+{
+	UEdGraphPin* SelfPin = GetSelfPin();
+	if (FromPin && SelfPin)
+	{
+		const UObjectTreeGraphSchema* GraphSchema = CastChecked<UObjectTreeGraphSchema>(GetSchema());
+		GraphSchema->TryCreateConnection(FromPin, SelfPin);
+	}
+
+	Super::AutowireNewNode(FromPin);
+}
+
+void UObjectTreeGraphNode::PinConnectionListChanged(UEdGraphPin* Pin)
+{
+	Super::PinConnectionListChanged(Pin);
+}
+
+void UObjectTreeGraphNode::NodeConnectionListChanged()
+{
+	Super::NodeConnectionListChanged();
+}
+
+void UObjectTreeGraphNode::OnPinRemoved(UEdGraphPin* InRemovedPin)
+{
+	Super::OnPinRemoved(InRemovedPin);
+	RefreshArrayPropertyPinNames();
+}
+
+void UObjectTreeGraphNode::CreateNewItemPin(FArrayProperty& InArrayProperty)
+{
+	UEdGraphPin** ParentArrayPinPtr = Pins.FindByPredicate([&InArrayProperty](UEdGraphPin* Item)
+			{
+				return Item->GetFName() == InArrayProperty.GetFName();
+			});
+	if (ensure(ParentArrayPinPtr))
+	{
+		CreateNewItemPin(*ParentArrayPinPtr);
+	}
+}
+
+void UObjectTreeGraphNode::CreateNewItemPin(UEdGraphPin* InParentArrayPin)
+{
+	const FObjectTreeGraphClassConfig& ObjectClassConfig = GetObjectClassConfig();
+
+	const FName PropertyName = InParentArrayPin->GetFName();
+	const int32 NewIndex = InParentArrayPin->SubPins.Num();
+
+	FEdGraphPinType ChildPinType;
+	ChildPinType.PinCategory = UObjectTreeGraphSchema::PC_Property;
+	ChildPinType.PinSubCategory = UObjectTreeGraphSchema::PSC_ArrayPropertyItem;
+
+	const EEdGraphPinDirection PinDirection = ObjectClassConfig.GetPropertyPinDirection(PropertyName);
+
+	UEdGraphPin* ChildPin = CreatePin(PinDirection, ChildPinType, PropertyName);
+	ChildPin->PinFriendlyName = FText::Format(LOCTEXT("ArrayPinFriendlyNameFmt", "{0} {1}"), FText::FromName(PropertyName), NewIndex);
+
+	ChildPin->ParentPin = InParentArrayPin;
+	InParentArrayPin->SubPins.Add(ChildPin);
+
+	// Always re-insert the child pin so that all child pins are just after
+	// the parent array pin.
+	const int32 ParentPinIndex = Pins.Find(InParentArrayPin);
+	if (ensure(ParentPinIndex >= 0))
+	{
+		const int32 ChildPinIndex = ParentPinIndex + InParentArrayPin->SubPins.Num();
+		Pins.Pop(EAllowShrinking::No);
+		Pins.Insert(ChildPin, ChildPinIndex);
+	}
+}
+
+void UObjectTreeGraphNode::RemoveItemPin(UEdGraphPin* InItemPin)
+{
+	if (ensure(InItemPin && 
+				InItemPin->ParentPin &&
+				InItemPin->PinType.PinCategory == UObjectTreeGraphSchema::PC_Property &&
+				InItemPin->PinType.PinSubCategory == UObjectTreeGraphSchema::PSC_ArrayPropertyItem))
+	{
+		// Don't call RemovePin() because that also removes the parent pin.
+		// We just want to tremove the child pin.
+		const int32 NumPinRemoved = Pins.Remove(InItemPin);
+		ensure(NumPinRemoved == 1);
+		const int32 NumSubPinRemoved = InItemPin->ParentPin->SubPins.Remove(InItemPin);
+		ensure(NumSubPinRemoved == 1);
+
+		OnPinRemoved(InItemPin);
+
+		InItemPin->MarkAsGarbage();
+	}
+}
+
+void UObjectTreeGraphNode::RefreshArrayPropertyPinNames()
+{
+	for (UEdGraphPin* Pin : Pins)
+	{
+		if (Pin && 
+				Pin->PinType.PinCategory == UObjectTreeGraphSchema::PC_Property &&
+				Pin->PinType.PinSubCategory == UObjectTreeGraphSchema::PSC_ArrayProperty)
+		{
+			const FName PropertyName = Pin->GetFName();
+			for (int32 PinIndex = 0; PinIndex < Pin->SubPins.Num(); ++PinIndex)
+			{
+				UEdGraphPin* ChildPin = Pin->SubPins[PinIndex];
+				ChildPin->PinFriendlyName = FText::Format(LOCTEXT("ArrayPinFriendlyNameFmt", "{0} {1}"), FText::FromName(PropertyName), PinIndex);
+			}
+		}
+	}
+}
+
+void UObjectTreeGraphNode::GetAllConnectableProperties(TArray<FProperty*>& OutProperties) const
+{
+	if (!ensure(Object))
+	{
+		return;
+	}
+
+	UClass* ObjectClass = Object->GetClass();
+	for (UEdGraphPin* Pin : Pins)
+	{
+		if (Pin 
+				&& Pin->PinType.PinCategory == UObjectTreeGraphSchema::PC_Property
+				&& (Pin->PinType.PinSubCategory == UObjectTreeGraphSchema::PSC_ObjectProperty || Pin->PinType.PinSubCategory == UObjectTreeGraphSchema::PSC_ArrayProperty))
+		{
+			FProperty* Property = ObjectClass->FindPropertyByName(Pin->GetFName());
+			if (ensure(Property))
+			{
+				OutProperties.Add(Property);
+			}
+		}
+	}
+}
+
+UEdGraphPin* UObjectTreeGraphNode::GetSelfPin() const
+{
+	UEdGraphPin* const* SelfPinPtr = Pins.FindByPredicate(
+			[](UEdGraphPin* Item) { return Item->PinType.PinCategory == UObjectTreeGraphSchema::PC_Self; });
+	if (SelfPinPtr)
+	{
+		return *SelfPinPtr;
+	}
+	return nullptr;
+}
+
+void UObjectTreeGraphNode::OverrideSelfPinDirection(EEdGraphPinDirection Direction)
+{
+	bOverrideSelfPinDirection = true;
+	SelfPinDirectionOverride = Direction;
+	if (UEdGraphPin* SelfPin = GetSelfPin())
+	{
+		SelfPin->Direction = Direction;
+		GetGraph()->NotifyNodeChanged(this);
+	}
+}
+
+UEdGraphPin* UObjectTreeGraphNode::GetPinForProperty(FObjectProperty* InProperty) const
+{
+	UEdGraphPin* const* FoundItem = Pins.FindByPredicate(
+			[InProperty](UEdGraphPin* Item)
+			{ 
+				return Item->PinType.PinCategory == UObjectTreeGraphSchema::PC_Property &&
+					Item->PinType.PinSubCategory == UObjectTreeGraphSchema::PSC_ObjectProperty &&
+					Item->GetFName() == InProperty->GetFName(); 
+			});
+	if (FoundItem)
+	{
+		return *FoundItem;
+	}
+	return nullptr;
+}
+
+UEdGraphPin* UObjectTreeGraphNode::GetPinForProperty(FArrayProperty* InProperty, int32 Index) const
+{
+	UEdGraphPin* const* FoundItem = Pins.FindByPredicate(
+			[InProperty](UEdGraphPin* Item)
+			{ 
+				return Item->PinType.PinCategory == UObjectTreeGraphSchema::PC_Property &&
+					Item->PinType.PinSubCategory == UObjectTreeGraphSchema::PSC_ArrayProperty &&
+					Item->GetFName() == InProperty->GetFName(); 
+			});
+	if (FoundItem && ensure((*FoundItem)->SubPins.IsValidIndex(Index)))
+	{
+		return (*FoundItem)->SubPins[Index];
+	}
+	return nullptr;
+}
+
+UEdGraphPin* UObjectTreeGraphNode::GetPinForPropertyNewItem(FArrayProperty* InProperty, bool bCreateNew)
+{
+	UEdGraphPin* const* FoundItem = Pins.FindByPredicate(
+			[InProperty](UEdGraphPin* Item)
+			{ 
+				return Item->PinType.PinCategory == UObjectTreeGraphSchema::PC_Property &&
+					Item->PinType.PinSubCategory == UObjectTreeGraphSchema::PSC_ArrayProperty &&
+					Item->GetFName() == InProperty->GetFName(); 
+			});
+	if (FoundItem)
+	{
+		ensure((*FoundItem)->SubPins.Num() > 0 && (*FoundItem)->SubPins.Last()->LinkedTo.IsEmpty());
+		UEdGraphPin* NewItemPin = (*FoundItem)->SubPins.Last();
+		if (bCreateNew)
+		{
+			CreateNewItemPin(*FoundItem);
+		}
+		return NewItemPin;
+	}
+	return nullptr;
+}
+
+FProperty* UObjectTreeGraphNode::GetPropertyForPin(const UEdGraphPin* InPin) const
+{
+	if (!ensure(Object))
+	{
+		return nullptr;
+	}
+	if (InPin->PinType.PinCategory != UObjectTreeGraphSchema::PC_Property)
+	{
+		return nullptr;
+	}
+
+	UClass* ObjectClass = Object->GetClass();
+
+	if (InPin->PinType.PinSubCategory == UObjectTreeGraphSchema::PSC_ObjectProperty)
+	{
+		return ObjectClass->FindPropertyByName(InPin->GetFName());
+	}
+	else if (InPin->PinType.PinSubCategory == UObjectTreeGraphSchema::PSC_ArrayProperty)
+	{
+		return ObjectClass->FindPropertyByName(InPin->GetFName());
+	}
+	else if (InPin->PinType.PinSubCategory == UObjectTreeGraphSchema::PSC_ArrayPropertyItem)
+	{
+		UEdGraphPin* ParentArrayPin = InPin->ParentPin;
+		check(ParentArrayPin);
+		return ObjectClass->FindPropertyByName(ParentArrayPin->GetFName());
+	}
+
+	return nullptr;
+}
+
+UClass* UObjectTreeGraphNode::GetConnectedObjectClassForPin(const UEdGraphPin* InPin) const
+{
+	if (!ensure(Object))
+	{
+		return nullptr;
+	}
+	if (InPin->PinType.PinCategory != UObjectTreeGraphSchema::PC_Property)
+	{
+		return nullptr;
+	}
+
+	UClass* ObjectClass = Object->GetClass();
+
+	if (InPin->PinType.PinSubCategory == UObjectTreeGraphSchema::PSC_ObjectProperty)
+	{
+		FObjectProperty* ObjectProperty = CastFieldChecked<FObjectProperty>(ObjectClass->FindPropertyByName(InPin->GetFName()));
+		return ObjectProperty->PropertyClass;
+	}
+	else if (InPin->PinType.PinSubCategory == UObjectTreeGraphSchema::PSC_ArrayProperty)
+	{
+		FArrayProperty* ArrayProperty = CastFieldChecked<FArrayProperty>(ObjectClass->FindPropertyByName(InPin->GetFName()));
+		FObjectProperty* InnerProperty = CastFieldChecked<FObjectProperty>(ArrayProperty->Inner);
+		return InnerProperty->PropertyClass;
+	}
+	else if (InPin->PinType.PinSubCategory == UObjectTreeGraphSchema::PSC_ArrayPropertyItem)
+	{
+		UEdGraphPin* ParentArrayPin = InPin->ParentPin;
+		check(ParentArrayPin);
+		FArrayProperty* ArrayProperty = CastFieldChecked<FArrayProperty>(ObjectClass->FindPropertyByName(ParentArrayPin->GetFName()));
+		FObjectProperty* InnerProperty = CastFieldChecked<FObjectProperty>(ArrayProperty->Inner);
+		return InnerProperty->PropertyClass;
+	}
+
+	return nullptr;
+}
+
+int32 UObjectTreeGraphNode::GetIndexOfArrayPin(const UEdGraphPin* InPin) const
+{
+	if (!ensure(InPin &&
+				InPin->PinType.PinCategory == UObjectTreeGraphSchema::PC_Property &&
+				InPin->PinType.PinSubCategory == UObjectTreeGraphSchema::PSC_ArrayPropertyItem))
+	{
+		return INDEX_NONE;
+	}
+
+	UEdGraphPin* ParentArrayPin = InPin->ParentPin;
+	check(ParentArrayPin);
+	return ParentArrayPin->SubPins.Find(const_cast<UEdGraphPin*>(InPin));
+}
+
+void UObjectTreeGraphNode::PostPlacedNewNode()
+{
+	Super::PostPlacedNewNode();
+
+	IObjectTreeGraphObject* GraphObject = Cast<IObjectTreeGraphObject>(GetObject());
+	if (GraphObject)
+	{
+		GraphObject->GetGraphNodePosition(NodePosX, NodePosY);
+	}
+}
+
+void UObjectTreeGraphNode::GetNodeContextMenuActions(class UToolMenu* Menu, class UGraphNodeContextMenuContext* Context) const
+{
+	FToolMenuInsert MenuPosition = FToolMenuInsert(NAME_None, EToolMenuInsertType::First);
+
+	const FGraphEditorCommandsImpl& GraphEditorCommands = FGraphEditorCommands::Get();
+	const FGenericCommands& GenericCommands = FGenericCommands::Get();
+
+	// Common actions.
+	{
+		FToolMenuSection& NodeSection = Menu->AddSection(
+				"ObjectTreeGraphNodeActions", LOCTEXT("NodeActionsMenuHeader", "Node Actions"), MenuPosition);
+
+		NodeSection.AddMenuEntry(GraphEditorCommands.BreakNodeLinks);
+	}
+
+	// General actions.
+	{
+		FToolMenuSection& Section = Menu->AddSection(
+				"ObjectTreeGraphNodeGenericActions", LOCTEXT("GenericActionsMenuHeader", "General"));
+
+		Section.AddMenuEntry(GenericCommands.Delete);
+		Section.AddMenuEntry(GenericCommands.Cut);
+		Section.AddMenuEntry(GenericCommands.Copy);
+		Section.AddMenuEntry(GenericCommands.Duplicate);
+	}
+
+	// Graph organization.
+	{
+		FToolMenuSection& Section = Menu->AddSection(
+				"ObjectTreeGraphOrganizationActions", LOCTEXT("OrganizationActionsMenuHeader", "Organization"));
+
+		Section.AddSubMenu(
+				"Alignment",
+				LOCTEXT("AlignmentHeader", "Alignment"),
+				FText(),
+				FNewToolMenuDelegate::CreateLambda([&GraphEditorCommands](UToolMenu* InMenu)
+					{
+						FToolMenuSection& SubMenuSection = InMenu->AddSection(
+								"ObjectTreeGraphAlignmentActions", LOCTEXT("AlignmentHeader", "Align"));
+						SubMenuSection.AddMenuEntry(GraphEditorCommands.AlignNodesTop);
+						SubMenuSection.AddMenuEntry(GraphEditorCommands.AlignNodesMiddle);
+						SubMenuSection.AddMenuEntry(GraphEditorCommands.AlignNodesBottom);
+						SubMenuSection.AddMenuEntry(GraphEditorCommands.AlignNodesLeft);
+						SubMenuSection.AddMenuEntry(GraphEditorCommands.AlignNodesCenter);
+						SubMenuSection.AddMenuEntry(GraphEditorCommands.AlignNodesRight);
+						SubMenuSection.AddMenuEntry(GraphEditorCommands.StraightenConnections);
+					}));
+
+		Section.AddSubMenu(
+				"Distribution",
+				LOCTEXT("DistributionHeader", "Distribution"),
+				FText(),
+				FNewToolMenuDelegate::CreateLambda([&GraphEditorCommands](UToolMenu* InMenu)
+					{
+						FToolMenuSection& SubMenuSection = InMenu->AddSection(
+								"ObjectTreeGraphDistributionActions", LOCTEXT("DistributionHeader", "Distribution"));
+						SubMenuSection.AddMenuEntry(GraphEditorCommands.DistributeNodesHorizontally);
+						SubMenuSection.AddMenuEntry(GraphEditorCommands.DistributeNodesVertically);
+					}));
+	}
+}
+
+bool UObjectTreeGraphNode::GetCanRenameNode() const
+{
+	IObjectTreeGraphObject* GraphObject = Cast<IObjectTreeGraphObject>(Object);
+	return GraphObject && GraphObject->HasSupportFlags(EObjectTreeGraphObjectSupportFlags::CustomRename);
+}
+
+void UObjectTreeGraphNode::OnRenameNode(const FString& NewName)
+{
+	Super::OnRenameNode(NewName);
+
+	if (IObjectTreeGraphObject* GraphObject = Cast<IObjectTreeGraphObject>(Object))
+	{
+		const FScopedTransaction Transaction(LOCTEXT("RenameNode", "Rename Node"));
+		Object->Modify();
+		GraphObject->OnRenameGraphNode(NewName);
+	}
+}
+
+bool UObjectTreeGraphNode::CanDuplicateNode() const
+{
+	const FObjectTreeGraphClassConfig& ObjectClassConfig = GetObjectClassConfig();
+	if (!ObjectClassConfig.CanCreateNew() || !ObjectClassConfig.CanDuplicate())
+	{
+		return false;
+	}
+
+	return Super::CanDuplicateNode();
+}
+
+bool UObjectTreeGraphNode::CanUserDeleteNode() const
+{
+	const FObjectTreeGraphClassConfig& ObjectClassConfig = GetObjectClassConfig();
+	if (!ObjectClassConfig.CanDelete())
+	{
+		return false;
+	}
+
+	return Super::CanUserDeleteNode();
+}
+
+bool UObjectTreeGraphNode::SupportsCommentBubble() const
+{
+	IObjectTreeGraphObject* GraphObject = Cast<IObjectTreeGraphObject>(Object);
+	return GraphObject && GraphObject->HasSupportFlags(EObjectTreeGraphObjectSupportFlags::CommentText);
+}
+
+void UObjectTreeGraphNode::OnUpdateCommentText(const FString& NewComment)
+{
+	Super::OnUpdateCommentText(NewComment);
+
+	IObjectTreeGraphObject* GraphObject = Cast<IObjectTreeGraphObject>(Object);
+	if (GraphObject)
+	{
+		const FScopedTransaction Transaction(LOCTEXT("UpdateNodeComment", "Update Node Comment"));
+		Object->Modify();
+		GraphObject->OnUpdateGraphNodeCommentText(NewComment);
+	}
+}
+
+void UObjectTreeGraphNode::OnGraphNodeMoved()
+{
+	IObjectTreeGraphObject* GraphObject = Cast<IObjectTreeGraphObject>(Object);
+	if (GraphObject)
+	{
+		const FScopedTransaction Transaction(LOCTEXT("MoveNode", "Move Node"));
+		Object->Modify();
+		GraphObject->OnGraphNodeMoved(NodePosX, NodePosY);
+	}
+}
+
+UObjectTreeGraphNode::FNodeContext UObjectTreeGraphNode::GetNodeContext() const
+{
+	UObjectTreeGraph* OuterGraph = GetTypedOuter<UObjectTreeGraph>();
+	const FObjectTreeGraphConfig& OuterGraphConfig = OuterGraph->GetConfig();
+
+	UClass* ObjectClass = Object->GetClass();
+	const FObjectTreeGraphClassConfig& ObjectClassConfig = OuterGraphConfig.GetObjectClassConfig(ObjectClass);
+
+	return FNodeContext{ ObjectClass, OuterGraph, OuterGraphConfig, ObjectClassConfig };
+}
+
+const FObjectTreeGraphClassConfig& UObjectTreeGraphNode::GetObjectClassConfig() const
+{
+	return GetNodeContext().ObjectClassConfig;
+}
+
+#undef LOCTEXT_NAMESPACE
+
diff --git a/Engine/Plugins/Cameras/GameplayCameras/Source/GameplayCamerasEditor/Private/Editors/ObjectTreeGraphSchema.cpp b/Engine/Plugins/Cameras/GameplayCameras/Source/GameplayCamerasEditor/Private/Editors/ObjectTreeGraphSchema.cpp
new file mode 100644
--- /dev/null
+++ b/Engine/Plugins/Cameras/GameplayCameras/Source/GameplayCamerasEditor/Private/Editors/ObjectTreeGraphSchema.cpp
@@ -0,0 +1,509 @@
+// Copyright Epic Games, Inc. All Rights Reserved.
+
+#include "Editors/ObjectTreeGraphSchema.h"
+
+#include "Editors/ObjectTreeConnectionDrawingPolicy.h"
+#include "Editors/ObjectTreeGraph.h"
+#include "Editors/ObjectTreeGraphNode.h"
+#include "ScopedTransaction.h"
+#include "UObject/UObjectIterator.h"
+
+#include UE_INLINE_GENERATED_CPP_BY_NAME(ObjectTreeGraphSchema)
+
+#define LOCTEXT_NAMESPACE "ObjectTreeGraphSchema"
+
+const FName UObjectTreeGraphSchema::PC_Self("Self");
+const FName UObjectTreeGraphSchema::PC_Property("Property");
+
+const FName UObjectTreeGraphSchema::PSC_ObjectProperty("ObjectProperty");
+const FName UObjectTreeGraphSchema::PSC_ArrayProperty("ArrayProperty");
+const FName UObjectTreeGraphSchema::PSC_ArrayPropertyItem("ArrayPropertyItem");
+
+UObjectTreeGraphSchema::UObjectTreeGraphSchema(const FObjectInitializer& ObjInit)
+	: Super(ObjInit)
+{
+}
+
+UObjectTreeGraphNode* UObjectTreeGraphSchema::CreateObjectNode(UObjectTreeGraph* InGraph, UObject* InObject) const
+{
+	if (!InObject)
+	{
+		return nullptr;
+	}
+
+	if (!InGraph->GetConfig().IsConnectable(InObject->GetClass()))
+	{
+		return nullptr;
+	}
+
+	return CreateObjectNodeImpl(InGraph, InObject);
+}
+
+UObjectTreeGraphNode* UObjectTreeGraphSchema::CreateObjectNodeImpl(UObjectTreeGraph* InGraph, UObject* InObject) const
+{
+	FGraphNodeCreator<UObjectTreeGraphNode> GraphNodeCreator(*InGraph);
+	UObjectTreeGraphNode* NewNode = GraphNodeCreator.CreateNode(false);
+	NewNode->Initialize(InObject);
+	GraphNodeCreator.Finalize();
+	return NewNode;
+}
+
+void UObjectTreeGraphSchema::GetGraphContextActions(FGraphContextMenuBuilder& ContextMenuBuilder) const
+{
+	const UObjectTreeGraph* Graph = CastChecked<UObjectTreeGraph>(ContextMenuBuilder.CurrentGraph);
+	const FObjectTreeGraphConfig& GraphConfig = Graph->GetConfig();
+
+	// Find the common class restriction for all the dragged pins. We will only show actions that
+	// are compatible with them.
+	UClass* DraggedPinClass = nullptr;
+	if (const UEdGraphPin* DraggedPin = ContextMenuBuilder.FromPin)
+	{
+		UObjectTreeGraphNode* OwningNode = Cast<UObjectTreeGraphNode>(DraggedPin->GetOwningNode());
+		if (OwningNode)
+		{
+			if (DraggedPin->PinType.PinCategory == UObjectTreeGraphSchema::PC_Self)
+			{
+				DraggedPinClass = OwningNode->GetObject()->GetClass();
+			}
+			else if (DraggedPin->PinType.PinCategory == UObjectTreeGraphSchema::PC_Property)
+			{
+				DraggedPinClass = OwningNode->GetConnectedObjectClassForPin(DraggedPin);
+			}
+		}
+	}
+
+	// Find all the object classes we can create from those pins, for the given graph.
+	TArray<UClass*> PossibleObjectClasses;
+	for (TObjectIterator<UClass> ClassIt; ClassIt; ++ClassIt)
+	{
+		if (ClassIt->HasAnyClassFlags(CLASS_Abstract | CLASS_Deprecated | CLASS_NewerVersionExists))
+		{
+			continue;
+		}
+		if (ClassIt->HasAnyClassFlags(CLASS_Hidden | CLASS_NotPlaceable))
+		{
+			continue;
+		}
+
+		if (!GraphConfig.IsConnectable(*ClassIt))
+		{
+			continue;
+		}
+
+		const FObjectTreeGraphClassConfig& ClassConfig = GraphConfig.GetObjectClassConfig(*ClassIt);
+		if (!ClassConfig.CanCreateNew())
+		{
+			continue;
+		}
+
+		if (DraggedPinClass && !ClassIt->IsChildOf(DraggedPinClass))
+		{
+			continue;
+		}
+
+		PossibleObjectClasses.Add(*ClassIt);
+	}
+
+	FilterGraphContextPlaceableClasses(PossibleObjectClasses);
+
+	for (UClass* PossibleObjectClass : PossibleObjectClasses)
+	{
+		const FString* CategoryName = nullptr;
+		for (UClass* CurClass = PossibleObjectClass; CurClass; CurClass = CurClass->GetSuperClass())
+		{
+			CategoryName = CurClass->FindMetaData(TEXT("ObjectTreeGraphCategory"));
+			if (CategoryName)
+			{
+				break;
+			}
+		}
+
+		TSharedRef<FObjectGraphSchemaAction_NewNode> Action = MakeShared<FObjectGraphSchemaAction_NewNode>(
+				CategoryName ? FText::FromString(*CategoryName) : FText::GetEmpty(),
+				PossibleObjectClass->GetDisplayNameText(), 
+				PossibleObjectClass->GetDisplayNameText());
+		Action->ObjectClass = PossibleObjectClass;
+		ContextMenuBuilder.AddAction(StaticCastSharedPtr<FEdGraphSchemaAction>(Action.ToSharedPtr()));
+	}
+
+	// Don't call the base class, we want to control exactly what can be created.
+}
+
+void UObjectTreeGraphSchema::FilterGraphContextPlaceableClasses(TArray<UClass*>& InOutClasses) const
+{
+}
+
+void UObjectTreeGraphSchema::GetContextMenuActions(UToolMenu* Menu, UGraphNodeContextMenuContext* Context) const
+{
+	Super::GetContextMenuActions(Menu, Context);
+}
+
+FName UObjectTreeGraphSchema::GetParentContextMenuName() const
+{
+	// Return NAME_None if we don't want the default menu entries.
+	return Super::GetParentContextMenuName();
+}
+
+FLinearColor UObjectTreeGraphSchema::GetPinTypeColor(const FEdGraphPinType& PinType) const
+{
+	return FLinearColor::White;
+}
+
+FConnectionDrawingPolicy* UObjectTreeGraphSchema::CreateConnectionDrawingPolicy(int32 InBackLayerID, int32 InFrontLayerID, float InZoomFactor, const FSlateRect& InClippingRect, class FSlateWindowElementList& InDrawElements, UEdGraph* InGraph) const
+{
+	return new FObjectTreeConnectionDrawingPolicy(InBackLayerID, InFrontLayerID, InZoomFactor, InClippingRect, InDrawElements);
+}
+
+bool UObjectTreeGraphSchema::ShouldAlwaysPurgeOnModification() const
+{
+	return false;
+}
+
+FPinConnectionResponse UObjectTreeGraphSchema::CanCreateNewNodes(UEdGraphPin* InSourcePin) const
+{
+	return Super::CanCreateNewNodes(InSourcePin);
+}
+
+const FPinConnectionResponse UObjectTreeGraphSchema::CanCreateConnection(const UEdGraphPin* A, const UEdGraphPin* B) const
+{
+	UObjectTreeGraphNode* NodeA = Cast<UObjectTreeGraphNode>(A->GetOwningNode());
+	UObjectTreeGraphNode* NodeB = Cast<UObjectTreeGraphNode>(B->GetOwningNode());
+	if (!NodeA && !NodeB)
+	{
+		return FPinConnectionResponse(CONNECT_RESPONSE_DISALLOW, TEXT("Unsupported node types"));
+	}
+
+	if (A->Direction == B->Direction)
+	{
+		return FPinConnectionResponse(CONNECT_RESPONSE_DISALLOW, TEXT("Incompatible pins"));
+	}
+	
+	// Try to always reason back to A being the property pin, and B being the self pin of the
+	// object we want to set on the property.
+	if (A->PinType.PinCategory == PC_Self)
+	{
+		Swap(A, B);
+		Swap(NodeA, NodeB);
+	}
+
+	const bool bIsPropertyToSelf = (A->PinType.PinCategory == PC_Property && B->PinType.PinCategory == PC_Self);
+	if (!bIsPropertyToSelf)
+	{
+		return FPinConnectionResponse(CONNECT_RESPONSE_DISALLOW, TEXT("Connection must be between a property pin and a self pin"));
+	}
+
+	UObject* ObjectA = NodeA->GetObject();
+	UObject* ObjectB = NodeB->GetObject();
+	UClass* ObjectClassB = ObjectB->GetClass();
+
+	FProperty* PropertyA = NodeA->GetPropertyForPin(A);
+	if (!PropertyA)
+	{
+		return FPinConnectionResponse(CONNECT_RESPONSE_DISALLOW, TEXT("Unsupported source pin"));
+	}
+
+	if (FObjectProperty* ObjectPropertyA = CastField<FObjectProperty>(PropertyA))
+	{
+		if (ObjectClassB->IsChildOf(ObjectPropertyA->PropertyClass))
+		{
+			return FPinConnectionResponse(CONNECT_RESPONSE_BREAK_OTHERS_AB, TEXT("Compatible pin types"));
+		}
+		else
+		{
+			return FPinConnectionResponse(CONNECT_RESPONSE_DISALLOW, TEXT("Incompatible pin types"));
+		}
+	}
+	else if (FArrayProperty* ArrayPropertyA = CastField<FArrayProperty>(PropertyA))
+	{
+		FObjectProperty* InnerPropertyA = CastFieldChecked<FObjectProperty>(ArrayPropertyA->Inner);
+		if (ObjectClassB->IsChildOf(InnerPropertyA->PropertyClass))
+		{
+			return FPinConnectionResponse(CONNECT_RESPONSE_BREAK_OTHERS_AB, TEXT("Compatible array pin types"));
+		}
+		else
+		{
+			return FPinConnectionResponse(CONNECT_RESPONSE_DISALLOW, TEXT("Incompatible array pin types"));
+		}
+	}
+	else
+	{
+		return FPinConnectionResponse(CONNECT_RESPONSE_DISALLOW, TEXT("Unsupported source pin type"));
+	}
+}
+
+bool UObjectTreeGraphSchema::TryCreateConnection(UEdGraphPin* A, UEdGraphPin* B) const
+{
+	const bool bModified = UEdGraphSchema::TryCreateConnection(A, B);
+	if (!bModified)
+	{
+		return false;
+	}
+
+	const FScopedTransaction Transaction(LOCTEXT("CreateConnection", "Create Connection"));
+
+	UObjectTreeGraphNode* NodeA = Cast<UObjectTreeGraphNode>(A->GetOwningNode());
+	UObjectTreeGraphNode* NodeB = Cast<UObjectTreeGraphNode>(B->GetOwningNode());
+
+	// Try to always reason back to A being the property pin, and B being the self pin of the
+	// object we want to set on the property.
+	if (A->PinType.PinCategory == PC_Self)
+	{
+		Swap(A, B);
+		Swap(NodeA, NodeB);
+	}
+	// We know we are in the right configuration now because UEdGraphSchema::TryCreateConnection
+	// already called CanCreateConnection, which we implemented above as checking that A and B
+	// are a property/self pin pair, one way or the other.
+
+	UObject* ObjectA = NodeA->GetObject();
+	UObject* ObjectB = NodeB->GetObject();
+
+	FProperty* PropertyA = NodeA->GetPropertyForPin(A);
+
+	if (FObjectProperty* ObjectPropertyA = CastField<FObjectProperty>(PropertyA))
+	{
+		ObjectA->Modify();
+
+		ObjectPropertyA->SetValue_InContainer(ObjectA, TObjectPtr<UObject>(ObjectB));
+	}
+	else if (FArrayProperty* ArrayPropertyA = CastField<FArrayProperty>(PropertyA))
+	{
+		ObjectA->Modify();
+
+		const int32 Index = NodeA->GetIndexOfArrayPin(A);
+		ensure(Index != INDEX_NONE);
+
+		FScriptArrayHelper ArrayHelper(ArrayPropertyA, ArrayPropertyA->ContainerPtrToValuePtr<void>(ObjectA));
+		const bool bAddNewItemPin = ArrayHelper.ExpandForIndex(Index);
+
+		FObjectProperty* InnerProperty = CastFieldChecked<FObjectProperty>(ArrayPropertyA->Inner);
+		InnerProperty->SetObjectPropertyValue(ArrayHelper.GetRawPtr(Index), ObjectB);
+
+		if (bAddNewItemPin)
+		{
+			NodeA->CreateNewItemPin(*ArrayPropertyA);
+			NodeA->GetGraph()->NotifyNodeChanged(NodeA);
+		}
+	}
+
+	return true;
+}
+
+void UObjectTreeGraphSchema::BreakPinLinks(UEdGraphPin& TargetPin, bool bSendsNodeNotifcation) const
+{
+	const FScopedTransaction Transaction(LOCTEXT("BreakPinLinks", "Break Pin Links"));
+
+	if (TargetPin.LinkedTo.IsEmpty())
+	{
+		Super::BreakPinLinks(TargetPin, bSendsNodeNotifcation);
+		return;
+	}
+
+	// TargetPin could be a self pin or a property pin, we need to handle both cases and directions.
+	UEdGraphPin* PropertyPin = &TargetPin;
+	UObjectTreeGraphNode* PropertyOwningNode = Cast<UObjectTreeGraphNode>(PropertyPin->GetOwningNode());
+	if (PropertyPin->PinType.PinCategory == PC_Self)
+	{
+		PropertyPin = TargetPin.LinkedTo[0];
+		PropertyOwningNode = Cast<UObjectTreeGraphNode>(PropertyPin->GetOwningNode());
+	}
+
+	bool bRemovePropertyPin = false;
+	UObject* OwningObject = PropertyOwningNode->GetObject();
+	FProperty* Property = PropertyOwningNode->GetPropertyForPin(PropertyPin);
+
+	if (FObjectProperty* ObjectProperty = CastField<FObjectProperty>(Property))
+	{
+		OwningObject->Modify();
+
+		ObjectProperty->ClearValue_InContainer(OwningObject);
+	}
+	else if (FArrayProperty* ArrayProperty = CastField<FArrayProperty>(Property))
+	{
+		OwningObject->Modify();
+
+		int32 Index = PropertyOwningNode->GetIndexOfArrayPin(PropertyPin);
+		ensure(Index != INDEX_NONE);
+
+		FScriptArrayHelper ArrayHelper(ArrayProperty, ArrayProperty->ContainerPtrToValuePtr<void>(OwningObject));
+		ArrayHelper.RemoveValues(Index);
+
+		bRemovePropertyPin = true;
+	}
+
+	Super::BreakPinLinks(TargetPin, bSendsNodeNotifcation);
+
+	if (bRemovePropertyPin)
+	{
+		PropertyOwningNode->RemoveItemPin(PropertyPin);
+		PropertyOwningNode->GetGraph()->NotifyNodeChanged(PropertyOwningNode);
+	}
+}
+
+void UObjectTreeGraphSchema::BreakSinglePinLink(UEdGraphPin* SourcePin, UEdGraphPin* TargetPin) const
+{
+	const FScopedTransaction Transaction(LOCTEXT("BreakSinglePinLink", "Break Pin Link"));
+
+	// SourcePin could be the self-pin, and TargetPin the property pin, if the directions
+	// are reversed for that type of object/graph.
+	UEdGraphPin* PropertyPin = SourcePin;
+	UObjectTreeGraphNode* PropertyOwningNode = Cast<UObjectTreeGraphNode>(SourcePin->GetOwningNode());
+	if (SourcePin->PinType.PinCategory == PC_Self)
+	{
+		PropertyPin = SourcePin->LinkedTo[0];
+		PropertyOwningNode = Cast<UObjectTreeGraphNode>(PropertyPin->GetOwningNode());
+	}
+
+	bool bRemovePropertyPin = false;
+	UObject* OwningObject = PropertyOwningNode->GetObject();
+	FProperty* Property = PropertyOwningNode->GetPropertyForPin(PropertyPin);
+
+	if (FObjectProperty* ObjectProperty = CastField<FObjectProperty>(Property))
+	{
+		OwningObject->Modify();
+
+		ObjectProperty->ClearValue_InContainer(OwningObject);
+	}
+	else if (FArrayProperty* ArrayProperty = CastField<FArrayProperty>(Property))
+	{
+		OwningObject->Modify();
+
+		int32 Index = PropertyOwningNode->GetIndexOfArrayPin(PropertyPin);
+		ensure(Index != INDEX_NONE);
+
+		FScriptArrayHelper ArrayHelper(ArrayProperty, ArrayProperty->ContainerPtrToValuePtr<void>(OwningObject));
+		ArrayHelper.RemoveValues(Index);
+
+		bRemovePropertyPin = true;
+	}
+
+	Super::BreakSinglePinLink(SourcePin, TargetPin);
+
+	if (bRemovePropertyPin)
+	{
+		PropertyOwningNode->RemoveItemPin(PropertyPin);
+		PropertyOwningNode->GetGraph()->NotifyNodeChanged(PropertyOwningNode);
+	}
+}
+
+bool UObjectTreeGraphSchema::SupportsDropPinOnNode(UEdGraphNode* InTargetNode, const FEdGraphPinType& InSourcePinType, EEdGraphPinDirection InSourcePinDirection, FText& OutErrorMessage) const
+{
+	return Super::SupportsDropPinOnNode(InTargetNode, InSourcePinType, InSourcePinDirection, OutErrorMessage);
+}
+
+void UObjectTreeGraphSchema::ProcessDuplicatedNodes(UObjectTreeGraph* InGraph, const TMap<UEdGraphNode*, UEdGraphNode*>& NodeMap) const
+{
+	UPackage* RootObjectPackage = InGraph->GetRootObject()->GetOutermost();
+
+	// Create object duplicates for all the new nodes.
+	for (TPair<UEdGraphNode*, UEdGraphNode*> Pair : NodeMap)
+	{
+		const UObjectTreeGraphNode* OldNode = Cast<UObjectTreeGraphNode>(Pair.Key);
+		UObjectTreeGraphNode* NewNode = Cast<UObjectTreeGraphNode>(Pair.Value);
+
+		if (OldNode && NewNode)
+		{
+			ensure(OldNode->GetObject() == NewNode->GetObject());
+
+			UObject* NewNodeObject = DuplicateObject<UObject>(OldNode->GetObject(), RootObjectPackage);
+			NewNode->Initialize(NewNodeObject);
+		}
+	}
+
+	// Patch-up property references on the new nodes' objects.
+	for (TPair<UEdGraphNode*, UEdGraphNode*> Pair : NodeMap)
+	{
+		UObjectTreeGraphNode* NewNode = Cast<UObjectTreeGraphNode>(Pair.Value);
+		if (NewNode)
+		{
+			NewNode->PostDuplicateObject(NodeMap);
+		}
+	}
+
+	// Notify for everything.
+	for (TPair<UEdGraphNode*, UEdGraphNode*> Pair : NodeMap)
+	{
+		UObjectTreeGraphNode* NewNode = Cast<UObjectTreeGraphNode>(Pair.Value);
+		InGraph->NotifyNodeChanged(NewNode);
+	}
+	InGraph->NotifyGraphChanged();
+}
+
+const FObjectTreeGraphClassConfig& UObjectTreeGraphSchema::GetObjectClassConfig(const UObjectTreeGraphNode* InNode) const
+{
+	const UObjectTreeGraph* Graph = CastChecked<UObjectTreeGraph>(InNode->GetGraph());
+	return GetObjectClassConfig(Graph, InNode->GetObject()->GetClass());
+}
+
+const FObjectTreeGraphClassConfig& UObjectTreeGraphSchema::GetObjectClassConfig(const UObjectTreeGraph* InGraph, UClass* InObjectClass) const
+{
+	return InGraph->GetConfig().GetObjectClassConfig(InObjectClass);
+}
+
+FObjectGraphSchemaAction_NewNode::FObjectGraphSchemaAction_NewNode()
+{
+}
+
+FObjectGraphSchemaAction_NewNode::FObjectGraphSchemaAction_NewNode(FText InNodeCategory, FText InMenuDesc, FText InToolTip, const int32 InGrouping, FText InKeywords)
+	: FEdGraphSchemaAction(InNodeCategory, InMenuDesc, InToolTip, InGrouping, InKeywords)
+{
+}
+
+UEdGraphNode* FObjectGraphSchemaAction_NewNode::PerformAction(UEdGraph* ParentGraph, UEdGraphPin* FromPin, const FVector2D Location, bool bSelectNewNode)
+{
+	UObjectTreeGraph* ObjectTreeGraph = Cast<UObjectTreeGraph>(ParentGraph);
+	if (!ensure(ObjectTreeGraph))
+	{
+		return nullptr;
+	}
+
+	if (!ensure(ObjectClass))
+	{
+		return nullptr;
+	}
+
+	if (!ObjectOuter)
+	{
+		if (ensure(ObjectTreeGraph))
+		{
+			ObjectOuter = ObjectTreeGraph->GetRootObject();
+		}
+	}
+
+	if (!ensure(ObjectOuter))
+	{
+		return nullptr;
+	}
+
+	const FScopedTransaction Transaction(LOCTEXT("CreateNewNodeAction", "Create New Node"));
+	const UObjectTreeGraphSchema* Schema = CastChecked<UObjectTreeGraphSchema>(ParentGraph->GetSchema());
+
+	UObject* NewObject = CreateObject();
+
+	if (NewObject)
+	{
+		UObjectTreeGraphNode* NewGraphNode = Schema->CreateObjectNode(ObjectTreeGraph, NewObject);
+
+		NewGraphNode->NodePosX = Location.X;
+		NewGraphNode->NodePosY = Location.Y;
+
+		AutoSetupNewNode(NewGraphNode, FromPin);
+
+		return NewGraphNode;
+	}
+	
+	return nullptr;
+}
+
+UObject* FObjectGraphSchemaAction_NewNode::CreateObject()
+{
+	return NewObject<UObject>(ObjectOuter, ObjectClass, NAME_None, RF_Transactional);
+}
+
+void FObjectGraphSchemaAction_NewNode::AutoSetupNewNode(UObjectTreeGraphNode* NewNode, UEdGraphPin* FromPin)
+{
+	NewNode->AutowireNewNode(FromPin);
+}
+
+#undef LOCTEXT_NAMESPACE
+
diff --git a/Engine/Plugins/Cameras/GameplayCameras/Source/GameplayCamerasEditor/Private/Editors/SCameraRigAssetEditor.cpp b/Engine/Plugins/Cameras/GameplayCameras/Source/GameplayCamerasEditor/Private/Editors/SCameraRigAssetEditor.cpp
new file mode 100644
--- /dev/null
+++ b/Engine/Plugins/Cameras/GameplayCameras/Source/GameplayCamerasEditor/Private/Editors/SCameraRigAssetEditor.cpp
@@ -0,0 +1,217 @@
+// Copyright Epic Games, Inc. All Rights Reserved.
+
+#include "Editors/SCameraRigAssetEditor.h"
+
+#include "Core/BlendCameraNode.h"
+#include "Core/CameraNode.h"
+#include "Core/CameraRigAsset.h"
+#include "Core/CameraRigTransition.h"
+#include "EdGraph/EdGraphPin.h"
+#include "Editors/CameraTransitionGraphSchema.h"
+#include "Editors/ObjectTreeGraph.h"
+#include "Editors/ObjectTreeGraphNode.h"
+#include "Editors/SObjectTreeGraphEditor.h"
+#include "GameplayCamerasEditorSettings.h"
+#include "ObjectEditorUtils.h"
+#include "Widgets/Layout/SBox.h"
+
+#define LOCTEXT_NAMESPACE "SCameraRigAssetEditor"
+
+namespace UE::Cameras
+{
+
+void SCameraRigAssetEditor::Construct(const FArguments& InArgs)
+{
+	CameraRigAsset = InArgs._CameraRigAsset;
+
+	CurrentMode = ECameraRigAssetEditorMode::NodeGraph;
+
+	CreateNodeGraphEditor(InArgs);
+	CreateTransitionGraphEditor(InArgs);
+
+	ChildSlot
+	[
+		SAssignNew(BoxPanel, SBox)
+		[
+			NodeGraphEditor.ToSharedRef()
+		]
+	];
+}
+
+void SCameraRigAssetEditor::CreateNodeGraphEditor(const FArguments& InArgs)
+{
+	const UGameplayCamerasEditorSettings* Settings = GetDefault<UGameplayCamerasEditorSettings>();
+
+	FObjectTreeGraphConfig GraphConfig;
+	GraphConfig.ConnectableObjectClasses.Add(UCameraRigAsset::StaticClass());
+	GraphConfig.ConnectableObjectClasses.Add(UCameraNode::StaticClass());
+	GraphConfig.NonConnectableObjectClasses.Add(UBlendCameraNode::StaticClass());
+	GraphConfig.ObjectClassConfigs.Emplace(UCameraRigAsset::StaticClass())
+		.OnlyAsRoot()
+		.HasSelfPin(false)
+		.NodeTitleUsesObjectName(true)
+		.NodeTitleColor(Settings->CameraRigAssetTitleColor);
+
+	NodeGraph = NewObject<UObjectTreeGraph>(GetTransientPackage(), NAME_None, RF_Transactional);
+	NodeGraph->AddToRoot();
+	NodeGraph->Initialize(CameraRigAsset, GraphConfig);
+	NodeGraph->RebuildGraph(EObjectTreeGraphBuildSource::RootObjectPackage);
+
+	FGraphAppearanceInfo Appearance;
+	Appearance.CornerText = LOCTEXT("CameraRigGraphText", "CAMERA RIG");
+
+	NodeGraphEditor = SNew(SObjectTreeGraphEditor)
+		.Appearance(Appearance)
+		.DetailsView(InArgs._DetailsView)
+		.GraphTitle(this, &SCameraRigAssetEditor::GetCameraRigAssetName)
+		.GraphToEdit(NodeGraph)
+		.AssetEditorToolkit(InArgs._AssetEditorToolkit);
+}
+
+void SCameraRigAssetEditor::CreateTransitionGraphEditor(const FArguments& InArgs)
+{
+	const UGameplayCamerasEditorSettings* Settings = GetDefault<UGameplayCamerasEditorSettings>();
+
+	FObjectTreeGraphConfig GraphConfig;
+	GraphConfig.ConnectableObjectClasses.Add(UCameraRigAsset::StaticClass());
+	GraphConfig.ConnectableObjectClasses.Add(UCameraRigTransition::StaticClass());
+	GraphConfig.ConnectableObjectClasses.Add(UCameraRigTransitionCondition::StaticClass());
+	GraphConfig.ConnectableObjectClasses.Add(UBlendCameraNode::StaticClass());
+	GraphConfig.ObjectClassConfigs.Emplace(UCameraRigAsset::StaticClass())
+		.HasSelfPin(false)
+		.OnlyAsRoot()
+		.SetPropertyPinDirection(GET_MEMBER_NAME_STRING_CHECKED(UCameraRigAsset, EnterTransitions), EGPD_Input)
+		.SetPropertyPinDirection(GET_MEMBER_NAME_STRING_CHECKED(UCameraRigAsset, ExitTransitions), EGPD_Output)
+		.NodeTitleUsesObjectName(true)
+		.NodeTitleColor(Settings->CameraRigAssetTitleColor);
+	GraphConfig.ObjectClassConfigs.Emplace(UCameraRigTransition::StaticClass())
+		.SetPropertyPinDirection(GET_MEMBER_NAME_STRING_CHECKED(UCameraRigTransition, Conditions), EGPD_Input)
+		.NodeTitleColor(Settings->CameraRigTransitionTitleColor);
+	GraphConfig.ObjectClassConfigs.Emplace(UCameraRigTransitionCondition::StaticClass())
+		.SelfPinDirection(EGPD_Output)
+		.DefaultPropertyPinDirection(EGPD_Input)
+		.NodeTitleColor(Settings->CameraRigTransitionConditionTitleColor);
+
+	TransitionGraph = NewObject<UObjectTreeGraph>(GetTransientPackage(), NAME_None, RF_Transactional);
+	TransitionGraph->Schema = UCameraTransitionGraphSchema::StaticClass();
+	TransitionGraph->AddToRoot();
+	TransitionGraph->Initialize(CameraRigAsset, GraphConfig);
+	TransitionGraph->RebuildGraph(EObjectTreeGraphBuildSource::RootObjectPackage);
+
+	FGraphAppearanceInfo Appearance;
+	Appearance.CornerText = LOCTEXT("TransitionGraphText", "TRANSITIONS");
+
+	TransitionGraphEditor = SNew(SObjectTreeGraphEditor)
+		.Appearance(Appearance)
+		.DetailsView(InArgs._DetailsView)
+		.GraphTitle(this, &SCameraRigAssetEditor::GetCameraRigAssetName)
+		.GraphToEdit(TransitionGraph)
+		.AssetEditorToolkit(InArgs._AssetEditorToolkit);
+}
+
+ECameraRigAssetEditorMode SCameraRigAssetEditor::GetEditorMode() const
+{
+	return CurrentMode;
+}
+
+bool SCameraRigAssetEditor::IsEditorMode(ECameraRigAssetEditorMode InMode) const
+{
+	return CurrentMode == InMode;
+}
+
+void SCameraRigAssetEditor::SetEditorMode(ECameraRigAssetEditorMode InMode)
+{
+	if (InMode != CurrentMode)
+	{
+		TSharedPtr<SObjectTreeGraphEditor> CurrentGraphEditor;
+		switch(InMode)
+		{
+			case ECameraRigAssetEditorMode::NodeGraph:
+			default:
+				CurrentGraphEditor = NodeGraphEditor;
+				break;
+			case ECameraRigAssetEditorMode::TransitionGraph:
+				CurrentGraphEditor = TransitionGraphEditor;
+				break;
+		}
+
+		BoxPanel->SetContent(CurrentGraphEditor.ToSharedRef());
+		CurrentGraphEditor->ResyncDetailsView();
+		CurrentMode = InMode;
+	}
+}
+
+void SCameraRigAssetEditor::GetGraphs(TArray<UEdGraph*>& OutGraphs) const
+{
+	OutGraphs.Add(NodeGraph);
+	OutGraphs.Add(TransitionGraph);
+}
+
+const FObjectTreeGraphConfig& SCameraRigAssetEditor::GetFocusedGraphConfig() const
+{
+	static const FObjectTreeGraphConfig DefaultConfig;
+
+	switch (CurrentMode)
+	{
+		case ECameraRigAssetEditorMode::NodeGraph:
+			return NodeGraph->GetConfig();
+		case ECameraRigAssetEditorMode::TransitionGraph:
+			return TransitionGraph->GetConfig();
+		default:
+			ensure(false);
+			return DefaultConfig;
+	}
+}
+
+void SCameraRigAssetEditor::FocusHome()
+{
+	UObjectTreeGraph* Graph = nullptr;
+	TSharedPtr<SObjectTreeGraphEditor> GraphEditor = nullptr;
+
+	switch (CurrentMode)
+	{
+		case ECameraRigAssetEditorMode::NodeGraph:
+			Graph = NodeGraph;
+			GraphEditor = NodeGraphEditor;
+			break;
+		case ECameraRigAssetEditorMode::TransitionGraph:
+			Graph = TransitionGraph;
+			GraphEditor = TransitionGraphEditor;
+	}
+
+	if (Graph && GraphEditor)
+	{
+		if (UObjectTreeGraphNode* RootObjectNode = Graph->GetRootObjectNode())
+		{
+			JumpToNode(RootObjectNode);
+		}
+	}
+}
+
+void SCameraRigAssetEditor::JumpToNode(UEdGraphNode* InGraphNode)
+{
+	if (InGraphNode)
+	{
+		UEdGraph* Graph = InGraphNode->GetGraph();
+		if (Graph == NodeGraph)
+		{
+			SetEditorMode(ECameraRigAssetEditorMode::NodeGraph);
+			NodeGraphEditor->JumpToNode(InGraphNode);
+		}
+		else if (Graph == TransitionGraph)
+		{
+			SetEditorMode(ECameraRigAssetEditorMode::TransitionGraph);
+			TransitionGraphEditor->JumpToNode(InGraphNode);
+		}
+	}
+}
+
+FText SCameraRigAssetEditor::GetCameraRigAssetName() const
+{
+	return FText::FromString(CameraRigAsset->GetName());
+}
+
+}  // namespace UE::Cameras
+
+#undef LOCTEXT_NAMESPACE
+
diff --git a/Engine/Plugins/Cameras/GameplayCameras/Source/GameplayCamerasEditor/Private/Editors/SFindInObjectTreeGraph.cpp b/Engine/Plugins/Cameras/GameplayCameras/Source/GameplayCamerasEditor/Private/Editors/SFindInObjectTreeGraph.cpp
new file mode 100644
--- /dev/null
+++ b/Engine/Plugins/Cameras/GameplayCameras/Source/GameplayCamerasEditor/Private/Editors/SFindInObjectTreeGraph.cpp
@@ -0,0 +1,380 @@
+// Copyright Epic Games, Inc. All Rights Reserved.
+
+#include "Editors/SFindInObjectTreeGraph.h"
+
+#include "EdGraph/EdGraph.h"
+#include "EdGraph/EdGraphNode.h"
+#include "EdGraph/EdGraphPin.h"
+#include "EdGraph/EdGraphSchema.h"
+#include "Framework/Application/SlateApplication.h"
+#include "Types/SlateEnums.h"
+#include "Widgets/Images/SImage.h"
+#include "Widgets/Input/SSearchBox.h"
+#include "Widgets/Layout/SBorder.h"
+#include "Widgets/SBoxPanel.h"
+
+#define LOCTEXT_NAMESPACE "SFindInObjectTreeGraph"
+
+FFindInObjectTreeGraphResult::FFindInObjectTreeGraphResult(const FText& InCustomText)
+	: CustomText(InCustomText)
+{
+}
+
+FFindInObjectTreeGraphResult::FFindInObjectTreeGraphResult(TSharedPtr<FFindInObjectTreeGraphResult>& InParent, UEdGraphNode* InGraphNode)
+	: Parent(InParent)
+	, GraphNode(InGraphNode)
+{
+}
+
+FFindInObjectTreeGraphResult::FFindInObjectTreeGraphResult(TSharedPtr<FFindInObjectTreeGraphResult>& InParent, UEdGraphPin* InGraphPin)
+	: Parent(InParent)
+	, GraphPin(InGraphPin)
+{
+}
+
+TSharedRef<SWidget>	FFindInObjectTreeGraphResult::GetIcon() const
+{
+	FSlateColor IconColor = FSlateColor::UseForeground();
+	const FSlateBrush* Brush = NULL;
+
+	if (UEdGraphPin* ResolvedPin = GraphPin.Get())
+	{
+		if (ResolvedPin->PinType.IsArray())
+		{
+			Brush = FAppStyle::GetBrush(TEXT("GraphEditor.ArrayPinIcon"));
+		}
+		else if (ResolvedPin->PinType.bIsReference)
+		{
+			Brush = FAppStyle::GetBrush(TEXT("GraphEditor.RefPinIcon"));
+		}
+		else
+		{
+			Brush = FAppStyle::GetBrush(TEXT("GraphEditor.PinIcon"));
+		}
+
+		const UEdGraphSchema* Schema = ResolvedPin->GetSchema();
+		IconColor = Schema->GetPinTypeColor(ResolvedPin->PinType);
+	}
+	else if (GraphNode.IsValid())
+	{
+		Brush = FAppStyle::GetBrush(TEXT("GraphEditor.NodeGlyph"));
+	}
+
+	return SNew(SImage)
+		.Image(Brush)
+		.ColorAndOpacity(IconColor)
+		.ToolTipText(GetCategory());
+}
+
+FText FFindInObjectTreeGraphResult::GetCategory() const
+{
+	if (GraphNode.IsValid())
+	{
+		return LOCTEXT("NodeCategory", "Node");
+	}
+	else if (GraphPin.Get())
+	{
+		return LOCTEXT("PinCategory", "Pin");
+	}
+	return FText::GetEmpty();
+}
+
+FText FFindInObjectTreeGraphResult::GetText() const
+{
+	if (UEdGraphNode* ResolvedNode = GraphNode.Get())
+	{
+		const FText NodeFullTitle = ResolvedNode->GetNodeTitle(ENodeTitleType::FullTitle);
+		const FText NodeListViewTitle = ResolvedNode->GetNodeTitle(ENodeTitleType::ListView);
+		if (NodeFullTitle.EqualToCaseIgnored(NodeListViewTitle))
+		{
+			return NodeFullTitle;
+		}
+		return FText::Format(LOCTEXT("NodeResultFmt", "{0} - {1}"), NodeListViewTitle, NodeFullTitle);
+	}
+	else if (UEdGraphPin* ResolvedPin = GraphPin.Get())
+	{
+		const UEdGraphSchema* GraphSchema = ResolvedPin->GetSchema();
+		return GraphSchema->GetPinDisplayName(ResolvedPin);
+	}
+	return CustomText;
+}
+
+FText FFindInObjectTreeGraphResult::GetCommentText() const
+{
+	if (UEdGraphNode* ResolvedNode = GraphNode.Get())
+	{
+		return FText::FromString(ResolvedNode->NodeComment);
+	}
+	return FText::GetEmpty();
+}
+
+FReply FFindInObjectTreeGraphResult::OnClick(TSharedRef<SFindInObjectTreeGraph> FindInObjectTreeGraph)
+{
+	if (UEdGraphNode* ResolvedNode = GraphNode.Get())
+	{
+		FindInObjectTreeGraph->OnJumpToNodeRequested.ExecuteIfBound(ResolvedNode);
+		return FReply::Handled();
+	}
+	else if (UEdGraphPin* ResolvedPin = GraphPin.Get())
+	{
+		if (FindInObjectTreeGraph->OnJumpToPinRequested.IsBound())
+		{
+			FindInObjectTreeGraph->OnJumpToPinRequested.Execute(ResolvedPin);
+		}
+		else
+		{
+			FindInObjectTreeGraph->OnJumpToNodeRequested.ExecuteIfBound(ResolvedPin->GetOwningNode());
+		}
+		return FReply::Handled();
+	}
+	return FReply::Unhandled();
+}
+
+void SFindInObjectTreeGraph::Construct(const FArguments& InArgs)
+{
+	GraphsToSearch = InArgs._GraphsToSearch;
+
+	OnJumpToNodeRequested = InArgs._OnJumpToNodeRequested;
+	OnJumpToPinRequested = InArgs._OnJumpToPinRequested;
+
+	ChildSlot
+	[
+		SNew(SVerticalBox)
+		+ SVerticalBox::Slot()
+		.AutoHeight()
+		[
+			SNew(SHorizontalBox)
+			+ SHorizontalBox::Slot()
+			.FillWidth(1)
+			[
+				SAssignNew(SearchBox, SSearchBox)
+				.HintText(LOCTEXT("SearchHint", "Search"))
+				.OnTextChanged(this, &SFindInObjectTreeGraph::OnSearchTextChanged)
+				.OnTextCommitted(this, &SFindInObjectTreeGraph::OnSearchTextCommitted)
+			]
+		]
+		+ SVerticalBox::Slot()
+		.FillHeight(1.f)
+		.Padding(0, 4, 0, 0)
+		[
+			SNew(SBorder)
+			.BorderImage(FAppStyle::GetBrush("Menu.Background"))
+			[
+				SAssignNew(ResultTreeView, SResultTreeView)
+				.ItemHeight(24)
+				.TreeItemsSource(&Results)
+				.SelectionMode(ESelectionMode::Multi)
+				.OnGenerateRow(this, &SFindInObjectTreeGraph::OnResultTreeViewGenerateRow)
+				.OnGetChildren(this, &SFindInObjectTreeGraph::OnResultTreeViewGetChildren)
+				.OnSelectionChanged(this, &SFindInObjectTreeGraph::OnResultTreeViewSelectionChanged)
+				.OnMouseButtonDoubleClick(this, &SFindInObjectTreeGraph::OnResultTreeViewMouseButtonDoubleClick)
+			]
+		]
+	];
+}
+
+void SFindInObjectTreeGraph::FocusSearchEditBox()
+{
+	FSlateApplication::Get().SetKeyboardFocus(SearchBox, EFocusCause::SetDirectly);
+}
+
+void SFindInObjectTreeGraph::OnSearchTextChanged(const FText& Text)
+{
+	SearchQuery = Text.ToString();
+}
+
+void SFindInObjectTreeGraph::OnSearchTextCommitted(const FText& Text, ETextCommit::Type CommitType)
+{
+	if (CommitType == ETextCommit::OnEnter)
+	{
+		StartSearch();
+	}
+}
+
+TSharedRef<ITableRow> SFindInObjectTreeGraph::OnResultTreeViewGenerateRow(FResultPtr InItem, const TSharedRef<STableViewBase>& OwnerTable)
+{
+	const FText CommentText = InItem->GetCommentText();
+
+	return SNew(STableRow<FResultPtr>, OwnerTable)
+		[
+			SNew(SHorizontalBox)
+			+ SHorizontalBox::Slot()
+			.AutoWidth()
+			.VAlign(VAlign_Center)
+			[
+				InItem->GetIcon()
+			]
+			+ SHorizontalBox::Slot()
+			.AutoWidth()
+			.VAlign(VAlign_Center)
+			.Padding(2, 0)
+			[
+				SNew(STextBlock)
+				.Text(InItem->GetText())
+				.HighlightText(HighlightText)
+				.ToolTipText(FText::Format(LOCTEXT("ResultToolTipFmt", "{0} : {1}"), InItem->GetCategory(), InItem->GetText()))
+			]
+			+ SHorizontalBox::Slot()
+			.FillWidth(1)
+			.HAlign(HAlign_Right)
+			.VAlign(VAlign_Center)
+			.Padding(2, 0)
+			[
+				SNew(STextBlock)
+				.Text(CommentText.IsEmpty() 
+						? FText::GetEmpty() 
+						: FText::Format(LOCTEXT("NodeCommentFmt", "Node Comment: {0}"), CommentText))
+				.HighlightText(HighlightText)
+			]
+		];
+}
+
+void SFindInObjectTreeGraph::OnResultTreeViewGetChildren(FResultPtr InItem, TArray<FResultPtr>& OutChildren)
+{
+	OutChildren += InItem->Children;
+}
+
+void SFindInObjectTreeGraph::OnResultTreeViewSelectionChanged(FResultPtr Item, ESelectInfo::Type SelectInfo)
+{
+	if (Item.IsValid())
+	{
+		Item->OnClick(SharedThis(this));
+	}
+}
+
+void SFindInObjectTreeGraph::OnResultTreeViewMouseButtonDoubleClick(FResultPtr Item)
+{
+	if (Item.IsValid())
+	{
+		Item->OnClick(SharedThis(this));
+	}
+}
+
+void SFindInObjectTreeGraph::StartSearch()
+{
+	TArray<FString> Tokens;
+	if (SearchQuery.Contains("\"") && SearchQuery.ParseIntoArray(Tokens, TEXT("\""), true) > 0)
+	{
+		for (FString& Token : Tokens)
+		{
+			Token = Token.TrimQuotes();
+		}
+	}
+	else
+	{
+		SearchQuery.ParseIntoArray(Tokens, TEXT(" "), true);
+	}
+	Tokens.RemoveAll([](const FString& Item) { return Item.IsEmpty(); });
+
+	Results.Empty();
+	HighlightText = FText::GetEmpty();
+	if (Tokens.Num() > 0)
+	{
+		HighlightText = FText::FromString(SearchQuery);
+		for (UEdGraph* Graph : GraphsToSearch)
+		{
+			MatchTokensInGraph(Graph, Tokens);
+		}
+	}
+	
+	if (Results.IsEmpty())
+	{
+		Results.Add(MakeShared<FFindInObjectTreeGraphResult>(LOCTEXT("NoResults", "No results found")));
+	}
+
+	ResultTreeView->RequestTreeRefresh();
+	for (FResultPtr Result : Results)
+	{
+		ResultTreeView->SetItemExpansion(Result, true);
+	}
+}
+
+void SFindInObjectTreeGraph::MatchTokensInGraph(UEdGraph* Graph, TArrayView<FString> Tokens)
+{
+	FResultPtr GraphResult(new FFindInObjectTreeGraphResult(FText::FromName(Graph->GetFName())));
+
+	const UEdGraphSchema* GraphSchema = Graph->GetSchema();
+
+	for (UEdGraphNode* Node : Graph->Nodes)
+	{
+		const bool bNodeMatches = GraphNodeMatchesSearchTokens(GraphSchema, Node, Tokens);
+
+		TArray<UEdGraphPin*> MatchingPins;
+		for (UEdGraphPin* Pin : Node->Pins)
+		{
+			if (GraphPinMatchesSearchTokens(GraphSchema, Pin, Tokens))
+			{
+				MatchingPins.Add(Pin);
+			}
+		}
+
+		if (bNodeMatches || MatchingPins.Num() > 0)
+		{
+			FResultPtr NodeResult(new FFindInObjectTreeGraphResult(GraphResult, Node));
+			Results.Add(NodeResult);
+
+			for (UEdGraphPin* Pin : MatchingPins)
+			{
+				FResultPtr PinResult(new FFindInObjectTreeGraphResult(NodeResult, Pin));
+			}
+		}
+	}
+
+	for (UEdGraph* SubGraph : Graph->SubGraphs)
+	{
+		MatchTokensInGraph(SubGraph, Tokens);
+	}
+}
+
+bool SFindInObjectTreeGraph::GraphNodeMatchesSearchTokens(const UEdGraphSchema* GraphSchema, UEdGraphNode* GraphNode, TArrayView<FString> Tokens)
+{
+	const FString NodeFullTitle = GraphNode->GetNodeTitle(ENodeTitleType::FullTitle).ToString();
+	if (StringMatchesSearchTokens(NodeFullTitle, Tokens))
+	{
+		return true;
+	}
+
+	const FString NodeListViewTitle = GraphNode->GetNodeTitle(ENodeTitleType::ListView).ToString();
+	if (StringMatchesSearchTokens(NodeListViewTitle, Tokens))
+	{
+		return true;
+	}
+
+	if (StringMatchesSearchTokens(GraphNode->NodeComment, Tokens))
+	{
+		return true;
+	}
+
+	return false;
+}
+
+bool SFindInObjectTreeGraph::GraphPinMatchesSearchTokens(const UEdGraphSchema* GraphSchema, UEdGraphPin* GraphPin, TArrayView<FString> Tokens)
+{
+	const FString PinDisplayName = GraphSchema->GetPinDisplayName(GraphPin).ToString();
+	if (StringMatchesSearchTokens(PinDisplayName, Tokens))
+	{
+		return true;
+	}
+
+	return false;
+}
+
+bool SFindInObjectTreeGraph::StringMatchesSearchTokens(const FString& ComparisonString, TArrayView<FString> Tokens)
+{
+	if (ComparisonString.IsEmpty())
+	{
+		return false;
+	}
+
+	for (const FString& Token : Tokens)
+	{
+		if (!ComparisonString.Contains(Token))
+		{
+			return false;
+		}
+	}
+	return true;
+}
+
+#undef LOCTEXT_NAMESPACE
+
diff --git a/Engine/Plugins/Cameras/GameplayCameras/Source/GameplayCamerasEditor/Private/Editors/SObjectTreeGraphEditor.cpp b/Engine/Plugins/Cameras/GameplayCameras/Source/GameplayCamerasEditor/Private/Editors/SObjectTreeGraphEditor.cpp
new file mode 100644
--- /dev/null
+++ b/Engine/Plugins/Cameras/GameplayCameras/Source/GameplayCamerasEditor/Private/Editors/SObjectTreeGraphEditor.cpp
@@ -0,0 +1,610 @@
+// Copyright Epic Games, Inc. All Rights Reserved.
+
+#include "Editors/SObjectTreeGraphEditor.h"
+
+#include "Algo/AnyOf.h"
+#include "EdGraph/EdGraphPin.h"
+#include "EdGraph/EdGraphSchema.h"
+#include "EdGraphUtilities.h"
+#include "Editor.h"
+#include "Editors/ObjectTreeGraph.h"
+#include "Editors/ObjectTreeGraphNode.h"
+#include "Editors/ObjectTreeGraphSchema.h"
+#include "Editors/SObjectTreeGraphTitleBar.h"
+#include "Editors/SObjectTreeGraphToolbox.h"
+#include "Framework/Commands/GenericCommands.h"
+#include "GraphEditorActions.h"
+#include "HAL/PlatformApplicationMisc.h"
+#include "IDetailsView.h"
+#include "SNodePanel.h"
+#include "ScopedTransaction.h"
+#include "SGraphPanel.h"
+#include "Widgets/Input/SEditableTextBox.h"
+
+#define LOCTEXT_NAMESPACE "SObjectTreeGraphEditor"
+
+void SObjectTreeGraphEditor::Construct(const FArguments& InArgs)
+{
+	DetailsView = InArgs._DetailsView;
+
+	TSharedPtr<SWidget> GraphTitleBar = InArgs._GraphTitleBar;
+	if (!GraphTitleBar.IsValid())
+	{
+		TSharedRef<SWidget> DefaultTitleBar = SNew(SObjectTreeGraphTitleBar)
+			.Graph(InArgs._GraphToEdit)
+			.TitleText(InArgs._GraphTitle);
+		GraphTitleBar = DefaultTitleBar.ToSharedPtr();
+	}
+
+	SGraphEditor::FGraphEditorEvents GraphEditorEvents;
+	GraphEditorEvents.OnSelectionChanged = SGraphEditor::FOnSelectionChanged::CreateSP(this, &SObjectTreeGraphEditor::OnGraphSelectionChanged);
+	GraphEditorEvents.OnTextCommitted = FOnNodeTextCommitted::CreateSP(this, &SObjectTreeGraphEditor::OnNodeTextCommitted);
+	GraphEditorEvents.OnDoubleClicked = SGraphEditor::FOnDoubleClicked::CreateSP(this, &SObjectTreeGraphEditor::OnDoubleClicked);
+	GraphEditorEvents.OnNodeDoubleClicked = FSingleNodeEvent::CreateSP(this, &SObjectTreeGraphEditor::OnNodeDoubleClicked);
+
+	InitializeBuiltInCommands();
+
+	TSharedPtr<FUICommandList> AdditionalCommands = BuiltInCommands;
+	if (InArgs._AdditionalCommands)
+	{
+		AdditionalCommands = MakeShared<FUICommandList>();
+		AdditionalCommands->Append(BuiltInCommands.ToSharedRef());
+		AdditionalCommands->Append(InArgs._AdditionalCommands.ToSharedRef());
+	}
+
+	GraphEditor = SNew(SGraphEditor)
+		.AdditionalCommands(AdditionalCommands)
+		.Appearance(InArgs._Appearance)
+		.TitleBar(GraphTitleBar)
+		.GraphToEdit(InArgs._GraphToEdit)
+		.GraphEvents(GraphEditorEvents)
+		.AssetEditorToolkit(InArgs._AssetEditorToolkit);
+
+	ChildSlot
+	[
+		GraphEditor.ToSharedRef()
+	];
+
+	GEditor->RegisterForUndo(this);
+}
+
+SObjectTreeGraphEditor::~SObjectTreeGraphEditor()
+{
+	GEditor->UnregisterForUndo(this);
+}
+
+void SObjectTreeGraphEditor::InitializeBuiltInCommands()
+{
+	if (BuiltInCommands.IsValid())
+	{
+		return;
+	}
+
+	const FGenericCommands& GenericCommands = FGenericCommands::Get();
+	const FGraphEditorCommandsImpl& GraphEditorCommands = FGraphEditorCommands::Get();
+
+	BuiltInCommands = MakeShared<FUICommandList>();
+
+	// Generic commands.
+	BuiltInCommands->MapAction(GenericCommands.SelectAll,
+		FExecuteAction::CreateSP(this, &SObjectTreeGraphEditor::SelectAllNodes),
+		FCanExecuteAction::CreateSP(this, &SObjectTreeGraphEditor::CanSelectAllNodes)
+		);
+
+	BuiltInCommands->MapAction(GenericCommands.Delete,
+		FExecuteAction::CreateSP(this, &SObjectTreeGraphEditor::DeleteSelectedNodes),
+		FCanExecuteAction::CreateSP(this, &SObjectTreeGraphEditor::CanDeleteSelectedNodes)
+		);
+
+	BuiltInCommands->MapAction(GenericCommands.Copy,
+		FExecuteAction::CreateSP(this, &SObjectTreeGraphEditor::CopySelectedNodes),
+		FCanExecuteAction::CreateSP(this, &SObjectTreeGraphEditor::CanCopySelectedNodes)
+		);
+
+	BuiltInCommands->MapAction(GenericCommands.Cut,
+		FExecuteAction::CreateSP(this, &SObjectTreeGraphEditor::CutSelectedNodes),
+		FCanExecuteAction::CreateSP(this, &SObjectTreeGraphEditor::CanCutSelectedNodes)
+		);
+
+	BuiltInCommands->MapAction(GenericCommands.Paste,
+		FExecuteAction::CreateSP(this, &SObjectTreeGraphEditor::PasteNodes),
+		FCanExecuteAction::CreateSP(this, &SObjectTreeGraphEditor::CanPasteNodes)
+		);
+
+	BuiltInCommands->MapAction(GenericCommands.Duplicate,
+		FExecuteAction::CreateSP(this, &SObjectTreeGraphEditor::DuplicateNodes),
+		FCanExecuteAction::CreateSP(this, &SObjectTreeGraphEditor::CanDuplicateNodes)
+		);
+	BuiltInCommands->MapAction(FGenericCommands::Get().Rename,
+		FExecuteAction::CreateSP(this, &SObjectTreeGraphEditor::OnRenameNode),
+		FCanExecuteAction::CreateSP(this, &SObjectTreeGraphEditor::CanRenameNode)
+	);
+
+	// Alignment commands.
+	BuiltInCommands->MapAction(GraphEditorCommands.AlignNodesTop,
+			FExecuteAction::CreateSP(this, &SObjectTreeGraphEditor::OnAlignTop)
+			);
+
+	BuiltInCommands->MapAction(GraphEditorCommands.AlignNodesMiddle,
+			FExecuteAction::CreateSP(this, &SObjectTreeGraphEditor::OnAlignMiddle)
+			);
+
+	BuiltInCommands->MapAction(GraphEditorCommands.AlignNodesBottom,
+			FExecuteAction::CreateSP(this, &SObjectTreeGraphEditor::OnAlignBottom)
+			);
+
+	BuiltInCommands->MapAction(GraphEditorCommands.AlignNodesLeft,
+			FExecuteAction::CreateSP(this, &SObjectTreeGraphEditor::OnAlignLeft)
+			);
+
+	BuiltInCommands->MapAction(GraphEditorCommands.AlignNodesCenter,
+			FExecuteAction::CreateSP(this, &SObjectTreeGraphEditor::OnAlignCenter)
+			);
+
+	BuiltInCommands->MapAction(GraphEditorCommands.AlignNodesRight,
+			FExecuteAction::CreateSP(this, &SObjectTreeGraphEditor::OnAlignRight)
+			);
+
+	BuiltInCommands->MapAction(GraphEditorCommands.StraightenConnections,
+			FExecuteAction::CreateSP(this, &SObjectTreeGraphEditor::OnStraightenConnections)
+			);
+
+	// Distribution commands.
+	BuiltInCommands->MapAction(GraphEditorCommands.DistributeNodesHorizontally,
+			FExecuteAction::CreateSP(this, &SObjectTreeGraphEditor::OnDistributeNodesHorizontally)
+			);
+
+	BuiltInCommands->MapAction(GraphEditorCommands.DistributeNodesVertically,
+			FExecuteAction::CreateSP(this, &SObjectTreeGraphEditor::OnDistributeNodesVertically)
+			);
+}
+
+void SObjectTreeGraphEditor::JumpToNode(UEdGraphNode* InNode)
+{
+	GraphEditor->JumpToNode(InNode);
+}
+
+void SObjectTreeGraphEditor::ResyncDetailsView()
+{
+	OnGraphSelectionChanged(GraphEditor->GetSelectedNodes());
+}
+
+FReply SObjectTreeGraphEditor::OnDragOver(const FGeometry& MyGeometry, const FDragDropEvent& DragDropEvent)
+{
+	TSharedPtr<FObjectTreeClassDragDropOp> ObjectClassOp = DragDropEvent.GetOperationAs<FObjectTreeClassDragDropOp>();
+	if (ObjectClassOp)
+	{
+		TArray<UClass*> PlaceableClasses = FilterPlaceableObjectClasses(ObjectClassOp->GetObjectClasses());
+		if (PlaceableClasses.Num() == ObjectClassOp->GetObjectClasses().Num())
+		{
+			const FSlateBrush* OKIcon = FAppStyle::GetBrush(TEXT("Graph.ConnectorFeedback.OK"));
+			ObjectClassOp->SetToolTip(
+					FText::Format(
+						LOCTEXT("OnDragOver_Success", "Create {0} node(s) from the dragged object classes"),
+						ObjectClassOp->GetObjectClasses().Num()),
+					OKIcon);
+		}
+		else if (PlaceableClasses.Num() > 0)
+		{
+			const FSlateBrush* WarnIcon = FAppStyle::GetBrush(TEXT("Graph.ConnectorFeedback.OKWarn"));
+			ObjectClassOp->SetToolTip(
+					FText::Format(
+						LOCTEXT("OnDragOver_Warning", "Create {0} node(s) from the dragged object classes, ignoring {1} that can't be created in this graph"),
+						PlaceableClasses.Num(), (ObjectClassOp->GetObjectClasses().Num() - PlaceableClasses.Num())),
+					WarnIcon);
+		}
+		else
+		{
+			const FSlateBrush* ErrorIcon = FAppStyle::GetBrush(TEXT("Graph.ConnectorFeedback.Error"));
+			ObjectClassOp->SetToolTip(
+					LOCTEXT("OnDragOver_Error", "The dragged object classes can't be created in this graph"),
+					ErrorIcon);
+		}
+
+		return FReply::Handled();
+	}
+
+	return SCompoundWidget::OnDragOver(MyGeometry, DragDropEvent);
+}
+
+FReply SObjectTreeGraphEditor::OnDrop(const FGeometry& MyGeometry, const FDragDropEvent& DragDropEvent)
+{
+	TSharedPtr<FObjectTreeClassDragDropOp> ObjectClassOp = DragDropEvent.GetOperationAs<FObjectTreeClassDragDropOp>();
+	if (ObjectClassOp)
+	{
+		const FScopedTransaction Transaction(LOCTEXT("DropObjectClasses", "Drop New Nodes"));
+
+		TArray<UClass*> PlaceableClasses = FilterPlaceableObjectClasses(ObjectClassOp->GetObjectClasses());
+		UEdGraph* Graph = GraphEditor->GetCurrentGraph();
+
+		GraphEditor->ClearSelectionSet();
+
+		SGraphPanel* GraphPanel = GraphEditor->GetGraphPanel();
+		FVector2D NewLocation = GraphPanel->PanelCoordToGraphCoord(MyGeometry.AbsoluteToLocal(DragDropEvent.GetScreenSpacePosition()));
+
+		for (UClass* PlaceableClass : PlaceableClasses)
+		{
+			FObjectGraphSchemaAction_NewNode Action;
+			Action.ObjectClass = PlaceableClass;
+			UEdGraphNode* NewNode = Action.PerformAction(Graph, nullptr, NewLocation, false);
+			GraphEditor->SetNodeSelection(NewNode, true);
+
+			NewLocation += FVector2D(20, 20);
+		}
+	}
+
+	return SCompoundWidget::OnDrop(MyGeometry, DragDropEvent);
+}
+
+TArray<UClass*> SObjectTreeGraphEditor::FilterPlaceableObjectClasses(TArrayView<UClass* const> InObjectClasses)
+{
+	UObjectTreeGraph* Graph = CastChecked<UObjectTreeGraph>(GraphEditor->GetCurrentGraph());
+	const FObjectTreeGraphConfig& GraphConfig = Graph->GetConfig();
+	TArray<UClass*> PlaceableClasses = InObjectClasses.FilterByPredicate(
+			[&GraphConfig](UClass* ObjectClass)
+			{
+				return GraphConfig.IsConnectable(ObjectClass);
+			});
+	return PlaceableClasses;
+}
+
+void SObjectTreeGraphEditor::PostUndo(bool bSuccess)
+{
+	if (bSuccess)
+	{
+		GraphEditor->NotifyGraphChanged();
+	}
+}
+
+void SObjectTreeGraphEditor::PostRedo(bool bSuccess)
+{
+	PostUndo(bSuccess);
+}
+
+void SObjectTreeGraphEditor::OnGraphSelectionChanged(const FGraphPanelSelectionSet& SelectionSet)
+{
+	if (DetailsView)
+	{
+		TArray<UObject*> SelectedObjects;
+		for (UObject* Selection : SelectionSet)
+		{
+			if (UObjectTreeGraphNode* GraphNode = Cast<UObjectTreeGraphNode>(Selection))
+			{
+				SelectedObjects.Add(GraphNode->GetObject());
+			}
+		}
+
+		DetailsView->SetObjects(SelectedObjects);
+	}
+}
+
+void SObjectTreeGraphEditor::OnNodeTextCommitted(const FText& InText, ETextCommit::Type InCommitType, UEdGraphNode* InEditedNode)
+{
+	if (InEditedNode)
+	{
+		FString NewName = InText.ToString().TrimStartAndEnd();
+		if (NewName.IsEmpty())
+		{
+			return;
+		}
+
+		if (NewName.Len() >= NAME_SIZE)
+		{
+			NewName = NewName.Left(NAME_SIZE - 1);
+		}
+
+		const FScopedTransaction Transaction(LOCTEXT("RenameNode", "Rename Node"));
+
+		InEditedNode->Modify();
+		InEditedNode->OnRenameNode(NewName);
+
+		GraphEditor->GetCurrentGraph()->NotifyNodeChanged(InEditedNode);
+	}
+}
+
+void SObjectTreeGraphEditor::OnNodeDoubleClicked(UEdGraphNode* InClickedNode)
+{
+	UObjectTreeGraphNode* SelectedNode = Cast<UObjectTreeGraphNode>(InClickedNode);
+	if (SelectedNode)
+	{
+		SelectedNode->OnDoubleClicked();
+	}
+}
+
+void SObjectTreeGraphEditor::OnDoubleClicked()
+{
+}
+
+FString SObjectTreeGraphEditor::ExportNodesToText(const FGraphPanelSelectionSet& Nodes, bool bOnlyCanDuplicateNodes, bool bOnlyCanDeleteNodes)
+{
+	FString ExportedText;
+	TSet<UObject*> NodesToExport;
+
+	for (FGraphPanelSelectionSet::TConstIterator NodeIt(Nodes); NodeIt; ++NodeIt)
+	{
+		UEdGraphNode* Node = Cast<UEdGraphNode>(*NodeIt);
+		if (Node && 
+				(!bOnlyCanDuplicateNodes || Node->CanDuplicateNode()) &&
+				(!bOnlyCanDeleteNodes || Node->CanUserDeleteNode()))
+		{
+			Node->PrepareForCopying();
+			NodesToExport.Add(Node);
+		}
+	}
+
+	FEdGraphUtilities::ExportNodesToText(NodesToExport, ExportedText);
+
+	return ExportedText;
+}
+
+void SObjectTreeGraphEditor::ImportNodesFromText(const FVector2D& Location, const FString& TextToImport)
+{
+	// Start a transaction and flag things as modified.
+	const FScopedTransaction Transaction(LOCTEXT("PasteNodes", "Paste Nodes"));
+
+	UObjectTreeGraph* Graph = CastChecked<UObjectTreeGraph>(GraphEditor->GetCurrentGraph());
+	Graph->Modify();
+
+	UPackage* ObjectPackage = Graph->GetRootObject()->GetOutermost();
+	ObjectPackage->Modify();
+
+	// Import the nodes
+	TSet<UEdGraphNode*> PastedNodes;
+	FEdGraphUtilities::ImportNodesFromText(Graph, TextToImport, PastedNodes);
+
+	// Process all the new nodes.
+	TMap<FGuid, UEdGraphNode*> GuidToNode;
+	for (UEdGraphNode* Node : Graph->Nodes)
+	{
+		GuidToNode.Add(Node->NodeGuid, Node);
+	}
+	
+	TMap<UEdGraphNode*, UEdGraphNode*> PastedNodeMap;
+	for (UEdGraphNode* PastedNode : PastedNodes)
+	{
+		const FGuid OldNodeGuid = PastedNode->NodeGuid;
+		PastedNode->CreateNewGuid();
+
+		UEdGraphNode** OldNode = GuidToNode.Find(OldNodeGuid);
+		if (ensure(OldNode))
+		{
+			PastedNodeMap.Add(*OldNode, PastedNode);
+		}
+	}
+	
+	const UObjectTreeGraphSchema* Schema = CastChecked<UObjectTreeGraphSchema>(Graph->GetSchema());
+	Schema->ProcessDuplicatedNodes(Graph, PastedNodeMap);
+
+	// Compute the center of the pasted nodes.
+	FVector2D PastedNodesClusterCenter(FVector2D::ZeroVector);
+	for (UEdGraphNode* PastedNode : PastedNodes)
+	{
+		PastedNodesClusterCenter.X += PastedNode->NodePosX;
+		PastedNodesClusterCenter.Y += PastedNode->NodePosY;
+	}
+	if (PastedNodes.Num() > 0)
+	{
+		float InvNumNodes = 1.0f / float(PastedNodes.Num());
+		PastedNodesClusterCenter.X *= InvNumNodes;
+		PastedNodesClusterCenter.Y *= InvNumNodes;
+	}
+
+	// Move all pasted nodes to the new location, and select them.
+	GraphEditor->ClearSelectionSet();
+
+	for (UEdGraphNode* PastedNode : PastedNodes)
+	{
+		PastedNode->NodePosX = (PastedNode->NodePosX - PastedNodesClusterCenter.X) + Location.X ;
+		PastedNode->NodePosY = (PastedNode->NodePosY - PastedNodesClusterCenter.Y) + Location.Y ;
+
+		PastedNode->SnapToGrid(SNodePanel::GetSnapGridSize());
+
+		GraphEditor->SetNodeSelection(PastedNode, true);
+	}
+
+	// Update the UI.
+	GraphEditor->NotifyGraphChanged();
+}
+
+void SObjectTreeGraphEditor::DeleteNodes(TArrayView<UObjectTreeGraphNode*> NodesToDelete)
+{
+	const UEdGraphSchema* Schema = GraphEditor->GetCurrentGraph()->GetSchema();
+
+	for (UObjectTreeGraphNode* Node : NodesToDelete)
+	{
+		if (Node)
+		{
+			Node->Modify();
+
+			Schema->BreakNodeLinks(*Node);
+
+			Node->DestroyNode();
+		}
+	}
+}
+
+void SObjectTreeGraphEditor::SelectAllNodes()
+{
+	GraphEditor->SelectAllNodes();
+}
+
+bool SObjectTreeGraphEditor::CanSelectAllNodes()
+{
+	return true;
+}
+
+void SObjectTreeGraphEditor::DeleteSelectedNodes()
+{
+	TArray<UObjectTreeGraphNode*> NodesToDelete;
+	const FGraphPanelSelectionSet SelectedNodes = GraphEditor->GetSelectedNodes();
+
+	for (FGraphPanelSelectionSet::TConstIterator NodeIt(SelectedNodes); NodeIt; ++NodeIt)
+	{
+		UEdGraphNode* GraphNode = Cast<UEdGraphNode>(*NodeIt);
+		if (GraphNode && GraphNode->CanUserDeleteNode())
+		{
+			NodesToDelete.Add(Cast<UObjectTreeGraphNode>(*NodeIt));
+		}
+	}
+	
+	DeleteNodes(NodesToDelete);
+}
+
+bool SObjectTreeGraphEditor::CanDeleteSelectedNodes()
+{
+	bool bDeletableNodeExists = false;
+	const FGraphPanelSelectionSet SelectedNodes = GraphEditor->GetSelectedNodes();
+
+	for (FGraphPanelSelectionSet::TConstIterator NodeIt(SelectedNodes); NodeIt; ++NodeIt)
+	{
+		UEdGraphNode* GraphNode = Cast<UEdGraphNode>(*NodeIt);
+		if (GraphNode && GraphNode->CanUserDeleteNode())
+		{
+			bDeletableNodeExists = true;
+		}
+	}
+
+	return SelectedNodes.Num() > 0 && bDeletableNodeExists;
+}
+
+void SObjectTreeGraphEditor::CopySelectedNodes()
+{
+	const FString Buffer = ExportNodesToText(GraphEditor->GetSelectedNodes(), true, false);
+	FPlatformApplicationMisc::ClipboardCopy(*Buffer);
+}
+
+bool SObjectTreeGraphEditor::CanCopySelectedNodes()
+{
+	const FGraphPanelSelectionSet SelectedNodes = GraphEditor->GetSelectedNodes();
+
+	for (FGraphPanelSelectionSet::TConstIterator NodeIt(SelectedNodes); NodeIt; ++NodeIt)
+	{
+		UEdGraphNode* Node = Cast<UEdGraphNode>(*NodeIt);
+		if (Node != nullptr && Node->CanDuplicateNode())
+		{
+			return true;
+		}
+	}
+
+	return false;
+}
+
+void SObjectTreeGraphEditor::CutSelectedNodes()
+{
+	const FString Buffer = ExportNodesToText(GraphEditor->GetSelectedNodes(), true, true);
+	FPlatformApplicationMisc::ClipboardCopy(*Buffer);
+
+	DeleteSelectedNodes();
+}
+
+bool SObjectTreeGraphEditor::CanCutSelectedNodes()
+{
+	return CanCopySelectedNodes() && CanDeleteSelectedNodes();
+}
+
+void SObjectTreeGraphEditor::PasteNodes()
+{
+	FString TextToImport;
+	FPlatformApplicationMisc::ClipboardPaste(TextToImport);
+	
+	ImportNodesFromText(GraphEditor->GetPasteLocation(), TextToImport);
+}
+
+bool SObjectTreeGraphEditor::CanPasteNodes()
+{
+	FString ClipboardContent;
+	FPlatformApplicationMisc::ClipboardPaste(ClipboardContent);
+
+	return FEdGraphUtilities::CanImportNodesFromText(GraphEditor->GetCurrentGraph(), ClipboardContent);
+}
+
+void SObjectTreeGraphEditor::DuplicateNodes()
+{
+	CopySelectedNodes();
+	PasteNodes();
+}
+
+bool SObjectTreeGraphEditor::CanDuplicateNodes()
+{
+	return CanCopySelectedNodes();
+}
+
+void SObjectTreeGraphEditor::OnRenameNode()
+{
+	const FGraphPanelSelectionSet SelectedNodes = GraphEditor->GetSelectedNodes();
+
+	for (FGraphPanelSelectionSet::TConstIterator NodeIt(SelectedNodes); NodeIt; ++NodeIt)
+	{
+		UEdGraphNode* SelectedNode = Cast<UEdGraphNode>(*NodeIt);
+		if (SelectedNode != nullptr && SelectedNode->GetCanRenameNode())
+		{
+			const bool bRequestRename = true;
+			GraphEditor->IsNodeTitleVisible(SelectedNode, bRequestRename);
+			break;
+		}
+	}
+}
+
+bool SObjectTreeGraphEditor::CanRenameNode()
+{
+	const FGraphPanelSelectionSet SelectedNodes = GraphEditor->GetSelectedNodes();
+
+	for (FGraphPanelSelectionSet::TConstIterator NodeIt(SelectedNodes); NodeIt; ++NodeIt)
+	{
+		UEdGraphNode* Node = Cast<UEdGraphNode>(*NodeIt);
+		if (Node != nullptr && Node->GetCanRenameNode())
+		{
+			return true;
+		}
+	}
+
+	return false;
+}
+
+void SObjectTreeGraphEditor::OnAlignTop()
+{
+	GraphEditor->OnAlignTop();
+}
+
+void SObjectTreeGraphEditor::OnAlignMiddle()
+{
+	GraphEditor->OnAlignMiddle();
+}
+
+void SObjectTreeGraphEditor::OnAlignBottom()
+{
+	GraphEditor->OnAlignBottom();
+}
+
+void SObjectTreeGraphEditor::OnAlignLeft()
+{
+	GraphEditor->OnAlignLeft();
+}
+
+void SObjectTreeGraphEditor::OnAlignCenter()
+{
+	GraphEditor->OnAlignCenter();
+}
+
+void SObjectTreeGraphEditor::OnAlignRight()
+{
+	GraphEditor->OnAlignRight();
+}
+
+void SObjectTreeGraphEditor::OnStraightenConnections()
+{
+	GraphEditor->OnStraightenConnections();
+}
+
+void SObjectTreeGraphEditor::OnDistributeNodesHorizontally()
+{
+	GraphEditor->OnDistributeNodesH();
+}
+
+void SObjectTreeGraphEditor::OnDistributeNodesVertically()
+{
+	GraphEditor->OnDistributeNodesV();
+}
+
+#undef LOCTEXT_NAMESPACE
+
diff --git a/Engine/Plugins/Cameras/GameplayCameras/Source/GameplayCamerasEditor/Private/Editors/SObjectTreeGraphNode.cpp b/Engine/Plugins/Cameras/GameplayCameras/Source/GameplayCamerasEditor/Private/Editors/SObjectTreeGraphNode.cpp
new file mode 100644
--- /dev/null
+++ b/Engine/Plugins/Cameras/GameplayCameras/Source/GameplayCamerasEditor/Private/Editors/SObjectTreeGraphNode.cpp
@@ -0,0 +1,31 @@
+// Copyright Epic Games, Inc. All Rights Reserved.
+
+#include "Editors/SObjectTreeGraphNode.h"
+
+#include "Editors/ObjectTreeGraphNode.h"
+
+void SObjectTreeGraphNode::Construct(const FArguments& InArgs)
+{
+	GraphNode = InArgs._GraphNode;
+	ObjectGraphNode = InArgs._GraphNode;
+
+	SetCursor(EMouseCursor::CardinalCross);
+
+	UpdateGraphNode();
+}
+
+void SObjectTreeGraphNode::MoveTo(const FVector2D& NewPosition, FNodeSet& NodeFilter, bool bMarkDirty)
+{
+	SGraphNode::MoveTo(NewPosition, NodeFilter, bMarkDirty);
+
+	if (ObjectGraphNode)
+	{
+		ObjectGraphNode->OnGraphNodeMoved();
+	}
+}
+
+void SObjectTreeGraphAddArrayItemPin::Construct(const FArguments& InArgs, UEdGraphPin* InGraphPinObj)
+{
+	SGraphPin::Construct(SGraphPin::FArguments(), InGraphPinObj);
+}
+
diff --git a/Engine/Plugins/Cameras/GameplayCameras/Source/GameplayCamerasEditor/Private/Editors/SObjectTreeGraphTitleBar.cpp b/Engine/Plugins/Cameras/GameplayCameras/Source/GameplayCamerasEditor/Private/Editors/SObjectTreeGraphTitleBar.cpp
new file mode 100644
--- /dev/null
+++ b/Engine/Plugins/Cameras/GameplayCameras/Source/GameplayCamerasEditor/Private/Editors/SObjectTreeGraphTitleBar.cpp
@@ -0,0 +1,230 @@
+// Copyright Epic Games, Inc. All Rights Reserved.
+
+#include "Editors/SObjectTreeGraphTitleBar.h"
+
+#include "EdGraph/EdGraph.h"
+#include "EdGraph/EdGraphSchema.h"
+#include "Styling/AppStyle.h"
+#include "Styling/CoreStyle.h"
+#include "Widgets/Images/SImage.h"
+#include "Widgets/Layout/SBorder.h"
+#include "Widgets/Layout/SScrollBox.h"
+#include "Widgets/Layout/SScrollBox.h"
+#include "Widgets/Layout/SSeparator.h"
+#include "Widgets/Layout/SSpacer.h"
+#include "Widgets/Navigation/SBreadcrumbTrail.h"
+#include "Widgets/SBoxPanel.h"
+#include "Widgets/Text/STextBlock.h"
+#include "Widgets/Views/SListView.h"
+#include "Widgets/Views/STableRow.h"
+#include "Widgets/Views/STableViewBase.h"
+
+#define LOCTEXT_NAMESPACE "SObjectTreeGraphTitleBar"
+
+void SObjectTreeGraphTitleBar::Construct(const FArguments& InArgs)
+{
+	Graph = InArgs._Graph;
+	OnBreadcrumbClicked = InArgs._OnBreadcrumbClicked;
+
+	const ISlateStyle& AppStyle = FAppStyle::Get();
+	const FMargin BreadcrumbTrailPadding = FMargin(4.f, 2.f);
+	const FSlateBrush* BreadcrumbButtonImage = FAppStyle::GetBrush("BreadcrumbTrail.Delimiter");
+
+	TSharedRef<SHorizontalBox> HorizontalBox = SNew(SHorizontalBox)
+		+ SHorizontalBox::Slot()
+		.AutoWidth()
+		[
+			SNew(SSeparator)
+			.Orientation(Orient_Vertical)
+		];
+
+	if (InArgs._HistoryNavigationWidget)
+	{
+		// Navigation widget.
+		HorizontalBox->AddSlot()
+		.AutoWidth()
+		[
+			InArgs._HistoryNavigationWidget.ToSharedRef()
+		];
+		HorizontalBox->AddSlot()
+		.AutoWidth()
+		[
+			SNew(SSeparator)
+			.Orientation(Orient_Vertical)
+		];
+	}
+
+	{
+		// Title icon and breadcrumb trail.
+		HorizontalBox->AddSlot()
+		[
+			SNew(SHorizontalBox)
+			+SHorizontalBox::Slot()
+			.FillWidth(1.f)
+			[
+				SNew(SHorizontalBox)
+				+SHorizontalBox::Slot()
+				.AutoWidth()
+				.Padding( 10.0f,5.0f )
+				.VAlign(VAlign_Center)
+				[
+					// Icon.
+					SNew(SImage)
+					.Image(FAppStyle::GetBrush(TEXT("GraphEditor.EventGraph_24x")))
+				]
+				+SHorizontalBox::Slot()
+				.FillWidth(1.f)
+				[
+					// Breadcrumb trail.
+					SAssignNew(BreadcrumbTrailScrollBox, SScrollBox)
+					.Orientation(Orient_Horizontal)
+					.ScrollBarVisibility(EVisibility::Collapsed)
+
+					+SScrollBox::Slot()
+					.Padding(0.f)
+					.VAlign(VAlign_Center)
+					[
+						// Root breadcrumb, defined by title text.
+						SNew(SHorizontalBox)
+						+SHorizontalBox::Slot()
+						.AutoWidth()
+						.VAlign(VAlign_Center)
+						.Padding(BreadcrumbTrailPadding)
+						[
+							SNew(STextBlock)
+							.Text(InArgs._TitleText)
+							.TextStyle(AppStyle, TEXT("GraphBreadcrumbButtonText"))
+							.Visibility(EVisibility::Visible)
+						]
+
+						+SHorizontalBox::Slot()
+						.AutoWidth()
+						.VAlign(VAlign_Center)
+						[
+							SNew(SImage)
+							.Image(BreadcrumbButtonImage)
+							.Visibility(EVisibility::Visible)
+						]
+
+						// Graph name, defined by current graph.
+						+SHorizontalBox::Slot()
+						.AutoWidth()
+						.VAlign(VAlign_Center)
+						[
+							SAssignNew(BreadcrumbTrail, SBreadcrumbTrail<UEdGraph*>)
+							.ButtonStyle(AppStyle, "GraphBreadcrumbButton")
+							.TextStyle(AppStyle, "GraphBreadcrumbButtonText")
+							.ButtonContentPadding(BreadcrumbTrailPadding)
+							.DelimiterImage(BreadcrumbButtonImage)
+							.PersistentBreadcrumbs(true)
+							.OnCrumbClicked(this, &SObjectTreeGraphTitleBar::OnBreadcrumbClickedImpl)
+						]
+					]
+				]
+			]
+		];
+	}
+	
+	ChildSlot
+	[
+		SNew(SVerticalBox)
+		+ SVerticalBox::Slot()
+		.HAlign(HAlign_Fill)
+		.AutoHeight()
+		[
+			SNew(SBorder)
+			.BorderImage(FAppStyle::GetBrush(TEXT("Graph.TitleBackground")))
+			.HAlign(HAlign_Fill)
+			.AddMetaData<FTagMetaData>(FTagMetaData(TEXT("EventGraphTitleBar")))
+			[
+				HorizontalBox
+			]
+		]
+		+SVerticalBox::Slot()
+		.VAlign(VAlign_Top)
+		.AutoHeight()
+		[
+			SNew(SHorizontalBox)
+			+SHorizontalBox::Slot()
+			.HAlign(HAlign_Center)
+			.AutoWidth()
+			[
+				SAssignNew(GraphListView, SListView<TSharedPtr<FObjectTreeGraphInfo>>)
+				.ListItemsSource(InArgs._GraphList)
+				.OnGenerateRow(this, &SObjectTreeGraphTitleBar::GenerateGraphInfoRow)
+				.SelectionMode( ESelectionMode::None )
+				.Visibility(EVisibility::Collapsed)
+			]
+		]
+	];
+
+	RebuildBreadcrumbTrail();
+	BreadcrumbTrailScrollBox->ScrollToEnd();
+}
+
+TSharedRef<ITableRow> SObjectTreeGraphTitleBar::GenerateGraphInfoRow(TSharedPtr<FObjectTreeGraphInfo> Item, const TSharedRef<STableViewBase>& OwnerTable)
+{
+	const int32 FontSize = 9;
+
+	const FObjectTreeGraphInfo* GraphInfo = Item.Get();
+
+	if (GraphInfo->GraphName.IsEmpty())
+	{
+		return SNew(STableRow<TSharedPtr<FObjectTreeGraphInfo>>, OwnerTable)
+			[
+				SNew(SSpacer)
+			];
+	}
+	else 
+	{
+		return SNew(STableRow<TSharedPtr<FObjectTreeGraphInfo>>, OwnerTable)
+			[
+				SNew(STextBlock)
+				.Font(FCoreStyle::GetDefaultFontStyle("Regular", FontSize))
+				.Text(FText::FromString(GraphInfo->GraphName))
+			];
+	}
+}
+
+void SObjectTreeGraphTitleBar::OnBreadcrumbClickedImpl(UEdGraph* const& Item)
+{
+	OnBreadcrumbClicked.ExecuteIfBound(Item);
+}
+
+void SObjectTreeGraphTitleBar::RebuildBreadcrumbTrail()
+{
+	TArray<UEdGraph*> GraphStack;
+	for (UEdGraph* OuterChain = Graph; OuterChain != nullptr; OuterChain = UEdGraph::GetOuterGraph(OuterChain))
+	{
+		GraphStack.Push(OuterChain);
+	}
+
+	BreadcrumbTrail->ClearCrumbs(false);
+
+	const UEdGraph* LastGraph = (GraphStack.Num() > 0 ? GraphStack.Last() : nullptr);
+
+	while (GraphStack.Num() > 0)
+	{
+		UEdGraph* CurGraph = GraphStack.Pop();
+		
+		TAttribute<FText> TitleText = TAttribute<FText>::Create(
+				TAttribute<FText>::FGetter::CreateStatic(&SObjectTreeGraphTitleBar::GetTitleForOneCrumb, LastGraph, (const UEdGraph*)CurGraph));
+		BreadcrumbTrail->PushCrumb(TitleText, CurGraph);
+	}
+}
+
+FText SObjectTreeGraphTitleBar::GetTitleForOneCrumb(const UEdGraph* BaseGraph, const UEdGraph* CurGraph)
+{
+	const UEdGraphSchema* Schema = CurGraph->GetSchema();
+
+	FGraphDisplayInfo DisplayInfo;
+	Schema->GetGraphDisplayInformation(*CurGraph, DisplayInfo);
+
+	FFormatNamedArguments Args;
+	Args.Add(TEXT("BreadcrumbDisplayName"), DisplayInfo.DisplayName);
+	Args.Add(TEXT("BreadcrumbNotes"), FText::FromString(DisplayInfo.GetNotesAsString()));
+	return FText::Format(LOCTEXT("BreadcrumbTitle", "{BreadcrumbDisplayName} {BreadcrumbNotes}"), Args);
+}
+
+#undef LOCTEXT_NAMESPACE
+
diff --git a/Engine/Plugins/Cameras/GameplayCameras/Source/GameplayCamerasEditor/Private/Editors/SObjectTreeGraphToolbox.cpp b/Engine/Plugins/Cameras/GameplayCameras/Source/GameplayCamerasEditor/Private/Editors/SObjectTreeGraphToolbox.cpp
new file mode 100644
--- /dev/null
+++ b/Engine/Plugins/Cameras/GameplayCameras/Source/GameplayCamerasEditor/Private/Editors/SObjectTreeGraphToolbox.cpp
@@ -0,0 +1,292 @@
+// Copyright Epic Games, Inc. All Rights Reserved.
+
+#include "Editors/SObjectTreeGraphToolbox.h"
+
+#include "Algo/AnyOf.h"
+#include "DragAndDrop/AssetDragDropOp.h"
+#include "EditorClassUtils.h"
+#include "Editors/ObjectTreeGraphConfig.h"
+#include "Styles/ObjectTreeGraphEditorStyle.h"
+#include "UObject/UObjectIterator.h"
+#include "Widgets/Input/SSearchBox.h"
+#include "Widgets/Layout/SSpacer.h"
+#include "Widgets/SBoxPanel.h"
+#include "Widgets/Views/SListView.h"
+
+#define LOCTEXT_NAMESPACE "SObjectTreeGraphToolbox"
+
+void SObjectTreeGraphToolboxEntry::Construct(const FArguments& InArgs)
+{
+	ObjectClass = InArgs._ObjectClass;
+
+	if (ObjectClass)
+	{
+		if (InArgs._GraphConfig)
+		{
+			DisplayNameText = InArgs._GraphConfig->GetDisplayNameText(ObjectClass);
+		}
+		else
+		{
+			DisplayNameText = ObjectClass->GetDisplayNameText();
+		}
+	}
+
+	TSharedRef<FObjectTreeGraphEditorStyle> ObjectTreeStyle = FObjectTreeGraphEditorStyle::Get();
+
+	TSharedPtr<IToolTip> EntryToolTip = FEditorClassUtils::GetTooltip(ObjectClass);
+	TSharedRef<SWidget> DocWidget = FEditorClassUtils::GetDocumentationLinkWidget(ObjectClass);
+
+	const FButtonStyle& ButtonStyle = ObjectTreeStyle->GetWidgetStyle<FButtonStyle>("ObjectTreeGraphToolbox.Entry");
+	NormalImage = &ButtonStyle.Normal;
+	HoverImage = &ButtonStyle.Hovered;
+	PressedImage = &ButtonStyle.Pressed; 
+
+	ChildSlot
+	.Padding(FMargin(8.f, 2.f, 12.f, 2.f))
+	[
+		SNew(SOverlay)
+		+SOverlay::Slot()
+		[
+			SNew(SBorder)
+			.BorderImage(ObjectTreeStyle->GetBrush("ObjectTreeGraphToolbox.Entry.Background"))
+			.Cursor(EMouseCursor::GrabHand)
+			.ToolTip(EntryToolTip)
+			.Padding(0)
+			[
+				SNew( SHorizontalBox )
+
+				+ SHorizontalBox::Slot()
+				.VAlign(VAlign_Fill)
+				.Padding(0)
+				[
+					SNew(SBorder)
+					.BorderImage(ObjectTreeStyle->GetBrush("ObjectTreeGraphToolbox.Entry.LabelBack"))
+					[
+						SNew(SHorizontalBox)
+						+SHorizontalBox::Slot()
+						.Padding(8.f, 4.f)
+						.VAlign(VAlign_Center)
+						[
+							SNew(STextBlock)
+							.TextStyle(ObjectTreeStyle, "ObjectTreeGraphToolbox.Entry.Name")
+							.Text(DisplayNameText)
+							.HighlightText(InArgs._HighlightText)
+						]
+
+						+ SHorizontalBox::Slot()
+						.VAlign(VAlign_Center)
+						.AutoWidth()
+						[
+							DocWidget
+						]
+					]
+				]
+			]
+		]
+		+SOverlay::Slot()
+		[
+			SNew(SBorder)
+			.BorderImage(this, &SObjectTreeGraphToolboxEntry::GetBorder)
+			.Cursor(EMouseCursor::GrabHand)
+			.ToolTip(EntryToolTip )
+		]
+	];
+}
+
+FReply SObjectTreeGraphToolboxEntry::OnMouseButtonDown(const FGeometry& MyGeometry, const FPointerEvent& MouseEvent)
+{
+	if (MouseEvent.GetEffectingButton() == EKeys::LeftMouseButton)
+	{
+		bIsPressed = true;
+
+		return FReply::Handled().DetectDrag(SharedThis(this), MouseEvent.GetEffectingButton());
+	}
+
+	return FReply::Unhandled();
+}
+
+FReply SObjectTreeGraphToolboxEntry::OnMouseButtonUp(const FGeometry& MyGeometry, const FPointerEvent& MouseEvent)
+{
+	if (MouseEvent.GetEffectingButton() == EKeys::LeftMouseButton)
+	{
+		bIsPressed = false;
+	}
+
+	return FReply::Unhandled();
+}
+
+FReply SObjectTreeGraphToolboxEntry::OnDragDetected(const FGeometry& MyGeometry, const FPointerEvent& MouseEvent)
+{
+	bIsPressed = false;
+
+	if (MouseEvent.IsMouseButtonDown(EKeys::LeftMouseButton))
+	{
+		return FReply::Handled().BeginDragDrop(FObjectTreeClassDragDropOp::New(ObjectClass));
+	}
+	else
+	{
+		return FReply::Handled();
+	}
+}
+
+const FSlateBrush* SObjectTreeGraphToolboxEntry::GetBorder() const
+{
+	if (bIsPressed)
+	{
+		return PressedImage;
+	}
+	else if (IsHovered())
+	{
+		return HoverImage;
+	}
+	else
+	{
+		return NormalImage;
+	}
+}
+
+void SObjectTreeGraphToolbox::Construct(const FArguments& InArgs)
+{
+	GraphConfig = InArgs._GraphConfig;
+
+	SearchTextFilter = MakeShareable(new FEntryTextFilter(
+		FEntryTextFilter::FItemToStringArray::CreateSP(this, &SObjectTreeGraphToolbox::GetEntryStrings)));
+
+	ChildSlot
+	[
+		SNew(SVerticalBox)
+		+SVerticalBox::Slot()
+		.AutoHeight()
+		[
+			SNew(SBorder)
+			.BorderImage(FAppStyle::GetBrush("ToolPanel.GroupBorder"))
+			.Padding(8.f)
+			[
+				SAssignNew(SearchBox, SSearchBox)
+				.HintText(LOCTEXT("SearchHint", "Search"))
+				.OnTextChanged(this, &SObjectTreeGraphToolbox::OnSearchTextChanged)
+				.OnTextCommitted(this, &SObjectTreeGraphToolbox::OnSearchTextCommitted)
+			]
+		]
+		+SVerticalBox::Slot()
+		.Padding(0.f, 3.f)
+		[
+			SAssignNew(ListView, SListView<UClass*>)
+				.ListItemsSource(&FilteredItemSource)
+				.OnGenerateRow(this, &SObjectTreeGraphToolbox::OnGenerateItemRow)
+		]
+	];
+
+	bUpdateItemSource = true;
+	bUpdateFilteredItemSource = true;
+}
+
+void SObjectTreeGraphToolbox::GetEntryStrings(const UClass* InItem, TArray<FString>& OutStrings)
+{
+	FText DisplayNameText = GraphConfig.GetDisplayNameText(InItem);
+	OutStrings.Add(DisplayNameText.ToString());
+}
+
+void SObjectTreeGraphToolbox::SetGraphConfig(const FObjectTreeGraphConfig& InGraphConfig)
+{
+	GraphConfig = InGraphConfig;
+	bUpdateItemSource = true;
+}
+
+void SObjectTreeGraphToolbox::Tick(const FGeometry& AllottedGeometry, const double InCurrentTime, const float InDeltaTime)
+{
+	if (bUpdateItemSource)
+	{
+		UpdateItemSource();
+	}
+
+	if (bUpdateItemSource || bUpdateFilteredItemSource)
+	{
+		UpdateFilteredItemSource();
+	}
+
+	const bool bRequestListRefresh = bUpdateItemSource || bUpdateFilteredItemSource;
+	bUpdateItemSource = false;
+	bUpdateFilteredItemSource = false;
+
+	if (bRequestListRefresh)
+	{
+		ListView->RequestListRefresh();
+	}
+
+	SCompoundWidget::Tick(AllottedGeometry, InCurrentTime, InDeltaTime);
+}
+
+void SObjectTreeGraphToolbox::UpdateItemSource()
+{
+	ItemSource.Reset();
+	GraphConfig.GetConnectableClasses(ItemSource, true);
+	ItemSource.Sort([](UClass& A, UClass& B) { return A.GetFName().Compare(B.GetFName()) < 0; });
+}
+
+void SObjectTreeGraphToolbox::UpdateFilteredItemSource()
+{
+	FilteredItemSource.Reset();
+
+	if (!SearchTextFilter->GetRawFilterText().IsEmpty())
+	{
+		for (UClass* Item : ItemSource)
+		{
+			if (SearchTextFilter->PassesFilter(Item))
+			{
+				FilteredItemSource.Add(Item);
+			}
+		}
+	}
+	else
+	{
+		FilteredItemSource = ItemSource;
+	}
+}
+
+TSharedRef<ITableRow> SObjectTreeGraphToolbox::OnGenerateItemRow(UClass* Item, const TSharedRef<STableViewBase>& OwnerTable)
+{
+	return SNew(STableRow<UClass*>, OwnerTable)
+		[
+			SNew(SObjectTreeGraphToolboxEntry)
+				.ObjectClass(Item)
+				.GraphConfig(&GraphConfig)
+				.HighlightText(this, &SObjectTreeGraphToolbox::GetHighlightText)
+		];
+}
+
+void SObjectTreeGraphToolbox::OnSearchTextChanged(const FText& InFilterText)
+{
+	SearchTextFilter->SetRawFilterText(InFilterText);
+	SearchBox->SetError(SearchTextFilter->GetFilterErrorText());
+
+	bUpdateFilteredItemSource = true;
+}
+
+void SObjectTreeGraphToolbox::OnSearchTextCommitted(const FText& InFilterText, ETextCommit::Type InCommitType)
+{
+	OnSearchTextChanged(InFilterText);
+}
+
+FText SObjectTreeGraphToolbox::GetHighlightText() const
+{
+	return SearchTextFilter->GetRawFilterText();
+}
+
+TSharedRef<FObjectTreeClassDragDropOp> FObjectTreeClassDragDropOp::New(UClass* ObjectClass)
+{
+	TArray<UClass*> ObjectClasses;
+	ObjectClasses.Add(ObjectClass);
+	return New(ObjectClasses);
+}
+
+TSharedRef<FObjectTreeClassDragDropOp> FObjectTreeClassDragDropOp::New(TArrayView<UClass*> ObjectClasses)
+{
+	TSharedRef<FObjectTreeClassDragDropOp> Operation = MakeShared<FObjectTreeClassDragDropOp>();
+	Operation->ObjectClasses = ObjectClasses;
+	Operation->Construct();
+	return Operation;
+}
+
+#undef LOCTEXT_NAMESPACE
+
diff --git a/Engine/Plugins/Cameras/GameplayCameras/Source/GameplayCamerasEditor/Private/GameplayCamerasEditorModule.cpp b/Engine/Plugins/Cameras/GameplayCameras/Source/GameplayCamerasEditor/Private/GameplayCamerasEditorModule.cpp
--- a/Engine/Plugins/Cameras/GameplayCameras/Source/GameplayCamerasEditor/Private/GameplayCamerasEditorModule.cpp
+++ b/Engine/Plugins/Cameras/GameplayCameras/Source/GameplayCamerasEditor/Private/GameplayCamerasEditorModule.cpp
@@ -170,9 +170,9 @@ private:
 
 		if (SettingsModule != nullptr)
 		{
-			SettingsModule->RegisterSettings("Project", "Plugins", "GameplayCamerasEditor",
-				LOCTEXT("GameplayCamerasEditorProjectSettingsName", "Gameplay Cameras Editor"),
-				LOCTEXT("GameplayCamerasEditorProjectSettingsDescription", "Configure the gameplay cameras editor."),
+			SettingsModule->RegisterSettings("Editor", "Plugins", "Gameplay Cameras",
+				LOCTEXT("GameplayCamerasEditorProjectSettingsName", "Gameplay Cameras"),
+				LOCTEXT("GameplayCamerasEditorProjectSettingsDescription", "Configure the gameplay cameras editors."),
 				GetMutableDefault<UGameplayCamerasEditorSettings>()
 			);
 		}
@@ -184,8 +184,7 @@ private:
 
 		if (SettingsModule != nullptr)
 		{
-			SettingsModule->UnregisterSettings("Project", "Plugins", "GameplayCamerasEditor");
-			SettingsModule->UnregisterSettings("Editor", "ContentEditors", "GameplayCamerasEditor");
+			SettingsModule->UnregisterSettings("Editor", "Plugins", "Gameplay Cameras");
 		}
 	}
 
diff --git a/Engine/Plugins/Cameras/GameplayCameras/Source/GameplayCamerasEditor/Private/GameplayCamerasEditorSettings.cpp b/Engine/Plugins/Cameras/GameplayCameras/Source/GameplayCamerasEditor/Private/GameplayCamerasEditorSettings.cpp
--- a/Engine/Plugins/Cameras/GameplayCameras/Source/GameplayCamerasEditor/Private/GameplayCamerasEditorSettings.cpp
+++ b/Engine/Plugins/Cameras/GameplayCameras/Source/GameplayCamerasEditor/Private/GameplayCamerasEditorSettings.cpp
@@ -7,5 +7,10 @@
 UGameplayCamerasEditorSettings::UGameplayCamerasEditorSettings(const FObjectInitializer& ObjectInitializer)
 	: Super(ObjectInitializer)
 {
+	CameraNodeTitleColor = FLinearColor(0.190525f, 0.583898f, 1.0f, 1.0f);           // Light blue
+	CameraRigAssetTitleColor = FLinearColor(1.0f, 0.170000f, 0.0f, 1.0f);            // Orange
+	CameraRigTransitionTitleColor = FLinearColor(1.0f, 0.65f, 0.4f, 1.0f);           // Beige
+	CameraRigTransitionConditionTitleColor = FLinearColor(0.8f, 0.4f, 0.4f, 1.0f);   // Salmon
+	CameraBlendNodeTitleColor = FLinearColor(0.6f, 0.0f, 1.0f, 1.0f);                // Purple
 }
 
diff --git a/Engine/Plugins/Cameras/GameplayCameras/Source/GameplayCamerasEditor/Private/Styles/GameplayCamerasEditorStyle.cpp b/Engine/Plugins/Cameras/GameplayCameras/Source/GameplayCamerasEditor/Private/Styles/GameplayCamerasEditorStyle.cpp
--- a/Engine/Plugins/Cameras/GameplayCameras/Source/GameplayCamerasEditor/Private/Styles/GameplayCamerasEditorStyle.cpp
+++ b/Engine/Plugins/Cameras/GameplayCameras/Source/GameplayCamerasEditor/Private/Styles/GameplayCamerasEditorStyle.cpp
@@ -5,20 +5,24 @@
 #include "Brushes/SlateImageBrush.h"
 #include "Interfaces/IPluginManager.h"
 #include "Misc/Paths.h"
+#include "Styling/AppStyle.h"
 #include "Styling/SlateStyleRegistry.h"
+#include "Styling/SlateStyleMacros.h"
+#include "Styling/SlateTypes.h"
+#include "Styling/StyleColors.h"
 
 namespace UE::Cameras
 {
 
-#define IMAGE_BRUSH(RelativePath, ...) FSlateImageBrush(RootToContentDir(RelativePath, TEXT(".png")), __VA_ARGS__)
-
 TSharedPtr<FGameplayCamerasEditorStyle> FGameplayCamerasEditorStyle::Singleton;
 
 FGameplayCamerasEditorStyle::FGameplayCamerasEditorStyle()
 	: FSlateStyleSet("GameplayCamerasEditorStyle")
 {
 	const FVector2D Icon16x16(16.0f, 16.0f);
+	const FVector2D Icon20x20(20.0f, 20.0f);
 	const FVector2D Icon24x24(24.0f, 24.0f);
+	const FVector2D Icon40x40(40.0f, 40.0f);
 	const FVector2D Icon48x48(48.0f, 48.0f);
 	const FVector2D Icon64x64(64.0f, 64.0f);
 
@@ -26,14 +30,37 @@ FGameplayCamerasEditorStyle::FGameplayCamerasEditorStyle()
 	SetContentRoot(ContentDir);
 	SetCoreContentRoot(FPaths::EngineContentDir() / TEXT("Slate"));
 
-	Set("Debugger.EnableDebugInfo.Icon", new IMAGE_BRUSH("/Icons/EnableDebugInfo", Icon16x16));
-	Set("Debugger.DisableDebugInfo.Icon", new IMAGE_BRUSH("/Icons/DisableDebugInfo", Icon16x16));
+	// Camera rig editor icons.
+	Set("CameraRigAssetEditor.Tabs.Toolbox", new IMAGE_BRUSH_SVG("Icons/CameraRig-Toolbox", Icon16x16));
+	Set("CameraRigAssetEditor.Tabs.Search", new CORE_IMAGE_BRUSH_SVG("Starship/Common/search", Icon16x16));
+	Set("CameraRigAssetEditor.Tabs.Messages", new CORE_IMAGE_BRUSH_SVG("Starship/Common/OutputLog", Icon16x16));
+	Set("CameraRigAssetEditor.Tabs.NodeHierarchy", new IMAGE_BRUSH_SVG("Icons/CameraRig-NodeHierarchy", Icon16x16));
+	Set("CameraRigAssetEditor.Tabs.Transitions", new IMAGE_BRUSH_SVG("Icons/CameraRig-Transitions", Icon16x16));
+
+	Set("CameraRigAssetEditor.Build", new IMAGE_BRUSH_SVG("Icons/CameraRig-BuildStatus_Background", Icon20x20));
+	Set("CameraRigAssetEditor.BuildStatus.Background", new IMAGE_BRUSH_SVG("Icons/CameraRig-BuildStatus_Background", Icon20x20));
+	Set("CameraRigAssetEditor.BuildStatus.Overlay.Fail", new IMAGE_BRUSH_SVG("Icons/CameraRig-BuildStatus_Fail_Badge", Icon20x20, FStyleColors::Error));
+	Set("CameraRigAssetEditor.BuildStatus.Overlay.Good", new IMAGE_BRUSH_SVG("Icons/CameraRig-BuildStatus_Good_Badge", Icon20x20, FStyleColors::AccentGreen));
+	Set("CameraRigAssetEditor.BuildStatus.Overlay.Unknown", new IMAGE_BRUSH_SVG("Icons/CameraRig-BuildStatus_Unknown_Badge", Icon20x20, FStyleColors::AccentYellow));
+	Set("CameraRigAssetEditor.BuildStatus.Overlay.Warning", new IMAGE_BRUSH_SVG("Icons/CameraRig-BuildStatus_Warning_Badge", Icon20x20, FStyleColors::Warning));
+
+	Set("CameraRigAssetEditor.ShowNodeHierarchy", new IMAGE_BRUSH_SVG("Icons/CameraRig-NodeHierarchy", Icon20x20));
+	Set("CameraRigAssetEditor.ShowTransitions", new IMAGE_BRUSH_SVG("Icons/CameraRig-Transitions", Icon20x20));
+	Set("CameraRigAssetEditor.FocusHome", new IMAGE_BRUSH_SVG("Icons/GraphEditor-Home", Icon20x20));
 
-	Set("DebugCategory.NodeTree.Icon", new IMAGE_BRUSH("/Icons/DebugCategory-NodeTree", Icon16x16));
-	Set("DebugCategory.DirectorTree.Icon", new IMAGE_BRUSH("/Icons/DebugCategory-DirectorTree", Icon16x16));
-	Set("DebugCategory.BlendStacks.Icon", new IMAGE_BRUSH("/Icons/DebugCategory-BlendStacks", Icon16x16));
-	Set("DebugCategory.PoseStats.Icon", new IMAGE_BRUSH("/Icons/DebugCategory-PoseStats", Icon16x16));
-	Set("DebugCategory.Viewfinder.Icon", new IMAGE_BRUSH("/Icons/DebugCategory-Viewfinder", Icon16x16));
+	Set("CameraRigAssetEditor.LiveUpdate", new IMAGE_BRUSH_SVG("Icons/CameraRig-LiveUpdate", Icon20x20));
+	Set("CameraRigAssetEditor.Apply", new IMAGE_BRUSH_SVG("CameraEditor-Apply", Icon20x20));
+	Set("CameraRigAssetEditor.FindInCameraRig", new CORE_IMAGE_BRUSH_SVG("Starship/Common/Search", Icon20x20));
+
+	// Debugger tool icons.
+	Set("Debugger.EnableDebugInfo.Icon", new IMAGE_BRUSH("Icons/EnableDebugInfo", Icon16x16));
+	Set("Debugger.DisableDebugInfo.Icon", new IMAGE_BRUSH("Icons/DisableDebugInfo", Icon16x16));
+
+	Set("DebugCategory.NodeTree.Icon", new IMAGE_BRUSH("Icons/DebugCategory-NodeTree", Icon16x16));
+	Set("DebugCategory.DirectorTree.Icon", new IMAGE_BRUSH("Icons/DebugCategory-DirectorTree", Icon16x16));
+	Set("DebugCategory.BlendStacks.Icon", new IMAGE_BRUSH("Icons/DebugCategory-BlendStacks", Icon16x16));
+	Set("DebugCategory.PoseStats.Icon", new IMAGE_BRUSH("Icons/DebugCategory-PoseStats", Icon16x16));
+	Set("DebugCategory.Viewfinder.Icon", new IMAGE_BRUSH("Icons/DebugCategory-Viewfinder", Icon16x16));
 
 	FSlateStyleRegistry::RegisterSlateStyle(*this);
 }
@@ -52,7 +79,5 @@ TSharedRef<FGameplayCamerasEditorStyle> FGameplayCamerasEditorStyle::Get()
 	return Singleton.ToSharedRef();
 }
 
-#undef IMAGE_BRUSH
-
 }  // namespace UE::Cameras
 
diff --git a/Engine/Plugins/Cameras/GameplayCameras/Source/GameplayCamerasEditor/Private/Styles/ObjectTreeGraphEditorStyle.cpp b/Engine/Plugins/Cameras/GameplayCameras/Source/GameplayCamerasEditor/Private/Styles/ObjectTreeGraphEditorStyle.cpp
new file mode 100644
--- /dev/null
+++ b/Engine/Plugins/Cameras/GameplayCameras/Source/GameplayCamerasEditor/Private/Styles/ObjectTreeGraphEditorStyle.cpp
@@ -0,0 +1,77 @@
+// Copyright Epic Games, Inc. All Rights Reserved.
+
+#include "Styles/ObjectTreeGraphEditorStyle.h"
+
+#include "Brushes/SlateImageBrush.h"
+#include "Interfaces/IPluginManager.h"
+#include "Misc/Paths.h"
+#include "Styling/AppStyle.h"
+#include "Styling/CoreStyle.h"
+#include "Styling/SlateStyleRegistry.h"
+#include "Styling/SlateStyleMacros.h"
+#include "Styling/SlateTypes.h"
+#include "Styling/StyleColors.h"
+
+TSharedPtr<FObjectTreeGraphEditorStyle> FObjectTreeGraphEditorStyle::Singleton;
+
+FObjectTreeGraphEditorStyle::FObjectTreeGraphEditorStyle()
+	: FSlateStyleSet("ObjectTreeGraphEditorStyle")
+{
+	const FVector2D Icon16x16(16.0f, 16.0f);
+	const FVector2D Icon20x20(20.0f, 20.0f);
+	const FVector2D Icon24x24(24.0f, 24.0f);
+	const FVector2D Icon40x40(40.0f, 40.0f);
+	const FVector2D Icon48x48(48.0f, 48.0f);
+	const FVector2D Icon64x64(64.0f, 64.0f);
+
+	const FString ContentDir = IPluginManager::Get().FindPlugin(TEXT("GameplayCameras"))->GetContentDir();
+	SetContentRoot(ContentDir);
+	SetCoreContentRoot(FPaths::EngineContentDir() / TEXT("Slate"));
+
+	const FButtonStyle& DefaultButton = FAppStyle::Get().GetWidgetStyle<FButtonStyle>("Button");
+	const FTextBlockStyle& NormalText = FAppStyle::Get().GetWidgetStyle<FTextBlockStyle>("NormalText");
+
+	// Object tree graph toolbox styles.
+	Set("ObjectTreeGraphToolbox.Entry", FButtonStyle(DefaultButton)
+		.SetNormal(FSlateRoundedBoxBrush(FLinearColor::Transparent, 6.0f, FStyleColors::Dropdown, 1.0f))
+		.SetHovered(FSlateRoundedBoxBrush(FLinearColor::Transparent, 6.0f, FStyleColors::Hover, 1.0f))
+		.SetPressed(FSlateRoundedBoxBrush(FLinearColor::Transparent, 6.0f, FStyleColors::Primary, 1.0f))
+		.SetNormalPadding(0)
+		.SetPressedPadding(0));
+
+	Set("ObjectTreeGraphToolbox.Entry.Name", FTextBlockStyle(NormalText)
+		.SetFont(DEFAULT_FONT("Regular", 10))
+		.SetColorAndOpacity(FLinearColor(1.0f, 1.0f, 1.0f, 0.9f))
+		.SetShadowOffset(FVector2D(1, 1))
+		.SetShadowColorAndOpacity(FLinearColor(0, 0, 0, 0.9f)));
+
+	Set("ObjectTreeGraphToolbox.Entry.Type", FTextBlockStyle(NormalText)
+		.SetFont(DEFAULT_FONT("Regular", 8))
+		.SetColorAndOpacity(FLinearColor(0.8f, 0.8f, 0.8f, 0.9f))
+		.SetShadowOffset(FVector2D(1, 1))
+		.SetShadowColorAndOpacity(FLinearColor(0, 0, 0, 0.9f)));
+
+	Set("ObjectTreeGraphToolbox.Entry.Background", new FSlateRoundedBoxBrush(FStyleColors::Recessed, 6.f));
+	Set("ObjectTreeGraphToolbox.Entry.LabelBack", new BOX_BRUSH("Icons/Toolbox-LabelBack", 6.f/18.f, FStyleColors::Dropdown));
+
+	Set("ObjectTreeGraphToolbox.EntryToolTip.Name", FTextBlockStyle(NormalText) .SetFont(DEFAULT_FONT("Bold", 9)));
+	Set("ObjectTreeGraphToolbox.EntryToolTip.ClassName", FTextBlockStyle(NormalText) .SetFont(DEFAULT_FONT("Regular", 9)));
+	Set("ObjectTreeGraphToolbox.EntryToolTip.Path", FTextBlockStyle(NormalText) .SetFont(DEFAULT_FONT("Regular", 8)));
+
+	FSlateStyleRegistry::RegisterSlateStyle(*this);
+}
+
+FObjectTreeGraphEditorStyle::~FObjectTreeGraphEditorStyle()
+{
+	FSlateStyleRegistry::UnRegisterSlateStyle(*this);
+}
+
+TSharedRef<FObjectTreeGraphEditorStyle> FObjectTreeGraphEditorStyle::Get()
+{
+	if (!Singleton.IsValid())
+	{
+		Singleton = MakeShareable(new FObjectTreeGraphEditorStyle);
+	}
+	return Singleton.ToSharedRef();
+}
+
diff --git a/Engine/Plugins/Cameras/GameplayCameras/Source/GameplayCamerasEditor/Private/Styles/ObjectTreeGraphEditorStyle.h b/Engine/Plugins/Cameras/GameplayCameras/Source/GameplayCamerasEditor/Private/Styles/ObjectTreeGraphEditorStyle.h
new file mode 100644
--- /dev/null
+++ b/Engine/Plugins/Cameras/GameplayCameras/Source/GameplayCamerasEditor/Private/Styles/ObjectTreeGraphEditorStyle.h
@@ -0,0 +1,23 @@
+// Copyright Epic Games, Inc. All Rights Reserved.
+
+#pragma once
+
+#include "Styling/SlateStyle.h"
+
+/**
+ * Implements the visual style of generic object tree graph editors
+ */
+class FObjectTreeGraphEditorStyle final : public FSlateStyleSet
+{
+public:
+	
+	FObjectTreeGraphEditorStyle();
+	virtual ~FObjectTreeGraphEditorStyle();
+
+	static TSharedRef<FObjectTreeGraphEditorStyle> Get();
+
+private:
+
+	static TSharedPtr<FObjectTreeGraphEditorStyle> Singleton;
+};
+
diff --git a/Engine/Plugins/Cameras/GameplayCameras/Source/GameplayCamerasEditor/Private/Toolkits/CameraAssetEditorToolkit.cpp b/Engine/Plugins/Cameras/GameplayCameras/Source/GameplayCamerasEditor/Private/Toolkits/CameraAssetEditorToolkit.cpp
--- a/Engine/Plugins/Cameras/GameplayCameras/Source/GameplayCamerasEditor/Private/Toolkits/CameraAssetEditorToolkit.cpp
+++ b/Engine/Plugins/Cameras/GameplayCameras/Source/GameplayCamerasEditor/Private/Toolkits/CameraAssetEditorToolkit.cpp
@@ -79,6 +79,8 @@ void FCameraAssetEditorToolkit::CreateWidgets()
 	DetailsViewArgs.bHideSelectionTip = true;
 	DetailsView = PropertyEditorModule.CreateDetailView(DetailsViewArgs);
 
+	// Now do our custom stuff.
+	
 	// Create the message log.
 	FMessageLogModule& MessageLogModule = FModuleManager::LoadModuleChecked<FMessageLogModule>("MessageLog");
 	FMessageLogInitializationOptions LogOptions;
diff --git a/Engine/Plugins/Cameras/GameplayCameras/Source/GameplayCamerasEditor/Private/Toolkits/CameraRigAssetEditorToolkit.cpp b/Engine/Plugins/Cameras/GameplayCameras/Source/GameplayCamerasEditor/Private/Toolkits/CameraRigAssetEditorToolkit.cpp
--- a/Engine/Plugins/Cameras/GameplayCameras/Source/GameplayCamerasEditor/Private/Toolkits/CameraRigAssetEditorToolkit.cpp
+++ b/Engine/Plugins/Cameras/GameplayCameras/Source/GameplayCamerasEditor/Private/Toolkits/CameraRigAssetEditorToolkit.cpp
@@ -7,7 +7,14 @@
 #include "Core/CameraRigAllocationInfoBuilder.h"
 #include "Core/CameraRigAsset.h"
 #include "EditorModeManager.h"
+#include "Editors/ObjectTreeGraph.h"
+#include "Editors/ObjectTreeGraphConfig.h"
+#include "Editors/SCameraRigAssetEditor.h"
+#include "Editors/SFindInObjectTreeGraph.h"
+#include "Editors/SObjectTreeGraphEditor.h"
+#include "Editors/SObjectTreeGraphToolbox.h"
 #include "Framework/Docking/LayoutExtender.h"
+#include "Framework/Docking/TabManager.h"
 #include "IGameplayCamerasLiveEditManager.h"
 #include "IGameplayCamerasModule.h"
 #include "IMessageLogListing.h"
@@ -15,6 +22,7 @@
 #include "MessageLogModule.h"
 #include "Modules/ModuleManager.h"
 #include "PropertyEditorModule.h"
+#include "Styles/GameplayCamerasEditorStyle.h"
 #include "Widgets/Docking/SDockTab.h"
 
 #include UE_INLINE_GENERATED_CPP_BY_NAME(CameraRigAssetEditorToolkit)
@@ -24,28 +32,52 @@
 namespace UE::Cameras
 {
 
+const FName FCameraRigAssetEditorToolkit::ToolboxTabId(TEXT("CameraRigAssetEditor_Toolbox"));
+const FName FCameraRigAssetEditorToolkit::CameraRigEditorTabId(TEXT("CameraRigAssetEditor_CameraRigEditor"));
+const FName FCameraRigAssetEditorToolkit::SearchTabId(TEXT("CameraRigAssetEditor_Search"));
+const FName FCameraRigAssetEditorToolkit::MessagesTabId(TEXT("CameraRigAssetEditor_Messages"));
 const FName FCameraRigAssetEditorToolkit::DetailsViewTabId(TEXT("CameraRigAssetEditor_DetailsView"));
 
 FCameraRigAssetEditorToolkit::FCameraRigAssetEditorToolkit(UCameraRigAssetEditor* InOwningAssetEditor)
 	: FBaseAssetToolkit(InOwningAssetEditor)
-	, CommandBindings(new FUICommandList())
 	, CameraRigAsset(InOwningAssetEditor->GetCameraRigAsset())
+	, CommandBindings(new FUICommandList())
 {
 	// Override base class default layout.
-	StandaloneDefaultLayout = FTabManager::NewLayout("CameraRigAssetEditor_Layout")
+	StandaloneDefaultLayout = FTabManager::NewLayout("CameraRigAssetEditor_Layout_v5")
 		->AddArea
 		(
 			FTabManager::NewPrimaryArea()->SetOrientation(Orient_Vertical)
 			->Split
 			(
-				FTabManager::NewSplitter()->SetOrientation(Orient_Horizontal)->SetSizeCoefficient(0.9f)
+				FTabManager::NewSplitter()->SetOrientation(Orient_Horizontal)
+				->Split
+				(
+					FTabManager::NewStack()
+					->SetSizeCoefficient(0.2f)
+					->AddTab(ToolboxTabId, ETabState::OpenedTab)
+				)
 				->Split
 				(
 					FTabManager::NewStack()
+					->SetSizeCoefficient(0.6f)
+					->AddTab(CameraRigEditorTabId, ETabState::OpenedTab)
+					->SetHideTabWell(true)
+				)
+				->Split
+				(
+					FTabManager::NewStack()
+					->SetSizeCoefficient(0.2f)
 					->AddTab(DetailsViewTabId, ETabState::OpenedTab)
-					->SetForegroundTab(DetailsViewTabId)
 				)
 			)
+			->Split
+			(
+				FTabManager::NewStack()
+				->SetSizeCoefficient(0.2f)
+				->AddTab(SearchTabId, ETabState::ClosedTab)
+				->AddTab(MessagesTabId, ETabState::ClosedTab)
+			)
 		);
 }
 
@@ -58,17 +90,86 @@ void FCameraRigAssetEditorToolkit::RegisterTabSpawners(const TSharedRef<FTabMana
 	// Skip FBaseAssetToolkit here because we don't want a viewport tab.
 	FAssetEditorToolkit::RegisterTabSpawners(InTabManager);
 
+	const FName CamerasStyleSetName = FGameplayCamerasEditorStyle::Get()->GetStyleSetName();
+
+	InTabManager->RegisterTabSpawner(ToolboxTabId, FOnSpawnTab::CreateSP(this, &FCameraRigAssetEditorToolkit::SpawnTab_Toolbox))
+		.SetDisplayName(LOCTEXT("Toolbox", "Toolbox"))
+		.SetGroup(AssetEditorTabsCategory.ToSharedRef())
+		.SetIcon(FSlateIcon(CamerasStyleSetName, "CameraRigAssetEditor.Tabs.Toolbox"));
+
+	InTabManager->RegisterTabSpawner(CameraRigEditorTabId, FOnSpawnTab::CreateSP(this, &FCameraRigAssetEditorToolkit::SpawnTab_CameraRigEditor))
+		.SetDisplayName(LOCTEXT("CameraRigEditor", "Camera Rig"))
+		.SetGroup(AssetEditorTabsCategory.ToSharedRef());
+
+	InTabManager->RegisterTabSpawner(SearchTabId, FOnSpawnTab::CreateSP(this, &FCameraRigAssetEditorToolkit::SpawnTab_Search))
+		.SetDisplayName(LOCTEXT("Search", "Search"))
+		.SetGroup(AssetEditorTabsCategory.ToSharedRef())
+		.SetIcon(FSlateIcon(CamerasStyleSetName, "CameraRigAssetEditor.Tabs.Search"));
+
+	InTabManager->RegisterTabSpawner(MessagesTabId, FOnSpawnTab::CreateSP(this, &FCameraRigAssetEditorToolkit::SpawnTab_Messages))
+		.SetDisplayName(LOCTEXT("Messages", "Messages"))
+		.SetGroup(AssetEditorTabsCategory.ToSharedRef())
+		.SetIcon(FSlateIcon(CamerasStyleSetName, "CameraRigAssetEditor.Tabs.Messages"));
+
 	InTabManager->RegisterTabSpawner(DetailsViewTabId, FOnSpawnTab::CreateSP(this, &FCameraRigAssetEditorToolkit::SpawnTab_Details))
 		.SetDisplayName(LOCTEXT("Details", "Details"))
 		.SetGroup(AssetEditorTabsCategory.ToSharedRef())
 		.SetIcon(FSlateIcon(FAppStyle::GetAppStyleSetName(), "LevelEditor.Tabs.Details"));
 }
 
+TSharedRef<SDockTab> FCameraRigAssetEditorToolkit::SpawnTab_Toolbox(const FSpawnTabArgs& Args)
+{
+	TSharedPtr<SDockTab> ToolboxTab = SNew(SDockTab)
+		.Label(LOCTEXT("ToolboxTabTitle", "Toolbox"))
+		[
+			ToolboxWidget.ToSharedRef()
+		];
+
+	return ToolboxTab.ToSharedRef();
+}
+
+TSharedRef<SDockTab> FCameraRigAssetEditorToolkit::SpawnTab_CameraRigEditor(const FSpawnTabArgs& Args)
+{
+	TSharedPtr<SDockTab> CameraRigEditorTab = SNew(SDockTab)
+		.Label(LOCTEXT("CameraRigEditorTabTitle", "Camera Rig Editor"))
+		[
+			CameraRigEditorWidget.ToSharedRef()
+		];
+
+	return CameraRigEditorTab.ToSharedRef();
+}
+
+TSharedRef<SDockTab> FCameraRigAssetEditorToolkit::SpawnTab_Search(const FSpawnTabArgs& Args)
+{
+	TSharedPtr<SDockTab> SearchTab = SNew(SDockTab)
+		.Label(LOCTEXT("SearchTabTitle", "Search"))
+		[
+			SearchWidget.ToSharedRef()
+		];
+
+	return SearchTab.ToSharedRef();
+}
+
+TSharedRef<SDockTab> FCameraRigAssetEditorToolkit::SpawnTab_Messages(const FSpawnTabArgs& Args)
+{
+	TSharedPtr<SDockTab> MessagesTab = SNew(SDockTab)
+		.Label(LOCTEXT("MessagesTabTitle", "Messages"))
+		[
+			MessagesWidget.ToSharedRef()
+		];
+
+	return MessagesTab.ToSharedRef();
+}
+
 void FCameraRigAssetEditorToolkit::UnregisterTabSpawners(const TSharedRef<FTabManager>& InTabManager)
 {
 	// Skip FBaseAssetToolkit here because we don't want a viewport tab.
 	FAssetEditorToolkit::UnregisterTabSpawners(InTabManager);
 
+	InTabManager->UnregisterTabSpawner(ToolboxTabId);
+	InTabManager->UnregisterTabSpawner(CameraRigEditorTabId);
+	InTabManager->UnregisterTabSpawner(SearchTabId);
+	InTabManager->UnregisterTabSpawner(MessagesTabId);
 	InTabManager->UnregisterTabSpawner(DetailsViewTabId);
 }
 
@@ -88,6 +189,26 @@ void FCameraRigAssetEditorToolkit::CreateWidgets()
 	DetailsViewArgs.NotifyHook = this;
 	DetailsView = PropertyEditorModule.CreateDetailView(DetailsViewArgs);
 
+	// Now do our custom stuff.
+
+	// Create the camera rig editor.
+	CameraRigEditorWidget = SNew(SCameraRigAssetEditor)
+		.DetailsView(DetailsView)
+		.CameraRigAsset(CameraRigAsset);
+
+	// Create the toolbox, default to the rig editor items.
+	ToolboxWidget = SNew(SObjectTreeGraphToolbox)
+		.GraphConfig(CameraRigEditorWidget->GetFocusedGraphConfig());
+
+	// Create the search panel.
+	TArray<UEdGraph*> CameraRigGraphs;
+	CameraRigEditorWidget->GetGraphs(CameraRigGraphs);
+	SearchWidget = SNew(SFindInObjectTreeGraph)
+		.GraphsToSearch(CameraRigGraphs)
+		.OnJumpToNodeRequested(
+				SFindInObjectTreeGraph::FOnJumpToNodeRequested::CreateSP(
+					CameraRigEditorWidget.ToSharedRef(), &SCameraRigAssetEditor::JumpToNode));
+
 	// Create the message log.
 	FMessageLogModule& MessageLogModule = FModuleManager::LoadModuleChecked<FMessageLogModule>("MessageLog");
 	FMessageLogInitializationOptions LogOptions;
@@ -95,9 +216,9 @@ void FCameraRigAssetEditorToolkit::CreateWidgets()
 	LogOptions.bShowFilters = false;
 	LogOptions.bAllowClear = false;
 	LogOptions.MaxPageCount = 1;
-	StatsListing = MessageLogModule.CreateLogListing("CameraRigAssetEditorStats", LogOptions);
+	MessageListing = MessageLogModule.CreateLogListing("CameraRigAssetEditorStats", LogOptions);
 
-	Stats = MessageLogModule.CreateLogListingWidget(StatsListing.ToSharedRef());
+	MessagesWidget = MessageLogModule.CreateLogListingWidget(MessageListing.ToSharedRef());
 }
 
 void FCameraRigAssetEditorToolkit::RegisterToolbar()
@@ -112,20 +233,31 @@ void FCameraRigAssetEditorToolkit::RegisterToolbar()
 		UToolMenu* ToolbarMenu = UToolMenus::Get()->RegisterMenu(
 				MenuName, ParentName, EMultiBoxType::ToolBar);
 
-		ToolbarMenu->AddDynamicSection(TEXT("Build"), FNewToolMenuDelegate::CreateLambda(
-				[](UToolMenu* InMenu)
+		FToolMenuInsert InsertAfterAssetSection("Asset", EToolMenuInsertType::After);
+		const FCameraRigAssetEditorCommands& Commands = FCameraRigAssetEditorCommands::Get();
+
+		ToolbarMenu->AddDynamicSection("Tools", FNewToolMenuDelegate::CreateLambda(
+				[&Commands](UToolMenu* InMenu)
 				{
-					const FCameraRigAssetEditorCommands& Commands = FCameraRigAssetEditorCommands::Get();
 					UCameraRigAssetEditorMenuContext* Context = InMenu->FindContext<UCameraRigAssetEditorMenuContext>();
 					FCameraRigAssetEditorToolkit* This = Context->CameraRigAssetEditorToolkit.Pin().Get();
 
-					FToolMenuSection& BuildSection = InMenu->AddSection(TEXT("Build"));
+					FToolMenuSection& ToolsSection = InMenu->AddSection("Tools");
 
 					FToolMenuEntry BuildButton = FToolMenuEntry::InitToolBarButton(Commands.Build);
 					BuildButton.Icon = TAttribute<FSlateIcon>(This, &FCameraRigAssetEditorToolkit::GetBuildButtonIcon);
 					BuildButton.ToolTip = TAttribute<FText>(This, &FCameraRigAssetEditorToolkit::GetBuildButtonTooltip);
-					BuildSection.AddEntry(BuildButton);
-				}));
+					ToolsSection.AddEntry(BuildButton);
+
+					ToolsSection.AddEntry(FToolMenuEntry::InitToolBarButton(Commands.FindInCameraRig));
+				}),
+				InsertAfterAssetSection);
+
+		FToolMenuSection& GraphsSection = ToolbarMenu->AddSection("Graphs", TAttribute<FText>(), InsertAfterAssetSection);
+
+		GraphsSection.AddEntry(FToolMenuEntry::InitToolBarButton(Commands.FocusHome));
+		GraphsSection.AddEntry(FToolMenuEntry::InitToolBarButton(Commands.ShowNodeHierarchy));
+		GraphsSection.AddEntry(FToolMenuEntry::InitToolBarButton(Commands.ShowTransitions));
 	}
 }
 
@@ -140,30 +272,30 @@ void FCameraRigAssetEditorToolkit::InitToolMenuContext(FToolMenuContext& MenuCon
 
 FSlateIcon FCameraRigAssetEditorToolkit::GetBuildButtonIcon() const
 {
-	static const FName BuildStatusBackground("Blueprint.CompileStatus.Background");
-	static const FName BuildStatusUnknown("Blueprint.CompileStatus.Overlay.Unknown");
-	static const FName BuildStatusError("Blueprint.CompileStatus.Overlay.Error");
-	static const FName BuildStatusGood("Blueprint.CompileStatus.Overlay.Good");
-	static const FName BuildStatusWarning("Blueprint.CompileStatus.Overlay.Warning");
+	static const FName BuildStatusBackground("CameraRigAssetEditor.BuildStatus.Background");
+	static const FName BuildStatusError("CameraRigAssetEditor.BuildStatus.Overlay.Error");
+	static const FName BuildStatusGood("CameraRigAssetEditor.BuildStatus.Overlay.Good");
+	static const FName BuildStatusUnknown("CameraRigAssetEditor.BuildStatus.Overlay.Unknown");
+	static const FName BuildStatusWarning("CameraRigAssetEditor.BuildStatus.Overlay.Warning");
 
-	const FName AppStyle = FAppStyle::GetAppStyleSetName();
+	const FName CamerasStyleSetName = FGameplayCamerasEditorStyle::Get()->GetStyleSetName();
 
 	if (!CameraRigAsset)
 	{
-		return FSlateIcon(AppStyle, BuildStatusBackground, NAME_None, BuildStatusError);
+		return FSlateIcon(CamerasStyleSetName, BuildStatusBackground, NAME_None, BuildStatusError);
 	}
 
 	switch (CameraRigAsset->BuildStatus)
 	{
 		default:
 		case ECameraRigBuildStatus::Dirty:
-			return FSlateIcon(AppStyle, BuildStatusBackground, NAME_None, BuildStatusUnknown);
+			return FSlateIcon(CamerasStyleSetName, BuildStatusBackground, NAME_None, BuildStatusUnknown);
 		case ECameraRigBuildStatus::WithErrors:
-			return FSlateIcon(AppStyle, BuildStatusBackground, NAME_None, BuildStatusError);
+			return FSlateIcon(CamerasStyleSetName, BuildStatusBackground, NAME_None, BuildStatusError);
 		case ECameraRigBuildStatus::Clean:
-			return FSlateIcon(AppStyle, BuildStatusBackground, NAME_None, BuildStatusGood);
+			return FSlateIcon(CamerasStyleSetName, BuildStatusBackground, NAME_None, BuildStatusGood);
 		case ECameraRigBuildStatus::CleanWithWarnings:
-			return FSlateIcon(AppStyle, BuildStatusBackground, NAME_None, BuildStatusWarning);
+			return FSlateIcon(CamerasStyleSetName, BuildStatusBackground, NAME_None, BuildStatusWarning);
 	}
 }
 
@@ -196,6 +328,28 @@ void FCameraRigAssetEditorToolkit::PostInitAssetEditor()
 		Commands.Build,
 		FExecuteAction::CreateSP(this, &FCameraRigAssetEditorToolkit::OnBuild));
 
+	ToolkitCommands->MapAction(
+		Commands.FindInCameraRig,
+		FExecuteAction::CreateSP(this, &FCameraRigAssetEditorToolkit::OnFindInCameraRig));
+
+	TSharedRef<SCameraRigAssetEditor> CameraRigEditor = CameraRigEditorWidget.ToSharedRef();
+
+	ToolkitCommands->MapAction(
+		Commands.FocusHome,
+		FExecuteAction::CreateSP(CameraRigEditor, &SCameraRigAssetEditor::FocusHome));
+
+	ToolkitCommands->MapAction(
+		Commands.ShowNodeHierarchy,
+		FExecuteAction::CreateSP(CameraRigEditor, &SCameraRigAssetEditor::SetEditorMode, ECameraRigAssetEditorMode::NodeGraph),
+		FCanExecuteAction(),
+		FIsActionChecked::CreateSP(CameraRigEditor, &SCameraRigAssetEditor::IsEditorMode, ECameraRigAssetEditorMode::NodeGraph));
+
+	ToolkitCommands->MapAction(
+		Commands.ShowTransitions,
+		FExecuteAction::CreateSP(CameraRigEditor, &SCameraRigAssetEditor::SetEditorMode, ECameraRigAssetEditorMode::TransitionGraph),
+		FCanExecuteAction(),
+		FIsActionChecked::CreateSP(CameraRigEditor, &SCameraRigAssetEditor::IsEditorMode, ECameraRigAssetEditorMode::TransitionGraph));
+
 	IGameplayCamerasModule& GameplayCamerasModule = FModuleManager::GetModuleChecked<IGameplayCamerasModule>("GameplayCameras");
 	LiveEditManager = GameplayCamerasModule.GetLiveEditManager();
 }
@@ -246,6 +400,12 @@ void FCameraRigAssetEditorToolkit::OnBuild()
 	}
 }
 
+void FCameraRigAssetEditorToolkit::OnFindInCameraRig()
+{
+	TabManager->TryInvokeTab(SearchTabId);
+	SearchWidget->FocusSearchEditBox();
+}
+
 }  // namespace UE::Cameras
 
 #undef LOCTEXT_NAMESPACE
diff --git a/Engine/Plugins/Cameras/GameplayCameras/Source/GameplayCamerasEditor/Public/Commands/CameraRigAssetEditorCommands.h b/Engine/Plugins/Cameras/GameplayCameras/Source/GameplayCamerasEditor/Public/Commands/CameraRigAssetEditorCommands.h
--- a/Engine/Plugins/Cameras/GameplayCameras/Source/GameplayCamerasEditor/Public/Commands/CameraRigAssetEditorCommands.h
+++ b/Engine/Plugins/Cameras/GameplayCameras/Source/GameplayCamerasEditor/Public/Commands/CameraRigAssetEditorCommands.h
@@ -23,6 +23,13 @@ public:
 public:
 
 	TSharedPtr<FUICommandInfo> Build;
+	
+	TSharedPtr<FUICommandInfo> ShowNodeHierarchy;
+	TSharedPtr<FUICommandInfo> ShowTransitions;
+
+	TSharedPtr<FUICommandInfo> ShowMessages;
+	TSharedPtr<FUICommandInfo> FindInCameraRig;
+	TSharedPtr<FUICommandInfo> FocusHome;
 };
 
 }  // namespace UE::Cameras
diff --git a/Engine/Plugins/Cameras/GameplayCameras/Source/GameplayCamerasEditor/Public/Editors/CameraTransitionGraphSchema.h b/Engine/Plugins/Cameras/GameplayCameras/Source/GameplayCamerasEditor/Public/Editors/CameraTransitionGraphSchema.h
new file mode 100644
--- /dev/null
+++ b/Engine/Plugins/Cameras/GameplayCameras/Source/GameplayCamerasEditor/Public/Editors/CameraTransitionGraphSchema.h
@@ -0,0 +1,53 @@
+// Copyright Epic Games, Inc. All Rights Reserved.
+
+#pragma once
+
+#include "Editors/ObjectTreeGraphSchema.h"
+
+#include "CameraTransitionGraphSchema.generated.h"
+
+UCLASS()
+class UCameraTransitionGraphSchema : public UObjectTreeGraphSchema
+{
+	GENERATED_BODY()
+
+protected:
+
+	// UEdGraphSchema interface.
+	virtual void GetGraphContextActions(FGraphContextMenuBuilder& ContextMenuBuilder) const override;
+
+	// UObjectTreeGraphSchema interface.
+	virtual void FilterGraphContextPlaceableClasses(TArray<UClass*>& InOutClasses) const override;
+};
+
+USTRUCT()
+struct FCameraTransitionGraphSchemaAction_NewTransitionNode : public FObjectGraphSchemaAction_NewNode
+{
+	GENERATED_BODY()
+
+public:
+
+	enum class ETransitionType
+	{
+		Enter,
+		Exit
+	};
+
+	ETransitionType TransitionType;
+
+public:
+
+	FCameraTransitionGraphSchemaAction_NewTransitionNode();
+	FCameraTransitionGraphSchemaAction_NewTransitionNode(FText InNodeCategory, FText InMenuDesc, FText InToolTip, const int32 InGrouping = 0, FText InKeywords = FText());
+
+public:
+
+	// FEdGraphSchemaAction interface.
+	static FName StaticGetTypeId() { static FName Type("FCameraTransitionGraphSchemaAction_NewTransitionNode"); return Type; }
+	virtual FName GetTypeId() const override { return StaticGetTypeId(); } 
+
+protected:
+
+	virtual void AutoSetupNewNode(UObjectTreeGraphNode* NewNode, UEdGraphPin* FromPin) override;
+};
+
diff --git a/Engine/Plugins/Cameras/GameplayCameras/Source/GameplayCamerasEditor/Public/Editors/ObjectTreeConnectionDrawingPolicy.h b/Engine/Plugins/Cameras/GameplayCameras/Source/GameplayCamerasEditor/Public/Editors/ObjectTreeConnectionDrawingPolicy.h
new file mode 100644
--- /dev/null
+++ b/Engine/Plugins/Cameras/GameplayCameras/Source/GameplayCamerasEditor/Public/Editors/ObjectTreeConnectionDrawingPolicy.h
@@ -0,0 +1,13 @@
+// Copyright Epic Games, Inc. All Rights Reserved.
+
+#pragma once
+
+#include "ConnectionDrawingPolicy.h"
+
+class FObjectTreeConnectionDrawingPolicy : public FConnectionDrawingPolicy
+{
+public:
+
+	FObjectTreeConnectionDrawingPolicy(int32 InBackLayerID, int32 InFrontLayerID, float InZoomFactor, const FSlateRect& InClippingRect, FSlateWindowElementList& InDrawElements);
+};
+
diff --git a/Engine/Plugins/Cameras/GameplayCameras/Source/GameplayCamerasEditor/Public/Editors/ObjectTreeGraph.h b/Engine/Plugins/Cameras/GameplayCameras/Source/GameplayCamerasEditor/Public/Editors/ObjectTreeGraph.h
new file mode 100644
--- /dev/null
+++ b/Engine/Plugins/Cameras/GameplayCameras/Source/GameplayCamerasEditor/Public/Editors/ObjectTreeGraph.h
@@ -0,0 +1,60 @@
+// Copyright Epic Games, Inc. All Rights Reserved.
+
+#pragma once
+
+#include "Containers/Array.h"
+#include "Containers/Map.h"
+#include "CoreTypes.h"
+#include "EdGraph/EdGraph.h"
+#include "Editors/ObjectTreeGraphConfig.h"
+#include "Templates/SubclassOf.h"
+#include "UObject/WeakObjectPtrFwd.h"
+
+#include "ObjectTreeGraph.generated.h"
+
+class UObjectTreeGraphNode;
+class UObjectTreeGraphSchema;
+
+enum class EObjectTreeGraphBuildSource
+{
+	RootObjectPackage
+};
+
+UCLASS()
+class UObjectTreeGraph : public UEdGraph
+{
+	GENERATED_BODY()
+
+public:
+
+	UObjectTreeGraph(const FObjectInitializer& ObjInit);
+
+	void Initialize(TObjectPtr<UObject> InRootObject, const FObjectTreeGraphConfig& InConfig);
+
+	UObject* GetRootObject() const { return WeakRootObject.Get(); }
+	UObjectTreeGraphNode* GetRootObjectNode() const { return RootObjectNode; }
+	const FObjectTreeGraphConfig& GetConfig() const;
+
+	void RebuildGraph(EObjectTreeGraphBuildSource InSource);
+
+private:
+
+	struct FCreatedNodes
+	{
+		TMap<UObject*, UObjectTreeGraphNode*> CreatedNodes;
+	};
+
+	void RemoveAllNodes();
+	void CreateAllNodes(EObjectTreeGraphBuildSource InSource);
+	void CreateConnections(UObjectTreeGraphNode* InGraphNode, const FCreatedNodes& InCreatedNodes);
+
+private:
+
+	TWeakObjectPtr<> WeakRootObject;
+
+	FObjectTreeGraphConfig Config;
+
+	UPROPERTY()
+	TObjectPtr<UObjectTreeGraphNode> RootObjectNode;
+};
+
diff --git a/Engine/Plugins/Cameras/GameplayCameras/Source/GameplayCamerasEditor/Public/Editors/ObjectTreeGraphConfig.h b/Engine/Plugins/Cameras/GameplayCameras/Source/GameplayCamerasEditor/Public/Editors/ObjectTreeGraphConfig.h
new file mode 100644
--- /dev/null
+++ b/Engine/Plugins/Cameras/GameplayCameras/Source/GameplayCamerasEditor/Public/Editors/ObjectTreeGraphConfig.h
@@ -0,0 +1,112 @@
+// Copyright Epic Games, Inc. All Rights Reserved.
+
+#pragma once
+
+#include "Containers/Array.h"
+#include "Containers/Map.h"
+#include "Containers/UnrealString.h"
+#include "CoreTypes.h"
+#include "EdGraph/EdGraphNode.h"
+#include "Misc/Optional.h"
+#include "Templates/SubclassOf.h"
+#include "Templates/UnrealTypeTraits.h"
+
+class FText;
+class UClass;
+class UObject;
+class UObjectTreeGraphNode;
+
+DECLARE_DELEGATE_RetVal_OneParam(FText, FOnGetObjectClassDisplayName, const UClass*);
+DECLARE_DELEGATE_TwoParams(FOnFormatObjectDisplayName, const UObject*, FText&);
+
+#define OTGCC_FIELD(FieldType, FieldName)\
+	public:\
+		typename TCallTraits<FieldType>::ConstReference FieldName() const\
+			{ return _##FieldName; }\
+		FObjectTreeGraphClassConfig& FieldName(typename TCallTraits<FieldType>::ParamType InValue)\
+			{ _##FieldName = InValue; return *this; }\
+	private:\
+		FieldType _##FieldName;
+
+struct FObjectTreeGraphClassConfig
+{
+public:
+
+	OTGCC_FIELD(TSubclassOf<UObjectTreeGraphNode>, GraphNodeClass)
+
+	OTGCC_FIELD(FText, SelfPinFriendlyName)
+	OTGCC_FIELD(EEdGraphPinDirection, SelfPinDirection)
+	OTGCC_FIELD(bool, HasSelfPin)
+
+	OTGCC_FIELD(EEdGraphPinDirection, DefaultPropertyPinDirection)
+
+	OTGCC_FIELD(TOptional<FLinearColor>, NodeTitleColor)
+	OTGCC_FIELD(TOptional<FLinearColor>, NodeBodyTintColor)
+
+	OTGCC_FIELD(bool, NodeTitleUsesObjectName)
+	OTGCC_FIELD(FOnGetObjectClassDisplayName, OnGetObjectClassDisplayName)
+
+	OTGCC_FIELD(bool, CanCreateNew)
+	OTGCC_FIELD(bool, CanDuplicate)
+	OTGCC_FIELD(bool, CanDelete)
+
+public:
+
+	FObjectTreeGraphClassConfig();
+
+	FObjectTreeGraphClassConfig& OnlyAsRoot();
+
+public:
+
+	TMap<FName, EEdGraphPinDirection> PropertyPinDirections;
+
+	FObjectTreeGraphClassConfig& SetPropertyPinDirection(const FName& InPropertyName, EEdGraphPinDirection InDirection)
+	{
+		PropertyPinDirections.Add(InPropertyName, InDirection);
+		return *this;
+	}
+
+	EEdGraphPinDirection GetPropertyPinDirection(const FName& InPropertyName) const
+	{
+		if (const EEdGraphPinDirection* PinDirection = PropertyPinDirections.Find(InPropertyName))
+		{
+			return *PinDirection;
+		}
+		return _DefaultPropertyPinDirection;
+	}
+};
+
+struct FObjectTreeGraphConfig
+{
+public:
+
+	TArray<UClass*> ConnectableObjectClasses;
+	TArray<UClass*> NonConnectableObjectClasses;
+	
+	TSubclassOf<UObjectTreeGraphNode> DefaultGraphNodeClass;
+
+	FLinearColor DefaultGraphNodeTitleColor;
+	FLinearColor DefaultGraphNodeBodyTintColor;
+
+	TArray<FString> StripDisplayNameSuffixes;
+	FOnFormatObjectDisplayName OnFormatObjectDisplayName;
+
+	TMap<UClass*, FObjectTreeGraphClassConfig> ObjectClassConfigs;
+
+public:
+
+	FObjectTreeGraphConfig();
+
+	bool IsConnectable(UClass* InObjectClass) const;
+	void GetConnectableClasses(TArray<UClass*>& OutClasses, bool bPlaceableOnly = false);
+	
+	const FObjectTreeGraphClassConfig& GetObjectClassConfig(const UClass* InObjectClass) const;
+
+	FText GetDisplayNameText(const UObject* InObject) const;
+	FText GetDisplayNameText(const UClass* InClass) const;
+
+private:
+
+	void FormatDisplayNameText(const UObject* InObject, FText& InOutDisplayNameText) const;
+};
+
diff --git a/Engine/Plugins/Cameras/GameplayCameras/Source/GameplayCamerasEditor/Public/Editors/ObjectTreeGraphNode.h b/Engine/Plugins/Cameras/GameplayCameras/Source/GameplayCamerasEditor/Public/Editors/ObjectTreeGraphNode.h
new file mode 100644
--- /dev/null
+++ b/Engine/Plugins/Cameras/GameplayCameras/Source/GameplayCamerasEditor/Public/Editors/ObjectTreeGraphNode.h
@@ -0,0 +1,97 @@
+// Copyright Epic Games, Inc. All Rights Reserved.
+
+#pragma once
+
+#include "Containers/Array.h"
+#include "EdGraph/EdGraphNode.h"
+#include "UObject/UObjectGlobals.h"
+
+#include "ObjectTreeGraphNode.generated.h"
+
+class FObjectProperty;
+class UEdGraphPin;
+class UObjectTreeGraph;
+struct FObjectTreeGraphConfig;
+struct FObjectTreeGraphClassConfig;
+
+UCLASS()
+class UObjectTreeGraphNode : public UEdGraphNode
+{
+	GENERATED_BODY()
+
+public:
+
+	UObjectTreeGraphNode(const FObjectInitializer& ObjInit);
+
+	void Initialize(UObject* InObject);
+
+	UObject* GetObject() const { return Object.Get(); }
+	void GetAllConnectableProperties(TArray<FProperty*>& OutProperties) const;
+
+	UEdGraphPin* GetSelfPin() const;
+	void OverrideSelfPinDirection(EEdGraphPinDirection Direction);
+
+	UEdGraphPin* GetPinForProperty(FObjectProperty* InProperty) const;
+	UEdGraphPin* GetPinForProperty(FArrayProperty* InProperty, int32 Index) const;
+	UEdGraphPin* GetPinForPropertyNewItem(FArrayProperty* InProperty, bool bCreateNew);
+	FProperty* GetPropertyForPin(const UEdGraphPin* InPin) const;
+	UClass* GetConnectedObjectClassForPin(const UEdGraphPin* InPin) const;
+	int32 GetIndexOfArrayPin(const UEdGraphPin* InPin) const;
+
+	// UEdGraphNode interface.
+	virtual FText GetNodeTitle(ENodeTitleType::Type TitleType) const override;
+	virtual TSharedPtr<SGraphNode> CreateVisualWidget() override;
+	virtual FLinearColor GetNodeTitleColor() const override;
+	virtual FLinearColor GetNodeBodyTintColor() const override;
+	virtual void AllocateDefaultPins() override;
+	virtual void PostPlacedNewNode() override;
+	virtual void AutowireNewNode(UEdGraphPin* FromPin) override;
+	virtual void PinConnectionListChanged(UEdGraphPin* Pin) override;
+	virtual void NodeConnectionListChanged() override;
+	virtual void OnPinRemoved(UEdGraphPin* InRemovedPin) override;
+	virtual void GetNodeContextMenuActions(class UToolMenu* Menu, class UGraphNodeContextMenuContext* Context) const override;
+	virtual bool GetCanRenameNode() const override;
+	virtual void OnRenameNode(const FString& NewName) override;
+	virtual bool CanDuplicateNode() const override;
+	virtual bool CanUserDeleteNode() const override;
+	virtual bool SupportsCommentBubble() const override;
+	virtual void OnUpdateCommentText(const FString& NewComment) override;
+
+	// UObjectTreeGraphNode interface.
+	virtual void OnGraphNodeMoved();
+	virtual void OnDoubleClicked() {}
+
+public:
+
+	// Internal API.
+	void PostDuplicateObject(const TMap<UEdGraphNode*, UEdGraphNode*>& NodeMap);
+	void CreateNewItemPin(FArrayProperty& InArrayProperty);
+	void CreateNewItemPin(UEdGraphPin* InParentArrayPin);
+	void RemoveItemPin(UEdGraphPin* InItemPin);
+	void RefreshArrayPropertyPinNames();
+
+protected:
+
+	struct FNodeContext
+	{
+		UClass* ObjectClass;
+		UObjectTreeGraph* Graph;
+		const FObjectTreeGraphConfig& GraphConfig;
+		const FObjectTreeGraphClassConfig& ObjectClassConfig;
+	};
+
+	FNodeContext GetNodeContext() const;
+	const FObjectTreeGraphClassConfig& GetObjectClassConfig() const;
+
+private:
+
+	UPROPERTY()
+	TObjectPtr<UObject> Object;
+
+	UPROPERTY()
+	TEnumAsByte<EEdGraphPinDirection> SelfPinDirectionOverride;
+
+	UPROPERTY()
+	bool bOverrideSelfPinDirection;
+};
+
diff --git a/Engine/Plugins/Cameras/GameplayCameras/Source/GameplayCamerasEditor/Public/Editors/ObjectTreeGraphSchema.h b/Engine/Plugins/Cameras/GameplayCameras/Source/GameplayCamerasEditor/Public/Editors/ObjectTreeGraphSchema.h
new file mode 100644
--- /dev/null
+++ b/Engine/Plugins/Cameras/GameplayCameras/Source/GameplayCamerasEditor/Public/Editors/ObjectTreeGraphSchema.h
@@ -0,0 +1,97 @@
+// Copyright Epic Games, Inc. All Rights Reserved.
+
+#pragma once
+
+#include "CoreTypes.h"
+#include "EdGraph/EdGraphSchema.h"
+
+#include "ObjectTreeGraphSchema.generated.h"
+
+class UEdGraph;
+class UObjectTreeGraph;
+class UObjectTreeGraphNode;
+struct FObjectTreeGraphClassConfig;
+
+UCLASS()
+class UObjectTreeGraphSchema : public UEdGraphSchema
+{
+	GENERATED_BODY()
+
+public:
+
+	// Pin categories.
+	static const FName PC_Self;
+	static const FName PC_Property;
+
+	// Pin sub-categories.
+	static const FName PSC_ObjectProperty;
+	static const FName PSC_ArrayProperty;
+	static const FName PSC_ArrayPropertyItem;
+
+public:
+
+	UObjectTreeGraphSchema(const FObjectInitializer& ObjInit);
+
+	UObjectTreeGraphNode* CreateObjectNode(UObjectTreeGraph* InGraph, UObject* InObject) const;
+
+public:
+
+	// UEdGraphSchema interface.
+	virtual void GetGraphContextActions(FGraphContextMenuBuilder& ContextMenuBuilder) const override;
+	virtual void GetContextMenuActions(class UToolMenu* Menu, class UGraphNodeContextMenuContext* Context) const override;
+	virtual FName GetParentContextMenuName() const override;
+	virtual FLinearColor GetPinTypeColor(const FEdGraphPinType& PinType) const override;
+	virtual class FConnectionDrawingPolicy* CreateConnectionDrawingPolicy(int32 InBackLayerID, int32 InFrontLayerID, float InZoomFactor, const FSlateRect& InClippingRect, class FSlateWindowElementList& InDrawElements, UEdGraph* InGraph) const override;
+	virtual bool ShouldAlwaysPurgeOnModification() const override;
+	virtual FPinConnectionResponse CanCreateNewNodes(UEdGraphPin* InSourcePin) const override;
+	virtual const FPinConnectionResponse CanCreateConnection(const UEdGraphPin* A, const UEdGraphPin* B) const override;
+	virtual bool TryCreateConnection(UEdGraphPin* A, UEdGraphPin* B) const override;
+	virtual void BreakPinLinks(UEdGraphPin& TargetPin, bool bSendsNodeNotifcation) const override;
+	virtual void BreakSinglePinLink(UEdGraphPin* SourcePin, UEdGraphPin* TargetPin) const override;
+	virtual bool SupportsDropPinOnNode(UEdGraphNode* InTargetNode, const FEdGraphPinType& InSourcePinType, EEdGraphPinDirection InSourcePinDirection, FText& OutErrorMessage) const override;
+
+	// UObjectTreeGraphSchema interface.
+	virtual void FilterGraphContextPlaceableClasses(TArray<UClass*>& InOutClasses) const;
+	virtual void ProcessDuplicatedNodes(UObjectTreeGraph* InGraph, const TMap<UEdGraphNode*, UEdGraphNode*>& NodeMap) const;
+
+protected:
+
+	virtual UObjectTreeGraphNode* CreateObjectNodeImpl(UObjectTreeGraph* InGraph, UObject* InObject) const;
+
+protected:
+
+	const FObjectTreeGraphClassConfig& GetObjectClassConfig(const UObjectTreeGraphNode* InNode) const;
+	const FObjectTreeGraphClassConfig& GetObjectClassConfig(const UObjectTreeGraph* InGraph, UClass* InObjectClass) const;
+};
+
+USTRUCT()
+struct FObjectGraphSchemaAction_NewNode : public FEdGraphSchemaAction
+{
+	GENERATED_BODY()
+
+public:
+
+	UPROPERTY()
+	TObjectPtr<UObject> ObjectOuter;
+
+	UPROPERTY()
+	TObjectPtr<UClass> ObjectClass;
+
+public:
+
+	FObjectGraphSchemaAction_NewNode();
+	FObjectGraphSchemaAction_NewNode(FText InNodeCategory, FText InMenuDesc, FText InToolTip, const int32 InGrouping = 0, FText InKeywords = FText());
+
+public:
+
+	// FEdGraphSchemaAction interface.
+	static FName StaticGetTypeId() { static FName Type("FObjectGraphSchemaAction_NewNode"); return Type; }
+	virtual FName GetTypeId() const override { return StaticGetTypeId(); } 
+	virtual UEdGraphNode* PerformAction(UEdGraph* ParentGraph, UEdGraphPin* FromPin, const FVector2D Location, bool bSelectNewNode = true) override;
+
+protected:
+
+	virtual UObject* CreateObject();
+	virtual void AutoSetupNewNode(UObjectTreeGraphNode* NewNode, UEdGraphPin* FromPin);
+};
+
diff --git a/Engine/Plugins/Cameras/GameplayCameras/Source/GameplayCamerasEditor/Public/Editors/SCameraRigAssetEditor.h b/Engine/Plugins/Cameras/GameplayCameras/Source/GameplayCamerasEditor/Public/Editors/SCameraRigAssetEditor.h
new file mode 100644
--- /dev/null
+++ b/Engine/Plugins/Cameras/GameplayCameras/Source/GameplayCamerasEditor/Public/Editors/SCameraRigAssetEditor.h
@@ -0,0 +1,85 @@
+// Copyright Epic Games, Inc. All Rights Reserved.
+
+#pragma once
+
+#include "CoreTypes.h"
+#include "GraphEditor.h"
+#include "Templates/SharedPointerFwd.h"
+#include "UObject/ObjectPtr.h"
+#include "Widgets/SCompoundWidget.h"
+
+class FAssetEditorToolkit;
+class IDetailsView;
+class SBox;
+class SObjectTreeGraphEditor;
+class UCameraRigAsset;
+class UEdGraphNode;
+class UObjectTreeGraph;
+struct FObjectTreeGraphConfig;
+
+namespace UE::Cameras
+{
+
+enum class ECameraRigAssetEditorMode
+{
+	NodeGraph,
+	TransitionGraph
+};
+
+class SCameraRigAssetEditor : public SCompoundWidget
+{
+public:
+
+	SLATE_BEGIN_ARGS(SCameraRigAssetEditor)
+	{}
+		SLATE_ARGUMENT(TObjectPtr<UCameraRigAsset>, CameraRigAsset)
+		SLATE_ARGUMENT(TSharedPtr<IDetailsView>, DetailsView)
+		SLATE_ARGUMENT(TWeakPtr<FAssetEditorToolkit>, AssetEditorToolkit)
+	SLATE_END_ARGS()
+
+	void Construct(const FArguments& InArgs);
+
+public:
+
+	ECameraRigAssetEditorMode GetEditorMode() const;
+	bool IsEditorMode(ECameraRigAssetEditorMode InMode) const;
+	void SetEditorMode(ECameraRigAssetEditorMode InMode);
+
+	void GetGraphs(TArray<UEdGraph*>& OutGraphs) const;
+
+	const FObjectTreeGraphConfig& GetFocusedGraphConfig() const;
+
+	void FocusHome();
+	void JumpToNode(UEdGraphNode* InGraphNode);
+
+protected:
+
+	void CreateNodeGraphEditor(const FArguments& InArgs);
+	void CreateTransitionGraphEditor(const FArguments& InArgs);
+
+	FText GetCameraRigAssetName() const;
+
+private:
+
+	/** The asset being edited */
+	TObjectPtr<UCameraRigAsset> CameraRigAsset;
+
+	/** The node hierarchy graph */
+	TObjectPtr<UObjectTreeGraph> NodeGraph;
+	/** The node hierarchy graph editor */
+	TSharedPtr<SObjectTreeGraphEditor> NodeGraphEditor;
+
+	/** The transition graph */
+	TObjectPtr<UObjectTreeGraph> TransitionGraph;
+	/** The transition graph editor */
+	TSharedPtr<SObjectTreeGraphEditor> TransitionGraphEditor;
+
+	/** Box panel holding either the node hierarchy or transition graph editor */
+	TSharedPtr<SBox> BoxPanel;
+
+	/** The mode for the currently shown graph editor */
+	ECameraRigAssetEditorMode CurrentMode;
+};
+
+}  // namespace UE::Cameras
+
diff --git a/Engine/Plugins/Cameras/GameplayCameras/Source/GameplayCamerasEditor/Public/Editors/SFindInObjectTreeGraph.h b/Engine/Plugins/Cameras/GameplayCameras/Source/GameplayCamerasEditor/Public/Editors/SFindInObjectTreeGraph.h
new file mode 100644
--- /dev/null
+++ b/Engine/Plugins/Cameras/GameplayCameras/Source/GameplayCamerasEditor/Public/Editors/SFindInObjectTreeGraph.h
@@ -0,0 +1,102 @@
+// Copyright Epic Games, Inc. All Rights Reserved.
+
+#pragma once
+
+#include "CoreTypes.h"
+#include "EdGraph/EdGraphPin.h"
+#include "UObject/WeakObjectPtr.h"
+#include "Templates/SharedPointerFwd.h"
+#include "Widgets/SCompoundWidget.h"
+#include "Widgets/Views/STableRow.h"
+#include "Widgets/Views/STableViewBase.h"
+#include "Widgets/Views/STreeView.h"
+
+class UEdGraph;
+class UEdGraphNode;
+class UEdGraphSchema;
+class SSearchBox;
+
+class SFindInObjectTreeGraph;
+
+struct FFindInObjectTreeGraphResult
+{
+public:
+
+	TWeakPtr<FFindInObjectTreeGraphResult> Parent;
+	TArray<TSharedPtr<FFindInObjectTreeGraphResult>> Children;
+
+	FText CustomText;
+	TWeakObjectPtr<UEdGraphNode> GraphNode;
+	FEdGraphPinReference GraphPin;
+
+public:
+
+	FFindInObjectTreeGraphResult(const FText& InCustomText);
+	FFindInObjectTreeGraphResult(TSharedPtr<FFindInObjectTreeGraphResult>& InParent, UEdGraphNode* InGraphNode);
+	FFindInObjectTreeGraphResult(TSharedPtr<FFindInObjectTreeGraphResult>& InParent, UEdGraphPin* InGraphPin);
+
+	TSharedRef<SWidget>	GetIcon() const;
+	FText GetCategory() const;
+	FText GetText() const;
+	FText GetCommentText() const;
+
+	FReply OnClick(TSharedRef<SFindInObjectTreeGraph> FindInObjectTreeGraph);
+};
+
+class SFindInObjectTreeGraph : public SCompoundWidget
+{
+public:
+
+	DECLARE_DELEGATE_OneParam(FOnJumpToNodeRequested, UEdGraphNode*);
+	DECLARE_DELEGATE_OneParam(FOnJumpToPinRequested, UEdGraphPin*);
+
+	SLATE_BEGIN_ARGS(SFindInObjectTreeGraph)
+	{}
+		SLATE_ARGUMENT(TArray<UEdGraph*>, GraphsToSearch)
+		SLATE_ARGUMENT(FOnJumpToNodeRequested, OnJumpToNodeRequested)
+		SLATE_ARGUMENT(FOnJumpToPinRequested, OnJumpToPinRequested)
+	SLATE_END_ARGS()
+
+	void Construct(const FArguments& InArgs);
+
+	void FocusSearchEditBox();
+
+protected:
+
+	typedef TSharedPtr<FFindInObjectTreeGraphResult> FResultPtr;
+	typedef STreeView<FResultPtr> SResultTreeView;
+
+	void OnSearchTextChanged(const FText& Text);
+	void OnSearchTextCommitted(const FText& Text, ETextCommit::Type CommitType);
+
+	TSharedRef<ITableRow> OnResultTreeViewGenerateRow(FResultPtr InItem, const TSharedRef<STableViewBase>& OwnerTable);
+	void OnResultTreeViewGetChildren(FResultPtr InItem, TArray<FResultPtr>& OutChildren);
+	void OnResultTreeViewSelectionChanged(FResultPtr Item, ESelectInfo::Type SelectInfo);
+	void OnResultTreeViewMouseButtonDoubleClick(FResultPtr Item);
+
+protected:
+
+	void StartSearch();
+	void MatchTokensInGraph(UEdGraph* Graph, TArrayView<FString> Tokens);
+	bool GraphNodeMatchesSearchTokens(const UEdGraphSchema* GraphSchema, UEdGraphNode* GraphNode, TArrayView<FString> Tokens);
+	bool GraphPinMatchesSearchTokens(const UEdGraphSchema* GraphSchema, UEdGraphPin* GraphPin, TArrayView<FString> Tokens);
+	bool StringMatchesSearchTokens(const FString& ComparisonString, TArrayView<FString> Tokens);
+
+protected:
+
+	TSharedPtr<SSearchBox> SearchBox;
+	TSharedPtr<SResultTreeView> ResultTreeView;
+
+	FString SearchQuery;
+	TArray<FResultPtr> Results;
+
+	FText HighlightText;
+
+	TArray<UEdGraph*> GraphsToSearch;
+
+	FOnJumpToNodeRequested OnJumpToNodeRequested;
+	FOnJumpToPinRequested OnJumpToPinRequested;
+
+	friend struct FFindInObjectTreeGraphResult;
+};
+
diff --git a/Engine/Plugins/Cameras/GameplayCameras/Source/GameplayCamerasEditor/Public/Editors/SObjectTreeGraphEditor.h b/Engine/Plugins/Cameras/GameplayCameras/Source/GameplayCamerasEditor/Public/Editors/SObjectTreeGraphEditor.h
new file mode 100644
--- /dev/null
+++ b/Engine/Plugins/Cameras/GameplayCameras/Source/GameplayCamerasEditor/Public/Editors/SObjectTreeGraphEditor.h
@@ -0,0 +1,100 @@
+// Copyright Epic Games, Inc. All Rights Reserved.
+
+#pragma once
+
+#include "EditorUndoClient.h"
+#include "GraphEditor.h"
+#include "Widgets/SCompoundWidget.h"
+
+class FAssetEditorToolkit;
+class IDetailsView;
+class UObjectTreeGraph;
+class UObjectTreeGraphNode;
+
+class SObjectTreeGraphEditor 
+	: public SCompoundWidget
+	, public FEditorUndoClient
+{
+public:
+
+	SLATE_BEGIN_ARGS(SObjectTreeGraphEditor)
+	{}
+		SLATE_ARGUMENT(TSharedPtr<FUICommandList>, AdditionalCommands)
+		SLATE_ARGUMENT(TSharedPtr<SWidget>, GraphTitleBar)
+		SLATE_ARGUMENT(TSharedPtr<IDetailsView>, DetailsView)
+		SLATE_ARGUMENT(UObjectTreeGraph*, GraphToEdit)
+		SLATE_ARGUMENT(TWeakPtr<FAssetEditorToolkit>, AssetEditorToolkit)
+		SLATE_ATTRIBUTE(FGraphAppearanceInfo, Appearance)
+		SLATE_ATTRIBUTE(FText, GraphTitle)
+	SLATE_END_ARGS()
+
+	void Construct(const FArguments& InArgs);
+	~SObjectTreeGraphEditor();
+
+	void JumpToNode(UEdGraphNode* InNode);
+	void ResyncDetailsView();
+
+protected:
+
+	// SWidget interface.
+	virtual FReply OnDragOver(const FGeometry& MyGeometry, const FDragDropEvent& DragDropEvent) override;
+	virtual FReply OnDrop(const FGeometry& MyGeometry, const FDragDropEvent& DragDropEvent) override;
+
+	// FEditorUndoClient interface.
+	virtual void PostUndo(bool bSuccess) override;
+	virtual void PostRedo(bool bSuccess) override;
+
+protected:
+
+	void InitializeBuiltInCommands();
+
+	void OnGraphSelectionChanged(const FGraphPanelSelectionSet& SelectionSet);
+	void OnNodeTextCommitted(const FText& InText, ETextCommit::Type InCommitType, UEdGraphNode* InEditedNode);
+	void OnNodeDoubleClicked(UEdGraphNode* InClickedNode);
+	void OnDoubleClicked();
+
+	FString ExportNodesToText(const FGraphPanelSelectionSet& Nodes, bool bOnlyCanDuplicateNodes, bool bOnlyCanDeleteNodes);
+	void ImportNodesFromText(const FVector2D& Location, const FString& TextToImport);
+	void DeleteNodes(TArrayView<UObjectTreeGraphNode*> NodesToDelete);
+
+	void SelectAllNodes();
+	bool CanSelectAllNodes();
+
+	void DeleteSelectedNodes();
+	bool CanDeleteSelectedNodes();
+
+	void CopySelectedNodes();
+	bool CanCopySelectedNodes();
+
+	void CutSelectedNodes();
+	bool CanCutSelectedNodes();
+
+	void PasteNodes();
+	bool CanPasteNodes();
+
+	void DuplicateNodes();
+	bool CanDuplicateNodes();
+
+	void OnRenameNode();
+	bool CanRenameNode();
+
+	void OnAlignTop();
+	void OnAlignMiddle();
+	void OnAlignBottom();
+	void OnAlignLeft();
+	void OnAlignCenter();
+	void OnAlignRight();
+
+	void OnStraightenConnections();
+	void OnDistributeNodesHorizontally();
+	void OnDistributeNodesVertically();
+
+	TArray<UClass*> FilterPlaceableObjectClasses(TArrayView<UClass* const> InObjectClasses);
+
+protected:
+
+	TSharedPtr<SGraphEditor> GraphEditor;
+	TSharedPtr<IDetailsView> DetailsView;
+	TSharedPtr<FUICommandList> BuiltInCommands;
+};
+
diff --git a/Engine/Plugins/Cameras/GameplayCameras/Source/GameplayCamerasEditor/Public/Editors/SObjectTreeGraphNode.h b/Engine/Plugins/Cameras/GameplayCameras/Source/GameplayCamerasEditor/Public/Editors/SObjectTreeGraphNode.h
new file mode 100644
--- /dev/null
+++ b/Engine/Plugins/Cameras/GameplayCameras/Source/GameplayCamerasEditor/Public/Editors/SObjectTreeGraphNode.h
@@ -0,0 +1,49 @@
+// Copyright Epic Games, Inc. All Rights Reserved.
+
+#pragma once
+
+#include "SGraphNode.h"
+#include "SGraphPin.h"
+
+class UObjectTreeGraphNode;
+
+class SObjectTreeGraphNode : public SGraphNode
+{
+public:
+
+	SLATE_BEGIN_ARGS(SObjectTreeGraphNode)
+		: _GraphNode(nullptr)
+	{}
+		SLATE_ARGUMENT(UObjectTreeGraphNode*, GraphNode)
+	SLATE_END_ARGS()
+
+	void Construct(const FArguments& InArgs);
+
+	UObjectTreeGraphNode* GetObjectGraphNode() const { return ObjectGraphNode; }
+
+public:
+
+	// SNodePanel::SNode interface.
+	virtual void MoveTo(const FVector2D& NewPosition, FNodeSet& NodeFilter, bool bMarkDirty=true) override;
+
+protected:
+
+	UObjectTreeGraphNode* ObjectGraphNode;
+};
+
+class SObjectTreeGraphAddArrayItemPin : public SGraphPin
+{
+public:
+
+	SLATE_BEGIN_ARGS(SObjectTreeGraphAddArrayItemPin)
+	{}
+	SLATE_END_ARGS()
+
+	void Construct(const FArguments& InArgs, UEdGraphPin* InGraphPinObj);
+
+public:
+
+	// SGraphPin interface.
+
+};
+
diff --git a/Engine/Plugins/Cameras/GameplayCameras/Source/GameplayCamerasEditor/Public/Editors/SObjectTreeGraphTitleBar.h b/Engine/Plugins/Cameras/GameplayCameras/Source/GameplayCamerasEditor/Public/Editors/SObjectTreeGraphTitleBar.h
new file mode 100644
--- /dev/null
+++ b/Engine/Plugins/Cameras/GameplayCameras/Source/GameplayCamerasEditor/Public/Editors/SObjectTreeGraphTitleBar.h
@@ -0,0 +1,57 @@
+// Copyright Epic Games, Inc. All Rights Reserved.
+
+#pragma once
+
+#include "CoreTypes.h"
+#include "Templates/SharedPointerFwd.h"
+#include "Widgets/SCompoundWidget.h"
+
+class ITableRow;
+class SScrollBox;
+class STableViewBase;
+class UEdGraph;
+
+template<typename> class SBreadcrumbTrail;
+template<typename> class SListView;
+
+DECLARE_DELEGATE_OneParam(FObjectTreeGraphEvent, UEdGraph*);
+
+struct FObjectTreeGraphInfo
+{
+	FString GraphName;
+};
+
+class SObjectTreeGraphTitleBar : public SCompoundWidget
+{
+public:
+	SLATE_BEGIN_ARGS(SObjectTreeGraphTitleBar)
+		: _Graph(nullptr)
+	{}
+		SLATE_ATTRIBUTE(FText, TitleText)
+		SLATE_ARGUMENT(UEdGraph*, Graph)
+		SLATE_ARGUMENT(const TArray<TSharedPtr<FObjectTreeGraphInfo>>*, GraphList)
+		SLATE_ARGUMENT(TSharedPtr<SWidget>, HistoryNavigationWidget)
+		SLATE_EVENT(FObjectTreeGraphEvent, OnBreadcrumbClicked)
+	SLATE_END_ARGS()
+
+	void Construct(const FArguments& InArgs);
+
+protected:
+
+	TSharedRef<ITableRow> GenerateGraphInfoRow(TSharedPtr<FObjectTreeGraphInfo> Item, const TSharedRef<STableViewBase>& OwnerTable);
+
+	void OnBreadcrumbClickedImpl(UEdGraph* const& Item);
+	void RebuildBreadcrumbTrail();
+	static FText GetTitleForOneCrumb(const UEdGraph* BaseGraph, const UEdGraph* CurGraph);
+
+protected:
+
+	UEdGraph* Graph;
+
+	TSharedPtr<SListView<TSharedPtr<FObjectTreeGraphInfo>>> GraphListView;
+
+	TSharedPtr<SBreadcrumbTrail<UEdGraph*>> BreadcrumbTrail;
+	TSharedPtr<SScrollBox> BreadcrumbTrailScrollBox;
+	FObjectTreeGraphEvent OnBreadcrumbClicked;
+};
+
diff --git a/Engine/Plugins/Cameras/GameplayCameras/Source/GameplayCamerasEditor/Public/Editors/SObjectTreeGraphToolbox.h b/Engine/Plugins/Cameras/GameplayCameras/Source/GameplayCamerasEditor/Public/Editors/SObjectTreeGraphToolbox.h
new file mode 100644
--- /dev/null
+++ b/Engine/Plugins/Cameras/GameplayCameras/Source/GameplayCamerasEditor/Public/Editors/SObjectTreeGraphToolbox.h
@@ -0,0 +1,118 @@
+// Copyright Epic Games, Inc. All Rights Reserved.
+
+#pragma once
+
+#include "CoreTypes.h"
+#include "DragAndDrop/DecoratedDragDropOp.h"
+#include "Editors/ObjectTreeGraphConfig.h"
+#include "Misc/TextFilter.h"
+#include "Widgets/SCompoundWidget.h"
+
+class ITableRow;
+class SSearchBox;
+class STableViewBase;
+class UClass;
+
+template<typename> class SListView;
+
+class SObjectTreeGraphToolboxEntry : public SCompoundWidget
+{
+public:
+
+	SLATE_BEGIN_ARGS(SObjectTreeGraphToolboxEntry)
+		: _ObjectClass(nullptr)
+		, _GraphConfig(nullptr)
+	{}
+		SLATE_ARGUMENT(UClass*, ObjectClass)
+		SLATE_ARGUMENT(const FObjectTreeGraphConfig*, GraphConfig)
+		SLATE_ATTRIBUTE(FText, HighlightText)
+	SLATE_END_ARGS()
+
+	void Construct(const FArguments& InArgs);
+
+	UClass* GetObjectClass() const { return ObjectClass; }
+
+protected:
+
+	virtual FReply OnMouseButtonDown(const FGeometry& MyGeometry, const FPointerEvent& MouseEvent) override;
+	virtual FReply OnMouseButtonUp(const FGeometry& MyGeometry, const FPointerEvent& MouseEvent) override;
+	virtual FReply OnDragDetected(const FGeometry& MyGeometry, const FPointerEvent& MouseEvent) override;
+
+	const FSlateBrush* GetBorder() const;
+
+private:
+
+	UClass* ObjectClass = nullptr;
+	FText DisplayNameText;
+
+	bool bIsPressed = false;
+
+	const FSlateBrush* NormalImage = nullptr;
+	const FSlateBrush* HoverImage = nullptr;
+	const FSlateBrush* PressedImage = nullptr;
+};
+
+class SObjectTreeGraphToolbox : public SCompoundWidget
+{
+public:
+
+	SLATE_BEGIN_ARGS(SObjectTreeGraphToolbox)
+	{}
+		SLATE_ARGUMENT(FObjectTreeGraphConfig, GraphConfig)
+	SLATE_END_ARGS()
+
+	void Construct(const FArguments& InArgs);
+
+	void SetGraphConfig(const FObjectTreeGraphConfig& InGraphConfig);
+
+protected:
+
+	// SWidget interface
+	virtual void Tick(const FGeometry& AllottedGeometry, const double InCurrentTime, const float InDeltaTime) override;
+
+private:
+
+	void GetEntryStrings(const UClass* InItem, TArray<FString>& OutStrings);
+
+	void UpdateItemSource();
+	void UpdateFilteredItemSource();
+
+	TSharedRef<ITableRow> OnGenerateItemRow(UClass* Item, const TSharedRef<STableViewBase>& OwnerTable);
+
+	void OnSearchTextChanged(const FText& InFilterText);
+	void OnSearchTextCommitted(const FText& InFilterText, ETextCommit::Type InCommitType);
+	FText GetHighlightText() const;
+
+private:
+
+	FObjectTreeGraphConfig GraphConfig;
+
+	TArray<UClass*> ItemSource;
+	TSharedPtr<SListView<UClass*>> ListView;
+
+	using FEntryTextFilter = TTextFilter<const UClass*>;
+	TSharedPtr<FEntryTextFilter> SearchTextFilter;
+	TSharedPtr<SSearchBox> SearchBox;
+
+	TArray<UClass*> FilteredItemSource;
+
+	bool bUpdateItemSource = false;
+	bool bUpdateFilteredItemSource = false;
+};
+
+class FObjectTreeClassDragDropOp : public FDecoratedDragDropOp
+{
+public:
+	
+	DRAG_DROP_OPERATOR_TYPE(FObjectTreeClassDragDropOp, FDecoratedDragDropOp)
+
+	static TSharedRef<FObjectTreeClassDragDropOp> New(UClass* InObjectClass);
+	static TSharedRef<FObjectTreeClassDragDropOp> New(TArrayView<UClass*> InObjectClasses);
+
+	TArrayView<UClass* const> GetObjectClasses() const { return ObjectClasses; }
+
+private:
+
+	TArray<UClass*> ObjectClasses;
+};
+
diff --git a/Engine/Plugins/Cameras/GameplayCameras/Source/GameplayCamerasEditor/Public/GameplayCamerasEditorSettings.h b/Engine/Plugins/Cameras/GameplayCameras/Source/GameplayCamerasEditor/Public/GameplayCamerasEditorSettings.h
--- a/Engine/Plugins/Cameras/GameplayCameras/Source/GameplayCamerasEditor/Public/GameplayCamerasEditorSettings.h
+++ b/Engine/Plugins/Cameras/GameplayCameras/Source/GameplayCamerasEditor/Public/GameplayCamerasEditorSettings.h
@@ -2,7 +2,7 @@
 
 #pragma once
 
-#include "CoreMinimal.h"
+#include "CoreTypes.h"
 #include "UObject/ObjectMacros.h"
 #include "UObject/Object.h"
 #include "UObject/Package.h"
@@ -18,5 +18,22 @@ class GAMEPLAYCAMERASEDITOR_API UGameplayCamerasEditorSettings : public UObject
 public:
 
 	UGameplayCamerasEditorSettings(const FObjectInitializer& ObjectInitializer);	
+
+public:
+
+	UPROPERTY(EditAnywhere, Config, Category=NodeTitleColors)
+	FLinearColor CameraNodeTitleColor;
+
+	UPROPERTY(EditAnywhere, Config, Category=NodeTitleColors)
+	FLinearColor CameraRigAssetTitleColor;
+
+	UPROPERTY(EditAnywhere, Config, Category=NodeTitleColors)
+	FLinearColor CameraRigTransitionTitleColor;
+
+	UPROPERTY(EditAnywhere, Config, Category=NodeTitleColors)
+	FLinearColor CameraRigTransitionConditionTitleColor;
+
+	UPROPERTY(EditAnywhere, Config, Category=NodeTitleColors)
+	FLinearColor CameraBlendNodeTitleColor;
 };
 
diff --git a/Engine/Plugins/Cameras/GameplayCameras/Source/GameplayCamerasEditor/Public/IGameplayCamerasEditorModule.h b/Engine/Plugins/Cameras/GameplayCameras/Source/GameplayCamerasEditor/Public/IGameplayCamerasEditorModule.h
--- a/Engine/Plugins/Cameras/GameplayCameras/Source/GameplayCamerasEditor/Public/IGameplayCamerasEditorModule.h
+++ b/Engine/Plugins/Cameras/GameplayCameras/Source/GameplayCamerasEditor/Public/IGameplayCamerasEditorModule.h
@@ -2,7 +2,7 @@
 
 #pragma once
 
-#include "CoreMinimal.h"
+#include "CoreTypes.h"
 #include "Delegates/Delegate.h"
 #include "Modules/ModuleInterface.h"
 #include "Toolkits/AssetEditorToolkit.h"
diff --git a/Engine/Plugins/Cameras/GameplayCameras/Source/GameplayCamerasEditor/Public/Toolkits/CameraAssetEditorToolkit.h b/Engine/Plugins/Cameras/GameplayCameras/Source/GameplayCamerasEditor/Public/Toolkits/CameraAssetEditorToolkit.h
--- a/Engine/Plugins/Cameras/GameplayCameras/Source/GameplayCamerasEditor/Public/Toolkits/CameraAssetEditorToolkit.h
+++ b/Engine/Plugins/Cameras/GameplayCameras/Source/GameplayCamerasEditor/Public/Toolkits/CameraAssetEditorToolkit.h
@@ -2,7 +2,7 @@
 
 #pragma once
 
-#include "CoreMinimal.h"
+#include "CoreTypes.h"
 #include "Tools/BaseAssetToolkit.h"
 #include "UObject/GCObject.h"
 
diff --git a/Engine/Plugins/Cameras/GameplayCameras/Source/GameplayCamerasEditor/Public/Toolkits/CameraRigAssetEditorToolkit.h b/Engine/Plugins/Cameras/GameplayCameras/Source/GameplayCamerasEditor/Public/Toolkits/CameraRigAssetEditorToolkit.h
--- a/Engine/Plugins/Cameras/GameplayCameras/Source/GameplayCamerasEditor/Public/Toolkits/CameraRigAssetEditorToolkit.h
+++ b/Engine/Plugins/Cameras/GameplayCameras/Source/GameplayCamerasEditor/Public/Toolkits/CameraRigAssetEditorToolkit.h
@@ -2,7 +2,8 @@
 
 #pragma once
 
-#include "CoreMinimal.h"
+#include "CoreTypes.h"
+#include "Editors/ObjectTreeGraphConfig.h"
 #include "Misc/NotifyHook.h"
 #include "Tools/BaseAssetToolkit.h"
 #include "UObject/GCObject.h"
@@ -10,14 +11,19 @@
 #include "CameraRigAssetEditorToolkit.generated.h"
 
 class IMessageLogListing;
+class SFindInObjectTreeGraph;
+class SObjectTreeGraphToolbox;
+class SWidget;
 class UCameraRigAsset;
 class UCameraRigAssetEditor;
-class SWidget;
+class UEdGraphNode;
 
 namespace UE::Cameras
 {
 
 class IGameplayCamerasLiveEditManager;
+class SCameraRigAssetEditor;
+class SCameraTransitionEditor;
 
 /**
  * Editor toolkit for a camera rig asset.
@@ -63,26 +69,45 @@ protected:
 
 private:
 
+	TSharedRef<SDockTab> SpawnTab_Toolbox(const FSpawnTabArgs& Args);
+	TSharedRef<SDockTab> SpawnTab_CameraRigEditor(const FSpawnTabArgs& Args);
+	TSharedRef<SDockTab> SpawnTab_Search(const FSpawnTabArgs& Args);
+	TSharedRef<SDockTab> SpawnTab_Messages(const FSpawnTabArgs& Args);
+
 	FSlateIcon GetBuildButtonIcon() const;
 	FText GetBuildButtonTooltip() const;
 
 	void OnBuild();
+	void OnFindInCameraRig();
 
 private:
 
+	static const FName ToolboxTabId;
+	static const FName CameraRigEditorTabId;
+	static const FName SearchTabId;
+	static const FName MessagesTabId;
 	static const FName DetailsViewTabId;
 
+	/** The asset being edited */
+	TObjectPtr<UCameraRigAsset> CameraRigAsset;
+
 	/** Command bindings */
 	TSharedRef<FUICommandList> CommandBindings;
 
-	/** Message log widget */
-	TSharedPtr<SWidget> Stats;
-
 	/** Message log listing */
-	TSharedPtr<IMessageLogListing> StatsListing;
+	TSharedPtr<IMessageLogListing> MessageListing;
 
-	/** The asset being edited */
-	TObjectPtr<UCameraRigAsset> CameraRigAsset;
+	/** Camera rig editor widget */
+	TSharedPtr<SCameraRigAssetEditor> CameraRigEditorWidget;
+
+	/** Toolbox widget */
+	TSharedPtr<SObjectTreeGraphToolbox> ToolboxWidget;
+
+	/** Message log widget */
+	TSharedPtr<SWidget> MessagesWidget;
+
+	/** Search widget */
+	TSharedPtr<SFindInObjectTreeGraph> SearchWidget;
 
 	/** Live edit manager for updating the assets in the runtime */
 	TSharedPtr<UE::Cameras::IGameplayCamerasLiveEditManager> LiveEditManager;
