Cameras: fix static analysis build

diff --git a/Engine/Plugins/Cameras/GameplayCameras/Source/GameplayCamerasEditor/Private/Editors/ObjectTreeGraphSchema.cpp b/Engine/Plugins/Cameras/GameplayCameras/Source/GameplayCamerasEditor/Private/Editors/ObjectTreeGraphSchema.cpp
--- a/Engine/Plugins/Cameras/GameplayCameras/Source/GameplayCamerasEditor/Private/Editors/ObjectTreeGraphSchema.cpp
+++ b/Engine/Plugins/Cameras/GameplayCameras/Source/GameplayCamerasEditor/Private/Editors/ObjectTreeGraphSchema.cpp
@@ -108,6 +108,11 @@ void UObjectTreeGraphSchema::GetGraphContextActions(FGraphContextMenuBuilder& Co
 
 	for (UClass* PossibleObjectClass : PossibleObjectClasses)
 	{
+		if (!PossibleObjectClass)
+		{
+			continue;
+		}
+
 		const FString* CategoryName = nullptr;
 		for (UClass* CurClass = PossibleObjectClass; CurClass; CurClass = CurClass->GetSuperClass())
 		{
@@ -168,7 +173,7 @@ const FPinConnectionResponse UObjectTreeGraphSchema::CanCreateConnection(const U
 {
 	UObjectTreeGraphNode* NodeA = Cast<UObjectTreeGraphNode>(A->GetOwningNode());
 	UObjectTreeGraphNode* NodeB = Cast<UObjectTreeGraphNode>(B->GetOwningNode());
-	if (!NodeA && !NodeB)
+	if (!NodeA || !NodeB)
 	{
 		return FPinConnectionResponse(CONNECT_RESPONSE_DISALLOW, TEXT("Unsupported node types"));
 	}
