Cameras: fix localization string mismatch

diff --git a/Engine/Plugins/Cameras/GameplayCameras/Source/GameplayCamerasEditor/Private/Editors/ObjectTreeGraphNode.cpp b/Engine/Plugins/Cameras/GameplayCameras/Source/GameplayCamerasEditor/Private/Editors/ObjectTreeGraphNode.cpp
--- a/Engine/Plugins/Cameras/GameplayCameras/Source/GameplayCamerasEditor/Private/Editors/ObjectTreeGraphNode.cpp
+++ b/Engine/Plugins/Cameras/GameplayCameras/Source/GameplayCamerasEditor/Private/Editors/ObjectTreeGraphNode.cpp
@@ -536,7 +536,7 @@ void UObjectTreeGraphNode::GetNodeContextMenuActions(class UToolMenu* Menu, clas
 				FNewToolMenuDelegate::CreateLambda([&GraphEditorCommands](UToolMenu* InMenu)
 					{
 						FToolMenuSection& SubMenuSection = InMenu->AddSection(
-								"ObjectTreeGraphAlignmentActions", LOCTEXT("AlignmentHeader", "Align"));
+								"ObjectTreeGraphAlignmentActions", LOCTEXT("AlignmentHeader", "Alignment"));
 						SubMenuSection.AddMenuEntry(GraphEditorCommands.AlignNodesTop);
 						SubMenuSection.AddMenuEntry(GraphEditorCommands.AlignNodesMiddle);
 						SubMenuSection.AddMenuEntry(GraphEditorCommands.AlignNodesBottom);
