Cameras: trim suffixes for node titles

diff --git a/Engine/Plugins/Cameras/GameplayCameras/Source/GameplayCamerasEditor/Private/Editors/ObjectTreeGraphConfig.cpp b/Engine/Plugins/Cameras/GameplayCameras/Source/GameplayCamerasEditor/Private/Editors/ObjectTreeGraphConfig.cpp
--- a/Engine/Plugins/Cameras/GameplayCameras/Source/GameplayCamerasEditor/Private/Editors/ObjectTreeGraphConfig.cpp
+++ b/Engine/Plugins/Cameras/GameplayCameras/Source/GameplayCamerasEditor/Private/Editors/ObjectTreeGraphConfig.cpp
@@ -109,27 +109,26 @@ FText FObjectTreeGraphConfig::GetDisplayNameText(const UObject* InObject) const
 	if (InObject)
 	{
 		FText DisplayNameText;
+		const FObjectTreeGraphClassConfig& ClassConfig = GetObjectClassConfig(InObject->GetClass());
 
 		const IObjectTreeGraphObject* GraphObject = Cast<IObjectTreeGraphObject>(InObject);
 		if (GraphObject && GraphObject->HasSupportFlags(EObjectTreeGraphObjectSupportFlags::CustomRename))
 		{
 			DisplayNameText = FText::FromString(GraphObject->GetGraphNodeName());
 		}
-		if (!DisplayNameText.IsEmpty())
+
+		if (DisplayNameText.IsEmpty() && ClassConfig.NodeTitleUsesObjectName())
 		{
-			FormatDisplayNameText(InObject, DisplayNameText);
-			return DisplayNameText;
+			DisplayNameText = FText::FromString(InObject->GetName());
 		}
 
-		const FObjectTreeGraphClassConfig& ClassConfig = GetObjectClassConfig(InObject->GetClass());
-		if (ClassConfig.NodeTitleUsesObjectName())
+		if (!DisplayNameText.IsEmpty())
 		{
-			DisplayNameText = FText::FromString(InObject->GetName());
-			FormatDisplayNameText(InObject, DisplayNameText);
+			FormatDisplayNameText(InObject, ClassConfig, DisplayNameText);
 			return DisplayNameText;
 		}
 
-		return GetDisplayNameText(InObject->GetClass());
+		return GetDisplayNameText(InObject->GetClass(), ClassConfig);
 	}
 	return FText::GetEmpty();
 }
@@ -139,24 +138,31 @@ FText FObjectTreeGraphConfig::GetDisplayNameText(const UClass* InClass) const
 	if (InClass)
 	{
 		const FObjectTreeGraphClassConfig& ClassConfig = GetObjectClassConfig(InClass);
-		if (ClassConfig.OnGetObjectClassDisplayName().IsBound())
-		{
-			return ClassConfig.OnGetObjectClassDisplayName().Execute(InClass);
-		}
-
-		FText DisplayNameText = InClass->GetDisplayNameText();
-		FormatDisplayNameText(InClass, DisplayNameText);
-		return DisplayNameText;
+		return GetDisplayNameText(InClass, ClassConfig);
 	}
 	return FText::GetEmpty();
 }
 
-void FObjectTreeGraphConfig::FormatDisplayNameText(const UObject* InObject, FText& InOutDisplayNameText) const
+FText FObjectTreeGraphConfig::GetDisplayNameText(const UClass* InClass, const FObjectTreeGraphClassConfig& InClassConfig) const
+{
+	check(InClass);
+	
+	if (InClassConfig.OnGetObjectClassDisplayName().IsBound())
+	{
+		return InClassConfig.OnGetObjectClassDisplayName().Execute(InClass);
+	}
+
+	FText DisplayNameText = InClass->GetDisplayNameText();
+	FormatDisplayNameText(InClass, InClassConfig, DisplayNameText);
+	return DisplayNameText;
+}
+
+void FObjectTreeGraphConfig::FormatDisplayNameText(const UObject* InObject, const FObjectTreeGraphClassConfig& InClassConfig, FText& InOutDisplayNameText) const
 {
-	if (StripDisplayNameSuffixes.Num() > 0)
+	if (InClassConfig.StripDisplayNameSuffixes().Num() > 0)
 	{
 		FString DisplayName = InOutDisplayNameText.ToString();
-		for (const FString& StripSuffix : StripDisplayNameSuffixes)
+		for (const FString& StripSuffix : InClassConfig.StripDisplayNameSuffixes())
 		{
 			if (DisplayName.RemoveFromEnd(StripSuffix))
 			{
diff --git a/Engine/Plugins/Cameras/GameplayCameras/Source/GameplayCamerasEditor/Private/Editors/SCameraRigAssetEditor.cpp b/Engine/Plugins/Cameras/GameplayCameras/Source/GameplayCamerasEditor/Private/Editors/SCameraRigAssetEditor.cpp
--- a/Engine/Plugins/Cameras/GameplayCameras/Source/GameplayCamerasEditor/Private/Editors/SCameraRigAssetEditor.cpp
+++ b/Engine/Plugins/Cameras/GameplayCameras/Source/GameplayCamerasEditor/Private/Editors/SCameraRigAssetEditor.cpp
@@ -51,6 +51,8 @@ void SCameraRigAssetEditor::CreateNodeGraphEditor(const FArguments& InArgs)
 		.HasSelfPin(false)
 		.NodeTitleUsesObjectName(true)
 		.NodeTitleColor(Settings->CameraRigAssetTitleColor);
+	GraphConfig.ObjectClassConfigs.Emplace(UCameraNode::StaticClass())
+		.StripDisplayNameSuffix(TEXT("Camera Node"));
 
 	NodeGraph = NewObject<UObjectTreeGraph>(GetTransientPackage(), NAME_None, RF_Transactional);
 	NodeGraph->AddToRoot();
@@ -90,7 +92,10 @@ void SCameraRigAssetEditor::CreateTransitionGraphEditor(const FArguments& InArgs
 	GraphConfig.ObjectClassConfigs.Emplace(UCameraRigTransitionCondition::StaticClass())
 		.SelfPinDirection(EGPD_Output)
 		.DefaultPropertyPinDirection(EGPD_Input)
+		.StripDisplayNameSuffix(TEXT("Transition Condition"))
 		.NodeTitleColor(Settings->CameraRigTransitionConditionTitleColor);
+	GraphConfig.ObjectClassConfigs.Emplace(UBlendCameraNode::StaticClass())
+		.StripDisplayNameSuffix(TEXT("Camera Node"));
 
 	TransitionGraph = NewObject<UObjectTreeGraph>(GetTransientPackage(), NAME_None, RF_Transactional);
 	TransitionGraph->Schema = UCameraTransitionGraphSchema::StaticClass();
diff --git a/Engine/Plugins/Cameras/GameplayCameras/Source/GameplayCamerasEditor/Public/Editors/ObjectTreeGraphConfig.h b/Engine/Plugins/Cameras/GameplayCameras/Source/GameplayCamerasEditor/Public/Editors/ObjectTreeGraphConfig.h
--- a/Engine/Plugins/Cameras/GameplayCameras/Source/GameplayCamerasEditor/Public/Editors/ObjectTreeGraphConfig.h
+++ b/Engine/Plugins/Cameras/GameplayCameras/Source/GameplayCamerasEditor/Public/Editors/ObjectTreeGraphConfig.h
@@ -11,6 +11,8 @@
 #include "Templates/SubclassOf.h"
 #include "Templates/UnrealTypeTraits.h"
 
+#include <initializer_list>
+
 class FText;
 class UClass;
 class UObject;
@@ -58,22 +60,41 @@ public:
 
 public:
 
-	TMap<FName, EEdGraphPinDirection> PropertyPinDirections;
+	TArrayView<const FString> StripDisplayNameSuffixes() const { return _StripDisplayNameSuffixes; }
+
+	FObjectTreeGraphClassConfig& StripDisplayNameSuffix(const FString& InSuffix)
+	{
+		_StripDisplayNameSuffixes.Add(InSuffix);
+		return *this;
+	}
+
+	FObjectTreeGraphClassConfig& StripDisplayNameSuffixes(std::initializer_list<FString> InSuffixes)
+	{
+		_StripDisplayNameSuffixes.Append(InSuffixes);
+		return *this;
+	}
+
+	const TMap<FName, EEdGraphPinDirection>& PropertyPinDirections() const { return _PropertyPinDirections; }
 
 	FObjectTreeGraphClassConfig& SetPropertyPinDirection(const FName& InPropertyName, EEdGraphPinDirection InDirection)
 	{
-		PropertyPinDirections.Add(InPropertyName, InDirection);
+		_PropertyPinDirections.Add(InPropertyName, InDirection);
 		return *this;
 	}
 
 	EEdGraphPinDirection GetPropertyPinDirection(const FName& InPropertyName) const
 	{
-		if (const EEdGraphPinDirection* PinDirection = PropertyPinDirections.Find(InPropertyName))
+		if (const EEdGraphPinDirection* PinDirection = _PropertyPinDirections.Find(InPropertyName))
 		{
 			return *PinDirection;
 		}
 		return _DefaultPropertyPinDirection;
 	}
+
+private:
+
+	TArray<FString> _StripDisplayNameSuffixes;
+	TMap<FName, EEdGraphPinDirection> _PropertyPinDirections;
 };
 
 struct FObjectTreeGraphConfig
@@ -88,7 +109,6 @@ public:
 	FLinearColor DefaultGraphNodeTitleColor;
 	FLinearColor DefaultGraphNodeBodyTintColor;
 
-	TArray<FString> StripDisplayNameSuffixes;
 	FOnFormatObjectDisplayName OnFormatObjectDisplayName;
 
 	TMap<UClass*, FObjectTreeGraphClassConfig> ObjectClassConfigs;
@@ -107,6 +127,7 @@ public:
 
 private:
 
-	void FormatDisplayNameText(const UObject* InObject, FText& InOutDisplayNameText) const;
+	FText GetDisplayNameText(const UClass* InClass, const FObjectTreeGraphClassConfig& InClassConfig) const;
+	void FormatDisplayNameText(const UObject* InObject, const FObjectTreeGraphClassConfig& InClassConfig, FText& InOutDisplayNameText) const;
 };
 
