Add GAMEPLAYCAMERAS_API to some classes and functions in GameplayCameras to prepare so I can use them in another module.
(Was reviewed under changelist 32702823 - Splitting this out to make merge with the branch I'm working in easier)

diff --git a/Engine/Plugins/Cameras/GameplayCameras/Source/GameplayCameras/Public/Core/CameraDirector.h b/Engine/Plugins/Cameras/GameplayCameras/Source/GameplayCameras/Public/Core/CameraDirector.h
--- a/Engine/Plugins/Cameras/GameplayCameras/Source/GameplayCameras/Public/Core/CameraDirector.h
+++ b/Engine/Plugins/Cameras/GameplayCameras/Source/GameplayCameras/Public/Core/CameraDirector.h
@@ -11,8 +11,8 @@
 /**
  * Base class for a camera director.
  */
-UCLASS(Abstract, DefaultToInstanced, MinimalAPI)
-class UCameraDirector : public UObject
+UCLASS(Abstract, DefaultToInstanced)
+class GAMEPLAYCAMERAS_API UCameraDirector : public UObject
 {
 	GENERATED_BODY()
 
diff --git a/Engine/Plugins/Cameras/GameplayCameras/Source/GameplayCameras/Public/Core/CameraDirectorEvaluator.h b/Engine/Plugins/Cameras/GameplayCameras/Source/GameplayCameras/Public/Core/CameraDirectorEvaluator.h
--- a/Engine/Plugins/Cameras/GameplayCameras/Source/GameplayCameras/Public/Core/CameraDirectorEvaluator.h
+++ b/Engine/Plugins/Cameras/GameplayCameras/Source/GameplayCameras/Public/Core/CameraDirectorEvaluator.h
@@ -85,8 +85,8 @@ class FCameraDirectorEvaluator
 
 public:
 
-	FCameraDirectorEvaluator();
-	virtual ~FCameraDirectorEvaluator() {}
+	GAMEPLAYCAMERAS_API FCameraDirectorEvaluator();
+	GAMEPLAYCAMERAS_API virtual ~FCameraDirectorEvaluator() {}
 	
 	/** Runs the camera director to determine what camera rig(s) should be active this frame. */
 	void Run(const FCameraDirectorEvaluationParams& Params, FCameraDirectorEvaluationResult& OutResult);
diff --git a/Engine/Plugins/Cameras/GameplayCameras/Source/GameplayCameras/Public/Core/CameraEvaluationContext.h b/Engine/Plugins/Cameras/GameplayCameras/Source/GameplayCameras/Public/Core/CameraEvaluationContext.h
--- a/Engine/Plugins/Cameras/GameplayCameras/Source/GameplayCameras/Public/Core/CameraEvaluationContext.h
+++ b/Engine/Plugins/Cameras/GameplayCameras/Source/GameplayCameras/Public/Core/CameraEvaluationContext.h
@@ -24,10 +24,10 @@ class FCameraEvaluationContext : public TSharedFromThis<FCameraEvaluationContext
 public:
 
 	/** Constructs an evaluation context. */
-	FCameraEvaluationContext();
+	GAMEPLAYCAMERAS_API FCameraEvaluationContext();
 
 	/** Destroys this evaluation context. */
-	virtual ~FCameraEvaluationContext();
+	GAMEPLAYCAMERAS_API virtual ~FCameraEvaluationContext();
 
 	/**
 	 * Gets the player controller (if any) in control of the cameras running inside
diff --git a/Engine/Plugins/Cameras/GameplayCameras/Source/GameplayCameras/Public/Core/CameraNode.h b/Engine/Plugins/Cameras/GameplayCameras/Source/GameplayCameras/Public/Core/CameraNode.h
--- a/Engine/Plugins/Cameras/GameplayCameras/Source/GameplayCameras/Public/Core/CameraNode.h
+++ b/Engine/Plugins/Cameras/GameplayCameras/Source/GameplayCameras/Public/Core/CameraNode.h
@@ -19,8 +19,8 @@ using FCameraNodeChildrenView = UE::Cameras::TObjectChildrenView<TObjectPtr<UCam
 /**
  * The base class for a camera node.
  */
-UCLASS(Abstract, DefaultToInstanced, EditInlineNew, MinimalAPI, meta=(ObjectTreeGraphCategory="Camera Nodes"))
-class UCameraNode 
+UCLASS(Abstract, DefaultToInstanced, EditInlineNew, meta=(ObjectTreeGraphCategory="Camera Nodes"))
+class GAMEPLAYCAMERAS_API UCameraNode 
 	: public UObject
 	, public IObjectTreeGraphObject
 {
diff --git a/Engine/Plugins/Cameras/GameplayCameras/Source/GameplayCameras/Public/Core/CameraNodeEvaluator.h b/Engine/Plugins/Cameras/GameplayCameras/Source/GameplayCameras/Public/Core/CameraNodeEvaluator.h
--- a/Engine/Plugins/Cameras/GameplayCameras/Source/GameplayCameras/Public/Core/CameraNodeEvaluator.h
+++ b/Engine/Plugins/Cameras/GameplayCameras/Source/GameplayCameras/Public/Core/CameraNodeEvaluator.h
@@ -44,7 +44,7 @@ struct FCameraNodeEvaluationParams
 /**
  * Input/output result structure for running a camera node evaluator.
  */
-struct FCameraNodeEvaluationResult
+struct GAMEPLAYCAMERAS_API FCameraNodeEvaluationResult
 {
 	/** The camera pose. */
 	FCameraPose CameraPose;
@@ -118,8 +118,8 @@ class FCameraNodeEvaluator
 
 public:
 
-	FCameraNodeEvaluator();
-	virtual ~FCameraNodeEvaluator() {}
+	GAMEPLAYCAMERAS_API FCameraNodeEvaluator();
+	GAMEPLAYCAMERAS_API virtual ~FCameraNodeEvaluator() {}
 
 	/** Called to build any children evaluators. */
 	void Build(const FCameraNodeEvaluatorBuildParams& Params);
@@ -156,23 +156,23 @@ public:
 protected:
 
 	/** Called to build any children evaluators. */
-	virtual void OnBuild(const FCameraNodeEvaluatorBuildParams& Params) {}
+	GAMEPLAYCAMERAS_API virtual void OnBuild(const FCameraNodeEvaluatorBuildParams& Params) {}
 
 	/** Initialize this evaluator. */
-	virtual void OnInitialize(const FCameraNodeEvaluatorInitializeParams& Params) {}
+	GAMEPLAYCAMERAS_API virtual void OnInitialize(const FCameraNodeEvaluatorInitializeParams& Params) {}
 
 	/** Get the list of children under this evaluator. */
 	virtual FCameraNodeEvaluatorChildrenView OnGetChildren() { return FCameraNodeEvaluatorChildrenView(); }
 
 	/** Run this evaluator. */
-	virtual void OnRun(const FCameraNodeEvaluationParams& Params, FCameraNodeEvaluationResult& OutResult) {}
+	GAMEPLAYCAMERAS_API virtual void OnRun(const FCameraNodeEvaluationParams& Params, FCameraNodeEvaluationResult& OutResult) {}
 
 	/** Collect referenced UObjects. */
-	virtual void OnAddReferencedObjects(FReferenceCollector& Collector) {}
+	GAMEPLAYCAMERAS_API virtual void OnAddReferencedObjects(FReferenceCollector& Collector) {}
 
 #if UE_GAMEPLAY_CAMERAS_DEBUG
 	/** Called to create debug blocks for this node evaluator. */
-	virtual void OnBuildDebugBlocks(const FCameraDebugBlockBuildParams& Params, FCameraDebugBlockBuilder& Builder);
+	GAMEPLAYCAMERAS_API virtual void OnBuildDebugBlocks(const FCameraDebugBlockBuildParams& Params, FCameraDebugBlockBuilder& Builder);
 #endif  // UE_GAMEPLAY_CAMERAS_DEBUG
 
 private:
diff --git a/Engine/Plugins/Cameras/GameplayCameras/Source/GameplayCameras/Public/Core/CameraObjectRtti.h b/Engine/Plugins/Cameras/GameplayCameras/Source/GameplayCameras/Public/Core/CameraObjectRtti.h
--- a/Engine/Plugins/Cameras/GameplayCameras/Source/GameplayCameras/Public/Core/CameraObjectRtti.h
+++ b/Engine/Plugins/Cameras/GameplayCameras/Source/GameplayCameras/Public/Core/CameraObjectRtti.h
@@ -82,7 +82,7 @@ class FCameraObjectTypeRegistry
 {
 public:
 
-	static FCameraObjectTypeRegistry& Get();
+	GAMEPLAYCAMERAS_API static FCameraObjectTypeRegistry& Get();
 
 	GAMEPLAYCAMERAS_API void RegisterType(FCameraObjectTypeID TypeID, FCameraObjectTypeInfo&& TypeInfo);
 	GAMEPLAYCAMERAS_API FCameraObjectTypeID FindTypeByName(const FName& TypeName);
diff --git a/Engine/Plugins/Cameras/GameplayCameras/Source/GameplayCameras/Public/Core/CameraParameters.h b/Engine/Plugins/Cameras/GameplayCameras/Source/GameplayCameras/Public/Core/CameraParameters.h
--- a/Engine/Plugins/Cameras/GameplayCameras/Source/GameplayCameras/Public/Core/CameraParameters.h
+++ b/Engine/Plugins/Cameras/GameplayCameras/Source/GameplayCameras/Public/Core/CameraParameters.h
@@ -10,7 +10,7 @@
 
 /** Boolean camera parameter. */
 USTRUCT()
-struct FBooleanCameraParameter
+struct GAMEPLAYCAMERAS_API FBooleanCameraParameter
 {
 	GENERATED_BODY()
 
@@ -27,7 +27,7 @@ struct FBooleanCameraParameter
 
 /** Integer camera parameter. */
 USTRUCT()
-struct FInteger32CameraParameter
+struct GAMEPLAYCAMERAS_API FInteger32CameraParameter
 {
 	GENERATED_BODY()
 
@@ -61,7 +61,7 @@ struct FFloatCameraParameter
 
 /** Double camera parameter. */
 USTRUCT()
-struct FDoubleCameraParameter
+struct GAMEPLAYCAMERAS_API FDoubleCameraParameter
 {
 	GENERATED_BODY()
 
@@ -96,7 +96,7 @@ struct FVector2fCameraParameter
 
 /** Vector2d camera parameter. */
 USTRUCT()
-struct FVector2dCameraParameter
+struct GAMEPLAYCAMERAS_API FVector2dCameraParameter
 {
 	GENERATED_BODY()
 
@@ -132,7 +132,7 @@ struct FVector3fCameraParameter
 
 /** Vector3d camera parameter. */
 USTRUCT()
-struct FVector3dCameraParameter
+struct GAMEPLAYCAMERAS_API FVector3dCameraParameter
 {
 	GENERATED_BODY()
 
@@ -186,7 +186,7 @@ struct FVector4dCameraParameter
 
 /** Rotator3f camera parameter. */
 USTRUCT()
-struct FRotator3fCameraParameter
+struct GAMEPLAYCAMERAS_API FRotator3fCameraParameter
 {
 	GENERATED_BODY()
 
@@ -204,7 +204,7 @@ struct FRotator3fCameraParameter
 
 /** Rotator3d camera parameter. */
 USTRUCT()
-struct FRotator3dCameraParameter
+struct GAMEPLAYCAMERAS_API FRotator3dCameraParameter
 {
 	GENERATED_BODY()
 
@@ -239,7 +239,7 @@ struct FTransform3fCameraParameter
 
 /** Transform3d camera parameter. */
 USTRUCT()
-struct FTransform3dCameraParameter
+struct GAMEPLAYCAMERAS_API FTransform3dCameraParameter
 {
 	GENERATED_BODY()
 
diff --git a/Engine/Plugins/Cameras/GameplayCameras/Source/GameplayCameras/Public/Core/CameraSystemEvaluator.h b/Engine/Plugins/Cameras/GameplayCameras/Source/GameplayCameras/Public/Core/CameraSystemEvaluator.h
--- a/Engine/Plugins/Cameras/GameplayCameras/Source/GameplayCameras/Public/Core/CameraSystemEvaluator.h
+++ b/Engine/Plugins/Cameras/GameplayCameras/Source/GameplayCameras/Public/Core/CameraSystemEvaluator.h
@@ -33,7 +33,7 @@ class FRootCameraDebugBlock;
 /**
  * Parameter structure for initializing a new camera system evaluator.
  */
-struct FCameraSystemEvaluatorCreateParams
+struct GAMEPLAYCAMERAS_API FCameraSystemEvaluatorCreateParams
 {
 	TObjectPtr<UObject> Owner;
 
@@ -44,7 +44,7 @@ struct FCameraSystemEvaluatorCreateParams
 /**
  * Parameter structure for updating the camera system.
  */
-struct FCameraSystemEvaluationUpdateParams
+struct GAMEPLAYCAMERAS_API FCameraSystemEvaluationUpdateParams
 {
 	/** Time interface for the update. */
 	float DeltaTime = 0.f;
@@ -80,21 +80,21 @@ class FCameraSystemEvaluator : public TSharedFromThis<FCameraSystemEvaluator>
 public:
 
 	/** Builds a new camera system. Initialize must be called before the system is used. */
-	FCameraSystemEvaluator();
+	GAMEPLAYCAMERAS_API FCameraSystemEvaluator();
 
 	/** Initializes the camera system. */
-	void Initialize(TObjectPtr<UObject> InOwner = nullptr);
+	GAMEPLAYCAMERAS_API void Initialize(TObjectPtr<UObject> InOwner = nullptr);
 	/** Initializes the camera system. */
-	void Initialize(const FCameraSystemEvaluatorCreateParams& Params);
+	GAMEPLAYCAMERAS_API void Initialize(const FCameraSystemEvaluatorCreateParams& Params);
 
 public:
 
 	/** Push a new evaluation context on the stack. */
-	void PushEvaluationContext(TSharedRef<FCameraEvaluationContext> EvaluationContext);
+	GAMEPLAYCAMERAS_API void PushEvaluationContext(TSharedRef<FCameraEvaluationContext> EvaluationContext);
 	/** Remove an existing evaluation context from the stack. */
-	void RemoveEvaluationContext(TSharedRef<FCameraEvaluationContext> EvaluationContext);
+	GAMEPLAYCAMERAS_API void RemoveEvaluationContext(TSharedRef<FCameraEvaluationContext> EvaluationContext);
 	/** Pop the active (top) evaluation context from the stack. */
-	void PopEvaluationContext();
+	GAMEPLAYCAMERAS_API void PopEvaluationContext();
 
 	/** Gets the context stack. */
 	FCameraEvaluationContextStack& GetEvaluationContextStack() { return ContextStack; }
@@ -104,7 +104,7 @@ public:
 public:
 
 	/** Run an update of the camera system. */
-	void Update(const FCameraSystemEvaluationUpdateParams& Params);
+	GAMEPLAYCAMERAS_API void Update(const FCameraSystemEvaluationUpdateParams& Params);
 
 	/** Returns the root node evaluator. */
 	FRootCameraNodeEvaluator* GetRootNodeEvaluator() const { return RootEvaluator; }
@@ -113,15 +113,15 @@ public:
 	const FCameraSystemEvaluationUpdateResult& GetEvaluatedResult() const { return Result; }
 
 	/** Get the last evaluated camera. */
-	void GetEvaluatedCameraView(FMinimalViewInfo& DesiredView);
+	GAMEPLAYCAMERAS_API void GetEvaluatedCameraView(FMinimalViewInfo& DesiredView);
 
 #if UE_GAMEPLAY_CAMERAS_DEBUG
-	void DebugUpdate(const FCameraSystemDebugUpdateParams& Params);
+	GAMEPLAYCAMERAS_API void DebugUpdate(const FCameraSystemDebugUpdateParams& Params);
 #endif  // UE_GAMEPLAY_CAMERAS_DEBUG
 
 public:
 
-	void AddReferencedObjects(FReferenceCollector& Collector);
+	GAMEPLAYCAMERAS_API void AddReferencedObjects(FReferenceCollector& Collector);
 
 private:
 
diff --git a/Engine/Plugins/Cameras/GameplayCameras/Source/GameplayCameras/Public/Core/CameraVariableTable.h b/Engine/Plugins/Cameras/GameplayCameras/Source/GameplayCameras/Public/Core/CameraVariableTable.h
--- a/Engine/Plugins/Cameras/GameplayCameras/Source/GameplayCameras/Public/Core/CameraVariableTable.h
+++ b/Engine/Plugins/Cameras/GameplayCameras/Source/GameplayCameras/Public/Core/CameraVariableTable.h
@@ -53,7 +53,7 @@ public:
 	FCameraVariableTable();
 	FCameraVariableTable(FCameraVariableTable&& Other);
 	FCameraVariableTable& operator=(FCameraVariableTable&& Other);
-	~FCameraVariableTable();
+	GAMEPLAYCAMERAS_API ~FCameraVariableTable();
 
 	FCameraVariableTable(const FCameraVariableTable&) = delete;
 	FCameraVariableTable& operator=(const FCameraVariableTable&) = delete;
@@ -66,7 +66,7 @@ public:
 	 * This may re-allocate the internal memory buffer. It's recommended to pre-compute
 	 * the allocation information needed for a table, and initialize it once.
 	 */
-	void AddVariable(const FCameraVariableDefinition& VariableDefinition);
+	GAMEPLAYCAMERAS_API void AddVariable(const FCameraVariableDefinition& VariableDefinition);
 
 public:
 
