Cameras: add comments on recently added classes

diff --git a/Engine/Plugins/Cameras/GameplayCameras/Source/GameplayCameras/Public/Core/ObjectTreeGraphObject.h b/Engine/Plugins/Cameras/GameplayCameras/Source/GameplayCameras/Public/Core/ObjectTreeGraphObject.h
--- a/Engine/Plugins/Cameras/GameplayCameras/Source/GameplayCameras/Public/Core/ObjectTreeGraphObject.h
+++ b/Engine/Plugins/Cameras/GameplayCameras/Source/GameplayCameras/Public/Core/ObjectTreeGraphObject.h
@@ -16,17 +16,24 @@ class UObjectTreeGraphObject : public UInterface
 
 #if WITH_EDITOR
 
+/** Flags representing supported optional APIs of an IObjectTreeGraphObject. */
 enum class EObjectTreeGraphObjectSupportFlags
 {
+	/** None supported, only graph node position. */
 	None,
+	/** Supports storing a comment text. */
 	CommentText = 1 << 0,
+	/** Supports custom renaming the graph node. */
 	CustomRename = 1 << 1
 };
 ENUM_CLASS_FLAGS(EObjectTreeGraphObjectSupportFlags);
 
 #endif  // WITH_EDITOR
 
-
+/**
+ * An interface that UObjects can implement to customize how they are represented
+ * and interacted with inside an object tree graph (see UObjectTreeGraph).
+ */
 class IObjectTreeGraphObject
 {
 	GENERATED_BODY()
@@ -35,16 +42,28 @@ class IObjectTreeGraphObject
 
 public:
 
+	// Minimal API.
+	
+	/** Gets the canvas position for the graph node representing this object. */
 	virtual void GetGraphNodePosition(int32& NodePosX, int32& NodePosY) const {}
+	/** Called to save the canvas position of the graph node representing this object. */
 	virtual void OnGraphNodeMoved(int32 NodePosX, int32 NodePosY) {}
 
+	// Optional API.
+
+	/** Gets optional APIs support flags. */
 	virtual EObjectTreeGraphObjectSupportFlags GetSupportFlags() const { return EObjectTreeGraphObjectSupportFlags::None; }
+	/** Gets whether this object supports the given optional API. */
 	bool HasSupportFlags(EObjectTreeGraphObjectSupportFlags InFlags) const { return EnumHasAllFlags(GetSupportFlags(), InFlags); }
 
+	/** Gets the graph node's comment text. */
 	virtual const FString& GetGraphNodeCommentText() const { static const FString EmptyString; return EmptyString; }
+	/** Called to save a new comment text. */
 	virtual void OnUpdateGraphNodeCommentText(const FString& NewComment) {}
 
+	/** Gets the custom name for the graph node. */
 	virtual const FString& GetGraphNodeName() const { static const FString EmptyString; return EmptyString; }
+	/** Called to save a new custom name. */
 	virtual void OnRenameGraphNode(const FString& NewName) {}
 
 #endif  // WITH_EDITOR
diff --git a/Engine/Plugins/Cameras/GameplayCameras/Source/GameplayCamerasEditor/Public/Editors/CameraTransitionGraphSchema.h b/Engine/Plugins/Cameras/GameplayCameras/Source/GameplayCamerasEditor/Public/Editors/CameraTransitionGraphSchema.h
--- a/Engine/Plugins/Cameras/GameplayCameras/Source/GameplayCamerasEditor/Public/Editors/CameraTransitionGraphSchema.h
+++ b/Engine/Plugins/Cameras/GameplayCameras/Source/GameplayCamerasEditor/Public/Editors/CameraTransitionGraphSchema.h
@@ -6,6 +6,9 @@
 
 #include "CameraTransitionGraphSchema.generated.h"
 
+/**
+ * Schema class for camera transition graph.
+ */
 UCLASS()
 class UCameraTransitionGraphSchema : public UObjectTreeGraphSchema
 {
@@ -20,6 +23,12 @@ protected:
 	virtual void FilterGraphContextPlaceableClasses(TArray<UClass*>& InOutClasses) const override;
 };
 
+/**
+ * Graph action to create a new transition node.
+ *
+ * We need a custom action for this because we need to switch the "self" pin according to whether
+ * we want an enter or exit transition.
+ */
 USTRUCT()
 struct FCameraTransitionGraphSchemaAction_NewTransitionNode : public FObjectGraphSchemaAction_NewNode
 {
@@ -33,6 +42,7 @@ public:
 		Exit
 	};
 
+	/** The transition type to create. */
 	ETransitionType TransitionType;
 
 public:
diff --git a/Engine/Plugins/Cameras/GameplayCameras/Source/GameplayCamerasEditor/Public/Editors/ObjectTreeGraph.h b/Engine/Plugins/Cameras/GameplayCameras/Source/GameplayCamerasEditor/Public/Editors/ObjectTreeGraph.h
--- a/Engine/Plugins/Cameras/GameplayCameras/Source/GameplayCamerasEditor/Public/Editors/ObjectTreeGraph.h
+++ b/Engine/Plugins/Cameras/GameplayCameras/Source/GameplayCamerasEditor/Public/Editors/ObjectTreeGraph.h
@@ -15,11 +15,16 @@
 class UObjectTreeGraphNode;
 class UObjectTreeGraphSchema;
 
+/** Specifies where to find graph node objects. */
 enum class EObjectTreeGraphBuildSource
 {
+	/** Create graph nodes from all eligible objects found in the same package. */
 	RootObjectPackage
 };
 
+/**
+ * A node graph that represents a hierarchy of objects and their relationships.
+ */
 UCLASS()
 class UObjectTreeGraph : public UEdGraph
 {
@@ -27,14 +32,20 @@ class UObjectTreeGraph : public UEdGraph
 
 public:
 
+	/** Creates a new graph. */
 	UObjectTreeGraph(const FObjectInitializer& ObjInit);
 
+	/** Initializes the graph given a root object and a graph configuration. */
 	void Initialize(TObjectPtr<UObject> InRootObject, const FObjectTreeGraphConfig& InConfig);
 
+	/** Gets the root object. */
 	UObject* GetRootObject() const { return WeakRootObject.Get(); }
+	/** Gets the root object's graph node. */
 	UObjectTreeGraphNode* GetRootObjectNode() const { return RootObjectNode; }
+	/** Gets the graph configuration. */
 	const FObjectTreeGraphConfig& GetConfig() const;
 
+	/** Rebuilds the node graph. */
 	void RebuildGraph(EObjectTreeGraphBuildSource InSource);
 
 private:
diff --git a/Engine/Plugins/Cameras/GameplayCameras/Source/GameplayCamerasEditor/Public/Editors/ObjectTreeGraphConfig.h b/Engine/Plugins/Cameras/GameplayCameras/Source/GameplayCamerasEditor/Public/Editors/ObjectTreeGraphConfig.h
--- a/Engine/Plugins/Cameras/GameplayCameras/Source/GameplayCamerasEditor/Public/Editors/ObjectTreeGraphConfig.h
+++ b/Engine/Plugins/Cameras/GameplayCameras/Source/GameplayCamerasEditor/Public/Editors/ObjectTreeGraphConfig.h
@@ -30,58 +30,80 @@ DECLARE_DELEGATE_TwoParams(FOnFormatObjectDisplayName, const UObject*, FText&);
 	private:\
 		FieldType _##FieldName;
 
+/**
+ * A structure providing optional configuration options for a given object class.
+ */
 struct FObjectTreeGraphClassConfig
 {
 public:
 
+	/** The subclass of graph nodes to create. */
 	OTGCC_FIELD(TSubclassOf<UObjectTreeGraphNode>, GraphNodeClass)
 
+	/** The display name of the self pin. */
 	OTGCC_FIELD(FText, SelfPinFriendlyName)
+	/** The direction of the self pin. */
 	OTGCC_FIELD(EEdGraphPinDirection, SelfPinDirection)
+	/** Whether graph nodes for this class have a self pin. */
 	OTGCC_FIELD(bool, HasSelfPin)
 
+	/** Default direction of property pins. */
 	OTGCC_FIELD(EEdGraphPinDirection, DefaultPropertyPinDirection)
 
+	/** Color of the graph node's title. */
 	OTGCC_FIELD(TOptional<FLinearColor>, NodeTitleColor)
+	/** Color of the graph node's body. */
 	OTGCC_FIELD(TOptional<FLinearColor>, NodeBodyTintColor)
 
+	/** Whether the graph node title uses the underlying object's name instead of its class name. */
 	OTGCC_FIELD(bool, NodeTitleUsesObjectName)
+	/** A custom call back to get the object's display name used in the graph node title. */
 	OTGCC_FIELD(FOnGetObjectClassDisplayName, OnGetObjectClassDisplayName)
 
+	/** Whether users can create new objects of this class in the graph. */
 	OTGCC_FIELD(bool, CanCreateNew)
+	/** Whether users can duplicate objects of this class in the graph. */
 	OTGCC_FIELD(bool, CanDuplicate)
+	/** Whether users can delete objects of this class in the graph. */
 	OTGCC_FIELD(bool, CanDelete)
 
 public:
 
 	FObjectTreeGraphClassConfig();
 
+	/** A shortcut for disabling CanCreateNew, CanDuplicate, and CanDelete. */
 	FObjectTreeGraphClassConfig& OnlyAsRoot();
 
 public:
 
+	/** Gets the name suffixes to strip. */
 	TArrayView<const FString> StripDisplayNameSuffixes() const { return _StripDisplayNameSuffixes; }
 
+	/** Adds a new suffix to strip from the display name. */
 	FObjectTreeGraphClassConfig& StripDisplayNameSuffix(const FString& InSuffix)
 	{
 		_StripDisplayNameSuffixes.Add(InSuffix);
 		return *this;
 	}
 
+	/** Adds multiple suffixes to strip from the display name. */
 	FObjectTreeGraphClassConfig& StripDisplayNameSuffixes(std::initializer_list<FString> InSuffixes)
 	{
 		_StripDisplayNameSuffixes.Append(InSuffixes);
 		return *this;
 	}
 
+	/** Gets the custom property pin directions for given named properties. */
 	const TMap<FName, EEdGraphPinDirection>& PropertyPinDirections() const { return _PropertyPinDirections; }
 
+	/** Adds a new custom property pin direction for a given named property. */
 	FObjectTreeGraphClassConfig& SetPropertyPinDirection(const FName& InPropertyName, EEdGraphPinDirection InDirection)
 	{
 		_PropertyPinDirections.Add(InPropertyName, InDirection);
 		return *this;
 	}
 
+	/** Gets the custom property pin direction for a given named property. */
 	EEdGraphPinDirection GetPropertyPinDirection(const FName& InPropertyName) const
 	{
 		if (const EEdGraphPinDirection* PinDirection = _PropertyPinDirections.Find(InPropertyName))
@@ -97,32 +119,72 @@ private:
 	TMap<FName, EEdGraphPinDirection> _PropertyPinDirections;
 };
 
+/**
+ * A structure that provides all the information needed to build, edit, and maintain an 
+ * object tree graph.
+ */
 struct FObjectTreeGraphConfig
 {
 public:
 
+	/**
+	 * The list of connectable object classes.
+	 *
+	 * Objects whose class is connectable (which includes sub-classes) will be eligible 
+	 * to be nodes in the graph. Properties on those objects that point to other connectable
+	 * objects (either with a direct object property or an array property) will show up
+	 * as pins on the object's node.
+	 */
 	TArray<UClass*> ConnectableObjectClasses;
+	/**
+	 * The list of unconnectable object classes.
+	 *
+	 * This serves as an exception list to the ConnectableObjectClasses list.
+	 */
 	TArray<UClass*> NonConnectableObjectClasses;
 	
+	/**
+	 * The default graph node class to use in the graph. Defaults to UObjectTreeGraphNode.
+	 */
 	TSubclassOf<UObjectTreeGraphNode> DefaultGraphNodeClass;
 
+	/** The default title color for an object's graph node. */
 	FLinearColor DefaultGraphNodeTitleColor;
+	/** The default body color for an object's graph node. */
 	FLinearColor DefaultGraphNodeBodyTintColor;
 
+	/** A custom callback to format an object's display name. */
 	FOnFormatObjectDisplayName OnFormatObjectDisplayName;
 
+	/** Advanced, optional bits of configuration for specific classes and sub-classes of objects. */
 	TMap<UClass*, FObjectTreeGraphClassConfig> ObjectClassConfigs;
 
 public:
 
+	/** Creates a new graph config. */
 	FObjectTreeGraphConfig();
 
+	/**
+	 * Returns whether the given class is connectable.
+	 *
+	 * It is connectable if it, or one of its parent classes, is inside ConnectableObjectClasses,
+	 * and nor it or any of its parent classes is in NonConnectableObjectClasses.
+	 */
 	bool IsConnectable(UClass* InObjectClass) const;
+
+	/**
+	 * Gets all possible known connectable classes.
+	 *
+	 * @param bPlaceableOnly  If set, only return those that can be created.
+	 */
 	void GetConnectableClasses(TArray<UClass*>& OutClasses, bool bPlaceableOnly = false);
 	
+	/** Gets the advanced class-specific configuration for the given class. */
 	const FObjectTreeGraphClassConfig& GetObjectClassConfig(const UClass* InObjectClass) const;
 
+	/** Computes the display name of the given object. */
 	FText GetDisplayNameText(const UObject* InObject) const;
+	/** Computes the display name of the given object class. */
 	FText GetDisplayNameText(const UClass* InClass) const;
 
 private:
diff --git a/Engine/Plugins/Cameras/GameplayCameras/Source/GameplayCamerasEditor/Public/Editors/ObjectTreeGraphNode.h b/Engine/Plugins/Cameras/GameplayCameras/Source/GameplayCamerasEditor/Public/Editors/ObjectTreeGraphNode.h
--- a/Engine/Plugins/Cameras/GameplayCameras/Source/GameplayCamerasEditor/Public/Editors/ObjectTreeGraphNode.h
+++ b/Engine/Plugins/Cameras/GameplayCameras/Source/GameplayCamerasEditor/Public/Editors/ObjectTreeGraphNode.h
@@ -14,6 +14,9 @@ class UObjectTreeGraph;
 struct FObjectTreeGraphConfig;
 struct FObjectTreeGraphClassConfig;
 
+/**
+ * A graph node that represents an object inside an object tree graph.
+ */
 UCLASS()
 class UObjectTreeGraphNode : public UEdGraphNode
 {
@@ -21,23 +24,37 @@ class UObjectTreeGraphNode : public UEdGraphNode
 
 public:
 
+	/** Creates a new graph node. */
 	UObjectTreeGraphNode(const FObjectInitializer& ObjInit);
 
+	/** Initializes this graph node for the given object. */
 	void Initialize(UObject* InObject);
 
+	/** Gets the underlying object represented by this graph node. */
 	UObject* GetObject() const { return Object.Get(); }
+	/** Gets all connectable properties on the underlying object. */
 	void GetAllConnectableProperties(TArray<FProperty*>& OutProperties) const;
 
+	/** Finds the self pin that represents the underlying object itself. */
 	UEdGraphPin* GetSelfPin() const;
+	/** Changes the direction of the self pin. */
 	void OverrideSelfPinDirection(EEdGraphPinDirection Direction);
 
+	/** Finds the pin for the given object property. */
 	UEdGraphPin* GetPinForProperty(FObjectProperty* InProperty) const;
+	/** Finds the pin for the given item in an array property. */
 	UEdGraphPin* GetPinForProperty(FArrayProperty* InProperty, int32 Index) const;
+	/** Finds the extra free pin used to add new items in an array property. */
 	UEdGraphPin* GetPinForPropertyNewItem(FArrayProperty* InProperty, bool bCreateNew);
+	/** Gets the underlying property represented by the given pin. */
 	FProperty* GetPropertyForPin(const UEdGraphPin* InPin) const;
+	/** Gets the type of object that can connect to the given pin. */
 	UClass* GetConnectedObjectClassForPin(const UEdGraphPin* InPin) const;
+	/** Gets the index of the given pin's underlying value inside an array property. */
 	int32 GetIndexOfArrayPin(const UEdGraphPin* InPin) const;
 
+public:
+
 	// UEdGraphNode interface.
 	virtual FText GetNodeTitle(ENodeTitleType::Type TitleType) const override;
 	virtual TSharedPtr<SGraphNode> CreateVisualWidget() override;
diff --git a/Engine/Plugins/Cameras/GameplayCameras/Source/GameplayCamerasEditor/Public/Editors/ObjectTreeGraphSchema.h b/Engine/Plugins/Cameras/GameplayCameras/Source/GameplayCamerasEditor/Public/Editors/ObjectTreeGraphSchema.h
--- a/Engine/Plugins/Cameras/GameplayCameras/Source/GameplayCamerasEditor/Public/Editors/ObjectTreeGraphSchema.h
+++ b/Engine/Plugins/Cameras/GameplayCameras/Source/GameplayCamerasEditor/Public/Editors/ObjectTreeGraphSchema.h
@@ -12,6 +12,9 @@ class UObjectTreeGraph;
 class UObjectTreeGraphNode;
 struct FObjectTreeGraphClassConfig;
 
+/**
+ * Schema class for an object tree graph.
+ */
 UCLASS()
 class UObjectTreeGraphSchema : public UEdGraphSchema
 {
@@ -20,18 +23,20 @@ class UObjectTreeGraphSchema : public UEdGraphSchema
 public:
 
 	// Pin categories.
-	static const FName PC_Self;
-	static const FName PC_Property;
+	static const FName PC_Self;			// A "self" pin.
+	static const FName PC_Property;		// A property pin.
 
 	// Pin sub-categories.
-	static const FName PSC_ObjectProperty;
-	static const FName PSC_ArrayProperty;
-	static const FName PSC_ArrayPropertyItem;
+	static const FName PSC_ObjectProperty;		// A normal object property pin.
+	static const FName PSC_ArrayProperty;		// An array property pin (generally hidden).
+	static const FName PSC_ArrayPropertyItem;	// A pin for an item inside an array property.
 
 public:
 
+	/** Creates a new schema. */
 	UObjectTreeGraphSchema(const FObjectInitializer& ObjInit);
 
+	/** Creates an object graph node for the given object. */
 	UObjectTreeGraphNode* CreateObjectNode(UObjectTreeGraph* InGraph, UObject* InObject) const;
 
 public:
@@ -64,6 +69,9 @@ protected:
 	const FObjectTreeGraphClassConfig& GetObjectClassConfig(const UObjectTreeGraph* InGraph, UClass* InObjectClass) const;
 };
 
+/**
+ * Graph action to create a new object (and corresponding graph node) of a given class.
+ */
 USTRUCT()
 struct FObjectGraphSchemaAction_NewNode : public FEdGraphSchemaAction
 {
@@ -71,9 +79,11 @@ struct FObjectGraphSchemaAction_NewNode : public FEdGraphSchemaAction
 
 public:
 
+	/** The outer for the new object. Defaults to the root object's package. */
 	UPROPERTY()
 	TObjectPtr<UObject> ObjectOuter;
 
+	/** The class of the new object. */
 	UPROPERTY()
 	TObjectPtr<UClass> ObjectClass;
 
diff --git a/Engine/Plugins/Cameras/GameplayCameras/Source/GameplayCamerasEditor/Public/Editors/SCameraRigAssetEditor.h b/Engine/Plugins/Cameras/GameplayCameras/Source/GameplayCamerasEditor/Public/Editors/SCameraRigAssetEditor.h
--- a/Engine/Plugins/Cameras/GameplayCameras/Source/GameplayCamerasEditor/Public/Editors/SCameraRigAssetEditor.h
+++ b/Engine/Plugins/Cameras/GameplayCameras/Source/GameplayCamerasEditor/Public/Editors/SCameraRigAssetEditor.h
@@ -20,20 +20,32 @@ struct FObjectTreeGraphConfig;
 namespace UE::Cameras
 {
 
+/** The current mode of the camera rig asset editor. */
 enum class ECameraRigAssetEditorMode
 {
+	/** Show the node hierarchy editor. */
 	NodeGraph,
+	/** Show the transition editor. */
 	TransitionGraph
 };
 
+/**
+ * A camera rig asset editor.
+ *
+ * This implements only the dual-graph editor, for the node hierarchy and transitions.
+ * The rest of the UI such as the details view or the toolbox aren't included here.
+ */
 class SCameraRigAssetEditor : public SCompoundWidget
 {
 public:
 
 	SLATE_BEGIN_ARGS(SCameraRigAssetEditor)
 	{}
+		/** The camera rig asset to edit. */
 		SLATE_ARGUMENT(TObjectPtr<UCameraRigAsset>, CameraRigAsset)
+		/** The details view to synchronize with the graph selection. */
 		SLATE_ARGUMENT(TSharedPtr<IDetailsView>, DetailsView)
+		/** The toolkit inside which this editor lives, if any. */
 		SLATE_ARGUMENT(TWeakPtr<FAssetEditorToolkit>, AssetEditorToolkit)
 	SLATE_END_ARGS()
 
@@ -41,15 +53,22 @@ public:
 
 public:
 
+	/** Gets the current editor mode. */
 	ECameraRigAssetEditorMode GetEditorMode() const;
+	/** Checks if the editor is in the current mode. */
 	bool IsEditorMode(ECameraRigAssetEditorMode InMode) const;
+	/** Changes the editor's current mode. */
 	void SetEditorMode(ECameraRigAssetEditorMode InMode);
 
+	/** Gets both the node hierarchy and transition graphs. */
 	void GetGraphs(TArray<UEdGraph*>& OutGraphs) const;
 
+	/** Gets the graph configuration for the current mode. */
 	const FObjectTreeGraphConfig& GetFocusedGraphConfig() const;
 
+	/** Focuses the current graph to the root object node. */
 	void FocusHome();
+	/** Jumps the current graph to the given node. */
 	void JumpToNode(UEdGraphNode* InGraphNode);
 
 protected:
diff --git a/Engine/Plugins/Cameras/GameplayCameras/Source/GameplayCamerasEditor/Public/Editors/SFindInObjectTreeGraph.h b/Engine/Plugins/Cameras/GameplayCameras/Source/GameplayCamerasEditor/Public/Editors/SFindInObjectTreeGraph.h
--- a/Engine/Plugins/Cameras/GameplayCameras/Source/GameplayCamerasEditor/Public/Editors/SFindInObjectTreeGraph.h
+++ b/Engine/Plugins/Cameras/GameplayCameras/Source/GameplayCamerasEditor/Public/Editors/SFindInObjectTreeGraph.h
@@ -18,31 +18,50 @@ class SSearchBox;
 
 class SFindInObjectTreeGraph;
 
+/** 
+ * Structure for a search result inside an object tree graph.
+ */
 struct FFindInObjectTreeGraphResult
 {
 public:
 
+	/** Parent result. */
 	TWeakPtr<FFindInObjectTreeGraphResult> Parent;
+	/** Children results. */
 	TArray<TSharedPtr<FFindInObjectTreeGraphResult>> Children;
 
+	/** Custom text for this result. */
 	FText CustomText;
+	/** The graph node that this result refers to. */
 	TWeakObjectPtr<UEdGraphNode> GraphNode;
+	/** The graph pin that this result refers to. */
 	FEdGraphPinReference GraphPin;
 
 public:
 
+	/** Creates a new result with a custom text. */
 	FFindInObjectTreeGraphResult(const FText& InCustomText);
+	/** Creates a new result referring to a graph node, under a parent result. */
 	FFindInObjectTreeGraphResult(TSharedPtr<FFindInObjectTreeGraphResult>& InParent, UEdGraphNode* InGraphNode);
+	/** Creates a new result referring to a graph pin, under a parent result. */
 	FFindInObjectTreeGraphResult(TSharedPtr<FFindInObjectTreeGraphResult>& InParent, UEdGraphPin* InGraphPin);
 
+	/** Gets the icon for this result. */
 	TSharedRef<SWidget>	GetIcon() const;
+	/** Gets the category for this result. */
 	FText GetCategory() const;
+	/** Gets the display text for this result. */
 	FText GetText() const;
+	/** Gets the comment text for this result. */
 	FText GetCommentText() const;
 
+	/** Go to the graph node, pin, etc. */
 	FReply OnClick(TSharedRef<SFindInObjectTreeGraph> FindInObjectTreeGraph);
 };
 
+/**
+ * A search panel to find things in one or more object tree graphs.
+ */
 class SFindInObjectTreeGraph : public SCompoundWidget
 {
 public:
@@ -52,8 +71,11 @@ public:
 
 	SLATE_BEGIN_ARGS(SFindInObjectTreeGraph)
 	{}
+		/** The graphs to search. */
 		SLATE_ARGUMENT(TArray<UEdGraph*>, GraphsToSearch)
+		/** The callback to invoke when a search result wants to focus a node. */
 		SLATE_ARGUMENT(FOnJumpToNodeRequested, OnJumpToNodeRequested)
+		/** The callback to invoke when a search result wants to focus a pin. */
 		SLATE_ARGUMENT(FOnJumpToPinRequested, OnJumpToPinRequested)
 	SLATE_END_ARGS()
 
diff --git a/Engine/Plugins/Cameras/GameplayCameras/Source/GameplayCamerasEditor/Public/Editors/SObjectTreeGraphEditor.h b/Engine/Plugins/Cameras/GameplayCameras/Source/GameplayCamerasEditor/Public/Editors/SObjectTreeGraphEditor.h
--- a/Engine/Plugins/Cameras/GameplayCameras/Source/GameplayCamerasEditor/Public/Editors/SObjectTreeGraphEditor.h
+++ b/Engine/Plugins/Cameras/GameplayCameras/Source/GameplayCamerasEditor/Public/Editors/SObjectTreeGraphEditor.h
@@ -11,6 +11,9 @@ class IDetailsView;
 class UObjectTreeGraph;
 class UObjectTreeGraphNode;
 
+/**
+ * A graph editor for an object tree graph.
+ */
 class SObjectTreeGraphEditor 
 	: public SCompoundWidget
 	, public FEditorUndoClient
@@ -19,12 +22,19 @@ public:
 
 	SLATE_BEGIN_ARGS(SObjectTreeGraphEditor)
 	{}
+		/** Any additional command mappings to use in the graph editor. */
 		SLATE_ARGUMENT(TSharedPtr<FUICommandList>, AdditionalCommands)
+		/** A custom graph title bar. A default is provided if this isn't specified. */
 		SLATE_ARGUMENT(TSharedPtr<SWidget>, GraphTitleBar)
+		/** The details view to use for showing the current graph selection. */
 		SLATE_ARGUMENT(TSharedPtr<IDetailsView>, DetailsView)
+		/** The graph to show in the editor. */
 		SLATE_ARGUMENT(UObjectTreeGraph*, GraphToEdit)
+		/** The toolkit inside which this editor lives, if any. */
 		SLATE_ARGUMENT(TWeakPtr<FAssetEditorToolkit>, AssetEditorToolkit)
+		/** The graph editor appearance. */
 		SLATE_ATTRIBUTE(FGraphAppearanceInfo, Appearance)
+		/** The graph editor title. */
 		SLATE_ATTRIBUTE(FText, GraphTitle)
 	SLATE_END_ARGS()
 
diff --git a/Engine/Plugins/Cameras/GameplayCameras/Source/GameplayCamerasEditor/Public/Editors/SObjectTreeGraphTitleBar.h b/Engine/Plugins/Cameras/GameplayCameras/Source/GameplayCamerasEditor/Public/Editors/SObjectTreeGraphTitleBar.h
--- a/Engine/Plugins/Cameras/GameplayCameras/Source/GameplayCamerasEditor/Public/Editors/SObjectTreeGraphTitleBar.h
+++ b/Engine/Plugins/Cameras/GameplayCameras/Source/GameplayCamerasEditor/Public/Editors/SObjectTreeGraphTitleBar.h
@@ -21,6 +21,9 @@ struct FObjectTreeGraphInfo
 	FString GraphName;
 };
 
+/**
+ * The default object graph title bar used by the object tree graph editor.
+ */
 class SObjectTreeGraphTitleBar : public SCompoundWidget
 {
 public:
diff --git a/Engine/Plugins/Cameras/GameplayCameras/Source/GameplayCamerasEditor/Public/Editors/SObjectTreeGraphToolbox.h b/Engine/Plugins/Cameras/GameplayCameras/Source/GameplayCamerasEditor/Public/Editors/SObjectTreeGraphToolbox.h
--- a/Engine/Plugins/Cameras/GameplayCameras/Source/GameplayCamerasEditor/Public/Editors/SObjectTreeGraphToolbox.h
+++ b/Engine/Plugins/Cameras/GameplayCameras/Source/GameplayCamerasEditor/Public/Editors/SObjectTreeGraphToolbox.h
@@ -15,6 +15,10 @@ class UClass;
 
 template<typename> class SListView;
 
+/**
+ * A widget for an object tree graph toolbox entry, showing a specific
+ * instantiable object class that can be added to a graph.
+ */
 class SObjectTreeGraphToolboxEntry : public SCompoundWidget
 {
 public:
@@ -23,8 +27,11 @@ public:
 		: _ObjectClass(nullptr)
 		, _GraphConfig(nullptr)
 	{}
+		/** The object class represented by this entry. */
 		SLATE_ARGUMENT(UClass*, ObjectClass)
+		/** The configuration of the graph this toolbox works for. */
 		SLATE_ARGUMENT(const FObjectTreeGraphConfig*, GraphConfig)
+		/** Text to highlight if a search is ongoing. */
 		SLATE_ATTRIBUTE(FText, HighlightText)
 	SLATE_END_ARGS()
 
@@ -52,6 +59,10 @@ private:
 	const FSlateBrush* PressedImage = nullptr;
 };
 
+/**
+ * A toolbox widget that shows all the possible instantiable classes of objects for a
+ * given object tree graph.
+ */
 class SObjectTreeGraphToolbox : public SCompoundWidget
 {
 public:
@@ -63,6 +74,10 @@ public:
 
 	void Construct(const FArguments& InArgs);
 
+	/**
+	 * Updates the toolbox to reflect the list of instantiable objects for
+	 * the given graph configuration.
+	 */
 	void SetGraphConfig(const FObjectTreeGraphConfig& InGraphConfig);
 
 protected:
@@ -100,6 +115,10 @@ private:
 	bool bUpdateFilteredItemSource = false;
 };
 
+/**
+ * Drag-drop operation for creating a new object (and corresponding graph node) in an object tree graph
+ * by dragging one of the entries from the toolbox widget.
+ */
 class FObjectTreeClassDragDropOp : public FDecoratedDragDropOp
 {
 public:
