Appease static analyzer
#rnx

diff --git a/Engine/Plugins/Cameras/GameplayCameras/Source/GameplayCamerasEditor/Private/Editors/ObjectTreeGraphSchema.cpp b/Engine/Plugins/Cameras/GameplayCameras/Source/GameplayCamerasEditor/Private/Editors/ObjectTreeGraphSchema.cpp
--- a/Engine/Plugins/Cameras/GameplayCameras/Source/GameplayCamerasEditor/Private/Editors/ObjectTreeGraphSchema.cpp
+++ b/Engine/Plugins/Cameras/GameplayCameras/Source/GameplayCamerasEditor/Private/Editors/ObjectTreeGraphSchema.cpp
@@ -123,6 +123,7 @@ void UObjectTreeGraphSchema::GetGraphContextActions(FGraphContextMenuBuilder& Co
 			}
 		}
 
+		checkSlow(PossibleObjectClass);
 		TSharedRef<FObjectGraphSchemaAction_NewNode> Action = MakeShared<FObjectGraphSchemaAction_NewNode>(
 				CategoryName ? FText::FromString(*CategoryName) : FText::GetEmpty(),
 				PossibleObjectClass->GetDisplayNameText(), 
