Cameras: use strongly-typed variable IDs instead of raw integer types

diff --git a/Engine/Plugins/Cameras/GameplayCameras/Source/GameplayCameras/Private/Core/BlendStackCameraNode.cpp b/Engine/Plugins/Cameras/GameplayCameras/Source/GameplayCameras/Private/Core/BlendStackCameraNode.cpp
--- a/Engine/Plugins/Cameras/GameplayCameras/Source/GameplayCameras/Private/Core/BlendStackCameraNode.cpp
+++ b/Engine/Plugins/Cameras/GameplayCameras/Source/GameplayCameras/Private/Core/BlendStackCameraNode.cpp
@@ -38,10 +38,10 @@ static FAutoConsoleVariableRef CVarGameplayCamerasDebugBlendStackShowUnchanged(
 	GGameplayCamerasDebugBlendStackShowUnchanged,
 	TEXT(""));
 
-bool GGameplayCamerasDebugBlendStackShowVariableIds = false;
-static FAutoConsoleVariableRef CVarGameplayCamerasDebugBlendStackShowVariableIds(
-	TEXT("GameplayCameras.Debug.BlendStack.ShowVariableIds"),
-	GGameplayCamerasDebugBlendStackShowVariableIds,
+bool GGameplayCamerasDebugBlendStackShowVariableIDs = false;
+static FAutoConsoleVariableRef CVarGameplayCamerasDebugBlendStackShowVariableIDs(
+	TEXT("GameplayCameras.Debug.BlendStack.ShowVariableIDs"),
+	GGameplayCamerasDebugBlendStackShowVariableIDs,
 	TEXT(""));
 
 UE_DEFINE_CAMERA_NODE_EVALUATOR(FBlendStackCameraNodeEvaluator)
@@ -482,7 +482,7 @@ FBlendStackCameraDebugBlock* FBlendStackCameraNodeEvaluator::BuildDetailedDebugB
 			{
 				ResultDebugBlock.Initialize(Entry.Result, Builder);
 				ResultDebugBlock.GetCameraPoseDebugBlock()->WithShowUnchangedCVar(TEXT("GameplayCameras.Debug.BlendStack.ShowUnchanged"));
-				ResultDebugBlock.GetVariableTableDebugBlock()->WithShowVariableIdsCVar(TEXT("GameplayCameras.Debug.BlendStack.ShowVariableIds"));
+				ResultDebugBlock.GetVariableTableDebugBlock()->WithShowVariableIDsCVar(TEXT("GameplayCameras.Debug.BlendStack.ShowVariableIDs"));
 			}
 		}
 	}
diff --git a/Engine/Plugins/Cameras/GameplayCameras/Source/GameplayCameras/Private/Core/CameraVariableAssets.cpp b/Engine/Plugins/Cameras/GameplayCameras/Source/GameplayCameras/Private/Core/CameraVariableAssets.cpp
--- a/Engine/Plugins/Cameras/GameplayCameras/Source/GameplayCameras/Private/Core/CameraVariableAssets.cpp
+++ b/Engine/Plugins/Cameras/GameplayCameras/Source/GameplayCameras/Private/Core/CameraVariableAssets.cpp
@@ -9,10 +9,15 @@ UCameraVariableAsset::UCameraVariableAsset(const FObjectInitializer& ObjectInit)
 {
 }
 
+void UCameraVariableAsset::RegenerateVariableID()
+{
+	VariableID = FCameraVariableID::FromHashValue(GetTypeHash(GetFullName()));
+}
+
 FCameraVariableDefinition UCameraVariableAsset::GetVariableDefinition() const
 {
 	FCameraVariableDefinition VariableDefinition;
-	VariableDefinition.VariableId = GetVariableId();
+	VariableDefinition.VariableID = GetVariableID();
 	VariableDefinition.VariableType = GetVariableType();
 #if WITH_EDITORONLY_DATA
 	VariableDefinition.VariableName = GetName();
@@ -25,10 +30,28 @@ void UCameraVariableAsset::Serialize(FArchive& Ar)
 	Super::Serialize(Ar);
 
 #if WITH_EDITORONLY_DATA
-	if ((Ar.IsLoading() && VariableId == 0) || Ar.IsSaving())
+	if ((Ar.IsLoading() && !VariableID.IsValid()) || Ar.IsSaving())
 	{
-		VariableId = GetTypeHash(GetFullName());
+		RegenerateVariableID();
 	}
 #endif
 }
 
+void UCameraVariableAsset::PostInitProperties()
+{
+	RegenerateVariableID();
+	Super::PostInitProperties();
+}
+
+void UCameraVariableAsset::PostRename(UObject* OldOuter, const FName OldName)
+{
+	RegenerateVariableID();
+	Super::PostRename(OldOuter, OldName);
+}
+
+void UCameraVariableAsset::PostDuplicate(EDuplicateMode::Type DuplicateMode)
+{
+	RegenerateVariableID();
+	Super::PostDuplicate(DuplicateMode);
+}
+
diff --git a/Engine/Plugins/Cameras/GameplayCameras/Source/GameplayCameras/Private/Core/CameraVariableTable.cpp b/Engine/Plugins/Cameras/GameplayCameras/Source/GameplayCameras/Private/Core/CameraVariableTable.cpp
--- a/Engine/Plugins/Cameras/GameplayCameras/Source/GameplayCameras/Private/Core/CameraVariableTable.cpp
+++ b/Engine/Plugins/Cameras/GameplayCameras/Source/GameplayCameras/Private/Core/CameraVariableTable.cpp
@@ -13,11 +13,11 @@ namespace Private
 
 static const FString GUnavailableVariableDebugName(TEXT("<no debug info>"));
 
-bool IsVariableInMask(uint32 VariableId, const FCameraVariableTableFlags* InMask, bool bInvertMask)
+bool IsVariableInMask(FCameraVariableID VariableID, const FCameraVariableTableFlags* InMask, bool bInvertMask)
 {
 	if (InMask)
 	{
-		const bool bIsInMask = InMask->VariableIds.Contains(VariableId);
+		const bool bIsInMask = InMask->VariableIDs.Contains(VariableID);
 		return bInvertMask ? !bIsInMask : bIsInMask;
 	}
 	return true;
@@ -86,13 +86,13 @@ void FCameraVariableTable::Initialize(const FCameraVariableTableAllocationInfo&
 		MaxAlignOf = FMath::Max(MaxAlignOf, CurAlignOf);
 
 		FEntry NewEntry;
-		NewEntry.Id = VariableDefinition.VariableId;
+		NewEntry.ID = VariableDefinition.VariableID;
 		NewEntry.Type = VariableDefinition.VariableType;
 		NewEntry.Offset = NewEntryOffset;
 #if WITH_EDITORONLY_DATA
 		NewEntry.DebugName = VariableDefinition.VariableName;
 #endif
-		Entries.Add(NewEntry.Id, NewEntry);
+		Entries.Add(NewEntry.ID, NewEntry);
 	}
 
 	// Allocate the memory buffer.
@@ -138,7 +138,7 @@ void FCameraVariableTable::AddVariable(const FCameraVariableDefinition& Variable
 	Used = NewUsed;
 
 	FEntry NewEntry;
-	NewEntry.Id = VariableDefinition.VariableId;
+	NewEntry.ID = VariableDefinition.VariableID;
 	NewEntry.Type = VariableDefinition.VariableType;
 	NewEntry.Offset = VariablePtr - Memory;
 	NewEntry.Flags = VariableDefinition.bIsPrivate ? EEntryFlags::Private : EEntryFlags::None;
@@ -146,7 +146,7 @@ void FCameraVariableTable::AddVariable(const FCameraVariableDefinition& Variable
 	NewEntry.DebugName = VariableDefinition.VariableName;
 #endif
 
-	Entries.Add(VariableDefinition.VariableId, NewEntry);
+	Entries.Add(VariableDefinition.VariableID, NewEntry);
 }
 
 void FCameraVariableTable::ReallocateBuffer(uint32 NewCapacity)
@@ -174,23 +174,23 @@ UE_CAMERA_VARIABLE_FOR_ALL_TYPES()
 	return false;
 }
 
-bool FCameraVariableTable::ContainsValue(uint32 VariableId) const
+bool FCameraVariableTable::ContainsValue(FCameraVariableID VariableID) const
 {
-	return Entries.Contains(VariableId);
+	return Entries.Contains(VariableID);
 }
 
-bool FCameraVariableTable::IsValueWritten(uint32 VariableId) const
+bool FCameraVariableTable::IsValueWritten(FCameraVariableID VariableID) const
 {
-	if (const FEntry* Entry = Entries.Find(VariableId))
+	if (const FEntry* Entry = Entries.Find(VariableID))
 	{
 		return EnumHasAnyFlags(Entry->Flags, EEntryFlags::Written);
 	}
 	return false;
 }
 
-void FCameraVariableTable::UnsetValue(uint32 VariableId)
+void FCameraVariableTable::UnsetValue(FCameraVariableID VariableID)
 {
-	if (FEntry* Entry = Entries.Find(VariableId))
+	if (FEntry* Entry = Entries.Find(VariableID))
 	{
 		EnumRemoveFlags(Entry->Flags, EEntryFlags::Written | EEntryFlags::WrittenThisFrame);
 	}
@@ -204,9 +204,9 @@ void FCameraVariableTable::UnsetAllValues()
 	}
 }
 
-bool FCameraVariableTable::IsValueWrittenThisFrame(uint32 VariableId) const
+bool FCameraVariableTable::IsValueWrittenThisFrame(FCameraVariableID VariableID) const
 {
-	if (const FEntry* Entry = Entries.Find(VariableId))
+	if (const FEntry* Entry = Entries.Find(VariableID))
 	{
 		return EnumHasAnyFlags(Entry->Flags, EEntryFlags::WrittenThisFrame);
 	}
@@ -248,10 +248,10 @@ void FCameraVariableTable::InternalOverride(const FCameraVariableTable& OtherTab
 		if (EnumHasAnyFlags(OtherFlags, EEntryFlags::Written)
 				&& (!bChangedOnly || EnumHasAnyFlags(OtherFlags, EEntryFlags::WrittenThisFrame))
 				&& !EnumHasAnyFlags(OtherFlags, EEntryFlags::Private)
-				&& IsVariableInMask(OtherEntry.Id, InMask, bInvertMask))
+				&& IsVariableInMask(OtherEntry.ID, InMask, bInvertMask))
 		{
 			// See if we know this variable.
-			FEntry* ThisEntry = Entries.Find(OtherEntry.Id);
+			FEntry* ThisEntry = Entries.Find(OtherEntry.ID);
 			if (ThisEntry)
 			{
 				// We already have the other table's variable in our table. Let's check
@@ -259,7 +259,7 @@ void FCameraVariableTable::InternalOverride(const FCameraVariableTable& OtherTab
 #if WITH_EDITORONLY_DATA
 				checkf(ThisEntry->DebugName == OtherEntry.DebugName,
 						TEXT("Camera variable name collision! Expected variable '%d' to be named '%s', but other table has '%s'!"),
-						ThisEntry->Id, *ThisEntry->DebugName, *OtherEntry.DebugName);
+						ThisEntry->ID.GetValue(), *ThisEntry->DebugName, *OtherEntry.DebugName);
 #endif
 
 #if WITH_EDITORONLY_DATA
@@ -269,21 +269,21 @@ void FCameraVariableTable::InternalOverride(const FCameraVariableTable& OtherTab
 #endif
 				checkf(ThisEntry->Type == OtherEntry.Type, 
 						TEXT("Camera variable name collision! Expected '%d' (%s) to be of type '%s' but other table has type '%s'!"),
-						ThisEntry->Id, *DebugName,
+						ThisEntry->ID.GetValue(), *DebugName,
 						*UEnum::GetValueAsString(ThisEntry->Type), *UEnum::GetValueAsString(OtherEntry.Type));
 			}
 			else
 			{
 				// We don't have this variable in our table. Let's add it.
 				FCameraVariableDefinition NewVariableDefinition;
-				NewVariableDefinition.VariableId = OtherEntry.Id;
+				NewVariableDefinition.VariableID = OtherEntry.ID;
 				NewVariableDefinition.VariableType = OtherEntry.Type;
 #if WITH_EDITORONLY_DATA
 				NewVariableDefinition.VariableName = OtherEntry.DebugName;
 #endif
 				AddVariable(NewVariableDefinition);
 
-				ThisEntry = Entries.Find(OtherEntry.Id);
+				ThisEntry = Entries.Find(OtherEntry.ID);
 			}
 
 			if (ensure(ThisEntry))
@@ -299,7 +299,7 @@ void FCameraVariableTable::InternalOverride(const FCameraVariableTable& OtherTab
 
 				if (OutMask)
 				{
-					OutMask->VariableIds.Add(ThisEntry->Id);
+					OutMask->VariableIDs.Add(ThisEntry->ID);
 				}
 			}
 		}
@@ -333,10 +333,10 @@ void FCameraVariableTable::InternalLerp(const FCameraVariableTable& ToTable, flo
 		if (EnumHasAnyFlags(ToFlags, EEntryFlags::Written)
 				&& (!bChangedOnly || EnumHasAnyFlags(ToFlags, EEntryFlags::WrittenThisFrame))
 				&& !EnumHasAnyFlags(ToFlags, EEntryFlags::Private)
-				&& IsVariableInMask(ToEntry.Id, InMask, bInvertMask))
+				&& IsVariableInMask(ToEntry.ID, InMask, bInvertMask))
 		{
 			// See if we know this variable.
-			FEntry* FromEntry = Entries.Find(ToEntry.Id);
+			FEntry* FromEntry = Entries.Find(ToEntry.ID);
 			if (FromEntry)
 			{
 				// We already have the other table's variable in our table. Let's check
@@ -344,7 +344,7 @@ void FCameraVariableTable::InternalLerp(const FCameraVariableTable& ToTable, flo
 #if WITH_EDITORONLY_DATA
 				checkf(FromEntry->DebugName == ToEntry.DebugName,
 						TEXT("Camera variable name collision! Expected variable '%d' to be named '%s', but other table has '%s'!"),
-						FromEntry->Id, *FromEntry->DebugName, *ToEntry.DebugName);
+						FromEntry->ID.GetValue(), *FromEntry->DebugName, *ToEntry.DebugName);
 #endif
 
 #if WITH_EDITORONLY_DATA
@@ -354,7 +354,7 @@ void FCameraVariableTable::InternalLerp(const FCameraVariableTable& ToTable, flo
 #endif
 				checkf(FromEntry->Type == ToEntry.Type, 
 						TEXT("Camera variable name collision! Expected '%d' (%s) to be of type '%s' but other table has type '%s'!"),
-						FromEntry->Id, *DebugName,
+						FromEntry->ID.GetValue(), *DebugName,
 						*UEnum::GetValueAsString(FromEntry->Type), *UEnum::GetValueAsString(ToEntry.Type));
 
 				uint8* FromValuePtr = Memory + FromEntry->Offset;
@@ -379,14 +379,14 @@ UE_CAMERA_VARIABLE_FOR_ALL_TYPES()
 			{
 				// We don't have this variable in our table. Let's add it.
 				FCameraVariableDefinition NewVariableDefinition;
-				NewVariableDefinition.VariableId = ToEntry.Id;
+				NewVariableDefinition.VariableID = ToEntry.ID;
 				NewVariableDefinition.VariableType = ToEntry.Type;
 #if WITH_EDITORONLY_DATA
 				NewVariableDefinition.VariableName = ToEntry.DebugName;
 #endif
 				AddVariable(NewVariableDefinition);
 
-				FromEntry = Entries.Find(ToEntry.Id);
+				FromEntry = Entries.Find(ToEntry.ID);
 				check(FromEntry);
 
 				uint32 ValueSize = 0, ValueAlignment = 0;
@@ -401,7 +401,7 @@ UE_CAMERA_VARIABLE_FOR_ALL_TYPES()
 
 			if (OutMask)
 			{
-				OutMask->VariableIds.Add(FromEntry->Id);
+				OutMask->VariableIDs.Add(FromEntry->ID);
 			}
 		}
 	}
diff --git a/Engine/Plugins/Cameras/GameplayCameras/Source/GameplayCameras/Private/Debug/VariableTableDebugBlock.cpp b/Engine/Plugins/Cameras/GameplayCameras/Source/GameplayCameras/Private/Debug/VariableTableDebugBlock.cpp
--- a/Engine/Plugins/Cameras/GameplayCameras/Source/GameplayCameras/Private/Debug/VariableTableDebugBlock.cpp
+++ b/Engine/Plugins/Cameras/GameplayCameras/Source/GameplayCameras/Private/Debug/VariableTableDebugBlock.cpp
@@ -26,9 +26,9 @@ FVariableTableDebugBlock::FVariableTableDebugBlock(const FCameraVariableTable& I
 
 void FVariableTableDebugBlock::Initialize(const FCameraVariableTable& InVariableTable)
 {
-	for (const TPair<uint32, FCameraVariableTable::FEntry>& Pair : InVariableTable.Entries)
+	for (const TPair<FCameraVariableID, FCameraVariableTable::FEntry>& Pair : InVariableTable.Entries)
 	{
-		const int32 EntryId = Pair.Key;
+		const uint32 EntryID = Pair.Key.GetValue();
 		const FCameraVariableTable::FEntry& Entry = Pair.Value;
 
 		FString EntryName;
@@ -41,7 +41,7 @@ void FVariableTableDebugBlock::Initialize(const FCameraVariableTable& InVariable
 			case ECameraVariableType::ValueName:\
 				if (EnumHasAnyFlags(Entry.Flags, FCameraVariableTable::EEntryFlags::Written))\
 				{\
-					const ValueType EntryValue = InVariableTable.GetValue<ValueType>(EntryId);\
+					const ValueType EntryValue = InVariableTable.GetValue<ValueType>(FCameraVariableID::FromHashValue(EntryID));\
 					EntryValueStr = ToDebugString(EntryValue);\
 				}\
 				break;
@@ -51,7 +51,7 @@ void FVariableTableDebugBlock::Initialize(const FCameraVariableTable& InVariable
 		}
 #undef UE_CAMERA_VARIABLE_FOR_TYPE
 
-		FEntryDebugInfo EntryDebugInfo{ EntryId, EntryName, EntryValueStr };
+		FEntryDebugInfo EntryDebugInfo{ EntryID, EntryName, EntryValueStr };
 		EntryDebugInfo.bWritten = EnumHasAnyFlags(Entry.Flags, FCameraVariableTable::EEntryFlags::Written);
 		EntryDebugInfo.bWrittenThisFrame = EnumHasAnyFlags(Entry.Flags, FCameraVariableTable::EEntryFlags::WrittenThisFrame);
 		Entries.Add(EntryDebugInfo);
@@ -61,13 +61,13 @@ void FVariableTableDebugBlock::Initialize(const FCameraVariableTable& InVariable
 void FVariableTableDebugBlock::OnDebugDraw(const FCameraDebugBlockDrawParams& Params, FCameraDebugRenderer& Renderer)
 {
 #if WITH_EDITORONLY_DATA
-	bool bShowVariableIds = false;
-	if (!ShowVariableIdsCVarName.IsEmpty())
+	bool bShowVariableIDs = false;
+	if (!ShowVariableIDsCVarName.IsEmpty())
 	{
-		IConsoleVariable* ShowVariableIdsCVar = IConsoleManager::Get().FindConsoleVariable(*ShowVariableIdsCVarName, false);
-		if (ensureMsgf(ShowVariableIdsCVar, TEXT("No such console variable: %s"), *ShowVariableIdsCVarName))
+		IConsoleVariable* ShowVariableIDsCVar = IConsoleManager::Get().FindConsoleVariable(*ShowVariableIDsCVarName, false);
+		if (ensureMsgf(ShowVariableIDsCVar, TEXT("No such console variable: %s"), *ShowVariableIDsCVarName))
 		{
-			bShowVariableIds = ShowVariableIdsCVar->GetBool();
+			bShowVariableIDs = ShowVariableIDsCVar->GetBool();
 		}
 	}
 #endif
@@ -77,9 +77,9 @@ void FVariableTableDebugBlock::OnDebugDraw(const FCameraDebugBlockDrawParams& Pa
 	for (const FEntryDebugInfo& Entry : Entries)
 	{
 #if WITH_EDITORONLY_DATA
-		if (bShowVariableIds)
+		if (bShowVariableIDs)
 		{
-			Renderer.AddText(TEXT("{cam_passive}[%d]{cam_default} "), Entry.Id);
+			Renderer.AddText(TEXT("{cam_passive}[%d]{cam_default} "), Entry.ID);
 		}
 		if (!Entry.Name.IsEmpty())
 		{
@@ -87,10 +87,10 @@ void FVariableTableDebugBlock::OnDebugDraw(const FCameraDebugBlockDrawParams& Pa
 		}
 		else
 		{
-			Renderer.AddText(TEXT("<no name data> : "), Entry.Id);
+			Renderer.AddText(TEXT("<no name data> : "), Entry.ID);
 		}
 #else
-		Renderer.AddText(TEXT("[%d] <no name data> : "), Entry.Id);
+		Renderer.AddText(TEXT("[%d] <no name data> : "), Entry.ID);
 #endif
 
 		if (Entry.bWritten)
@@ -114,12 +114,12 @@ void FVariableTableDebugBlock::OnDebugDraw(const FCameraDebugBlockDrawParams& Pa
 void FVariableTableDebugBlock::OnSerialize(FArchive& Ar)
 {
 	Ar << Entries;
-	Ar << ShowVariableIdsCVarName;
+	Ar << ShowVariableIDsCVarName;
 }
 
 FArchive& operator<< (FArchive& Ar, FVariableTableDebugBlock::FEntryDebugInfo& EntryDebugInfo)
 {
-	Ar << EntryDebugInfo.Id;
+	Ar << EntryDebugInfo.ID;
 	Ar << EntryDebugInfo.Name;
 	Ar << EntryDebugInfo.Value;
 	Ar << EntryDebugInfo.bWritten;
diff --git a/Engine/Plugins/Cameras/GameplayCameras/Source/GameplayCameras/Private/GameFramework/GameplayCameraEvaluationResultBlueprintInterop.cpp b/Engine/Plugins/Cameras/GameplayCameras/Source/GameplayCameras/Private/GameFramework/GameplayCameraEvaluationResultBlueprintInterop.cpp
--- a/Engine/Plugins/Cameras/GameplayCameras/Source/GameplayCameras/Private/GameFramework/GameplayCameraEvaluationResultBlueprintInterop.cpp
+++ b/Engine/Plugins/Cameras/GameplayCameras/Source/GameplayCameras/Private/GameFramework/GameplayCameraEvaluationResultBlueprintInterop.cpp
@@ -25,7 +25,7 @@
 	static VariableType ErrorResult {};\
 	UE_PRIVATE_CAMERA_EVALUATION_RESULT_BLUEPRINT_INTEROP_VALIDATE_RESULT(ErrorResult)\
 	UE_PRIVATE_CAMERA_EVALUATION_RESULT_BLUEPRINT_INTEROP_VALIDATE_VARIABLE_ASSET_PARAM(ErrorResult)\
-	return Result->VariableTable.GetValue<VariableType>(InVariableAsset->GetVariableId());
+	return Result->VariableTable.GetValue<VariableType>(InVariableAsset->GetVariableID());
 
 #define UE_PRIVATE_CAMERA_EVALUATION_RESULT_BLUEPRINT_INTEROP_SET_VARIABLE(VariableType)\
 	UE_PRIVATE_CAMERA_EVALUATION_RESULT_BLUEPRINT_INTEROP_VALIDATE_RESULT()\
diff --git a/Engine/Plugins/Cameras/GameplayCameras/Source/GameplayCameras/Public/Core/CameraParameterReader.h b/Engine/Plugins/Cameras/GameplayCameras/Source/GameplayCameras/Public/Core/CameraParameterReader.h
--- a/Engine/Plugins/Cameras/GameplayCameras/Source/GameplayCameras/Public/Core/CameraParameterReader.h
+++ b/Engine/Plugins/Cameras/GameplayCameras/Source/GameplayCameras/Public/Core/CameraParameterReader.h
@@ -37,7 +37,7 @@ public:
 
 		if (Parameter.Variable)
 		{
-			VariableId = Parameter.Variable->GetVariableId();
+			VariableID = Parameter.Variable->GetVariableID();
 			DefaultValuePtr = reinterpret_cast<const ValueType*>(Parameter.Variable->GetDefaultValuePtr());
 			ensureMsgf(DefaultValuePtr, TEXT("The given parameter's driving variable doesn't have a default value!"));
 		}
@@ -49,7 +49,7 @@ public:
 	typename TCallTraits<ValueType>::ParamType Get(const FCameraVariableTable& VariableTable) const
 	{
 		checkf(DefaultValuePtr, TEXT("Parameter reader has no value pointer!"));
-		if (VariableId == 0)
+		if (!VariableID.IsValid())
 		{
 			// No variable is driving the parameter, just return the parameter value.
 			return *DefaultValuePtr;
@@ -57,7 +57,7 @@ public:
 		else
 		{
 			// The parameter is driven by a variable. Find it in the variable table.
-			if (const ValueType* ActualValue = VariableTable.FindValue<ValueType>(VariableId))
+			if (const ValueType* ActualValue = VariableTable.FindValue<ValueType>(VariableID))
 			{
 				return *ActualValue;
 			}
@@ -70,7 +70,7 @@ private:
 	/** Pointer to the value in the parameter. */
 	const ValueType* DefaultValuePtr = nullptr;
 	/** The ID of the variable driving the parameter, if any. */
-	uint32 VariableId = 0;
+	FCameraVariableID VariableID;
 };
 
 }  // namespace UE::Cameras
diff --git a/Engine/Plugins/Cameras/GameplayCameras/Source/GameplayCameras/Public/Core/CameraVariableAssets.h b/Engine/Plugins/Cameras/GameplayCameras/Source/GameplayCameras/Public/Core/CameraVariableAssets.h
--- a/Engine/Plugins/Cameras/GameplayCameras/Source/GameplayCameras/Public/Core/CameraVariableAssets.h
+++ b/Engine/Plugins/Cameras/GameplayCameras/Source/GameplayCameras/Public/Core/CameraVariableAssets.h
@@ -21,7 +21,7 @@ public:
 
 	UCameraVariableAsset(const FObjectInitializer& ObjectInit);
 
-	uint32 GetVariableId() const { return VariableId; }
+	FCameraVariableID GetVariableID() const { return VariableID; }
 
 	FCameraVariableDefinition GetVariableDefinition() const;
 
@@ -32,9 +32,20 @@ public:
 
 	// UObject interface
 	virtual void Serialize(FArchive& Ar) override;
+	virtual void PostInitProperties() override;
+	virtual void PostRename(UObject* OldOuter, const FName OldName) override;
+	virtual void PostDuplicate(EDuplicateMode::Type DuplicateMode) override;
+
+private:
+
+	void RegenerateVariableID();
 
 public:
 
+	/** Whether this variable is private and shouldn't be propagated with evaluation results. */
+	UPROPERTY()
+	bool bIsPrivate = false;
+
 	/** Whether this variable auto-resets to its default value every frame. */
 	UPROPERTY(EditAnywhere, Category=Camera)
 	bool bAutoReset = false;
@@ -42,10 +53,7 @@ public:
 private:
 
 	UPROPERTY()
-	uint32 VariableId = 0;
-
-	UPROPERTY()
-	bool bIsLocal = false;
+	FCameraVariableID VariableID;
 };
 
 /** Boolean camera variable. */
diff --git a/Engine/Plugins/Cameras/GameplayCameras/Source/GameplayCameras/Public/Core/CameraVariableTable.h b/Engine/Plugins/Cameras/GameplayCameras/Source/GameplayCameras/Public/Core/CameraVariableTable.h
--- a/Engine/Plugins/Cameras/GameplayCameras/Source/GameplayCameras/Public/Core/CameraVariableTable.h
+++ b/Engine/Plugins/Cameras/GameplayCameras/Source/GameplayCameras/Public/Core/CameraVariableTable.h
@@ -27,7 +27,7 @@ struct TCameraVariableInterpolation;
 struct FCameraVariableTableFlags
 {
 	/** The list of processed variable IDs. */
-	TSet<uint32> VariableIds;
+	TSet<FCameraVariableID> VariableIDs;
 };
 
 /**
@@ -73,28 +73,28 @@ public:
 	// Getter methods.
 
 	template<typename ValueType>
-	const ValueType* FindValue(uint32 VariableId) const;
+	const ValueType* FindValue(FCameraVariableID VariableID) const;
 
 	template<typename ValueType>
-	const ValueType& GetValue(uint32 VariableId) const;
+	const ValueType& GetValue(FCameraVariableID VariableID) const;
 
 	template<typename ValueType>
-	ValueType GetValue(uint32 VariableId, typename TCallTraits<ValueType>::ParamType DefaultValue) const;
+	ValueType GetValue(FCameraVariableID VariableID, typename TCallTraits<ValueType>::ParamType DefaultValue) const;
 
 	template<typename ValueType>
-	bool TryGetValue(uint32 VariableId, ValueType& OutValue) const;
+	bool TryGetValue(FCameraVariableID VariableID, ValueType& OutValue) const;
 
-	bool ContainsValue(uint32 VariableId) const;
+	bool ContainsValue(FCameraVariableID VariableID) const;
 
 public:
 
 	// Setter methods.
 	
 	template<typename ValueType>
-	void SetValue(uint32 VariableId, typename TCallTraits<ValueType>::ParamType Value);
+	void SetValue(FCameraVariableID VariableID, typename TCallTraits<ValueType>::ParamType Value);
 
 	template<typename ValueType>
-	bool TrySetValue(uint32 VariableId, typename TCallTraits<ValueType>::ParamType Value);
+	bool TrySetValue(FCameraVariableID VariableID, typename TCallTraits<ValueType>::ParamType Value);
 
 	template<typename VariableAssetType>
 	void SetValue(
@@ -118,11 +118,11 @@ public:
 
 	// Lower level API.
 
-	bool IsValueWritten(uint32 VariableId) const;
-	void UnsetValue(uint32 VariableId);
+	bool IsValueWritten(FCameraVariableID VariableID) const;
+	void UnsetValue(FCameraVariableID VariableID);
 	void UnsetAllValues();
 
-	bool IsValueWrittenThisFrame(uint32 VariableId) const;
+	bool IsValueWrittenThisFrame(FCameraVariableID VariableID) const;
 	void ClearAllWrittenThisFrameFlags();
 
 private:
@@ -153,7 +153,7 @@ private:
 
 	struct FEntry
 	{
-		uint32 Id;
+		FCameraVariableID ID;
 		ECameraVariableType Type;
 		uint32 Offset;
 		mutable EEntryFlags Flags;
@@ -162,7 +162,7 @@ private:
 #endif
 	};
 
-	TMap<uint32, FEntry> Entries;
+	TMap<FCameraVariableID, FEntry> Entries;
 
 	uint8* Memory = nullptr;
 	uint32 Capacity = 0;
@@ -179,9 +179,9 @@ private:
 ENUM_CLASS_FLAGS(FCameraVariableTable::EEntryFlags)
 
 template<typename ValueType>
-const ValueType* FCameraVariableTable::FindValue(uint32 VariableId) const
+const ValueType* FCameraVariableTable::FindValue(FCameraVariableID VariableID) const
 {
-	if (const FEntry* Entry = Entries.Find(VariableId))
+	if (const FEntry* Entry = Entries.Find(VariableID))
 	{
 		CheckVariableType<ValueType>(Entry->Type);
 		if (EnumHasAnyFlags(Entry->Flags, EEntryFlags::Written))
@@ -193,9 +193,9 @@ const ValueType* FCameraVariableTable::FindValue(uint32 VariableId) const
 }
 
 template<typename ValueType>
-const ValueType& FCameraVariableTable::GetValue(uint32 VariableId) const
+const ValueType& FCameraVariableTable::GetValue(FCameraVariableID VariableID) const
 {
-	const FEntry& Entry = Entries.FindChecked(VariableId);
+	const FEntry& Entry = Entries.FindChecked(VariableID);
 	CheckVariableType<ValueType>(Entry.Type);
 #if WITH_EDITORONLY_DATA
 	checkf(
@@ -206,15 +206,15 @@ const ValueType& FCameraVariableTable::GetValue(uint32 VariableId) const
 	checkf(
 			EnumHasAnyFlags(Entry.Flags, EEntryFlags::Written),
 			TEXT("Variable '%s' has never been written to. GetValue() will return uninitialized memory!"),
-			*LexToString(VariableId));
+			*LexToString(VariableID.GetValue()));
 #endif
 	return *reinterpret_cast<ValueType*>(Memory + Entry.Offset);
 }
 
 template<typename ValueType>
-ValueType FCameraVariableTable::GetValue(uint32 VariableId, typename TCallTraits<ValueType>::ParamType DefaultValue) const
+ValueType FCameraVariableTable::GetValue(FCameraVariableID VariableID, typename TCallTraits<ValueType>::ParamType DefaultValue) const
 {
-	if (const ValueType* Value = FindValue<ValueType>(VariableId))
+	if (const ValueType* Value = FindValue<ValueType>(VariableID))
 	{
 		return *Value;
 	}
@@ -222,9 +222,9 @@ ValueType FCameraVariableTable::GetValue(uint32 VariableId, typename TCallTraits
 }
 
 template<typename ValueType>
-bool FCameraVariableTable::TryGetValue(uint32 VariableId, ValueType& OutValue) const
+bool FCameraVariableTable::TryGetValue(FCameraVariableID VariableID, ValueType& OutValue) const
 {
-	if (const ValueType* Value = FindValue<ValueType>(VariableId))
+	if (const ValueType* Value = FindValue<ValueType>(VariableID))
 	{
 		OutValue = *Value;
 		return true;
@@ -233,9 +233,9 @@ bool FCameraVariableTable::TryGetValue(uint32 VariableId, ValueType& OutValue) c
 }
 
 template<typename ValueType>
-void FCameraVariableTable::SetValue(uint32 VariableId, typename TCallTraits<ValueType>::ParamType Value)
+void FCameraVariableTable::SetValue(FCameraVariableID VariableID, typename TCallTraits<ValueType>::ParamType Value)
 {
-	FEntry& Entry = Entries.FindChecked(VariableId);
+	FEntry& Entry = Entries.FindChecked(VariableID);
 	CheckVariableType<ValueType>(Entry.Type);
 	ValueType* ValuePtr = reinterpret_cast<ValueType*>(Memory + Entry.Offset);
 	*ValuePtr = Value;
@@ -243,9 +243,9 @@ void FCameraVariableTable::SetValue(uint32 VariableId, typename TCallTraits<Valu
 }
 
 template<typename ValueType>
-bool FCameraVariableTable::TrySetValue(uint32 VariableId, typename TCallTraits<ValueType>::ParamType Value)
+bool FCameraVariableTable::TrySetValue(FCameraVariableID VariableID, typename TCallTraits<ValueType>::ParamType Value)
 {
-	if (FEntry* Entry = Entries.Find(VariableId))
+	if (FEntry* Entry = Entries.Find(VariableID))
 	{
 		CheckVariableType<ValueType>(Entry->Type);
 		ValueType* ValuePtr = reinterpret_cast<ValueType*>(Memory + Entry->Offset);
@@ -264,7 +264,7 @@ void FCameraVariableTable::SetValue(
 {
 	if (ensure(VariableAsset))
 	{
-		if (TrySetValue<typename VariableAssetType::ValueType>(VariableAsset->GetVariableId(), Value))
+		if (TrySetValue<typename VariableAssetType::ValueType>(VariableAsset->GetVariableID(), Value))
 		{
 			return;
 		}
@@ -273,7 +273,7 @@ void FCameraVariableTable::SetValue(
 		{
 			FCameraVariableDefinition VariableDefinition = VariableAsset->GetVariableDefinition();
 			AddVariable(VariableDefinition);
-			SetValue<typename VariableAssetType::ValueType>(VariableDefinition.VariableId, Value);
+			SetValue<typename VariableAssetType::ValueType>(VariableDefinition.VariableID, Value);
 		}
 	}
 }
diff --git a/Engine/Plugins/Cameras/GameplayCameras/Source/GameplayCameras/Public/Core/CameraVariableTableFwd.h b/Engine/Plugins/Cameras/GameplayCameras/Source/GameplayCameras/Public/Core/CameraVariableTableFwd.h
--- a/Engine/Plugins/Cameras/GameplayCameras/Source/GameplayCameras/Public/Core/CameraVariableTableFwd.h
+++ b/Engine/Plugins/Cameras/GameplayCameras/Source/GameplayCameras/Public/Core/CameraVariableTableFwd.h
@@ -42,13 +42,65 @@ enum class ECameraVariableType
 	Transform3d
 };
 
+USTRUCT()
+struct FCameraVariableID
+{
+	GENERATED_BODY()
+
+public:
+
+	FCameraVariableID() : Value(INVALID) {}
+
+	uint32 GetValue() const { return Value; }
+
+	bool IsValid() const { return Value != INVALID; }
+
+	explicit operator bool() const { return IsValid(); }
+
+	static FCameraVariableID FromHashValue(uint32 InValue)
+	{
+		return FCameraVariableID(InValue);
+	}
+
+public:
+
+	friend bool operator<(FCameraVariableID A, FCameraVariableID B)
+	{
+		return A.Value < B.Value;
+	}
+
+	friend bool operator==(FCameraVariableID A, FCameraVariableID B)
+	{
+		return A.Value == B.Value;
+	}
+
+	friend bool operator!=(FCameraVariableID A, FCameraVariableID B)
+	{
+		return A.Value != B.Value;
+	}
+
+	friend uint32 GetTypeHash(FCameraVariableID In)
+	{
+		return In.Value;
+	}
+
+private:
+
+	FCameraVariableID(uint32 InValue) : Value(InValue) {}
+
+	static const uint32 INVALID = uint32(-1);
+
+	UPROPERTY()
+	uint32 Value;
+};
+
 USTRUCT()
 struct FCameraVariableDefinition
 {
 	GENERATED_BODY()
 
 	UPROPERTY()
-	uint32 VariableId = 0;
+	FCameraVariableID VariableID;
 
 	UPROPERTY()
 	ECameraVariableType VariableType = ECameraVariableType::Boolean;
diff --git a/Engine/Plugins/Cameras/GameplayCameras/Source/GameplayCameras/Public/Debug/VariableTableDebugBlock.h b/Engine/Plugins/Cameras/GameplayCameras/Source/GameplayCameras/Public/Debug/VariableTableDebugBlock.h
--- a/Engine/Plugins/Cameras/GameplayCameras/Source/GameplayCameras/Public/Debug/VariableTableDebugBlock.h
+++ b/Engine/Plugins/Cameras/GameplayCameras/Source/GameplayCameras/Public/Debug/VariableTableDebugBlock.h
@@ -26,9 +26,9 @@ public:
 	GAMEPLAYCAMERAS_API FVariableTableDebugBlock(const FCameraVariableTable& InVariableTable);
 
 	/** Specifies the console variable to use to toggle the printing of variable IDs. */
-	FVariableTableDebugBlock& WithShowVariableIdsCVar(const TCHAR* InShowVariableIdsCVarName)
+	FVariableTableDebugBlock& WithShowVariableIDsCVar(const TCHAR* InShowVariableIDsCVarName)
 	{
-		ShowVariableIdsCVarName = InShowVariableIdsCVarName;
+		ShowVariableIDsCVarName = InShowVariableIDsCVarName;
 		return *this;
 	}
 
@@ -43,7 +43,7 @@ private:
 
 	struct FEntryDebugInfo
 	{
-		uint32 Id;
+		uint32 ID;
 		FString Name;
 		FString Value;
 		bool bWritten;
@@ -51,7 +51,7 @@ private:
 	};
 	TArray<FEntryDebugInfo> Entries;
 
-	FString ShowVariableIdsCVarName;
+	FString ShowVariableIDsCVarName;
 
 	friend FArchive& operator<< (FArchive&, FEntryDebugInfo&);
 };
diff --git a/Engine/Plugins/Cameras/GameplayCameras/Source/GameplayCamerasEditor/Private/Debugger/SBlendStacksDebugPanel.cpp b/Engine/Plugins/Cameras/GameplayCameras/Source/GameplayCamerasEditor/Private/Debugger/SBlendStacksDebugPanel.cpp
--- a/Engine/Plugins/Cameras/GameplayCameras/Source/GameplayCamerasEditor/Private/Debugger/SBlendStacksDebugPanel.cpp
+++ b/Engine/Plugins/Cameras/GameplayCameras/Source/GameplayCamerasEditor/Private/Debugger/SBlendStacksDebugPanel.cpp
@@ -28,8 +28,8 @@ void SBlendStacksDebugPanel::Construct(const FArguments& InArgs)
 			.AutoHeight()
 			[
 				SDebugWidgetUtils::CreateConsoleVariableCheckBox(
-						LOCTEXT("ShowVariableIds", "Show variable IDs"),
-						TEXT("GameplayCameras.Debug.BlendStack.ShowVariableIds"))
+						LOCTEXT("ShowVariableIDs", "Show variable IDs"),
+						TEXT("GameplayCameras.Debug.BlendStack.ShowVariableIDs"))
 			]
 	];
 }
diff --git a/Engine/Plugins/Cameras/GameplayCameras/Source/GameplayCamerasEditor/Private/Debugger/SCameraPoseStatsDebugPanel.cpp b/Engine/Plugins/Cameras/GameplayCameras/Source/GameplayCamerasEditor/Private/Debugger/SCameraPoseStatsDebugPanel.cpp
--- a/Engine/Plugins/Cameras/GameplayCameras/Source/GameplayCamerasEditor/Private/Debugger/SCameraPoseStatsDebugPanel.cpp
+++ b/Engine/Plugins/Cameras/GameplayCameras/Source/GameplayCamerasEditor/Private/Debugger/SCameraPoseStatsDebugPanel.cpp
@@ -27,8 +27,8 @@ void SCameraPoseStatsDebugPanel::Construct(const FArguments& InArgs)
 			.AutoHeight()
 			[
 				SDebugWidgetUtils::CreateConsoleVariableCheckBox(
-						LOCTEXT("ShowVariableIds", "Show variable IDs"),
-						TEXT("GameplayCameras.Debug.PoseStats.ShowVariableIds"))
+						LOCTEXT("ShowVariableIDs", "Show variable IDs"),
+						TEXT("GameplayCameras.Debug.PoseStats.ShowVariableIDs"))
 			]
 	];
 }
