Cameras: add declspec to macros to make it possible to declare nodes/contexts/etc. in other modules

diff --git a/Engine/Plugins/Cameras/GameplayCameras/Source/GameplayCameras/Private/Core/BlendStackRootCameraNode.cpp b/Engine/Plugins/Cameras/GameplayCameras/Source/GameplayCameras/Private/Core/BlendStackRootCameraNode.cpp
--- a/Engine/Plugins/Cameras/GameplayCameras/Source/GameplayCameras/Private/Core/BlendStackRootCameraNode.cpp
+++ b/Engine/Plugins/Cameras/GameplayCameras/Source/GameplayCameras/Private/Core/BlendStackRootCameraNode.cpp
@@ -35,7 +35,7 @@ namespace UE::Cameras
 
 UE_DEFINE_CAMERA_NODE_EVALUATOR(FBlendStackRootCameraNodeEvaluator)
 
-UE_DECLARE_CAMERA_DEBUG_BLOCK_START(FBlendStackRootCameraDebugBlock)
+UE_DECLARE_CAMERA_DEBUG_BLOCK_START(GAMEPLAYCAMERAS_API, FBlendStackRootCameraDebugBlock)
 	UE_DECLARE_CAMERA_DEBUG_BLOCK_FIELD(FString, CameraRigAssetName)
 UE_DECLARE_CAMERA_DEBUG_BLOCK_END()
 
diff --git a/Engine/Plugins/Cameras/GameplayCameras/Source/GameplayCameras/Private/Core/BlendStackRootCameraNode.h b/Engine/Plugins/Cameras/GameplayCameras/Source/GameplayCameras/Private/Core/BlendStackRootCameraNode.h
--- a/Engine/Plugins/Cameras/GameplayCameras/Source/GameplayCameras/Private/Core/BlendStackRootCameraNode.h
+++ b/Engine/Plugins/Cameras/GameplayCameras/Source/GameplayCameras/Private/Core/BlendStackRootCameraNode.h
@@ -47,7 +47,7 @@ class FBlendCameraNodeEvaluator;
  */
 class FBlendStackRootCameraNodeEvaluator : public FCameraNodeEvaluator
 {
-	UE_DECLARE_CAMERA_NODE_EVALUATOR(FBlendStackRootCameraNodeEvaluator)
+	UE_DECLARE_CAMERA_NODE_EVALUATOR(GAMEPLAYCAMERAS_API, FBlendStackRootCameraNodeEvaluator)
 
 public:
 
diff --git a/Engine/Plugins/Cameras/GameplayCameras/Source/GameplayCameras/Private/Core/DefaultRootCameraNode.cpp b/Engine/Plugins/Cameras/GameplayCameras/Source/GameplayCameras/Private/Core/DefaultRootCameraNode.cpp
--- a/Engine/Plugins/Cameras/GameplayCameras/Source/GameplayCameras/Private/Core/DefaultRootCameraNode.cpp
+++ b/Engine/Plugins/Cameras/GameplayCameras/Source/GameplayCameras/Private/Core/DefaultRootCameraNode.cpp
@@ -101,7 +101,7 @@ void FDefaultRootCameraNodeEvaluator::OnActivateCameraRig(const FActivateCameraR
 
 #if UE_GAMEPLAY_CAMERAS_DEBUG
 
-UE_DECLARE_CAMERA_DEBUG_BLOCK_START(FDefaultRootCameraNodeEvaluatorDebugBlock)
+UE_DECLARE_CAMERA_DEBUG_BLOCK_START(GAMEPLAYCAMERAS_API, FDefaultRootCameraNodeEvaluatorDebugBlock)
 UE_DECLARE_CAMERA_DEBUG_BLOCK_END()
 
 UE_DEFINE_CAMERA_DEBUG_BLOCK_WITH_FIELDS(FDefaultRootCameraNodeEvaluatorDebugBlock)
diff --git a/Engine/Plugins/Cameras/GameplayCameras/Source/GameplayCameras/Private/Directors/SingleCameraDirector.cpp b/Engine/Plugins/Cameras/GameplayCameras/Source/GameplayCameras/Private/Directors/SingleCameraDirector.cpp
--- a/Engine/Plugins/Cameras/GameplayCameras/Source/GameplayCameras/Private/Directors/SingleCameraDirector.cpp
+++ b/Engine/Plugins/Cameras/GameplayCameras/Source/GameplayCameras/Private/Directors/SingleCameraDirector.cpp
@@ -9,7 +9,7 @@ namespace UE::Cameras
 
 class FSingleCameraDirectorEvaluator : public FCameraDirectorEvaluator
 {
-	UE_DECLARE_CAMERA_DIRECTOR_EVALUATOR(FSingleCameraDirectorEvaluator)
+	UE_DECLARE_CAMERA_DIRECTOR_EVALUATOR(GAMEPLAYCAMERAS_API, FSingleCameraDirectorEvaluator)
 protected:
 	virtual void OnRun(const FCameraDirectorEvaluationParams& Params, FCameraDirectorEvaluationResult& OutResult) override
 	{
diff --git a/Engine/Plugins/Cameras/GameplayCameras/Source/GameplayCameras/Private/Nodes/Blends/LinearBlendCameraNode.cpp b/Engine/Plugins/Cameras/GameplayCameras/Source/GameplayCameras/Private/Nodes/Blends/LinearBlendCameraNode.cpp
--- a/Engine/Plugins/Cameras/GameplayCameras/Source/GameplayCameras/Private/Nodes/Blends/LinearBlendCameraNode.cpp
+++ b/Engine/Plugins/Cameras/GameplayCameras/Source/GameplayCameras/Private/Nodes/Blends/LinearBlendCameraNode.cpp
@@ -9,7 +9,8 @@ namespace UE::Cameras
 
 class FLinearBlendCameraNodeEvaluator : public FSimpleFixedTimeBlendCameraNodeEvaluator
 {
-	UE_DECLARE_BLEND_CAMERA_NODE_EVALUATOR_EX(FLinearBlendCameraNodeEvaluator, FSimpleFixedTimeBlendCameraNodeEvaluator)
+	UE_DECLARE_BLEND_CAMERA_NODE_EVALUATOR_EX(GAMEPLAYCAMERAS_API, FLinearBlendCameraNodeEvaluator, FSimpleFixedTimeBlendCameraNodeEvaluator)
+
 protected:
 	virtual void OnComputeBlendFactor(const FCameraNodeEvaluationParams& Params, FSimpleBlendCameraNodeEvaluationResult& OutResult) override;
 };
diff --git a/Engine/Plugins/Cameras/GameplayCameras/Source/GameplayCameras/Private/Nodes/Blends/PopBlendCameraNode.cpp b/Engine/Plugins/Cameras/GameplayCameras/Source/GameplayCameras/Private/Nodes/Blends/PopBlendCameraNode.cpp
--- a/Engine/Plugins/Cameras/GameplayCameras/Source/GameplayCameras/Private/Nodes/Blends/PopBlendCameraNode.cpp
+++ b/Engine/Plugins/Cameras/GameplayCameras/Source/GameplayCameras/Private/Nodes/Blends/PopBlendCameraNode.cpp
@@ -9,7 +9,8 @@ namespace UE::Cameras
 
 class FPopBlendCameraNodeEvaluator : public FBlendCameraNodeEvaluator
 {
-	UE_DECLARE_BLEND_CAMERA_NODE_EVALUATOR(FPopBlendCameraNodeEvaluator)
+	UE_DECLARE_BLEND_CAMERA_NODE_EVALUATOR(GAMEPLAYCAMERAS_API, FPopBlendCameraNodeEvaluator)
+
 protected:
 	virtual void OnRun(const FCameraNodeEvaluationParams& Params, FCameraNodeEvaluationResult& OutResult) override;
 	virtual void OnBlendResults(const FCameraNodeBlendParams& Params, FCameraNodeBlendResult& OutResult) override;
diff --git a/Engine/Plugins/Cameras/GameplayCameras/Source/GameplayCameras/Private/Nodes/Blends/SimpleBlendCameraNode.cpp b/Engine/Plugins/Cameras/GameplayCameras/Source/GameplayCameras/Private/Nodes/Blends/SimpleBlendCameraNode.cpp
--- a/Engine/Plugins/Cameras/GameplayCameras/Source/GameplayCameras/Private/Nodes/Blends/SimpleBlendCameraNode.cpp
+++ b/Engine/Plugins/Cameras/GameplayCameras/Source/GameplayCameras/Private/Nodes/Blends/SimpleBlendCameraNode.cpp
@@ -13,7 +13,7 @@ namespace UE::Cameras
 
 UE_DEFINE_BLEND_CAMERA_NODE_EVALUATOR(FSimpleBlendCameraNodeEvaluator)
 
-UE_DECLARE_CAMERA_DEBUG_BLOCK_START(FSimpleBlendCameraDebugBlock)
+UE_DECLARE_CAMERA_DEBUG_BLOCK_START(GAMEPLAYCAMERAS_API, FSimpleBlendCameraDebugBlock)
 	UE_DECLARE_CAMERA_DEBUG_BLOCK_FIELD(float, BlendFactor)
 UE_DECLARE_CAMERA_DEBUG_BLOCK_END()
 
diff --git a/Engine/Plugins/Cameras/GameplayCameras/Source/GameplayCameras/Private/Nodes/Blends/SmoothBlendCameraNode.cpp b/Engine/Plugins/Cameras/GameplayCameras/Source/GameplayCameras/Private/Nodes/Blends/SmoothBlendCameraNode.cpp
--- a/Engine/Plugins/Cameras/GameplayCameras/Source/GameplayCameras/Private/Nodes/Blends/SmoothBlendCameraNode.cpp
+++ b/Engine/Plugins/Cameras/GameplayCameras/Source/GameplayCameras/Private/Nodes/Blends/SmoothBlendCameraNode.cpp
@@ -11,7 +11,8 @@ namespace UE::Cameras
 
 class FSmoothBlendCameraNodeEvaluator : public FSimpleFixedTimeBlendCameraNodeEvaluator
 {
-	UE_DECLARE_BLEND_CAMERA_NODE_EVALUATOR_EX(FSmoothBlendCameraNodeEvaluator, FSimpleFixedTimeBlendCameraNodeEvaluator)
+	UE_DECLARE_BLEND_CAMERA_NODE_EVALUATOR_EX(GAMEPLAYCAMERAS_API, FSmoothBlendCameraNodeEvaluator, FSimpleFixedTimeBlendCameraNodeEvaluator)
+
 protected:
 	virtual void OnComputeBlendFactor(const FCameraNodeEvaluationParams& Params, FSimpleBlendCameraNodeEvaluationResult& OutResult) override;
 };
diff --git a/Engine/Plugins/Cameras/GameplayCameras/Source/GameplayCameras/Private/Nodes/Common/ArrayCameraNode.cpp b/Engine/Plugins/Cameras/GameplayCameras/Source/GameplayCameras/Private/Nodes/Common/ArrayCameraNode.cpp
--- a/Engine/Plugins/Cameras/GameplayCameras/Source/GameplayCameras/Private/Nodes/Common/ArrayCameraNode.cpp
+++ b/Engine/Plugins/Cameras/GameplayCameras/Source/GameplayCameras/Private/Nodes/Common/ArrayCameraNode.cpp
@@ -11,7 +11,7 @@ namespace UE::Cameras
 
 class FArrayCameraNodeEvaluator : public FCameraNodeEvaluator
 {
-	UE_DECLARE_CAMERA_NODE_EVALUATOR(FArrayCameraNodeEvaluator)
+	UE_DECLARE_CAMERA_NODE_EVALUATOR(GAMEPLAYCAMERAS_API, FArrayCameraNodeEvaluator)
 
 protected:
 
diff --git a/Engine/Plugins/Cameras/GameplayCameras/Source/GameplayCameras/Private/Nodes/Common/BoomArmCameraNode.cpp b/Engine/Plugins/Cameras/GameplayCameras/Source/GameplayCameras/Private/Nodes/Common/BoomArmCameraNode.cpp
--- a/Engine/Plugins/Cameras/GameplayCameras/Source/GameplayCameras/Private/Nodes/Common/BoomArmCameraNode.cpp
+++ b/Engine/Plugins/Cameras/GameplayCameras/Source/GameplayCameras/Private/Nodes/Common/BoomArmCameraNode.cpp
@@ -14,7 +14,7 @@ namespace UE::Cameras
 
 class FBoomArmCameraNodeEvaluator : public FCameraNodeEvaluator
 {
-	UE_DECLARE_CAMERA_NODE_EVALUATOR(FBoomArmCameraNodeEvaluator)
+	UE_DECLARE_CAMERA_NODE_EVALUATOR(GAMEPLAYCAMERAS_API, FBoomArmCameraNodeEvaluator)
 
 protected:
 
diff --git a/Engine/Plugins/Cameras/GameplayCameras/Source/GameplayCameras/Private/Nodes/Common/DampenPositionCameraNode.cpp b/Engine/Plugins/Cameras/GameplayCameras/Source/GameplayCameras/Private/Nodes/Common/DampenPositionCameraNode.cpp
--- a/Engine/Plugins/Cameras/GameplayCameras/Source/GameplayCameras/Private/Nodes/Common/DampenPositionCameraNode.cpp
+++ b/Engine/Plugins/Cameras/GameplayCameras/Source/GameplayCameras/Private/Nodes/Common/DampenPositionCameraNode.cpp
@@ -18,7 +18,7 @@ namespace UE::Cameras
 
 class FDampenPositionCameraNodeEvaluator : public FCameraNodeEvaluator
 {
-	UE_DECLARE_CAMERA_NODE_EVALUATOR(FDampenPositionCameraNodeEvaluator)
+	UE_DECLARE_CAMERA_NODE_EVALUATOR(GAMEPLAYCAMERAS_API, FDampenPositionCameraNodeEvaluator)
 
 protected:
 
@@ -45,7 +45,7 @@ private:
 
 UE_DEFINE_CAMERA_NODE_EVALUATOR(FDampenPositionCameraNodeEvaluator)
 
-UE_DECLARE_CAMERA_DEBUG_BLOCK_START(FDampenPositionCameraDebugBlock)
+UE_DECLARE_CAMERA_DEBUG_BLOCK_START(GAMEPLAYCAMERAS_API, FDampenPositionCameraDebugBlock)
 	UE_DECLARE_CAMERA_DEBUG_BLOCK_FIELD(float, ForwardX0);
 	UE_DECLARE_CAMERA_DEBUG_BLOCK_FIELD(float, LateralX0);
 	UE_DECLARE_CAMERA_DEBUG_BLOCK_FIELD(float, VerticalX0);
diff --git a/Engine/Plugins/Cameras/GameplayCameras/Source/GameplayCameras/Private/Nodes/Common/LensParametersCameraNode.cpp b/Engine/Plugins/Cameras/GameplayCameras/Source/GameplayCameras/Private/Nodes/Common/LensParametersCameraNode.cpp
--- a/Engine/Plugins/Cameras/GameplayCameras/Source/GameplayCameras/Private/Nodes/Common/LensParametersCameraNode.cpp
+++ b/Engine/Plugins/Cameras/GameplayCameras/Source/GameplayCameras/Private/Nodes/Common/LensParametersCameraNode.cpp
@@ -13,7 +13,7 @@ namespace UE::Cameras
 
 class FLensParametersCameraNodeEvaluator : public FCameraNodeEvaluator
 {
-	UE_DECLARE_CAMERA_NODE_EVALUATOR(FLensParametersCameraNodeEvaluator)
+	UE_DECLARE_CAMERA_NODE_EVALUATOR(GAMEPLAYCAMERAS_API, FLensParametersCameraNodeEvaluator)
 
 protected:
 
diff --git a/Engine/Plugins/Cameras/GameplayCameras/Source/GameplayCameras/Private/Nodes/Common/OffsetCameraNode.cpp b/Engine/Plugins/Cameras/GameplayCameras/Source/GameplayCameras/Private/Nodes/Common/OffsetCameraNode.cpp
--- a/Engine/Plugins/Cameras/GameplayCameras/Source/GameplayCameras/Private/Nodes/Common/OffsetCameraNode.cpp
+++ b/Engine/Plugins/Cameras/GameplayCameras/Source/GameplayCameras/Private/Nodes/Common/OffsetCameraNode.cpp
@@ -13,7 +13,7 @@ namespace UE::Cameras
 
 class FOffsetCameraNodeEvaluator : public FCameraNodeEvaluator
 {
-	UE_DECLARE_CAMERA_NODE_EVALUATOR(FOffsetCameraNodeEvaluator)
+	UE_DECLARE_CAMERA_NODE_EVALUATOR(GAMEPLAYCAMERAS_API, FOffsetCameraNodeEvaluator)
 
 protected:
 
diff --git a/Engine/Plugins/Cameras/GameplayCameras/Source/GameplayCameras/Public/Core/BlendCameraNode.h b/Engine/Plugins/Cameras/GameplayCameras/Source/GameplayCameras/Public/Core/BlendCameraNode.h
--- a/Engine/Plugins/Cameras/GameplayCameras/Source/GameplayCameras/Public/Core/BlendCameraNode.h
+++ b/Engine/Plugins/Cameras/GameplayCameras/Source/GameplayCameras/Public/Core/BlendCameraNode.h
@@ -61,7 +61,7 @@ struct FCameraNodeBlendResult
  */
 class FBlendCameraNodeEvaluator : public FCameraNodeEvaluator
 {
-	UE_DECLARE_CAMERA_NODE_EVALUATOR(FBlendCameraNodeEvaluator)
+	UE_DECLARE_CAMERA_NODE_EVALUATOR(GAMEPLAYCAMERAS_API, FBlendCameraNodeEvaluator)
 
 public:
 
@@ -81,11 +81,11 @@ protected:
 // from having to specify FBlendCameraNodeEvaluator as the base class, which saves
 // a little bit of typing.
 //
-#define UE_DECLARE_BLEND_CAMERA_NODE_EVALUATOR(ClassName)\
-	UE_DECLARE_CAMERA_NODE_EVALUATOR_EX(ClassName, ::UE::Cameras::FBlendCameraNodeEvaluator)
+#define UE_DECLARE_BLEND_CAMERA_NODE_EVALUATOR(ApiDeclSpec, ClassName)\
+	UE_DECLARE_CAMERA_NODE_EVALUATOR_EX(ApiDeclSpec, ClassName, ::UE::Cameras::FBlendCameraNodeEvaluator)
 
-#define UE_DECLARE_BLEND_CAMERA_NODE_EVALUATOR_EX(ClassName, BaseClassName)\
-	UE_DECLARE_CAMERA_NODE_EVALUATOR_EX(ClassName, BaseClassName)
+#define UE_DECLARE_BLEND_CAMERA_NODE_EVALUATOR_EX(ApiDeclSpec, ClassName, BaseClassName)\
+	UE_DECLARE_CAMERA_NODE_EVALUATOR_EX(ApiDeclSpec, ClassName, BaseClassName)
 
 #define UE_DEFINE_BLEND_CAMERA_NODE_EVALUATOR(ClassName)\
 	UE_DEFINE_CAMERA_NODE_EVALUATOR(ClassName)
diff --git a/Engine/Plugins/Cameras/GameplayCameras/Source/GameplayCameras/Public/Core/BlendStackCameraNode.h b/Engine/Plugins/Cameras/GameplayCameras/Source/GameplayCameras/Public/Core/BlendStackCameraNode.h
--- a/Engine/Plugins/Cameras/GameplayCameras/Source/GameplayCameras/Public/Core/BlendStackCameraNode.h
+++ b/Engine/Plugins/Cameras/GameplayCameras/Source/GameplayCameras/Public/Core/BlendStackCameraNode.h
@@ -85,7 +85,7 @@ class FBlendStackCameraNodeEvaluator
 	, public IGameplayCamerasLiveEditListener
 #endif
 {
-	UE_DECLARE_CAMERA_NODE_EVALUATOR(FBlendStackCameraNodeEvaluator)
+	UE_DECLARE_CAMERA_NODE_EVALUATOR(GAMEPLAYCAMERAS_API, FBlendStackCameraNodeEvaluator)
 
 public:
 
@@ -171,7 +171,7 @@ protected:
 
 class FBlendStackSummaryCameraDebugBlock : public FCameraDebugBlock
 {
-	UE_DECLARE_CAMERA_DEBUG_BLOCK(FBlendStackSummaryCameraDebugBlock)
+	UE_DECLARE_CAMERA_DEBUG_BLOCK(GAMEPLAYCAMERAS_API, FBlendStackSummaryCameraDebugBlock)
 
 public:
 
@@ -190,7 +190,7 @@ private:
 
 class FBlendStackCameraDebugBlock : public FCameraDebugBlock
 {
-	UE_DECLARE_CAMERA_DEBUG_BLOCK(FBlendStackCameraDebugBlock)
+	UE_DECLARE_CAMERA_DEBUG_BLOCK(GAMEPLAYCAMERAS_API, FBlendStackCameraDebugBlock)
 
 public:
 
diff --git a/Engine/Plugins/Cameras/GameplayCameras/Source/GameplayCameras/Public/Core/CameraDirectorEvaluator.h b/Engine/Plugins/Cameras/GameplayCameras/Source/GameplayCameras/Public/Core/CameraDirectorEvaluator.h
--- a/Engine/Plugins/Cameras/GameplayCameras/Source/GameplayCameras/Public/Core/CameraDirectorEvaluator.h
+++ b/Engine/Plugins/Cameras/GameplayCameras/Source/GameplayCameras/Public/Core/CameraDirectorEvaluator.h
@@ -81,7 +81,7 @@ private:
  */
 class FCameraDirectorEvaluator
 {
-	UE_GAMEPLAY_CAMERAS_DECLARE_RTTI_BASE(FCameraDirectorEvaluator)
+	UE_GAMEPLAY_CAMERAS_DECLARE_RTTI_BASE(GAMEPLAYCAMERAS_API, FCameraDirectorEvaluator)
 
 public:
 
@@ -136,11 +136,11 @@ using FCameraDirectorEvaluatorPtr = UE::Cameras::FCameraDirectorEvaluator*;
 
 // Utility macros for declaring and defining camera director evaluators.
 //
-#define UE_DECLARE_CAMERA_DIRECTOR_EVALUATOR(ClassName)\
-	UE_GAMEPLAY_CAMERAS_DECLARE_RTTI(ClassName, ::UE::Cameras::FCameraDirectorEvaluator)
+#define UE_DECLARE_CAMERA_DIRECTOR_EVALUATOR(ApiDeclSpec, ClassName)\
+	UE_GAMEPLAY_CAMERAS_DECLARE_RTTI(ApiDeclSpec, ClassName, ::UE::Cameras::FCameraDirectorEvaluator)
 
-#define UE_DECLARE_CAMERA_DIRECTOR_EVALUATOR_EX(ClassName, BaseClassName)\
-	UE_GAMEPLAY_CAMERAS_DECLARE_RTTI(ClassName, BaseClassName)
+#define UE_DECLARE_CAMERA_DIRECTOR_EVALUATOR_EX(ApiDeclSpec, ClassName, BaseClassName)\
+	UE_GAMEPLAY_CAMERAS_DECLARE_RTTI(ApiDeclSpec, ClassName, BaseClassName)
 
 #define UE_DEFINE_CAMERA_DIRECTOR_EVALUATOR(ClassName)\
 	UE_GAMEPLAY_CAMERAS_DEFINE_RTTI(ClassName)
diff --git a/Engine/Plugins/Cameras/GameplayCameras/Source/GameplayCameras/Public/Core/CameraEvaluationContext.h b/Engine/Plugins/Cameras/GameplayCameras/Source/GameplayCameras/Public/Core/CameraEvaluationContext.h
--- a/Engine/Plugins/Cameras/GameplayCameras/Source/GameplayCameras/Public/Core/CameraEvaluationContext.h
+++ b/Engine/Plugins/Cameras/GameplayCameras/Source/GameplayCameras/Public/Core/CameraEvaluationContext.h
@@ -19,7 +19,7 @@ namespace UE::Cameras
  */
 class FCameraEvaluationContext : public TSharedFromThis<FCameraEvaluationContext>
 {
-	UE_GAMEPLAY_CAMERAS_DECLARE_RTTI_BASE(FCameraEvaluationContext)
+	UE_GAMEPLAY_CAMERAS_DECLARE_RTTI_BASE(GAMEPLAYCAMERAS_API, FCameraEvaluationContext)
 
 public:
 
@@ -67,11 +67,11 @@ protected:
 
 // Utility macros for declaring and defining camera evaluation contexts.
 //
-#define UE_DECLARE_CAMERA_EVALUATION_CONTEXT(ClassName)\
-	UE_GAMEPLAY_CAMERAS_DECLARE_RTTI(ClassName, ::UE::Cameras::FCameraEvaluationContext)
+#define UE_DECLARE_CAMERA_EVALUATION_CONTEXT(ApiDeclSpec, ClassName)\
+	UE_GAMEPLAY_CAMERAS_DECLARE_RTTI(ApiDeclSpec, ClassName, ::UE::Cameras::FCameraEvaluationContext)
 
-#define UE_DECLARE_CAMERA_EVALUATION_CONTEXT_EX(ClassName, BaseClassName)\
-	UE_GAMEPLAY_CAMERAS_DECLARE_RTTI(ClassName, BaseClassName)
+#define UE_DECLARE_CAMERA_EVALUATION_CONTEXT_EX(ApiDeclSpec, ClassName, BaseClassName)\
+	UE_GAMEPLAY_CAMERAS_DECLARE_RTTI(ApiDeclSpec, ClassName, BaseClassName)
 
 #define UE_DEFINE_CAMERA_EVALUATION_CONTEXT(ClassName)\
 	UE_GAMEPLAY_CAMERAS_DEFINE_RTTI(ClassName)
diff --git a/Engine/Plugins/Cameras/GameplayCameras/Source/GameplayCameras/Public/Core/CameraNodeEvaluator.h b/Engine/Plugins/Cameras/GameplayCameras/Source/GameplayCameras/Public/Core/CameraNodeEvaluator.h
--- a/Engine/Plugins/Cameras/GameplayCameras/Source/GameplayCameras/Public/Core/CameraNodeEvaluator.h
+++ b/Engine/Plugins/Cameras/GameplayCameras/Source/GameplayCameras/Public/Core/CameraNodeEvaluator.h
@@ -114,7 +114,7 @@ struct FCameraDebugBlockBuildParams
  */
 class FCameraNodeEvaluator
 {
-	UE_GAMEPLAY_CAMERAS_DECLARE_RTTI_BASE(FCameraNodeEvaluator)
+	UE_GAMEPLAY_CAMERAS_DECLARE_RTTI_BASE(GAMEPLAYCAMERAS_API, FCameraNodeEvaluator)
 
 public:
 
@@ -208,11 +208,11 @@ EvaluatorType* FCameraNodeEvaluatorBuildParams::BuildEvaluatorAs(const UCameraNo
 
 // Utility macros for declaring and defining camera node evaluators.
 //
-#define UE_DECLARE_CAMERA_NODE_EVALUATOR(ClassName)\
-	UE_GAMEPLAY_CAMERAS_DECLARE_RTTI(ClassName, FCameraNodeEvaluator)
+#define UE_DECLARE_CAMERA_NODE_EVALUATOR(ApiDeclSpec, ClassName)\
+	UE_GAMEPLAY_CAMERAS_DECLARE_RTTI(ApiDeclSpec, ClassName, FCameraNodeEvaluator)
 
-#define UE_DECLARE_CAMERA_NODE_EVALUATOR_EX(ClassName, BaseClassName)\
-	UE_GAMEPLAY_CAMERAS_DECLARE_RTTI(ClassName, BaseClassName)
+#define UE_DECLARE_CAMERA_NODE_EVALUATOR_EX(ApiDeclSpec, ClassName, BaseClassName)\
+	UE_GAMEPLAY_CAMERAS_DECLARE_RTTI(ApiDeclSpec, ClassName, BaseClassName)
 
 #define UE_DEFINE_CAMERA_NODE_EVALUATOR(ClassName)\
 	UE_GAMEPLAY_CAMERAS_DEFINE_RTTI(ClassName)
diff --git a/Engine/Plugins/Cameras/GameplayCameras/Source/GameplayCameras/Public/Core/CameraObjectRtti.h b/Engine/Plugins/Cameras/GameplayCameras/Source/GameplayCameras/Public/Core/CameraObjectRtti.h
--- a/Engine/Plugins/Cameras/GameplayCameras/Source/GameplayCameras/Public/Core/CameraObjectRtti.h
+++ b/Engine/Plugins/Cameras/GameplayCameras/Source/GameplayCameras/Public/Core/CameraObjectRtti.h
@@ -139,8 +139,9 @@ private:
 // The first macro is for the root of the class hierarchy, while the second is for all
 // other classes below it. The third macro goes in the cpp file of each class.
 //
-#define UE_GAMEPLAY_CAMERAS_DECLARE_RTTI_BASE(ClassName)\
+#define UE_GAMEPLAY_CAMERAS_DECLARE_RTTI_BASE(ApiDeclSpec, ClassName)\
 	public:\
+		using FCameraObjectTypeID = ::UE::Cameras::FCameraObjectTypeID;\
 		static const ::UE::Cameras::TCameraObjectTypeID<ClassName>& StaticTypeID() { return ClassName::PrivateTypeID; }\
 		virtual const FCameraObjectTypeID& GetTypeID() const { return ClassName::PrivateTypeID; }\
 		virtual bool IsKindOf(const FCameraObjectTypeID& InTypeID) const { return InTypeID == ClassName::PrivateTypeID; }\
@@ -150,15 +151,15 @@ private:
 		template<typename Type> Type* CastThisChecked() { check(IsKindOf<Type>()); return static_cast<Type*>(this); }\
 		template<typename Type> const Type* CastThisChecked() const { check(IsKindOf<Type>()); return static_cast<Type*>(this); }\
 	private:\
-		GAMEPLAYCAMERAS_API static const ::UE::Cameras::TCameraObjectTypeID<ClassName> PrivateTypeID;
+		ApiDeclSpec static const ::UE::Cameras::TCameraObjectTypeID<ClassName> PrivateTypeID;
 
-#define UE_GAMEPLAY_CAMERAS_DECLARE_RTTI(ClassName, BaseClassName)\
+#define UE_GAMEPLAY_CAMERAS_DECLARE_RTTI(ApiDeclSpec, ClassName, BaseClassName)\
 	public:\
 		static const ::UE::Cameras::TCameraObjectTypeID<ClassName>& StaticTypeID() { return ClassName::PrivateTypeID; }\
 		virtual const FCameraObjectTypeID& GetTypeID() const override { return ClassName::PrivateTypeID; }\
 		virtual bool IsKindOf(const FCameraObjectTypeID& InTypeID) const override { return (InTypeID == ClassName::PrivateTypeID) || BaseClassName::IsKindOf(InTypeID); }\
 	private:\
-		GAMEPLAYCAMERAS_API static const ::UE::Cameras::TCameraObjectTypeID<ClassName> PrivateTypeID;
+		ApiDeclSpec static const ::UE::Cameras::TCameraObjectTypeID<ClassName> PrivateTypeID;
 
 #define UE_GAMEPLAY_CAMERAS_DEFINE_RTTI(ClassName)\
 	const ::UE::Cameras::TCameraObjectTypeID<ClassName> ClassName::PrivateTypeID = ::UE::Cameras::TCameraObjectTypeID<ClassName>::RegisterType(#ClassName);
diff --git a/Engine/Plugins/Cameras/GameplayCameras/Source/GameplayCameras/Public/Core/DefaultRootCameraNode.h b/Engine/Plugins/Cameras/GameplayCameras/Source/GameplayCameras/Public/Core/DefaultRootCameraNode.h
--- a/Engine/Plugins/Cameras/GameplayCameras/Source/GameplayCameras/Public/Core/DefaultRootCameraNode.h
+++ b/Engine/Plugins/Cameras/GameplayCameras/Source/GameplayCameras/Public/Core/DefaultRootCameraNode.h
@@ -49,7 +49,7 @@ class FBlendStackCameraNodeEvaluator;
  */
 class FDefaultRootCameraNodeEvaluator : public FRootCameraNodeEvaluator
 {
-	UE_DECLARE_CAMERA_NODE_EVALUATOR(FDefaultRootCameraNodeEvaluator)
+	UE_DECLARE_CAMERA_NODE_EVALUATOR(GAMEPLAYCAMERAS_API, FDefaultRootCameraNodeEvaluator)
 
 protected:
 
diff --git a/Engine/Plugins/Cameras/GameplayCameras/Source/GameplayCameras/Public/Debug/BlendStacksCameraDebugBlock.h b/Engine/Plugins/Cameras/GameplayCameras/Source/GameplayCameras/Public/Debug/BlendStacksCameraDebugBlock.h
--- a/Engine/Plugins/Cameras/GameplayCameras/Source/GameplayCameras/Public/Debug/BlendStacksCameraDebugBlock.h
+++ b/Engine/Plugins/Cameras/GameplayCameras/Source/GameplayCameras/Public/Debug/BlendStacksCameraDebugBlock.h
@@ -16,7 +16,7 @@ class FBlendStackCameraDebugBlock;
  */
 class FBlendStacksCameraDebugBlock : public FCameraDebugBlock
 {
-	UE_DECLARE_CAMERA_DEBUG_BLOCK(FBlendStacksCameraDebugBlock)
+	UE_DECLARE_CAMERA_DEBUG_BLOCK(GAMEPLAYCAMERAS_API, FBlendStacksCameraDebugBlock)
 
 public:
 
diff --git a/Engine/Plugins/Cameras/GameplayCameras/Source/GameplayCameras/Public/Debug/CameraDebugBlock.h b/Engine/Plugins/Cameras/GameplayCameras/Source/GameplayCameras/Public/Debug/CameraDebugBlock.h
--- a/Engine/Plugins/Cameras/GameplayCameras/Source/GameplayCameras/Public/Debug/CameraDebugBlock.h
+++ b/Engine/Plugins/Cameras/GameplayCameras/Source/GameplayCameras/Public/Debug/CameraDebugBlock.h
@@ -67,7 +67,7 @@ struct TCameraDebugBlockField : FCameraDebugBlockField
  */
 class FCameraDebugBlock
 {
-	UE_GAMEPLAY_CAMERAS_DECLARE_RTTI_BASE(FCameraDebugBlock)
+	UE_GAMEPLAY_CAMERAS_DECLARE_RTTI_BASE(GAMEPLAYCAMERAS_API, FCameraDebugBlock)
 
 public:
 
@@ -126,11 +126,11 @@ private:
 
 // Macros for defining debug blocks.
 //
-#define UE_DECLARE_CAMERA_DEBUG_BLOCK(ClassName)\
-	UE_GAMEPLAY_CAMERAS_DECLARE_RTTI(ClassName, ::UE::Cameras::FCameraDebugBlock)
+#define UE_DECLARE_CAMERA_DEBUG_BLOCK(ApiDeclSpec, ClassName)\
+	UE_GAMEPLAY_CAMERAS_DECLARE_RTTI(ApiDeclSpec, ClassName, ::UE::Cameras::FCameraDebugBlock)
 
-#define UE_DECLARE_CAMERA_DEBUG_BLOCK_EX(ClassName, BaseClassName)\
-	UE_GAMEPLAY_CAMERAS_DECLARE_RTTI(ClassName, BaseClassName)
+#define UE_DECLARE_CAMERA_DEBUG_BLOCK_EX(ApiDeclSpec, ClassName, BaseClassName)\
+	UE_GAMEPLAY_CAMERAS_DECLARE_RTTI(ApiDeclSpec, ClassName, BaseClassName)
 
 #define UE_DEFINE_CAMERA_DEBUG_BLOCK(ClassName)\
 	UE_GAMEPLAY_CAMERAS_DEFINE_RTTI(ClassName)
@@ -145,17 +145,17 @@ private:
 //
 // Example:
 //
-//     UE_DECLARE_CAMERA_DEBUG_BLOCK_START(FMyDebugBlock)
+//     UE_DECLARE_CAMERA_DEBUG_BLOCK_START(MYMODULE_API, FMyDebugBlock)
 //			UE_DECLARE_CAMERA_DEBUG_BLOCK_FIELD(float, Something)
 //			UE_DECLARE_CAMERA_DEBUG_BLOCK_FIELD(bool, bSomethingElse)
 //	   UE_DECLARE_CAMERA_DEBUG_BLOCK_END()
 //
 //	   UE_DEFINE_CAMERA_DEBUG_BLOCK(FMyDebugBlock)
 //
-#define UE_DECLARE_CAMERA_DEBUG_BLOCK_START(ClassName)\
+#define UE_DECLARE_CAMERA_DEBUG_BLOCK_START(ApiDeclSpec, ClassName)\
 	class ClassName : public ::UE::Cameras::FCameraDebugBlock\
 	{\
-		UE_DECLARE_CAMERA_DEBUG_BLOCK(ClassName)\
+		UE_DECLARE_CAMERA_DEBUG_BLOCK(ApiDeclSpec, ClassName)\
 	private:\
 		using Super = ::UE::Cameras::FCameraDebugBlock;\
 		using ThisClassName = ClassName;\
@@ -196,7 +196,7 @@ private:
 
 // Empty macros for shipping builds.
 //
-#define UE_DECLARE_CAMERA_DEBUG_BLOCK_START(ClassName)
+#define UE_DECLARE_CAMERA_DEBUG_BLOCK_START(ApiDeclSpec, ClassName)
 #define UE_DECLARE_CAMERA_DEBUG_BLOCK_FIELD(FieldType, FieldName)
 #define UE_DECLARE_CAMERA_DEBUG_BLOCK_END()
 #define UE_DEFINE_CAMERA_DEBUG_BLOCK_WITH_FIELDS(ClassName)
diff --git a/Engine/Plugins/Cameras/GameplayCameras/Source/GameplayCameras/Public/Debug/CameraDirectorTreeDebugBlock.h b/Engine/Plugins/Cameras/GameplayCameras/Source/GameplayCameras/Public/Debug/CameraDirectorTreeDebugBlock.h
--- a/Engine/Plugins/Cameras/GameplayCameras/Source/GameplayCameras/Public/Debug/CameraDirectorTreeDebugBlock.h
+++ b/Engine/Plugins/Cameras/GameplayCameras/Source/GameplayCameras/Public/Debug/CameraDirectorTreeDebugBlock.h
@@ -17,7 +17,7 @@ struct FCameraEvaluationContextStack;
  */
 class FCameraDirectorTreeDebugBlock : public FCameraDebugBlock
 {
-	UE_DECLARE_CAMERA_DEBUG_BLOCK(FCameraDirectorTreeDebugBlock)
+	UE_DECLARE_CAMERA_DEBUG_BLOCK(GAMEPLAYCAMERAS_API, FCameraDirectorTreeDebugBlock)
 
 public:
 
diff --git a/Engine/Plugins/Cameras/GameplayCameras/Source/GameplayCameras/Public/Debug/CameraNodeEvaluationResultDebugBlock.h b/Engine/Plugins/Cameras/GameplayCameras/Source/GameplayCameras/Public/Debug/CameraNodeEvaluationResultDebugBlock.h
--- a/Engine/Plugins/Cameras/GameplayCameras/Source/GameplayCameras/Public/Debug/CameraNodeEvaluationResultDebugBlock.h
+++ b/Engine/Plugins/Cameras/GameplayCameras/Source/GameplayCameras/Public/Debug/CameraNodeEvaluationResultDebugBlock.h
@@ -22,7 +22,7 @@ struct FCameraNodeEvaluationResult;
  */
 class FCameraNodeEvaluationResultDebugBlock : public FCameraDebugBlock
 {
-	UE_DECLARE_CAMERA_DEBUG_BLOCK(FCameraNodeEvaluationResultDebugBlock)
+	UE_DECLARE_CAMERA_DEBUG_BLOCK(GAMEPLAYCAMERAS_API, FCameraNodeEvaluationResultDebugBlock)
 
 public:
 
diff --git a/Engine/Plugins/Cameras/GameplayCameras/Source/GameplayCameras/Public/Debug/CameraNodeEvaluatorDebugBlock.h b/Engine/Plugins/Cameras/GameplayCameras/Source/GameplayCameras/Public/Debug/CameraNodeEvaluatorDebugBlock.h
--- a/Engine/Plugins/Cameras/GameplayCameras/Source/GameplayCameras/Public/Debug/CameraNodeEvaluatorDebugBlock.h
+++ b/Engine/Plugins/Cameras/GameplayCameras/Source/GameplayCameras/Public/Debug/CameraNodeEvaluatorDebugBlock.h
@@ -19,7 +19,7 @@ namespace UE::Cameras
  */
 class FCameraNodeEvaluatorDebugBlock : public FCameraDebugBlock
 {
-	UE_DECLARE_CAMERA_DEBUG_BLOCK(FCameraNodeEvaluatorDebugBlock)
+	UE_DECLARE_CAMERA_DEBUG_BLOCK(GAMEPLAYCAMERAS_API, FCameraNodeEvaluatorDebugBlock)
 
 public:
 
diff --git a/Engine/Plugins/Cameras/GameplayCameras/Source/GameplayCameras/Public/Debug/CameraPoseDebugBlock.h b/Engine/Plugins/Cameras/GameplayCameras/Source/GameplayCameras/Public/Debug/CameraPoseDebugBlock.h
--- a/Engine/Plugins/Cameras/GameplayCameras/Source/GameplayCameras/Public/Debug/CameraPoseDebugBlock.h
+++ b/Engine/Plugins/Cameras/GameplayCameras/Source/GameplayCameras/Public/Debug/CameraPoseDebugBlock.h
@@ -15,7 +15,7 @@ namespace UE::Cameras
  */
 class FCameraPoseDebugBlock : public FCameraDebugBlock
 {
-	UE_DECLARE_CAMERA_DEBUG_BLOCK(FCameraPoseDebugBlock)
+	UE_DECLARE_CAMERA_DEBUG_BLOCK(GAMEPLAYCAMERAS_API, FCameraPoseDebugBlock)
 
 public:
 
diff --git a/Engine/Plugins/Cameras/GameplayCameras/Source/GameplayCameras/Public/Debug/CategoryTitleDebugBlock.h b/Engine/Plugins/Cameras/GameplayCameras/Source/GameplayCameras/Public/Debug/CategoryTitleDebugBlock.h
--- a/Engine/Plugins/Cameras/GameplayCameras/Source/GameplayCameras/Public/Debug/CategoryTitleDebugBlock.h
+++ b/Engine/Plugins/Cameras/GameplayCameras/Source/GameplayCameras/Public/Debug/CategoryTitleDebugBlock.h
@@ -17,7 +17,7 @@ class FCameraNodeEvaluatorDebugBlock;
  */
 class FCategoryTitleDebugBlock : public FCameraDebugBlock
 {
-	UE_DECLARE_CAMERA_DEBUG_BLOCK(FCategoryTitleDebugBlock)
+	UE_DECLARE_CAMERA_DEBUG_BLOCK(GAMEPLAYCAMERAS_API, FCategoryTitleDebugBlock)
 
 public:
 
diff --git a/Engine/Plugins/Cameras/GameplayCameras/Source/GameplayCameras/Public/Debug/RootCameraDebugBlock.h b/Engine/Plugins/Cameras/GameplayCameras/Source/GameplayCameras/Public/Debug/RootCameraDebugBlock.h
--- a/Engine/Plugins/Cameras/GameplayCameras/Source/GameplayCameras/Public/Debug/RootCameraDebugBlock.h
+++ b/Engine/Plugins/Cameras/GameplayCameras/Source/GameplayCameras/Public/Debug/RootCameraDebugBlock.h
@@ -22,7 +22,7 @@ GAMEPLAYCAMERAS_API extern FString GGameplayCamerasDebugCategories;
  */
 class FRootCameraDebugBlock : public FCameraDebugBlock
 {
-	UE_DECLARE_CAMERA_DEBUG_BLOCK(FRootCameraDebugBlock)
+	UE_DECLARE_CAMERA_DEBUG_BLOCK(GAMEPLAYCAMERAS_API, FRootCameraDebugBlock)
 
 public:
 
diff --git a/Engine/Plugins/Cameras/GameplayCameras/Source/GameplayCameras/Public/Debug/VariableTableDebugBlock.h b/Engine/Plugins/Cameras/GameplayCameras/Source/GameplayCameras/Public/Debug/VariableTableDebugBlock.h
--- a/Engine/Plugins/Cameras/GameplayCameras/Source/GameplayCameras/Public/Debug/VariableTableDebugBlock.h
+++ b/Engine/Plugins/Cameras/GameplayCameras/Source/GameplayCameras/Public/Debug/VariableTableDebugBlock.h
@@ -16,7 +16,7 @@ class FCameraVariableTable;
  */
 class FVariableTableDebugBlock : public FCameraDebugBlock
 {
-	UE_DECLARE_CAMERA_DEBUG_BLOCK(FVariableTableDebugBlock)
+	UE_DECLARE_CAMERA_DEBUG_BLOCK(GAMEPLAYCAMERAS_API, FVariableTableDebugBlock)
 
 public:
 
diff --git a/Engine/Plugins/Cameras/GameplayCameras/Source/GameplayCameras/Public/Debug/ViewfinderDebugBlock.h b/Engine/Plugins/Cameras/GameplayCameras/Source/GameplayCameras/Public/Debug/ViewfinderDebugBlock.h
--- a/Engine/Plugins/Cameras/GameplayCameras/Source/GameplayCameras/Public/Debug/ViewfinderDebugBlock.h
+++ b/Engine/Plugins/Cameras/GameplayCameras/Source/GameplayCameras/Public/Debug/ViewfinderDebugBlock.h
@@ -11,7 +11,7 @@ namespace UE::Cameras
 
 class FViewfinderDebugBlock : public FCameraDebugBlock
 {
-	UE_DECLARE_CAMERA_DEBUG_BLOCK(FViewfinderDebugBlock)
+	UE_DECLARE_CAMERA_DEBUG_BLOCK(GAMEPLAYCAMERAS_API, FViewfinderDebugBlock)
 
 public:
 
diff --git a/Engine/Plugins/Cameras/GameplayCameras/Source/GameplayCameras/Public/GameFramework/GameplayCameraComponent.h b/Engine/Plugins/Cameras/GameplayCameras/Source/GameplayCameras/Public/GameFramework/GameplayCameraComponent.h
--- a/Engine/Plugins/Cameras/GameplayCameras/Source/GameplayCameras/Public/GameFramework/GameplayCameraComponent.h
+++ b/Engine/Plugins/Cameras/GameplayCameras/Source/GameplayCameras/Public/GameFramework/GameplayCameraComponent.h
@@ -96,7 +96,7 @@ namespace UE::Cameras
  */
 class FGameplayCameraComponentEvaluationContext : public FCameraEvaluationContext
 {
-	UE_DECLARE_CAMERA_EVALUATION_CONTEXT(FGameplayCameraComponentEvaluationContext)
+	UE_DECLARE_CAMERA_EVALUATION_CONTEXT(GAMEPLAYCAMERAS_API, FGameplayCameraComponentEvaluationContext)
 
 public:
 
diff --git a/Engine/Plugins/Cameras/GameplayCameras/Source/GameplayCameras/Public/Nodes/Blends/SimpleBlendCameraNode.h b/Engine/Plugins/Cameras/GameplayCameras/Source/GameplayCameras/Public/Nodes/Blends/SimpleBlendCameraNode.h
--- a/Engine/Plugins/Cameras/GameplayCameras/Source/GameplayCameras/Public/Nodes/Blends/SimpleBlendCameraNode.h
+++ b/Engine/Plugins/Cameras/GameplayCameras/Source/GameplayCameras/Public/Nodes/Blends/SimpleBlendCameraNode.h
@@ -43,7 +43,7 @@ struct FSimpleBlendCameraNodeEvaluationResult
 
 class FSimpleBlendCameraNodeEvaluator : public FBlendCameraNodeEvaluator
 {
-	UE_DECLARE_BLEND_CAMERA_NODE_EVALUATOR(FSimpleBlendCameraNodeEvaluator)
+	UE_DECLARE_BLEND_CAMERA_NODE_EVALUATOR(GAMEPLAYCAMERAS_API, FSimpleBlendCameraNodeEvaluator)
 
 public:
 
@@ -71,7 +71,7 @@ private:
 
 class FSimpleFixedTimeBlendCameraNodeEvaluator : public FSimpleBlendCameraNodeEvaluator
 {
-	UE_DECLARE_BLEND_CAMERA_NODE_EVALUATOR_EX(FSimpleFixedTimeBlendCameraNodeEvaluator, FSimpleBlendCameraNodeEvaluator)
+	UE_DECLARE_BLEND_CAMERA_NODE_EVALUATOR_EX(GAMEPLAYCAMERAS_API, FSimpleFixedTimeBlendCameraNodeEvaluator, FSimpleBlendCameraNodeEvaluator)
 
 protected:
 
