Cameras: add error/warning logging when invalid situations are encountered in a blend stack

diff --git a/Engine/Plugins/Cameras/GameplayCameras/Source/GameplayCameras/Private/Core/BlendStackCameraNode.cpp b/Engine/Plugins/Cameras/GameplayCameras/Source/GameplayCameras/Private/Core/BlendStackCameraNode.cpp
--- a/Engine/Plugins/Cameras/GameplayCameras/Source/GameplayCameras/Private/Core/BlendStackCameraNode.cpp
+++ b/Engine/Plugins/Cameras/GameplayCameras/Source/GameplayCameras/Private/Core/BlendStackCameraNode.cpp
@@ -154,6 +154,17 @@ void FBlendStackCameraNodeEvaluator::OnRun(const FCameraNodeEvaluationParams& Pa
 		if (UNLIKELY(!CurContext.IsValid()))
 		{
 			Entry.Result.bIsValid = false;
+
+#if UE_GAMEPLAY_CAMERAS_TRACE
+			if (Entry.bLogWarnings)
+			{
+				UE_LOG(LogCameraSystem, Warning,
+						TEXT("Can't update camera rig '%s' because its evaluation context isn't valid."),
+						*GetNameSafe(Entry.CameraRig));
+				Entry.bLogWarnings = false;
+			}
+#endif  // UE_GAMEPLAY_CAMERAS_TRACE
+
 			continue;
 		}
 
@@ -171,6 +182,17 @@ void FBlendStackCameraNodeEvaluator::OnRun(const FCameraNodeEvaluationParams& Pa
 			if (UNLIKELY(!ContextResult.bIsValid))
 			{
 				CurResult.bIsValid = false;
+
+#if UE_GAMEPLAY_CAMERAS_TRACE
+				if (Entry.bLogWarnings)
+				{
+					UE_LOG(LogCameraSystem, Warning,
+							TEXT("Can't update camera rig '%s' because its initial result isn't valid."),
+							*GetNameSafe(Entry.CameraRig));
+					Entry.bLogWarnings = false;
+				}
+#endif  // UE_GAMEPLAY_CAMERAS_TRACE
+
 				continue;
 			}
 
@@ -194,6 +216,10 @@ void FBlendStackCameraNodeEvaluator::OnRun(const FCameraNodeEvaluationParams& Pa
 			// Only evaluate the blend via the root node.
 			Entry.RootEvaluator->Run(CurParams, CurResult);
 		}
+
+#if UE_GAMEPLAY_CAMERAS_TRACE
+		Entry.bLogWarnings = true;
+#endif  // UE_GAMEPLAY_CAMERAS_TRACE
 	}
 
 	// Now blend all the results, keeping track of blends that have reached 100% so
@@ -521,7 +547,7 @@ FBlendStackCameraDebugBlock::FBlendStackCameraDebugBlock(const FBlendStackCamera
 	for (const FBlendStackCameraNodeEvaluator::FCameraRigEntry& Entry : InEvaluator.Entries)
 	{
 		FEntryDebugInfo EntryDebugInfo;
-		EntryDebugInfo.CameraRigName = Entry.CameraRig ? Entry.CameraRig->GetName() : TEXT("<no camera rig>");
+		EntryDebugInfo.CameraRigName = GetNameSafe(Entry.CameraRig);
 		Entries.Add(EntryDebugInfo);
 	}
 }
diff --git a/Engine/Plugins/Cameras/GameplayCameras/Source/GameplayCameras/Public/Core/BlendStackCameraNode.h b/Engine/Plugins/Cameras/GameplayCameras/Source/GameplayCameras/Public/Core/BlendStackCameraNode.h
--- a/Engine/Plugins/Cameras/GameplayCameras/Source/GameplayCameras/Public/Core/BlendStackCameraNode.h
+++ b/Engine/Plugins/Cameras/GameplayCameras/Source/GameplayCameras/Public/Core/BlendStackCameraNode.h
@@ -146,6 +146,11 @@ protected:
 		bool bIsFirstFrame = false;
 		/** Whether this entry is frozen. */
 		bool bIsFrozen = false;
+
+#if UE_GAMEPLAY_CAMERAS_TRACE
+		bool bLogWarnings = true;
+#endif  // UE_GAMEPLAY_CAMERAS_TRACE
+
 #if WITH_EDITOR
 		FCameraRigPackages ListenedPackages;
 #endif  // WITH_EDITOR
