Cameras: add "super" class alias to rtti framework

diff --git a/Engine/Plugins/Cameras/GameplayCameras/Source/GameplayCameras/Public/Core/CameraObjectRtti.h b/Engine/Plugins/Cameras/GameplayCameras/Source/GameplayCameras/Public/Core/CameraObjectRtti.h
--- a/Engine/Plugins/Cameras/GameplayCameras/Source/GameplayCameras/Public/Core/CameraObjectRtti.h
+++ b/Engine/Plugins/Cameras/GameplayCameras/Source/GameplayCameras/Public/Core/CameraObjectRtti.h
@@ -155,6 +155,7 @@ private:
 
 #define UE_GAMEPLAY_CAMERAS_DECLARE_RTTI(ApiDeclSpec, ClassName, BaseClassName)\
 	public:\
+		using Super = BaseClassName;\
 		static const ::UE::Cameras::TCameraObjectTypeID<ClassName>& StaticTypeID() { return ClassName::PrivateTypeID; }\
 		virtual const FCameraObjectTypeID& GetTypeID() const override { return ClassName::PrivateTypeID; }\
 		virtual bool IsKindOf(const FCameraObjectTypeID& InTypeID) const override { return (InTypeID == ClassName::PrivateTypeID) || BaseClassName::IsKindOf(InTypeID); }\
diff --git a/Engine/Plugins/Cameras/GameplayCameras/Source/GameplayCameras/Public/Debug/CameraDebugBlock.h b/Engine/Plugins/Cameras/GameplayCameras/Source/GameplayCameras/Public/Debug/CameraDebugBlock.h
--- a/Engine/Plugins/Cameras/GameplayCameras/Source/GameplayCameras/Public/Debug/CameraDebugBlock.h
+++ b/Engine/Plugins/Cameras/GameplayCameras/Source/GameplayCameras/Public/Debug/CameraDebugBlock.h
@@ -157,7 +157,6 @@ private:
 	{\
 		UE_DECLARE_CAMERA_DEBUG_BLOCK(ApiDeclSpec, ClassName)\
 	private:\
-		using Super = ::UE::Cameras::FCameraDebugBlock;\
 		using ThisClassName = ClassName;\
 		static FStaticFieldArray StaticFields;\
 	protected:\
