Cameras: implement hierarchical evaluation contexts and directors

- Rework the relationship between evaluation contexts and camera directors: the former are now creating and owning the latter.
- Children contexts can now be created, and their director is "injected" inside their parent context's director.

diff --git a/Engine/Plugins/Cameras/GameplayCameras/Source/GameplayCameras/Private/Core/CameraDirectorEvaluator.cpp b/Engine/Plugins/Cameras/GameplayCameras/Source/GameplayCameras/Private/Core/CameraDirectorEvaluator.cpp
--- a/Engine/Plugins/Cameras/GameplayCameras/Source/GameplayCameras/Private/Core/CameraDirectorEvaluator.cpp
+++ b/Engine/Plugins/Cameras/GameplayCameras/Source/GameplayCameras/Private/Core/CameraDirectorEvaluator.cpp
@@ -3,24 +3,126 @@
 #include "Core/CameraDirectorEvaluator.h"
 
 #include "Core/CameraDirector.h"
+#include "Core/CameraEvaluationContext.h"
 
 namespace UE::Cameras
 {
 
+void FCameraDirectorEvaluatorStorage::DestroyEvaluator()
+{
+	Evaluator.Reset();
+}
+
 UE_GAMEPLAY_CAMERAS_DEFINE_RTTI(FCameraDirectorEvaluator)
 
 FCameraDirectorEvaluator::FCameraDirectorEvaluator()
 {
 }
 
+void FCameraDirectorEvaluator::SetPrivateCameraDirector(const UCameraDirector* InCameraDirector)
+{
+	PrivateCameraDirector = InCameraDirector;
+}
+
+void FCameraDirectorEvaluator::Initialize(const FCameraDirectorInitializeParams& Params)
+{
+	WeakOwnerContext = Params.OwnerContext;
+
+	OnInitialize(Params);
+}
+
 void FCameraDirectorEvaluator::Run(const FCameraDirectorEvaluationParams& Params, FCameraDirectorEvaluationResult& OutResult)
 {
 	OnRun(Params, OutResult);
 }
 
-void FCameraDirectorEvaluator::SetPrivateCameraDirector(TObjectPtr<const UCameraDirector> InCameraDirector)
+bool FCameraDirectorEvaluator::AddChildEvaluationContext(TSharedRef<FCameraEvaluationContext> InContext)
 {
-	PrivateCameraDirector = InCameraDirector;
+	TSharedPtr<FCameraEvaluationContext> OwnerContext = WeakOwnerContext.Pin();
+	if (!ensureMsgf(OwnerContext, TEXT("Can't add child evaluation context when the parent/owner context is invalid!")))
+	{
+		return false;
+	}
+
+	FChildContextManulationParams Params;
+	Params.ParentContext = OwnerContext;
+	Params.ChildContext = InContext.ToSharedPtr();
+	FChildContextManulationResult Result;
+	OnAddChildEvaluationContext(Params, Result);
+
+	bool bReturn = false;
+	bool bRegisterAndActivateChildContext = false;
+	switch (Result.Result)
+	{
+		case EChildContextManipulationResult::Failure:
+		default:
+			// Nothing to do.
+			break;
+		case EChildContextManipulationResult::Success:
+			// Our director evaluator accepted the child context.
+			bRegisterAndActivateChildContext = true;
+			bReturn = true;
+			break;
+		case EChildContextManipulationResult::ChildContextSuccess:
+			// A sub-director of our director accepted the child context, so it already
+			// activate it and we don't need to do it ourselves.
+			bReturn = true;
+			break;
+	}
+	if (bRegisterAndActivateChildContext)
+	{
+		OwnerContext->RegisterChildContext(InContext);
+
+		FCameraEvaluationContextActivateParams ActivateParams;
+		InContext->Activate(ActivateParams);
+	}
+	return bReturn;
+}
+
+bool FCameraDirectorEvaluator::RemoveChildEvaluationContext(TSharedRef<FCameraEvaluationContext> InContext)
+{
+	TSharedPtr<FCameraEvaluationContext> OwnerContext = WeakOwnerContext.Pin();
+	if (!ensureMsgf(OwnerContext, TEXT("Can't remove child evaluation context when the parent/owner context is invalid!")))
+	{
+		return false;
+	}
+
+	FChildContextManulationParams Params;
+	Params.ParentContext = OwnerContext;
+	Params.ChildContext = InContext;
+	FChildContextManulationResult Result;
+	OnRemoveChildEvaluationContext(Params, Result);
+
+	bool bReturn = false;
+	bool bUnregisterAndDeactivateChildContext = false;
+	switch (Result.Result)
+	{
+		case EChildContextManipulationResult::Failure:
+		default:
+			break;
+		case EChildContextManipulationResult::Success:
+			bUnregisterAndDeactivateChildContext = true;
+			bReturn = true;
+			break;
+		case EChildContextManipulationResult::ChildContextSuccess:
+			bReturn = true;
+			break;
+	}
+	if (bUnregisterAndDeactivateChildContext)
+	{
+		OwnerContext->UnregisterChildContext(InContext);
+
+		FCameraEvaluationContextDeactivateParams DeactivateParams;
+		InContext->Deactivate(DeactivateParams);
+	}
+	return bReturn;
+}
+
+void FCameraDirectorEvaluator::AddReferencedObjects(FReferenceCollector& Collector)
+{
+	Collector.AddReferencedObject(PrivateCameraDirector);
+
+	OnAddReferencedObjects(Collector);
 }
 
 }  // namespace UE::Cameras
diff --git a/Engine/Plugins/Cameras/GameplayCameras/Source/GameplayCameras/Private/Core/CameraEvaluationContext.cpp b/Engine/Plugins/Cameras/GameplayCameras/Source/GameplayCameras/Private/Core/CameraEvaluationContext.cpp
--- a/Engine/Plugins/Cameras/GameplayCameras/Source/GameplayCameras/Private/Core/CameraEvaluationContext.cpp
+++ b/Engine/Plugins/Cameras/GameplayCameras/Source/GameplayCameras/Private/Core/CameraEvaluationContext.cpp
@@ -2,6 +2,9 @@
 
 #include "Core/CameraEvaluationContext.h"
 
+#include "Core/CameraAsset.h"
+#include "Core/CameraDirector.h"
+
 namespace UE::Cameras
 {
 
@@ -11,8 +14,117 @@ FCameraEvaluationContext::FCameraEvaluationContext()
 {
 }
 
+void FCameraEvaluationContext::Initialize(const FCameraEvaluationContextInitializeParams& Params)
+{
+	if (!ensureMsgf(!bInitialized, TEXT("This evaluation context has already been initialized!")))
+	{
+		return;
+	}
+
+	CameraAsset = Params.CameraAsset;
+	PlayerController = Params.PlayerController;
+
+	bInitialized = true;
+}
+
 FCameraEvaluationContext::~FCameraEvaluationContext()
 {
+	// Camera director evaluator usually gets destroyed here since the storage object generally
+	// holds the only shared pointer to it.
+}
+
+void FCameraEvaluationContext::AddReferencedObjects(FReferenceCollector& Collector)
+{
+	Collector.AddReferencedObject(CameraAsset);
+
+	if (DirectorEvaluator)
+	{
+		DirectorEvaluator->AddReferencedObjects(Collector);
+	}
+
+	for (TSharedPtr<FCameraEvaluationContext> ChildContext : ChildrenContexts)
+	{
+		ChildContext->AddReferencedObjects(Collector);
+	}
+}
+
+void FCameraEvaluationContext::AutoCreateDirectorEvaluator()
+{
+	if (DirectorEvaluator == nullptr)
+	{
+		if (const UCameraDirector* CameraDirector = CameraAsset->CameraDirector)
+		{
+			FCameraDirectorEvaluatorBuilder DirectorBuilder(DirectorEvaluatorStorage);
+			DirectorEvaluator = CameraDirector->BuildEvaluator(DirectorBuilder);
+
+			FCameraDirectorInitializeParams InitParams;
+			InitParams.OwnerContext = SharedThis(this);
+			InitParams.CameraDirector = CameraDirector;
+			DirectorEvaluator->Initialize(InitParams);
+		}
+		else
+		{
+			UE_LOG(LogCameraSystem, Warning, TEXT("Activating an evaluation context without a camera director!"));
+		}
+	}
+}
+
+void FCameraEvaluationContext::Activate(const FCameraEvaluationContextActivateParams& Params)
+{
+	if (!ensureMsgf(bInitialized, TEXT("This evaluation context needs to be initialized!")))
+	{
+		return;
+	}
+	if (!ensureMsgf(!bActivated, TEXT("This evaluation context has already been activated!")))
+	{
+		return;
+	}
+
+	OnActivate(Params);
+
+	AutoCreateDirectorEvaluator();
+
+	bActivated = true;
+}
+
+void FCameraEvaluationContext::Deactivate(const FCameraEvaluationContextDeactivateParams& Params)
+{
+	if (!ensureMsgf(bActivated, TEXT("This evaluation context has not been activated!")))
+	{
+		return;
+	}
+
+	// Don't destroy the camera director evaluator, it could still be useful. We only destroy it
+	// along with this context.
+
+	OnDeactivate(Params);
+
+	bActivated = false;
+}
+
+bool FCameraEvaluationContext::RegisterChildContext(TSharedRef<FCameraEvaluationContext> ChildContext)
+{
+	if (!ensureMsgf(ChildContext->WeakParent == nullptr, TEXT("The given evaluation context already has a parent!")))
+	{
+		return false;
+	}
+
+	ChildContext->WeakParent = SharedThis(this);
+	ChildrenContexts.Add(ChildContext);
+	return true;
+}
+
+bool FCameraEvaluationContext::UnregisterChildContext(TSharedRef<FCameraEvaluationContext> ChildContext)
+{
+	if (!ensureMsgf(ChildContext->WeakParent == SharedThis(this), TEXT("The given evaluation context isn't our child!")))
+	{
+		return false;
+	}
+
+	ChildContext->WeakParent = nullptr;
+	const int32 NumRemoved = ChildrenContexts.Remove(ChildContext);
+	ensureMsgf(NumRemoved == 1, TEXT("The given evaluation context wasn't in our list of children!"));
+	return true;
 }
 
 }  // namespace UE::Cameras
diff --git a/Engine/Plugins/Cameras/GameplayCameras/Source/GameplayCameras/Private/Core/CameraEvaluationContextStack.cpp b/Engine/Plugins/Cameras/GameplayCameras/Source/GameplayCameras/Private/Core/CameraEvaluationContextStack.cpp
--- a/Engine/Plugins/Cameras/GameplayCameras/Source/GameplayCameras/Private/Core/CameraEvaluationContextStack.cpp
+++ b/Engine/Plugins/Cameras/GameplayCameras/Source/GameplayCameras/Private/Core/CameraEvaluationContextStack.cpp
@@ -12,16 +12,16 @@
 namespace UE::Cameras
 {
 
-FCameraEvaluationContextInfo FCameraEvaluationContextStack::GetActiveContext() const
+TSharedPtr<FCameraEvaluationContext> FCameraEvaluationContextStack::GetActiveContext() const
 {
 	for (const FContextEntry& Entry : ReverseIterate(Entries))
 	{
 		if (TSharedPtr<FCameraEvaluationContext> Context = Entry.WeakContext.Pin())
 		{
-			return FCameraEvaluationContextInfo{ Context, Entry.CameraDirector, Entry.Evaluator };
+			return Context;
 		}
 	}
-	return FCameraEvaluationContextInfo();
+	return nullptr;
 }
 
 bool FCameraEvaluationContextStack::HasContext(TSharedRef<FCameraEvaluationContext> Context) const
@@ -54,24 +54,45 @@ void FCameraEvaluationContextStack::PushContext(TSharedRef<FCameraEvaluationCont
 		return;
 	}
 
-	// Make a new entry and build the director evaluator using the entry's storage.
+	// Make a new entry and activate the context. This will build the director evaluator.
 	FContextEntry NewEntry;
 
-	FCameraDirectorEvaluatorBuilder Builder(NewEntry.EvaluatorStorage);
-	const UCameraDirector* CameraDirector = Context->GetCameraAsset()->CameraDirector;
-	FCameraDirectorEvaluator* DirectorEvaluator = CameraDirector->BuildEvaluator(Builder);
+	FCameraEvaluationContextActivateParams ActivateParams;
+	Context->Activate(ActivateParams);
 	
 	NewEntry.WeakContext = Context;
-	NewEntry.CameraDirector = CameraDirector;
-	NewEntry.Evaluator = DirectorEvaluator;
 	Entries.Push(MoveTemp(NewEntry));
 }
 
+bool FCameraEvaluationContextStack::AddChildContext(TSharedRef<FCameraEvaluationContext> Context)
+{
+	TSharedPtr<FCameraEvaluationContext> ActiveContext = GetActiveContext();
+	if (ensureMsgf(ActiveContext.IsValid(), TEXT("Can't add child context to the stack, no active context was found!")))
+	{
+		FCameraDirectorEvaluator* DirectorEvaluator = ActiveContext->GetDirectorEvaluator();
+		if (ensureMsgf(DirectorEvaluator, TEXT("Can't add child context, active context has no camera director evaluator!")))
+		{
+			return DirectorEvaluator->AddChildEvaluationContext(Context);
+		}
+	}
+	return false;
+}
+
 bool FCameraEvaluationContextStack::RemoveContext(TSharedRef<FCameraEvaluationContext> Context)
 {
-	const int32 NumRemoved = Entries.RemoveAll(
-			[Context](FContextEntry& Entry) { return Entry.WeakContext == Context; });
-	return (NumRemoved > 0);
+	for (auto It = Entries.CreateIterator(); It; ++It)
+	{
+		FContextEntry& Entry = (*It);
+		if (Entry.WeakContext == Context)
+		{
+			FCameraEvaluationContextDeactivateParams DeactivateParams;
+			Context->Deactivate(DeactivateParams);
+
+			It.RemoveCurrent();
+			return true;
+		}
+	}
+	return false;
 }
 
 void FCameraEvaluationContextStack::PopContext()
@@ -88,7 +109,10 @@ void FCameraEvaluationContextStack::AddReferencedObjects(FReferenceCollector& Co
 {
 	for (FContextEntry& Entry : Entries)
 	{
-		Collector.AddReferencedObject(Entry.CameraDirector);
+		if (TSharedPtr<FCameraEvaluationContext> Context = Entry.WeakContext.Pin())
+		{
+			Context->AddReferencedObjects(Collector);
+		}
 	}
 }
 
diff --git a/Engine/Plugins/Cameras/GameplayCameras/Source/GameplayCameras/Private/Core/CameraSystemEvaluator.cpp b/Engine/Plugins/Cameras/GameplayCameras/Source/GameplayCameras/Private/Core/CameraSystemEvaluator.cpp
--- a/Engine/Plugins/Cameras/GameplayCameras/Source/GameplayCameras/Private/Core/CameraSystemEvaluator.cpp
+++ b/Engine/Plugins/Cameras/GameplayCameras/Source/GameplayCameras/Private/Core/CameraSystemEvaluator.cpp
@@ -93,20 +93,20 @@ void FCameraSystemEvaluator::Update(const FCameraSystemEvaluationUpdateParams& P
 	SCOPE_CYCLE_COUNTER(CameraSystemEval_Total);
 
 	// Get the active evaluation context.
-	FCameraEvaluationContextInfo ActiveContextInfo = ContextStack.GetActiveContext();
-	if (UNLIKELY(!ActiveContextInfo.IsValid()))
+	TSharedPtr<FCameraEvaluationContext> ActiveContext = ContextStack.GetActiveContext();
+	if (UNLIKELY(!ActiveContext.IsValid()))
 	{
 		Result.bIsValid = false;
 		return;
 	}
 
 	// Run the camera director, and activate any camera rig(s) it returns to us.
-	FCameraDirectorEvaluator* ActiveDirectorEvaluator = ActiveContextInfo.Evaluator;
+	FCameraDirectorEvaluator* ActiveDirectorEvaluator = ActiveContext->GetDirectorEvaluator();
 	if (ActiveDirectorEvaluator)
 	{
 		FCameraDirectorEvaluationParams DirectorParams;
 		DirectorParams.DeltaTime = Params.DeltaTime;
-		DirectorParams.OwnerContext = ActiveContextInfo.EvaluationContext;
+		DirectorParams.OwnerContext = ActiveContext;
 
 		FCameraDirectorEvaluationResult DirectorResult;
 
@@ -114,12 +114,15 @@ void FCameraSystemEvaluator::Update(const FCameraSystemEvaluationUpdateParams& P
 
 		if (DirectorResult.ActiveCameraRigs.Num() == 1)
 		{
+			const FActiveCameraRigInfo& ActiveCameraRig = DirectorResult.ActiveCameraRigs[0];
+
 			FActivateCameraRigParams CameraRigParams;
 			CameraRigParams.Evaluator = this;
-			CameraRigParams.EvaluationContext = ActiveContextInfo.EvaluationContext;
-			CameraRigParams.CameraRig = DirectorResult.ActiveCameraRigs[0];
+			CameraRigParams.EvaluationContext = ActiveCameraRig.EvaluationContext;
+			CameraRigParams.CameraRig = ActiveCameraRig.CameraRig;
 			RootEvaluator->ActivateCameraRig(CameraRigParams);
 		}
+		// TODO: handle the case of composite camera rigs.
 	}
 
 	// Setup the params/result for running the root camera node.
diff --git a/Engine/Plugins/Cameras/GameplayCameras/Source/GameplayCameras/Private/Directors/SingleCameraDirector.cpp b/Engine/Plugins/Cameras/GameplayCameras/Source/GameplayCameras/Private/Directors/SingleCameraDirector.cpp
--- a/Engine/Plugins/Cameras/GameplayCameras/Source/GameplayCameras/Private/Directors/SingleCameraDirector.cpp
+++ b/Engine/Plugins/Cameras/GameplayCameras/Source/GameplayCameras/Private/Directors/SingleCameraDirector.cpp
@@ -16,7 +16,7 @@ protected:
 		const USingleCameraDirector* SingleDirector = GetCameraDirectorAs<USingleCameraDirector>();
 		if (SingleDirector->CameraRig)
 		{
-			OutResult.ActiveCameraRigs.Add(SingleDirector->CameraRig);
+			OutResult.Add(Params.OwnerContext, SingleDirector->CameraRig);
 		}
 	}
 };
diff --git a/Engine/Plugins/Cameras/GameplayCameras/Source/GameplayCameras/Public/Core/CameraDirectorEvaluator.h b/Engine/Plugins/Cameras/GameplayCameras/Source/GameplayCameras/Public/Core/CameraDirectorEvaluator.h
--- a/Engine/Plugins/Cameras/GameplayCameras/Source/GameplayCameras/Public/Core/CameraDirectorEvaluator.h
+++ b/Engine/Plugins/Cameras/GameplayCameras/Source/GameplayCameras/Public/Core/CameraDirectorEvaluator.h
@@ -4,6 +4,7 @@
 
 #include "Core/CameraObjectRtti.h"
 #include "CoreTypes.h"
+#include "GameplayCameras.h"
 #include "Templates/UniquePtr.h"
 #include "UObject/ObjectPtr.h"
 
@@ -17,6 +18,12 @@ class FCameraDirectorEvaluator;
 class FCameraDirectorEvaluatorStorage;
 class FCameraEvaluationContext;
 
+struct FCameraDirectorInitializeParams
+{
+	TSharedPtr<FCameraEvaluationContext> OwnerContext;
+	TObjectPtr<const UCameraDirector> CameraDirector;
+};
+
 /**
  * Parameter structure for running a camera director.
  */
@@ -29,13 +36,25 @@ struct FCameraDirectorEvaluationParams
 	TSharedPtr<const FCameraEvaluationContext> OwnerContext;
 };
 
+struct FActiveCameraRigInfo
+{
+	TSharedPtr<const FCameraEvaluationContext> EvaluationContext;
+	TObjectPtr<const UCameraRigAsset> CameraRig;
+};
+
 /**
  * Result structure for running a camera director.
  */
 struct FCameraDirectorEvaluationResult
 {
+	using FActiveCameraRigInfos = TArray<FActiveCameraRigInfo, TInlineAllocator<2>>;
 	/** The camera rig(s) that the director says should be active this frame. */
-	TArray<TObjectPtr<const UCameraRigAsset>, TInlineAllocator<2>> ActiveCameraRigs;
+	FActiveCameraRigInfos ActiveCameraRigs;
+
+	void Add(TSharedPtr<const FCameraEvaluationContext> InContext, TObjectPtr<const UCameraRigAsset> InCameraRig)
+	{
+		ActiveCameraRigs.Add(FActiveCameraRigInfo{ InContext, InCameraRig });
+	}
 };
 
 /**
@@ -66,11 +85,13 @@ public:
 	/** Gets the stored evaluator, if any. */
 	FCameraDirectorEvaluator* GetEvaluator() const { return Evaluator.Get(); }
 
-private:
-
 	template<typename EvaluatorType, typename ...ArgTypes>
 	EvaluatorType* BuildEvaluator(ArgTypes&&... InArgs);
 
+	void DestroyEvaluator();
+
+private:
+
 	TSharedPtr<FCameraDirectorEvaluator> Evaluator;
 
 	friend struct FCameraDirectorEvaluatorBuilder;
@@ -87,10 +108,14 @@ public:
 
 	GAMEPLAYCAMERAS_API FCameraDirectorEvaluator();
 	GAMEPLAYCAMERAS_API virtual ~FCameraDirectorEvaluator() {}
+
+	void Initialize(const FCameraDirectorInitializeParams& Params);
 	
 	/** Runs the camera director to determine what camera rig(s) should be active this frame. */
 	void Run(const FCameraDirectorEvaluationParams& Params, FCameraDirectorEvaluationResult& OutResult);
 
+	const UCameraDirector* GetCameraDirector() const { return PrivateCameraDirector; }
+
 	/** Gets the camera director. */
 	template<typename CameraDirectorType>
 	const CameraDirectorType* GetCameraDirectorAs() const
@@ -98,18 +123,52 @@ public:
 		return Cast<CameraDirectorType>(PrivateCameraDirector);
 	}
 
+	GAMEPLAYCAMERAS_API bool AddChildEvaluationContext(TSharedRef<FCameraEvaluationContext> InContext);
+	GAMEPLAYCAMERAS_API bool RemoveChildEvaluationContext(TSharedRef<FCameraEvaluationContext> InContext);
+
+	void AddReferencedObjects(FReferenceCollector& Collector);
+
 public:
 
 	// Internal API.
-	void SetPrivateCameraDirector(TObjectPtr<const UCameraDirector> InCameraDirector);
+	void SetPrivateCameraDirector(const UCameraDirector* InCameraDirector);
 
 protected:
 
+	enum class EChildContextManipulationResult
+	{
+		Failure,
+		Success,
+		ChildContextSuccess
+	};
+
+	struct FChildContextManulationParams
+	{
+		TSharedPtr<FCameraEvaluationContext> ParentContext;
+		TSharedPtr<FCameraEvaluationContext> ChildContext;
+	};
+
+	struct FChildContextManulationResult
+	{
+		EChildContextManipulationResult Result = EChildContextManipulationResult::Failure;
+#if UE_GAMEPLAY_CAMERAS_DEBUG
+#endif  // UE_GAMEPLAY_CAMERAS_DEBUG
+	};
+
+	virtual void OnInitialize(const FCameraDirectorInitializeParams& Params) {}
+
 	/** Runs the camera director to determine what camera rig(s) should be active this frame. */
 	virtual void OnRun(const FCameraDirectorEvaluationParams& Params, FCameraDirectorEvaluationResult& OutResult) {}
 
+	virtual void OnAddChildEvaluationContext(const FChildContextManulationParams& Params, FChildContextManulationResult& Result) {}
+	virtual void OnRemoveChildEvaluationContext(const FChildContextManulationParams& Params, FChildContextManulationResult& Result) {}
+
+	virtual void OnAddReferencedObjects(FReferenceCollector& Collector) {}
+
 private:
 
+	TWeakPtr<FCameraEvaluationContext> WeakOwnerContext;
+
 	/** The camera director this evaluator is running. */
 	TObjectPtr<const UCameraDirector> PrivateCameraDirector;
 };
diff --git a/Engine/Plugins/Cameras/GameplayCameras/Source/GameplayCameras/Public/Core/CameraEvaluationContext.h b/Engine/Plugins/Cameras/GameplayCameras/Source/GameplayCameras/Public/Core/CameraEvaluationContext.h
--- a/Engine/Plugins/Cameras/GameplayCameras/Source/GameplayCameras/Public/Core/CameraEvaluationContext.h
+++ b/Engine/Plugins/Cameras/GameplayCameras/Source/GameplayCameras/Public/Core/CameraEvaluationContext.h
@@ -4,16 +4,32 @@
 
 #include "CoreTypes.h"
 #include "UObject/ObjectPtr.h"
+#include "Core/CameraDirectorEvaluator.h"
 #include "Core/CameraNodeEvaluator.h"
 #include "Core/CameraObjectRtti.h"
 #include "Templates/SharedPointer.h"
 
-class UCameraAsset;
 class APlayerController;
+class UCameraAsset;
+class UCameraDirector;
 
 namespace UE::Cameras
 {
 
+struct FCameraEvaluationContextInitializeParams
+{
+	TObjectPtr<const UCameraAsset> CameraAsset;
+	TObjectPtr<APlayerController> PlayerController;
+};
+
+struct FCameraEvaluationContextActivateParams
+{
+};
+
+struct FCameraEvaluationContextDeactivateParams
+{
+};
+
 /**
  * Base class for providing a context to running camera rigs.
  */
@@ -26,6 +42,8 @@ public:
 	/** Constructs an evaluation context. */
 	GAMEPLAYCAMERAS_API FCameraEvaluationContext();
 
+	GAMEPLAYCAMERAS_API void Initialize(const FCameraEvaluationContextInitializeParams& Params);
+
 	/** Destroys this evaluation context. */
 	GAMEPLAYCAMERAS_API virtual ~FCameraEvaluationContext();
 
@@ -36,7 +54,7 @@ public:
 	APlayerController* GetPlayerController() const { return PlayerController; }
 
 	/** Gets the camera asset that is hosted in this context. */
-	UCameraAsset* GetCameraAsset() const { return CameraAsset; }
+	const UCameraAsset* GetCameraAsset() const { return CameraAsset; }
 
 	/** Gets the initial evaluation result for all camera rigs in this context. */
 	const FCameraNodeEvaluationResult& GetInitialResult() const { return InitialResult; }
@@ -44,10 +62,30 @@ public:
 	/** Gets the initial evaluation result for all camera rigs in this context. */
 	FCameraNodeEvaluationResult& GetInitialResult() { return InitialResult; }
 
+	FCameraDirectorEvaluator* GetDirectorEvaluator() const { return DirectorEvaluator; }
+
+	TArrayView<const TSharedPtr<FCameraEvaluationContext>> GetChildrenContexts() const { return ChildrenContexts; }
+
+public:
+
+	void Activate(const FCameraEvaluationContextActivateParams& Params);
+	void Deactivate(const FCameraEvaluationContextDeactivateParams& Params);
+
 public:
 
+	// Internal API.
 	void AddReferencedObjects(FReferenceCollector& Collector);
 
+	bool RegisterChildContext(TSharedRef<FCameraEvaluationContext> ChildContext);
+	bool UnregisterChildContext(TSharedRef<FCameraEvaluationContext> ChildContext);
+
+protected:
+
+	virtual void OnActivate(const FCameraEvaluationContextActivateParams& Params) {}
+	virtual void OnDeactivate(const FCameraEvaluationContextDeactivateParams& Params) {}
+
+	void AutoCreateDirectorEvaluator();
+
 protected:
 
 	/**
@@ -57,10 +95,23 @@ protected:
 	TObjectPtr<APlayerController> PlayerController;
 
 	/** The camera asset hosted in this context. */
-	TObjectPtr<UCameraAsset> CameraAsset;
+	TObjectPtr<const UCameraAsset> CameraAsset;
 
 	/** The initial result for all camera rigs in this context. */
 	FCameraNodeEvaluationResult InitialResult;
+
+private:
+
+	FCameraDirectorEvaluatorStorage DirectorEvaluatorStorage;
+	FCameraDirectorEvaluator* DirectorEvaluator = nullptr;
+
+	TWeakPtr<FCameraEvaluationContext> WeakParent;
+
+	using FChildrenContexts = TArray<TSharedPtr<FCameraEvaluationContext>>;
+	FChildrenContexts ChildrenContexts;
+
+	bool bInitialized = false;
+	bool bActivated = false;
 };
 
 }  // namespace UE::Cameras
diff --git a/Engine/Plugins/Cameras/GameplayCameras/Source/GameplayCameras/Public/Core/CameraEvaluationContextStack.h b/Engine/Plugins/Cameras/GameplayCameras/Source/GameplayCameras/Public/Core/CameraEvaluationContextStack.h
--- a/Engine/Plugins/Cameras/GameplayCameras/Source/GameplayCameras/Public/Core/CameraEvaluationContextStack.h
+++ b/Engine/Plugins/Cameras/GameplayCameras/Source/GameplayCameras/Public/Core/CameraEvaluationContextStack.h
@@ -3,7 +3,6 @@
 #pragma once
 
 #include "Containers/Array.h"
-#include "Core/CameraDirectorEvaluator.h"
 #include "GameplayCameras.h"
 #include "UObject/WeakObjectPtr.h"
 
@@ -17,22 +16,6 @@ class FCameraDirectorEvaluator;
 class FCameraEvaluationContext;
 class FCameraSystemEvaluator;
 
-/** Information about a running camera evaluation context. */
-struct FCameraEvaluationContextInfo
-{
-	/** The evaluation context. */
-	TSharedPtr<FCameraEvaluationContext> EvaluationContext;
-
-	/** The instantiated camera director running in this context. */
-	TObjectPtr<const UCameraDirector> CameraDirector;
-
-	/** The evaluator for the running camera director. */
-	FCameraDirectorEvaluator* Evaluator = nullptr;
-
-	/** Returns whether this structure has a valid context and director. */
-	bool IsValid() const { return EvaluationContext && CameraDirector; }
-};
-
 /**
  * A simple stack of evaluation contexts. The top one is the active one.
  */
@@ -41,7 +24,7 @@ struct FCameraEvaluationContextStack
 public:
 
 	/** Gets the active (top) context. */
-	FCameraEvaluationContextInfo GetActiveContext() const;
+	TSharedPtr<FCameraEvaluationContext> GetActiveContext() const;
 
 	/** Returns whether the given context exists in the stack. */
 	bool HasContext(TSharedRef<FCameraEvaluationContext> Context) const;
@@ -49,6 +32,8 @@ public:
 	/** Push a new context on the stack and instantiate its director. */
 	void PushContext(TSharedRef<FCameraEvaluationContext> Context);
 
+	bool AddChildContext(TSharedRef<FCameraEvaluationContext> Context);
+
 	/** Remove an existing context from the stack. */
 	bool RemoveContext(TSharedRef<FCameraEvaluationContext> Context);
 
@@ -69,9 +54,6 @@ private:
 	struct FContextEntry
 	{
 		TWeakPtr<FCameraEvaluationContext> WeakContext;
-		TObjectPtr<const UCameraDirector> CameraDirector;
-		FCameraDirectorEvaluatorStorage EvaluatorStorage;
-		FCameraDirectorEvaluator* Evaluator = nullptr;
 	};
 
 	/** The entries in the stack. */
diff --git a/Engine/Plugins/Cameras/GameplayCameras/Source/GameplayCameras/Public/Directors/SingleCameraDirector.h b/Engine/Plugins/Cameras/GameplayCameras/Source/GameplayCameras/Public/Directors/SingleCameraDirector.h
--- a/Engine/Plugins/Cameras/GameplayCameras/Source/GameplayCameras/Public/Directors/SingleCameraDirector.h
+++ b/Engine/Plugins/Cameras/GameplayCameras/Source/GameplayCameras/Public/Directors/SingleCameraDirector.h
@@ -9,7 +9,7 @@
 /**
  * A simple camera director that only ever returns one single camera rig.
  */
-UCLASS(EditInlineNew)
+UCLASS(MinimalAPI, EditInlineNew)
 class USingleCameraDirector : public UCameraDirector
 {
 	GENERATED_BODY()
