Cameras: add owner object to evaluation contexts, better handle errors in activation.

diff --git a/Engine/Plugins/Cameras/GameplayCameras/Source/GameplayCameras/Private/Core/CameraEvaluationContext.cpp b/Engine/Plugins/Cameras/GameplayCameras/Source/GameplayCameras/Private/Core/CameraEvaluationContext.cpp
--- a/Engine/Plugins/Cameras/GameplayCameras/Source/GameplayCameras/Private/Core/CameraEvaluationContext.cpp
+++ b/Engine/Plugins/Cameras/GameplayCameras/Source/GameplayCameras/Private/Core/CameraEvaluationContext.cpp
@@ -4,6 +4,7 @@
 
 #include "Core/CameraAsset.h"
 #include "Core/CameraDirector.h"
+#include "GameFramework/PlayerController.h"
 
 namespace UE::Cameras
 {
@@ -14,6 +15,11 @@ FCameraEvaluationContext::FCameraEvaluationContext()
 {
 }
 
+FCameraEvaluationContext::FCameraEvaluationContext(const FCameraEvaluationContextInitializeParams& Params)
+{
+	Initialize(Params);
+}
+
 void FCameraEvaluationContext::Initialize(const FCameraEvaluationContextInitializeParams& Params)
 {
 	if (!ensureMsgf(!bInitialized, TEXT("This evaluation context has already been initialized!")))
@@ -21,8 +27,9 @@ void FCameraEvaluationContext::Initialize(const FCameraEvaluationContextInitiali
 		return;
 	}
 
+	WeakOwner = Params.Owner;
 	CameraAsset = Params.CameraAsset;
-	PlayerController = Params.PlayerController;
+	WeakPlayerController = Params.PlayerController;
 
 	bInitialized = true;
 }
@@ -52,20 +59,25 @@ void FCameraEvaluationContext::AutoCreateDirectorEvaluator()
 {
 	if (DirectorEvaluator == nullptr)
 	{
-		if (const UCameraDirector* CameraDirector = CameraAsset->CameraDirector)
+		if (!CameraAsset)
 		{
-			FCameraDirectorEvaluatorBuilder DirectorBuilder(DirectorEvaluatorStorage);
-			DirectorEvaluator = CameraDirector->BuildEvaluator(DirectorBuilder);
-
-			FCameraDirectorInitializeParams InitParams;
-			InitParams.OwnerContext = SharedThis(this);
-			InitParams.CameraDirector = CameraDirector;
-			DirectorEvaluator->Initialize(InitParams);
+			UE_LOG(LogCameraSystem, Error, TEXT("Activating an evaluation context without a camera!"));
+			return;
 		}
-		else
+		if (!CameraAsset->CameraDirector)
 		{
-			UE_LOG(LogCameraSystem, Warning, TEXT("Activating an evaluation context without a camera director!"));
+			UE_LOG(LogCameraSystem, Error, TEXT("Activating an evaluation context without a camera director!"));
+			return;
 		}
+
+		const UCameraDirector* CameraDirector = CameraAsset->CameraDirector;
+		FCameraDirectorEvaluatorBuilder DirectorBuilder(DirectorEvaluatorStorage);
+		DirectorEvaluator = CameraDirector->BuildEvaluator(DirectorBuilder);
+
+		FCameraDirectorInitializeParams InitParams;
+		InitParams.OwnerContext = SharedThis(this);
+		InitParams.CameraDirector = CameraDirector;
+		DirectorEvaluator->Initialize(InitParams);
 	}
 }
 
diff --git a/Engine/Plugins/Cameras/GameplayCameras/Source/GameplayCameras/Private/GameFramework/GameplayCameraComponent.cpp b/Engine/Plugins/Cameras/GameplayCameras/Source/GameplayCameras/Private/GameFramework/GameplayCameraComponent.cpp
--- a/Engine/Plugins/Cameras/GameplayCameras/Source/GameplayCameras/Private/GameFramework/GameplayCameraComponent.cpp
+++ b/Engine/Plugins/Cameras/GameplayCameras/Source/GameplayCameras/Private/GameFramework/GameplayCameraComponent.cpp
@@ -105,7 +105,12 @@ void UGameplayCameraComponent::ActivateCamera(APlayerController* PlayerControlle
 	if (!EvaluationContext.IsValid())
 	{
 		EvaluationContext = MakeShared<FGameplayCameraComponentEvaluationContext>();
-		EvaluationContext->Initialize(this, PlayerController);
+
+		FCameraEvaluationContextInitializeParams InitParams;
+		InitParams.Owner = this;
+		InitParams.CameraAsset = Camera;
+		InitParams.PlayerController = PlayerController;
+		EvaluationContext->Initialize(InitParams);
 
 		InitialResultInterop->Setup(&EvaluationContext->GetInitialResult());
 	}
@@ -209,12 +214,6 @@ namespace UE::Cameras
 
 UE_DEFINE_CAMERA_EVALUATION_CONTEXT(FGameplayCameraComponentEvaluationContext)
 
-void FGameplayCameraComponentEvaluationContext::Initialize(UGameplayCameraComponent* Owner, APlayerController* InPlayerController)
-{
-	PlayerController = InPlayerController;
-	CameraAsset = Owner->Camera;
-}
-
 void FGameplayCameraComponentEvaluationContext::Update(UGameplayCameraComponent* Owner)
 {
 	const FTransform& OwnerTransform = Owner->GetComponentTransform();
diff --git a/Engine/Plugins/Cameras/GameplayCameras/Source/GameplayCameras/Public/Core/CameraEvaluationContext.h b/Engine/Plugins/Cameras/GameplayCameras/Source/GameplayCameras/Public/Core/CameraEvaluationContext.h
--- a/Engine/Plugins/Cameras/GameplayCameras/Source/GameplayCameras/Public/Core/CameraEvaluationContext.h
+++ b/Engine/Plugins/Cameras/GameplayCameras/Source/GameplayCameras/Public/Core/CameraEvaluationContext.h
@@ -16,16 +16,31 @@ class UCameraDirector;
 namespace UE::Cameras
 {
 
+/**
+ * Parameter struct for initializing an evaluation context.
+ */
 struct FCameraEvaluationContextInitializeParams
 {
+	/** The owner of the context, if any. */
+	TObjectPtr<UObject> Owner;
+	/** The camera asset to run inside the context, if any. */
 	TObjectPtr<const UCameraAsset> CameraAsset;
+	/** The player controller driving this context, if any. */
 	TObjectPtr<APlayerController> PlayerController;
 };
 
+/**
+ * Parameter struct for activating an evaluation context, which happens when it is added
+ * to the camera system's context stack.
+ */
 struct FCameraEvaluationContextActivateParams
 {
 };
 
+/**
+ * Parameter struct for deactivating an evaluation context, which happens when it is removed
+ * from the camera system's context stack.
+ */
 struct FCameraEvaluationContextDeactivateParams
 {
 };
@@ -41,17 +56,23 @@ public:
 
 	/** Constructs an evaluation context. */
 	GAMEPLAYCAMERAS_API FCameraEvaluationContext();
+	/** Constructs an evaluation context. */
+	GAMEPLAYCAMERAS_API FCameraEvaluationContext(const FCameraEvaluationContextInitializeParams& Params);
 
+	/** Initializes the evaluation context, if it was created with the default constructor. */
 	GAMEPLAYCAMERAS_API void Initialize(const FCameraEvaluationContextInitializeParams& Params);
 
 	/** Destroys this evaluation context. */
 	GAMEPLAYCAMERAS_API virtual ~FCameraEvaluationContext();
 
+	/** Gets the owner of this evaluation context, if any, and if still valid. */
+	UObject* GetOwner() const { return WeakOwner.Get(); }
+
 	/**
 	 * Gets the player controller (if any) in control of the cameras running inside
 	 * of this evaluation context.
 	 */
-	APlayerController* GetPlayerController() const { return PlayerController; }
+	APlayerController* GetPlayerController() const { return WeakPlayerController.Get(); }
 
 	/** Gets the camera asset that is hosted in this context. */
 	const UCameraAsset* GetCameraAsset() const { return CameraAsset; }
@@ -62,13 +83,23 @@ public:
 	/** Gets the initial evaluation result for all camera rigs in this context. */
 	FCameraNodeEvaluationResult& GetInitialResult() { return InitialResult; }
 
+	/** Gets the camera director evaluator that runs the camera rigs of this context. */
 	FCameraDirectorEvaluator* GetDirectorEvaluator() const { return DirectorEvaluator; }
 
+	/** Gets the children evaluation contexts running inside this context. */
 	TArrayView<const TSharedPtr<FCameraEvaluationContext>> GetChildrenContexts() const { return ChildrenContexts; }
 
 public:
 
+	/**
+	 * Activates this evaluation context.
+	 * This will create the camera director evaluator if necessary.
+	 */
 	void Activate(const FCameraEvaluationContextActivateParams& Params);
+
+	/**
+	 * Deactivates this evaluation context.
+	 */
 	void Deactivate(const FCameraEvaluationContextDeactivateParams& Params);
 
 public:
@@ -88,11 +119,14 @@ protected:
 
 protected:
 
+	/** The owner of this context, if any. */
+	TWeakObjectPtr<> WeakOwner;
+
 	/**
 	 * The player controller (if any) in control of the cameras running inside
 	 * of this evaluation context.
 	 */
-	TObjectPtr<APlayerController> PlayerController;
+	TWeakObjectPtr<APlayerController> WeakPlayerController;
 
 	/** The camera asset hosted in this context. */
 	TObjectPtr<const UCameraAsset> CameraAsset;
diff --git a/Engine/Plugins/Cameras/GameplayCameras/Source/GameplayCameras/Public/Core/CameraSystemEvaluator.h b/Engine/Plugins/Cameras/GameplayCameras/Source/GameplayCameras/Public/Core/CameraSystemEvaluator.h
--- a/Engine/Plugins/Cameras/GameplayCameras/Source/GameplayCameras/Public/Core/CameraSystemEvaluator.h
+++ b/Engine/Plugins/Cameras/GameplayCameras/Source/GameplayCameras/Public/Core/CameraSystemEvaluator.h
@@ -35,8 +35,10 @@ class FRootCameraDebugBlock;
  */
 struct GAMEPLAYCAMERAS_API FCameraSystemEvaluatorCreateParams
 {
+	/** The owner of the camera system, if any. */
 	TObjectPtr<UObject> Owner;
 
+	/** An optional factory for creating the root node. */
 	using FRootNodeFactory = TFunction<URootCameraNode*()>;
 	FRootNodeFactory RootNodeFactory;
 };
@@ -87,6 +89,11 @@ public:
 	/** Initializes the camera system. */
 	GAMEPLAYCAMERAS_API void Initialize(const FCameraSystemEvaluatorCreateParams& Params);
 
+public:
+
+	/** Gets the owner of this camera system, if any, and if still valid. */
+	UObject* GetOwner() const { return WeakOwner.Get(); }
+
 public:
 
 	/** Push a new evaluation context on the stack. */
diff --git a/Engine/Plugins/Cameras/GameplayCameras/Source/GameplayCameras/Public/GameFramework/GameplayCameraComponent.h b/Engine/Plugins/Cameras/GameplayCameras/Source/GameplayCameras/Public/GameFramework/GameplayCameraComponent.h
--- a/Engine/Plugins/Cameras/GameplayCameras/Source/GameplayCameras/Public/GameFramework/GameplayCameraComponent.h
+++ b/Engine/Plugins/Cameras/GameplayCameras/Source/GameplayCameras/Public/GameFramework/GameplayCameraComponent.h
@@ -100,7 +100,6 @@ class FGameplayCameraComponentEvaluationContext : public FCameraEvaluationContex
 
 public:
 
-	void Initialize(UGameplayCameraComponent* Owner, APlayerController* InPlayerController);
 	void Update(UGameplayCameraComponent* Owner);
 };
 
