Cameras: new APIs on the camera pose structure

diff --git a/Engine/Plugins/Cameras/GameplayCameras/Source/GameplayCameras/Private/Core/CameraPose.cpp b/Engine/Plugins/Cameras/GameplayCameras/Source/GameplayCameras/Private/Core/CameraPose.cpp
--- a/Engine/Plugins/Cameras/GameplayCameras/Source/GameplayCameras/Private/Core/CameraPose.cpp
+++ b/Engine/Plugins/Cameras/GameplayCameras/Source/GameplayCameras/Private/Core/CameraPose.cpp
@@ -122,6 +122,7 @@ float FCameraPose::GetEffectiveFieldOfView() const
 	{
 		// Compute FOV with similar code to UCineCameraComponent...
 		float CropedSensorWidth = SensorWidth * SqueezeFactor;
+		const float AspectRatio = GetSensorAspectRatio();
 		if (AspectRatio > 0.0f)
 		{
 			float DesqueezeAspectRatio = SensorWidth * SqueezeFactor / SensorHeight;
@@ -140,6 +141,25 @@ float FCameraPose::GetEffectiveFieldOfView() const
 	}
 }
 
+FCameraFieldsOfView FCameraPose::GetEffectiveFieldsOfView() const
+{
+	FCameraFieldsOfView FOVs;
+
+	FOVs.HorizontalFieldOfView = GetEffectiveFieldOfView();
+
+	const float AspectRatio = GetSensorAspectRatio();
+	FOVs.VerticalFieldOfView = FMath::RadiansToDegrees(
+			2.f * FMath::Atan(
+				FMath::Tan(FMath::DegreesToRadians(FOVs.HorizontalFieldOfView / 2.f)) / AspectRatio));
+
+	return FOVs;
+}
+
+float FCameraPose::GetSensorAspectRatio() const
+{
+	return (SensorHeight > 0.f) ? (SensorWidth / SensorHeight) : 0.f;
+}
+
 FRay3d FCameraPose::GetAimRay() const
 {
 	const bool bDirectionIsNormalized = false;
@@ -147,6 +167,11 @@ FRay3d FCameraPose::GetAimRay() const
 	return FRay3d(Location, Rotation.RotateVector(TargetDir), bDirectionIsNormalized);
 }
 
+FVector3d FCameraPose::GetAimDir() const
+{
+	return Rotation.RotateVector(FVector3d{ 1, 0, 0 });
+}
+
 void FCameraPose::OverrideAll(const FCameraPose& OtherPose)
 {
 	InternalOverrideChanged(OtherPose, false);
diff --git a/Engine/Plugins/Cameras/GameplayCameras/Source/GameplayCameras/Public/Core/CameraPose.h b/Engine/Plugins/Cameras/GameplayCameras/Source/GameplayCameras/Public/Core/CameraPose.h
--- a/Engine/Plugins/Cameras/GameplayCameras/Source/GameplayCameras/Public/Core/CameraPose.h
+++ b/Engine/Plugins/Cameras/GameplayCameras/Source/GameplayCameras/Public/Core/CameraPose.h
@@ -21,7 +21,6 @@ class FArchive;
 	UE_CAMERA_POSE_FOR_PROPERTY(float,  SensorWidth)\
 	UE_CAMERA_POSE_FOR_PROPERTY(float,  SensorHeight)\
 	UE_CAMERA_POSE_FOR_PROPERTY(float,  SqueezeFactor)\
-	UE_CAMERA_POSE_FOR_PROPERTY(float,  AspectRatio)\
 	UE_CAMERA_POSE_FOR_PROPERTY(float,  NearClippingPlane)\
 	UE_CAMERA_POSE_FOR_PROPERTY(float,  FarClippingPlane)
 
@@ -71,6 +70,15 @@ public:
 	FCameraPoseFlags& OR(const FCameraPoseFlags& OtherFlags);
 };
 
+/**
+ * Simple struct for holding horizontal and vertical fields of view.
+ */
+struct FCameraFieldsOfView
+{
+	float HorizontalFieldOfView;
+	float VerticalFieldOfView;
+};
+
 /**
  * Structure describing the state of a camera.
  *
@@ -98,10 +106,13 @@ public:
 	{\
 		return PropName;\
 	}\
-	void Set##PropName(TCallTraits<PropType>::ParamType InValue)\
+	void Set##PropName(TCallTraits<PropType>::ParamType InValue, bool bForceSet = false)\
 	{\
-		ChangedFlags.PropName = true;\
-		PropName = InValue;\
+		if (bForceSet || PropName != InValue)\
+		{\
+			ChangedFlags.PropName = true;\
+			PropName = InValue;\
+		}\
 	}
 
 UE_CAMERA_POSE_FOR_ALL_PROPERTIES()
@@ -135,17 +146,24 @@ public:
 	void SetTransform(FTransform3d Transform);
 
 	/**
-	 * Computes the field of view of the camera.
+	 * Computes the horizontal field of view of the camera.
 	 * The effective field of view can be driven by the FieldOfView property, or
 	 * the FocalLength property in combination with the sensor size.
 	 */
 	float GetEffectiveFieldOfView() const;
 
-	/**
-	 * Gets the aiming ray of the camera.
-	 */
+	/** Gets both horizontal and vertical effective fields of view.  */
+	FCameraFieldsOfView GetEffectiveFieldsOfView() const;
+
+	/** Gets the aspect ratio of the camera sensor. */
+	float GetSensorAspectRatio() const;
+
+	/** Gets the aiming ray of the camera. */
 	FRay3d GetAimRay() const;
 
+	/** Gets the aiming direction of the camera. */
+	FVector3d GetAimDir() const;
+
 public:
 
 	// Interpolation
@@ -191,13 +209,6 @@ private:
 	UPROPERTY()
 	float FieldOfView = -1.f;  // Default to using a focal length
 
-	/**
-	 * The aspect ratio of the camera
-	 * If zero or less, the sensor width and height are used instead
-	 */
-	UPROPERTY()
-	float AspectRatio = -1.f;
-
 	/**
 	 * The focal length of the camera's lens, in millimeters
 	 * If zero or less, field of view is used instead
