Cameras: add priority queue camera director

This director doesn't do anything per-se: it waits until sub-directors are inserted into it, and then runs the one with the higher priority.

diff --git a/Engine/Plugins/Cameras/GameplayCameras/Source/GameplayCameras/Private/Directors/PriorityQueueCameraDirector.cpp b/Engine/Plugins/Cameras/GameplayCameras/Source/GameplayCameras/Private/Directors/PriorityQueueCameraDirector.cpp
new file mode 100644
--- /dev/null
+++ b/Engine/Plugins/Cameras/GameplayCameras/Source/GameplayCameras/Private/Directors/PriorityQueueCameraDirector.cpp
@@ -0,0 +1,98 @@
+// Copyright Epic Games, Inc. All Rights Reserved.
+
+#include "Directors/PriorityQueueCameraDirector.h"
+
+#include "Algo/StableSort.h"
+#include "Core/CameraEvaluationContext.h"
+
+#include UE_INLINE_GENERATED_CPP_BY_NAME(PriorityQueueCameraDirector)
+
+namespace UE::Cameras
+{
+
+int32 FPriorityQueueCameraDirectorEvaluator::FEntry::GetPriority() const
+{
+	switch (PriorityGiver.GetIndex())
+	{
+		case FPriorityGiver::IndexOfType<int32>():
+			return PriorityGiver.Get<int32>();
+		case FPriorityGiver::IndexOfType<IPriorityQueueEntry*>():
+			return PriorityGiver.Get<IPriorityQueueEntry*>()->GetPriority();
+		default:
+			return 0;
+	}
+}
+
+UE_DEFINE_CAMERA_DIRECTOR_EVALUATOR(FPriorityQueueCameraDirectorEvaluator)
+
+void FPriorityQueueCameraDirectorEvaluator::AddChildEvaluationContext(TSharedRef<FCameraEvaluationContext> InContext, int32 Priority)
+{
+	Super::AddChildEvaluationContext(InContext);
+	Entries.Last().PriorityGiver = FPriorityGiver(TInPlaceType<int32>(), Priority);
+}
+
+void FPriorityQueueCameraDirectorEvaluator::AddChildEvaluationContext(TSharedRef<FCameraEvaluationContext> InContext, IPriorityQueueEntry* PriorityEntry)
+{
+	Super::AddChildEvaluationContext(InContext);
+	Entries.Last().PriorityGiver = FPriorityGiver(TInPlaceType<IPriorityQueueEntry*>(), PriorityEntry);
+}
+
+void FPriorityQueueCameraDirectorEvaluator::OnAddChildEvaluationContext(const FChildContextManulationParams& Params, FChildContextManulationResult& Result)
+{
+	Entries.Add(FEntry{ Params.ChildContext, FPriorityGiver(TInPlaceType<int32>(), 0) });
+	Result.Result = EChildContextManipulationResult::Success;
+}
+
+void FPriorityQueueCameraDirectorEvaluator::OnRemoveChildEvaluationContext(const FChildContextManulationParams& Params, FChildContextManulationResult& Result)
+{
+	Result.Result = EChildContextManipulationResult::Failure;
+
+	for (auto It = Entries.CreateIterator(); It; ++It)
+	{
+		if (It->ChildContext == Params.ChildContext)
+		{
+			It.RemoveCurrent();
+			Result.Result = EChildContextManipulationResult::Success;
+			break;
+		}
+	}
+}
+
+void FPriorityQueueCameraDirectorEvaluator::OnRun(const FCameraDirectorEvaluationParams& Params, FCameraDirectorEvaluationResult& OutResult)
+{
+	if (Entries.Num() > 0)
+	{
+		FEntryArray EntriesCopy(Entries);
+		Algo::StableSortBy(EntriesCopy, &FEntry::GetPriority);
+
+		const FEntry& HighestPriorityEntry = EntriesCopy.Last();
+		if (ensure(HighestPriorityEntry.ChildContext))
+		{
+			if (FCameraDirectorEvaluator* DirectorEvaluator = HighestPriorityEntry.ChildContext->GetDirectorEvaluator())
+			{
+				FCameraDirectorEvaluationParams ChildParams(Params);
+				ChildParams.OwnerContext = HighestPriorityEntry.ChildContext;
+				DirectorEvaluator->Run(ChildParams, OutResult);
+
+				if (OutResult.ActiveCameraRigs.Num() > 0)
+				{
+					return;
+				}
+			}
+		}
+	}
+}
+
+}  // namespace UE::Cameras
+
+UPriorityQueueCameraDirector::UPriorityQueueCameraDirector(const FObjectInitializer& ObjectInit)
+	: Super(ObjectInit)
+{
+}
+
+FCameraDirectorEvaluatorPtr UPriorityQueueCameraDirector::OnBuildEvaluator(FCameraDirectorEvaluatorBuilder& Builder) const
+{
+	using namespace UE::Cameras;
+	return Builder.BuildEvaluator<FPriorityQueueCameraDirectorEvaluator>();
+}
+
diff --git a/Engine/Plugins/Cameras/GameplayCameras/Source/GameplayCameras/Public/Directors/PriorityQueueCameraDirector.h b/Engine/Plugins/Cameras/GameplayCameras/Source/GameplayCameras/Public/Directors/PriorityQueueCameraDirector.h
new file mode 100644
--- /dev/null
+++ b/Engine/Plugins/Cameras/GameplayCameras/Source/GameplayCameras/Public/Directors/PriorityQueueCameraDirector.h
@@ -0,0 +1,94 @@
+// Copyright Epic Games, Inc. All Rights Reserved.
+
+#pragma once
+
+#include "Core/CameraDirector.h"
+#include "Core/CameraDirectorEvaluator.h"
+#include "Misc/TVariant.h"
+
+#include "PriorityQueueCameraDirector.generated.h"
+
+/**
+ * A camera director that holds multiple sub-directors, and runs the one that has the highest priority.
+ */
+UCLASS(MinimalAPI, EditInlineNew)
+class UPriorityQueueCameraDirector : public UCameraDirector
+{
+	GENERATED_BODY()
+
+public:
+
+	UPriorityQueueCameraDirector(const FObjectInitializer& ObjectInit);
+
+protected:
+
+	// UCameraDirector interface.
+	virtual FCameraDirectorEvaluatorPtr OnBuildEvaluator(FCameraDirectorEvaluatorBuilder& Builder) const override;
+};
+
+namespace UE::Cameras
+{
+
+/**
+ * Interface for sub-directors with dynamic priorities.
+ */
+class IPriorityQueueEntry
+{
+public:
+
+	virtual ~IPriorityQueueEntry() {}
+
+	/** Gets the current priority of the sub-director. */
+	virtual int32 GetPriority() = 0;
+};
+
+/**
+ * Evaluator for the priority queue camera director.
+ *
+ * Use the overriden AddChildEvaluationContext methods here to add sub-directors with a specific priority. Using the base class
+ * version will add a sub-director with a priority of 0.
+ *
+ * Highest priority values mean more priority, which means higher chances of being picked from the queue.
+ */
+class FPriorityQueueCameraDirectorEvaluator : public FCameraDirectorEvaluator
+{
+	UE_DECLARE_CAMERA_DIRECTOR_EVALUATOR(GAMEPLAYCAMERAS_API, FPriorityQueueCameraDirectorEvaluator)
+
+public:
+
+	/**
+	 * Adds a sub-director with the given fixed priority.
+	 */
+	GAMEPLAYCAMERAS_API void AddChildEvaluationContext(TSharedRef<FCameraEvaluationContext> InContext, int32 Priority);
+
+	/**
+	 * Adds a sub-director with the given dynamic priority.
+	 * It's customary to have the evaluation context itself implement the IPriorityQueueEntry interface, and therefore passing
+	 * the pointer twice to the function.
+	 */
+	GAMEPLAYCAMERAS_API void AddChildEvaluationContext(TSharedRef<FCameraEvaluationContext> InContext, IPriorityQueueEntry* PriorityEntry);
+
+protected:
+
+	virtual void OnRun(const FCameraDirectorEvaluationParams& Params, FCameraDirectorEvaluationResult& OutResult) override;
+
+	virtual void OnAddChildEvaluationContext(const FChildContextManulationParams& Params, FChildContextManulationResult& Result) override;
+	virtual void OnRemoveChildEvaluationContext(const FChildContextManulationParams& Params, FChildContextManulationResult& Result) override;
+
+private:
+
+	using FPriorityGiver = TVariant<int32, IPriorityQueueEntry*>;
+	struct FEntry
+	{
+		TSharedPtr<FCameraEvaluationContext> ChildContext;
+		FPriorityGiver PriorityGiver;
+
+		int32 GetPriority() const;
+	};
+
+	using FEntryArray = TArray<FEntry>;
+	FEntryArray Entries;
+};
+
+}  // namespace UE::Cameras
+
