Cameras: rename evaluation result interop BP node to a shorter/simpler name.

diff --git a/Engine/Plugins/Cameras/GameplayCameras/Source/GameplayCameras/Private/GameFramework/CameraEvaluationResultInterop.cpp b/Engine/Plugins/Cameras/GameplayCameras/Source/GameplayCameras/Private/GameFramework/CameraEvaluationResultInterop.cpp
new file mode 100644
--- /dev/null
+++ b/Engine/Plugins/Cameras/GameplayCameras/Source/GameplayCameras/Private/GameFramework/CameraEvaluationResultInterop.cpp
@@ -0,0 +1,155 @@
+// Copyright Epic Games, Inc. All Rights Reserved.
+
+#include "GameFramework/CameraEvaluationResultInterop.h"
+
+#include "Core/CameraNodeEvaluator.h"
+#include "Core/CameraPose.h"
+#include "Core/CameraVariableAssets.h"
+#include "Core/CameraVariableTable.h"
+
+#include UE_INLINE_GENERATED_CPP_BY_NAME(CameraEvaluationResultInterop)
+
+#define UE_PRIVATE_CAMERA_EVALUATION_RESULT_BLUEPRINT_INTEROP_VALIDATE_RESULT(ErrorResult)\
+	if (!Result)\
+	{\
+		FFrame::KismetExecutionMessage(TEXT("No camera evaluation result has been set"), ELogVerbosity::Error);\
+		return ErrorResult;\
+	}
+
+#define UE_PRIVATE_CAMERA_EVALUATION_RESULT_BLUEPRINT_INTEROP_VALIDATE_VARIABLE_ASSET_PARAM(ErrorResult)\
+	if (!InVariableAsset)\
+	{\
+		FFrame::KismetExecutionMessage(TEXT("No camera variable asset was given"), ELogVerbosity::Error);\
+		return ErrorResult;\
+	}\
+
+#define UE_PRIVATE_CAMERA_EVALUATION_RESULT_BLUEPRINT_INTEROP_GET_VARIABLE(VariableType)\
+	static VariableType ErrorResult {};\
+	UE_PRIVATE_CAMERA_EVALUATION_RESULT_BLUEPRINT_INTEROP_VALIDATE_RESULT(ErrorResult)\
+	UE_PRIVATE_CAMERA_EVALUATION_RESULT_BLUEPRINT_INTEROP_VALIDATE_VARIABLE_ASSET_PARAM(ErrorResult)\
+	return Result->VariableTable.GetValue<VariableType>(InVariableAsset->GetVariableID());
+
+#define UE_PRIVATE_CAMERA_EVALUATION_RESULT_BLUEPRINT_INTEROP_SET_VARIABLE(VariableType)\
+	UE_PRIVATE_CAMERA_EVALUATION_RESULT_BLUEPRINT_INTEROP_VALIDATE_RESULT()\
+	UE_PRIVATE_CAMERA_EVALUATION_RESULT_BLUEPRINT_INTEROP_VALIDATE_VARIABLE_ASSET_PARAM()\
+	Result->VariableTable.SetValue(InVariableAsset, InValue, true);
+
+UCameraEvaluationResultInterop::UCameraEvaluationResultInterop(const FObjectInitializer& ObjectInit)
+	: Super(ObjectInit)
+{
+}
+
+void UCameraEvaluationResultInterop::Setup(FCameraNodeEvaluationResult* InResult)
+{
+	Result = InResult;
+}
+
+void UCameraEvaluationResultInterop::Teardown()
+{
+	Result = nullptr;
+}
+
+const FCameraPose& UCameraEvaluationResultInterop::GetCameraPose() const
+{
+	static FCameraPose ErrorResult;
+	UE_PRIVATE_CAMERA_EVALUATION_RESULT_BLUEPRINT_INTEROP_VALIDATE_RESULT(ErrorResult);
+	return Result->CameraPose;
+}
+
+bool UCameraEvaluationResultInterop::GetBooleanCameraVariable(UBooleanCameraVariable* InVariableAsset) const
+{
+	UE_PRIVATE_CAMERA_EVALUATION_RESULT_BLUEPRINT_INTEROP_GET_VARIABLE(bool);
+}
+
+int32 UCameraEvaluationResultInterop::GetInteger32CameraVariable(UInteger32CameraVariable* InVariableAsset) const
+{
+	UE_PRIVATE_CAMERA_EVALUATION_RESULT_BLUEPRINT_INTEROP_GET_VARIABLE(int32);
+}
+
+float UCameraEvaluationResultInterop::GetFloatCameraVariable(UFloatCameraVariable* InVariableAsset) const
+{
+	UE_PRIVATE_CAMERA_EVALUATION_RESULT_BLUEPRINT_INTEROP_GET_VARIABLE(float);
+}
+
+double UCameraEvaluationResultInterop::GetDoubleCameraVariable(UDoubleCameraVariable* InVariableAsset) const
+{
+	UE_PRIVATE_CAMERA_EVALUATION_RESULT_BLUEPRINT_INTEROP_GET_VARIABLE(double);
+}
+
+FVector2D UCameraEvaluationResultInterop::GetVector2CameraVariable(UVector2dCameraVariable* InVariableAsset) const
+{
+	UE_PRIVATE_CAMERA_EVALUATION_RESULT_BLUEPRINT_INTEROP_GET_VARIABLE(FVector2d);
+}
+
+FVector UCameraEvaluationResultInterop::GetVector3CameraVariable(UVector3dCameraVariable* InVariableAsset) const
+{
+	UE_PRIVATE_CAMERA_EVALUATION_RESULT_BLUEPRINT_INTEROP_GET_VARIABLE(FVector3d);
+}
+
+FVector4 UCameraEvaluationResultInterop::GetVector4CameraVariable(UVector4dCameraVariable* InVariableAsset) const
+{
+	UE_PRIVATE_CAMERA_EVALUATION_RESULT_BLUEPRINT_INTEROP_GET_VARIABLE(FVector4d);
+}
+
+FRotator UCameraEvaluationResultInterop::GetRotatorCameraVariable(URotator3dCameraVariable* InVariableAsset) const
+{
+	UE_PRIVATE_CAMERA_EVALUATION_RESULT_BLUEPRINT_INTEROP_GET_VARIABLE(FRotator3d);
+}
+
+FTransform UCameraEvaluationResultInterop::GetTransformCameraVariable(UTransform3dCameraVariable* InVariableAsset) const
+{
+	UE_PRIVATE_CAMERA_EVALUATION_RESULT_BLUEPRINT_INTEROP_GET_VARIABLE(FTransform3d);
+}
+
+void UCameraEvaluationResultInterop::SetCameraPose(const FCameraPose& InCameraPose)
+{
+	UE_PRIVATE_CAMERA_EVALUATION_RESULT_BLUEPRINT_INTEROP_VALIDATE_RESULT();
+	Result->CameraPose = InCameraPose;
+	// TODO: auto-set flags based on differences
+}
+
+void UCameraEvaluationResultInterop::SetBooleanCameraVariable(UBooleanCameraVariable* InVariableAsset, bool InValue)
+{
+	UE_PRIVATE_CAMERA_EVALUATION_RESULT_BLUEPRINT_INTEROP_SET_VARIABLE(bool);
+}
+
+void UCameraEvaluationResultInterop::SetInteger32CameraVariable(UInteger32CameraVariable* InVariableAsset, int32 InValue)
+{
+	UE_PRIVATE_CAMERA_EVALUATION_RESULT_BLUEPRINT_INTEROP_SET_VARIABLE(int32);
+}
+
+void UCameraEvaluationResultInterop::SetFloatCameraVariable(UFloatCameraVariable* InVariableAsset, float InValue)
+{
+	UE_PRIVATE_CAMERA_EVALUATION_RESULT_BLUEPRINT_INTEROP_SET_VARIABLE(float);
+}
+
+void UCameraEvaluationResultInterop::SetDoubleCameraVariable(UDoubleCameraVariable* InVariableAsset, double InValue)
+{
+	UE_PRIVATE_CAMERA_EVALUATION_RESULT_BLUEPRINT_INTEROP_SET_VARIABLE(double);
+}
+
+void UCameraEvaluationResultInterop::SetVector2CameraVariable(UVector2dCameraVariable* InVariableAsset, const FVector2D& InValue)
+{
+	UE_PRIVATE_CAMERA_EVALUATION_RESULT_BLUEPRINT_INTEROP_SET_VARIABLE(FVector2d);
+}
+
+void UCameraEvaluationResultInterop::SetVector3CameraVariable(UVector3dCameraVariable* InVariableAsset, const FVector& InValue)
+{
+	UE_PRIVATE_CAMERA_EVALUATION_RESULT_BLUEPRINT_INTEROP_SET_VARIABLE(FVector3d);
+}
+
+void UCameraEvaluationResultInterop::SetVector4CameraVariable(UVector4dCameraVariable* InVariableAsset, const FVector4& InValue)
+{
+	UE_PRIVATE_CAMERA_EVALUATION_RESULT_BLUEPRINT_INTEROP_SET_VARIABLE(FVector4d);
+}
+
+void UCameraEvaluationResultInterop::SetRotatorCameraVariable(URotator3dCameraVariable* InVariableAsset, const FRotator& InValue)
+{
+	UE_PRIVATE_CAMERA_EVALUATION_RESULT_BLUEPRINT_INTEROP_SET_VARIABLE(FRotator3d);
+}
+
+void UCameraEvaluationResultInterop::SetTransformCameraVariable(UTransform3dCameraVariable* InVariableAsset, const FTransform& InValue)
+{
+	UE_PRIVATE_CAMERA_EVALUATION_RESULT_BLUEPRINT_INTEROP_SET_VARIABLE(FTransform3d);
+}
+
diff --git a/Engine/Plugins/Cameras/GameplayCameras/Source/GameplayCameras/Private/GameFramework/GameplayCameraComponent.cpp b/Engine/Plugins/Cameras/GameplayCameras/Source/GameplayCameras/Private/GameFramework/GameplayCameraComponent.cpp
--- a/Engine/Plugins/Cameras/GameplayCameras/Source/GameplayCameras/Private/GameFramework/GameplayCameraComponent.cpp
+++ b/Engine/Plugins/Cameras/GameplayCameras/Source/GameplayCameras/Private/GameFramework/GameplayCameraComponent.cpp
@@ -7,7 +7,7 @@
 #include "Engine/StaticMesh.h"
 #include "Engine/World.h"
 #include "GameFramework/Controller.h"
-#include "GameFramework/GameplayCameraEvaluationResultBlueprintInterop.h"
+#include "GameFramework/CameraEvaluationResultInterop.h"
 #include "GameFramework/GameplayCameraSystemActor.h"
 #include "GameFramework/GameplayCameraSystemComponent.h"
 #include "Logging/MessageLog.h"
@@ -23,7 +23,7 @@ UGameplayCameraComponent::UGameplayCameraComponent(const FObjectInitializer& Obj
 {
 	PrimaryComponentTick.bCanEverTick = true;
 
-	InitialResultInterop = ObjectInit.CreateDefaultSubobject<UGameplayCameraEvaluationResultBlueprintInterop>(this, "InitialResultInterop");
+	InitialResultInterop = ObjectInit.CreateDefaultSubobject<UCameraEvaluationResultInterop>(this, "InitialResultInterop");
 
 #if WITH_EDITORONLY_DATA
 	if (GIsEditor && !IsRunningCommandlet())
@@ -145,7 +145,7 @@ void UGameplayCameraComponent::DeactivateCamera(APlayerController* PlayerControl
 	Deactivate();
 }
 
-UGameplayCameraEvaluationResultBlueprintInterop* UGameplayCameraComponent::GetInitialResult() const
+UCameraEvaluationResultInterop* UGameplayCameraComponent::GetInitialResult() const
 {
 	return InitialResultInterop;
 }
diff --git a/Engine/Plugins/Cameras/GameplayCameras/Source/GameplayCameras/Private/GameFramework/GameplayCameraEvaluationResultBlueprintInterop.cpp b/Engine/Plugins/Cameras/GameplayCameras/Source/GameplayCameras/Private/GameFramework/GameplayCameraEvaluationResultBlueprintInterop.cpp
deleted file mode 100644
diff --git a/Engine/Plugins/Cameras/GameplayCameras/Source/GameplayCameras/Public/GameFramework/CameraEvaluationResultInterop.h b/Engine/Plugins/Cameras/GameplayCameras/Source/GameplayCameras/Public/GameFramework/CameraEvaluationResultInterop.h
new file mode 100644
--- /dev/null
+++ b/Engine/Plugins/Cameras/GameplayCameras/Source/GameplayCameras/Public/GameFramework/CameraEvaluationResultInterop.h
@@ -0,0 +1,113 @@
+// Copyright Epic Games, Inc. All Rights Reserved.
+
+#pragma once
+
+#include "CoreTypes.h"
+#include "UObject/ObjectMacros.h"
+
+#include "CameraEvaluationResultInterop.generated.h"
+
+class UBooleanCameraVariable;
+class UDoubleCameraVariable;
+class UFloatCameraVariable;
+class UInteger32CameraVariable;
+class URotator3dCameraVariable;
+class UTransform3dCameraVariable;
+class UVector2dCameraVariable;
+class UVector3dCameraVariable;
+class UVector4dCameraVariable;
+struct FCameraPose;
+
+namespace UE::Cameras
+{
+
+struct FCameraNodeEvaluationResult;
+
+}  // namespace UE::Cameras
+
+UCLASS(BlueprintType, DisplayName="Camera Evaluation Result")
+class GAMEPLAYCAMERAS_API UCameraEvaluationResultInterop : public UObject
+{
+	GENERATED_BODY()
+
+public:
+
+	using FCameraNodeEvaluationResult = UE::Cameras::FCameraNodeEvaluationResult;
+
+	UCameraEvaluationResultInterop(const FObjectInitializer& ObjectInit);
+
+	void Setup(FCameraNodeEvaluationResult* InResult);
+	void Teardown();
+
+	FCameraNodeEvaluationResult* GetEvaluationResult() { return Result; }
+	const FCameraNodeEvaluationResult* GetEvaluationResult() const { return Result; }
+
+public:
+
+	UFUNCTION(BlueprintPure, Category=Camera)
+	const FCameraPose& GetCameraPose() const;
+
+	UFUNCTION(BlueprintPure, Category=Camera, meta=(CompactNodeTitle="Get Bool"))
+	bool GetBooleanCameraVariable(UBooleanCameraVariable* InVariableAsset) const;
+
+	UFUNCTION(BlueprintPure, Category=Camera, meta=(CompactNodeTitle="Get Int"))
+	int32 GetInteger32CameraVariable(UInteger32CameraVariable* InVariableAsset) const;
+
+	UFUNCTION(BlueprintPure, Category=Camera, meta=(CompactNodeTitle="Get Float"))
+	float GetFloatCameraVariable(UFloatCameraVariable* InVariableAsset) const;
+
+	UFUNCTION(BlueprintPure, Category=Camera, meta=(CompactNodeTitle="Get Double"))
+	double GetDoubleCameraVariable(UDoubleCameraVariable* InVariableAsset) const;
+
+	UFUNCTION(BlueprintPure, Category=Camera, meta=(CompactNodeTitle="Get Vec2"))
+	FVector2D GetVector2CameraVariable(UVector2dCameraVariable* InVariableAsset) const;
+
+	UFUNCTION(BlueprintPure, Category=Camera, meta=(CompactNodeTitle="Get Vec3"))
+	FVector GetVector3CameraVariable(UVector3dCameraVariable* InVariableAsset) const;
+
+	UFUNCTION(BlueprintPure, Category=Camera, meta=(CompactNodeTitle="Get Vec4"))
+	FVector4 GetVector4CameraVariable(UVector4dCameraVariable* InVariableAsset) const;
+
+	UFUNCTION(BlueprintPure, Category=Camera, meta=(CompactNodeTitle="Get Rotator"))
+	FRotator GetRotatorCameraVariable(URotator3dCameraVariable* InVariableAsset) const;
+
+	UFUNCTION(BlueprintPure, Category=Camera, meta=(CompactNodeTitle="Get Transform"))
+	FTransform GetTransformCameraVariable(UTransform3dCameraVariable* InVariableAsset) const;
+
+public:
+
+	UFUNCTION(BlueprintCallable, Category=Camera)
+	void SetCameraPose(const FCameraPose& InCameraPose);
+
+	UFUNCTION(BlueprintCallable, Category=Camera)
+	void SetBooleanCameraVariable(UBooleanCameraVariable* InVariableAsset, bool InValue);
+
+	UFUNCTION(BlueprintCallable, Category=Camera)
+	void SetInteger32CameraVariable(UInteger32CameraVariable* InVariableAsset, int32 InValue);
+
+	UFUNCTION(BlueprintCallable, Category=Camera)
+	void SetFloatCameraVariable(UFloatCameraVariable* InVariableAsset, float InValue);
+
+	UFUNCTION(BlueprintCallable, Category=Camera)
+	void SetDoubleCameraVariable(UDoubleCameraVariable* InVariableAsset, double InValue);
+
+	UFUNCTION(BlueprintCallable, Category=Camera)
+	void SetVector2CameraVariable(UVector2dCameraVariable* InVariableAsset, const FVector2D& InValue);
+
+	UFUNCTION(BlueprintCallable, Category=Camera)
+	void SetVector3CameraVariable(UVector3dCameraVariable* InVariableAsset, const FVector& InValue);
+
+	UFUNCTION(BlueprintCallable, Category=Camera)
+	void SetVector4CameraVariable(UVector4dCameraVariable* InVariableAsset, const FVector4& InValue);
+
+	UFUNCTION(BlueprintCallable, Category=Camera)
+	void SetRotatorCameraVariable(URotator3dCameraVariable* InVariableAsset, const FRotator& InValue);
+
+	UFUNCTION(BlueprintCallable, Category=Camera)
+	void SetTransformCameraVariable(UTransform3dCameraVariable* InVariableAsset, const FTransform& InValue);
+
+private:
+
+	FCameraNodeEvaluationResult* Result = nullptr;
+};
+
diff --git a/Engine/Plugins/Cameras/GameplayCameras/Source/GameplayCameras/Public/GameFramework/GameplayCameraComponent.h b/Engine/Plugins/Cameras/GameplayCameras/Source/GameplayCameras/Public/GameFramework/GameplayCameraComponent.h
--- a/Engine/Plugins/Cameras/GameplayCameras/Source/GameplayCameras/Public/GameFramework/GameplayCameraComponent.h
+++ b/Engine/Plugins/Cameras/GameplayCameras/Source/GameplayCameras/Public/GameFramework/GameplayCameraComponent.h
@@ -11,7 +11,7 @@
 #include "GameplayCameraComponent.generated.h"
 
 class UCameraAsset;
-class UGameplayCameraEvaluationResultBlueprintInterop;
+class UCameraEvaluationResultInterop;
 
 namespace UE::Cameras
 {
@@ -39,7 +39,7 @@ public:
 	GAMEPLAYCAMERAS_API void DeactivateCamera();
 
 	UFUNCTION(BlueprintPure, Category=Camera)
-	GAMEPLAYCAMERAS_API UGameplayCameraEvaluationResultBlueprintInterop* GetInitialResult() const;
+	GAMEPLAYCAMERAS_API UCameraEvaluationResultInterop* GetInitialResult() const;
 
 public:
 
@@ -70,15 +70,15 @@ protected:
 
 	TSharedPtr<FGameplayCameraComponentEvaluationContext> EvaluationContext;
 
-	UPROPERTY()
-	TObjectPtr<UGameplayCameraEvaluationResultBlueprintInterop> InitialResultInterop;
+	UPROPERTY(Transient)
+	TObjectPtr<UCameraEvaluationResultInterop> InitialResultInterop;
 	
 #if WITH_EDITORONLY_DATA
 
-	UPROPERTY()
+	UPROPERTY(Transient)
 	TObjectPtr<UStaticMesh> PreviewMesh;
 
-	UPROPERTY()
+	UPROPERTY(Transient)
 	TObjectPtr<UStaticMeshComponent> PreviewMeshComponent;
 
 #endif	// WITH_EDITORONLY_DATA
diff --git a/Engine/Plugins/Cameras/GameplayCameras/Source/GameplayCameras/Public/GameFramework/GameplayCameraEvaluationResultBlueprintInterop.h b/Engine/Plugins/Cameras/GameplayCameras/Source/GameplayCameras/Public/GameFramework/GameplayCameraEvaluationResultBlueprintInterop.h
deleted file mode 100644
