Cameras: fix incorrect name for camera parameter reader struct

It's a templated struct so it should use the T prefix, not the F prefix.

diff --git a/Engine/Plugins/Cameras/GameplayCameras/Source/GameplayCameras/Private/Nodes/Common/LensParametersCameraNode.cpp b/Engine/Plugins/Cameras/GameplayCameras/Source/GameplayCameras/Private/Nodes/Common/LensParametersCameraNode.cpp
--- a/Engine/Plugins/Cameras/GameplayCameras/Source/GameplayCameras/Private/Nodes/Common/LensParametersCameraNode.cpp
+++ b/Engine/Plugins/Cameras/GameplayCameras/Source/GameplayCameras/Private/Nodes/Common/LensParametersCameraNode.cpp
@@ -22,11 +22,11 @@ protected:
 
 private:
 
-	FCameraParameterReader<float> SensorWidthReader;
-	FCameraParameterReader<float> SensorHeightReader;
-	FCameraParameterReader<float> FocusDistanceReader;
-	FCameraParameterReader<float> FocalLengthReader;
-	FCameraParameterReader<float> ApertureReader;
+	TCameraParameterReader<float> SensorWidthReader;
+	TCameraParameterReader<float> SensorHeightReader;
+	TCameraParameterReader<float> FocusDistanceReader;
+	TCameraParameterReader<float> FocalLengthReader;
+	TCameraParameterReader<float> ApertureReader;
 };
 
 UE_DEFINE_CAMERA_NODE_EVALUATOR(FLensParametersCameraNodeEvaluator)
diff --git a/Engine/Plugins/Cameras/GameplayCameras/Source/GameplayCameras/Private/Nodes/Common/OffsetCameraNode.cpp b/Engine/Plugins/Cameras/GameplayCameras/Source/GameplayCameras/Private/Nodes/Common/OffsetCameraNode.cpp
--- a/Engine/Plugins/Cameras/GameplayCameras/Source/GameplayCameras/Private/Nodes/Common/OffsetCameraNode.cpp
+++ b/Engine/Plugins/Cameras/GameplayCameras/Source/GameplayCameras/Private/Nodes/Common/OffsetCameraNode.cpp
@@ -22,7 +22,7 @@ protected:
 
 private:
 
-	FCameraParameterReader<FVector3d> OffsetReader;
+	TCameraParameterReader<FVector3d> OffsetReader;
 };
 
 UE_DEFINE_CAMERA_NODE_EVALUATOR(FOffsetCameraNodeEvaluator)
diff --git a/Engine/Plugins/Cameras/GameplayCameras/Source/GameplayCameras/Public/Core/CameraParameterReader.h b/Engine/Plugins/Cameras/GameplayCameras/Source/GameplayCameras/Public/Core/CameraParameterReader.h
--- a/Engine/Plugins/Cameras/GameplayCameras/Source/GameplayCameras/Public/Core/CameraParameterReader.h
+++ b/Engine/Plugins/Cameras/GameplayCameras/Source/GameplayCameras/Public/Core/CameraParameterReader.h
@@ -10,14 +10,14 @@ namespace UE::Cameras
 {
 
 template<typename ValueType>
-class FCameraParameterReader
+class TCameraParameterReader
 {
 public:
 
-	FCameraParameterReader() {}
+	TCameraParameterReader() {}
 
 	template<typename ParameterType>
-	FCameraParameterReader(const ParameterType& Parameter)
+	TCameraParameterReader(const ParameterType& Parameter)
 	{
 		Initialize(Parameter);
 	}
