Cameras: add rotation offset, fix context-space offset node evaluator logic

Breaking change: the Offset property is now TranslationOffset. Since this is early/alpha code, this is OK.

diff --git a/Engine/Plugins/Cameras/GameplayCameras/Source/GameplayCameras/Private/Nodes/Common/OffsetCameraNode.cpp b/Engine/Plugins/Cameras/GameplayCameras/Source/GameplayCameras/Private/Nodes/Common/OffsetCameraNode.cpp
--- a/Engine/Plugins/Cameras/GameplayCameras/Source/GameplayCameras/Private/Nodes/Common/OffsetCameraNode.cpp
+++ b/Engine/Plugins/Cameras/GameplayCameras/Source/GameplayCameras/Private/Nodes/Common/OffsetCameraNode.cpp
@@ -5,6 +5,7 @@
 #include "Core/CameraEvaluationContext.h"
 #include "Core/CameraParameterReader.h"
 #include "GameplayCameras.h"
+#include "Math/Axis.h"
 
 #include UE_INLINE_GENERATED_CPP_BY_NAME(OffsetCameraNode)
 
@@ -22,7 +23,8 @@ protected:
 
 private:
 
-	TCameraParameterReader<FVector3d> OffsetReader;
+	TCameraParameterReader<FVector3d> TranslationReader;
+	TCameraParameterReader<FRotator3d> RotationReader;
 };
 
 UE_DEFINE_CAMERA_NODE_EVALUATOR(FOffsetCameraNodeEvaluator)
@@ -30,21 +32,24 @@ UE_DEFINE_CAMERA_NODE_EVALUATOR(FOffsetCameraNodeEvaluator)
 void FOffsetCameraNodeEvaluator::OnInitialize(const FCameraNodeEvaluatorInitializeParams& Params)
 {
 	const UOffsetCameraNode* OffsetNode = GetCameraNodeAs<UOffsetCameraNode>();
-	OffsetReader.Initialize(OffsetNode->Offset);
+	TranslationReader.Initialize(OffsetNode->TranslationOffset);
+	RotationReader.Initialize(OffsetNode->RotationOffset);
 }
 
 void FOffsetCameraNodeEvaluator::OnRun(const FCameraNodeEvaluationParams& Params, FCameraNodeEvaluationResult& OutResult)
 {
-	FVector3d LocalOffset = OffsetReader.Get(OutResult.VariableTable);
+	const FVector3d TranslationOffset = TranslationReader.Get(OutResult.VariableTable);
+	const FRotator3d RotationOffset = RotationReader.Get(OutResult.VariableTable);
 
 	const UOffsetCameraNode* OffsetNode = GetCameraNodeAs<UOffsetCameraNode>();
-	switch(OffsetNode->OffsetSpace)
+	switch (OffsetNode->OffsetSpace)
 	{
 		case ECameraNodeSpace::CameraPose:
 		default:
 			{
-				const FRotator3d Rotation = OutResult.CameraPose.GetRotation();
-				LocalOffset = Rotation.RotateVector(LocalOffset);
+				FTransform3d Transform = OutResult.CameraPose.GetTransform();
+				Transform = FTransform3d(RotationOffset, TranslationOffset) * Transform;
+				OutResult.CameraPose.SetTransform(Transform);
 			}
 			break;
 		case ECameraNodeSpace::Context:
@@ -53,8 +58,25 @@ void FOffsetCameraNodeEvaluator::OnRun(const FCameraNodeEvaluationParams& Params
 				const FCameraNodeEvaluationResult& InitialResult = Params.EvaluationContext->GetInitialResult();
 				ensureMsgf(InitialResult.bIsValid,
 						TEXT("OffsetCameraNode: using invalid context result as offset space!"));
-				const FRotator3d Rotation = InitialResult.CameraPose.GetRotation();
-				LocalOffset = Rotation.RotateVector(LocalOffset);
+
+				// The offsets are meant to be treated as context-local. Let's get the context transform
+				// and apply the offsets using that transform's axes.
+				const FTransform3d ContextTransform = InitialResult.CameraPose.GetTransform();
+
+				const FVector3d WorldTranslationOffset = ContextTransform.TransformVector(TranslationOffset);
+
+				const FVector3d ContextForward = ContextTransform.GetUnitAxis(EAxis::X);
+				const FVector3d ContextRight = ContextTransform.GetUnitAxis(EAxis::Y);
+				const FVector3d ContextUp = ContextTransform.GetUnitAxis(EAxis::Z);
+				const FQuat WorldRotationOffset = 
+					FQuat(ContextUp, RotationOffset.Yaw) * 
+					FQuat(ContextRight, RotationOffset.Pitch) *
+					FQuat(ContextForward, RotationOffset.Roll);
+
+				FTransform3d Transform = OutResult.CameraPose.GetTransform();
+				Transform.SetTranslation(Transform.GetTranslation() + WorldTranslationOffset);
+				Transform.SetRotation(WorldRotationOffset * Transform.GetRotation());
+				OutResult.CameraPose.SetTransform(Transform);
 			}
 			else
 			{
@@ -65,12 +87,14 @@ void FOffsetCameraNodeEvaluator::OnRun(const FCameraNodeEvaluationParams& Params
 			}
 			break;
 		case ECameraNodeSpace::World:
-			// Nothing to do;
+			{
+				FTransform3d Transform = OutResult.CameraPose.GetTransform();
+				Transform.SetTranslation(Transform.GetTranslation() + TranslationOffset);
+				Transform.SetRotation(RotationOffset.Quaternion() * Transform.GetRotation());
+				OutResult.CameraPose.SetTransform(Transform);
+			}
 			break;
 	}
-
-	const FVector3d Location = OutResult.CameraPose.GetLocation();
-	OutResult.CameraPose.SetLocation(Location + LocalOffset);
 }
 
 }  // namespace UE::Cameras
diff --git a/Engine/Plugins/Cameras/GameplayCameras/Source/GameplayCameras/Public/Nodes/Common/OffsetCameraNode.h b/Engine/Plugins/Cameras/GameplayCameras/Source/GameplayCameras/Public/Nodes/Common/OffsetCameraNode.h
--- a/Engine/Plugins/Cameras/GameplayCameras/Source/GameplayCameras/Public/Nodes/Common/OffsetCameraNode.h
+++ b/Engine/Plugins/Cameras/GameplayCameras/Source/GameplayCameras/Public/Nodes/Common/OffsetCameraNode.h
@@ -23,9 +23,13 @@ protected:
 
 public:
 
-	/** The offset to apply to the camera, in local space. */
+	/** The translation offset to apply to the camera. */
 	UPROPERTY(EditAnywhere, Category=Common)
-	FVector3dCameraParameter Offset;
+	FVector3dCameraParameter TranslationOffset;
+
+	/** The rotation offset to apply to the camera. */
+	UPROPERTY(EditAnywhere, Category=Common)
+	FRotator3dCameraParameter RotationOffset;
 
 	/** The space in which to apply the offset. */
 	UPROPERTY(EditAnywhere, Category=Common)
