Cameras: attach to player pawn camera node

diff --git a/Engine/Plugins/Cameras/GameplayCameras/Source/GameplayCameras/Private/Nodes/Attach/AttachToPlayerPawnCameraNode.cpp b/Engine/Plugins/Cameras/GameplayCameras/Source/GameplayCameras/Private/Nodes/Attach/AttachToPlayerPawnCameraNode.cpp
new file mode 100644
--- /dev/null
+++ b/Engine/Plugins/Cameras/GameplayCameras/Source/GameplayCameras/Private/Nodes/Attach/AttachToPlayerPawnCameraNode.cpp
@@ -0,0 +1,77 @@
+// Copyright Epic Games, Inc. All Rights Reserved.
+
+#include "Nodes/Attach/AttachToPlayerPawnCameraNode.h"
+
+#include "Core/CameraEvaluationContext.h"
+#include "Core/CameraParameterReader.h"
+#include "GameFramework/Pawn.h"
+#include "GameFramework/PlayerController.h"
+#include "GameplayCameras.h"
+
+#include UE_INLINE_GENERATED_CPP_BY_NAME(AttachToPlayerPawnCameraNode)
+
+namespace UE::Cameras
+{
+
+class FAttachToPlayerPawnCameraNodeEvaluator : public FCameraNodeEvaluator
+{
+	UE_DECLARE_CAMERA_NODE_EVALUATOR(GAMEPLAYCAMERAS_API, FAttachToPlayerPawnCameraNodeEvaluator)
+
+protected:
+
+	virtual void OnInitialize(const FCameraNodeEvaluatorInitializeParams& Params) override;
+	virtual void OnRun(const FCameraNodeEvaluationParams& Params, FCameraNodeEvaluationResult& OutResult) override;
+
+private:
+
+	TCameraParameterReader<bool> AttachToLocationReader;
+	TCameraParameterReader<bool> AttachToRotationReader;
+};
+
+UE_DEFINE_CAMERA_NODE_EVALUATOR(FAttachToPlayerPawnCameraNodeEvaluator)
+
+void FAttachToPlayerPawnCameraNodeEvaluator::OnInitialize(const FCameraNodeEvaluatorInitializeParams& Params)
+{
+	const UAttachToPlayerPawnCameraNode* AttachNode = GetCameraNodeAs<UAttachToPlayerPawnCameraNode>();
+	AttachToLocationReader.Initialize(AttachNode->AttachToLocation);
+	AttachToRotationReader.Initialize(AttachNode->AttachToRotation);
+}
+
+void FAttachToPlayerPawnCameraNodeEvaluator::OnRun(const FCameraNodeEvaluationParams& Params, FCameraNodeEvaluationResult& OutResult)
+{
+	APlayerController* PlayerController = Params.EvaluationContext->GetPlayerController();
+	if (!ensure(PlayerController))
+	{
+		return;
+	}
+
+	const APawn* Pawn = PlayerController->GetPawnOrSpectator();
+	if (!ensure(Pawn))
+	{
+		return;
+	}
+
+	const bool bAttachToLocation = AttachToLocationReader.Get(OutResult.VariableTable);
+	const bool bAttachToRotation = AttachToRotationReader.Get(OutResult.VariableTable);
+	
+	if (bAttachToLocation)
+	{
+		const FVector3d PawnLocation = Pawn->GetActorLocation();
+		OutResult.CameraPose.SetLocation(PawnLocation);
+	}
+
+	if (bAttachToRotation)
+	{
+		const FRotator3d PawnRotation = Pawn->GetActorRotation();
+		OutResult.CameraPose.SetRotation(PawnRotation);
+	}
+}
+
+}  // namespace UE::Cameras
+
+FCameraNodeEvaluatorPtr UAttachToPlayerPawnCameraNode::OnBuildEvaluator(FCameraNodeEvaluatorBuilder& Builder) const
+{
+	using namespace UE::Cameras;
+	return Builder.BuildEvaluator<FAttachToPlayerPawnCameraNodeEvaluator>();
+}
+
diff --git a/Engine/Plugins/Cameras/GameplayCameras/Source/GameplayCameras/Public/Nodes/Attach/AttachToPlayerPawnCameraNode.h b/Engine/Plugins/Cameras/GameplayCameras/Source/GameplayCameras/Public/Nodes/Attach/AttachToPlayerPawnCameraNode.h
new file mode 100644
--- /dev/null
+++ b/Engine/Plugins/Cameras/GameplayCameras/Source/GameplayCameras/Public/Nodes/Attach/AttachToPlayerPawnCameraNode.h
@@ -0,0 +1,33 @@
+// Copyright Epic Games, Inc. All Rights Reserved.
+
+#pragma once
+
+#include "Core/CameraNode.h"
+#include "Core/CameraParameters.h"
+
+#include "AttachToPlayerPawnCameraNode.generated.h"
+
+/**
+ * A camera node that moves the camera to the player pawn.
+ */
+UCLASS(MinimalAPI)
+class UAttachToPlayerPawnCameraNode : public UCameraNode
+{
+	GENERATED_BODY()
+
+protected:
+
+	// UCameraNode interface.
+	virtual FCameraNodeEvaluatorPtr OnBuildEvaluator(FCameraNodeEvaluatorBuilder& Builder) const override;
+
+public:
+
+	/** Whether to move the camera to the pawn's location. */
+	UPROPERTY(EditAnywhere, Category=Common)
+	FBooleanCameraParameter AttachToLocation;
+
+	/** Whether to align the camera rotation to the pawn's orientation. */
+	UPROPERTY(EditAnywhere, Category=Common)
+	FBooleanCameraParameter AttachToRotation;
+};
+
