Cameras: field-of-view setter camera node

diff --git a/Engine/Plugins/Cameras/GameplayCameras/Source/GameplayCameras/Private/Nodes/Common/FieldOfViewCameraNode.cpp b/Engine/Plugins/Cameras/GameplayCameras/Source/GameplayCameras/Private/Nodes/Common/FieldOfViewCameraNode.cpp
new file mode 100644
--- /dev/null
+++ b/Engine/Plugins/Cameras/GameplayCameras/Source/GameplayCameras/Private/Nodes/Common/FieldOfViewCameraNode.cpp
@@ -0,0 +1,49 @@
+// Copyright Epic Games, Inc. All Rights Reserved.
+
+#include "Nodes/Common/FieldOfViewCameraNode.h"
+
+#include "Core/CameraEvaluationContext.h"
+#include "Core/CameraParameterReader.h"
+#include "GameplayCameras.h"
+
+#include UE_INLINE_GENERATED_CPP_BY_NAME(FieldOfViewCameraNode)
+
+namespace UE::Cameras
+{
+
+class FFieldOfViewCameraNodeEvaluator : public FCameraNodeEvaluator
+{
+	UE_DECLARE_CAMERA_NODE_EVALUATOR(GAMEPLAYCAMERAS_API, FFieldOfViewCameraNodeEvaluator)
+
+protected:
+
+	virtual void OnInitialize(const FCameraNodeEvaluatorInitializeParams& Params) override;
+	virtual void OnRun(const FCameraNodeEvaluationParams& Params, FCameraNodeEvaluationResult& OutResult) override;
+
+private:
+
+	TCameraParameterReader<float> FieldOfViewReader;
+};
+
+UE_DEFINE_CAMERA_NODE_EVALUATOR(FFieldOfViewCameraNodeEvaluator)
+
+void FFieldOfViewCameraNodeEvaluator::OnInitialize(const FCameraNodeEvaluatorInitializeParams& Params)
+{
+	const UFieldOfViewCameraNode* FieldOfViewNode = GetCameraNodeAs<UFieldOfViewCameraNode>();
+	FieldOfViewReader.Initialize(FieldOfViewNode->FieldOfView);
+}
+
+void FFieldOfViewCameraNodeEvaluator::OnRun(const FCameraNodeEvaluationParams& Params, FCameraNodeEvaluationResult& OutResult)
+{
+	const float FieldOfView = FieldOfViewReader.Get(OutResult.VariableTable);
+	OutResult.CameraPose.SetFieldOfView(FieldOfView);
+}
+
+}  // namespace UE::Cameras
+
+FCameraNodeEvaluatorPtr UFieldOfViewCameraNode::OnBuildEvaluator(FCameraNodeEvaluatorBuilder& Builder) const
+{
+	using namespace UE::Cameras;
+	return Builder.BuildEvaluator<FFieldOfViewCameraNodeEvaluator>();
+}
+
diff --git a/Engine/Plugins/Cameras/GameplayCameras/Source/GameplayCameras/Public/Nodes/Common/FieldOfViewCameraNode.h b/Engine/Plugins/Cameras/GameplayCameras/Source/GameplayCameras/Public/Nodes/Common/FieldOfViewCameraNode.h
new file mode 100644
--- /dev/null
+++ b/Engine/Plugins/Cameras/GameplayCameras/Source/GameplayCameras/Public/Nodes/Common/FieldOfViewCameraNode.h
@@ -0,0 +1,30 @@
+// Copyright Epic Games, Inc. All Rights Reserved.
+
+#pragma once
+
+#include "Core/CameraNode.h"
+#include "Core/CameraParameters.h"
+#include "Nodes/CameraNodeTypes.h"
+
+#include "FieldOfViewCameraNode.generated.h"
+
+/**
+ * A camera node that sets the field of view of the camera.
+ */
+UCLASS(MinimalAPI)
+class UFieldOfViewCameraNode : public UCameraNode
+{
+	GENERATED_BODY()
+
+protected:
+
+	// UCameraNode interface.
+	virtual FCameraNodeEvaluatorPtr OnBuildEvaluator(FCameraNodeEvaluatorBuilder& Builder) const override;
+
+public:
+
+	/** The field of view, in degrees. */
+	UPROPERTY(EditAnywhere, Category=Common)
+	FFloatCameraParameter FieldOfView;
+};
+
