Cameras: move initial creation of graph to the schema class

Since the schema class is typically sub-classed for specific graphs, it makes it easier for those sub-classes to add custom extra nodes too.

diff --git a/Engine/Plugins/Cameras/GameplayCameras/Source/GameplayCamerasEditor/Private/Editors/ObjectTreeGraph.cpp b/Engine/Plugins/Cameras/GameplayCameras/Source/GameplayCamerasEditor/Private/Editors/ObjectTreeGraph.cpp
--- a/Engine/Plugins/Cameras/GameplayCameras/Source/GameplayCamerasEditor/Private/Editors/ObjectTreeGraph.cpp
+++ b/Engine/Plugins/Cameras/GameplayCameras/Source/GameplayCamerasEditor/Private/Editors/ObjectTreeGraph.cpp
@@ -34,119 +34,11 @@ const FObjectTreeGraphConfig& UObjectTreeGraph::GetConfig() const
 }
 
 void UObjectTreeGraph::RebuildGraph(EObjectTreeGraphBuildSource InSource)
-{
-	RemoveAllNodes();
-	CreateAllNodes(InSource);
-	NotifyGraphChanged();
-}
-
-void UObjectTreeGraph::RemoveAllNodes()
-{
-	TArray<UEdGraphNode*> NodesToRemove(Nodes);  // Copy all nodes to remove them.
-	for (UEdGraphNode* NodeToRemove : NodesToRemove)
-	{
-		RemoveNode(NodeToRemove);
-	}
-}
-
-void UObjectTreeGraph::CreateAllNodes(EObjectTreeGraphBuildSource InSource)
 {
 	const UObjectTreeGraphSchema* GraphSchema = Cast<UObjectTreeGraphSchema>(GetSchema());
-
-	UObject* RootObject = WeakRootObject.Get();
-	if (!ensure(RootObject))
-	{
-		return;
-	}
-
-	TArray<UObject*> AllObjects;
-	UPackage* Package = RootObject->GetOutermost();
-	GetObjectsWithPackage(Package, AllObjects);
-
-	ensure(AllObjects.Contains(RootObject));
-
-	FCreatedNodes CreatedNodes;
-	for (UObject* Object : AllObjects)
-	{
-		if (UObjectTreeGraphNode* GraphNode = GraphSchema->CreateObjectNode(this, Object))
-		{
-			CreatedNodes.CreatedNodes.Add(Object, GraphNode);
-		}
-	}
-
-	RootObjectNode = nullptr;
-	UObjectTreeGraphNode** CreatedRootObjectNode = CreatedNodes.CreatedNodes.Find(RootObject);
-	if (ensure(CreatedRootObjectNode))
-	{
-		RootObjectNode = *CreatedRootObjectNode;
-	}
-
-	for (TPair<UObject*, UObjectTreeGraphNode*> Pair : CreatedNodes.CreatedNodes)
-	{
-		CreateConnections(Pair.Value, CreatedNodes);
-	}
-}
-
-void UObjectTreeGraph::CreateConnections(UObjectTreeGraphNode* InGraphNode, const FCreatedNodes& InCreatedNodes)
-{
-	UObject* Object = InGraphNode->GetObject();
-	UClass* ObjectClass = Object->GetClass();
-
-	TArray<FProperty*> ConnectableProperties;
-	InGraphNode->GetAllConnectableProperties(ConnectableProperties);
-
-	for (FProperty* ConnectableProperty : ConnectableProperties)
+	if (ensure(GraphSchema))
 	{
-		if (FObjectProperty* ObjectProperty = CastField<FObjectProperty>(ConnectableProperty))
-		{
-			UEdGraphPin* Pin = InGraphNode->GetPinForProperty(ObjectProperty);
-			if (!ensure(Pin))
-			{
-				continue;
-			}
-
-			TObjectPtr<UObject> OutConnectedObject;
-			ObjectProperty->GetValue_InContainer(Object, &OutConnectedObject);
-			if (!OutConnectedObject)
-			{
-				continue;
-			}
-
-			UObjectTreeGraphNode* const* ConnectedNode = InCreatedNodes.CreatedNodes.Find(OutConnectedObject);
-			if (ensure(ConnectedNode))
-			{
-				UEdGraphPin* ConnectedPin = (*ConnectedNode)->GetSelfPin();
-				Pin->MakeLinkTo(ConnectedPin);
-			}
-		}
-		else if (FArrayProperty* ArrayProperty = CastField<FArrayProperty>(ConnectableProperty))
-		{
-			FObjectProperty* InnerProperty = CastFieldChecked<FObjectProperty>(ArrayProperty->Inner);
-			FScriptArrayHelper ArrayHelper(ArrayProperty, ArrayProperty->ContainerPtrToValuePtr<void>(Object));
-
-			const int32 ArrayNum = ArrayHelper.Num();
-			for (int32 Index = 0; Index < ArrayNum; ++Index)
-			{
-				UEdGraphPin* Pin = InGraphNode->GetPinForPropertyNewItem(ArrayProperty, true);
-				if (!ensure(Pin))
-				{
-					continue;
-				}
-
-				UObject* ConnectedObject = InnerProperty->GetObjectPropertyValue(ArrayHelper.GetRawPtr(Index));
-				if (!ConnectedObject)
-				{
-					continue;
-				}
-
-				UObjectTreeGraphNode* const* ConnectedNode = InCreatedNodes.CreatedNodes.Find(ConnectedObject);
-				if (ensure(ConnectedNode))
-				{
-					UEdGraphPin* ConnectedPin = (*ConnectedNode)->GetSelfPin();
-					Pin->MakeLinkTo(ConnectedPin);
-				}
-			}
-		}
+		GraphSchema->RebuildGraph(this, InSource);
 	}
 }
 
diff --git a/Engine/Plugins/Cameras/GameplayCameras/Source/GameplayCamerasEditor/Private/Editors/ObjectTreeGraphSchema.cpp b/Engine/Plugins/Cameras/GameplayCameras/Source/GameplayCamerasEditor/Private/Editors/ObjectTreeGraphSchema.cpp
--- a/Engine/Plugins/Cameras/GameplayCameras/Source/GameplayCamerasEditor/Private/Editors/ObjectTreeGraphSchema.cpp
+++ b/Engine/Plugins/Cameras/GameplayCameras/Source/GameplayCamerasEditor/Private/Editors/ObjectTreeGraphSchema.cpp
@@ -24,6 +24,127 @@ UObjectTreeGraphSchema::UObjectTreeGraphSchema(const FObjectInitializer& ObjInit
 {
 }
 
+void UObjectTreeGraphSchema::RebuildGraph(UObjectTreeGraph* InGraph, EObjectTreeGraphBuildSource InSource) const
+{
+	RemoveAllNodes(InGraph);
+	CreateAllNodes(InGraph, InSource);
+	InGraph->NotifyGraphChanged();
+}
+
+void UObjectTreeGraphSchema::RemoveAllNodes(UObjectTreeGraph* InGraph) const
+{
+	TArray<UEdGraphNode*> NodesToRemove(InGraph->Nodes);  // Copy all nodes to remove them.
+	for (UEdGraphNode* NodeToRemove : NodesToRemove)
+	{
+		InGraph->RemoveNode(NodeToRemove);
+	}
+}
+
+void UObjectTreeGraphSchema::CreateAllNodes(UObjectTreeGraph* InGraph, EObjectTreeGraphBuildSource InSource) const
+{
+	UObject* RootObject = InGraph->GetRootObject();
+	if (!ensure(RootObject))
+	{
+		return;
+	}
+
+	TArray<UObject*> AllObjects;
+	UPackage* Package = RootObject->GetOutermost();
+	GetObjectsWithPackage(Package, AllObjects);
+
+	ensure(AllObjects.Contains(RootObject));
+
+	FCreatedNodes CreatedNodes;
+	for (UObject* Object : AllObjects)
+	{
+		if (UObjectTreeGraphNode* GraphNode = CreateObjectNode(InGraph, Object))
+		{
+			CreatedNodes.CreatedNodes.Add(Object, GraphNode);
+		}
+	}
+
+	InGraph->RootObjectNode = nullptr;
+	UObjectTreeGraphNode** CreatedRootObjectNode = CreatedNodes.CreatedNodes.Find(RootObject);
+	if (ensure(CreatedRootObjectNode))
+	{
+		InGraph->RootObjectNode = *CreatedRootObjectNode;
+	}
+
+	for (TPair<UObject*, UObjectTreeGraphNode*> Pair : CreatedNodes.CreatedNodes)
+	{
+		CreateConnections(Pair.Value, CreatedNodes);
+	}
+
+	OnCreateAllNodes(InGraph, CreatedNodes);
+}
+
+void UObjectTreeGraphSchema::CreateConnections(UObjectTreeGraphNode* InGraphNode, const FCreatedNodes& InCreatedNodes) const
+{
+	UObject* Object = InGraphNode->GetObject();
+	UClass* ObjectClass = Object->GetClass();
+
+	TArray<FProperty*> ConnectableProperties;
+	InGraphNode->GetAllConnectableProperties(ConnectableProperties);
+
+	for (FProperty* ConnectableProperty : ConnectableProperties)
+	{
+		if (FObjectProperty* ObjectProperty = CastField<FObjectProperty>(ConnectableProperty))
+		{
+			UEdGraphPin* Pin = InGraphNode->GetPinForProperty(ObjectProperty);
+			if (!ensure(Pin))
+			{
+				continue;
+			}
+
+			TObjectPtr<UObject> OutConnectedObject;
+			ObjectProperty->GetValue_InContainer(Object, &OutConnectedObject);
+			if (!OutConnectedObject)
+			{
+				continue;
+			}
+
+			UObjectTreeGraphNode* const* ConnectedNode = InCreatedNodes.CreatedNodes.Find(OutConnectedObject);
+			if (ensure(ConnectedNode))
+			{
+				UEdGraphPin* ConnectedPin = (*ConnectedNode)->GetSelfPin();
+				Pin->MakeLinkTo(ConnectedPin);
+			}
+		}
+		else if (FArrayProperty* ArrayProperty = CastField<FArrayProperty>(ConnectableProperty))
+		{
+			FObjectProperty* InnerProperty = CastFieldChecked<FObjectProperty>(ArrayProperty->Inner);
+			FScriptArrayHelper ArrayHelper(ArrayProperty, ArrayProperty->ContainerPtrToValuePtr<void>(Object));
+
+			const int32 ArrayNum = ArrayHelper.Num();
+			for (int32 Index = 0; Index < ArrayNum; ++Index)
+			{
+				UEdGraphPin* Pin = InGraphNode->GetPinForPropertyNewItem(ArrayProperty, true);
+				if (!ensure(Pin))
+				{
+					continue;
+				}
+
+				UObject* ConnectedObject = InnerProperty->GetObjectPropertyValue(ArrayHelper.GetRawPtr(Index));
+				if (!ConnectedObject)
+				{
+					continue;
+				}
+
+				UObjectTreeGraphNode* const* ConnectedNode = InCreatedNodes.CreatedNodes.Find(ConnectedObject);
+				if (ensure(ConnectedNode))
+				{
+					UEdGraphPin* ConnectedPin = (*ConnectedNode)->GetSelfPin();
+					Pin->MakeLinkTo(ConnectedPin);
+				}
+			}
+		}
+	}
+}
+
+void UObjectTreeGraphSchema::OnCreateAllNodes(UObjectTreeGraph* InGraph, const FCreatedNodes& InCreatedNodes) const
+{
+}
+
 UObjectTreeGraphNode* UObjectTreeGraphSchema::CreateObjectNode(UObjectTreeGraph* InGraph, UObject* InObject) const
 {
 	if (!InObject)
@@ -123,7 +244,6 @@ void UObjectTreeGraphSchema::GetGraphContextActions(FGraphContextMenuBuilder& Co
 			}
 		}
 
-		checkSlow(PossibleObjectClass);
 		TSharedRef<FObjectGraphSchemaAction_NewNode> Action = MakeShared<FObjectGraphSchemaAction_NewNode>(
 				CategoryName ? FText::FromString(*CategoryName) : FText::GetEmpty(),
 				PossibleObjectClass->GetDisplayNameText(), 
diff --git a/Engine/Plugins/Cameras/GameplayCameras/Source/GameplayCamerasEditor/Private/Editors/SObjectTreeGraphNode.cpp b/Engine/Plugins/Cameras/GameplayCameras/Source/GameplayCamerasEditor/Private/Editors/SObjectTreeGraphNode.cpp
--- a/Engine/Plugins/Cameras/GameplayCameras/Source/GameplayCamerasEditor/Private/Editors/SObjectTreeGraphNode.cpp
+++ b/Engine/Plugins/Cameras/GameplayCameras/Source/GameplayCamerasEditor/Private/Editors/SObjectTreeGraphNode.cpp
@@ -24,8 +24,3 @@ void SObjectTreeGraphNode::MoveTo(const FVector2D& NewPosition, FNodeSet& NodeFi
 	}
 }
 
-void SObjectTreeGraphAddArrayItemPin::Construct(const FArguments& InArgs, UEdGraphPin* InGraphPinObj)
-{
-	SGraphPin::Construct(SGraphPin::FArguments(), InGraphPinObj);
-}
-
diff --git a/Engine/Plugins/Cameras/GameplayCameras/Source/GameplayCamerasEditor/Public/Editors/ObjectTreeGraph.h b/Engine/Plugins/Cameras/GameplayCameras/Source/GameplayCamerasEditor/Public/Editors/ObjectTreeGraph.h
--- a/Engine/Plugins/Cameras/GameplayCameras/Source/GameplayCamerasEditor/Public/Editors/ObjectTreeGraph.h
+++ b/Engine/Plugins/Cameras/GameplayCameras/Source/GameplayCamerasEditor/Public/Editors/ObjectTreeGraph.h
@@ -15,12 +15,7 @@
 class UObjectTreeGraphNode;
 class UObjectTreeGraphSchema;
 
-/** Specifies where to find graph node objects. */
-enum class EObjectTreeGraphBuildSource
-{
-	/** Create graph nodes from all eligible objects found in the same package. */
-	RootObjectPackage
-};
+enum class EObjectTreeGraphBuildSource : uint8;
 
 /**
  * A node graph that represents a hierarchy of objects and their relationships.
@@ -48,17 +43,6 @@ public:
 	/** Rebuilds the node graph. */
 	void RebuildGraph(EObjectTreeGraphBuildSource InSource);
 
-private:
-
-	struct FCreatedNodes
-	{
-		TMap<UObject*, UObjectTreeGraphNode*> CreatedNodes;
-	};
-
-	void RemoveAllNodes();
-	void CreateAllNodes(EObjectTreeGraphBuildSource InSource);
-	void CreateConnections(UObjectTreeGraphNode* InGraphNode, const FCreatedNodes& InCreatedNodes);
-
 private:
 
 	TWeakObjectPtr<> WeakRootObject;
@@ -67,5 +51,7 @@ private:
 
 	UPROPERTY()
 	TObjectPtr<UObjectTreeGraphNode> RootObjectNode;
+
+	friend class UObjectTreeGraphSchema;
 };
 
diff --git a/Engine/Plugins/Cameras/GameplayCameras/Source/GameplayCamerasEditor/Public/Editors/ObjectTreeGraphSchema.h b/Engine/Plugins/Cameras/GameplayCameras/Source/GameplayCamerasEditor/Public/Editors/ObjectTreeGraphSchema.h
--- a/Engine/Plugins/Cameras/GameplayCameras/Source/GameplayCamerasEditor/Public/Editors/ObjectTreeGraphSchema.h
+++ b/Engine/Plugins/Cameras/GameplayCameras/Source/GameplayCamerasEditor/Public/Editors/ObjectTreeGraphSchema.h
@@ -12,6 +12,13 @@ class UObjectTreeGraph;
 class UObjectTreeGraphNode;
 struct FObjectTreeGraphClassConfig;
 
+/** Specifies where to find graph node objects. */
+enum class EObjectTreeGraphBuildSource : uint8
+{
+	/** Create graph nodes from all eligible objects found in the same package. */
+	RootObjectPackage
+};
+
 /**
  * Schema class for an object tree graph.
  */
@@ -36,6 +43,8 @@ public:
 	/** Creates a new schema. */
 	UObjectTreeGraphSchema(const FObjectInitializer& ObjInit);
 
+	void RebuildGraph(UObjectTreeGraph* InGraph, EObjectTreeGraphBuildSource InSource) const;
+
 	/** Creates an object graph node for the given object. */
 	UObjectTreeGraphNode* CreateObjectNode(UObjectTreeGraph* InGraph, UObject* InObject) const;
 
@@ -61,12 +70,25 @@ public:
 
 protected:
 
+	struct FCreatedNodes
+	{
+		TMap<UObject*, UObjectTreeGraphNode*> CreatedNodes;
+	};
+
+	// UObjectTreeGraphSchema interface.
+	virtual void OnCreateAllNodes(UObjectTreeGraph* InGraph, const FCreatedNodes& InCreatedNodes) const;
 	virtual UObjectTreeGraphNode* CreateObjectNodeImpl(UObjectTreeGraph* InGraph, UObject* InObject) const;
 
 protected:
 
 	const FObjectTreeGraphClassConfig& GetObjectClassConfig(const UObjectTreeGraphNode* InNode) const;
 	const FObjectTreeGraphClassConfig& GetObjectClassConfig(const UObjectTreeGraph* InGraph, UClass* InObjectClass) const;
+
+private:
+
+	void RemoveAllNodes(UObjectTreeGraph* InGraph) const;
+	void CreateAllNodes(UObjectTreeGraph* InGraph, EObjectTreeGraphBuildSource InSource) const;
+	void CreateConnections(UObjectTreeGraphNode* InGraphNode, const FCreatedNodes& InCreatedNodes) const;
 };
 
 /**
diff --git a/Engine/Plugins/Cameras/GameplayCameras/Source/GameplayCamerasEditor/Public/Editors/SObjectTreeGraphNode.h b/Engine/Plugins/Cameras/GameplayCameras/Source/GameplayCamerasEditor/Public/Editors/SObjectTreeGraphNode.h
--- a/Engine/Plugins/Cameras/GameplayCameras/Source/GameplayCamerasEditor/Public/Editors/SObjectTreeGraphNode.h
+++ b/Engine/Plugins/Cameras/GameplayCameras/Source/GameplayCamerasEditor/Public/Editors/SObjectTreeGraphNode.h
@@ -7,6 +7,9 @@
 
 class UObjectTreeGraphNode;
 
+/**
+ * The widget used by default for object tree graph nodes.
+ */
 class SObjectTreeGraphNode : public SGraphNode
 {
 public:
@@ -31,19 +34,3 @@ protected:
 	UObjectTreeGraphNode* ObjectGraphNode;
 };
 
-class SObjectTreeGraphAddArrayItemPin : public SGraphPin
-{
-public:
-
-	SLATE_BEGIN_ARGS(SObjectTreeGraphAddArrayItemPin)
-	{}
-	SLATE_END_ARGS()
-
-	void Construct(const FArguments& InArgs, UEdGraphPin* InGraphPinObj);
-
-public:
-
-	// SGraphPin interface.
-
-};
-
