Cameras: add more extensibility to the ObjectTreeGraph API, improve schema logic

- Self pin name can be customized.
- Properties that would show up as pins can be hidden with the "ObjectTreeGraphHidden" metadata.
- Correctly allocate custom graph node sub-classes as per the configuration.
- Add virtual/overridable methods to make it easier to extend schema logic in schema sub-classes.

diff --git a/Engine/Plugins/Cameras/GameplayCameras/Source/GameplayCamerasEditor/Private/Editors/ObjectTreeGraphConfig.cpp b/Engine/Plugins/Cameras/GameplayCameras/Source/GameplayCamerasEditor/Private/Editors/ObjectTreeGraphConfig.cpp
--- a/Engine/Plugins/Cameras/GameplayCameras/Source/GameplayCamerasEditor/Private/Editors/ObjectTreeGraphConfig.cpp
+++ b/Engine/Plugins/Cameras/GameplayCameras/Source/GameplayCamerasEditor/Private/Editors/ObjectTreeGraphConfig.cpp
@@ -11,7 +11,8 @@
 #define LOCTEXT_NAMESPACE "ObjectTreeGraphConfig"
 
 FObjectTreeGraphClassConfig::FObjectTreeGraphClassConfig()
-	: _SelfPinFriendlyName(LOCTEXT("SelfPinFriendlyName", "self"))
+	: _SelfPinName(NAME_Self)
+	, _SelfPinFriendlyName(FText::GetEmpty())
 	, _SelfPinDirection(EGPD_Input)
 	, _HasSelfPin(true)
 	, _DefaultPropertyPinDirection(EGPD_Output)
@@ -38,6 +39,11 @@ FObjectTreeGraphConfig::FObjectTreeGraphConfig()
 
 bool FObjectTreeGraphConfig::IsConnectable(UClass* InObjectClass) const
 {
+	if (!ensure(InObjectClass))
+	{
+		return false;
+	}
+
 	const bool bIsConnectable = Algo::AnyOf(ConnectableObjectClasses, [InObjectClass](UClass* Item)
 			{
 				return InObjectClass->IsChildOf(Item);
@@ -59,6 +65,22 @@ bool FObjectTreeGraphConfig::IsConnectable(UClass* InObjectClass) const
 	return true;
 }
 
+bool FObjectTreeGraphConfig::IsConnectable(FObjectProperty* InObjectProperty) const
+{
+	if (!ensure(InObjectProperty))
+	{
+		return false;
+	}
+
+	if (!IsConnectable(InObjectProperty->PropertyClass))
+	{
+		return false;
+	}
+
+	const bool bIsHidden = InObjectProperty->GetBoolMetaData(TEXT("ObjectTreeGraphHidden"));
+	return !bIsHidden;
+}
+
 void FObjectTreeGraphConfig::GetConnectableClasses(TArray<UClass*>& OutClasses, bool bPlaceableOnly)
 {
 	for (TObjectIterator<UClass> ClassIt; ClassIt; ++ClassIt)
diff --git a/Engine/Plugins/Cameras/GameplayCameras/Source/GameplayCamerasEditor/Private/Editors/ObjectTreeGraphNode.cpp b/Engine/Plugins/Cameras/GameplayCameras/Source/GameplayCamerasEditor/Private/Editors/ObjectTreeGraphNode.cpp
--- a/Engine/Plugins/Cameras/GameplayCameras/Source/GameplayCamerasEditor/Private/Editors/ObjectTreeGraphNode.cpp
+++ b/Engine/Plugins/Cameras/GameplayCameras/Source/GameplayCamerasEditor/Private/Editors/ObjectTreeGraphNode.cpp
@@ -137,8 +137,6 @@ FLinearColor UObjectTreeGraphNode::GetNodeBodyTintColor() const
 
 void UObjectTreeGraphNode::AllocateDefaultPins()
 {
-	static const FName SelfPinName("self");
-
 	if (!ensure(Object))
 	{
 		return;
@@ -152,6 +150,7 @@ void UObjectTreeGraphNode::AllocateDefaultPins()
 	{
 		FEdGraphPinType SelfPinType;
 		SelfPinType.PinCategory = UObjectTreeGraphSchema::PC_Self;
+		const FName& SelfPinName = ObjectClassConfig.SelfPinName();
 		UEdGraphPin* SelfPin = CreatePin(ObjectClassConfig.SelfPinDirection(), SelfPinType, SelfPinName);
 		SelfPin->PinFriendlyName = ObjectClassConfig.SelfPinFriendlyName();
 	}
@@ -167,7 +166,7 @@ void UObjectTreeGraphNode::AllocateDefaultPins()
 
 		if (FObjectProperty* ObjectProperty = CastField<FObjectProperty>(*PropertyIt))
 		{
-			if (!ObjectProperty->PropertyClass || !OuterGraphConfig.IsConnectable(ObjectProperty->PropertyClass))
+			if (!OuterGraphConfig.IsConnectable(ObjectProperty))
 			{
 				continue;
 			}
diff --git a/Engine/Plugins/Cameras/GameplayCameras/Source/GameplayCamerasEditor/Private/Editors/ObjectTreeGraphSchema.cpp b/Engine/Plugins/Cameras/GameplayCameras/Source/GameplayCamerasEditor/Private/Editors/ObjectTreeGraphSchema.cpp
--- a/Engine/Plugins/Cameras/GameplayCameras/Source/GameplayCamerasEditor/Private/Editors/ObjectTreeGraphSchema.cpp
+++ b/Engine/Plugins/Cameras/GameplayCameras/Source/GameplayCamerasEditor/Private/Editors/ObjectTreeGraphSchema.cpp
@@ -162,8 +162,17 @@ UObjectTreeGraphNode* UObjectTreeGraphSchema::CreateObjectNode(UObjectTreeGraph*
 
 UObjectTreeGraphNode* UObjectTreeGraphSchema::CreateObjectNodeImpl(UObjectTreeGraph* InGraph, UObject* InObject) const
 {
+	const FObjectTreeGraphConfig& Config = InGraph->GetConfig();
+	const FObjectTreeGraphClassConfig& ClassConfig = Config.GetObjectClassConfig(InObject->GetClass());
+
+	TSubclassOf<UObjectTreeGraphNode> GraphNodeClass = ClassConfig.GraphNodeClass();
+	if (!GraphNodeClass.Get())
+	{
+		GraphNodeClass = Config.DefaultGraphNodeClass;
+	}
+
 	FGraphNodeCreator<UObjectTreeGraphNode> GraphNodeCreator(*InGraph);
-	UObjectTreeGraphNode* NewNode = GraphNodeCreator.CreateNode(false);
+	UObjectTreeGraphNode* NewNode = GraphNodeCreator.CreateNode(false, GraphNodeClass);
 	NewNode->Initialize(InObject);
 	GraphNodeCreator.Finalize();
 	return NewNode;
@@ -177,6 +186,7 @@ void UObjectTreeGraphSchema::GetGraphContextActions(FGraphContextMenuBuilder& Co
 	// Find the common class restriction for all the dragged pins. We will only show actions that
 	// are compatible with them.
 	UClass* DraggedPinClass = nullptr;
+	bool bShouldShowNewObjectActions = true;
 	if (const UEdGraphPin* DraggedPin = ContextMenuBuilder.FromPin)
 	{
 		UObjectTreeGraphNode* OwningNode = Cast<UObjectTreeGraphNode>(DraggedPin->GetOwningNode());
@@ -190,8 +200,23 @@ void UObjectTreeGraphSchema::GetGraphContextActions(FGraphContextMenuBuilder& Co
 			{
 				DraggedPinClass = OwningNode->GetConnectedObjectClassForPin(DraggedPin);
 			}
+			else
+			{
+				// Dragged an unknown pin...
+				bShouldShowNewObjectActions = false;
+			}
+		}
+		else
+		{
+			// Dragged a pin from an unknown node...
+				bShouldShowNewObjectActions = false;
 		}
 	}
+	if (!bShouldShowNewObjectActions)
+	{
+		// Don't show anything.
+		return;
+	}
 
 	// Find all the object classes we can create from those pins, for the given graph.
 	TArray<UClass*> PossibleObjectClasses;
@@ -244,6 +269,7 @@ void UObjectTreeGraphSchema::GetGraphContextActions(FGraphContextMenuBuilder& Co
 			}
 		}
 
+		checkSlow(PossibleObjectClass);
 		TSharedRef<FObjectGraphSchemaAction_NewNode> Action = MakeShared<FObjectGraphSchemaAction_NewNode>(
 				CategoryName ? FText::FromString(*CategoryName) : FText::GetEmpty(),
 				PossibleObjectClass->GetDisplayNameText(), 
@@ -365,6 +391,12 @@ bool UObjectTreeGraphSchema::TryCreateConnection(UEdGraphPin* A, UEdGraphPin* B)
 		return false;
 	}
 
+	const bool bHandled = OnCreateConnection(A, B);
+	if (bHandled)
+	{
+		return true;
+	}
+
 	const FScopedTransaction Transaction(LOCTEXT("CreateConnection", "Create Connection"));
 
 	UObjectTreeGraphNode* NodeA = Cast<UObjectTreeGraphNode>(A->GetOwningNode());
@@ -415,16 +447,27 @@ bool UObjectTreeGraphSchema::TryCreateConnection(UEdGraphPin* A, UEdGraphPin* B)
 	return true;
 }
 
-void UObjectTreeGraphSchema::BreakPinLinks(UEdGraphPin& TargetPin, bool bSendsNodeNotifcation) const
+bool UObjectTreeGraphSchema::OnCreateConnection(UEdGraphPin* A, UEdGraphPin* B) const
 {
-	const FScopedTransaction Transaction(LOCTEXT("BreakPinLinks", "Break Pin Links"));
+	return false;
+}
 
+void UObjectTreeGraphSchema::BreakPinLinks(UEdGraphPin& TargetPin, bool bSendsNodeNotification) const
+{
 	if (TargetPin.LinkedTo.IsEmpty())
 	{
-		Super::BreakPinLinks(TargetPin, bSendsNodeNotifcation);
+		Super::BreakPinLinks(TargetPin, bSendsNodeNotification);
 		return;
 	}
 
+	const bool bHandled = OnBreakPinLinks(TargetPin, bSendsNodeNotification);
+	if (bHandled)
+	{
+		return;
+	}
+
+	const FScopedTransaction Transaction(LOCTEXT("BreakPinLinks", "Break Pin Links"));
+
 	// TargetPin could be a self pin or a property pin, we need to handle both cases and directions.
 	UEdGraphPin* PropertyPin = &TargetPin;
 	UObjectTreeGraphNode* PropertyOwningNode = Cast<UObjectTreeGraphNode>(PropertyPin->GetOwningNode());
@@ -457,7 +500,7 @@ void UObjectTreeGraphSchema::BreakPinLinks(UEdGraphPin& TargetPin, bool bSendsNo
 		bRemovePropertyPin = true;
 	}
 
-	Super::BreakPinLinks(TargetPin, bSendsNodeNotifcation);
+	Super::BreakPinLinks(TargetPin, bSendsNodeNotification);
 
 	if (bRemovePropertyPin)
 	{
@@ -466,8 +509,19 @@ void UObjectTreeGraphSchema::BreakPinLinks(UEdGraphPin& TargetPin, bool bSendsNo
 	}
 }
 
+bool UObjectTreeGraphSchema::OnBreakPinLinks(UEdGraphPin& TargetPin, bool bSendsNodeNotification) const
+{
+	return false;
+}
+
 void UObjectTreeGraphSchema::BreakSinglePinLink(UEdGraphPin* SourcePin, UEdGraphPin* TargetPin) const
 {
+	const bool bHandled = OnBreakSinglePinLink(SourcePin, TargetPin);
+	if (bHandled)
+	{
+		return;
+	}
+
 	const FScopedTransaction Transaction(LOCTEXT("BreakSinglePinLink", "Break Pin Link"));
 
 	// SourcePin could be the self-pin, and TargetPin the property pin, if the directions
@@ -512,6 +566,11 @@ void UObjectTreeGraphSchema::BreakSinglePinLink(UEdGraphPin* SourcePin, UEdGraph
 	}
 }
 
+bool UObjectTreeGraphSchema::OnBreakSinglePinLink(UEdGraphPin* SourcePin, UEdGraphPin* TargetPin) const
+{
+	return false;
+}
+
 bool UObjectTreeGraphSchema::SupportsDropPinOnNode(UEdGraphNode* InTargetNode, const FEdGraphPinType& InSourcePinType, EEdGraphPinDirection InSourcePinDirection, FText& OutErrorMessage) const
 {
 	return Super::SupportsDropPinOnNode(InTargetNode, InSourcePinType, InSourcePinDirection, OutErrorMessage);
@@ -612,6 +671,7 @@ UEdGraphNode* FObjectGraphSchemaAction_NewNode::PerformAction(UEdGraph* ParentGr
 
 		NewGraphNode->NodePosX = Location.X;
 		NewGraphNode->NodePosY = Location.Y;
+		NewGraphNode->OnGraphNodeMoved();
 
 		AutoSetupNewNode(NewGraphNode, FromPin);
 
diff --git a/Engine/Plugins/Cameras/GameplayCameras/Source/GameplayCamerasEditor/Public/Editors/ObjectTreeGraphConfig.h b/Engine/Plugins/Cameras/GameplayCameras/Source/GameplayCamerasEditor/Public/Editors/ObjectTreeGraphConfig.h
--- a/Engine/Plugins/Cameras/GameplayCameras/Source/GameplayCamerasEditor/Public/Editors/ObjectTreeGraphConfig.h
+++ b/Engine/Plugins/Cameras/GameplayCameras/Source/GameplayCamerasEditor/Public/Editors/ObjectTreeGraphConfig.h
@@ -40,6 +40,8 @@ public:
 	/** The subclass of graph nodes to create. */
 	OTGCC_FIELD(TSubclassOf<UObjectTreeGraphNode>, GraphNodeClass)
 
+	/** The name of the self pin. */
+	OTGCC_FIELD(FName, SelfPinName)
 	/** The display name of the self pin. */
 	OTGCC_FIELD(FText, SelfPinFriendlyName)
 	/** The direction of the self pin. */
@@ -172,6 +174,14 @@ public:
 	 */
 	bool IsConnectable(UClass* InObjectClass) const;
 
+	/**
+	 * Returns whether the given object reference property is connectable.
+	 *
+	 * It is connectable if the property's reference type is for a connectable class, and if the
+	 * property doesn't have the ObjectTreeGraphHidden metadata.
+	 */
+	bool IsConnectable(FObjectProperty* InObjectProperty) const;
+
 	/**
 	 * Gets all possible known connectable classes.
 	 *
diff --git a/Engine/Plugins/Cameras/GameplayCameras/Source/GameplayCamerasEditor/Public/Editors/ObjectTreeGraphNode.h b/Engine/Plugins/Cameras/GameplayCameras/Source/GameplayCamerasEditor/Public/Editors/ObjectTreeGraphNode.h
--- a/Engine/Plugins/Cameras/GameplayCameras/Source/GameplayCamerasEditor/Public/Editors/ObjectTreeGraphNode.h
+++ b/Engine/Plugins/Cameras/GameplayCameras/Source/GameplayCamerasEditor/Public/Editors/ObjectTreeGraphNode.h
@@ -32,6 +32,11 @@ public:
 
 	/** Gets the underlying object represented by this graph node. */
 	UObject* GetObject() const { return Object.Get(); }
+	/** Gets whether we have a valid underlying object, and that it's a type of ObjectClass. */
+	template<typename ObjectClass> bool IsObjectA() const;
+	/** Gets the underlying object as a point to the given sub-class. */
+	template<typename ObjectClass> ObjectClass* CastObject() const;
+
 	/** Gets all connectable properties on the underlying object. */
 	void GetAllConnectableProperties(TArray<FProperty*>& OutProperties) const;
 
@@ -112,3 +117,23 @@ private:
 	bool bOverrideSelfPinDirection;
 };
 
+template<typename ObjectClass>
+bool UObjectTreeGraphNode::IsObjectA() const
+{
+	if (Object)
+	{
+		return Object->IsA<ObjectClass>();
+	}
+	return false;
+}
+
+template<typename ObjectClass> 
+ObjectClass* UObjectTreeGraphNode::CastObject() const
+{
+	if (Object)
+	{
+		return Cast<ObjectClass>(Object);
+	}
+	return nullptr;
+}
+
diff --git a/Engine/Plugins/Cameras/GameplayCameras/Source/GameplayCamerasEditor/Public/Editors/ObjectTreeGraphSchema.h b/Engine/Plugins/Cameras/GameplayCameras/Source/GameplayCamerasEditor/Public/Editors/ObjectTreeGraphSchema.h
--- a/Engine/Plugins/Cameras/GameplayCameras/Source/GameplayCamerasEditor/Public/Editors/ObjectTreeGraphSchema.h
+++ b/Engine/Plugins/Cameras/GameplayCameras/Source/GameplayCamerasEditor/Public/Editors/ObjectTreeGraphSchema.h
@@ -60,11 +60,14 @@ public:
 	virtual FPinConnectionResponse CanCreateNewNodes(UEdGraphPin* InSourcePin) const override;
 	virtual const FPinConnectionResponse CanCreateConnection(const UEdGraphPin* A, const UEdGraphPin* B) const override;
 	virtual bool TryCreateConnection(UEdGraphPin* A, UEdGraphPin* B) const override;
-	virtual void BreakPinLinks(UEdGraphPin& TargetPin, bool bSendsNodeNotifcation) const override;
+	virtual void BreakPinLinks(UEdGraphPin& TargetPin, bool bSendsNodeNotification) const override;
 	virtual void BreakSinglePinLink(UEdGraphPin* SourcePin, UEdGraphPin* TargetPin) const override;
 	virtual bool SupportsDropPinOnNode(UEdGraphNode* InTargetNode, const FEdGraphPinType& InSourcePinType, EEdGraphPinDirection InSourcePinDirection, FText& OutErrorMessage) const override;
 
 	// UObjectTreeGraphSchema interface.
+	virtual bool OnCreateConnection(UEdGraphPin* A, UEdGraphPin* B) const;
+	virtual bool OnBreakPinLinks(UEdGraphPin& TargetPin, bool bSendsNodeNotification) const;
+	virtual bool OnBreakSinglePinLink(UEdGraphPin* SourcePin, UEdGraphPin* TargetPin) const;
 	virtual void FilterGraphContextPlaceableClasses(TArray<UClass*>& InOutClasses) const;
 	virtual void ProcessDuplicatedNodes(UObjectTreeGraph* InGraph, const TMap<UEdGraphNode*, UEdGraphNode*>& NodeMap) const;
 
