Cameras: add ability to expose camera nodes' parameters on the parent rig, and override them externally.

This makes it possible to parameterize camera rig instances, enable "prefab workflows" (i.e. nesting and re-use of camera rigs), and more.

Internally, "hidden" camera variables are allocated, and exposed camera parameters get driven by them. A new BuildCameraRig method handles both this and the previous computation of the allocation information of the rig.

diff --git a/Engine/Plugins/Cameras/GameplayCameras/Source/GameplayCameras/Private/Core/CameraRigAsset.cpp b/Engine/Plugins/Cameras/GameplayCameras/Source/GameplayCameras/Private/Core/CameraRigAsset.cpp
--- a/Engine/Plugins/Cameras/GameplayCameras/Source/GameplayCameras/Private/Core/CameraRigAsset.cpp
+++ b/Engine/Plugins/Cameras/GameplayCameras/Source/GameplayCameras/Private/Core/CameraRigAsset.cpp
@@ -3,9 +3,197 @@
 #include "Core/CameraRigAsset.h"
 
 #include "Core/CameraNode.h"
+#include "Core/CameraParameters.h"
+#include "Core/CameraRigAllocationInfoBuilder.h"
+#include "Core/CameraVariableAssets.h"
+#include "Templates/UnrealTypeTraits.h"
+#include "UObject/ObjectSaveContext.h"
 
 #include UE_INLINE_GENERATED_CPP_BY_NAME(CameraRigAsset)
 
+namespace UE::Cameras::Internal
+{
+
+template<typename VariableAssetType, typename ValueType>
+void SetPrivateVariableDefaultValue(VariableAssetType* PrivateVariable, typename TCallTraits<ValueType>::ParamType Value)
+{
+	PrivateVariable->DefaultValue = Value;
+}
+
+template<>
+void SetPrivateVariableDefaultValue<UBooleanCameraVariable, bool>(UBooleanCameraVariable* PrivateVariable, bool bValue)
+{
+	PrivateVariable->bDefaultValue = bValue;
+}
+
+template<typename CameraParameterType>
+UCameraVariableAsset* CreatePrivateVariable(UCameraRigAsset* CameraRigAsset, UCameraRigInterfaceParameter* InterfaceParameter, const FString& VariableName, CameraParameterType* CameraParameter)
+{
+	if (CameraParameter->Variable != nullptr)
+	{
+		UE_LOG(LogCameraSystem, Error, 
+				TEXT("Invalid interface parameter in camera rig '%s': '%s.%s' is already driven by variable."), 
+				*CameraRigAsset->GetPathName(), 
+				*InterfaceParameter->Target->GetName(), 
+				*InterfaceParameter->TargetPropertyName.ToString());
+		return nullptr;
+	}
+
+	using ValueType = typename CameraParameterType::ValueType;
+	using VariableAssetType = typename CameraParameterType::VariableAssetType;
+
+	VariableAssetType* PrivateVariable = NewObject<VariableAssetType>(CameraRigAsset, *VariableName);
+	PrivateVariable->bIsPrivate = true;
+	PrivateVariable->bAutoReset = false;
+	SetPrivateVariableDefaultValue<VariableAssetType, ValueType>(PrivateVariable, CameraParameter->Value);
+
+	InterfaceParameter->PrivateVariable = PrivateVariable;
+	CameraParameter->Variable = PrivateVariable;
+
+	return PrivateVariable;
+}
+
+}  // namespace UE::Cameras::Internal
+
+#if WITH_EDITOR
+
+void UCameraRigInterfaceParameter::GetGraphNodePosition(int32& NodePosX, int32& NodePosY) const
+{
+	NodePosX = GraphNodePosX;
+	NodePosY = GraphNodePosY;
+}
+
+void UCameraRigInterfaceParameter::OnGraphNodeMoved(int32 NodePosX, int32 NodePosY)
+{
+	GraphNodePosX = NodePosX;
+	GraphNodePosY = NodePosY;
+}
+
+#endif
+
+UCameraRigInterfaceParameter* FCameraRigInterface::FindInterfaceParameterByName(const FString& ParameterName) const
+{
+	const TObjectPtr<UCameraRigInterfaceParameter>* FoundItem = InterfaceParameters.FindByPredicate([&ParameterName](UCameraRigInterfaceParameter* Item)
+			{
+				return Item->InterfaceParameterName == ParameterName;
+			});
+	return FoundItem ? *FoundItem : nullptr;
+}
+
+bool FCameraRigInterface::HasInterfaceParameter(const FString& ParameterName) const
+{
+	return FindInterfaceParameterByName(ParameterName) != nullptr;
+}
+
+void UCameraRigAsset::BuildCameraRig()
+{
+	using namespace UE::Cameras;
+	using namespace UE::Cameras::Internal;
+
+	// Build allocation info.
+	FCameraRigAllocationInfoBuilder CameraRigBuilder;
+	CameraRigBuilder.BuildAllocationInfo(this, AllocationInfo);
+	BuildStatus = ECameraRigBuildStatus::Clean;
+
+	// Build interface parameters.
+	TSet<FString> UsedInterfaceParameterNames;
+	for (UCameraRigInterfaceParameter* InterfaceParameter : Interface.InterfaceParameters)
+	{
+		// Do some basic validation.
+		if (!InterfaceParameter || !InterfaceParameter->Target)
+		{
+			UE_LOG(LogCameraSystem, Error, 
+					TEXT("Invalid interface parameter target in camera rig: '%s'."),
+					*GetPathName());
+			continue;
+		}
+		if (InterfaceParameter->TargetPropertyName.IsNone())
+		{
+			UE_LOG(LogCameraSystem, Error, 
+					TEXT("Invalid interface parameter target property name in camera rig: '%s'."), 
+					*GetPathName());
+			continue;
+		}
+		if (InterfaceParameter->InterfaceParameterName.IsEmpty())
+		{
+			UE_LOG(LogCameraSystem, Error, 
+					TEXT("Invalid interface parameter name in camera rig '%s'."), 
+					*GetPathName());
+			continue;
+		}
+
+		if (UsedInterfaceParameterNames.Contains(InterfaceParameter->InterfaceParameterName))
+		{
+			UE_LOG(LogCameraSystem, Error, 
+					TEXT("Multiple interface parameters named '%s' in camera rig '%s'. Ignoring duplicates."), 
+					*InterfaceParameter->InterfaceParameterName, 
+					*GetPathName());
+			continue;
+		}
+		UsedInterfaceParameterNames.Add(InterfaceParameter->InterfaceParameterName);
+
+		// Get the target camera node property and check that it is a camera parameter struct.
+		UClass* TargetClass = InterfaceParameter->Target->GetClass();
+		FProperty* TargetProperty = TargetClass->FindPropertyByName(InterfaceParameter->TargetPropertyName);
+		if (!TargetProperty)
+		{
+			UE_LOG(LogCameraSystem, Error, 
+					TEXT("Invalid interface parameter in camera rig '%s': no property '%s' on camera node '%s'."), 
+					*GetPathName(), 
+					*InterfaceParameter->InterfaceParameterName, 
+					*InterfaceParameter->Target->GetName());
+			continue;
+		}
+
+		FStructProperty* TargetStructProperty = CastField<FStructProperty>(TargetProperty);
+		if (!TargetStructProperty)
+		{
+			UE_LOG(LogCameraSystem, Error, 
+					TEXT("Invalid interface parameter in camera rig '%s': property '%s' on camera node '%s' is not a camera parameter."),
+					*GetPathName(), 
+					*InterfaceParameter->InterfaceParameterName, 
+					*InterfaceParameter->Target->GetName());
+			continue;
+		}
+
+		// Get the type of the camera parameter by matching the struct against all the types we support,
+		// and create a private camera variable asset to drive its value.
+		FString PrivateVariableName(FString::Format(TEXT("Override_{0}_{1}"), { GetName(), InterfaceParameter->InterfaceParameterName }));
+#define UE_CAMERA_VARIABLE_FOR_TYPE(ValueType, ValueName)\
+		if (TargetStructProperty->Struct == F##ValueName##CameraParameter::StaticStruct())\
+		{\
+			F##ValueName##CameraParameter CameraParameter;\
+			TargetStructProperty->GetValue_InContainer(InterfaceParameter->Target, reinterpret_cast<void*>(&CameraParameter));\
+			CreatePrivateVariable(this, InterfaceParameter, PrivateVariableName, &CameraParameter);\
+			TargetStructProperty->SetValue_InContainer(InterfaceParameter->Target, reinterpret_cast<void*>(&CameraParameter));\
+		}\
+		else
+UE_CAMERA_VARIABLE_FOR_ALL_TYPES()
+#undef UE_CAMERA_VARIABLE_FOR_TYPE
+		{
+			UE_LOG(LogCameraSystem, Error, 
+					TEXT("Invalid interface parameter in camera rig '%s': property '%s' on camera node '%s' is not a camera parameter."), 
+					*GetPathName(), 
+					*InterfaceParameter->InterfaceParameterName,
+					*InterfaceParameter->Target->GetName());
+			continue;
+		}
+	}
+}
+
+void UCameraRigAsset::PreSave(FObjectPreSaveContext ObjectSaveContext)
+{
+#if WITH_EDITOR
+	if (!HasAnyFlags(RF_ClassDefaultObject | RF_ArchetypeObject))
+	{
+		// Build on save.
+		BuildCameraRig();
+	}
+#endif
+
+	Super::PreSave(ObjectSaveContext);
+}
+
 #if WITH_EDITOR
 
 void UCameraRigAsset::GatherPackages(FCameraRigPackages& OutPackages) const
diff --git a/Engine/Plugins/Cameras/GameplayCameras/Source/GameplayCameras/Public/Core/CameraParameters.h b/Engine/Plugins/Cameras/GameplayCameras/Source/GameplayCameras/Public/Core/CameraParameters.h
--- a/Engine/Plugins/Cameras/GameplayCameras/Source/GameplayCameras/Public/Core/CameraParameters.h
+++ b/Engine/Plugins/Cameras/GameplayCameras/Source/GameplayCameras/Public/Core/CameraParameters.h
@@ -15,6 +15,7 @@ struct GAMEPLAYCAMERAS_API FBooleanCameraParameter
 	GENERATED_BODY()
 
 	using ValueType = bool;
+	using VariableAssetType = UBooleanCameraVariable;
 
 	UPROPERTY(EditAnywhere, Category=Common)
 	bool Value = false;
@@ -32,6 +33,7 @@ struct GAMEPLAYCAMERAS_API FInteger32CameraParameter
 	GENERATED_BODY()
 
 	using ValueType = int32;
+	using VariableAssetType = UInteger32CameraVariable;
 
 	UPROPERTY(EditAnywhere, Category=Common)
 	int32 Value = 0;
@@ -49,6 +51,7 @@ struct FFloatCameraParameter
 	GENERATED_BODY()
 
 	using ValueType = float;
+	using VariableAssetType = UFloatCameraVariable;
 
 	UPROPERTY(EditAnywhere, Category=Common)
 	float Value = 0.f;
@@ -66,6 +69,7 @@ struct GAMEPLAYCAMERAS_API FDoubleCameraParameter
 	GENERATED_BODY()
 
 	using ValueType = double;
+	using VariableAssetType = UDoubleCameraVariable;
 
 	UPROPERTY(EditAnywhere, Category=Common)
 	double Value = 0.0;
@@ -83,6 +87,7 @@ struct FVector2fCameraParameter
 	GENERATED_BODY()
 
 	using ValueType = FVector2f;
+	using VariableAssetType = UVector2fCameraVariable;
 
 	UPROPERTY(EditAnywhere, Category=Common)
 	FVector2f Value;
@@ -101,6 +106,7 @@ struct GAMEPLAYCAMERAS_API FVector2dCameraParameter
 	GENERATED_BODY()
 
 	using ValueType = FVector2D;
+	using VariableAssetType = UVector2dCameraVariable;
 
 	UPROPERTY(EditAnywhere, Category=Common)
 	FVector2D Value;
@@ -119,6 +125,7 @@ struct FVector3fCameraParameter
 	GENERATED_BODY()
 
 	using ValueType = FVector3f;
+	using VariableAssetType = UVector3fCameraVariable;
 
 	UPROPERTY(EditAnywhere, Category=Common)
 	FVector3f Value;
@@ -137,6 +144,7 @@ struct GAMEPLAYCAMERAS_API FVector3dCameraParameter
 	GENERATED_BODY()
 
 	using ValueType = FVector3d;
+	using VariableAssetType = UVector3dCameraVariable;
 
 	UPROPERTY(EditAnywhere, Category=Common)
 	FVector3d Value;
@@ -155,6 +163,7 @@ struct FVector4fCameraParameter
 	GENERATED_BODY()
 
 	using ValueType = FVector4f;
+	using VariableAssetType = UVector4fCameraVariable;
 
 	UPROPERTY(EditAnywhere, Category=Common)
 	FVector4f Value;
@@ -173,6 +182,7 @@ struct FVector4dCameraParameter
 	GENERATED_BODY()
 
 	using ValueType = FVector4d;
+	using VariableAssetType = UVector4dCameraVariable;
 
 	UPROPERTY(EditAnywhere, Category=Common)
 	FVector4d Value;
@@ -191,6 +201,7 @@ struct GAMEPLAYCAMERAS_API FRotator3fCameraParameter
 	GENERATED_BODY()
 
 	using ValueType = FRotator3f;
+	using VariableAssetType = URotator3fCameraVariable;
 
 	UPROPERTY(EditAnywhere, Category=Common)
 	FRotator3f Value;
@@ -209,6 +220,7 @@ struct GAMEPLAYCAMERAS_API FRotator3dCameraParameter
 	GENERATED_BODY()
 
 	using ValueType = FRotator3d;
+	using VariableAssetType = URotator3dCameraVariable;
 
 	UPROPERTY(EditAnywhere, Category=Common)
 	FRotator3d Value;
@@ -227,6 +239,7 @@ struct FTransform3fCameraParameter
 	GENERATED_BODY()
 
 	using ValueType = FTransform3f;
+	using VariableAssetType = UTransform3fCameraVariable;
 
 	UPROPERTY(EditAnywhere, Category=Common)
 	FTransform3f Value;
@@ -244,6 +257,7 @@ struct GAMEPLAYCAMERAS_API FTransform3dCameraParameter
 	GENERATED_BODY()
 
 	using ValueType = FTransform3d;
+	using VariableAssetType = UTransform3dCameraVariable;
 
 	UPROPERTY(EditAnywhere, Category=Common)
 	FTransform3d Value;
diff --git a/Engine/Plugins/Cameras/GameplayCameras/Source/GameplayCameras/Public/Core/CameraRigAsset.h b/Engine/Plugins/Cameras/GameplayCameras/Source/GameplayCameras/Public/Core/CameraRigAsset.h
--- a/Engine/Plugins/Cameras/GameplayCameras/Source/GameplayCameras/Public/Core/CameraRigAsset.h
+++ b/Engine/Plugins/Cameras/GameplayCameras/Source/GameplayCameras/Public/Core/CameraRigAsset.h
@@ -12,6 +12,7 @@
 #include "CameraRigAsset.generated.h"
 
 class UCameraNode;
+class UCameraVariableAsset;
 
 /**
  *
@@ -42,6 +43,85 @@ struct FCameraRigAllocationInfo
 	FCameraVariableTableAllocationInfo VariableTableInfo;
 };
 
+/**
+ * An exposed camera rig parameter that drives a specific parameter on one of
+ * its camera nodes.
+ */
+UCLASS(MinimalAPI)
+class UCameraRigInterfaceParameter
+	: public UObject
+	, public IObjectTreeGraphObject
+{
+	GENERATED_BODY()
+
+public:
+
+	/** The camera node that this parameter drives. */
+	UPROPERTY(meta=(ObjectTreeGraphHidden=true))
+	TObjectPtr<UCameraNode> Target;
+
+	/** The camera parameter on the target camera node that this parameter drives. */
+	UPROPERTY()
+	FName TargetPropertyName;
+
+	/** The exposed name for this parameter. */
+	UPROPERTY(EditAnywhere, Category=Camera)
+	FString InterfaceParameterName;
+
+	// Built on save/cook.
+
+	/**
+	 * The private camera variable created to drive the target camera parameter on
+	 * the target camera node. This variable is created by the build method on the
+	 * camera rig.
+	 */
+	UPROPERTY()
+	TObjectPtr<UCameraVariableAsset> PrivateVariable;
+
+protected:
+
+	// IObjectTreeGraphObject interface.
+#if WITH_EDITOR
+	virtual void GetGraphNodePosition(int32& NodePosX, int32& NodePosY) const override;
+	virtual void OnGraphNodeMoved(int32 NodePosX, int32 NodePosY) override;
+#endif
+
+private:
+
+#if WITH_EDITORONLY_DATA
+
+	UPROPERTY()
+	int32 GraphNodePosX = 0;
+
+	UPROPERTY()
+	int32 GraphNodePosY = 0;
+
+#endif  // WITH_EDITORONLY_DATA
+};
+
+/**
+ * Structure defining the public data interface of a camera rig asset.
+ */
+USTRUCT()
+struct FCameraRigInterface
+{
+	GENERATED_BODY()
+
+public:
+
+	/** The list of exposed parameters on the camera rig. */
+	UPROPERTY()
+	TArray<TObjectPtr<UCameraRigInterfaceParameter>> InterfaceParameters;
+
+public:
+	
+	/** Finds an exposed parameter by name. */
+	GAMEPLAYCAMERAS_API UCameraRigInterfaceParameter* FindInterfaceParameterByName(const FString& ParameterName) const;
+
+	/** Returns whether an exposed parameter with the given name exists. */
+	GAMEPLAYCAMERAS_API bool HasInterfaceParameter(const FString& ParameterName) const;
+};
+
 /**
  * List of packages that contain the definition of a camera rig.
  * In most cases there's only one, but with nested assets there could be more.
@@ -71,6 +151,10 @@ public:
 	UPROPERTY(EditAnywhere, Instanced, Category=Common)
 	TObjectPtr<UCameraNode> RootNode;
 
+	/** The public data interface of this camera rig. */
+	UPROPERTY()
+	FCameraRigInterface Interface;
+
 	/** List of enter transitions for this camera rig. */
 	UPROPERTY(EditAnywhere, Category=Blending)
 	TArray<TObjectPtr<UCameraRigTransition>> EnterTransitions;
@@ -79,14 +163,23 @@ public:
 	UPROPERTY(EditAnywhere, Category=Blending)
 	TArray<TObjectPtr<UCameraRigTransition>> ExitTransitions;
 
+	/** Allocation information for all the nodes and variables in this camera rig. */
 	UPROPERTY()
 	FCameraRigAllocationInfo AllocationInfo;
 
 public:
 
+	/** The current build state of this camera rig. */
 	UPROPERTY(Transient)
 	ECameraRigBuildStatus BuildStatus = ECameraRigBuildStatus::Dirty;
 
+	/**
+	 * Builds this camera rig.
+	 * This will validate the data, build the allocation info, and create internal
+	 * camera variables for any exposed parameters.
+	 */
+	GAMEPLAYCAMERAS_API void BuildCameraRig();
+
 protected:
 
 	// IObjectTreeGraphObject interface.
@@ -98,6 +191,9 @@ protected:
 	virtual void OnUpdateGraphNodeCommentText(const FString& NewComment) override;
 #endif
 
+	// UObject interface
+	virtual void PreSave(FObjectPreSaveContext ObjectSaveContext) override;
+
 private:
 
 #if WITH_EDITORONLY_DATA
