Cameras: add debug drawing helper functions on the camera debug renderer.

diff --git a/Engine/Plugins/Cameras/GameplayCameras/Source/GameplayCameras/Private/Debug/CameraDebugRenderer.cpp b/Engine/Plugins/Cameras/GameplayCameras/Source/GameplayCameras/Private/Debug/CameraDebugRenderer.cpp
--- a/Engine/Plugins/Cameras/GameplayCameras/Source/GameplayCameras/Private/Debug/CameraDebugRenderer.cpp
+++ b/Engine/Plugins/Cameras/GameplayCameras/Source/GameplayCameras/Private/Debug/CameraDebugRenderer.cpp
@@ -10,6 +10,7 @@
 #include "Engine/Engine.h"
 #include "Engine/Font.h"
 #include "HAL/IConsoleManager.h"
+#include "Math/Box2D.h"
 #include "Misc/TVariant.h"
 
 #if UE_GAMEPLAY_CAMERAS_DEBUG
@@ -202,6 +203,60 @@ void FCameraDebugRenderer::DrawTextBackgroundTile(float Opacity)
 	}
 }
 
+void FCameraDebugRenderer::Draw2DLine(const FVector2D& Start, const FVector2D& End, const FLinearColor& LineColor, float LineThickness)
+{
+	if (Canvas)
+	{
+		FCanvasLineItem LineItem(Start, End);
+		LineItem.SetColor(LineColor);
+		LineItem.LineThickness = LineThickness;
+		Canvas->DrawItem(LineItem);
+	}
+}
+
+void FCameraDebugRenderer::Draw2DBox(const FBox2D& Box, const FLinearColor& LineColor, float LineThickness)
+{
+	if (Canvas)
+	{
+		FCanvasBoxItem BoxItem(Box.Min, Box.GetSize());
+		BoxItem.SetColor(LineColor);
+		BoxItem.LineThickness = LineThickness;
+		Canvas->DrawItem(BoxItem);
+	}
+}
+
+void FCameraDebugRenderer::Draw2DBox(const FVector2D& BoxPosition, const FVector2D& BoxSize, const FLinearColor& LineColor, float LineThickness)
+{
+	if (Canvas)
+	{
+		FCanvasBoxItem BoxItem(BoxPosition, BoxSize);
+		BoxItem.SetColor(LineColor);
+		BoxItem.LineThickness = LineThickness;
+		Canvas->DrawItem(BoxItem);
+	}
+}
+
+void FCameraDebugRenderer::Draw2DCircle(const FVector2D& Center, float Radius, const FLinearColor& LineColor, float LineThickness, int32 NumSides)
+{
+	if (NumSides <= 0)
+	{
+		NumSides = FMath::Max(6, (int)Radius / 25);
+	}
+
+	const float	AngleDelta = 2.0f * UE_PI / NumSides;
+	const FVector2D AxisX(1.f, 0.f);
+	const FVector2D AxisY(0.f, -1.f);
+	FVector2D LastVertex = Center + AxisX * Radius;
+
+	for (int32 SideIndex = 0; SideIndex < NumSides; SideIndex++)
+	{
+		const float CurAngle = AngleDelta * (SideIndex + 1);
+		const FVector2D Vertex = Center + (AxisX * FMath::Cos(CurAngle) + AxisY * FMath::Sin(CurAngle)) * Radius;
+		Draw2DLine(LastVertex, Vertex, LineColor, LineThickness);
+		LastVertex = Vertex;
+	}
+}
+
 void FCameraDebugRenderer::SkipAttachedBlocks()
 {
 	VisitFlags |= ECameraDebugDrawVisitFlags::SkipAttachedBlocks;
diff --git a/Engine/Plugins/Cameras/GameplayCameras/Source/GameplayCameras/Public/Debug/CameraDebugRenderer.h b/Engine/Plugins/Cameras/GameplayCameras/Source/GameplayCameras/Public/Debug/CameraDebugRenderer.h
--- a/Engine/Plugins/Cameras/GameplayCameras/Source/GameplayCameras/Public/Debug/CameraDebugRenderer.h
+++ b/Engine/Plugins/Cameras/GameplayCameras/Source/GameplayCameras/Public/Debug/CameraDebugRenderer.h
@@ -5,6 +5,7 @@
 #include "CoreTypes.h"
 #include "GameplayCameras.h"
 #include "Math/Color.h"
+#include "Math/MathFwd.h"
 #include "Math/Vector.h"
 #include "Misc/StringBuilder.h"
 
@@ -62,6 +63,17 @@ public:
 	/** Draws a translucent background behind the text. */
 	GAMEPLAYCAMERAS_API void DrawTextBackgroundTile(float Opacity);
 
+public:
+
+	/** Draws a 2D line. */
+	GAMEPLAYCAMERAS_API void Draw2DLine(const FVector2D& Start, const FVector2D& End, const FLinearColor& LineColor, float LineThickness = 1.f);
+	/** Draws a 2D box. */
+	GAMEPLAYCAMERAS_API void Draw2DBox(const FBox2D& Box, const FLinearColor& LineColor, float LineThickness = 1.f);
+	/** Draws a 2D box. */
+	GAMEPLAYCAMERAS_API void Draw2DBox(const FVector2D& BoxPosition, const FVector2D& BoxSize, const FLinearColor& LineColor, float LineThickness = 1.f);
+	/** Draws a 2D circle. */
+	GAMEPLAYCAMERAS_API void Draw2DCircle(const FVector2D& Center, float Radius, const FLinearColor& LineColor, float LineThickness = 1.f, int32 NumSides = 0);
+
 public:
 
 	/** Request skipping drawing any blocks attached to the current block. */
@@ -83,6 +95,8 @@ public:
 	/** Gets the size of the canvas. */
 	FVector2D GetCanvasSize() const;
 
+	bool HasCanvas() const { return Canvas != nullptr; }
+
 private:
 
 	void AddTextFmtImpl(const TCHAR* Fmt, va_list Args);
