Cameras: add serialization to critical damper.

This makes it possible to save a damper's state in a debug block.

diff --git a/Engine/Plugins/Cameras/GameplayCameras/Source/GameplayCameras/Private/Math/CriticalDamper.cpp b/Engine/Plugins/Cameras/GameplayCameras/Source/GameplayCameras/Private/Math/CriticalDamper.cpp
--- a/Engine/Plugins/Cameras/GameplayCameras/Source/GameplayCameras/Private/Math/CriticalDamper.cpp
+++ b/Engine/Plugins/Cameras/GameplayCameras/Source/GameplayCameras/Private/Math/CriticalDamper.cpp
@@ -248,5 +248,18 @@ void FCriticalDamper::InternalUpdate(float ForcedMovement, float DeltaTime)
 	}
 }
 
+void FCriticalDamper::Serialize(FArchive& Ar)
+{
+	Ar << W0;
+	Ar << X0;
+	Ar << X0Derivative;
+}
+
+FArchive& operator <<(FArchive& Ar, FCriticalDamper& Damper)
+{
+	Damper.Serialize(Ar);
+	return Ar;
+}
+
 }  // namespace UE::Cameras
 
diff --git a/Engine/Plugins/Cameras/GameplayCameras/Source/GameplayCameras/Public/Math/CriticalDamper.h b/Engine/Plugins/Cameras/GameplayCameras/Source/GameplayCameras/Public/Math/CriticalDamper.h
--- a/Engine/Plugins/Cameras/GameplayCameras/Source/GameplayCameras/Public/Math/CriticalDamper.h
+++ b/Engine/Plugins/Cameras/GameplayCameras/Source/GameplayCameras/Public/Math/CriticalDamper.h
@@ -4,6 +4,8 @@
 
 #include "CoreTypes.h"
 
+class FArchive;
+
 namespace UE::Cameras
 {
 
@@ -93,6 +95,9 @@ public:
 	 */
 	float Update(float PreviousDamped, float NextTarget, float DeltaTime);
 
+	/** Saves or loads the state of this critical damper. */
+	void Serialize(FArchive& Ar);
+
 private:
 
 	void InternalUpdate(float ForcedMovement, float DeltaTime);
@@ -105,7 +110,11 @@ private:
 	float X0;
 	/** Initial velocity of the spring-mass system, aka the previous frame's velocity. */
 	float X0Derivative;
+
+	friend FArchive& operator <<(FArchive& Ar, FCriticalDamper& Damper);
 };
 
+FArchive& operator <<(FArchive& Ar, FCriticalDamper& Damper);
+
 }  // namespace UE::Cameras
 
