Cameras: add rig parameters to the camera rig graph editor

- Add custom schema class for adding rig parameter nodes to the graph.
- Add custom nodes and node widget classes for this rig parameter nodes.
- Add images and styles for these new node widgets.
- Specify a custom class for camera node graph nodes, so that we can show input pins for any camera parameter property.
- When clicking the Build button, run the new BuildCameraRig method that handles all the new stuff.

diff --git a/Engine/Plugins/Cameras/GameplayCameras/Source/GameplayCamerasEditor/Private/Editors/CameraNodeGraphNode.cpp b/Engine/Plugins/Cameras/GameplayCameras/Source/GameplayCamerasEditor/Private/Editors/CameraNodeGraphNode.cpp
new file mode 100644
--- /dev/null
+++ b/Engine/Plugins/Cameras/GameplayCameras/Source/GameplayCamerasEditor/Private/Editors/CameraNodeGraphNode.cpp
@@ -0,0 +1,77 @@
+// Copyright Epic Games, Inc. All Rights Reserved.
+
+#include "Editors/CameraNodeGraphNode.h"
+
+#include "Core/CameraParameters.h"
+#include "EdGraph/EdGraphPin.h"
+#include "Editors/CameraNodeGraphSchema.h"
+#include "ScopedTransaction.h"
+#include "ToolMenus.h"
+
+UCameraNodeGraphNode::UCameraNodeGraphNode(const FObjectInitializer& ObjInit)
+	: UObjectTreeGraphNode(ObjInit)
+{
+}
+
+void UCameraNodeGraphNode::AllocateDefaultPins()
+{
+	UObjectTreeGraphNode::AllocateDefaultPins();
+
+	// Add extra input pins for any camera parameter.
+	UClass* CameraNodeClass = GetObject()->GetClass();
+	for (TFieldIterator<FProperty> PropertyIt(CameraNodeClass); PropertyIt; ++PropertyIt)
+	{
+		FStructProperty* StructProperty = CastField<FStructProperty>(*PropertyIt);
+		if (!StructProperty)
+		{
+			continue;
+		}
+
+		const FName PropertyName = PropertyIt->GetFName();
+		
+		FEdGraphPinType PinType;
+		PinType.PinCategory = UCameraNodeGraphSchema::PC_CameraParameter;
+
+#define UE_CAMERA_VARIABLE_FOR_TYPE(ValueType, ValueName)\
+		if (StructProperty->Struct == F##ValueName##CameraParameter::StaticStruct())\
+		{\
+			UEdGraphPin* ParameterPin = CreatePin(EGPD_Input, PinType, PropertyName);\
+			ParameterPin->PinFriendlyName = FText::FromName(PropertyName);\
+			continue;\
+		}
+UE_CAMERA_VARIABLE_FOR_ALL_TYPES()
+#undef UE_CAMERA_VARIABLE_FOR_TYPE
+	}
+}
+
+UEdGraphPin* UCameraNodeGraphNode::GetPinForCameraParameterProperty(const FName& InPropertyName) const
+{
+	UEdGraphPin* const* FoundItem = Pins.FindByPredicate(
+			[InPropertyName](UEdGraphPin* Item)
+			{ 
+				return Item->PinType.PinCategory == UCameraNodeGraphSchema::PC_CameraParameter &&
+					Item->GetFName() == InPropertyName;
+			});
+	if (FoundItem)
+	{
+		return *FoundItem;
+	}
+	return nullptr;
+}
+
+FStructProperty* UCameraNodeGraphNode::GetCameraParameterPropertyForPin(const UEdGraphPin* InPin) const
+{
+	UObject* CameraNode = GetObject();
+	if (!ensure(CameraNode))
+	{
+		return nullptr;
+	}
+	if (InPin->PinType.PinCategory != UCameraNodeGraphSchema::PC_CameraParameter)
+	{
+		return nullptr;
+	}
+
+	UClass* CameraNodeClass = CameraNode->GetClass();
+	return CastField<FStructProperty>(CameraNodeClass->FindPropertyByName(InPin->GetFName()));
+}
+
diff --git a/Engine/Plugins/Cameras/GameplayCameras/Source/GameplayCamerasEditor/Private/Editors/CameraNodeGraphSchema.cpp b/Engine/Plugins/Cameras/GameplayCameras/Source/GameplayCamerasEditor/Private/Editors/CameraNodeGraphSchema.cpp
new file mode 100644
--- /dev/null
+++ b/Engine/Plugins/Cameras/GameplayCameras/Source/GameplayCamerasEditor/Private/Editors/CameraNodeGraphSchema.cpp
@@ -0,0 +1,268 @@
+// Copyright Epic Games, Inc. All Rights Reserved.
+
+#include "Editors/CameraNodeGraphSchema.h"
+
+#include "Core/CameraRigAsset.h"
+#include "Core/CameraNode.h"
+#include "EdGraph/EdGraphPin.h"
+#include "Editors/CameraNodeGraphNode.h"
+#include "Editors/ObjectTreeGraph.h"
+
+#include "ScopedTransaction.h"
+
+#include UE_INLINE_GENERATED_CPP_BY_NAME(CameraNodeGraphSchema)
+
+#define LOCTEXT_NAMESPACE "CameraNodeGraphSchema"
+
+const FName UCameraNodeGraphSchema::PC_CameraParameter("CameraParameter");
+
+void UCameraNodeGraphSchema::OnCreateAllNodes(UObjectTreeGraph* InGraph, const FCreatedNodes& InCreatedNodes) const
+{
+	Super::OnCreateAllNodes(InGraph, InCreatedNodes);
+
+	UCameraRigAsset* CameraRigAsset = Cast<UCameraRigAsset>(InGraph->GetRootObject());
+	if (ensure(CameraRigAsset))
+	{
+		for (UCameraRigInterfaceParameter* InterfaceParameter : CameraRigAsset->Interface.InterfaceParameters)
+		{
+			UObjectTreeGraphNode* const* InterfaceParameterNode = InCreatedNodes.CreatedNodes.Find(InterfaceParameter);
+			UObjectTreeGraphNode* const* CameraNodeNode = InCreatedNodes.CreatedNodes.Find(InterfaceParameter->Target);
+			if (InterfaceParameterNode && CameraNodeNode)
+			{
+				UEdGraphPin* InterfaceParameterSelfPin = (*InterfaceParameterNode)->GetSelfPin();
+				UEdGraphPin* CameraParameterPin = Cast<UCameraNodeGraphNode>(*CameraNodeNode)->GetPinForCameraParameterProperty(InterfaceParameter->TargetPropertyName);
+				InterfaceParameterSelfPin->MakeLinkTo(CameraParameterPin);
+			}
+		}
+	}
+}
+
+void UCameraNodeGraphSchema::GetGraphContextActions(FGraphContextMenuBuilder& ContextMenuBuilder) const
+{
+	// See if we were dragging a camera parameter pin.
+	if (const UEdGraphPin* DraggedPin = ContextMenuBuilder.FromPin)
+	{
+		if (DraggedPin->PinType.PinCategory == PC_CameraParameter)
+		{
+			UCameraNodeGraphNode* CameraNodeNode = Cast<UCameraNodeGraphNode>(DraggedPin->GetOwningNode());
+			FStructProperty* StructProperty = CameraNodeNode->GetCameraParameterPropertyForPin(DraggedPin);
+
+			TSharedRef<FCameraNodeGraphSchemaAction_NewInterfaceParameterNode> Action = 
+				MakeShared<FCameraNodeGraphSchemaAction_NewInterfaceParameterNode>(
+						FText::GetEmpty(),
+						LOCTEXT("NewInterfaceParameterAction", "Camera Rig Parameter"),
+						LOCTEXT("NewInterfaceParameterActionToolTip", "Exposes this parameter on the camera rig"));
+			Action->Target = Cast<UCameraNode>(CameraNodeNode->GetObject());
+			Action->TargetPropertyName = StructProperty->GetFName();
+			ContextMenuBuilder.AddAction(StaticCastSharedPtr<FEdGraphSchemaAction>(Action.ToSharedPtr()));
+
+			return;
+		}
+	}
+
+	Super::GetGraphContextActions(ContextMenuBuilder);
+}
+
+const FPinConnectionResponse UCameraNodeGraphSchema::CanCreateConnection(const UEdGraphPin* A, const UEdGraphPin* B) const
+{
+	if (A->PinType.PinCategory == PC_CameraParameter && B->PinType.PinCategory == PC_Self)
+	{
+		UObjectTreeGraphNode* NodeB = Cast<UObjectTreeGraphNode>(B->GetOwningNode());
+		if (NodeB && NodeB->IsObjectA<UCameraRigInterfaceParameter>())
+		{
+			return FPinConnectionResponse(CONNECT_RESPONSE_BREAK_OTHERS_AB, TEXT("Compatible pin types"));
+		}
+	}
+	else if (A->PinType.PinCategory == PC_Self && B->PinType.PinCategory == PC_CameraParameter)
+	{
+		UObjectTreeGraphNode* NodeA = Cast<UObjectTreeGraphNode>(A->GetOwningNode());
+		if (NodeA && NodeA->IsObjectA<UCameraRigInterfaceParameter>())
+		{
+			return FPinConnectionResponse(CONNECT_RESPONSE_BREAK_OTHERS_AB, TEXT("Compatible pin types"));
+		}
+	}
+
+	return Super::CanCreateConnection(A, B);
+}
+
+bool UCameraNodeGraphSchema::OnCreateConnection(UEdGraphPin* A, UEdGraphPin* B) const
+{
+	// Try to make a connection between a camera node's parameter pin and a camera rig interface parameter.
+	// First, figure out which is which.
+	UEdGraphPin* RigInterfacePin = nullptr;
+	UEdGraphPin* CameraParameterPin = nullptr;
+
+	if (A->PinType.PinCategory == PC_CameraParameter && B->PinType.PinCategory == PC_Self)
+	{
+		RigInterfacePin = B;
+		CameraParameterPin = A;
+	}
+	else if (A->PinType.PinCategory == PC_Self && B->PinType.PinCategory == PC_CameraParameter)
+	{
+		RigInterfacePin = A;
+		CameraParameterPin = B;
+	}
+
+	if (!RigInterfacePin && !CameraParameterPin)
+	{
+		return false;
+	}
+
+	// Now make sure both nodes are what we expect, and that they have what we need.
+	UObjectTreeGraphNode* RigParameterNode = Cast<UObjectTreeGraphNode>(RigInterfacePin->GetOwningNode());
+	if (!RigParameterNode)
+	{
+		return false;
+	}
+	UCameraRigInterfaceParameter* RigParameter = RigParameterNode->CastObject<UCameraRigInterfaceParameter>();
+	if (!RigParameter)
+	{
+		return false;
+	}
+
+	UCameraNodeGraphNode* CameraNodeNode = Cast<UCameraNodeGraphNode>(CameraParameterPin->GetOwningNode());
+	if (!CameraNodeNode)
+	{
+		return false;
+	}
+	UCameraNode* CameraNode = CameraNodeNode->CastObject<UCameraNode>();
+	if (!CameraNode)
+	{
+		return false;
+	}
+	FStructProperty* StructProperty = CameraNodeNode->GetCameraParameterPropertyForPin(CameraParameterPin);
+	if (!StructProperty)
+	{
+		return false;
+	}
+
+	// Make the connection.
+	const FScopedTransaction Transaction(LOCTEXT("ExposeCameraRigParameter", "Expose Camera Rig Parameter"));
+
+	RigParameter->Modify();
+
+	RigParameter->Target = CameraNode;
+	RigParameter->TargetPropertyName = StructProperty->GetFName();
+	if (RigParameter->InterfaceParameterName.IsEmpty())
+	{
+		RigParameter->InterfaceParameterName = StructProperty->GetName();
+	}
+
+	return true;
+}
+
+bool UCameraNodeGraphSchema::OnBreakPinLinks(UEdGraphPin& TargetPin, bool bSendsNodeNotification) const
+{
+	// See if we have a rig parameter connection to break.
+	if (TargetPin.PinType.PinCategory == PC_Self || TargetPin.PinType.PinCategory == PC_CameraParameter)
+	{
+		UEdGraphPin* RigParameterSelfPin = &TargetPin;
+		if (TargetPin.PinType.PinCategory == PC_CameraParameter)
+		{
+			RigParameterSelfPin = TargetPin.LinkedTo[0];
+		}
+
+		UObjectTreeGraphNode* RigParameterNode = Cast<UObjectTreeGraphNode>(RigParameterSelfPin->GetOwningNode());
+		if (RigParameterNode && RigParameterNode->IsObjectA<UCameraRigInterfaceParameter>())
+		{
+			const FScopedTransaction Transaction(LOCTEXT("BreakPinLinks", "Break Pin Links"));
+
+			UCameraRigInterfaceParameter* RigParameter = RigParameterNode->CastObject<UCameraRigInterfaceParameter>();
+
+			RigParameter->Modify();
+
+			RigParameter->Target = nullptr;
+			RigParameter->TargetPropertyName = NAME_None;
+
+			// Skip our parent, we just want the actual graph-change only.
+			UEdGraphSchema::BreakPinLinks(TargetPin, bSendsNodeNotification);
+
+			return true;
+		}
+	}
+
+	return false;
+}
+
+bool UCameraNodeGraphSchema::OnBreakSinglePinLink(UEdGraphPin* SourcePin, UEdGraphPin* TargetPin) const
+{
+	UCameraRigInterfaceParameter* RigParameter = nullptr;
+
+	if (SourcePin->PinType.PinCategory == PC_Self && TargetPin->PinType.PinCategory == PC_CameraParameter)
+	{
+		UObjectTreeGraphNode* SourceNode = Cast<UObjectTreeGraphNode>(SourcePin->GetOwningNode());
+		if (SourceNode)
+		{
+			RigParameter = SourceNode->CastObject<UCameraRigInterfaceParameter>();
+		}
+	}
+	else if (SourcePin->PinType.PinCategory == PC_CameraParameter && TargetPin->PinType.PinCategory == PC_Self)
+	{
+		UObjectTreeGraphNode* TargetNode = Cast<UObjectTreeGraphNode>(TargetPin->GetOwningNode());
+		if (TargetNode)
+		{ 
+			RigParameter = TargetNode->CastObject<UCameraRigInterfaceParameter>();
+		}
+	}
+
+	if (RigParameter)
+	{
+		const FScopedTransaction Transaction(LOCTEXT("BreakSinglePinLink", "Break Pin Link"));
+
+		RigParameter->Modify();
+
+		RigParameter->Target = nullptr;
+		RigParameter->TargetPropertyName = NAME_None;
+
+		// Skip our parent, we just want the actual graph-change only.
+		UEdGraphSchema::BreakSinglePinLink(SourcePin, TargetPin);
+
+		return true;
+	}
+	
+	return false;
+}
+
+FCameraNodeGraphSchemaAction_NewInterfaceParameterNode::FCameraNodeGraphSchemaAction_NewInterfaceParameterNode()
+{
+}
+
+FCameraNodeGraphSchemaAction_NewInterfaceParameterNode::FCameraNodeGraphSchemaAction_NewInterfaceParameterNode(FText InNodeCategory, FText InMenuDesc, FText InToolTip, const int32 InGrouping, FText InKeywords)
+	: FEdGraphSchemaAction(InNodeCategory, InMenuDesc, InToolTip, InGrouping, InKeywords)
+{
+}
+
+UEdGraphNode* FCameraNodeGraphSchemaAction_NewInterfaceParameterNode::PerformAction(UEdGraph* ParentGraph, UEdGraphPin* FromPin, const FVector2D Location, bool bSelectNewNode)
+{
+	UObjectTreeGraph* ObjectTreeGraph = Cast<UObjectTreeGraph>(ParentGraph);
+	if (!ensure(ObjectTreeGraph))
+	{
+		return nullptr;
+	}
+
+	UCameraRigAsset* CameraRigAsset = Cast<UCameraRigAsset>(ObjectTreeGraph->GetRootObject());
+	if (!ensure(CameraRigAsset))
+	{
+		return nullptr;
+	}
+
+	const FScopedTransaction Transaction(LOCTEXT("CreateNewNodeAction", "Create New Node"));
+	const UObjectTreeGraphSchema* Schema = CastChecked<UObjectTreeGraphSchema>(ParentGraph->GetSchema());
+
+	UCameraRigInterfaceParameter* NewInterfaceParameter = NewObject<UCameraRigInterfaceParameter>(CameraRigAsset, NAME_None, RF_Transactional);
+	NewInterfaceParameter->Target = Target;
+	NewInterfaceParameter->TargetPropertyName = TargetPropertyName;
+	NewInterfaceParameter->InterfaceParameterName = TargetPropertyName.ToString();
+	CameraRigAsset->Interface.InterfaceParameters.Add(NewInterfaceParameter);
+
+	UObjectTreeGraphNode* NewGraphNode = Schema->CreateObjectNode(ObjectTreeGraph, NewInterfaceParameter);
+	NewGraphNode->NodePosX = Location.X;
+	NewGraphNode->NodePosY = Location.Y;
+	NewGraphNode->OnGraphNodeMoved();
+	NewGraphNode->AutowireNewNode(FromPin);
+
+	return NewGraphNode;
+}
+
+#undef LOCTEXT_NAMESPACE
+
diff --git a/Engine/Plugins/Cameras/GameplayCameras/Source/GameplayCamerasEditor/Private/Editors/CameraRigInterfaceParameterGraphNode.cpp b/Engine/Plugins/Cameras/GameplayCameras/Source/GameplayCamerasEditor/Private/Editors/CameraRigInterfaceParameterGraphNode.cpp
new file mode 100644
--- /dev/null
+++ b/Engine/Plugins/Cameras/GameplayCameras/Source/GameplayCamerasEditor/Private/Editors/CameraRigInterfaceParameterGraphNode.cpp
@@ -0,0 +1,17 @@
+// Copyright Epic Games, Inc. All Rights Reserved.
+
+#include "Editors/CameraRigInterfaceParameterGraphNode.h"
+
+#include "Core/CameraRigAsset.h"
+#include "Editors/SCameraRigInterfaceParameterGraphNode.h"
+
+UCameraRigInterfaceParameterGraphNode::UCameraRigInterfaceParameterGraphNode(const FObjectInitializer& ObjInit)
+	: UObjectTreeGraphNode(ObjInit)
+{
+}
+
+TSharedPtr<SGraphNode> UCameraRigInterfaceParameterGraphNode::CreateVisualWidget()
+{
+	return SNew(SCameraRigInterfaceParameterGraphNode).GraphNode(this);
+}
+
diff --git a/Engine/Plugins/Cameras/GameplayCameras/Source/GameplayCamerasEditor/Private/Editors/SCameraRigAssetEditor.cpp b/Engine/Plugins/Cameras/GameplayCameras/Source/GameplayCamerasEditor/Private/Editors/SCameraRigAssetEditor.cpp
--- a/Engine/Plugins/Cameras/GameplayCameras/Source/GameplayCamerasEditor/Private/Editors/SCameraRigAssetEditor.cpp
+++ b/Engine/Plugins/Cameras/GameplayCameras/Source/GameplayCamerasEditor/Private/Editors/SCameraRigAssetEditor.cpp
@@ -7,6 +7,9 @@
 #include "Core/CameraRigAsset.h"
 #include "Core/CameraRigTransition.h"
 #include "EdGraph/EdGraphPin.h"
+#include "Editors/CameraNodeGraphNode.h"
+#include "Editors/CameraNodeGraphSchema.h"
+#include "Editors/CameraRigInterfaceParameterGraphNode.h"
 #include "Editors/CameraTransitionGraphSchema.h"
 #include "Editors/ObjectTreeGraph.h"
 #include "Editors/ObjectTreeGraphNode.h"
@@ -45,6 +48,7 @@ void SCameraRigAssetEditor::CreateNodeGraphEditor(const FArguments& InArgs)
 	FObjectTreeGraphConfig GraphConfig;
 	GraphConfig.ConnectableObjectClasses.Add(UCameraRigAsset::StaticClass());
 	GraphConfig.ConnectableObjectClasses.Add(UCameraNode::StaticClass());
+	GraphConfig.ConnectableObjectClasses.Add(UCameraRigInterfaceParameter::StaticClass());
 	GraphConfig.NonConnectableObjectClasses.Add(UBlendCameraNode::StaticClass());
 	GraphConfig.ObjectClassConfigs.Emplace(UCameraRigAsset::StaticClass())
 		.OnlyAsRoot()
@@ -52,9 +56,15 @@ void SCameraRigAssetEditor::CreateNodeGraphEditor(const FArguments& InArgs)
 		.NodeTitleUsesObjectName(true)
 		.NodeTitleColor(Settings->CameraRigAssetTitleColor);
 	GraphConfig.ObjectClassConfigs.Emplace(UCameraNode::StaticClass())
-		.StripDisplayNameSuffix(TEXT("Camera Node"));
+		.StripDisplayNameSuffix(TEXT("Camera Node"))
+		.GraphNodeClass(UCameraNodeGraphNode::StaticClass());
+	GraphConfig.ObjectClassConfigs.Emplace(UCameraRigInterfaceParameter::StaticClass())
+		.SelfPinDirection(EGPD_Output)
+		.SelfPinName(NAME_None)  // No self pin name, we just want the title
+		.GraphNodeClass(UCameraRigInterfaceParameterGraphNode::StaticClass());
 
 	NodeGraph = NewObject<UObjectTreeGraph>(GetTransientPackage(), NAME_None, RF_Transactional);
+	NodeGraph->Schema = UCameraNodeGraphSchema::StaticClass();
 	NodeGraph->AddToRoot();
 	NodeGraph->Initialize(CameraRigAsset, GraphConfig);
 	NodeGraph->RebuildGraph(EObjectTreeGraphBuildSource::RootObjectPackage);
@@ -152,6 +162,20 @@ void SCameraRigAssetEditor::GetGraphs(TArray<UEdGraph*>& OutGraphs) const
 	OutGraphs.Add(TransitionGraph);
 }
 
+UEdGraph* SCameraRigAssetEditor::GetFocusedGraph() const
+{
+	switch (CurrentMode)
+	{
+		case ECameraRigAssetEditorMode::NodeGraph:
+			return NodeGraph;
+		case ECameraRigAssetEditorMode::TransitionGraph:
+			return TransitionGraph;
+		default:
+			ensure(false);
+			return nullptr;
+	}
+}
+
 const FObjectTreeGraphConfig& SCameraRigAssetEditor::GetFocusedGraphConfig() const
 {
 	static const FObjectTreeGraphConfig DefaultConfig;
diff --git a/Engine/Plugins/Cameras/GameplayCameras/Source/GameplayCamerasEditor/Private/Editors/SCameraRigInterfaceParameterGraphNode.cpp b/Engine/Plugins/Cameras/GameplayCameras/Source/GameplayCamerasEditor/Private/Editors/SCameraRigInterfaceParameterGraphNode.cpp
new file mode 100644
--- /dev/null
+++ b/Engine/Plugins/Cameras/GameplayCameras/Source/GameplayCamerasEditor/Private/Editors/SCameraRigInterfaceParameterGraphNode.cpp
@@ -0,0 +1,145 @@
+// Copyright Epic Games, Inc. All Rights Reserved.
+
+#include "Editors/SCameraRigInterfaceParameterGraphNode.h"
+
+#include "Core/CameraRigAsset.h"
+#include "Editors/ObjectTreeGraphNode.h"
+#include "SCommentBubble.h"
+#include "SGraphNode.h"
+#include "SNodePanel.h"
+#include "Styles/GameplayCamerasEditorStyle.h"
+#include "TutorialMetaData.h"
+#include "Widgets/Images/SImage.h"
+#include "Widgets/Layout/SBorder.h"
+#include "Widgets/Layout/SSpacer.h"
+#include "Widgets/Notifications/SErrorText.h"
+#include "Widgets/SBoxPanel.h"
+#include "Widgets/SNullWidget.h"
+#include "Widgets/SOverlay.h"
+#include "Widgets/Text/STextBlock.h"
+
+#define LOCTEXT_NAMESPACE "SCameraRigInterfaceParameterGraphNode"
+
+void SCameraRigInterfaceParameterGraphNode::Construct(const FArguments& InArgs)
+{
+	SObjectTreeGraphNode::FArguments SuperArgs;
+	SuperArgs
+		.GraphNode(InArgs._GraphNode);
+	SObjectTreeGraphNode::Construct(SuperArgs);
+}
+
+void SCameraRigInterfaceParameterGraphNode::UpdateGraphNode()
+{
+	using namespace UE::Cameras;
+
+	TSharedRef<FGameplayCamerasEditorStyle> CamerasEditorStyle = FGameplayCamerasEditorStyle::Get();
+
+	InputPins.Empty();
+	OutputPins.Empty();
+
+	RightNodeBox.Reset();
+	LeftNodeBox.Reset();
+
+	SetupErrorReporting();
+
+	ContentScale.Bind(this, &SGraphNode::GetContentScale);
+
+	GetOrAddSlot(ENodeZone::Center)
+	.HAlign(HAlign_Center)
+	.VAlign(VAlign_Center)
+	[
+		SNew(SVerticalBox)
+		+SVerticalBox::Slot()
+		[
+			SNew(SOverlay)
+			+SOverlay::Slot()
+			[
+				SNew(SImage)
+				.Image(CamerasEditorStyle->GetBrush("Graph.CameraRigParameterNode.Body"))
+			]
+			+ SOverlay::Slot()
+			.VAlign(VAlign_Top)
+			[
+				SNew(SImage)
+				.Image(CamerasEditorStyle->GetBrush("Graph.CameraRigParameterNode.ColorSpill"))
+				.ColorAndOpacity(this, &SCameraRigInterfaceParameterGraphNode::GetNodeTitleColor)
+			]
+			+SOverlay::Slot()
+			[
+				SNew(SImage)
+				.Image(CamerasEditorStyle->GetBrush("Graph.CameraRigParameterNode.Gloss"))
+			]
+			+SOverlay::Slot()
+			.VAlign(VAlign_Center)
+			.HAlign(HAlign_Left)
+			.Padding(FMargin(12, 8, 38, 8))
+			[
+				// NODE TITLE
+				SNew(STextBlock)
+				.TextStyle(FAppStyle::Get(), "Graph.Node.NodeTitle")
+				.Text(this, &SCameraRigInterfaceParameterGraphNode::GetInterfaceParameterName)
+			]
+			+SOverlay::Slot()
+			.Padding(FMargin(0, 4))
+			[
+				// NODE CONTENT AREA
+				SNew(SHorizontalBox)
+				+SHorizontalBox::Slot()
+				.HAlign(HAlign_Left)
+				.FillWidth(1.0f)
+				.Padding(FMargin(2, 0))
+				[
+					// LEFT
+					SAssignNew(LeftNodeBox, SVerticalBox)
+				]
+				+ SHorizontalBox::Slot()
+				.AutoWidth()
+				.HAlign(HAlign_Right)
+				.Padding(FMargin(2, 0))
+				[
+					// RIGHT
+					SAssignNew(RightNodeBox, SVerticalBox)
+				]
+			]
+		]
+		+SVerticalBox::Slot()
+		.VAlign(VAlign_Top)
+		.AutoHeight() 
+		.Padding( FMargin(5.0f, 1.0f) )
+		[
+			ErrorReporting->AsWidget()
+		]
+	];
+
+	// Create widgets for each of the real pins
+	CreatePinWidgets();
+}
+
+const FSlateBrush* SCameraRigInterfaceParameterGraphNode::GetShadowBrush(bool bSelected) const
+{
+	using namespace UE::Cameras;
+	TSharedRef<FGameplayCamerasEditorStyle> CamerasEditorStyle = FGameplayCamerasEditorStyle::Get();
+	return bSelected ? 
+		CamerasEditorStyle->GetBrush(TEXT("Graph.CameraRigParameterNode.ShadowSelected")) : 
+		CamerasEditorStyle->GetBrush(TEXT("Graph.CameraRigParameterNode.Shadow"));
+}
+
+void SCameraRigInterfaceParameterGraphNode::GetDiffHighlightBrushes(const FSlateBrush*& BackgroundOut, const FSlateBrush*& ForegroundOut) const
+{
+	using namespace UE::Cameras;
+	TSharedRef<FGameplayCamerasEditorStyle> CamerasEditorStyle = FGameplayCamerasEditorStyle::Get();
+	BackgroundOut = CamerasEditorStyle->GetBrush(TEXT("Graph.CameraRigParameterNode.DiffHighlight"));
+	ForegroundOut = CamerasEditorStyle->GetBrush(TEXT("Graph.CameraRigParameterNode.DiffHighlightShading"));
+}
+
+FText SCameraRigInterfaceParameterGraphNode::GetInterfaceParameterName() const
+{
+	if (UCameraRigInterfaceParameter* RigParameterObject = GetObjectGraphNode()->CastObject<UCameraRigInterfaceParameter>())
+	{
+		return FText::FromString(RigParameterObject->InterfaceParameterName);
+	}
+	return LOCTEXT("InvalidParameterName", "Invalid");
+}
+
+#undef LOCTEXT_NAMESPACE
+
diff --git a/Engine/Plugins/Cameras/GameplayCameras/Source/GameplayCamerasEditor/Private/Styles/GameplayCamerasEditorStyle.cpp b/Engine/Plugins/Cameras/GameplayCameras/Source/GameplayCamerasEditor/Private/Styles/GameplayCamerasEditorStyle.cpp
--- a/Engine/Plugins/Cameras/GameplayCameras/Source/GameplayCamerasEditor/Private/Styles/GameplayCamerasEditorStyle.cpp
+++ b/Engine/Plugins/Cameras/GameplayCameras/Source/GameplayCamerasEditor/Private/Styles/GameplayCamerasEditorStyle.cpp
@@ -62,6 +62,15 @@ FGameplayCamerasEditorStyle::FGameplayCamerasEditorStyle()
 	Set("DebugCategory.PoseStats.Icon", new IMAGE_BRUSH("Icons/DebugCategory-PoseStats", Icon16x16));
 	Set("DebugCategory.Viewfinder.Icon", new IMAGE_BRUSH("Icons/DebugCategory-Viewfinder", Icon16x16));
 
+	// Graph editor brushes.
+	Set("Graph.CameraRigParameterNode.Body", new BOX_BRUSH("Graph/CameraRigParameterNode_Body", FMargin(16.f/64.f, 12.f/28.f)));
+	Set("Graph.CameraRigParameterNode.ColorSpill", new IMAGE_BRUSH("Graph/CameraRigParameterNode_ColorSpill", FVector2D(132,28)));
+	Set("Graph.CameraRigParameterNode.Gloss", new BOX_BRUSH("Graph/CameraRigParameterNode_Gloss", FMargin(16.f/64.f, 16.f/28.f, 16.f/64.f, 4.f/28.f)));
+	Set("Graph.CameraRigParameterNode.Shadow", new BOX_BRUSH("Graph/CameraRigParameterNode_Shadow", FMargin(26.0f/64.0f)));
+	Set("Graph.CameraRigParameterNode.ShadowSelected", new BOX_BRUSH("Graph/CameraRigParameterNode_ShadowSelected", FMargin(26.0f/64.0f)));
+	Set("Graph.CameraRigParameterNode.DiffHighlight", new BOX_BRUSH("Graph/CameraRigParameterNode_DiffHighlight", FMargin(18.0f/64.0f)));
+	Set("Graph.CameraRigParameterNode.DiffHighlightShading", new BOX_BRUSH("Graph/CameraRigParameterNode_DiffHighlightShading", FMargin(18.0f/64.0f)));
+
 	FSlateStyleRegistry::RegisterSlateStyle(*this);
 }
 
diff --git a/Engine/Plugins/Cameras/GameplayCameras/Source/GameplayCamerasEditor/Private/Toolkits/CameraRigAssetEditorToolkit.cpp b/Engine/Plugins/Cameras/GameplayCameras/Source/GameplayCamerasEditor/Private/Toolkits/CameraRigAssetEditorToolkit.cpp
--- a/Engine/Plugins/Cameras/GameplayCameras/Source/GameplayCamerasEditor/Private/Toolkits/CameraRigAssetEditorToolkit.cpp
+++ b/Engine/Plugins/Cameras/GameplayCameras/Source/GameplayCamerasEditor/Private/Toolkits/CameraRigAssetEditorToolkit.cpp
@@ -4,7 +4,7 @@
 
 #include "AssetTools/CameraRigAssetEditor.h"
 #include "Commands/CameraRigAssetEditorCommands.h"
-#include "Core/CameraRigAllocationInfoBuilder.h"
+#include "Core/CameraNode.h"
 #include "Core/CameraRigAsset.h"
 #include "EditorModeManager.h"
 #include "Editors/ObjectTreeGraph.h"
@@ -384,12 +384,7 @@ void FCameraRigAssetEditorToolkit::OnBuild()
 {
 	using namespace UE::Cameras;
 
-	FCameraRigAllocationInfo AllocationInfo;
-	FCameraRigAllocationInfoBuilder CameraRigBuilder;
-	CameraRigBuilder.BuildAllocationInfo(CameraRigAsset, AllocationInfo);
-
-	CameraRigAsset->BuildStatus = ECameraRigBuildStatus::Clean;
-	CameraRigAsset->AllocationInfo = AllocationInfo;
+	CameraRigAsset->BuildCameraRig();
 
 	FCameraRigPackages BuiltPackages;
 	CameraRigAsset->GatherPackages(BuiltPackages);
diff --git a/Engine/Plugins/Cameras/GameplayCameras/Source/GameplayCamerasEditor/Public/Editors/CameraNodeGraphNode.h b/Engine/Plugins/Cameras/GameplayCameras/Source/GameplayCamerasEditor/Public/Editors/CameraNodeGraphNode.h
new file mode 100644
--- /dev/null
+++ b/Engine/Plugins/Cameras/GameplayCameras/Source/GameplayCamerasEditor/Public/Editors/CameraNodeGraphNode.h
@@ -0,0 +1,33 @@
+// Copyright Epic Games, Inc. All Rights Reserved.
+
+#pragma once
+
+#include "Editors/ObjectTreeGraphNode.h"
+
+#include "CameraNodeGraphNode.generated.h"
+
+/**
+ * Custom graph node for camera nodes. They mostly differ by showing input pins for any 
+ * camera parameter property.
+ */
+UCLASS()
+class UCameraNodeGraphNode : public UObjectTreeGraphNode
+{
+	GENERATED_BODY()
+
+public:
+
+	/** Creates a new graph node. */
+	UCameraNodeGraphNode(const FObjectInitializer& ObjInit);
+
+	/** Gets input pin for given camera parameter property on the underlying camera node. */
+	UEdGraphPin* GetPinForCameraParameterProperty(const FName& InPropertyName) const;
+	/** Gets camera parameter property on the underlying camera node for the given input pin. */
+	FStructProperty* GetCameraParameterPropertyForPin(const UEdGraphPin* InPin) const;
+
+public:
+
+	// UEdGraphNode interface
+	virtual void AllocateDefaultPins() override;
+};
+
diff --git a/Engine/Plugins/Cameras/GameplayCameras/Source/GameplayCamerasEditor/Public/Editors/CameraNodeGraphSchema.h b/Engine/Plugins/Cameras/GameplayCameras/Source/GameplayCamerasEditor/Public/Editors/CameraNodeGraphSchema.h
new file mode 100644
--- /dev/null
+++ b/Engine/Plugins/Cameras/GameplayCameras/Source/GameplayCamerasEditor/Public/Editors/CameraNodeGraphSchema.h
@@ -0,0 +1,66 @@
+// Copyright Epic Games, Inc. All Rights Reserved.
+
+#pragma once
+
+#include "Editors/ObjectTreeGraphSchema.h"
+
+#include "CameraNodeGraphSchema.generated.h"
+
+class UCameraNode;
+
+/**
+ * Schema class for camera node graph.
+ */
+UCLASS()
+class UCameraNodeGraphSchema : public UObjectTreeGraphSchema
+{
+	GENERATED_BODY()
+
+public:
+
+	static const FName PC_CameraParameter;			// A camera parameter pin.
+
+protected:
+
+	// UEdGraphSchema interface.
+	virtual void GetGraphContextActions(FGraphContextMenuBuilder& ContextMenuBuilder) const override;
+	virtual const FPinConnectionResponse CanCreateConnection(const UEdGraphPin* A, const UEdGraphPin* B) const override;
+
+	// UObjectTreeGraphSchema interface.
+	virtual void OnCreateAllNodes(UObjectTreeGraph* InGraph, const FCreatedNodes& InCreatedNodes) const override;
+	virtual bool OnCreateConnection(UEdGraphPin* A, UEdGraphPin* B) const override;
+	virtual bool OnBreakPinLinks(UEdGraphPin& TargetPin, bool bSendsNodeNotification) const override;
+	virtual bool OnBreakSinglePinLink(UEdGraphPin* SourcePin, UEdGraphPin* TargetPin) const override;
+};
+
+/**
+ * Graph editor action for adding a new camera rig parameter node.
+ */
+USTRUCT()
+struct FCameraNodeGraphSchemaAction_NewInterfaceParameterNode : public FEdGraphSchemaAction
+{
+	GENERATED_BODY()
+
+public:
+
+	/** The camera node being driven by the camera rig parameter. */
+	UPROPERTY()
+	TObjectPtr<UCameraNode> Target;
+
+	/** The property on the target camera node being driven by the camera rig parameter. */
+	UPROPERTY()
+	FName TargetPropertyName;
+
+public:
+
+	FCameraNodeGraphSchemaAction_NewInterfaceParameterNode();
+	FCameraNodeGraphSchemaAction_NewInterfaceParameterNode(FText InNodeCategory, FText InMenuDesc, FText InToolTip, const int32 InGrouping = 0, FText InKeywords = FText());
+
+public:
+
+	// FEdGraphSchemaAction interface.
+	static FName StaticGetTypeId() { static FName Type("FCameraNodeGraphSchemaAction_NewInterfaceParameterNode"); return Type; }
+	virtual FName GetTypeId() const override { return StaticGetTypeId(); } 
+	virtual UEdGraphNode* PerformAction(UEdGraph* ParentGraph, UEdGraphPin* FromPin, const FVector2D Location, bool bSelectNewNode = true) override;
+};
+
diff --git a/Engine/Plugins/Cameras/GameplayCameras/Source/GameplayCamerasEditor/Public/Editors/CameraRigInterfaceParameterGraphNode.h b/Engine/Plugins/Cameras/GameplayCameras/Source/GameplayCamerasEditor/Public/Editors/CameraRigInterfaceParameterGraphNode.h
new file mode 100644
--- /dev/null
+++ b/Engine/Plugins/Cameras/GameplayCameras/Source/GameplayCamerasEditor/Public/Editors/CameraRigInterfaceParameterGraphNode.h
@@ -0,0 +1,27 @@
+// Copyright Epic Games, Inc. All Rights Reserved.
+
+#pragma once
+
+#include "Editors/ObjectTreeGraphNode.h"
+
+#include "CameraRigInterfaceParameterGraphNode.generated.h"
+
+/**
+ * Custom graph editor node for a camera rig parameter.
+ */
+UCLASS()
+class UCameraRigInterfaceParameterGraphNode : public UObjectTreeGraphNode
+{
+	GENERATED_BODY()
+
+public:
+
+	/** Creates a new graph node. */
+	UCameraRigInterfaceParameterGraphNode(const FObjectInitializer& ObjInit);
+
+public:
+
+	// UEdGraphNode interface
+	virtual TSharedPtr<SGraphNode> CreateVisualWidget() override;
+};
+
diff --git a/Engine/Plugins/Cameras/GameplayCameras/Source/GameplayCamerasEditor/Public/Editors/SCameraRigAssetEditor.h b/Engine/Plugins/Cameras/GameplayCameras/Source/GameplayCamerasEditor/Public/Editors/SCameraRigAssetEditor.h
--- a/Engine/Plugins/Cameras/GameplayCameras/Source/GameplayCamerasEditor/Public/Editors/SCameraRigAssetEditor.h
+++ b/Engine/Plugins/Cameras/GameplayCameras/Source/GameplayCamerasEditor/Public/Editors/SCameraRigAssetEditor.h
@@ -63,6 +63,8 @@ public:
 	/** Gets both the node hierarchy and transition graphs. */
 	void GetGraphs(TArray<UEdGraph*>& OutGraphs) const;
 
+	/** Gets the graph for the current mode. */
+	UEdGraph* GetFocusedGraph() const;
 	/** Gets the graph configuration for the current mode. */
 	const FObjectTreeGraphConfig& GetFocusedGraphConfig() const;
 
diff --git a/Engine/Plugins/Cameras/GameplayCameras/Source/GameplayCamerasEditor/Public/Editors/SCameraRigInterfaceParameterGraphNode.h b/Engine/Plugins/Cameras/GameplayCameras/Source/GameplayCamerasEditor/Public/Editors/SCameraRigInterfaceParameterGraphNode.h
new file mode 100644
--- /dev/null
+++ b/Engine/Plugins/Cameras/GameplayCameras/Source/GameplayCamerasEditor/Public/Editors/SCameraRigInterfaceParameterGraphNode.h
@@ -0,0 +1,33 @@
+// Copyright Epic Games, Inc. All Rights Reserved.
+
+#pragma once
+
+#include "Editors/SObjectTreeGraphNode.h"
+
+/**
+ * Custom graph editor node widget for a camera rig parameter node.
+ */
+class SCameraRigInterfaceParameterGraphNode : public SObjectTreeGraphNode
+{
+public:
+
+	SLATE_BEGIN_ARGS(SCameraRigInterfaceParameterGraphNode)
+		: _GraphNode(nullptr)
+	{}
+		SLATE_ARGUMENT(UObjectTreeGraphNode*, GraphNode)
+	SLATE_END_ARGS()
+
+	void Construct(const FArguments& InArgs);
+
+protected:
+
+	// SGraphNode interface.
+	virtual void UpdateGraphNode() override;
+	virtual const FSlateBrush* GetShadowBrush(bool bSelected) const override;
+	virtual void GetDiffHighlightBrushes(const FSlateBrush*& BackgroundOut, const FSlateBrush*& ForegroundOut) const override;
+
+private:
+	
+	FText GetInterfaceParameterName() const;
+};
+
diff --git a/Engine/Plugins/Cameras/GameplayCameras/Source/GameplayCamerasEditor/Public/Toolkits/CameraRigAssetEditorToolkit.h b/Engine/Plugins/Cameras/GameplayCameras/Source/GameplayCamerasEditor/Public/Toolkits/CameraRigAssetEditorToolkit.h
--- a/Engine/Plugins/Cameras/GameplayCameras/Source/GameplayCamerasEditor/Public/Toolkits/CameraRigAssetEditorToolkit.h
+++ b/Engine/Plugins/Cameras/GameplayCameras/Source/GameplayCamerasEditor/Public/Toolkits/CameraRigAssetEditorToolkit.h
@@ -23,7 +23,6 @@ namespace UE::Cameras
 
 class IGameplayCamerasLiveEditManager;
 class SCameraRigAssetEditor;
-class SCameraTransitionEditor;
 
 /**
  * Editor toolkit for a camera rig asset.
