Cameras: new BP node for setting camera rig parameters of a given rig into a given evaluation result

- Add individual camera rig parameter setter functions, but flag them as BP-internal.
- Add a BP node that dynamically lists a given camera rig's exposed parameters as input pins on himself. On BP compile, the node is discarded and replaced by the individual setter calls, as appropriate.
- Add a new "UncookedOnly" module in the GameplayCameras plugin to put the previously mentioned new node.

diff --git a/Engine/Plugins/Cameras/GameplayCameras/GameplayCameras.uplugin b/Engine/Plugins/Cameras/GameplayCameras/GameplayCameras.uplugin
--- a/Engine/Plugins/Cameras/GameplayCameras/GameplayCameras.uplugin
+++ b/Engine/Plugins/Cameras/GameplayCameras/GameplayCameras.uplugin
@@ -28,6 +28,11 @@
 			"Type" : "Runtime",
 			"LoadingPhase" : "PreDefault"
 		},
+		{
+			"Name" : "GameplayCamerasUncookedOnly",
+			"Type" : "UncookedOnly",
+			"LoadingPhase" : "Default"
+		},
 		{
 			"Name": "GameplayCamerasEditor",
 			"Type": "Editor",
diff --git a/Engine/Plugins/Cameras/GameplayCameras/Source/GameplayCameras/Private/GameFramework/CameraRigParameterInterop.cpp b/Engine/Plugins/Cameras/GameplayCameras/Source/GameplayCameras/Private/GameFramework/CameraRigParameterInterop.cpp
new file mode 100644
--- /dev/null
+++ b/Engine/Plugins/Cameras/GameplayCameras/Source/GameplayCameras/Private/GameFramework/CameraRigParameterInterop.cpp
@@ -0,0 +1,141 @@
+// Copyright Epic Games, Inc. All Rights Reserved.
+
+#include "GameFramework/CameraRigParameterInterop.h"
+
+#include "Core/CameraRigAsset.h"
+#include "Core/CameraNodeEvaluator.h"
+#include "Core/CameraVariableTable.h"
+#include "GameFramework/CameraEvaluationResultInterop.h"
+
+#define LOCTEXT_NAMESPACE "CameraRigParameterInterop"
+
+#include UE_INLINE_GENERATED_CPP_BY_NAME(CameraRigParameterInterop)
+
+namespace UE::Cameras::Private
+{
+
+template<typename VariableAssetType>
+void SetCameraRigParameter(UCameraEvaluationResultInterop* ResultInterop, VariableAssetType* PrivateVariable, typename VariableAssetType::ValueType Value)
+{
+	if (ResultInterop == nullptr)
+	{
+		FFrame::KismetExecutionMessage(TEXT("No camera evaluation result was passed."), ELogVerbosity::Error);
+		return;
+	}
+	if (PrivateVariable == nullptr)
+	{
+		FFrame::KismetExecutionMessage(TEXT("No camera rig was passed."), ELogVerbosity::Error);
+		return;
+	}
+
+	FCameraNodeEvaluationResult* Result = ResultInterop->GetEvaluationResult();
+	if (Result == nullptr)
+	{
+		FFrame::KismetExecutionMessage(TEXT("The given camera evaluation result is invalid."), ELogVerbosity::Error);
+		return;
+	}
+
+	Result->VariableTable.SetValue(PrivateVariable, Value, true);
+}
+
+}  // namespace UE::Cameras::Private
+
+UCameraRigParameterInterop::UCameraRigParameterInterop(const FObjectInitializer& ObjectInit)
+	: Super(ObjectInit)
+{
+}
+
+void UCameraRigParameterInterop::SetBooleanParameter(UCameraEvaluationResultInterop* ResultInterop, UCameraRigAsset* CameraRig, const FString& ParameterName, bool bParameterValue)
+{
+	UE::Cameras::Private::SetCameraRigParameter(
+			ResultInterop, 
+			Cast<UBooleanCameraVariable>(GetParameterPrivateVariable(CameraRig, ParameterName)), 
+			bParameterValue);
+}
+
+void UCameraRigParameterInterop::SetIntegerParameter(UCameraEvaluationResultInterop* ResultInterop, UCameraRigAsset* CameraRig, const FString& ParameterName, int32 ParameterValue)
+{
+	UE::Cameras::Private::SetCameraRigParameter(
+			ResultInterop, 
+			Cast<UInteger32CameraVariable>(GetParameterPrivateVariable(CameraRig, ParameterName)), 
+			ParameterValue);
+}
+
+void UCameraRigParameterInterop::SetFloatParameter(UCameraEvaluationResultInterop* ResultInterop, UCameraRigAsset* CameraRig, const FString& ParameterName, double ParameterValue)
+{
+	UE::Cameras::Private::SetCameraRigParameter(
+			ResultInterop, 
+			Cast<UFloatCameraVariable>(GetParameterPrivateVariable(CameraRig, ParameterName)),
+			(float)ParameterValue);
+}
+
+void UCameraRigParameterInterop::SetDoubleParameter(UCameraEvaluationResultInterop* ResultInterop, UCameraRigAsset* CameraRig, const FString& ParameterName, double ParameterValue)
+{
+	UE::Cameras::Private::SetCameraRigParameter(
+			ResultInterop, 
+			Cast<UDoubleCameraVariable>(GetParameterPrivateVariable(CameraRig, ParameterName)),
+			ParameterValue);
+}
+
+void UCameraRigParameterInterop::SetVector2Parameter(UCameraEvaluationResultInterop* ResultInterop, UCameraRigAsset* CameraRig, const FString& ParameterName, FVector2D ParameterValue)
+{
+	UE::Cameras::Private::SetCameraRigParameter(
+			ResultInterop, 
+			Cast<UVector2dCameraVariable>(GetParameterPrivateVariable(CameraRig, ParameterName)),
+			ParameterValue);
+}
+
+void UCameraRigParameterInterop::SetVector3Parameter(UCameraEvaluationResultInterop* ResultInterop, UCameraRigAsset* CameraRig, const FString& ParameterName, FVector ParameterValue)
+{
+	UE::Cameras::Private::SetCameraRigParameter(
+			ResultInterop, 
+			Cast<UVector3dCameraVariable>(GetParameterPrivateVariable(CameraRig, ParameterName)),
+			ParameterValue);
+}
+
+void UCameraRigParameterInterop::SetVector4Parameter(UCameraEvaluationResultInterop* ResultInterop, UCameraRigAsset* CameraRig, const FString& ParameterName, FVector4 ParameterValue)
+{
+	UE::Cameras::Private::SetCameraRigParameter(
+			ResultInterop, 
+			Cast<UVector4dCameraVariable>(GetParameterPrivateVariable(CameraRig, ParameterName)),
+			ParameterValue);
+}
+
+void UCameraRigParameterInterop::SetRotatorParameter(UCameraEvaluationResultInterop* ResultInterop, UCameraRigAsset* CameraRig, const FString& ParameterName, FRotator ParameterValue)
+{
+	UE::Cameras::Private::SetCameraRigParameter(
+			ResultInterop, 
+			Cast<URotator3dCameraVariable>(GetParameterPrivateVariable(CameraRig, ParameterName)),
+			ParameterValue);
+}
+
+void UCameraRigParameterInterop::SetTransformParameter(UCameraEvaluationResultInterop* ResultInterop, UCameraRigAsset* CameraRig, const FString& ParameterName, FTransform ParameterValue)
+{
+	UE::Cameras::Private::SetCameraRigParameter(
+			ResultInterop, 
+			Cast<UTransform3dCameraVariable>(GetParameterPrivateVariable(CameraRig, ParameterName)),
+			ParameterValue);
+}
+
+UCameraVariableAsset* UCameraRigParameterInterop::GetParameterPrivateVariable(UCameraRigAsset* CameraRig, const FString& ParameterName)
+{
+	UCameraRigInterfaceParameter* InterfaceParameter = CameraRig->Interface.FindInterfaceParameterByName(ParameterName);
+	if (!InterfaceParameter)
+	{
+		const FText Text = LOCTEXT("NoSuchParameter", "No parameter '{0}' found on camera rig '{1}'. Setting this camera variable table value will most probably accomplish nothing.");
+		FFrame::KismetExecutionMessage(*FText::Format(Text, FText::FromString(ParameterName), FText::FromString(CameraRig->GetPathName())).ToString(), ELogVerbosity::Warning);
+		return nullptr;
+	}
+
+	if (!InterfaceParameter->PrivateVariable)
+	{
+		const FText Text = LOCTEXT("CameraRigNeedsBuilding", "Parameter '{0}' isn't built. Please build camera rig '{1}'.");
+		FFrame::KismetExecutionMessage(*FText::Format(Text, FText::FromString(ParameterName), FText::FromString(CameraRig->GetPathName())).ToString(), ELogVerbosity::Warning);
+		return nullptr;
+	}
+
+	return InterfaceParameter->PrivateVariable;
+}
+
+#undef LOCTEXT_NAMESPACE
+
diff --git a/Engine/Plugins/Cameras/GameplayCameras/Source/GameplayCameras/Public/GameFramework/CameraRigParameterInterop.h b/Engine/Plugins/Cameras/GameplayCameras/Source/GameplayCameras/Public/GameFramework/CameraRigParameterInterop.h
new file mode 100644
--- /dev/null
+++ b/Engine/Plugins/Cameras/GameplayCameras/Source/GameplayCameras/Public/GameFramework/CameraRigParameterInterop.h
@@ -0,0 +1,65 @@
+// Copyright Epic Games, Inc. All Rights Reserved.
+
+#pragma once
+
+#include "Core/CameraParameters.h"
+#include "CoreTypes.h"
+#include "Kismet/BlueprintFunctionLibrary.h"
+#include "UObject/ObjectMacros.h"
+#include "UObject/ObjectPtr.h"
+
+#include "CameraRigParameterInterop.generated.h"
+
+class UCameraRigAsset;
+class UCameraEvaluationResultInterop;
+class UCameraVariableAsset;
+
+/**
+ * Blueprint internal methods to set values on a camera rig's exposed parameters.
+ *
+ * These functions are internal because users are supposed to use the K2Node_SetCameraRigParameters node instead. That node then
+ * gets compiled into one or more of these internal functions.
+ */
+UCLASS(MinimalAPI)
+class UCameraRigParameterInterop : public UBlueprintFunctionLibrary
+{
+	GENERATED_BODY()
+
+public:
+
+	UCameraRigParameterInterop(const FObjectInitializer& ObjectInit);
+
+public:
+
+	UFUNCTION(BlueprintCallable, Category="Camera", meta=(BlueprintInternalUseOnly="true"))
+	static void SetBooleanParameter(UCameraEvaluationResultInterop* ResultInterop, UCameraRigAsset* CameraRig, const FString& ParameterName, bool bParameterValue);
+
+	UFUNCTION(BlueprintCallable, Category="Camera", meta=(BlueprintInternalUseOnly="true"))
+	static void SetIntegerParameter(UCameraEvaluationResultInterop* ResultInterop, UCameraRigAsset* CameraRig, const FString& ParameterName, int32 ParameterValue);
+
+	UFUNCTION(BlueprintCallable, Category="Camera", meta=(BlueprintInternalUseOnly="true"))
+	static void SetFloatParameter(UCameraEvaluationResultInterop* ResultInterop, UCameraRigAsset* CameraRig, const FString& ParameterName, double ParameterValue);
+
+	UFUNCTION(BlueprintCallable, Category="Camera", meta=(BlueprintInternalUseOnly="true"))
+	static void SetDoubleParameter(UCameraEvaluationResultInterop* ResultInterop, UCameraRigAsset* CameraRig, const FString& ParameterName, double ParameterValue);
+
+	UFUNCTION(BlueprintCallable, Category="Camera", meta=(BlueprintInternalUseOnly="true"))
+	static void SetVector2Parameter(UCameraEvaluationResultInterop* ResultInterop, UCameraRigAsset* CameraRig, const FString& ParameterName, FVector2D ParameterValue);
+
+	UFUNCTION(BlueprintCallable, Category="Camera", meta=(BlueprintInternalUseOnly="true"))
+	static void SetVector3Parameter(UCameraEvaluationResultInterop* ResultInterop, UCameraRigAsset* CameraRig, const FString& ParameterName, FVector ParameterValue);
+
+	UFUNCTION(BlueprintCallable, Category="Camera", meta=(BlueprintInternalUseOnly="true"))
+	static void SetVector4Parameter(UCameraEvaluationResultInterop* ResultInterop, UCameraRigAsset* CameraRig, const FString& ParameterName, FVector4 ParameterValue);
+
+	UFUNCTION(BlueprintCallable, Category="Camera", meta=(BlueprintInternalUseOnly="true"))
+	static void SetRotatorParameter(UCameraEvaluationResultInterop* ResultInterop, UCameraRigAsset* CameraRig, const FString& ParameterName, FRotator ParameterValue);
+
+	UFUNCTION(BlueprintCallable, Category="Camera", meta=(BlueprintInternalUseOnly="true"))
+	static void SetTransformParameter(UCameraEvaluationResultInterop* ResultInterop, UCameraRigAsset* CameraRig, const FString& ParameterName, FTransform ParameterValue);
+
+private:
+
+	static UCameraVariableAsset* GetParameterPrivateVariable(UCameraRigAsset* CameraRig, const FString& ParameterName);
+};
+
diff --git a/Engine/Plugins/Cameras/GameplayCameras/Source/GameplayCamerasUncookedOnly/GameplayCamerasUncookedOnly.Build.cs b/Engine/Plugins/Cameras/GameplayCameras/Source/GameplayCamerasUncookedOnly/GameplayCamerasUncookedOnly.Build.cs
new file mode 100644
--- /dev/null
+++ b/Engine/Plugins/Cameras/GameplayCameras/Source/GameplayCamerasUncookedOnly/GameplayCamerasUncookedOnly.Build.cs
@@ -0,0 +1,32 @@
+// Copyright Epic Games, Inc. All Rights Reserved.
+
+namespace UnrealBuildTool.Rules
+{
+	public class GameplayCamerasUncookedOnly : ModuleRules
+	{
+		public GameplayCamerasUncookedOnly(ReadOnlyTargetRules Target) : base(Target)
+		{
+			PublicDependencyModuleNames.AddRange(
+				new string[]
+				{
+					"GameplayCameras"
+				}
+			);
+
+			PrivateDependencyModuleNames.AddRange(
+				new string[]
+				{
+					"BlueprintGraph",
+					"Core",
+					"CoreUObject",
+					"Engine",
+					"BlueprintGraph",
+					"Kismet",
+					"KismetCompiler",
+					"UnrealEd",
+				}
+			);
+		}
+	}
+}
+
diff --git a/Engine/Plugins/Cameras/GameplayCameras/Source/GameplayCamerasUncookedOnly/Private/BlueprintGraph/K2Node_SetCameraRigParameters.cpp b/Engine/Plugins/Cameras/GameplayCameras/Source/GameplayCamerasUncookedOnly/Private/BlueprintGraph/K2Node_SetCameraRigParameters.cpp
new file mode 100644
--- /dev/null
+++ b/Engine/Plugins/Cameras/GameplayCameras/Source/GameplayCamerasUncookedOnly/Private/BlueprintGraph/K2Node_SetCameraRigParameters.cpp
@@ -0,0 +1,431 @@
+// Copyright Epic Games, Inc. All Rights Reserved.
+
+#include "BlueprintGraph/K2Node_SetCameraRigParameters.h"
+
+#include "BlueprintActionDatabaseRegistrar.h"
+#include "BlueprintNodeSpawner.h"
+#include "Core/CameraRigAsset.h"
+#include "Core/CameraVariableAssets.h"
+#include "EdGraphSchema_K2.h"
+#include "EditorCategoryUtils.h"
+#include "GameFramework/CameraEvaluationResultInterop.h"
+#include "GameFramework/CameraRigParameterInterop.h"
+#include "K2Node_CallFunction.h"
+#include "KismetCompiler.h"
+#include "Kismet2/BlueprintEditorUtils.h"
+
+#define LOCTEXT_NAMESPACE "K2Node_SetCameraRigParameters"
+
+const FName UK2Node_SetCameraRigParameters::CameraRigPinName(TEXT("CameraRig"));
+const FName UK2Node_SetCameraRigParameters::CameraEvaluationResultPinName(TEXT("CameraEvaluationResult"));
+
+UK2Node_SetCameraRigParameters::UK2Node_SetCameraRigParameters(const FObjectInitializer& ObjectInit)
+	: Super(ObjectInit)
+{
+}
+
+void UK2Node_SetCameraRigParameters::AllocateDefaultPins()
+{
+	// Add execution pins.
+	CreatePin(EGPD_Input, UEdGraphSchema_K2::PC_Exec, UEdGraphSchema_K2::PN_Execute);
+	CreatePin(EGPD_Output, UEdGraphSchema_K2::PC_Exec, UEdGraphSchema_K2::PN_Then);
+
+	// Add evalation result pin.
+	CreatePin(EGPD_Input, UEdGraphSchema_K2::PC_Object, UCameraEvaluationResultInterop::StaticClass(), CameraEvaluationResultPinName);
+
+	// Add camera rig pin.
+	CreatePin(EGPD_Input, UEdGraphSchema_K2::PC_Object, UCameraRigAsset::StaticClass(), CameraRigPinName);
+	
+	Super::AllocateDefaultPins();
+}
+
+void UK2Node_SetCameraRigParameters::ReallocatePinsDuringReconstruction(TArray<UEdGraphPin*>& OldPins) 
+{
+	AllocateDefaultPins();
+
+	TArrayView<UEdGraphPin* const> PinsToSearch = OldPins;
+	if (UCameraRigAsset* CameraRig = GetCameraRig(&PinsToSearch))
+	{
+		CreatePinsForCameraRig(CameraRig);
+	}
+}
+
+void UK2Node_SetCameraRigParameters::PostPlacedNewNode()
+{
+	Super::PostPlacedNewNode();
+
+	if (UCameraRigAsset* CameraRig = GetCameraRig())
+	{
+		CreatePinsForCameraRig(CameraRig);
+	}
+}
+
+void UK2Node_SetCameraRigParameters::PinConnectionListChanged(UEdGraphPin* Pin)
+{
+	Super::PinConnectionListChanged(Pin);
+
+	if (Pin && (Pin->PinName == CameraRigPinName))
+	{
+		OnCameraRigChanged();
+	}
+}
+
+void UK2Node_SetCameraRigParameters::PinDefaultValueChanged(UEdGraphPin* ChangedPin) 
+{
+	if (ChangedPin && (ChangedPin->PinName == CameraRigPinName))
+	{
+		OnCameraRigChanged();
+	}
+}
+
+FText UK2Node_SetCameraRigParameters::GetTooltipText() const
+{
+	return LOCTEXT("NodeTooltip", "Sets values for the exposed parameters on the given camera rig.");
+}
+
+FText UK2Node_SetCameraRigParameters::GetNodeTitle(ENodeTitleType::Type TitleType) const
+{
+	return LOCTEXT("BaseNodeTitle", "Set Camera Rig Parameters");
+}
+
+void UK2Node_SetCameraRigParameters::GetMenuActions(FBlueprintActionDatabaseRegistrar& ActionRegistrar) const
+{
+	UClass* ActionKey = GetClass();
+	if (ActionRegistrar.IsOpenForRegistration(ActionKey))
+	{
+		UBlueprintNodeSpawner* NodeSpawner = UBlueprintNodeSpawner::Create(GetClass());
+		check(NodeSpawner != nullptr);
+
+		ActionRegistrar.AddBlueprintAction(ActionKey, NodeSpawner);
+	}
+}
+
+FText UK2Node_SetCameraRigParameters::GetMenuCategory() const
+{
+	return FEditorCategoryUtils::GetCommonCategory(FCommonEditorCategory::Gameplay);
+}
+
+void UK2Node_SetCameraRigParameters::ExpandNode(class FKismetCompilerContext& CompilerContext, UEdGraph* SourceGraph)
+{
+	Super::ExpandNode(CompilerContext, SourceGraph);
+
+	UCameraRigAsset* CameraRig = GetCameraRig();
+
+	if (!CameraRig)
+	{
+		CompilerContext.MessageLog.Error(*LOCTEXT("ErrorMissingCameraRig", "SetCameraRigParameters node @@ doesn't have a camera rig set.").ToString(), this);
+		BreakAllNodeLinks();
+		return;
+	}
+
+	// Get all the pins that correspond to parameters we want to override.
+	TArray<UEdGraphPin*> RigParameterPins;
+	GetCameraRigParameterPins(RigParameterPins);
+
+	UEdGraphPin* const CameraRigPin = FindPinChecked(CameraRigPinName);
+	UEdGraphPin* const EvaluationResultPin = FindPinChecked(CameraEvaluationResultPinName);
+
+	UEdGraphPin* OriginalThenPin = GetThenPin();
+	UEdGraphPin* PreviousThenPin = nullptr;
+
+	for (UEdGraphPin* RigParameterPin : RigParameterPins)
+	{
+		UCameraRigInterfaceParameter* InterfaceParameter = CameraRig->Interface.FindInterfaceParameterByName(RigParameterPin->GetName());
+		if (!InterfaceParameter)
+		{
+			CompilerContext.MessageLog.Error(*LOCTEXT("ErrorMissingParameter", "SetCameraRigParameters node @@ is trying to set parameter @@ but camera rig @@ has no such parameter.").ToString(), this, *RigParameterPin->GetName(), CameraRig);
+			continue;
+		}
+
+		// Figure out the sort of SetXxxParameter function we want to call for this parameter.
+		FName CallSetParameterFuncName;
+		switch (InterfaceParameter->PrivateVariable->GetVariableType())
+		{
+			case ECameraVariableType::Boolean:
+				CallSetParameterFuncName = GET_FUNCTION_NAME_CHECKED(UCameraRigParameterInterop, SetBooleanParameter);
+				break;
+			case ECameraVariableType::Integer32:
+				CallSetParameterFuncName = GET_FUNCTION_NAME_CHECKED(UCameraRigParameterInterop, SetIntegerParameter);
+				break;
+			case ECameraVariableType::Float:
+				CallSetParameterFuncName = GET_FUNCTION_NAME_CHECKED(UCameraRigParameterInterop, SetFloatParameter);
+				break;
+			case ECameraVariableType::Double:
+				CallSetParameterFuncName = GET_FUNCTION_NAME_CHECKED(UCameraRigParameterInterop, SetDoubleParameter);
+				break;
+			case ECameraVariableType::Vector2d:
+				CallSetParameterFuncName = GET_FUNCTION_NAME_CHECKED(UCameraRigParameterInterop, SetVector2Parameter);
+				break;
+			case ECameraVariableType::Vector3d:
+				CallSetParameterFuncName = GET_FUNCTION_NAME_CHECKED(UCameraRigParameterInterop, SetVector3Parameter);
+				break;
+			case ECameraVariableType::Vector4d:
+				CallSetParameterFuncName = GET_FUNCTION_NAME_CHECKED(UCameraRigParameterInterop, SetVector4Parameter);
+				break;
+			case ECameraVariableType::Rotator3d:
+				CallSetParameterFuncName = GET_FUNCTION_NAME_CHECKED(UCameraRigParameterInterop, SetRotatorParameter);
+				break;
+			case ECameraVariableType::Transform3d:
+				CallSetParameterFuncName = GET_FUNCTION_NAME_CHECKED(UCameraRigParameterInterop, SetTransformParameter);
+				break;
+		}
+		if (CallSetParameterFuncName.IsNone())
+		{
+			CompilerContext.MessageLog.Error(*LOCTEXT("ErrorUnsupportedParameterType", "SetCameraRigParameters node @@ is trying to set parameter @@ but it has an unsupported type.").ToString(), this, *RigParameterPin->GetName());
+			continue;
+		}
+
+		// Make the SetXxxParameter function call node.
+		UK2Node_CallFunction* CallSetParameter = CompilerContext.SpawnIntermediateNode<UK2Node_CallFunction>(this, SourceGraph);
+		CallSetParameter->FunctionReference.SetExternalMember(CallSetParameterFuncName, UCameraRigParameterInterop::StaticClass());
+		CallSetParameter->AllocateDefaultPins();
+
+		// Connect the evaluation result pin that specifies where the parameter should be overriden.
+		UEdGraphPin* CallSetParameterResultInteropPin = CallSetParameter->FindPinChecked(TEXT("ResultInterop"));
+		CompilerContext.CopyPinLinksToIntermediate(*EvaluationResultPin, *CallSetParameterResultInteropPin);
+
+		// Connect the camera rig argument.
+		UEdGraphPin* CallSetParameterCameraRigPin = CallSetParameter->FindPinChecked(TEXT("CameraRig"));
+		CompilerContext.CopyPinLinksToIntermediate(*CameraRigPin, *CallSetParameterCameraRigPin);
+
+		// Set the parameter name argument.
+		UEdGraphPin* CallSetParameterNamePin = CallSetParameter->FindPinChecked(TEXT("ParameterName"));
+		CallSetParameterNamePin->DefaultValue = InterfaceParameter->InterfaceParameterName;
+
+		// Set or connect the parameter value argument.
+		UEdGraphPin* CallSetParameterValuePin = CallSetParameter->FindPinChecked(TEXT("ParameterValue"));
+		if (RigParameterPin->LinkedTo.Num() > 0)
+		{
+			CompilerContext.MovePinLinksToIntermediate(*RigParameterPin, *CallSetParameterValuePin);
+		}
+		else
+		{
+			CallSetParameterValuePin->DefaultObject = RigParameterPin->DefaultObject;
+		}
+
+		// Connect the SetXxxParameter node to the chain of other similar nodes. The SetCameraRigParameters node
+		// effectively transforms into a chain of individual setter function calls.
+		UEdGraphPin* CallSetParameterExecPin = CallSetParameter->GetExecPin();
+		if (PreviousThenPin)
+		{
+			PreviousThenPin->MakeLinkTo(CallSetParameterExecPin);
+		}
+		else
+		{
+			UEdGraphPin* ThisExecPin = GetExecPin();
+			CompilerContext.MovePinLinksToIntermediate(*ThisExecPin, *CallSetParameterExecPin);
+		}
+
+		PreviousThenPin = CallSetParameter->GetThenPin();
+	}
+
+	// Connect the last node if necessary.
+	if (OriginalThenPin && OriginalThenPin->LinkedTo.Num() > 0)
+	{
+		CompilerContext.MovePinLinksToIntermediate(*OriginalThenPin, *PreviousThenPin);
+	}
+
+	BreakAllNodeLinks();
+}
+
+UEdGraphPin* UK2Node_SetCameraRigParameters::GetCameraRigPin(TArrayView<UEdGraphPin* const>* InPinsToSearch) const
+{
+	TArrayView<UEdGraphPin* const> PinsToSearch = MakeArrayView(Pins);
+	if (InPinsToSearch)
+	{
+		PinsToSearch = *InPinsToSearch;
+	}
+
+	UEdGraphPin* CameraRigPin = nullptr;
+	for (UEdGraphPin* Pin : PinsToSearch)
+	{
+		if (Pin && Pin->PinName == CameraRigPinName)
+		{
+			CameraRigPin = Pin;
+			break;
+		}
+	}
+	check(CameraRigPin == nullptr || CameraRigPin->Direction == EGPD_Input);
+	return CameraRigPin;
+}
+
+UEdGraphPin* UK2Node_SetCameraRigParameters::GetCameraEvaluationResultPin() const
+{
+	UEdGraphPin* ResultPin = nullptr;
+	for (UEdGraphPin* Pin : Pins)
+	{
+		if (Pin && Pin->PinName == CameraRigPinName)
+		{
+			ResultPin = Pin;
+			break;
+		}
+	}
+	check(ResultPin == nullptr || ResultPin->Direction == EGPD_Input);
+	return ResultPin;
+}
+
+void UK2Node_SetCameraRigParameters::GetCameraRigParameterPins(TArray<UEdGraphPin*>& OutParameterPins) const
+{
+	for (UEdGraphPin* Pin : Pins)
+	{
+		if (IsCameraRigParameterPin(Pin))
+		{
+			OutParameterPins.Add(Pin);
+		}
+	}
+}
+
+bool UK2Node_SetCameraRigParameters::IsCameraRigParameterPin(UEdGraphPin* Pin) const
+{
+	return Pin->PinName != UEdGraphSchema_K2::PN_Execute &&
+		Pin->PinName != UEdGraphSchema_K2::PN_Then &&
+		Pin->PinName != UEdGraphSchema_K2::PN_ReturnValue &&
+		Pin->PinName != CameraRigPinName &&
+		Pin->PinName != CameraEvaluationResultPinName;
+}
+
+void UK2Node_SetCameraRigParameters::CreatePinsForCameraRig(UCameraRigAsset* CameraRig, TArray<UEdGraphPin*>* CreatedPins)
+{
+	check(CameraRig);
+
+	const UEdGraphSchema_K2* K2Schema = GetDefault<UEdGraphSchema_K2>();
+
+	for (UCameraRigInterfaceParameter* InterfaceParameter : CameraRig->Interface.InterfaceParameters)
+	{
+		if (!ensure(InterfaceParameter))
+		{
+			continue;
+		}
+
+		if (!InterfaceParameter->PrivateVariable)
+		{
+			// Camera rig isn't fully built.
+			continue;
+		}
+
+		FName NewPinCategory;
+		FName NewPinSubCategory;
+		UObject* NewPinSubCategoryObject = nullptr;
+		switch (InterfaceParameter->PrivateVariable->GetVariableType())
+		{
+			case ECameraVariableType::Boolean:
+				NewPinCategory = UEdGraphSchema_K2::PC_Boolean;
+				break;
+			case ECameraVariableType::Integer32:
+				NewPinCategory = UEdGraphSchema_K2::PC_Int;
+				break;
+			case ECameraVariableType::Float:
+				// We'll cast down to float.
+				NewPinCategory = UEdGraphSchema_K2::PC_Real;
+				NewPinSubCategory = UEdGraphSchema_K2::PC_Float;
+				break;
+			case ECameraVariableType::Double:
+				NewPinCategory = UEdGraphSchema_K2::PC_Real;
+				NewPinSubCategory = UEdGraphSchema_K2::PC_Double;
+				break;
+			case ECameraVariableType::Vector2d:
+				NewPinCategory = UEdGraphSchema_K2::PC_Struct;
+				NewPinSubCategoryObject = TBaseStructure<FVector2D>::Get();
+				break;
+			case ECameraVariableType::Vector3d:
+				NewPinCategory = UEdGraphSchema_K2::PC_Struct;
+				NewPinSubCategoryObject = TBaseStructure<FVector>::Get();
+				break;
+			case ECameraVariableType::Vector4d:
+				NewPinCategory = UEdGraphSchema_K2::PC_Struct;
+				NewPinSubCategoryObject = TBaseStructure<FVector4>::Get();
+				break;
+			case ECameraVariableType::Rotator3d:
+				NewPinCategory = UEdGraphSchema_K2::PC_Struct;
+				NewPinSubCategoryObject = TBaseStructure<FRotator>::Get();
+				break;
+			case ECameraVariableType::Transform3d:
+				NewPinCategory = UEdGraphSchema_K2::PC_Struct;
+				NewPinSubCategoryObject = TBaseStructure<FTransform>::Get();
+				break;
+		}
+		if (NewPinCategory.IsNone())
+		{
+			// Unsupported type for Blueprints.
+			continue;
+		}
+
+		UEdGraphPin* NewPin = CreatePin(
+				EGPD_Input, 
+				NewPinCategory, NewPinSubCategory, NewPinSubCategoryObject, 
+				FName(InterfaceParameter->InterfaceParameterName));
+		if (CreatedPins)
+		{
+			CreatedPins->Add(NewPin);
+		}
+	}
+}
+
+UCameraRigAsset* UK2Node_SetCameraRigParameters::GetCameraRig(TArrayView<UEdGraphPin* const>* InPinsToSearch) const
+{
+	TArrayView<UEdGraphPin* const> PinsToSearch(Pins);
+	if (InPinsToSearch)
+	{
+		PinsToSearch = *InPinsToSearch;
+	}
+
+	UEdGraphPin* CameraRigPin = GetCameraRigPin(&PinsToSearch);
+	if (CameraRigPin && CameraRigPin->DefaultObject && CameraRigPin->LinkedTo.Num() == 0)
+	{
+		return CastChecked<UCameraRigAsset>(CameraRigPin->DefaultObject);
+	}
+	else if (CameraRigPin && CameraRigPin->LinkedTo.Num() > 0)
+	{
+		if (UEdGraphPin* CameraRigSource = CameraRigPin->LinkedTo[0])
+		{
+			return Cast<UCameraRigAsset>(CameraRigSource->PinType.PinSubCategoryObject.Get());
+		}
+	}
+	return nullptr;
+}
+
+UCameraEvaluationResultInterop* UK2Node_SetCameraRigParameters::GetCameraEvaluationResult() const
+{
+	UEdGraphPin* ResultPin = GetCameraEvaluationResultPin();
+	if (ResultPin && ResultPin->DefaultObject && ResultPin->LinkedTo.Num() == 0)
+	{
+		return CastChecked<UCameraEvaluationResultInterop>(ResultPin->DefaultObject);
+	}
+	else if (ResultPin && ResultPin->LinkedTo.Num() > 0)
+	{
+		if (UEdGraphPin* ResultSource = ResultPin->LinkedTo[0])
+		{
+			return Cast<UCameraEvaluationResultInterop>(ResultSource->PinType.PinSubCategoryObject.Get());
+		}
+	}
+	return nullptr;
+}
+
+void UK2Node_SetCameraRigParameters::OnCameraRigChanged()
+{
+	TArray<UEdGraphPin*> OldPins = Pins;
+	TArray<UEdGraphPin*> OldCameraRigPins;
+
+	for (UEdGraphPin* OldPin : OldPins)
+	{
+		if (IsCameraRigParameterPin(OldPin))
+		{
+			Pins.Remove(OldPin);
+			OldCameraRigPins.Add(OldPin);
+		}
+	}
+
+	TArray<UEdGraphPin*> NewPins;
+	if (UCameraRigAsset* CameraRig = GetCameraRig())
+	{
+		CreatePinsForCameraRig(CameraRig, &NewPins);
+	}
+
+	RewireOldPinsToNewPins(OldCameraRigPins, Pins, nullptr);
+
+	GetGraph()->NotifyGraphChanged();
+	FBlueprintEditorUtils::MarkBlueprintAsModified(GetBlueprint());
+}
+
+#undef LOCTEXT_NAMESPACE
+
diff --git a/Engine/Plugins/Cameras/GameplayCameras/Source/GameplayCamerasUncookedOnly/Private/BlueprintGraph/K2Node_SetCameraRigParameters.h b/Engine/Plugins/Cameras/GameplayCameras/Source/GameplayCamerasUncookedOnly/Private/BlueprintGraph/K2Node_SetCameraRigParameters.h
new file mode 100644
--- /dev/null
+++ b/Engine/Plugins/Cameras/GameplayCameras/Source/GameplayCamerasUncookedOnly/Private/BlueprintGraph/K2Node_SetCameraRigParameters.h
@@ -0,0 +1,62 @@
+// Copyright Epic Games, Inc. All Rights Reserved.
+
+#pragma once
+
+#include "CoreTypes.h"
+#include "EdGraph/EdGraphNode.h"
+#include "K2Node.h"
+#include "UObject/ObjectMacros.h"
+#include "UObject/UObjectGlobals.h"
+
+#include "K2Node_SetCameraRigParameters.generated.h"
+
+class UCameraEvaluationResultInterop;
+class UCameraRigAsset;
+
+/**
+ * Dynamic Blueprint node that, given a camera rig, exposes input pins for any camera rig parameters
+ * found on it. On compile, this node gets expanded into the appropriate number of individual setter
+ * function calls for each parameter (see UCameraRigParameterInterop).
+ */
+UCLASS(MinimalAPI)
+class UK2Node_SetCameraRigParameters : public UK2Node
+{
+	GENERATED_BODY()
+
+public:
+
+	UK2Node_SetCameraRigParameters(const FObjectInitializer& ObjectInit);
+
+public:
+
+	// UEdGraphNode interface.
+	virtual void AllocateDefaultPins() override;
+	virtual FText GetNodeTitle(ENodeTitleType::Type TitleType) const override;
+	virtual void PinDefaultValueChanged(UEdGraphPin* Pin) override;
+	virtual FText GetTooltipText() const override;
+	virtual void PinConnectionListChanged(UEdGraphPin* Pin);
+	virtual void PostPlacedNewNode() override;
+	virtual void ExpandNode(class FKismetCompilerContext& CompilerContext, UEdGraph* SourceGraph) override;
+
+	// UK2Node interface.
+	virtual bool IsNodeSafeToIgnore() const override { return true; }
+	virtual void ReallocatePinsDuringReconstruction(TArray<UEdGraphPin*>& OldPins) override;
+	virtual void GetMenuActions(FBlueprintActionDatabaseRegistrar& ActionRegistrar) const override;
+	virtual FText GetMenuCategory() const override;
+
+protected:
+
+	static const FName CameraRigPinName;
+	static const FName CameraEvaluationResultPinName;
+
+	UEdGraphPin* GetCameraRigPin(TArrayView<UEdGraphPin* const>* InPinsToSearch = nullptr) const;
+	UEdGraphPin* GetCameraEvaluationResultPin() const;
+	void GetCameraRigParameterPins(TArray<UEdGraphPin*>& OutParameterPins) const;
+	bool IsCameraRigParameterPin(UEdGraphPin* Pin) const;
+	void CreatePinsForCameraRig(UCameraRigAsset* CameraRig, TArray<UEdGraphPin*>* CreatedPins = nullptr);
+
+	UCameraRigAsset* GetCameraRig(TArrayView<UEdGraphPin* const>* InPinsToSearch = nullptr) const;
+	UCameraEvaluationResultInterop* GetCameraEvaluationResult() const;
+	void OnCameraRigChanged();
+};
+
diff --git a/Engine/Plugins/Cameras/GameplayCameras/Source/GameplayCamerasUncookedOnly/Private/GameplayCamerasUncookedOnlyModule.cpp b/Engine/Plugins/Cameras/GameplayCameras/Source/GameplayCamerasUncookedOnly/Private/GameplayCamerasUncookedOnlyModule.cpp
new file mode 100644
--- /dev/null
+++ b/Engine/Plugins/Cameras/GameplayCameras/Source/GameplayCamerasUncookedOnly/Private/GameplayCamerasUncookedOnlyModule.cpp
@@ -0,0 +1,20 @@
+// Copyright Epic Games, Inc. All Rights Reserved.
+
+#include "CoreTypes.h"
+#include "Modules/ModuleManager.h"
+
+class FGameplayCamerasUncookedOnlyModule : public IModuleInterface
+{
+protected:
+
+	virtual void StartupModule() override
+	{
+	}
+
+	virtual void ShutdownModule() override
+	{
+	}
+};
+
+IMPLEMENT_MODULE(FGameplayCamerasUncookedOnlyModule, GameplayCamerasUncookedOnly);
+
