Cameras: add 3D line drawing helper functions to the debug renderer

diff --git a/Engine/Plugins/Cameras/GameplayCameras/Source/GameplayCameras/Private/Core/CameraSystemEvaluator.cpp b/Engine/Plugins/Cameras/GameplayCameras/Source/GameplayCameras/Private/Core/CameraSystemEvaluator.cpp
--- a/Engine/Plugins/Cameras/GameplayCameras/Source/GameplayCameras/Private/Core/CameraSystemEvaluator.cpp
+++ b/Engine/Plugins/Cameras/GameplayCameras/Source/GameplayCameras/Private/Core/CameraSystemEvaluator.cpp
@@ -180,16 +180,17 @@ void FCameraSystemEvaluator::DebugUpdate(const FCameraSystemDebugUpdateParams& P
 	FCameraDebugBlockBuilder DebugBlockBuilder(DebugBlockStorage, *RootDebugBlock);
 	RootDebugBlock->BuildDebugBlocks(*this, BuildParams, DebugBlockBuilder);
 
+	UObject* Owner = WeakOwner.Get();
+	UWorld* OwnerWorld = Owner ? Owner->GetWorld() : nullptr;
+
 #if UE_GAMEPLAY_CAMERAS_TRACE
 	if (bTraceEnabled)
 	{
-		UObject* Owner = WeakOwner.Get();
-		UWorld* OwnerWorld = Owner ? Owner->GetWorld() : nullptr;
 		FCameraSystemTrace::TraceEvaluation(OwnerWorld, Result, *RootDebugBlock);
 	}
 #endif
 	
-	FCameraDebugRenderer Renderer(Params.Canvas);
+	FCameraDebugRenderer Renderer(OwnerWorld, Params.Canvas);
 	RootDebugBlock->RootDebugDraw(Renderer);
 }
 
diff --git a/Engine/Plugins/Cameras/GameplayCameras/Source/GameplayCameras/Private/Debug/CameraDebugRenderer.cpp b/Engine/Plugins/Cameras/GameplayCameras/Source/GameplayCameras/Private/Debug/CameraDebugRenderer.cpp
--- a/Engine/Plugins/Cameras/GameplayCameras/Source/GameplayCameras/Private/Debug/CameraDebugRenderer.cpp
+++ b/Engine/Plugins/Cameras/GameplayCameras/Source/GameplayCameras/Private/Debug/CameraDebugRenderer.cpp
@@ -5,10 +5,12 @@
 #include "Algo/Find.h"
 #include "CanvasItem.h"
 #include "CanvasTypes.h"
+#include "Components/LineBatchComponent.h"
 #include "Debug/CameraDebugColors.h"
 #include "Debug/DebugTextRenderer.h"
 #include "Engine/Engine.h"
 #include "Engine/Font.h"
+#include "Engine/World.h"
 #include "HAL/IConsoleManager.h"
 #include "Math/Box2D.h"
 #include "Misc/TVariant.h"
@@ -48,8 +50,9 @@ static FAutoConsoleVariableRef CVarGameplayCamerasDebugBackgroundDepthSortKey(
 	GGameplayCamerasDebugBackgroundDepthSortKey,
 	TEXT(""));
 
-FCameraDebugRenderer::FCameraDebugRenderer(FCanvas* InCanvas)
-	: Canvas(InCanvas)
+FCameraDebugRenderer::FCameraDebugRenderer(UWorld* InWorld, FCanvas* InCanvas)
+	: World(InWorld)
+	, Canvas(InCanvas)
 	, DrawColor(FColor::White)
 {
 	RenderFont = GEngine->GetSmallFont();
@@ -257,6 +260,19 @@ void FCameraDebugRenderer::Draw2DCircle(const FVector2D& Center, float Radius, c
 	}
 }
 
+void FCameraDebugRenderer::DrawLine(const FVector3d& Start, const FVector3d& End, const FLinearColor& LineColor, float LineThickness)
+{
+	if (ULineBatchComponent* LineBatcher = GetDebugLineBatcher())
+	{
+		LineBatcher->DrawLine(Start, End, LineColor, SDPG_Foreground, LineThickness);
+	}
+}
+
+ULineBatchComponent* FCameraDebugRenderer::GetDebugLineBatcher() const
+{
+	return World ? World->ForegroundLineBatcher : nullptr;
+}
+
 void FCameraDebugRenderer::SkipAttachedBlocks()
 {
 	VisitFlags |= ECameraDebugDrawVisitFlags::SkipAttachedBlocks;
diff --git a/Engine/Plugins/Cameras/GameplayCameras/Source/GameplayCameras/Public/Debug/CameraDebugRenderer.h b/Engine/Plugins/Cameras/GameplayCameras/Source/GameplayCameras/Public/Debug/CameraDebugRenderer.h
--- a/Engine/Plugins/Cameras/GameplayCameras/Source/GameplayCameras/Public/Debug/CameraDebugRenderer.h
+++ b/Engine/Plugins/Cameras/GameplayCameras/Source/GameplayCameras/Public/Debug/CameraDebugRenderer.h
@@ -11,6 +11,8 @@
 
 class FCanvas;
 class UFont;
+class ULineBatchComponent;
+class UWorld;
 
 #if UE_GAMEPLAY_CAMERAS_DEBUG
 
@@ -35,7 +37,7 @@ class FCameraDebugRenderer
 public:
 
 	/** Creates a new debug renderer. */
-	GAMEPLAYCAMERAS_API FCameraDebugRenderer(FCanvas* InCanvas);
+	GAMEPLAYCAMERAS_API FCameraDebugRenderer(UWorld* InWorld, FCanvas* InCanvas);
 	/** Destroys the debug renderer. */
 	GAMEPLAYCAMERAS_API ~FCameraDebugRenderer();
 
@@ -74,6 +76,8 @@ public:
 	/** Draws a 2D circle. */
 	GAMEPLAYCAMERAS_API void Draw2DCircle(const FVector2D& Center, float Radius, const FLinearColor& LineColor, float LineThickness = 1.f, int32 NumSides = 0);
 
+	GAMEPLAYCAMERAS_API void DrawLine(const FVector3d& Start, const FVector3d& End, const FLinearColor& LineColor, float LineThickness = 1.f);
+
 public:
 
 	/** Request skipping drawing any blocks attached to the current block. */
@@ -105,8 +109,12 @@ private:
 	float GetIndentMargin() const;
 	void FlushText();
 
+	ULineBatchComponent* GetDebugLineBatcher() const;
+
 private:
 
+	/** The world in which we might draw debug primitives. */
+	UWorld* World;
 	/** The canvas used to draw the text wall. */
 	FCanvas* Canvas;
 	/** The draw color of the canvas. */
diff --git a/Engine/Plugins/Cameras/GameplayCameras/Source/GameplayCamerasEditor/Private/Trace/CameraSystemRewindDebuggerExtension.cpp b/Engine/Plugins/Cameras/GameplayCameras/Source/GameplayCamerasEditor/Private/Trace/CameraSystemRewindDebuggerExtension.cpp
--- a/Engine/Plugins/Cameras/GameplayCameras/Source/GameplayCamerasEditor/Private/Trace/CameraSystemRewindDebuggerExtension.cpp
+++ b/Engine/Plugins/Cameras/GameplayCameras/Source/GameplayCamerasEditor/Private/Trace/CameraSystemRewindDebuggerExtension.cpp
@@ -101,7 +101,8 @@ void FCameraSystemRewindDebuggerExtension::DebugDraw(UCanvas* Canvas, APlayerCon
 {
 	if (RootDebugBlock)
 	{
-		FCameraDebugRenderer CameraDebugRenderer(Canvas->Canvas);
+		UWorld* World = PlayController->GetWorld();
+		FCameraDebugRenderer CameraDebugRenderer(World, Canvas->Canvas);
 		RootDebugBlock->RootDebugDraw(CameraDebugRenderer);
 	}
 }
