Cameras: fix static code analyzer errors

diff --git a/Engine/Plugins/Cameras/GameplayCameras/Source/GameplayCamerasEditor/Private/Editors/CameraNodeGraphSchema.cpp b/Engine/Plugins/Cameras/GameplayCameras/Source/GameplayCamerasEditor/Private/Editors/CameraNodeGraphSchema.cpp
--- a/Engine/Plugins/Cameras/GameplayCameras/Source/GameplayCamerasEditor/Private/Editors/CameraNodeGraphSchema.cpp
+++ b/Engine/Plugins/Cameras/GameplayCameras/Source/GameplayCamerasEditor/Private/Editors/CameraNodeGraphSchema.cpp
@@ -103,7 +103,7 @@ bool UCameraNodeGraphSchema::OnCreateConnection(UEdGraphPin* A, UEdGraphPin* B)
 		CameraParameterPin = B;
 	}
 
-	if (!RigInterfacePin && !CameraParameterPin)
+	if (!RigInterfacePin || !CameraParameterPin)
 	{
 		return false;
 	}
diff --git a/Engine/Plugins/Cameras/GameplayCameras/Source/GameplayCamerasUncookedOnly/Private/BlueprintGraph/K2Node_SetCameraRigParameters.cpp b/Engine/Plugins/Cameras/GameplayCameras/Source/GameplayCamerasUncookedOnly/Private/BlueprintGraph/K2Node_SetCameraRigParameters.cpp
--- a/Engine/Plugins/Cameras/GameplayCameras/Source/GameplayCamerasUncookedOnly/Private/BlueprintGraph/K2Node_SetCameraRigParameters.cpp
+++ b/Engine/Plugins/Cameras/GameplayCameras/Source/GameplayCamerasUncookedOnly/Private/BlueprintGraph/K2Node_SetCameraRigParameters.cpp
@@ -220,7 +220,7 @@ void UK2Node_SetCameraRigParameters::ExpandNode(class FKismetCompilerContext& Co
 	}
 
 	// Connect the last node if necessary.
-	if (OriginalThenPin && OriginalThenPin->LinkedTo.Num() > 0)
+	if (OriginalThenPin && PreviousThenPin && OriginalThenPin->LinkedTo.Num() > 0)
 	{
 		CompilerContext.MovePinLinksToIntermediate(*OriginalThenPin, *PreviousThenPin);
 	}
