Cameras: basic screen-space framing camera node, plus a dolly shot implementation

These nodes try to frame a target in screen-space using a set of three framing rectangular zones (dead zone, soft zone, hard zone).
The base node evaluation implements the screen-space math, and the sub-class implements how to reframe the shot based on these results. The dolly shot node specifically does this by translating across a plane. A future node might implement panning instead, i.e. rotating in place to frame the subject.

Some new math classes are also added to help compute the shot.

diff --git a/Engine/Plugins/Cameras/GameplayCameras/Source/GameplayCameras/Private/Math/CameraFramingZoneMath.cpp b/Engine/Plugins/Cameras/GameplayCameras/Source/GameplayCameras/Private/Math/CameraFramingZoneMath.cpp
new file mode 100644
--- /dev/null
+++ b/Engine/Plugins/Cameras/GameplayCameras/Source/GameplayCameras/Private/Math/CameraFramingZoneMath.cpp
@@ -0,0 +1,121 @@
+// Copyright Epic Games, Inc. All Rights Reserved.
+
+#include "Math/CameraFramingZoneMath.h"
+
+#include "Core/CameraPose.h"
+#include "Nodes/Framing/CameraFramingZone.h"
+
+namespace UE::Cameras
+{
+
+FFramingZone::FFramingZone()
+{
+	LeftBound = 0;
+	TopBound = 0;
+	RightBound = 1;
+	BottomBound = 1;
+}
+
+FFramingZone::FFramingZone(const FFramingZoneMargins& FramingZoneMargins)
+{
+	LeftBound = FramingZoneMargins.LeftMargin;
+	TopBound = FramingZoneMargins.TopMargin;
+	RightBound = 1.0f - FramingZoneMargins.RightMargin;
+	BottomBound = 1.0f - FramingZoneMargins.BottomMargin;
+
+	ClampBounds();
+}
+
+void FFramingZone::ClampBounds()
+{
+	LeftBound = FMath::Clamp(LeftBound, 0.f, 1.f);
+	TopBound = FMath::Clamp(TopBound, 0.f, 1.f);
+	RightBound = FMath::Clamp(RightBound, 0.f, 1.f);
+	BottomBound = FMath::Clamp(BottomBound, 0.f, 1.f);
+
+	RightBound = FMath::Max(LeftBound, RightBound);
+	BottomBound = FMath::Max(TopBound, BottomBound);
+}
+
+void FFramingZone::ClampBounds(const FVector2d& MustContain)
+{
+	LeftBound = FMath::Clamp(LeftBound, 0.f, MustContain.X);
+	TopBound = FMath::Clamp(TopBound, 0.f, MustContain.Y);
+	RightBound = FMath::Clamp(RightBound, MustContain.X, 1.f);
+	BottomBound = FMath::Clamp(BottomBound, MustContain.Y, 1.f);
+}
+
+void FFramingZone::ClampBounds(const FFramingZone& MustContain)
+{
+	LeftBound = FMath::Clamp(LeftBound, 0.f, MustContain.LeftBound);
+	TopBound = FMath::Clamp(TopBound, 0.f, MustContain.TopBound);
+	RightBound = FMath::Clamp(RightBound, MustContain.RightBound, 1.f);
+	BottomBound = FMath::Clamp(BottomBound, MustContain.BottomBound, 1.f);
+}
+
+bool FFramingZone::Contains(const FVector2d& Point) const
+{
+	return Point.X >= LeftBound && Point.X <= RightBound &&
+		Point.Y >= TopBound && Point.Y <= BottomBound;
+}
+
+FVector4d FFramingZone::GetMarginsFromCenter() const
+{
+	return FVector4d(
+		0.5 - LeftBound,
+		0.5 - TopBound,
+		0.5 - (1.0 - RightBound),
+		0.5 - (1.0 - BottomBound));
+}
+
+FVector2d FFramingZone::GetCanvasPosition(const FVector2d& CanvasSize) const
+{
+	return FVector2d(LeftBound * CanvasSize.X, TopBound * CanvasSize.Y);
+}
+
+FVector2d FFramingZone::GetCanvasSize(const FVector2d& CanvasSize) const
+{
+	return FVector2d((RightBound - LeftBound) * CanvasSize.X, (BottomBound - TopBound) * CanvasSize.Y);
+}
+
+void FFramingZone::Serialize(FArchive& Ar)
+{
+	Ar << LeftBound;
+	Ar << TopBound;
+	Ar << RightBound;
+	Ar << BottomBound;
+}
+
+FArchive& operator <<(FArchive& Ar, FFramingZone& FramingZone)
+{
+	FramingZone.Serialize(Ar);
+	return Ar;
+}
+
+FFramingZoneAngles FFramingZoneMath::GetFramingZoneAngles(const FFramingZone& FramingZone, const FCameraFieldsOfView& FieldsOfView)
+{
+	const float BackingHalfWidth = FMath::Tan(FMath::DegreesToRadians(FieldsOfView.HorizontalFieldOfView / 2.f));
+	const float BackingHalfHeight = FMath::Tan(FMath::DegreesToRadians(FieldsOfView.VerticalFieldOfView / 2.f));
+
+	const FVector4d MarginsFromCenter = FramingZone.GetMarginsFromCenter();
+
+	FFramingZoneAngles Angles;
+	Angles.LeftHalfAngle = GetFramingMarginAngle(MarginsFromCenter.X, BackingHalfWidth);
+	Angles.TopHalfAngle = GetFramingMarginAngle(MarginsFromCenter.Y, BackingHalfHeight);
+	Angles.RightHalfAngle = GetFramingMarginAngle(MarginsFromCenter.Z, BackingHalfWidth);
+	Angles.BottomHalfAngle = GetFramingMarginAngle(MarginsFromCenter.W, BackingHalfHeight);
+	return Angles;
+}
+
+double FFramingZoneMath::GetFramingMarginAngle(float MarginFromCenter, float BackingHalfSize)
+{
+	// The margin should always be a percentage from the center of the screen.
+	// So a margin of zero is the center, a margin of 0.5 is half way between the center and the edge,
+	// and a margin of 1 is at the edge.
+	const double MarginSize = FMath::Clamp(MarginFromCenter, -1.f, 1.f) * BackingHalfSize;
+	const double MarginAngle = FMath::Atan(MarginSize);
+	return MarginAngle;
+}
+
+}  // namespace UE::Cameras
+
diff --git a/Engine/Plugins/Cameras/GameplayCameras/Source/GameplayCameras/Private/Math/CameraPoseMath.cpp b/Engine/Plugins/Cameras/GameplayCameras/Source/GameplayCameras/Private/Math/CameraPoseMath.cpp
new file mode 100644
--- /dev/null
+++ b/Engine/Plugins/Cameras/GameplayCameras/Source/GameplayCameras/Private/Math/CameraPoseMath.cpp
@@ -0,0 +1,196 @@
+// Copyright Epic Games, Inc. All Rights Reserved.
+
+#include "Math/CameraPoseMath.h"
+
+#include "Core/CameraPose.h"
+#include "CoreGlobals.h"
+#include "Math/InverseRotationMatrix.h"
+#include "Math/PerspectiveMatrix.h"
+#include "Math/TranslationMatrix.h"
+#include "RHI.h"
+
+namespace UE::Cameras
+{
+
+FMatrix FCameraPoseMath::BuildProjectionMatrix(const FCameraPose& CameraPose)
+{
+	const double NearClippingPlane = CameraPose.GetNearClippingPlane() > 0.0f ? 
+		CameraPose.GetNearClippingPlane() : GNearClippingPlane;
+	const double FieldOfView = FMath::Max(CameraPose.GetEffectiveFieldOfView(), 0.001f);
+	const float AspectRatio = CameraPose.GetSensorAspectRatio();
+
+	FMatrix ProjectionMatrix = FReversedZPerspectiveMatrix(
+			FMath::DegreesToRadians(FieldOfView / 2.f),
+			AspectRatio,
+			1.0f,
+			NearClippingPlane);
+	//ProjectionMatrix = AdjustProjectionMatrixForRHI(ProjectionMatrix);
+	return ProjectionMatrix;
+}
+
+FMatrix FCameraPoseMath::BuildViewProjectionMatrix(const FCameraPose& CameraPose)
+{
+	const FTranslationMatrix InverseOrigin(-CameraPose.GetLocation());
+
+	// Somehow we need to also transpose... code stolen from ULocalPlayer::GetProjectionData...
+	const FMatrix InverseRotation = FInverseRotationMatrix(CameraPose.GetRotation()) * FMatrix(
+		FPlane(0,	0,	1,	0),
+		FPlane(1,	0,	0,	0),
+		FPlane(0,	1,	0,	0),
+		FPlane(0,	0,	0,	1));
+
+	const FMatrix ProjectionMatrix = BuildProjectionMatrix(CameraPose);
+
+	return InverseOrigin * InverseRotation * ProjectionMatrix;
+}
+
+TOptional<FVector2d> FCameraPoseMath::ProjectWorldToScreen(
+		const FCameraPose& CameraPose, const FVector3d& WorldLocation, bool bForceLocationInsideFrustum)
+{
+	const FMatrix ViewProjectionMatrix = BuildViewProjectionMatrix(CameraPose);
+	return ProjectToScreen(ViewProjectionMatrix, WorldLocation, bForceLocationInsideFrustum);
+}
+
+TOptional<FVector2d> FCameraPoseMath::ProjectCameraToScreen(
+			const FCameraPose& CameraPose, const FVector3d& CameraSpaceLocation, bool bForceLocationInsideFrustum)
+{
+	const FMatrix ProjectionMatrix = BuildProjectionMatrix(CameraPose);
+	return ProjectToScreen(ProjectionMatrix, CameraSpaceLocation, bForceLocationInsideFrustum);
+}
+
+TOptional<FVector2d> FCameraPoseMath::ProjectToScreen(
+		const FMatrix& ViewProjectionMatrix, const FVector3d& Location, bool bForceLocationInsideFrustum)
+{
+	const FVector4d ProjectedLocation = ViewProjectionMatrix.TransformFVector4(FVector4(Location, 1.f));
+
+	// See if we need to handle the case of a point outside of the view frustum.
+	const bool bIsInsideFrustum = (ProjectedLocation.W > 0.f);
+	double W = ProjectedLocation.W;
+	if (!bIsInsideFrustum)
+	{
+		if (!bForceLocationInsideFrustum)
+		{
+			return TOptional<FVector2d>();
+		}
+
+		W = FMath::Abs(W);
+	}
+
+	// The result of this will be coordinates in -1..1 projection space.
+	const double RHW = 1.0f / W;
+	const FVector4d ScreenSpaceLocation(
+			ProjectedLocation.X * RHW, 
+			ProjectedLocation.Y * RHW, 
+			ProjectedLocation.Z * RHW, 
+			ProjectedLocation.W);
+
+	// Move from projection space to normalized 0..1 UI space.
+	const double ScreenSpaceX = (ScreenSpaceLocation.X / 2.f) + 0.5f;
+	const double ScreenSpaceY = 1.f - (ScreenSpaceLocation.Y / 2.f) - 0.5f;
+
+	return FVector2d(ScreenSpaceX, ScreenSpaceY);
+}
+
+FRay3d FCameraPoseMath::UnprojectScreenToCamera(const FCameraPose& CameraPose, const FVector2D& ScreenSpacePoint)
+{
+	const FMatrix ProjectionMatrix = BuildProjectionMatrix(CameraPose);
+	const FMatrix InvProjectionMatrix = ProjectionMatrix.InverseFast();
+	return UnprojectFromScreen(InvProjectionMatrix, ScreenSpacePoint);
+}
+
+FVector3d FCameraPoseMath::UnprojectScreenToCamera(const FCameraPose& CameraPose, const FVector2D& ScreenSpacePoint, double PredictedDistance)
+{
+	const FRay3d UnprojectedRay = UnprojectScreenToCamera(CameraPose, ScreenSpacePoint);
+	const FVector3d WorldPoint = UnprojectedRay.PointAt(PredictedDistance);
+	return WorldPoint;
+}
+
+FRay3d FCameraPoseMath::UnprojectScreenToWorld(const FCameraPose& CameraPose, const FVector2D& ScreenSpacePoint)
+{
+	const FMatrix ViewProjectionMatrix = BuildViewProjectionMatrix(CameraPose);
+	const FMatrix InvViewProjectionMatrix = ViewProjectionMatrix.InverseFast();
+	return UnprojectFromScreen(InvViewProjectionMatrix, ScreenSpacePoint);
+}
+
+FVector3d FCameraPoseMath::UnprojectScreenToWorld(const FCameraPose& CameraPose, const FVector2D& ScreenSpacePoint, double PredictedDistance)
+{
+	const FRay3d UnprojectedRay = UnprojectScreenToWorld(CameraPose, ScreenSpacePoint);
+	const FVector3d WorldPoint = UnprojectedRay.PointAt(PredictedDistance);
+	return WorldPoint;
+}
+
+FRay3d FCameraPoseMath::UnprojectFromScreen(const FMatrix& InverseViewProjectionMatrix, const FVector2D& ScreenSpacePoint)
+{
+	// Convert the given screen-space point from 0..1 UI space to -1..1 projection space.
+	const double ScreenSpaceX = (ScreenSpacePoint.X - 0.5f) * 2.f;
+	const double ScreenSpaceY = ((1.f - ScreenSpacePoint.Y) - 0.5f) * 2.f;
+
+	// Build a ray from the front of the frustum to the back of the frustum, starting at the screen-space point.
+	// We use reverse-Z projection matrices for better precision, so near is Z=1, and far is Z=0.
+	const FVector4 RayStartProjectionSpace = FVector4(ScreenSpaceX, ScreenSpaceY, 1.f, 1.f);
+	const FVector4 RayEndProjectionSpace = FVector4(ScreenSpaceX, ScreenSpaceY, 0.01f, 1.f);
+
+	// Unproject the ray points and normalize them.
+	const FVector4 RayStartProjected = InverseViewProjectionMatrix.TransformFVector4(RayStartProjectionSpace);
+	const FVector4 RayEndProjected = InverseViewProjectionMatrix.TransformFVector4(RayEndProjectionSpace);
+
+	FVector RayStartWorldSpace(RayStartProjected.X, RayStartProjected.Y, RayStartProjected.Z);
+	if (RayStartProjected.W != 0.f)
+	{
+		RayStartWorldSpace /= RayStartProjected.W;
+	}
+	FVector RayEndWorldSpace(RayEndProjected.X, RayEndProjected.Y, RayEndProjected.Z);
+	if (RayEndProjected.W != 0.f)
+	{
+		RayEndWorldSpace /= RayEndProjected.W;
+	}
+
+	// Make the 3D ray.
+	const bool bDirectionIsNormalized = true;
+	const FVector RayDirWorldSpace = (RayEndWorldSpace - RayStartWorldSpace).GetSafeNormal();
+	return FRay3d(RayStartWorldSpace, RayDirWorldSpace, bDirectionIsNormalized);
+}
+
+FVector3d FCameraPoseMath::UnprojectFromScreen(const FMatrix& InverseViewProjectionMatrix, const FVector2D& ScreenSpacePoint, double PredictedDistance)
+{
+	const FRay3d UnprojectedRay = UnprojectFromScreen(InverseViewProjectionMatrix, ScreenSpacePoint);
+	const FVector3d WorldPoint = UnprojectedRay.PointAt(PredictedDistance);
+	return WorldPoint;
+}
+
+FMatrix FCameraPoseMath::InverseProjectionMatrix(const FMatrix& ProjectionMatrix)
+{
+	// Stolen from SceneView.h
+	if (ProjectionMatrix.M[1][0] == 0.0f &&
+		ProjectionMatrix.M[3][0] == 0.0f &&
+		ProjectionMatrix.M[0][1] == 0.0f &&
+		ProjectionMatrix.M[3][1] == 0.0f &&
+		ProjectionMatrix.M[0][2] == 0.0f &&
+		ProjectionMatrix.M[1][2] == 0.0f &&
+		ProjectionMatrix.M[0][3] == 0.0f &&
+		ProjectionMatrix.M[1][3] == 0.0f &&
+		ProjectionMatrix.M[2][3] == 1.0f &&
+		ProjectionMatrix.M[3][3] == 0.0f)
+	{
+		double a = ProjectionMatrix.M[0][0];
+		double b = ProjectionMatrix.M[1][1];
+		double c = ProjectionMatrix.M[2][2];
+		double d = ProjectionMatrix.M[3][2];
+		double s = ProjectionMatrix.M[2][0];
+		double t = ProjectionMatrix.M[2][1];
+
+		return FMatrix(
+			FPlane(1.0 / a, 0.0f, 0.0f, 0.0f),
+			FPlane(0.0f, 1.0 / b, 0.0f, 0.0f),
+			FPlane(0.0f, 0.0f, 0.0f, 1.0 / d),
+			FPlane(-s/a, -t/b, 1.0f, -c/d)
+		);
+	}
+	else
+	{
+		return ProjectionMatrix.Inverse();
+	}
+}
+
+}  // namespace UE::Cameras
+
diff --git a/Engine/Plugins/Cameras/GameplayCameras/Source/GameplayCameras/Private/Nodes/Framing/BaseFramingCameraNode.cpp b/Engine/Plugins/Cameras/GameplayCameras/Source/GameplayCameras/Private/Nodes/Framing/BaseFramingCameraNode.cpp
new file mode 100644
--- /dev/null
+++ b/Engine/Plugins/Cameras/GameplayCameras/Source/GameplayCameras/Private/Nodes/Framing/BaseFramingCameraNode.cpp
@@ -0,0 +1,392 @@
+// Copyright Epic Games, Inc. All Rights Reserved.
+
+#include "Nodes/Framing/BaseFramingCameraNode.h"
+
+#include "CanvasItem.h"
+#include "CanvasTypes.h"
+#include "Core/CameraEvaluationContext.h"
+#include "Core/CameraParameterReader.h"
+#include "Debug/CameraDebugBlock.h"
+#include "Debug/CameraDebugBlockBuilder.h"
+#include "Debug/CameraDebugRenderer.h"
+#include "GameplayCameras.h"
+#include "Math/CameraPoseMath.h"
+#include "Math/ColorList.h"
+
+#include UE_INLINE_GENERATED_CPP_BY_NAME(BaseFramingCameraNode)
+
+namespace UE::Cameras
+{
+
+UE_DECLARE_CAMERA_DEBUG_BLOCK_START(GAMEPLAYCAMERAS_API, FBaseFramingCameraDebugBlock)
+	UE_DECLARE_CAMERA_DEBUG_BLOCK_FIELD(FBaseFramingCameraNodeEvaluator::FState, State);
+	UE_DECLARE_CAMERA_DEBUG_BLOCK_FIELD(FBaseFramingCameraNodeEvaluator::FDesired, Desired);
+UE_DECLARE_CAMERA_DEBUG_BLOCK_END()
+
+UE_DEFINE_CAMERA_DEBUG_BLOCK_WITH_FIELDS(FBaseFramingCameraDebugBlock)
+
+void FCameraFramingZoneParameterReader::Initialize(const FCameraFramingZone& FramingZone)
+{
+	LeftMargin.Initialize(FramingZone.LeftMargin);
+	TopMargin.Initialize(FramingZone.TopMargin);
+	RightMargin.Initialize(FramingZone.RightMargin);
+	BottomMargin.Initialize(FramingZone.BottomMargin);
+}
+
+FFramingZoneMargins FCameraFramingZoneParameterReader::GetZoneMargins(const FCameraVariableTable& VariableTable) const
+{
+	FFramingZoneMargins Margins;
+	Margins.LeftMargin = LeftMargin.Get(VariableTable);
+	Margins.TopMargin = TopMargin.Get(VariableTable);
+	Margins.RightMargin = RightMargin.Get(VariableTable);
+	Margins.BottomMargin = BottomMargin.Get(VariableTable);
+	return Margins;
+}
+
+UE_DEFINE_CAMERA_NODE_EVALUATOR(FBaseFramingCameraNodeEvaluator)
+
+void FBaseFramingCameraNodeEvaluator::OnInitialize(const FCameraNodeEvaluatorInitializeParams& Params)
+{
+	const UBaseFramingCameraNode* BaseFramingNode = GetCameraNodeAs<UBaseFramingCameraNode>();
+
+	Readers.HorizontalFraming.Initialize(BaseFramingNode->HorizontalFraming);
+	Readers.VerticalFraming.Initialize(BaseFramingNode->VerticalFraming);
+	Readers.ReframeDampingFactor.Initialize(BaseFramingNode->ReframeDampingFactor);
+	Readers.ReframeUnlockRadius.Initialize(BaseFramingNode->ReframeUnlockRadius);
+
+	Readers.DeadZoneMargin.Initialize(BaseFramingNode->DeadZone);
+	Readers.SoftZoneMargin.Initialize(BaseFramingNode->SoftZone);
+
+	State.bIsFirstFrame = true;
+}
+
+void FBaseFramingCameraNodeEvaluator::UpdateFramingState(const FCameraNodeEvaluationResult& OutResult, const FVector3d& TargetLocation)
+{
+	// Get screen-space coordinates of the ideal framing point. These are in 0..1 UI space.
+	State.IdealTarget.X = Readers.HorizontalFraming.Get(OutResult.VariableTable);
+	State.IdealTarget.Y = Readers.VerticalFraming.Get(OutResult.VariableTable);
+	State.ReframeDampingFactor = Readers.ReframeDampingFactor.Get(OutResult.VariableTable);
+	State.ReframeUnlockRadius = Readers.ReframeUnlockRadius.Get(OutResult.VariableTable);
+
+	// Update the damping factor in case it's driven by a variable.
+	State.ReframeDamper.SetW0(Readers.ReframeDampingFactor.Get(OutResult.VariableTable));
+
+	// Get the effective margins of the framing zones for this frame.
+	const FFramingZoneMargins DeadZoneMargins = Readers.DeadZoneMargin.GetZoneMargins(OutResult.VariableTable);
+	const FFramingZoneMargins SoftZoneMargins = Readers.SoftZoneMargin.GetZoneMargins(OutResult.VariableTable);
+
+	// Compute the UI space coordinates of the framing zones.
+	State.DeadZone = FFramingZone(DeadZoneMargins);
+	State.SoftZone = FFramingZone(SoftZoneMargins);
+
+	// On the first frame, initialize LastFrameTransform to the current so we don't use garbage.
+	if (State.bIsFirstFrame)
+	{
+		State.LastFramedTransform = OutResult.CameraPose.GetTransform();
+		State.bIsFirstFrame = false;
+	}
+
+	// We are going to reframing things iteratively, so we'll use a temporary pose defined by last frame's
+	// shot transform.
+	FCameraPose TempPose(OutResult.CameraPose);
+	TempPose.SetTransform(State.LastFramedTransform);
+
+	// Get the target in screen-space.
+	const TOptional<FVector2d> ScreenTarget = FCameraPoseMath::ProjectWorldToScreen(TempPose, TargetLocation, true);
+	State.WorldTarget = TargetLocation;
+	State.ScreenTarget = ScreenTarget.Get(FVector2d(0.5, 0.5));
+
+	// Make sure our framing zones are hierarchically correct: soft zone contains the dead zone, which contains
+	// the ideal target.
+	State.DeadZone.ClampBounds(State.IdealTarget);
+	State.SoftZone.ClampBounds(State.DeadZone);
+
+	const bool bIsInSoftZone = State.SoftZone.Contains(State.ScreenTarget);
+	const bool bIsInDeadZone = State.DeadZone.Contains(State.ScreenTarget);
+	if (!ScreenTarget.IsSet() || !bIsInSoftZone)
+	{
+		// Target is out of view or outside the soft zone -- it's therefore in the hard zone and we will
+		// do everything we can to put it back in the soft zone ASAP.
+		State.TargetFramingState = ETargetFramingState::InHardZone;
+		State.bIsReframingTarget = true;
+	}
+	else if (!bIsInDeadZone)
+	{
+		// Target is in the soft zone so we will gently reframe it towards the ideal framing.
+		State.TargetFramingState = ETargetFramingState::InSoftZone;
+		State.bIsReframingTarget = true;
+	}
+	else
+	{
+		// Target is in the dead zone.
+		State.TargetFramingState = ETargetFramingState::InDeadZone;
+
+		// Even though the target is free to move inside the dead zone, we might still want to continue reframing
+		// it towars the ideal position... if we didn't do that, reframing from the soft zone would stop entirely 
+		// once we reach the edge of the dead zone, and we would never really ever see the target near the ideal
+		// position. So if we re-enter the dead zone from the soft zone, we keep reframing until we hit a smaller
+		// "unlock reframing" zone defined by the ReframeUnlockRadius.
+		//
+		// If reframing wasn't active in the first place (e.g. the target is freely moving inside the dead zone), 
+		// we don't do anything and let it be.
+		if (State.bIsReframingTarget)
+		{
+			// Since screen-space positions are in 0..1 space, we need to modulate vertical coordinates by the
+			// aspect ratio, otherwise we end up comparing against a squished ellipse instead of a circle.
+			const float AspectRatio = OutResult.CameraPose.GetSensorAspectRatio();
+			const double DistanceToIdeal = FVector2d::Distance(
+					FVector2d(State.ScreenTarget.X, (State.ScreenTarget.Y - 0.5) / AspectRatio + 0.5),
+					FVector2d(State.IdealTarget.X, (State.IdealTarget.Y - 0.5) / AspectRatio + 0.5));
+			if (DistanceToIdeal <= State.ReframeUnlockRadius)
+			{
+				State.bIsReframingTarget = false;
+			}
+		}
+	}
+}
+
+void FBaseFramingCameraNodeEvaluator::ComputeDesiredState(float DeltaTime)
+{
+	// If we  don't have any reframing to do, bail out.
+	FVector2d IdealToTarget(State.ScreenTarget - State.IdealTarget);
+	double DistanceToGo = IdealToTarget.Length();
+	if (!State.bIsReframingTarget || DistanceToGo == 0.f)
+	{
+		ensure(State.TargetFramingState == ETargetFramingState::InDeadZone);
+		Desired.ScreenTarget = State.ScreenTarget;
+		Desired.FramingCorrection = FVector2d::ZeroVector;
+		Desired.bHasCorrection = false;
+		return;
+	}
+
+	if (State.TargetFramingState == ETargetFramingState::InHardZone)
+	{
+		// Bring the target immediately to edge of the soft zone, in the direction of the 
+		// ideal position. From there, follow-up with applying the soft zone effect.
+		Desired.ScreenTarget = GetHardReframeCoords();
+
+		IdealToTarget = (Desired.ScreenTarget - State.IdealTarget);
+		DistanceToGo = IdealToTarget.Length();
+	}
+
+	// Move the target towards the ideal position using some damping.
+	const double NewDistanceToGo = State.ReframeDamper.Update(DistanceToGo, DeltaTime);
+
+	// Compute where we want the target this frame.
+	const FVector2d InvReframeDir(IdealToTarget / DistanceToGo);
+	const FVector2d NewScreenTarget = State.IdealTarget + InvReframeDir * NewDistanceToGo;
+	Desired.ScreenTarget = NewScreenTarget;
+
+	Desired.FramingCorrection = Desired.ScreenTarget - State.ScreenTarget;
+	Desired.bHasCorrection = true;
+}
+
+const FTransform& FBaseFramingCameraNodeEvaluator::GetLastFraming() const
+{
+	return State.LastFramedTransform;
+}
+
+void FBaseFramingCameraNodeEvaluator::RegisterNewFraming(const FTransform& NewPoseTransform)
+{
+	State.LastFramedTransform = NewPoseTransform;
+}
+
+FVector2d FBaseFramingCameraNodeEvaluator::GetHardReframeCoords() const
+{
+	// The target is in the hard zone and must be brought back to the edge of the soft zone.
+	// Let's compute the diagonal between the target and the ideal framing point, and bring
+	// the target where that diagonal intersects the soft zone.
+	const FVector2d Diagonal(State.IdealTarget - State.ScreenTarget);
+	if (Diagonal.X == 0.f && Diagonal.Y == 0.f)
+	{
+		// Somehow we're already on the desired framing. This shouldn't happen, we're supposed
+		// to be in the hard zone right now...
+		ensure(false);
+		return State.ScreenTarget;
+	}
+	if (Diagonal.X == 0.f)
+	{
+		// The target is directly above/below the ideal position.
+		return FVector2d(
+				State.IdealTarget.X,
+				(Diagonal.Y > 0 ? State.SoftZone.TopBound : State.SoftZone.BottomBound)
+				);
+	}
+	else if (Diagonal.Y == 0.f)
+	{
+		// The target is directly to the left/right of the ideal position.
+		return FVector2d(
+				(Diagonal.X > 0 ? State.SoftZone.LeftBound : State.SoftZone.RightBound),
+				State.IdealTarget.Y
+				);
+	}
+	else
+	{
+		// The diagonal's equation is P = V*d + P0
+		//
+		//		V = the unit direction vector of the diagonal
+		//		d = the distance
+		//		P0 = a known reference point (we can use the ideal target for this)
+		//
+		// We want to find P where P is on the edge of the soft zone. This means we know
+		// either P.x or P.y depending on the edge. 
+		//
+		// Let's say we deal with a vertical edge... then we know about P.x, because we
+		// want P to be on that edge. We can therefor compute d:
+		//
+		//		P.x = V.x*d + P0.x
+		//		d = (P.x - P0.x) / V.x
+		//
+		// If it's a horizontal edge, we can do the same:
+		//
+		//		P.y = V.y*d + P0.y
+		//		d = (P.y - P0.y) / V.y
+		//
+		// Keep going around the zone until we're sure the point is on the edge.
+		//
+		const FVector2d P0(State.IdealTarget);
+		const FVector2d V(Diagonal.GetSafeNormal());
+		double D = 0.f;
+		FVector2d P(State.ScreenTarget);
+
+		if (P.X < State.SoftZone.LeftBound)
+		{
+			D = (State.SoftZone.LeftBound - P0.X) / V.X;
+			P = V * D + P0;
+		}
+		if (P.Y < State.SoftZone.TopBound)
+		{
+			D = (State.SoftZone.TopBound - P0.Y) / V.Y;
+			P = V * D + P0;
+		}
+		if (P.X > State.SoftZone.RightBound)
+		{
+			D = (State.SoftZone.RightBound - P0.X) / V.X;
+			P = V * D + P0;
+		}
+		if (P.X > State.SoftZone.BottomBound)
+		{
+			D = (State.SoftZone.BottomBound - P0.Y) / V.Y;
+			P = V * D + P0;
+		}
+
+		return P;
+	}
+}
+
+#if UE_GAMEPLAY_CAMERAS_DEBUG
+
+void FBaseFramingCameraNodeEvaluator::OnBuildDebugBlocks(const FCameraDebugBlockBuildParams& Params, FCameraDebugBlockBuilder& Builder)
+{
+	FBaseFramingCameraDebugBlock& DebugBlock = Builder.AttachDebugBlock<FBaseFramingCameraDebugBlock>();
+
+	DebugBlock.State = State;
+	DebugBlock.Desired = Desired;
+}
+
+void FBaseFramingCameraDebugBlock::OnDebugDraw(const FCameraDebugBlockDrawParams& Params, FCameraDebugRenderer& Renderer)
+{
+	FString StateString(TEXT("Unknown"));
+	switch (State.TargetFramingState)
+	{
+		case FBaseFramingCameraNodeEvaluator::ETargetFramingState::InDeadZone:
+			StateString = TEXT("Dead Zone");
+			break;
+		case FBaseFramingCameraNodeEvaluator::ETargetFramingState::InSoftZone:
+			StateString = TEXT("Soft Zone");
+			break;
+		case FBaseFramingCameraNodeEvaluator::ETargetFramingState::InHardZone:
+			StateString = TEXT("Hard Zone");
+			break;
+	}
+
+	Renderer.AddText(TEXT("state [%s]"), *StateString);
+	if (State.bIsReframingTarget)
+	{
+		Renderer.AddText(TEXT("[REFRAMING]"));
+	}
+
+	Renderer.AddText(TEXT(" (damping = %0.3f)"), State.ReframeDamper.GetX0());
+
+	if (Renderer.HasCanvas())
+	{
+		const FVector2D CanvasSize = Renderer.GetCanvasSize();
+
+		const float ReframeUnlockRadius(State.ReframeUnlockRadius * CanvasSize.X);
+		const FVector2D IdealTarget(State.IdealTarget.X * CanvasSize.X, State.IdealTarget.Y * CanvasSize.Y);
+
+		Renderer.Draw2DBox(
+				State.SoftZone.GetCanvasPosition(CanvasSize), 
+				State.SoftZone.GetCanvasSize(CanvasSize),
+				FLinearColor::Red,
+				1.f);
+		Renderer.Draw2DBox(
+				State.DeadZone.GetCanvasPosition(CanvasSize), 
+				State.DeadZone.GetCanvasSize(CanvasSize),
+				FLinearColor::Green,
+				1.f);
+		Renderer.Draw2DCircle(IdealTarget, ReframeUnlockRadius, FLinearColor(FColorList::PaleGreen), 1.f);
+
+		const FVector2D FramingCorrection(Desired.FramingCorrection.X * CanvasSize.X, Desired.FramingCorrection.Y * CanvasSize.Y);
+
+		Renderer.AddText(TEXT("  correction (%0.1f ; %0.1f)"),
+				State.ReframeDamper.GetX0(),
+				FramingCorrection.X, FramingCorrection.Y);
+
+		const FVector2D ScreenTarget(State.ScreenTarget.X * CanvasSize.X, State.ScreenTarget.Y * CanvasSize.Y);
+		const FVector2D NextScreenTarget(Desired.ScreenTarget.X * CanvasSize.X, Desired.ScreenTarget.Y * CanvasSize.Y);
+
+		Renderer.AddText(TEXT("  target (%0.1f; %0.1f)"), ScreenTarget.X, ScreenTarget.Y);
+		Renderer.Draw2DLine(ScreenTarget, NextScreenTarget, FLinearColor(FColorList::Salmon), 1.f);
+		Renderer.Draw2DCircle(ScreenTarget, 2.f, FLinearColor(FColorList::Orange), 2.f);
+		Renderer.Draw2DCircle(IdealTarget, 2.f, FLinearColor::Green, 2.f);
+	}
+}
+
+#endif  // UE_GAMEPLAY_CAMERAS_DEBUG
+
+void FBaseFramingCameraNodeEvaluator::FState::Serialize(FArchive& Ar)
+{
+	Ar << IdealTarget;
+	Ar << ReframeDampingFactor;
+	Ar << ReframeUnlockRadius;
+	Ar << DeadZone;
+	Ar << SoftZone;
+	Ar << ScreenTarget;
+	Ar << TargetFramingState;
+	Ar << bIsReframingTarget;
+	Ar << ReframeDamper;
+}
+
+void FBaseFramingCameraNodeEvaluator::FDesired::Serialize(FArchive& Ar)
+{
+	Ar << ScreenTarget;
+	Ar << FramingCorrection;
+	Ar << bHasCorrection;
+}
+
+FArchive& operator <<(FArchive& Ar, FBaseFramingCameraNodeEvaluator::FState& State)
+{
+	State.Serialize(Ar);
+	return Ar;
+}
+
+FArchive& operator <<(FArchive& Ar, FBaseFramingCameraNodeEvaluator::FDesired& Desired)
+{
+	Desired.Serialize(Ar);
+	return Ar;
+}
+
+}  // namespace UE::Cameras
+
+UBaseFramingCameraNode::UBaseFramingCameraNode(const FObjectInitializer& ObjectInit)
+	: Super(ObjectInit)
+	, DeadZone(0.45f)
+	, SoftZone(0.05f)
+{
+	HorizontalFraming.Value = 0.5f;
+	VerticalFraming.Value = 0.5f;
+	ReframeDampingFactor.Value = 10.f;
+	ReframeUnlockRadius.Value = 0.005f;
+}
+
diff --git a/Engine/Plugins/Cameras/GameplayCameras/Source/GameplayCameras/Private/Nodes/Framing/DollyFramingCameraNode.cpp b/Engine/Plugins/Cameras/GameplayCameras/Source/GameplayCameras/Private/Nodes/Framing/DollyFramingCameraNode.cpp
new file mode 100644
--- /dev/null
+++ b/Engine/Plugins/Cameras/GameplayCameras/Source/GameplayCameras/Private/Nodes/Framing/DollyFramingCameraNode.cpp
@@ -0,0 +1,192 @@
+// Copyright Epic Games, Inc. All Rights Reserved.
+
+#include "Nodes/Framing/DollyFramingCameraNode.h"
+
+#include "Core/CameraEvaluationContext.h"
+#include "Core/CameraParameterReader.h"
+#include "Debug/CameraDebugBlock.h"
+#include "Debug/CameraDebugBlockBuilder.h"
+#include "Debug/CameraDebugRenderer.h"
+#include "GameFramework/Pawn.h"
+#include "GameFramework/PlayerController.h"
+#include "GameplayCameras.h"
+#include "Math/CameraFramingZoneMath.h"
+#include "Math/CameraPoseMath.h"
+#include "Math/InverseRotationMatrix.h"
+
+#include "Engine/LocalPlayer.h"
+#include "Engine/GameViewportClient.h"
+#include "Kismet/GameplayStatics.h"
+#include "SceneView.h"
+
+#include UE_INLINE_GENERATED_CPP_BY_NAME(DollyFramingCameraNode)
+
+namespace UE::Cameras
+{
+
+class FDollyFramingCameraNodeEvaluator : public FBaseFramingCameraNodeEvaluator
+{
+	UE_DECLARE_CAMERA_NODE_EVALUATOR_EX(GAMEPLAYCAMERAS_API, FDollyFramingCameraNodeEvaluator, FBaseFramingCameraNodeEvaluator)
+
+protected:
+
+	virtual void OnInitialize(const FCameraNodeEvaluatorInitializeParams& Params) override;
+	virtual void OnRun(const FCameraNodeEvaluationParams& Params, FCameraNodeEvaluationResult& OutResult) override;
+
+#if UE_GAMEPLAY_CAMERAS_DEBUG
+	virtual void OnBuildDebugBlocks(const FCameraDebugBlockBuildParams& Params, FCameraDebugBlockBuilder& Builder) override;
+#endif  // UE_GAMEPLAY_CAMERAS_DEBUG
+
+private:
+
+	FVector3d ComputeFramingTranslation(const FCameraPose& CameraPose);
+
+private:
+
+	TCameraParameterReader<bool> CanMoveLaterallyReader;
+	TCameraParameterReader<bool> CanMoveVerticallyReader;
+
+	FVector2d DollyPosition;
+
+#if UE_GAMEPLAY_CAMERAS_DEBUG
+	FVector3d DebugUnprojectedNextScreenTarget;
+	FVector2d DebugDollyCorrection;
+#endif  // UE_GAMEPLAY_CAMERAS_DEBUG
+};
+
+UE_DEFINE_CAMERA_NODE_EVALUATOR(FDollyFramingCameraNodeEvaluator)
+
+UE_DECLARE_CAMERA_DEBUG_BLOCK_START(GAMEPLAYCAMERAS_API, FDollyFramingCameraDebugBlock)
+	UE_DECLARE_CAMERA_DEBUG_BLOCK_FIELD(FVector2d, DollyPosition);
+	UE_DECLARE_CAMERA_DEBUG_BLOCK_FIELD(FVector2d, DollyCorrection);
+	UE_DECLARE_CAMERA_DEBUG_BLOCK_FIELD(FVector3d, WorldTarget);
+	UE_DECLARE_CAMERA_DEBUG_BLOCK_FIELD(FVector3d, UnprojectedNextScreenTarget);
+UE_DECLARE_CAMERA_DEBUG_BLOCK_END()
+
+UE_DEFINE_CAMERA_DEBUG_BLOCK_WITH_FIELDS(FDollyFramingCameraDebugBlock)
+
+void FDollyFramingCameraNodeEvaluator::OnInitialize(const FCameraNodeEvaluatorInitializeParams& Params)
+{
+	Super::OnInitialize(Params);
+
+	const UDollyFramingCameraNode* DollyNode = GetCameraNodeAs<UDollyFramingCameraNode>();
+	CanMoveLaterallyReader.Initialize(DollyNode->CanMoveLaterally);
+	CanMoveVerticallyReader.Initialize(DollyNode->CanMoveVertically);
+
+	DollyPosition = FVector2d::ZeroVector;
+}
+
+void FDollyFramingCameraNodeEvaluator::OnRun(const FCameraNodeEvaluationParams& Params, FCameraNodeEvaluationResult& OutResult)
+{
+	// TODO get the target some other way
+	APlayerController* PlayerController = Params.EvaluationContext->GetPlayerController();
+	APawn* Pawn = PlayerController->GetPawn();
+	const FVector3d TargetLocation = Pawn->GetActorLocation();
+
+	// Let the base class figure out all the screen-space framing stuff.
+	UpdateFramingState(OutResult, TargetLocation);
+	ComputeDesiredState(Params.DeltaTime);
+
+	// If we need to reframe the target this tick, figure out how much we need to move the dolly
+	// to accomplish that.
+	if (Desired.bHasCorrection)
+	{
+		const FTransform& LastFraming = GetLastFraming();
+		FCameraPose LastFramingPose(OutResult.CameraPose);
+		LastFramingPose.SetTransform(LastFraming);
+
+		FVector3d DesiredLocalOffset = ComputeFramingTranslation(LastFramingPose);
+
+		// We never bring the dolly forward or backward (we only move it vertically or horizontally).
+		DesiredLocalOffset.X = 0.f;
+
+		if (!CanMoveLaterallyReader.Get(OutResult.VariableTable))
+		{
+			DesiredLocalOffset.Y = 0.f;
+		}
+		if (!CanMoveVerticallyReader.Get(OutResult.VariableTable))
+		{
+			DesiredLocalOffset.Z = 0.f;
+		}
+
+		const FVector2d DollyCorrection = FVector2d(DesiredLocalOffset.Y, DesiredLocalOffset.Z);
+		DollyPosition += DollyCorrection;
+
+#if UE_GAMEPLAY_CAMERAS_DEBUG
+		DebugDollyCorrection = DollyCorrection;
+#endif  //UE_GAMEPLAY_CAMERAS_DEBUG
+	}
+
+	// Build the new shot transform and register it with the base class for next tick.
+	FTransform3d Transform = OutResult.CameraPose.GetTransform();
+	Transform = FTransform3d(FVector3d(0, DollyPosition.X, DollyPosition.Y)) * Transform;
+	OutResult.CameraPose.SetTransform(Transform);
+
+	RegisterNewFraming(Transform);
+}
+
+FVector3d FDollyFramingCameraNodeEvaluator::ComputeFramingTranslation(const FCameraPose& CameraPose)
+{
+	// Unproject the desired screen-space position of our target for this tick. This normally gives us
+	// a ray into the camera frustum, since there is an infinity of world-space points that correspond
+	// to a point on the screen... we arbitrarily pick a point along this ray that is at the same 
+	// distance as our target, since generally it's not far from there.
+	const FVector3d CameraToTarget(State.WorldTarget - CameraPose.GetLocation());
+	const double TargetDist(CameraToTarget.Length());
+
+	const FVector2D DesiredScreenTarget(Desired.ScreenTarget);
+	FVector3d RoughDesiredWorldTarget = FCameraPoseMath::UnprojectScreenToWorld(CameraPose, DesiredScreenTarget, TargetDist);
+
+#if UE_GAMEPLAY_CAMERAS_DEBUG
+	DebugUnprojectedNextScreenTarget = RoughDesiredWorldTarget;
+#endif  // UE_GAMEPLAY_CAMERAS_DEBUG
+
+	// Now we know that we need to translate the dolly by that much in world-space. Convert that to
+	// a camera-space offset, since we want to decompose that into lateral and vertical dolly 
+	// movement.
+	const FVector3d DesiredWorldOffset(State.WorldTarget - RoughDesiredWorldTarget);
+	const FInverseRotationMatrix InverseRotation(CameraPose.GetRotation());
+	const FVector3d DesiredLocalOffset = InverseRotation.TransformVector(DesiredWorldOffset);
+	return DesiredLocalOffset;
+}
+
+#if UE_GAMEPLAY_CAMERAS_DEBUG
+
+void FDollyFramingCameraNodeEvaluator::OnBuildDebugBlocks(const FCameraDebugBlockBuildParams& Params, FCameraDebugBlockBuilder& Builder)
+{
+	Super::OnBuildDebugBlocks(Params, Builder);
+
+	FDollyFramingCameraDebugBlock& DebugBlock = Builder.AttachDebugBlock<FDollyFramingCameraDebugBlock>();
+	DebugBlock.DollyPosition = DollyPosition;
+	DebugBlock.DollyCorrection = DebugDollyCorrection;
+	DebugBlock.WorldTarget = State.WorldTarget;
+	DebugBlock.UnprojectedNextScreenTarget = DebugUnprojectedNextScreenTarget;
+}
+
+void FDollyFramingCameraDebugBlock::OnDebugDraw(const FCameraDebugBlockDrawParams& Params, FCameraDebugRenderer& Renderer)
+{
+	Renderer.AddText(
+			TEXT("dolly position (%0.1f ; %0.1f)  correction (%0.1f ; %0.1f)"), 
+			DollyPosition.X, DollyPosition.Y,
+			DollyCorrection.X, DollyCorrection.Y);
+
+	Renderer.DrawLine(WorldTarget, UnprojectedNextScreenTarget, FLinearColor(FColorList::NavyBlue), 1.f);
+}
+
+#endif  // UE_GAMEPLAY_CAMERAS_DEBUG
+
+}  // namespace UE::Cameras
+
+UDollyFramingCameraNode::UDollyFramingCameraNode(const FObjectInitializer& ObjectInit)
+	: Super(ObjectInit)
+{
+	CanMoveLaterally.Value = true;
+	CanMoveVertically.Value = true;
+}
+
+FCameraNodeEvaluatorPtr UDollyFramingCameraNode::OnBuildEvaluator(FCameraNodeEvaluatorBuilder& Builder) const
+{
+	using namespace UE::Cameras;
+	return Builder.BuildEvaluator<FDollyFramingCameraNodeEvaluator>();
+}
+
diff --git a/Engine/Plugins/Cameras/GameplayCameras/Source/GameplayCameras/Public/Math/CameraFramingZoneMath.h b/Engine/Plugins/Cameras/GameplayCameras/Source/GameplayCameras/Public/Math/CameraFramingZoneMath.h
new file mode 100644
--- /dev/null
+++ b/Engine/Plugins/Cameras/GameplayCameras/Source/GameplayCameras/Public/Math/CameraFramingZoneMath.h
@@ -0,0 +1,103 @@
+// Copyright Epic Games, Inc. All Rights Reserved.
+
+#pragma once
+
+#include "CoreTypes.h"
+#include "Math/Vector.h"
+
+class FArchive;
+struct FCameraFieldsOfView;
+struct FCameraPose;
+
+namespace UE::Cameras
+{
+
+/**
+ * Effective margins for a rectangular screen-space zone.
+ */
+struct FFramingZoneMargins
+{
+	float LeftMargin = 0;
+	float TopMargin = 0;
+	float RightMargin = 0;
+	float BottomMargin = 0;
+};
+
+/**
+ * Effective coordinates for a rectangular screen-space zone.
+ * Unlike FFramingZoneMargins, which stores margin values from the screen's edges,
+ * this struct is expected to store actual screen coordinates in 0..1 UI space.
+ */
+struct FFramingZone
+{
+	float LeftBound = 0;
+	float TopBound = 0;
+	float RightBound = 0;
+	float BottomBound = 0;
+
+	/** Builds an empty framing zone. */
+	FFramingZone();
+
+	/** Build a framing zone from a set of margins. */
+	FFramingZone(const FFramingZoneMargins& FramingZoneMargins);
+
+	/** Makes sure all the bounds have valid values between 0 and 1. */
+	void ClampBounds();
+
+	/**
+	 * Makes sure all the bounds have valid values between 0 and 1, and that the
+	 * enclosed rectangle contains the given target point.
+	 */
+	void ClampBounds(const FVector2d& MustContain);
+
+	/**
+	 * Makes sure all the bounds have valid values between 0 and 1, and that the
+	 * enclosed rectangle contains the given inner rectangle.
+	 */
+	void ClampBounds(const FFramingZone& MustContain);
+
+	/** Checks whether the given point (in 0..1 UI space) is inside this zone. */
+	bool Contains(const FVector2d& Point) const;
+
+	/** Gets the inner margins of this zone compared to the screen's center. */
+	FVector4d GetMarginsFromCenter() const;
+
+	/** Gets the coordinates of the top-left corner of the zone, in 0..Width/Height canvas units. */
+	FVector2d GetCanvasPosition(const FVector2d& CanvasSize) const;
+	/** Gets the size of the zone, in 0..Width/Height canvas units. */
+	FVector2d GetCanvasSize(const FVector2d& CanvasSize) const;
+
+	void Serialize(FArchive& Ar);
+};
+
+FArchive& operator <<(FArchive& Ar, FFramingZone& FramingZone);
+
+/**
+ * The half-angles (in radians) of a rectangular screen framing zone, relative to the
+ * camera pose's aim direction.
+ */
+struct FFramingZoneAngles
+{
+	double LeftHalfAngle = 0;
+	double TopHalfAngle = 0;
+	double RightHalfAngle = 0;
+	double BottomHalfAngle = 0;
+};
+
+/**
+ * Utility class for mathematical functions related to framing zones.
+ */
+class FFramingZoneMath
+{
+public:
+
+	/** Gets the framing zone's half-angles for a given camera FOV. */
+	static FFramingZoneAngles GetFramingZoneAngles(const FFramingZone& FramingZone, const FCameraFieldsOfView& FieldsOfView);
+
+private:
+
+	static double GetFramingMarginAngle(float MarginFromCenter, float BackingHalfSize);
+};
+
+}  // namesdoubleUE::Cameras
+
diff --git a/Engine/Plugins/Cameras/GameplayCameras/Source/GameplayCameras/Public/Math/CameraPoseMath.h b/Engine/Plugins/Cameras/GameplayCameras/Source/GameplayCameras/Public/Math/CameraPoseMath.h
new file mode 100644
--- /dev/null
+++ b/Engine/Plugins/Cameras/GameplayCameras/Source/GameplayCameras/Public/Math/CameraPoseMath.h
@@ -0,0 +1,70 @@
+// Copyright Epic Games, Inc. All Rights Reserved.
+
+#pragma once
+
+#include "CoreTypes.h"
+#include "Math/MathFwd.h"
+#include "Math/Matrix.h"
+
+struct FCameraPose;
+
+namespace UE::Cameras
+{
+
+/**
+ * A utility class for mathematical functions related to a camera pose.
+ */
+class FCameraPoseMath
+{
+public:
+
+	/** 
+	 * Builds the projection matrix of the given camera pose. 
+	 * This matrix is suitable for projecting camera-space points onto screen-space.
+	 */
+	static FMatrix BuildProjectionMatrix(const FCameraPose& CameraPose);
+	/** 
+	 * Builds the view-projection matrix of the given camera pose.
+	 * This matrix combines the camera transform and the projection matrix, making it
+	 * suitable for projecting world-space points onto screen-space.
+	 */
+	static FMatrix BuildViewProjectionMatrix(const FCameraPose& CameraPose);
+
+	/** Projects the given world point onto screen-space. */
+	static TOptional<FVector2d> ProjectWorldToScreen(
+			const FCameraPose& CameraPose, const FVector3d& WorldLocation, bool bForceLocationInsideFrustum = false);
+	/** Projects the given camera-local point onto screen-space. */
+	static TOptional<FVector2d> ProjectCameraToScreen(
+			const FCameraPose& CameraPose, const FVector3d& CameraSpaceLocation, bool bForceLocationInsideFrustum = false);
+
+	/**
+	 * Projects the given point onto screen-space.
+	 * The caller is responsible for making sure the view-projection matrix and coordinate system 
+	 * for the given point are compatible.
+	 */
+	static TOptional<FVector2d> ProjectToScreen(
+			const FMatrix& ViewProjectionMatrix, const FVector3d& Location, bool bForceLocationInsideFrustum = false);
+
+	/** Unproject a screen-space point into a camera-local ray. */
+	static FRay3d UnprojectScreenToCamera(const FCameraPose& CameraPose, const FVector2D& ScreenSpacePoint);
+	/** Unproject a screen-space point into a camera-local point given an expected distance. */
+	static FVector3d UnprojectScreenToCamera(const FCameraPose& CameraPose, const FVector2D& ScreenSpacePoint, double PredictedDistance);
+	/** Unproject a screen-space point into a world-space ray. */
+	static FRay3d UnprojectScreenToWorld(const FCameraPose& CameraPose, const FVector2D& ScreenSpacePoint);
+	/** Unproject a screen-space point into a world-space point given an expected distance. */
+	static FVector3d UnprojectScreenToWorld(const FCameraPose& CameraPose, const FVector2D& ScreenSpacePoint, double PredictedDistance);
+
+	/** Unprojects the given screen-space point into a ray. */
+	static FRay3d UnprojectFromScreen(
+			const FMatrix& InverseViewProjectionMatrix, const FVector2D& ScreenSpacePoint);
+	/** Unprojects the given screen-space point into a point given an expected distance. */
+	static FVector3d UnprojectFromScreen(
+			const FMatrix& InverseViewProjectionMatrix, const FVector2D& ScreenSpacePoint, double PredictedDistance);
+
+private:
+
+	static FMatrix InverseProjectionMatrix(const FMatrix& ProjectionMatrix);
+};
+
+}  // namespace UE::Cameras
+
diff --git a/Engine/Plugins/Cameras/GameplayCameras/Source/GameplayCameras/Public/Nodes/Framing/BaseFramingCameraNode.h b/Engine/Plugins/Cameras/GameplayCameras/Source/GameplayCameras/Public/Nodes/Framing/BaseFramingCameraNode.h
new file mode 100644
--- /dev/null
+++ b/Engine/Plugins/Cameras/GameplayCameras/Source/GameplayCameras/Public/Nodes/Framing/BaseFramingCameraNode.h
@@ -0,0 +1,244 @@
+// Copyright Epic Games, Inc. All Rights Reserved.
+
+#pragma once
+
+#include "Core/CameraNode.h"
+#include "Core/CameraParameterReader.h"
+#include "Core/CameraParameters.h"
+#include "Nodes/Framing/CameraFramingZone.h"
+#include "Math/CameraFramingZoneMath.h"
+#include "Math/CriticalDamper.h"
+
+#include "BaseFramingCameraNode.generated.h"
+
+class FArchive;
+
+/**
+ * The base class for a standard scren-space framing camera node.
+ */
+UCLASS(MinimalAPI, Abstract)
+class UBaseFramingCameraNode : public UCameraNode
+{
+	GENERATED_BODY()
+
+public:
+
+	/** The ideal horizontal screen-space position of the target. */
+	UPROPERTY(EditAnywhere, Category="Framing Target")
+	FFloatCameraParameter HorizontalFraming;
+
+	/** The ideal vertical screen-space position of the target. */
+	UPROPERTY(EditAnywhere, Category="Framing Target")
+	FFloatCameraParameter VerticalFraming;
+
+	/** The damping factor for how fast the framing recenters on the target. */
+	UPROPERTY(EditAnywhere, Category="Framing Target")
+	FFloatCameraParameter ReframeDampingFactor;
+
+	/**
+	 * The distance from the ideal framing position at which we can disengage reframing.
+	 * This should be a very small value, but if it is too small the reframing will keep "chasing"
+	 * the target for a long time even if it stays in the dead zone.
+	 */
+	UPROPERTY(EditAnywhere, Category="Framing Target")
+	FFloatCameraParameter ReframeUnlockRadius;
+
+	/** 
+	 * The margins of the dead zone, i.e. the zone inside which the target can freely move.
+	 * Margins are expressed in screen percentages from the edges.
+	 */
+	UPROPERTY(EditAnywhere, Category="Framing Zones")
+	FCameraFramingZone DeadZone;
+
+	/**
+	 * The margins of the soft zone, i.e. the zone inside which the reframing will engage, in order
+	 * to bring the target back towards the ideal framing position. If the target is outside of the
+	 * soft zone, it will be forcibly and immedialy brought back to its edges, so this zone also 
+	 * defines the "hard" or "safe" zone of framing.
+	 * Margins are expressed in screen percentages from the edges.
+	 */
+	UPROPERTY(EditAnywhere, Category="Framing Zones")
+	FCameraFramingZone SoftZone;
+
+public:
+
+	UBaseFramingCameraNode(const FObjectInitializer& ObjectInit);
+};
+
+namespace UE::Cameras
+{
+
+class FCameraVariableTable;
+
+/**
+ * Utility struct for reading a framing zone's margin parameters.
+ */
+struct FCameraFramingZoneParameterReader
+{
+public:
+
+	TCameraParameterReader<float> LeftMargin;
+	TCameraParameterReader<float> TopMargin;
+	TCameraParameterReader<float> RightMargin;
+	TCameraParameterReader<float> BottomMargin;
+
+public:
+
+	void Initialize(const FCameraFramingZone& FramingZone);
+	FFramingZoneMargins GetZoneMargins(const FCameraVariableTable& VariableTable) const;
+};
+
+/**
+ * The base class for a framing camera node evaluator.
+ *
+ * This evaluator does nothing per se but provides utility functions to be called in 
+ * a sub-class' OnRun method. Namely:
+ *
+ * - UpdateFramingState() : computes the current state of the framing node. The result
+ *			can be obtained from the State field.
+ *
+ * - ComputeDesiredState() : one the current state has been written, this method computes
+ *			the desired framing state for the current tick, including the desired framing
+ *			correction. It is up to the sub-class to implement the necessary logic to
+ *			honor this correction. For instance, a dolly shot would translate left/right 
+ *			(and maybe up/down too) to try and reframe things accordingly, whereas a panning
+ *			shot would rotate the camera left/right/up/down to accomplish the same.
+ *
+ * - RegisterNewFraming() : once the framing correction has been executed by the sub-class,
+ *			it's important to register the new camera transform with RegisterNewFraming, 
+ *			otherwise the reframing will always act only on the incoming camera pose! If this
+ *			incoming camera pose is fixed (e.g. the previous nodes are only fixed offsets)
+ *			then if RegisterNewFraming isn't called, the reframing will always do the same
+ *			thing every frame!
+ */
+class FBaseFramingCameraNodeEvaluator : public FCameraNodeEvaluator
+{
+	UE_DECLARE_CAMERA_NODE_EVALUATOR(GAMEPLAYCAMERAS_API, FBaseFramingCameraNodeEvaluator)
+
+protected:
+
+	virtual void OnInitialize(const FCameraNodeEvaluatorInitializeParams& Params) override;
+
+#if UE_GAMEPLAY_CAMERAS_DEBUG
+	virtual void OnBuildDebugBlocks(const FCameraDebugBlockBuildParams& Params, FCameraDebugBlockBuilder& Builder) override;
+#endif  // UE_GAMEPLAY_CAMERAS_DEBUG
+
+protected:
+
+	/** Updates the framing state for the current tick, see State member field. */
+	void UpdateFramingState(const FCameraNodeEvaluationResult& OutResult, const FVector3d& TargetLocation);
+	/** Computes the desired reframing for the current tick, see Desired member field. */
+	void ComputeDesiredState(float DeltaTime);
+
+	/** Gets the last registered shot transform (probably from last tick). */
+	const FTransform& GetLastFraming() const;
+	/** Registers a new shot transform computed by a sub-class. */
+	void RegisterNewFraming(const FTransform& NewPoseTransform);
+
+private:
+
+	FVector2d GetHardReframeCoords() const;
+
+protected:
+
+	/** The current location of the target. */
+	enum class ETargetFramingState
+	{
+		/**
+		 * The target is in the dead zone, i.e. it can roam freely unless we have an 
+		 * active reframing to finish.
+		 */
+		InDeadZone,
+		/**
+		 * The target is in the soft zone, i.e. we will attempt to gently bring it back 
+		 * to the ideal framing position.
+		 */
+		InSoftZone,
+		/**
+		 * The target is in the hard zone, i.e. it has exited the soft zone and we need
+		 * to bring it back ASAP.
+		 */
+		InHardZone
+	};
+
+	/** Utility structure for all the parameter readers we need every frame. */
+	struct FReaders
+	{
+		TCameraParameterReader<float> HorizontalFraming;
+		TCameraParameterReader<float> VerticalFraming;
+		TCameraParameterReader<float> ReframeDampingFactor;
+		TCameraParameterReader<float> ReframeUnlockRadius;
+
+		FCameraFramingZoneParameterReader DeadZoneMargin;
+		FCameraFramingZoneParameterReader SoftZoneMargin;
+	};
+	FReaders Readers;
+
+	/** Utility struct for storing the current known state. */
+	struct FState
+	{
+		/** World position of the tracked target. */
+		FVector3d WorldTarget;
+
+		/** Screen-space position of the ideal framing position. */
+		FVector2d IdealTarget;
+		/** Current reframing damping factor. */
+		float ReframeDampingFactor;
+		/** Current reframe unlock radius. */
+		float ReframeUnlockRadius;
+		/** Current coordinates of the dead zone. */
+		FFramingZone DeadZone;
+		/** Current coordinates of the soft zone. */
+		FFramingZone SoftZone;
+
+		/** Current screen-space position of the tracked target. */
+		FVector2d ScreenTarget;
+		/** Current state of the tracked target. */
+		ETargetFramingState TargetFramingState;
+		/** Last tick's shot transform, as per RegisterNewFraming. */
+		FTransform3d LastFramedTransform;
+		/** Whether we are actively trying to bring the target back to the ideal position. */
+		bool bIsReframingTarget = false;
+		/** Whether this is the first tick. */
+		bool bIsFirstFrame = false;
+
+		/** The damper for reframing from the soft zone. */
+		FCriticalDamper ReframeDamper;
+
+		void Serialize(FArchive& Ar);
+	};
+	FState State;
+
+	/** Utility struct for the desired reframing to be done in the current tick. */
+	struct FDesired
+	{
+		/**
+		 * The desired screen-space position of the tracked target. For instance, if the target
+		 * is in the soft zone, this desired position will be the next step to get us closer to
+		 * the ideal position.
+		 */
+		FVector2d ScreenTarget;
+		/** 
+		 * The screen-space correction we want this tick.
+		 * This is effectively equal to: Desired.ScreenTarget - State.ScreenTarget
+		 */
+		FVector2d FramingCorrection;
+		/**
+		 * Whether we have any correction to do.
+		 */
+		bool bHasCorrection = false;
+
+		void Serialize(FArchive& Ar);
+	};
+	FDesired Desired;
+
+	friend class FBaseFramingCameraDebugBlock;
+	friend FArchive& operator <<(FArchive& Ar, FState& State);
+	friend FArchive& operator <<(FArchive& Ar, FDesired& Desired);
+};
+
+FArchive& operator <<(FArchive& Ar, FBaseFramingCameraNodeEvaluator::FState& State);
+FArchive& operator <<(FArchive& Ar, FBaseFramingCameraNodeEvaluator::FDesired& Desired);
+
+}  // namespace UE::Cameras
+
diff --git a/Engine/Plugins/Cameras/GameplayCameras/Source/GameplayCameras/Public/Nodes/Framing/CameraFramingZone.h b/Engine/Plugins/Cameras/GameplayCameras/Source/GameplayCameras/Public/Nodes/Framing/CameraFramingZone.h
new file mode 100644
--- /dev/null
+++ b/Engine/Plugins/Cameras/GameplayCameras/Source/GameplayCameras/Public/Nodes/Framing/CameraFramingZone.h
@@ -0,0 +1,69 @@
+// Copyright Epic Games, Inc. All Rights Reserved.
+
+#pragma once
+
+#include "CoreTypes.h"
+#include "Core/CameraParameters.h"
+
+#include "CameraFramingZone.generated.h"
+
+/**
+ * A structure that defines a zone for use in framing subjects in screen-space.
+ *
+ * All margins are defined in percentages of the screen's horizontal size. They are also 
+ * all defined relative to their respective edges.
+ */
+USTRUCT()
+struct FCameraFramingZone
+{
+	GENERATED_BODY()
+
+public:
+
+	UPROPERTY(EditAnywhere, Category="Framing")
+	FFloatCameraParameter LeftMargin;
+
+	UPROPERTY(EditAnywhere, Category="Framing")
+	FFloatCameraParameter TopMargin;
+
+	UPROPERTY(EditAnywhere, Category="Framing")
+	FFloatCameraParameter RightMargin;
+
+	UPROPERTY(EditAnywhere, Category="Framing")
+	FFloatCameraParameter BottomMargin;
+
+public:
+
+	FCameraFramingZone()
+	{
+		LeftMargin.Value = 0;
+		TopMargin.Value = 0;
+		RightMargin.Value = 0;
+		BottomMargin.Value = 0;
+	}
+
+	FCameraFramingZone(float UniformMargin)
+	{
+		LeftMargin.Value = UniformMargin;
+		TopMargin.Value = UniformMargin;
+		RightMargin.Value = UniformMargin;
+		BottomMargin.Value = UniformMargin;
+	}
+
+	FCameraFramingZone(float HorizontalMargin, float VerticalMargin)
+	{
+		LeftMargin.Value = HorizontalMargin;
+		TopMargin.Value = VerticalMargin;
+		RightMargin.Value = HorizontalMargin;
+		BottomMargin.Value = VerticalMargin;
+	}
+
+	FCameraFramingZone(float InLeftMargin, float InTopMargin, float InRightMargin, float InBottomMargin)
+	{
+		LeftMargin.Value = InLeftMargin;
+		TopMargin.Value = InTopMargin;
+		RightMargin.Value = InRightMargin;
+		BottomMargin.Value = InBottomMargin;
+	}
+};
+
diff --git a/Engine/Plugins/Cameras/GameplayCameras/Source/GameplayCameras/Public/Nodes/Framing/DollyFramingCameraNode.h b/Engine/Plugins/Cameras/GameplayCameras/Source/GameplayCameras/Public/Nodes/Framing/DollyFramingCameraNode.h
new file mode 100644
--- /dev/null
+++ b/Engine/Plugins/Cameras/GameplayCameras/Source/GameplayCameras/Public/Nodes/Framing/DollyFramingCameraNode.h
@@ -0,0 +1,37 @@
+// Copyright Epic Games, Inc. All Rights Reserved.
+
+#pragma once
+
+#include "Nodes/Framing/BaseFramingCameraNode.h"
+
+#include "DollyFramingCameraNode.generated.h"
+
+/**
+ * A camera node that frames a target by moving along a sideways rail (left/right relative to 
+ * the camera transform) and optionally also up and down.
+ */
+UCLASS(MinimalAPI)
+class UDollyFramingCameraNode : public UBaseFramingCameraNode
+{
+	GENERATED_BODY()
+
+public:
+
+	/** Whether the dolly can move laterally. */
+	UPROPERTY(EditAnywhere, Category="Dolly")
+	FBooleanCameraParameter CanMoveLaterally;
+
+	/** Whether the dolly can move vertically. */
+	UPROPERTY(EditAnywhere, Category="Dolly")
+	FBooleanCameraParameter CanMoveVertically;
+
+public:
+
+	UDollyFramingCameraNode(const FObjectInitializer& ObjectInit);
+
+protected:
+
+	// UCameraNode interface.
+	virtual FCameraNodeEvaluatorPtr OnBuildEvaluator(FCameraNodeEvaluatorBuilder& Builder) const override;
+};
+
