Cameras: fix memory management problems in the variable table

We were never actually allocating memory because Capacity was initialized to zero, and twice zero is still zero. Plus, even if we did allocate, we also weren't freeing the old buffer when moving to a bigger one. Woops.

diff --git a/Engine/Plugins/Cameras/GameplayCameras/Source/GameplayCameras/Private/Core/CameraVariableTable.cpp b/Engine/Plugins/Cameras/GameplayCameras/Source/GameplayCameras/Private/Core/CameraVariableTable.cpp
--- a/Engine/Plugins/Cameras/GameplayCameras/Source/GameplayCameras/Private/Core/CameraVariableTable.cpp
+++ b/Engine/Plugins/Cameras/GameplayCameras/Source/GameplayCameras/Private/Core/CameraVariableTable.cpp
@@ -130,7 +130,7 @@ void FCameraVariableTable::AddVariable(const FCameraVariableDefinition& Variable
 
 	if (NewUsed > Capacity)
 	{
-		ReallocateBuffer(Capacity * 2);
+		ReallocateBuffer(NewUsed);
 
 		VariablePtr = Align(Memory + Used, AlignOf);
 	}
@@ -149,12 +149,26 @@ void FCameraVariableTable::AddVariable(const FCameraVariableDefinition& Variable
 	Entries.Add(VariableDefinition.VariableID, NewEntry);
 }
 
-void FCameraVariableTable::ReallocateBuffer(uint32 NewCapacity)
+void FCameraVariableTable::ReallocateBuffer(uint32 MinRequired)
 {
+	static const uint32 DefaultCapacity = 64;
 	static const uint32 DefaultAlignment = 32;
 
+	uint32 NewCapacity = Capacity <= 0 ? DefaultCapacity : Capacity * 2;
+	if (MinRequired > 0)
+	{
+		NewCapacity = FMath::Max(NewCapacity, MinRequired);
+	}
+
+	uint8* OldMemory = Memory;
 	uint8* NewMemory = reinterpret_cast<uint8*>(FMemory::Malloc(NewCapacity, DefaultAlignment));
-	FMemory::Memmove(NewMemory, Memory, Capacity);
+
+	if (OldMemory)
+	{
+		FMemory::Memmove(NewMemory, OldMemory, Capacity);
+		FMemory::Free(OldMemory);
+	}
+
 	Memory = NewMemory;
 	Capacity = NewCapacity;
 }
diff --git a/Engine/Plugins/Cameras/GameplayCameras/Source/GameplayCameras/Public/Core/CameraVariableTable.h b/Engine/Plugins/Cameras/GameplayCameras/Source/GameplayCameras/Public/Core/CameraVariableTable.h
--- a/Engine/Plugins/Cameras/GameplayCameras/Source/GameplayCameras/Public/Core/CameraVariableTable.h
+++ b/Engine/Plugins/Cameras/GameplayCameras/Source/GameplayCameras/Public/Core/CameraVariableTable.h
@@ -135,7 +135,7 @@ private:
 		return ensure(TCameraVariableTraits<ValueType>::Type == InType);
 	}
 
-	void ReallocateBuffer(uint32 NewCapacity);
+	void ReallocateBuffer(uint32 MinRequired = 0);
 
 	void InternalOverride(const FCameraVariableTable& OtherTable, const FCameraVariableTableFlags* InMask, bool bInvertMask, FCameraVariableTableFlags* OutMask, bool bChangedOnly);
 	void InternalLerp(const FCameraVariableTable& ToTable, float Factor, const FCameraVariableTableFlags* InMask, bool bInvertMask, FCameraVariableTableFlags* OutMask, bool bChangedOnly);
