Cameras: remember original buffer alignment when initializing an object storage from an allocation info.

diff --git a/Engine/Plugins/Cameras/GameplayCameras/Source/GameplayCameras/Public/Core/CameraObjectStorage.h b/Engine/Plugins/Cameras/GameplayCameras/Source/GameplayCameras/Public/Core/CameraObjectStorage.h
--- a/Engine/Plugins/Cameras/GameplayCameras/Source/GameplayCameras/Public/Core/CameraObjectStorage.h
+++ b/Engine/Plugins/Cameras/GameplayCameras/Source/GameplayCameras/Public/Core/CameraObjectStorage.h
@@ -259,6 +259,7 @@ void TCameraObjectStorage<BaseObjectType>::AllocatePage(uint32 InCapacity, uint3
 {
 	FAllocation& NewAllocation = Allocations.Emplace_GetRef();
 	NewAllocation.Memory = reinterpret_cast<uint8*>(FMemory::Malloc(InCapacity, InAlignment));
+	NewAllocation.Alignment = InAlignment;
 	NewAllocation.Capacity = InCapacity;
 	NewAllocation.Used = 0;
 }
