Cameras: give a sane FOV default to the FOV node, and invalidate focal length/FOV as appropriate

diff --git a/Engine/Plugins/Cameras/GameplayCameras/Source/GameplayCameras/Private/Nodes/Common/FieldOfViewCameraNode.cpp b/Engine/Plugins/Cameras/GameplayCameras/Source/GameplayCameras/Private/Nodes/Common/FieldOfViewCameraNode.cpp
--- a/Engine/Plugins/Cameras/GameplayCameras/Source/GameplayCameras/Private/Nodes/Common/FieldOfViewCameraNode.cpp
+++ b/Engine/Plugins/Cameras/GameplayCameras/Source/GameplayCameras/Private/Nodes/Common/FieldOfViewCameraNode.cpp
@@ -37,10 +37,17 @@ void FFieldOfViewCameraNodeEvaluator::OnRun(const FCameraNodeEvaluationParams& P
 {
 	const float FieldOfView = FieldOfViewReader.Get(OutResult.VariableTable);
 	OutResult.CameraPose.SetFieldOfView(FieldOfView);
+	OutResult.CameraPose.SetFocalLength(-1);
 }
 
 }  // namespace UE::Cameras
 
+UFieldOfViewCameraNode::UFieldOfViewCameraNode(const FObjectInitializer& ObjectInit)
+	: Super(ObjectInit)
+{
+	FieldOfView.Value = 90.f;
+}
+
 FCameraNodeEvaluatorPtr UFieldOfViewCameraNode::OnBuildEvaluator(FCameraNodeEvaluatorBuilder& Builder) const
 {
 	using namespace UE::Cameras;
diff --git a/Engine/Plugins/Cameras/GameplayCameras/Source/GameplayCameras/Private/Nodes/Common/LensParametersCameraNode.cpp b/Engine/Plugins/Cameras/GameplayCameras/Source/GameplayCameras/Private/Nodes/Common/LensParametersCameraNode.cpp
--- a/Engine/Plugins/Cameras/GameplayCameras/Source/GameplayCameras/Private/Nodes/Common/LensParametersCameraNode.cpp
+++ b/Engine/Plugins/Cameras/GameplayCameras/Source/GameplayCameras/Private/Nodes/Common/LensParametersCameraNode.cpp
@@ -64,6 +64,7 @@ void FLensParametersCameraNodeEvaluator::OnRun(const FCameraNodeEvaluationParams
 	if (FocalLength > 0)
 	{
 		OutPose.SetFocalLength(FocalLength);
+		OutPose.SetFieldOfView(-1);
 	}
 	float Aperture = ApertureReader.Get(OutResult.VariableTable);
 	if (Aperture > 0)
diff --git a/Engine/Plugins/Cameras/GameplayCameras/Source/GameplayCameras/Public/Nodes/Common/FieldOfViewCameraNode.h b/Engine/Plugins/Cameras/GameplayCameras/Source/GameplayCameras/Public/Nodes/Common/FieldOfViewCameraNode.h
--- a/Engine/Plugins/Cameras/GameplayCameras/Source/GameplayCameras/Public/Nodes/Common/FieldOfViewCameraNode.h
+++ b/Engine/Plugins/Cameras/GameplayCameras/Source/GameplayCameras/Public/Nodes/Common/FieldOfViewCameraNode.h
@@ -16,15 +16,19 @@ class UFieldOfViewCameraNode : public UCameraNode
 {
 	GENERATED_BODY()
 
-protected:
-
-	// UCameraNode interface.
-	virtual FCameraNodeEvaluatorPtr OnBuildEvaluator(FCameraNodeEvaluatorBuilder& Builder) const override;
-
 public:
 
 	/** The field of view, in degrees. */
 	UPROPERTY(EditAnywhere, Category=Common)
 	FFloatCameraParameter FieldOfView;
+
+public:
+
+	UFieldOfViewCameraNode(const FObjectInitializer& ObjectInit);
+
+protected:
+
+	// UCameraNode interface.
+	virtual FCameraNodeEvaluatorPtr OnBuildEvaluator(FCameraNodeEvaluatorBuilder& Builder) const override;
 };
 
