Cameras: properly convert rotation offsets to radians before using them as quaternions

diff --git a/Engine/Plugins/Cameras/GameplayCameras/Source/GameplayCameras/Private/Nodes/Common/OffsetCameraNode.cpp b/Engine/Plugins/Cameras/GameplayCameras/Source/GameplayCameras/Private/Nodes/Common/OffsetCameraNode.cpp
--- a/Engine/Plugins/Cameras/GameplayCameras/Source/GameplayCameras/Private/Nodes/Common/OffsetCameraNode.cpp
+++ b/Engine/Plugins/Cameras/GameplayCameras/Source/GameplayCameras/Private/Nodes/Common/OffsetCameraNode.cpp
@@ -69,9 +69,9 @@ void FOffsetCameraNodeEvaluator::OnRun(const FCameraNodeEvaluationParams& Params
 				const FVector3d ContextRight = ContextTransform.GetUnitAxis(EAxis::Y);
 				const FVector3d ContextUp = ContextTransform.GetUnitAxis(EAxis::Z);
 				const FQuat WorldRotationOffset = 
-					FQuat(ContextUp, RotationOffset.Yaw) * 
-					FQuat(ContextRight, RotationOffset.Pitch) *
-					FQuat(ContextForward, RotationOffset.Roll);
+					FQuat(ContextUp, FMath::DegreesToRadians(RotationOffset.Yaw)) * 
+					FQuat(ContextRight, -FMath::DegreesToRadians(RotationOffset.Pitch)) *
+					FQuat(ContextForward, -FMath::DegreesToRadians(RotationOffset.Roll));
 
 				FTransform3d Transform = OutResult.CameraPose.GetTransform();
 				Transform.SetTranslation(Transform.GetTranslation() + WorldTranslationOffset);
