Cameras: when building a camera rig, check for valid private variables before making new ones

diff --git a/Engine/Plugins/Cameras/GameplayCameras/Source/GameplayCameras/Private/Core/CameraRigAsset.cpp b/Engine/Plugins/Cameras/GameplayCameras/Source/GameplayCameras/Private/Core/CameraRigAsset.cpp
--- a/Engine/Plugins/Cameras/GameplayCameras/Source/GameplayCameras/Private/Core/CameraRigAsset.cpp
+++ b/Engine/Plugins/Cameras/GameplayCameras/Source/GameplayCameras/Private/Core/CameraRigAsset.cpp
@@ -29,20 +29,35 @@ void SetPrivateVariableDefaultValue<UBooleanCameraVariable, bool>(UBooleanCamera
 template<typename CameraParameterType>
 UCameraVariableAsset* CreatePrivateVariable(UCameraRigAsset* CameraRigAsset, UCameraRigInterfaceParameter* InterfaceParameter, const FString& VariableName, CameraParameterType* CameraParameter)
 {
+	using ValueType = typename CameraParameterType::ValueType;
+	using VariableAssetType = typename CameraParameterType::VariableAssetType;
+
+	VariableAssetType* PrivateVariable = nullptr;
+
 	if (CameraParameter->Variable != nullptr)
 	{
-		UE_LOG(LogCameraSystem, Error, 
-				TEXT("Invalid interface parameter in camera rig '%s': '%s.%s' is already driven by variable."), 
-				*CameraRigAsset->GetPathName(), 
-				*InterfaceParameter->Target->GetName(), 
-				*InterfaceParameter->TargetPropertyName.ToString());
-		return nullptr;
+		UObject* VariableOuter = CameraParameter->Variable->GetOuter();
+		if (VariableOuter != CameraRigAsset)
+		{
+			// If this parameter is driven by a user-defined variable, emit a warning, since we will 
+			// replace that driving variable with our private variable.
+			UE_LOG(LogCameraSystem, Warning, 
+					TEXT("Camera node parameter '%s.%s' in camera rig '%s' is both exposed and driven by a variable! That variable will be replaced by the rig parameter variable."), 
+					*InterfaceParameter->Target->GetName(), 
+					*InterfaceParameter->TargetPropertyName.ToString(),
+					*CameraRigAsset->GetPathName());
+		}
+		else
+		{
+			PrivateVariable = CameraParameter->Variable;
+		}
 	}
 
-	using ValueType = typename CameraParameterType::ValueType;
-	using VariableAssetType = typename CameraParameterType::VariableAssetType;
+	if (PrivateVariable == nullptr)
+	{
+		PrivateVariable = NewObject<VariableAssetType>(CameraRigAsset, *VariableName);
+	}
 
-	VariableAssetType* PrivateVariable = NewObject<VariableAssetType>(CameraRigAsset, *VariableName);
 	PrivateVariable->bIsPrivate = true;
 	PrivateVariable->bAutoReset = false;
 	SetPrivateVariableDefaultValue<VariableAssetType, ValueType>(PrivateVariable, CameraParameter->Value);
@@ -90,6 +105,8 @@ void UCameraRigAsset::BuildCameraRig()
 	using namespace UE::Cameras;
 	using namespace UE::Cameras::Internal;
 
+	Modify();
+
 	// Build allocation info.
 	FCameraRigAllocationInfoBuilder CameraRigBuilder;
 	CameraRigBuilder.BuildAllocationInfo(this, AllocationInfo);
@@ -165,6 +182,7 @@ void UCameraRigAsset::BuildCameraRig()
 			F##ValueName##CameraParameter CameraParameter;\
 			TargetStructProperty->GetValue_InContainer(InterfaceParameter->Target, reinterpret_cast<void*>(&CameraParameter));\
 			CreatePrivateVariable(this, InterfaceParameter, PrivateVariableName, &CameraParameter);\
+			InterfaceParameter->Target->Modify();\
 			TargetStructProperty->SetValue_InContainer(InterfaceParameter->Target, reinterpret_cast<void*>(&CameraParameter));\
 		}\
 		else
