Cameras: better handle FOV vs focal length issues

diff --git a/Engine/Plugins/Cameras/GameplayCameras/Source/GameplayCameras/Private/Core/CameraPose.cpp b/Engine/Plugins/Cameras/GameplayCameras/Source/GameplayCameras/Private/Core/CameraPose.cpp
--- a/Engine/Plugins/Cameras/GameplayCameras/Source/GameplayCameras/Private/Core/CameraPose.cpp
+++ b/Engine/Plugins/Cameras/GameplayCameras/Source/GameplayCameras/Private/Core/CameraPose.cpp
@@ -3,7 +3,7 @@
 #include "Core/CameraPose.h"
 
 #include "Math/Ray.h"
-#include "Runtime/Engine/Classes/PhysicsField/PhysicsFieldComponent.h"
+#include "GameplayCameras.h"
 
 #include UE_INLINE_GENERATED_CPP_BY_NAME(CameraPose)
 
@@ -114,9 +114,17 @@ void FCameraPose::SetTransform(FTransform3d Transform)
 
 float FCameraPose::GetEffectiveFieldOfView() const
 {
-	checkf(
-			(FocalLength > 0.f && FieldOfView <= 0.f) || (FocalLength <= 0.f && FieldOfView > 0.f),
-			TEXT("FocalLength or FieldOfView must have a valid, positive value."));
+	checkf((FocalLength > 0.f || FieldOfView > 0.f), TEXT("FocalLength or FieldOfView must have a valid, positive value."));
+
+#if !NO_LOGGING
+	static bool GEmitFocalLengthPrioritizationWarning = true;
+	if ((FocalLength > 0.f && FieldOfView > 0.f) && GEmitFocalLengthPrioritizationWarning)
+	{
+		UE_LOG(LogCameraSystem, Warning,
+				TEXT("Both FocalLength and FieldOfView are specified on a camera pose! Using FocalLength first."));
+		GEmitFocalLengthPrioritizationWarning = false;
+	}
+#endif  // NO_LOGGING	
 
 	if (FocalLength > 0.f)
 	{
