Cameras: fix screen-space math for dolly and panning camera nodes

- Use the effective aspect ratio of the screen (not the sensor's aspect ratio) when possible.
- Add aspect ratio axis constraint handling to the camera pose.
- Change the math of the dolly camera node for simpler and more precise translation correction.

diff --git a/Engine/Plugins/Cameras/GameplayCameras/Source/GameplayCameras/Private/Core/CameraPose.cpp b/Engine/Plugins/Cameras/GameplayCameras/Source/GameplayCameras/Private/Core/CameraPose.cpp
--- a/Engine/Plugins/Cameras/GameplayCameras/Source/GameplayCameras/Private/Core/CameraPose.cpp
+++ b/Engine/Plugins/Cameras/GameplayCameras/Source/GameplayCameras/Private/Core/CameraPose.cpp
@@ -2,8 +2,9 @@
 
 #include "Core/CameraPose.h"
 
-#include "Math/Ray.h"
+#include "Engine/EngineTypes.h"
 #include "GameplayCameras.h"
+#include "Math/Ray.h"
 
 #include UE_INLINE_GENERATED_CPP_BY_NAME(CameraPose)
 
@@ -149,20 +150,6 @@ float FCameraPose::GetEffectiveFieldOfView() const
 	}
 }
 
-FCameraFieldsOfView FCameraPose::GetEffectiveFieldsOfView() const
-{
-	FCameraFieldsOfView FOVs;
-
-	FOVs.HorizontalFieldOfView = GetEffectiveFieldOfView();
-
-	const float AspectRatio = GetSensorAspectRatio();
-	FOVs.VerticalFieldOfView = FMath::RadiansToDegrees(
-			2.f * FMath::Atan(
-				FMath::Tan(FMath::DegreesToRadians(FOVs.HorizontalFieldOfView / 2.f)) / AspectRatio));
-
-	return FOVs;
-}
-
 float FCameraPose::GetSensorAspectRatio() const
 {
 	return (SensorHeight > 0.f) ? (SensorWidth / SensorHeight) : 0.f;
diff --git a/Engine/Plugins/Cameras/GameplayCameras/Source/GameplayCameras/Private/Math/CameraFramingMath.cpp b/Engine/Plugins/Cameras/GameplayCameras/Source/GameplayCameras/Private/Math/CameraFramingMath.cpp
new file mode 100644
--- /dev/null
+++ b/Engine/Plugins/Cameras/GameplayCameras/Source/GameplayCameras/Private/Math/CameraFramingMath.cpp
@@ -0,0 +1,84 @@
+// Copyright Epic Games, Inc. All Rights Reserved.
+
+#include "Math/CameraFramingMath.h"
+
+#include "Math/CameraFramingZoneMath.h"
+#include "Math/CameraPoseMath.h"
+#include "Math/UnrealMathUtility.h"
+#include "Math/Vector2D.h"
+#include "Math/Vector4.h"
+
+namespace UE::Cameras
+{
+
+FVector2d FCameraFramingMath::GetTargetAngles(const FVector2d& Target, const FCameraFieldsOfView& FieldsOfView)
+{
+	const float TanHalfHorizontalFOV = FMath::Tan(FMath::DegreesToRadians(FieldsOfView.HorizontalFieldOfView / 2.f));
+	const float TanHalfVerticalFOV = FMath::Tan(FMath::DegreesToRadians(FieldsOfView.VerticalFieldOfView / 2.f));
+
+	// Target is in 0..1 UI space... convert to -1..1 space.
+	const double NormalizedTargetX = (Target.X - 0.5) * 2.0;
+	const double NormalizedTargetY = (Target.Y - 0.5) * 2.0;
+
+	const double BoundAngleYaw = GetBoundAngle(NormalizedTargetX, TanHalfHorizontalFOV);
+	const double BoundAnglePitch = GetBoundAngle(NormalizedTargetY, TanHalfVerticalFOV);
+
+	return FVector2d(
+			FMath::RadiansToDegrees(BoundAngleYaw),
+			FMath::RadiansToDegrees(BoundAnglePitch));
+}
+
+FFramingZoneAngles FCameraFramingMath::GetFramingZoneAngles(const FFramingZone& FramingZone, const FCameraFieldsOfView& FieldsOfView)
+{
+	const float TanHalfHorizontalFOV = FMath::Tan(FMath::DegreesToRadians(FieldsOfView.HorizontalFieldOfView / 2.f));
+	const float TanHalfVerticalFOV = FMath::Tan(FMath::DegreesToRadians(FieldsOfView.VerticalFieldOfView / 2.f));
+
+	const FVector4f BoundFactorsFromCenter = FramingZone.GetNormalizedBounds();
+
+	const double LeftHalfAngleRad = GetBoundAngle(BoundFactorsFromCenter.X, TanHalfHorizontalFOV);
+	const double TopHalfAngleRad = GetBoundAngle(BoundFactorsFromCenter.Y, TanHalfVerticalFOV);
+	const double RightHalfAngleRad = GetBoundAngle(BoundFactorsFromCenter.Z, TanHalfHorizontalFOV);
+	const double BottomHalfAngleRad = GetBoundAngle(BoundFactorsFromCenter.W, TanHalfVerticalFOV);
+
+	FFramingZoneAngles Angles;
+	Angles.LeftHalfAngle = FMath::DegreesToRadians(LeftHalfAngleRad); 
+	Angles.TopHalfAngle = FMath::DegreesToRadians(TopHalfAngleRad); 
+	Angles.RightHalfAngle = FMath::DegreesToRadians(RightHalfAngleRad); 
+	Angles.BottomHalfAngle = FMath::DegreesToRadians(BottomHalfAngleRad); 
+	return Angles;
+}
+
+double FCameraFramingMath::GetBoundAngle(float FactorFromCenter, double TanHalfFOV)
+{
+	// The FactorFromCenter should always be a percentage from the center of the screen.
+	// So a factor of zero is the center, a factor of 0.5 is half way between the center and the edge,
+	// and a factor of 1 is at the edge.
+	//
+	// Now here is what we do: let's consider a plane, orthogonal to the camera's aim direction, and
+	// located at a distance L from the camera:
+	//
+	//		tan(fov/2) = W/L
+	//
+	// Where W is the half-width of that plane, and fov is the horizontal field of view angle.
+	//
+	// We can do the same for the bound whose angle we're trying to determine. Since this factor (let's
+	// call it m) is a percentage of the plane's half-width from the center, we also have:
+	// 
+	//		tan(x) = W*m/L
+	//
+	// Where x is what we want (the angle for that bound).
+	//
+	// So:
+	//	
+	//		W = tan(fov/2) * L
+	//		tan(x) = tan(fov/2) * L * m/L
+	//		tan(x) = tan(fov/2) * m
+	//		x = atan(tan(fov/2) * m)
+	//
+	const float LocalFactor = FMath::Abs(FactorFromCenter);
+	const double BoundAngle = FMath::Atan(TanHalfFOV * LocalFactor);
+	return FactorFromCenter <= 0.f ? -BoundAngle : BoundAngle;
+}
+
+}  // namespace UE::Cameras
+
diff --git a/Engine/Plugins/Cameras/GameplayCameras/Source/GameplayCameras/Private/Math/CameraFramingZoneMath.cpp b/Engine/Plugins/Cameras/GameplayCameras/Source/GameplayCameras/Private/Math/CameraFramingZoneMath.cpp
--- a/Engine/Plugins/Cameras/GameplayCameras/Source/GameplayCameras/Private/Math/CameraFramingZoneMath.cpp
+++ b/Engine/Plugins/Cameras/GameplayCameras/Source/GameplayCameras/Private/Math/CameraFramingZoneMath.cpp
@@ -2,7 +2,6 @@
 
 #include "Math/CameraFramingZoneMath.h"
 
-#include "Core/CameraPose.h"
 #include "Nodes/Framing/CameraFramingZone.h"
 
 namespace UE::Cameras
@@ -59,13 +58,20 @@ bool FFramingZone::Contains(const FVector2d& Point) const
 		Point.Y >= TopBound && Point.Y <= BottomBound;
 }
 
-FVector4d FFramingZone::GetMarginsFromCenter() const
+FVector4f FFramingZone::GetNormalizedBounds() const
 {
-	return FVector4d(
-		0.5 - LeftBound,
-		0.5 - TopBound,
-		0.5 - (1.0 - RightBound),
-		0.5 - (1.0 - BottomBound));
+	// Returned margins are negative if in the left or upper halves, and
+	// positive if in the right or lower halves.
+	return FVector4f(
+		GetNormalizedBound(LeftBound),
+		GetNormalizedBound(TopBound),
+		GetNormalizedBound(RightBound),
+		GetNormalizedBound(BottomBound));
+}
+
+float FFramingZone::GetNormalizedBound(float Bound)
+{
+	return (Bound - 0.5f) * 2.f;
 }
 
 FVector2d FFramingZone::GetCanvasPosition(const FVector2d& CanvasSize) const
@@ -92,30 +98,5 @@ FArchive& operator <<(FArchive& Ar, FFramingZone& FramingZone)
 	return Ar;
 }
 
-FFramingZoneAngles FFramingZoneMath::GetFramingZoneAngles(const FFramingZone& FramingZone, const FCameraFieldsOfView& FieldsOfView)
-{
-	const float BackingHalfWidth = FMath::Tan(FMath::DegreesToRadians(FieldsOfView.HorizontalFieldOfView / 2.f));
-	const float BackingHalfHeight = FMath::Tan(FMath::DegreesToRadians(FieldsOfView.VerticalFieldOfView / 2.f));
-
-	const FVector4d MarginsFromCenter = FramingZone.GetMarginsFromCenter();
-
-	FFramingZoneAngles Angles;
-	Angles.LeftHalfAngle = GetFramingMarginAngle(MarginsFromCenter.X, BackingHalfWidth);
-	Angles.TopHalfAngle = GetFramingMarginAngle(MarginsFromCenter.Y, BackingHalfHeight);
-	Angles.RightHalfAngle = GetFramingMarginAngle(MarginsFromCenter.Z, BackingHalfWidth);
-	Angles.BottomHalfAngle = GetFramingMarginAngle(MarginsFromCenter.W, BackingHalfHeight);
-	return Angles;
-}
-
-double FFramingZoneMath::GetFramingMarginAngle(float MarginFromCenter, float BackingHalfSize)
-{
-	// The margin should always be a percentage from the center of the screen.
-	// So a margin of zero is the center, a margin of 0.5 is half way between the center and the edge,
-	// and a margin of 1 is at the edge.
-	const double MarginSize = FMath::Clamp(MarginFromCenter, -1.f, 1.f) * BackingHalfSize;
-	const double MarginAngle = FMath::Atan(MarginSize);
-	return MarginAngle;
-}
-
 }  // namespace UE::Cameras
 
diff --git a/Engine/Plugins/Cameras/GameplayCameras/Source/GameplayCameras/Private/Math/CameraPoseMath.cpp b/Engine/Plugins/Cameras/GameplayCameras/Source/GameplayCameras/Private/Math/CameraPoseMath.cpp
--- a/Engine/Plugins/Cameras/GameplayCameras/Source/GameplayCameras/Private/Math/CameraPoseMath.cpp
+++ b/Engine/Plugins/Cameras/GameplayCameras/Source/GameplayCameras/Private/Math/CameraPoseMath.cpp
@@ -4,6 +4,7 @@
 
 #include "Core/CameraPose.h"
 #include "CoreGlobals.h"
+#include "GameFramework/PlayerController.h"
 #include "Math/InverseRotationMatrix.h"
 #include "Math/PerspectiveMatrix.h"
 #include "Math/TranslationMatrix.h"
@@ -12,12 +13,89 @@
 namespace UE::Cameras
 {
 
-FMatrix FCameraPoseMath::BuildProjectionMatrix(const FCameraPose& CameraPose)
+FCameraFieldsOfView FCameraPoseMath::GetEffectiveFieldsOfView(const FCameraPose& CameraPose)
+{
+	return GetEffectiveFieldsOfView(CameraPose, CameraPose.GetSensorAspectRatio());
+}
+
+FCameraFieldsOfView FCameraPoseMath::GetEffectiveFieldsOfView(const FCameraPose& CameraPose, const APlayerController* PlayerController)
+{
+	const float AspectRatio = GetEffectiveAspectRatio(CameraPose, PlayerController);
+	return GetEffectiveFieldsOfView(CameraPose, AspectRatio);
+}
+
+FCameraFieldsOfView FCameraPoseMath::GetEffectiveFieldsOfView(const FCameraPose& CameraPose, float AspectRatio)
+{
+	if (CameraPose.GetConstrainAspectRatio() || 
+			!ensureMsgf(AspectRatio > 0.f, TEXT("Invalid aspect ratio! Using sensor aspect ratio instead.")))
+	{
+		// Aspect ratio is constrained, there will be black bars to enforce it to be equal to
+		// our sensor aspect ratio.
+		return GetEffectiveFieldsOfView(CameraPose, CameraPose.GetSensorAspectRatio());
+	}
+
+	const float HorizontalFOV = CameraPose.GetEffectiveFieldOfView();
+
+	// Check the sort of aspect ratio axis constraint we have.
+	const EAspectRatioAxisConstraint Constraint(CameraPose.GetAspectRatioAxisConstraint());
+	if (Constraint == AspectRatio_MaintainYFOV ||
+			(Constraint == AspectRatio_MajorAxisFOV && AspectRatio < 1.f))
+	{
+		FCameraFieldsOfView FOVs;
+
+		// We need to maintain vertical FOV... the horizontal FOV we have is for our "ideal" aspect ratio,
+		// i.e. our sensor's aspect ratio. Now we need to compute the vertical FOV in this ideal situation
+		// and re-compute the effective horizontal FOV using the effective aspect ratio.
+		const float VerticalFOVRad = 2.f * FMath::Atan(
+				FMath::Tan(FMath::DegreesToRadians(HorizontalFOV / 2.f)) / AspectRatio);
+		const float HorizontalFOVRad = 2.f * FMath::Atan(
+				FMath::Tan(VerticalFOVRad / 2.f) * AspectRatio);
+
+		FOVs.HorizontalFieldOfView = FMath::RadiansToDegrees(HorizontalFOVRad);
+		FOVs.VerticalFieldOfView = FMath::RadiansToDegrees(VerticalFOVRad);
+
+		return FOVs;
+	}
+	else
+	{
+		// Our horizontal FOV is the effective one, so just compute the vertical FOV.
+		FCameraFieldsOfView FOVs;
+
+		FOVs.HorizontalFieldOfView = HorizontalFOV;
+		FOVs.VerticalFieldOfView = FMath::RadiansToDegrees(2.f * FMath::Atan(
+					FMath::Tan(FMath::DegreesToRadians(FOVs.HorizontalFieldOfView / 2.f)) / AspectRatio));
+
+		return FOVs;
+	}
+}
+
+float FCameraPoseMath::GetEffectiveAspectRatio(const FCameraPose& CameraPose, const APlayerController* PlayerController)
+{
+	if (CameraPose.GetConstrainAspectRatio() || PlayerController == nullptr)
+	{
+		return CameraPose.GetSensorAspectRatio();
+	}
+	else
+	{
+		int32 ViewportSizeX = 0, ViewportSizeY = 0;
+		PlayerController->GetViewportSize(ViewportSizeX, ViewportSizeY);
+		if (ensureMsgf(ViewportSizeX > 0 && ViewportSizeY > 0,
+				TEXT("Can't get viewport aspect ratio! Using sensor aspect ratio instead.")))
+		{
+			return (float)ViewportSizeX / (float)ViewportSizeY;
+		}
+		else
+		{
+			return CameraPose.GetSensorAspectRatio();
+		}
+	}
+}
+
+FMatrix FCameraPoseMath::BuildProjectionMatrix(const FCameraPose& CameraPose, float AspectRatio)
 {
 	const double NearClippingPlane = CameraPose.GetNearClippingPlane() > 0.0f ? 
 		CameraPose.GetNearClippingPlane() : GNearClippingPlane;
 	const double FieldOfView = FMath::Max(CameraPose.GetEffectiveFieldOfView(), 0.001f);
-	const float AspectRatio = CameraPose.GetSensorAspectRatio();
 
 	FMatrix ProjectionMatrix = FReversedZPerspectiveMatrix(
 			FMath::DegreesToRadians(FieldOfView / 2.f),
@@ -28,7 +106,7 @@ FMatrix FCameraPoseMath::BuildProjectionMatrix(const FCameraPose& CameraPose)
 	return ProjectionMatrix;
 }
 
-FMatrix FCameraPoseMath::BuildViewProjectionMatrix(const FCameraPose& CameraPose)
+FMatrix FCameraPoseMath::BuildViewProjectionMatrix(const FCameraPose& CameraPose, float AspectRatio)
 {
 	const FTranslationMatrix InverseOrigin(-CameraPose.GetLocation());
 
@@ -39,32 +117,34 @@ FMatrix FCameraPoseMath::BuildViewProjectionMatrix(const FCameraPose& CameraPose
 		FPlane(0,	1,	0,	0),
 		FPlane(0,	0,	0,	1));
 
-	const FMatrix ProjectionMatrix = BuildProjectionMatrix(CameraPose);
+	const FMatrix ProjectionMatrix = BuildProjectionMatrix(CameraPose, AspectRatio);
 
 	return InverseOrigin * InverseRotation * ProjectionMatrix;
 }
 
 TOptional<FVector2d> FCameraPoseMath::ProjectWorldToScreen(
-		const FCameraPose& CameraPose, const FVector3d& WorldLocation, bool bForceLocationInsideFrustum)
+		const FCameraPose& CameraPose, float AspectRatio, 
+		const FVector3d& WorldLocation, bool bForceLocationInsideFrustum)
 {
-	const FMatrix ViewProjectionMatrix = BuildViewProjectionMatrix(CameraPose);
+	const FMatrix ViewProjectionMatrix = BuildViewProjectionMatrix(CameraPose, AspectRatio);
 	return ProjectToScreen(ViewProjectionMatrix, WorldLocation, bForceLocationInsideFrustum);
 }
 
 TOptional<FVector2d> FCameraPoseMath::ProjectCameraToScreen(
-			const FCameraPose& CameraPose, const FVector3d& CameraSpaceLocation, bool bForceLocationInsideFrustum)
+			const FCameraPose& CameraPose, float AspectRatio, 
+			const FVector3d& CameraSpaceLocation, bool bForceLocationInsideFrustum)
 {
-	const FMatrix ProjectionMatrix = BuildProjectionMatrix(CameraPose);
+	const FMatrix ProjectionMatrix = BuildProjectionMatrix(CameraPose, AspectRatio);
 	return ProjectToScreen(ProjectionMatrix, CameraSpaceLocation, bForceLocationInsideFrustum);
 }
 
 TOptional<FVector2d> FCameraPoseMath::ProjectToScreen(
 		const FMatrix& ViewProjectionMatrix, const FVector3d& Location, bool bForceLocationInsideFrustum)
 {
-	const FVector4d ProjectedLocation = ViewProjectionMatrix.TransformFVector4(FVector4(Location, 1.f));
+	const FVector4d ProjectedLocation = ViewProjectionMatrix.TransformFVector4(FVector4(Location, 1.0));
 
 	// See if we need to handle the case of a point outside of the view frustum.
-	const bool bIsInsideFrustum = (ProjectedLocation.W > 0.f);
+	const bool bIsInsideFrustum = (ProjectedLocation.W > 0);
 	double W = ProjectedLocation.W;
 	if (!bIsInsideFrustum)
 	{
@@ -77,7 +157,7 @@ TOptional<FVector2d> FCameraPoseMath::ProjectToScreen(
 	}
 
 	// The result of this will be coordinates in -1..1 projection space.
-	const double RHW = 1.0f / W;
+	const double RHW = 1.0 / W;
 	const FVector4d ScreenSpaceLocation(
 			ProjectedLocation.X * RHW, 
 			ProjectedLocation.Y * RHW, 
@@ -85,36 +165,36 @@ TOptional<FVector2d> FCameraPoseMath::ProjectToScreen(
 			ProjectedLocation.W);
 
 	// Move from projection space to normalized 0..1 UI space.
-	const double ScreenSpaceX = (ScreenSpaceLocation.X / 2.f) + 0.5f;
-	const double ScreenSpaceY = 1.f - (ScreenSpaceLocation.Y / 2.f) - 0.5f;
+	const double ScreenSpaceX = (ScreenSpaceLocation.X / 2.0) + 0.5;
+	const double ScreenSpaceY = 1.0 - (ScreenSpaceLocation.Y / 2.0) - 0.5;
 
 	return FVector2d(ScreenSpaceX, ScreenSpaceY);
 }
 
-FRay3d FCameraPoseMath::UnprojectScreenToCamera(const FCameraPose& CameraPose, const FVector2D& ScreenSpacePoint)
+FRay3d FCameraPoseMath::UnprojectScreenToCamera(const FCameraPose& CameraPose, float AspectRatio, const FVector2D& ScreenSpacePoint)
 {
-	const FMatrix ProjectionMatrix = BuildProjectionMatrix(CameraPose);
+	const FMatrix ProjectionMatrix = BuildProjectionMatrix(CameraPose, AspectRatio);
 	const FMatrix InvProjectionMatrix = ProjectionMatrix.InverseFast();
 	return UnprojectFromScreen(InvProjectionMatrix, ScreenSpacePoint);
 }
 
-FVector3d FCameraPoseMath::UnprojectScreenToCamera(const FCameraPose& CameraPose, const FVector2D& ScreenSpacePoint, double PredictedDistance)
+FVector3d FCameraPoseMath::UnprojectScreenToCamera(const FCameraPose& CameraPose, float AspectRatio, const FVector2D& ScreenSpacePoint, double PredictedDistance)
 {
-	const FRay3d UnprojectedRay = UnprojectScreenToCamera(CameraPose, ScreenSpacePoint);
+	const FRay3d UnprojectedRay = UnprojectScreenToCamera(CameraPose, AspectRatio, ScreenSpacePoint);
 	const FVector3d WorldPoint = UnprojectedRay.PointAt(PredictedDistance);
 	return WorldPoint;
 }
 
-FRay3d FCameraPoseMath::UnprojectScreenToWorld(const FCameraPose& CameraPose, const FVector2D& ScreenSpacePoint)
+FRay3d FCameraPoseMath::UnprojectScreenToWorld(const FCameraPose& CameraPose, float AspectRatio, const FVector2D& ScreenSpacePoint)
 {
-	const FMatrix ViewProjectionMatrix = BuildViewProjectionMatrix(CameraPose);
+	const FMatrix ViewProjectionMatrix = BuildViewProjectionMatrix(CameraPose, AspectRatio);
 	const FMatrix InvViewProjectionMatrix = ViewProjectionMatrix.InverseFast();
 	return UnprojectFromScreen(InvViewProjectionMatrix, ScreenSpacePoint);
 }
 
-FVector3d FCameraPoseMath::UnprojectScreenToWorld(const FCameraPose& CameraPose, const FVector2D& ScreenSpacePoint, double PredictedDistance)
+FVector3d FCameraPoseMath::UnprojectScreenToWorld(const FCameraPose& CameraPose, float AspectRatio, const FVector2D& ScreenSpacePoint, double PredictedDistance)
 {
-	const FRay3d UnprojectedRay = UnprojectScreenToWorld(CameraPose, ScreenSpacePoint);
+	const FRay3d UnprojectedRay = UnprojectScreenToWorld(CameraPose, AspectRatio, ScreenSpacePoint);
 	const FVector3d WorldPoint = UnprojectedRay.PointAt(PredictedDistance);
 	return WorldPoint;
 }
@@ -122,25 +202,25 @@ FVector3d FCameraPoseMath::UnprojectScreenToWorld(const FCameraPose& CameraPose,
 FRay3d FCameraPoseMath::UnprojectFromScreen(const FMatrix& InverseViewProjectionMatrix, const FVector2D& ScreenSpacePoint)
 {
 	// Convert the given screen-space point from 0..1 UI space to -1..1 projection space.
-	const double ScreenSpaceX = (ScreenSpacePoint.X - 0.5f) * 2.f;
-	const double ScreenSpaceY = ((1.f - ScreenSpacePoint.Y) - 0.5f) * 2.f;
+	const double ScreenSpaceX = (ScreenSpacePoint.X - 0.5) * 2.0;
+	const double ScreenSpaceY = ((1.0 - ScreenSpacePoint.Y) - 0.5) * 2.0;
 
 	// Build a ray from the front of the frustum to the back of the frustum, starting at the screen-space point.
 	// We use reverse-Z projection matrices for better precision, so near is Z=1, and far is Z=0.
-	const FVector4 RayStartProjectionSpace = FVector4(ScreenSpaceX, ScreenSpaceY, 1.f, 1.f);
-	const FVector4 RayEndProjectionSpace = FVector4(ScreenSpaceX, ScreenSpaceY, 0.01f, 1.f);
+	const FVector4 RayStartProjectionSpace = FVector4(ScreenSpaceX, ScreenSpaceY, 1.0, 1.0);
+	const FVector4 RayEndProjectionSpace = FVector4(ScreenSpaceX, ScreenSpaceY, 0.01, 1.0);
 
 	// Unproject the ray points and normalize them.
 	const FVector4 RayStartProjected = InverseViewProjectionMatrix.TransformFVector4(RayStartProjectionSpace);
 	const FVector4 RayEndProjected = InverseViewProjectionMatrix.TransformFVector4(RayEndProjectionSpace);
 
 	FVector RayStartWorldSpace(RayStartProjected.X, RayStartProjected.Y, RayStartProjected.Z);
-	if (RayStartProjected.W != 0.f)
+	if (RayStartProjected.W != 0)
 	{
 		RayStartWorldSpace /= RayStartProjected.W;
 	}
 	FVector RayEndWorldSpace(RayEndProjected.X, RayEndProjected.Y, RayEndProjected.Z);
-	if (RayEndProjected.W != 0.f)
+	if (RayEndProjected.W != 0)
 	{
 		RayEndWorldSpace /= RayEndProjected.W;
 	}
diff --git a/Engine/Plugins/Cameras/GameplayCameras/Source/GameplayCameras/Private/Nodes/Framing/BaseFramingCameraNode.cpp b/Engine/Plugins/Cameras/GameplayCameras/Source/GameplayCameras/Private/Nodes/Framing/BaseFramingCameraNode.cpp
--- a/Engine/Plugins/Cameras/GameplayCameras/Source/GameplayCameras/Private/Nodes/Framing/BaseFramingCameraNode.cpp
+++ b/Engine/Plugins/Cameras/GameplayCameras/Source/GameplayCameras/Private/Nodes/Framing/BaseFramingCameraNode.cpp
@@ -56,11 +56,9 @@ void FBaseFramingCameraNodeEvaluator::OnInitialize(const FCameraNodeEvaluatorIni
 
 	Readers.DeadZoneMargin.Initialize(BaseFramingNode->DeadZone);
 	Readers.SoftZoneMargin.Initialize(BaseFramingNode->SoftZone);
-
-	State.bIsFirstFrame = true;
 }
 
-void FBaseFramingCameraNodeEvaluator::UpdateFramingState(const FCameraNodeEvaluationResult& OutResult, const FVector3d& TargetLocation)
+void FBaseFramingCameraNodeEvaluator::UpdateFramingState(const FCameraNodeEvaluationParams& Params, const FCameraNodeEvaluationResult& OutResult, const FVector3d& TargetLocation, const FTransform3d& LastFraming)
 {
 	// Get screen-space coordinates of the ideal framing point. These are in 0..1 UI space.
 	State.IdealTarget.X = Readers.HorizontalFraming.Get(OutResult.VariableTable);
@@ -79,20 +77,15 @@ void FBaseFramingCameraNodeEvaluator::UpdateFramingState(const FCameraNodeEvalua
 	State.DeadZone = FFramingZone(DeadZoneMargins);
 	State.SoftZone = FFramingZone(SoftZoneMargins);
 
-	// On the first frame, initialize LastFrameTransform to the current so we don't use garbage.
-	if (State.bIsFirstFrame)
-	{
-		State.LastFramedTransform = OutResult.CameraPose.GetTransform();
-		State.bIsFirstFrame = false;
-	}
-
 	// We are going to reframing things iteratively, so we'll use a temporary pose defined by last frame's
 	// shot transform.
 	FCameraPose TempPose(OutResult.CameraPose);
-	TempPose.SetTransform(State.LastFramedTransform);
+	TempPose.SetTransform(LastFraming);
 
 	// Get the target in screen-space.
-	const TOptional<FVector2d> ScreenTarget = FCameraPoseMath::ProjectWorldToScreen(TempPose, TargetLocation, true);
+	APlayerController* PlayerController = Params.EvaluationContext->GetPlayerController();
+	const float AspectRatio = FCameraPoseMath::GetEffectiveAspectRatio(TempPose, PlayerController);
+	const TOptional<FVector2d> ScreenTarget = FCameraPoseMath::ProjectWorldToScreen(TempPose, AspectRatio, TargetLocation, true);
 	State.WorldTarget = TargetLocation;
 	State.ScreenTarget = ScreenTarget.Get(FVector2d(0.5, 0.5));
 
@@ -133,7 +126,6 @@ void FBaseFramingCameraNodeEvaluator::UpdateFramingState(const FCameraNodeEvalua
 		{
 			// Since screen-space positions are in 0..1 space, we need to modulate vertical coordinates by the
 			// aspect ratio, otherwise we end up comparing against a squished ellipse instead of a circle.
-			const float AspectRatio = OutResult.CameraPose.GetSensorAspectRatio();
 			const double DistanceToIdeal = FVector2d::Distance(
 					FVector2d(State.ScreenTarget.X, (State.ScreenTarget.Y - 0.5) / AspectRatio + 0.5),
 					FVector2d(State.IdealTarget.X, (State.IdealTarget.Y - 0.5) / AspectRatio + 0.5));
@@ -181,16 +173,6 @@ void FBaseFramingCameraNodeEvaluator::ComputeDesiredState(float DeltaTime)
 	Desired.bHasCorrection = true;
 }
 
-const FTransform& FBaseFramingCameraNodeEvaluator::GetLastFraming() const
-{
-	return State.LastFramedTransform;
-}
-
-void FBaseFramingCameraNodeEvaluator::RegisterNewFraming(const FTransform& NewPoseTransform)
-{
-	State.LastFramedTransform = NewPoseTransform;
-}
-
 FVector2d FBaseFramingCameraNodeEvaluator::GetHardReframeCoords() const
 {
 	// The target is in the hard zone and must be brought back to the edge of the soft zone.
diff --git a/Engine/Plugins/Cameras/GameplayCameras/Source/GameplayCameras/Private/Nodes/Framing/DollyFramingCameraNode.cpp b/Engine/Plugins/Cameras/GameplayCameras/Source/GameplayCameras/Private/Nodes/Framing/DollyFramingCameraNode.cpp
--- a/Engine/Plugins/Cameras/GameplayCameras/Source/GameplayCameras/Private/Nodes/Framing/DollyFramingCameraNode.cpp
+++ b/Engine/Plugins/Cameras/GameplayCameras/Source/GameplayCameras/Private/Nodes/Framing/DollyFramingCameraNode.cpp
@@ -10,15 +10,10 @@
 #include "GameFramework/Pawn.h"
 #include "GameFramework/PlayerController.h"
 #include "GameplayCameras.h"
-#include "Math/CameraFramingZoneMath.h"
+#include "Math/CameraFramingMath.h"
 #include "Math/CameraPoseMath.h"
 #include "Math/InverseRotationMatrix.h"
 
-#include "Engine/LocalPlayer.h"
-#include "Engine/GameViewportClient.h"
-#include "Kismet/GameplayStatics.h"
-#include "SceneView.h"
-
 #include UE_INLINE_GENERATED_CPP_BY_NAME(DollyFramingCameraNode)
 
 namespace UE::Cameras
@@ -39,7 +34,8 @@ protected:
 
 private:
 
-	FVector3d ComputeFramingTranslation(const FCameraPose& CameraPose);
+	FTransform3d BuildDollyShotTransform(const FCameraPose& CameraPose) const;
+	FVector3d ComputeFramingTranslation(const FCameraPose& CameraPose, APlayerController* PlayerController);
 
 private:
 
@@ -49,7 +45,7 @@ private:
 	FVector2d DollyPosition;
 
 #if UE_GAMEPLAY_CAMERAS_DEBUG
-	FVector3d DebugUnprojectedNextScreenTarget;
+	FVector3d DebugNextDesiredTarget;
 	FVector2d DebugDollyCorrection;
 #endif  // UE_GAMEPLAY_CAMERAS_DEBUG
 };
@@ -60,7 +56,7 @@ UE_DECLARE_CAMERA_DEBUG_BLOCK_START(GAMEPLAYCAMERAS_API, FDollyFramingCameraDebu
 	UE_DECLARE_CAMERA_DEBUG_BLOCK_FIELD(FVector2d, DollyPosition);
 	UE_DECLARE_CAMERA_DEBUG_BLOCK_FIELD(FVector2d, DollyCorrection);
 	UE_DECLARE_CAMERA_DEBUG_BLOCK_FIELD(FVector3d, WorldTarget);
-	UE_DECLARE_CAMERA_DEBUG_BLOCK_FIELD(FVector3d, UnprojectedNextScreenTarget);
+	UE_DECLARE_CAMERA_DEBUG_BLOCK_FIELD(FVector3d, NextWorldTarget);
 UE_DECLARE_CAMERA_DEBUG_BLOCK_END()
 
 UE_DEFINE_CAMERA_DEBUG_BLOCK_WITH_FIELDS(FDollyFramingCameraDebugBlock)
@@ -84,18 +80,18 @@ void FDollyFramingCameraNodeEvaluator::OnRun(const FCameraNodeEvaluationParams&
 	const FVector3d TargetLocation = Pawn->GetActorLocation();
 
 	// Let the base class figure out all the screen-space framing stuff.
-	UpdateFramingState(OutResult, TargetLocation);
+	const FTransform3d LastShotTransform = BuildDollyShotTransform(OutResult.CameraPose);
+	UpdateFramingState(Params, OutResult, TargetLocation, LastShotTransform);
 	ComputeDesiredState(Params.DeltaTime);
 
 	// If we need to reframe the target this tick, figure out how much we need to move the dolly
 	// to accomplish that.
 	if (Desired.bHasCorrection)
 	{
-		const FTransform& LastFraming = GetLastFraming();
-		FCameraPose LastFramingPose(OutResult.CameraPose);
-		LastFramingPose.SetTransform(LastFraming);
+		FCameraPose LastShotPose(OutResult.CameraPose);
+		LastShotPose.SetTransform(LastShotTransform);
 
-		FVector3d DesiredLocalOffset = ComputeFramingTranslation(LastFramingPose);
+		FVector3d DesiredLocalOffset = ComputeFramingTranslation(LastShotPose, PlayerController);
 
 		// We never bring the dolly forward or backward (we only move it vertically or horizontally).
 		DesiredLocalOffset.X = 0.f;
@@ -116,37 +112,70 @@ void FDollyFramingCameraNodeEvaluator::OnRun(const FCameraNodeEvaluationParams&
 		DebugDollyCorrection = DollyCorrection;
 #endif  //UE_GAMEPLAY_CAMERAS_DEBUG
 	}
+	else
+	{
+#if UE_GAMEPLAY_CAMERAS_DEBUG
+		DebugDollyCorrection = FVector2d::ZeroVector;
+#endif  //UE_GAMEPLAY_CAMERAS_DEBUG
+	}
 
-	// Build the new shot transform and register it with the base class for next tick.
-	FTransform3d Transform = OutResult.CameraPose.GetTransform();
-	Transform = FTransform3d(FVector3d(0, DollyPosition.X, DollyPosition.Y)) * Transform;
-	OutResult.CameraPose.SetTransform(Transform);
-
-	RegisterNewFraming(Transform);
+	const FTransform3d NewShotTransform = BuildDollyShotTransform(OutResult.CameraPose);
+	OutResult.CameraPose.SetTransform(NewShotTransform);
 }
 
-FVector3d FDollyFramingCameraNodeEvaluator::ComputeFramingTranslation(const FCameraPose& CameraPose)
+FTransform3d FDollyFramingCameraNodeEvaluator::BuildDollyShotTransform(const FCameraPose& CameraPose) const
 {
-	// Unproject the desired screen-space position of our target for this tick. This normally gives us
-	// a ray into the camera frustum, since there is an infinity of world-space points that correspond
-	// to a point on the screen... we arbitrarily pick a point along this ray that is at the same 
-	// distance as our target, since generally it's not far from there.
-	const FVector3d CameraToTarget(State.WorldTarget - CameraPose.GetLocation());
-	const double TargetDist(CameraToTarget.Length());
+	FTransform3d Transform = CameraPose.GetTransform();
+	Transform = FTransform3d(FVector3d(0, DollyPosition.X, DollyPosition.Y)) * Transform;
+	return Transform;
+}
 
-	const FVector2D DesiredScreenTarget(Desired.ScreenTarget);
-	FVector3d RoughDesiredWorldTarget = FCameraPoseMath::UnprojectScreenToWorld(CameraPose, DesiredScreenTarget, TargetDist);
+FVector3d FDollyFramingCameraNodeEvaluator::ComputeFramingTranslation(const FCameraPose& CameraPose, APlayerController* PlayerController)
+{
+	const float AspectRatio = FCameraPoseMath::GetEffectiveAspectRatio(CameraPose, PlayerController);
+	const FCameraFieldsOfView FOVs(FCameraPoseMath::GetEffectiveFieldsOfView(CameraPose, AspectRatio));
+
+	// Get the position of the current target in camera space.
+	const FTransform3d InverseCameraTransform = CameraPose.GetTransform().Inverse();
+	const FVector3d TargetInCameraSpace = InverseCameraTransform.TransformPosition(State.WorldTarget);
+
+	// Get the horizontal and vertical angles, relative to the aiming vector, for the desired 
+	// target position.
+	const FVector2d DesiredAngles = FCameraFramingMath::GetTargetAngles(Desired.ScreenTarget, FOVs);
+
+	// The position of the desired target in camera space can be deduced from the angle and the current
+	// target position. We know that:
+	//
+	// 1) The current and desired targets will be on a plane parallel to the focal plane (i.e. on a plane
+	//    orthogonal to the aiming vector). This is because our dolly shot only translate laterally and
+	//    vertically -- it doesn't translate forwards/backwards.
+	//
+	// 2) We can do basic trigonometry for each axis (left/right and up/down, a.k.a. Y and Z in camera 
+	//    space). There's a right triangle between the aim vector (up to the targets' distance) and the
+	//    vector from the camera position to the desired target position.
+	//
+	// From (1) we know that the desired target's distance from the focal plane is the same as the
+	// current target's. So the X coordinates (near/far) are the same.
+	//
+	// From (2) we can use the sin() of the horizontal/vertical angles to get the horizontal/vertical
+	// coordinates of the desired target (again, in camera space). We just need to invert the vertical
+	// one because positive Z is up, while positive vertical angle is down (because this was computed
+	// in -1..1 UI screen-space).
+	//
+	const FVector3d DesiredInCameraSpace(
+			TargetInCameraSpace.X,
+			TargetInCameraSpace.X * FMath::Sin(FMath::DegreesToRadians(DesiredAngles.X)),
+			TargetInCameraSpace.X * -FMath::Sin(FMath::DegreesToRadians(DesiredAngles.Y)));
+
+	// Now we can figure out the desired camera-space offset that the dolly needs to move by. Remember
+	// that, for instance, moving the camera to the *right* will result in the target moving to the
+	// *left* on screen, so that's why we move by Desired->Current, and not the other way around.
+	FVector3d DesiredLocalOffset = (TargetInCameraSpace - DesiredInCameraSpace);
 
 #if UE_GAMEPLAY_CAMERAS_DEBUG
-	DebugUnprojectedNextScreenTarget = RoughDesiredWorldTarget;
+	DebugNextDesiredTarget = CameraPose.GetTransform().TransformPosition(DesiredInCameraSpace);
 #endif  // UE_GAMEPLAY_CAMERAS_DEBUG
 
-	// Now we know that we need to translate the dolly by that much in world-space. Convert that to
-	// a camera-space offset, since we want to decompose that into lateral and vertical dolly 
-	// movement.
-	const FVector3d DesiredWorldOffset(State.WorldTarget - RoughDesiredWorldTarget);
-	const FInverseRotationMatrix InverseRotation(CameraPose.GetRotation());
-	const FVector3d DesiredLocalOffset = InverseRotation.TransformVector(DesiredWorldOffset);
 	return DesiredLocalOffset;
 }
 
@@ -160,7 +189,7 @@ void FDollyFramingCameraNodeEvaluator::OnBuildDebugBlocks(const FCameraDebugBloc
 	DebugBlock.DollyPosition = DollyPosition;
 	DebugBlock.DollyCorrection = DebugDollyCorrection;
 	DebugBlock.WorldTarget = State.WorldTarget;
-	DebugBlock.UnprojectedNextScreenTarget = DebugUnprojectedNextScreenTarget;
+	DebugBlock.NextWorldTarget = DebugNextDesiredTarget;
 }
 
 void FDollyFramingCameraDebugBlock::OnDebugDraw(const FCameraDebugBlockDrawParams& Params, FCameraDebugRenderer& Renderer)
@@ -170,7 +199,7 @@ void FDollyFramingCameraDebugBlock::OnDebugDraw(const FCameraDebugBlockDrawParam
 			DollyPosition.X, DollyPosition.Y,
 			DollyCorrection.X, DollyCorrection.Y);
 
-	Renderer.DrawLine(WorldTarget, UnprojectedNextScreenTarget, FLinearColor(FColorList::NavyBlue), 1.f);
+	Renderer.DrawLine(WorldTarget, NextWorldTarget, FLinearColor(FColorList::LightGrey), 1.f);
 }
 
 #endif  // UE_GAMEPLAY_CAMERAS_DEBUG
diff --git a/Engine/Plugins/Cameras/GameplayCameras/Source/GameplayCameras/Private/Nodes/Framing/PanningFramingCameraNode.cpp b/Engine/Plugins/Cameras/GameplayCameras/Source/GameplayCameras/Private/Nodes/Framing/PanningFramingCameraNode.cpp
new file mode 100644
--- /dev/null
+++ b/Engine/Plugins/Cameras/GameplayCameras/Source/GameplayCameras/Private/Nodes/Framing/PanningFramingCameraNode.cpp
@@ -0,0 +1,162 @@
+// Copyright Epic Games, Inc. All Rights Reserved.
+
+#include "Nodes/Framing/PanningFramingCameraNode.h"
+
+#include "Core/CameraEvaluationContext.h"
+#include "Core/CameraParameterReader.h"
+#include "Debug/CameraDebugBlock.h"
+#include "Debug/CameraDebugBlockBuilder.h"
+#include "Debug/CameraDebugRenderer.h"
+#include "GameFramework/Pawn.h"
+#include "GameFramework/PlayerController.h"
+#include "GameplayCameras.h"
+#include "Math/CameraFramingMath.h"
+#include "Math/CameraPoseMath.h"
+#include "Math/InverseRotationMatrix.h"
+
+#include UE_INLINE_GENERATED_CPP_BY_NAME(PanningFramingCameraNode)
+
+namespace UE::Cameras
+{
+
+class FPanningFramingCameraNodeEvaluator : public FBaseFramingCameraNodeEvaluator
+{
+	UE_DECLARE_CAMERA_NODE_EVALUATOR_EX(GAMEPLAYCAMERAS_API, FPanningFramingCameraNodeEvaluator, FBaseFramingCameraNodeEvaluator)
+
+protected:
+
+	virtual void OnInitialize(const FCameraNodeEvaluatorInitializeParams& Params) override;
+	virtual void OnRun(const FCameraNodeEvaluationParams& Params, FCameraNodeEvaluationResult& OutResult) override;
+
+#if UE_GAMEPLAY_CAMERAS_DEBUG
+	virtual void OnBuildDebugBlocks(const FCameraDebugBlockBuildParams& Params, FCameraDebugBlockBuilder& Builder) override;
+#endif  // UE_GAMEPLAY_CAMERAS_DEBUG
+
+private:
+
+	FTransform3d BuildPanningShotTransform(const FCameraPose& CameraPose) const;
+
+private:
+
+	TCameraParameterReader<bool> CanPanLaterallyReader;
+	TCameraParameterReader<bool> CanPanVerticallyReader;
+
+	FRotator3d PanningRotation;
+
+#if UE_GAMEPLAY_CAMERAS_DEBUG
+	FRotator3d DebugPanningCorrection;
+#endif  // UE_GAMEPLAY_CAMERAS_DEBUG
+};
+
+UE_DEFINE_CAMERA_NODE_EVALUATOR(FPanningFramingCameraNodeEvaluator)
+
+UE_DECLARE_CAMERA_DEBUG_BLOCK_START(GAMEPLAYCAMERAS_API, FPanningFramingCameraDebugBlock)
+	UE_DECLARE_CAMERA_DEBUG_BLOCK_FIELD(FRotator3d, PanningRotation);
+	UE_DECLARE_CAMERA_DEBUG_BLOCK_FIELD(FRotator3d, PanningCorrection);
+UE_DECLARE_CAMERA_DEBUG_BLOCK_END()
+
+UE_DEFINE_CAMERA_DEBUG_BLOCK_WITH_FIELDS(FPanningFramingCameraDebugBlock)
+
+void FPanningFramingCameraNodeEvaluator::OnInitialize(const FCameraNodeEvaluatorInitializeParams& Params)
+{
+	Super::OnInitialize(Params);
+
+	const UPanningFramingCameraNode* DollyNode = GetCameraNodeAs<UPanningFramingCameraNode>();
+	CanPanLaterallyReader.Initialize(DollyNode->CanPanLaterally);
+	CanPanVerticallyReader.Initialize(DollyNode->CanPanVertically);
+
+	PanningRotation = FRotator3d::ZeroRotator;
+}
+
+void FPanningFramingCameraNodeEvaluator::OnRun(const FCameraNodeEvaluationParams& Params, FCameraNodeEvaluationResult& OutResult)
+{
+	// TODO get the target some other way
+	APlayerController* PlayerController = Params.EvaluationContext->GetPlayerController();
+	APawn* Pawn = PlayerController->GetPawn();
+	const FVector3d TargetLocation = Pawn->GetActorLocation();
+
+	// Let the base class figure out all the screen-space framing stuff.
+	const FTransform3d LastShotTransform = BuildPanningShotTransform(OutResult.CameraPose);
+	UpdateFramingState(Params, OutResult, TargetLocation, LastShotTransform);
+	ComputeDesiredState(Params.DeltaTime);
+
+	// If we need to reframe the target this tick, figure out how much we need to rotate the camera
+	// to accomplish that.
+	if (Desired.bHasCorrection)
+	{
+		FCameraPose LastShotPose(OutResult.CameraPose);
+		LastShotPose.SetTransform(LastShotTransform);
+
+		const float AspectRatio = FCameraPoseMath::GetEffectiveAspectRatio(LastShotPose, PlayerController);
+		const FCameraFieldsOfView FOVs(FCameraPoseMath::GetEffectiveFieldsOfView(LastShotPose, AspectRatio));
+
+		// Get yaw/pitch angles for where the target is currently at, and where we want it to be at.
+		const FVector2d TargetAngles = FCameraFramingMath::GetTargetAngles(State.ScreenTarget, FOVs);
+		const FVector2d DesiredAngles = FCameraFramingMath::GetTargetAngles(Desired.ScreenTarget, FOVs);
+
+		// We need to reverse pitch because rotators treat positive pitch as up, whereas we were 
+		// treating that as negative (due to the -1..1 normalized space of our math).
+		FRotator3d PanningCorrection(
+				(DesiredAngles.Y - TargetAngles.Y),	// Pitch
+				(TargetAngles.X - DesiredAngles.X),	// Yaw
+				0);									// Roll
+		PanningRotation += PanningCorrection;
+
+#if UE_GAMEPLAY_CAMERAS_DEBUG
+		DebugPanningCorrection = PanningCorrection;
+#endif  //UE_GAMEPLAY_CAMERAS_DEBUG
+	}
+	else
+	{
+#if UE_GAMEPLAY_CAMERAS_DEBUG
+		DebugPanningCorrection = FRotator3d::ZeroRotator;
+#endif  //UE_GAMEPLAY_CAMERAS_DEBUG
+	}
+
+	const FTransform3d NewShotTransform = BuildPanningShotTransform(OutResult.CameraPose);
+	OutResult.CameraPose.SetTransform(NewShotTransform);
+}
+
+FTransform3d FPanningFramingCameraNodeEvaluator::BuildPanningShotTransform(const FCameraPose& CameraPose) const
+{
+	FTransform3d Transform = CameraPose.GetTransform();
+	Transform = FTransform3d(PanningRotation) * Transform;
+	return Transform;
+}
+
+#if UE_GAMEPLAY_CAMERAS_DEBUG
+
+void FPanningFramingCameraNodeEvaluator::OnBuildDebugBlocks(const FCameraDebugBlockBuildParams& Params, FCameraDebugBlockBuilder& Builder)
+{
+	Super::OnBuildDebugBlocks(Params, Builder);
+
+	FPanningFramingCameraDebugBlock& DebugBlock = Builder.AttachDebugBlock<FPanningFramingCameraDebugBlock>();
+	DebugBlock.PanningRotation = PanningRotation;
+	DebugBlock.PanningCorrection = DebugPanningCorrection;
+}
+
+void FPanningFramingCameraDebugBlock::OnDebugDraw(const FCameraDebugBlockDrawParams& Params, FCameraDebugRenderer& Renderer)
+{
+	Renderer.AddText(
+			TEXT("pan yaw/pitch (%0.1f ; %0.1f)  correction (%0.1f ; %0.1f)"), 
+			PanningRotation.Yaw, PanningRotation.Pitch,
+			PanningCorrection.Yaw, PanningCorrection.Pitch);
+}
+
+#endif  // UE_GAMEPLAY_CAMERAS_DEBUG
+
+}  // namespace UE::Cameras
+
+UPanningFramingCameraNode::UPanningFramingCameraNode(const FObjectInitializer& ObjectInit)
+	: Super(ObjectInit)
+{
+	CanPanLaterally.Value = true;
+	CanPanVertically.Value = true;
+}
+
+FCameraNodeEvaluatorPtr UPanningFramingCameraNode::OnBuildEvaluator(FCameraNodeEvaluatorBuilder& Builder) const
+{
+	using namespace UE::Cameras;
+	return Builder.BuildEvaluator<FPanningFramingCameraNodeEvaluator>();
+}
+
diff --git a/Engine/Plugins/Cameras/GameplayCameras/Source/GameplayCameras/Public/Core/CameraPose.h b/Engine/Plugins/Cameras/GameplayCameras/Source/GameplayCameras/Public/Core/CameraPose.h
--- a/Engine/Plugins/Cameras/GameplayCameras/Source/GameplayCameras/Public/Core/CameraPose.h
+++ b/Engine/Plugins/Cameras/GameplayCameras/Source/GameplayCameras/Public/Core/CameraPose.h
@@ -9,6 +9,7 @@
 #include "CameraPose.generated.h"
 
 class FArchive;
+enum EAspectRatioAxisConstraint : int;
 
 #define UE_CAMERA_POSE_FOR_TRANSFORM_PROPERTIES()\
 	UE_CAMERA_POSE_FOR_PROPERTY(FVector, Location)\
@@ -29,13 +30,17 @@ class FArchive;
 	UE_CAMERA_POSE_FOR_PROPERTY(float, FocalLength)
 
 #define UE_CAMERA_POSE_FOR_BOOL_PROPERTIES()\
-	UE_CAMERA_POSE_FOR_PROPERTY(bool, bConstrainAspectRatio)
+	UE_CAMERA_POSE_FOR_PROPERTY(bool, ConstrainAspectRatio)
+
+#define UE_CAMERA_POSE_FOR_INTEGER_PROPERTIES()\
+	UE_CAMERA_POSE_FOR_PROPERTY(EAspectRatioAxisConstraint, AspectRatioAxisConstraint)
 
 #define UE_CAMERA_POSE_FOR_ALL_PROPERTIES()\
 	UE_CAMERA_POSE_FOR_TRANSFORM_PROPERTIES()\
 	UE_CAMERA_POSE_FOR_INTERPOLABLE_PROPERTIES()\
 	UE_CAMERA_POSE_FOR_FOV_PROPERTIES()\
-	UE_CAMERA_POSE_FOR_BOOL_PROPERTIES()
+	UE_CAMERA_POSE_FOR_BOOL_PROPERTIES()\
+	UE_CAMERA_POSE_FOR_INTEGER_PROPERTIES()
 
 /**
  * Boolean flags for each of the properties inside FCameraPose.
@@ -70,15 +75,6 @@ public:
 	FCameraPoseFlags& OR(const FCameraPoseFlags& OtherFlags);
 };
 
-/**
- * Simple struct for holding horizontal and vertical fields of view.
- */
-struct FCameraFieldsOfView
-{
-	float HorizontalFieldOfView;
-	float VerticalFieldOfView;
-};
-
 /**
  * Structure describing the state of a camera.
  *
@@ -152,9 +148,6 @@ public:
 	 */
 	float GetEffectiveFieldOfView() const;
 
-	/** Gets both horizontal and vertical effective fields of view.  */
-	FCameraFieldsOfView GetEffectiveFieldsOfView() const;
-
 	/** Gets the aspect ratio of the camera sensor. */
 	float GetSensorAspectRatio() const;
 
@@ -246,7 +239,11 @@ private:
 
 	/** Whether to constrain aspect ratio */
 	UPROPERTY()
-	bool bConstrainAspectRatio = false;
+	bool ConstrainAspectRatio = false;
+
+	/** If ConstrainAspectRatio is false, how we should compute FieldOfView */
+	UPROPERTY()
+	TEnumAsByte<EAspectRatioAxisConstraint> AspectRatioAxisConstraint;
 
 private:
 	
diff --git a/Engine/Plugins/Cameras/GameplayCameras/Source/GameplayCameras/Public/Debug/DebugTextRenderer.h b/Engine/Plugins/Cameras/GameplayCameras/Source/GameplayCameras/Public/Debug/DebugTextRenderer.h
--- a/Engine/Plugins/Cameras/GameplayCameras/Source/GameplayCameras/Public/Debug/DebugTextRenderer.h
+++ b/Engine/Plugins/Cameras/GameplayCameras/Source/GameplayCameras/Public/Debug/DebugTextRenderer.h
@@ -4,6 +4,7 @@
 
 #include "CoreTypes.h"
 #include "Containers/StringView.h"
+#include "Engine/EngineTypes.h"
 #include "GameplayCameras.h"
 #include "Math/UnrealMath.h"
 #include "Misc/TVariant.h"
@@ -47,6 +48,17 @@ FString ToDebugString(const UE::Math::TTransform<T>& FieldValue)
 {
 	return FieldValue.ToString();
 }
+template<>
+inline FString ToDebugString(const EAspectRatioAxisConstraint& FieldValue)
+{
+	switch (FieldValue)
+	{
+		case AspectRatio_MaintainYFOV: return TEXT("Maintain Y-Axis FOV");
+		case AspectRatio_MaintainXFOV: return TEXT("Maintain X-Axis FOV");
+		case AspectRatio_MajorAxisFOV: return TEXT("Maintain Major Axis FOV");
+	}
+	return TEXT("Invalid");
+}
 
 /** Command for drawing text on a canvas. */
 struct FDebugTextDrawCommand
diff --git a/Engine/Plugins/Cameras/GameplayCameras/Source/GameplayCameras/Public/Math/CameraFramingMath.h b/Engine/Plugins/Cameras/GameplayCameras/Source/GameplayCameras/Public/Math/CameraFramingMath.h
new file mode 100644
--- /dev/null
+++ b/Engine/Plugins/Cameras/GameplayCameras/Source/GameplayCameras/Public/Math/CameraFramingMath.h
@@ -0,0 +1,33 @@
+// Copyright Epic Games, Inc. All Rights Reserved.
+
+#pragma once
+
+#include "CoreTypes.h"
+#include "Math/MathFwd.h"
+
+namespace UE::Cameras
+{
+
+struct FCameraFieldsOfView;
+struct FFramingZone;
+struct FFramingZoneAngles;
+
+/**
+ * Utility class for mathematical functions related to framing zones.
+ */
+class FCameraFramingMath
+{
+public:
+
+	static FVector2d GetTargetAngles(const FVector2d& Target, const FCameraFieldsOfView& FieldsOfView);
+
+	/** Gets the framing zone's half-angles for a given camera FOV. */
+	static FFramingZoneAngles GetFramingZoneAngles(const FFramingZone& FramingZone, const FCameraFieldsOfView& FieldsOfView);
+
+private:
+
+	static double GetBoundAngle(float FactorFromCenter, double TanHalfFOV);
+};
+
+}  // namespace UE::Cameras
+
diff --git a/Engine/Plugins/Cameras/GameplayCameras/Source/GameplayCameras/Public/Math/CameraFramingZoneMath.h b/Engine/Plugins/Cameras/GameplayCameras/Source/GameplayCameras/Public/Math/CameraFramingZoneMath.h
--- a/Engine/Plugins/Cameras/GameplayCameras/Source/GameplayCameras/Public/Math/CameraFramingZoneMath.h
+++ b/Engine/Plugins/Cameras/GameplayCameras/Source/GameplayCameras/Public/Math/CameraFramingZoneMath.h
@@ -3,15 +3,16 @@
 #pragma once
 
 #include "CoreTypes.h"
-#include "Math/Vector.h"
+#include "Math/MathFwd.h"
 
 class FArchive;
-struct FCameraFieldsOfView;
 struct FCameraPose;
 
 namespace UE::Cameras
 {
 
+struct FCameraFieldsOfView;
+
 /**
  * Effective margins for a rectangular screen-space zone.
  */
@@ -60,7 +61,7 @@ struct FFramingZone
 	bool Contains(const FVector2d& Point) const;
 
 	/** Gets the inner margins of this zone compared to the screen's center. */
-	FVector4d GetMarginsFromCenter() const;
+	FVector4f GetNormalizedBounds() const;
 
 	/** Gets the coordinates of the top-left corner of the zone, in 0..Width/Height canvas units. */
 	FVector2d GetCanvasPosition(const FVector2d& CanvasSize) const;
@@ -68,6 +69,10 @@ struct FFramingZone
 	FVector2d GetCanvasSize(const FVector2d& CanvasSize) const;
 
 	void Serialize(FArchive& Ar);
+
+private:
+
+	static float GetNormalizedBound(float Bound);
 };
 
 FArchive& operator <<(FArchive& Ar, FFramingZone& FramingZone);
@@ -84,20 +89,5 @@ struct FFramingZoneAngles
 	double BottomHalfAngle = 0;
 };
 
-/**
- * Utility class for mathematical functions related to framing zones.
- */
-class FFramingZoneMath
-{
-public:
-
-	/** Gets the framing zone's half-angles for a given camera FOV. */
-	static FFramingZoneAngles GetFramingZoneAngles(const FFramingZone& FramingZone, const FCameraFieldsOfView& FieldsOfView);
-
-private:
-
-	static double GetFramingMarginAngle(float MarginFromCenter, float BackingHalfSize);
-};
-
-}  // namesdoubleUE::Cameras
+}  // namespace UE::Cameras
 
diff --git a/Engine/Plugins/Cameras/GameplayCameras/Source/GameplayCameras/Public/Math/CameraPoseMath.h b/Engine/Plugins/Cameras/GameplayCameras/Source/GameplayCameras/Public/Math/CameraPoseMath.h
--- a/Engine/Plugins/Cameras/GameplayCameras/Source/GameplayCameras/Public/Math/CameraPoseMath.h
+++ b/Engine/Plugins/Cameras/GameplayCameras/Source/GameplayCameras/Public/Math/CameraPoseMath.h
@@ -6,11 +6,22 @@
 #include "Math/MathFwd.h"
 #include "Math/Matrix.h"
 
+class APlayerController;
+struct FCameraFieldsOfView;
 struct FCameraPose;
 
 namespace UE::Cameras
 {
 
+/**
+ * Simple struct for holding horizontal and vertical fields of view.
+ */
+struct FCameraFieldsOfView
+{
+	float HorizontalFieldOfView;
+	float VerticalFieldOfView;
+};
+
 /**
  * A utility class for mathematical functions related to a camera pose.
  */
@@ -18,24 +29,38 @@ class FCameraPoseMath
 {
 public:
 
+	/** Gets both horizontal and vertical effective fields of view, using the sensor aspect ratio. */
+	static FCameraFieldsOfView GetEffectiveFieldsOfView(const FCameraPose& CameraPose);
+
+	/** Gets both horizontal and vertical effective fields of view, using the given player controller's viewport aspect ratio. */
+	static FCameraFieldsOfView GetEffectiveFieldsOfView(const FCameraPose& CameraPose, const APlayerController* PlayerController);
+
+	/** Gets both horizontal and vertical effective fields of view, using the given aspect ratio. */
+	static FCameraFieldsOfView GetEffectiveFieldsOfView(const FCameraPose& CameraPose, float AspectRatio);
+
+	/** Gets the aspect ratio of the viewport associated with the given player controller. */
+	static float GetEffectiveAspectRatio(const FCameraPose& CameraPose, const APlayerController* PlayerController);
+
 	/** 
 	 * Builds the projection matrix of the given camera pose. 
 	 * This matrix is suitable for projecting camera-space points onto screen-space.
 	 */
-	static FMatrix BuildProjectionMatrix(const FCameraPose& CameraPose);
+	static FMatrix BuildProjectionMatrix(const FCameraPose& CameraPose, float AspectRatio);
 	/** 
 	 * Builds the view-projection matrix of the given camera pose.
 	 * This matrix combines the camera transform and the projection matrix, making it
 	 * suitable for projecting world-space points onto screen-space.
 	 */
-	static FMatrix BuildViewProjectionMatrix(const FCameraPose& CameraPose);
+	static FMatrix BuildViewProjectionMatrix(const FCameraPose& CameraPose, float AspectRatio);
 
 	/** Projects the given world point onto screen-space. */
 	static TOptional<FVector2d> ProjectWorldToScreen(
-			const FCameraPose& CameraPose, const FVector3d& WorldLocation, bool bForceLocationInsideFrustum = false);
+			const FCameraPose& CameraPose, float AspectRatio, 
+			const FVector3d& WorldLocation, bool bForceLocationInsideFrustum = false);
 	/** Projects the given camera-local point onto screen-space. */
 	static TOptional<FVector2d> ProjectCameraToScreen(
-			const FCameraPose& CameraPose, const FVector3d& CameraSpaceLocation, bool bForceLocationInsideFrustum = false);
+			const FCameraPose& CameraPose, float AspectRatio, 
+			const FVector3d& CameraSpaceLocation, bool bForceLocationInsideFrustum = false);
 
 	/**
 	 * Projects the given point onto screen-space.
@@ -46,13 +71,17 @@ public:
 			const FMatrix& ViewProjectionMatrix, const FVector3d& Location, bool bForceLocationInsideFrustum = false);
 
 	/** Unproject a screen-space point into a camera-local ray. */
-	static FRay3d UnprojectScreenToCamera(const FCameraPose& CameraPose, const FVector2D& ScreenSpacePoint);
+	static FRay3d UnprojectScreenToCamera(
+			const FCameraPose& CameraPose, float AspectRatio, const FVector2D& ScreenSpacePoint);
 	/** Unproject a screen-space point into a camera-local point given an expected distance. */
-	static FVector3d UnprojectScreenToCamera(const FCameraPose& CameraPose, const FVector2D& ScreenSpacePoint, double PredictedDistance);
+	static FVector3d UnprojectScreenToCamera(
+			const FCameraPose& CameraPose, float AspectRatio, const FVector2D& ScreenSpacePoint, double PredictedDistance);
 	/** Unproject a screen-space point into a world-space ray. */
-	static FRay3d UnprojectScreenToWorld(const FCameraPose& CameraPose, const FVector2D& ScreenSpacePoint);
+	static FRay3d UnprojectScreenToWorld(
+			const FCameraPose& CameraPose, float AspectRatio, const FVector2D& ScreenSpacePoint);
 	/** Unproject a screen-space point into a world-space point given an expected distance. */
-	static FVector3d UnprojectScreenToWorld(const FCameraPose& CameraPose, const FVector2D& ScreenSpacePoint, double PredictedDistance);
+	static FVector3d UnprojectScreenToWorld(
+			const FCameraPose& CameraPose, float AspectRatio, const FVector2D& ScreenSpacePoint, double PredictedDistance);
 
 	/** Unprojects the given screen-space point into a ray. */
 	static FRay3d UnprojectFromScreen(
diff --git a/Engine/Plugins/Cameras/GameplayCameras/Source/GameplayCameras/Public/Nodes/Framing/BaseFramingCameraNode.h b/Engine/Plugins/Cameras/GameplayCameras/Source/GameplayCameras/Public/Nodes/Framing/BaseFramingCameraNode.h
--- a/Engine/Plugins/Cameras/GameplayCameras/Source/GameplayCameras/Public/Nodes/Framing/BaseFramingCameraNode.h
+++ b/Engine/Plugins/Cameras/GameplayCameras/Source/GameplayCameras/Public/Nodes/Framing/BaseFramingCameraNode.h
@@ -126,15 +126,10 @@ protected:
 protected:
 
 	/** Updates the framing state for the current tick, see State member field. */
-	void UpdateFramingState(const FCameraNodeEvaluationResult& OutResult, const FVector3d& TargetLocation);
+	void UpdateFramingState(const FCameraNodeEvaluationParams& Params, const FCameraNodeEvaluationResult& OutResult, const FVector3d& TargetLocation, const FTransform3d& LastFraming);
 	/** Computes the desired reframing for the current tick, see Desired member field. */
 	void ComputeDesiredState(float DeltaTime);
 
-	/** Gets the last registered shot transform (probably from last tick). */
-	const FTransform& GetLastFraming() const;
-	/** Registers a new shot transform computed by a sub-class. */
-	void RegisterNewFraming(const FTransform& NewPoseTransform);
-
 private:
 
 	FVector2d GetHardReframeCoords() const;
@@ -195,12 +190,8 @@ protected:
 		FVector2d ScreenTarget;
 		/** Current state of the tracked target. */
 		ETargetFramingState TargetFramingState;
-		/** Last tick's shot transform, as per RegisterNewFraming. */
-		FTransform3d LastFramedTransform;
 		/** Whether we are actively trying to bring the target back to the ideal position. */
 		bool bIsReframingTarget = false;
-		/** Whether this is the first tick. */
-		bool bIsFirstFrame = false;
 
 		/** The damper for reframing from the soft zone. */
 		FCriticalDamper ReframeDamper;
diff --git a/Engine/Plugins/Cameras/GameplayCameras/Source/GameplayCameras/Public/Nodes/Framing/PanningFramingCameraNode.h b/Engine/Plugins/Cameras/GameplayCameras/Source/GameplayCameras/Public/Nodes/Framing/PanningFramingCameraNode.h
new file mode 100644
--- /dev/null
+++ b/Engine/Plugins/Cameras/GameplayCameras/Source/GameplayCameras/Public/Nodes/Framing/PanningFramingCameraNode.h
@@ -0,0 +1,36 @@
+// Copyright Epic Games, Inc. All Rights Reserved.
+
+#pragma once
+
+#include "Nodes/Framing/BaseFramingCameraNode.h"
+
+#include "PanningFramingCameraNode.generated.h"
+
+/**
+ * A camera node that frames a target by rotating in place.
+ */
+UCLASS(MinimalAPI)
+class UPanningFramingCameraNode : public UBaseFramingCameraNode
+{
+	GENERATED_BODY()
+
+public:
+
+	/** Whether the camera can rotate laterally (yaw). */
+	UPROPERTY(EditAnywhere, Category="Panning")
+	FBooleanCameraParameter CanPanLaterally;
+
+	/** Whether the camera can rotate vertically (pitch). */
+	UPROPERTY(EditAnywhere, Category="Panning")
+	FBooleanCameraParameter CanPanVertically;
+
+public:
+
+	UPanningFramingCameraNode(const FObjectInitializer& ObjectInit);
+
+protected:
+
+	// UCameraNode interface.
+	virtual FCameraNodeEvaluatorPtr OnBuildEvaluator(FCameraNodeEvaluatorBuilder& Builder) const override;
+};
+
