Cameras: add occlusion material camera node

This new node changes the material of meshes found to be occluding the line of sight between the camera and a given target (generally the player pawn). This is useful to make these meshes transparent.

diff --git a/Engine/Plugins/Cameras/GameplayCameras/Source/GameplayCameras/Private/Core/CameraEvaluationContext.cpp b/Engine/Plugins/Cameras/GameplayCameras/Source/GameplayCameras/Private/Core/CameraEvaluationContext.cpp
--- a/Engine/Plugins/Cameras/GameplayCameras/Source/GameplayCameras/Private/Core/CameraEvaluationContext.cpp
+++ b/Engine/Plugins/Cameras/GameplayCameras/Source/GameplayCameras/Private/Core/CameraEvaluationContext.cpp
@@ -40,6 +40,15 @@ FCameraEvaluationContext::~FCameraEvaluationContext()
 	// holds the only shared pointer to it.
 }
 
+UWorld* FCameraEvaluationContext::GetWorld() const
+{
+	if (UObject* Owner = GetOwner())
+	{
+		return Owner->GetWorld();
+	}
+	return nullptr;
+}
+
 void FCameraEvaluationContext::AddReferencedObjects(FReferenceCollector& Collector)
 {
 	Collector.AddReferencedObject(CameraAsset);
diff --git a/Engine/Plugins/Cameras/GameplayCameras/Source/GameplayCameras/Private/Nodes/Collision/OcclusionMaterialCameraNode.cpp b/Engine/Plugins/Cameras/GameplayCameras/Source/GameplayCameras/Private/Nodes/Collision/OcclusionMaterialCameraNode.cpp
new file mode 100644
--- /dev/null
+++ b/Engine/Plugins/Cameras/GameplayCameras/Source/GameplayCameras/Private/Nodes/Collision/OcclusionMaterialCameraNode.cpp
@@ -0,0 +1,223 @@
+// Copyright Epic Games, Inc. All Rights Reserved.
+
+#include "Nodes/Collision/OcclusionMaterialCameraNode.h"
+
+#include "CollisionQueryParams.h"
+#include "Components/StaticMeshComponent.h"
+#include "Containers/Map.h"
+#include "Containers/Set.h"
+#include "Core/CameraEvaluationContext.h"
+#include "Core/CameraParameterReader.h"
+#include "Core/CameraSystemEvaluator.h"
+#include "Engine/World.h"
+#include "GameFramework/Pawn.h"
+#include "GameplayCameras.h"
+#include "Materials/MaterialInstanceDynamic.h"
+#include "Misc/AssertionMacros.h"
+#include "WorldCollision.h"
+
+#include UE_INLINE_GENERATED_CPP_BY_NAME(OcclusionMaterialCameraNode)
+
+namespace UE::Cameras
+{
+
+struct FOcclusionMaterialOverrideInfo
+{
+	TArray<UMaterialInterface*> OriginalMaterials;
+	TArray<UMaterialInterface*> OverrideMaterials;
+};
+
+class FOcclusionMaterialCameraNodeEvaluator : public FCameraNodeEvaluator
+{
+	UE_DECLARE_CAMERA_NODE_EVALUATOR(GAMEPLAYCAMERAS_API, FOcclusionMaterialCameraNodeEvaluator)
+
+protected:
+
+	virtual void OnInitialize(const FCameraNodeEvaluatorInitializeParams& Params) override;
+	virtual void OnRun(const FCameraNodeEvaluationParams& Params, FCameraNodeEvaluationResult& OutResult) override;
+
+private:
+
+	void RunOcclusionTrace(UWorld* World, APlayerController* PlayerController, FCameraNodeEvaluationResult& OutResult);
+	void HandleOcclusionTraceResult(UWorld* World);
+
+	void ApplyOcclusionMaterial(TSet<UStaticMeshComponent*> MeshComponents);
+	void RemoveOcclusionMaterial(TSet<UStaticMeshComponent*> MeshComponents);
+
+private:
+
+	TCameraParameterReader<float> OcclusionSphereRadiusReader;
+	TCameraParameterReader<FVector3d> OcclusionTargetOffsetReader;
+
+	FTraceHandle OcclusionTraceHandle;
+	TSet<UStaticMeshComponent*> CurrentlyOccludedMeshComponents;
+	TMap<UStaticMeshComponent*, FOcclusionMaterialOverrideInfo> AppliedMaterialOverrides;
+};
+
+UE_DEFINE_CAMERA_NODE_EVALUATOR(FOcclusionMaterialCameraNodeEvaluator)
+
+void FOcclusionMaterialCameraNodeEvaluator::OnInitialize(const FCameraNodeEvaluatorInitializeParams& Params)
+{
+	const UOcclusionMaterialCameraNode* OcclusionMaterialNode = GetCameraNodeAs<UOcclusionMaterialCameraNode>();
+	OcclusionSphereRadiusReader.Initialize(OcclusionMaterialNode->OcclusionSphereRadius);
+	OcclusionTargetOffsetReader.Initialize(OcclusionMaterialNode->OcclusionTargetOffset);
+}
+
+void FOcclusionMaterialCameraNodeEvaluator::OnRun(const FCameraNodeEvaluationParams& Params, FCameraNodeEvaluationResult& OutResult)
+{
+	if (!ensure(Params.EvaluationContext))
+	{
+		return;
+	}
+
+	UWorld* World = Params.EvaluationContext->GetWorld();
+	APlayerController* PlayerController = Params.EvaluationContext->GetPlayerController();
+	if (!World || !PlayerController)
+	{
+		return;
+	}
+
+	HandleOcclusionTraceResult(World);
+	RunOcclusionTrace(World, PlayerController, OutResult);
+}
+
+void FOcclusionMaterialCameraNodeEvaluator::RunOcclusionTrace(UWorld* World, APlayerController* PlayerController, FCameraNodeEvaluationResult& OutResult)
+{
+	static FName OcclusionTraceTag(TEXT("CameraOcclusion"));
+	static FName OcclusionTraceOwnerTag(TEXT("OcclusionMaterialCameraNode"));
+
+	APawn* Pawn = PlayerController->GetPawn();
+
+	// TODO fix this hardcoded nonsense
+	const FVector3d OcclusionTarget = Pawn->GetActorLocation();
+
+	const UOcclusionMaterialCameraNode* OcclusionMaterialNode = GetCameraNodeAs<UOcclusionMaterialCameraNode>();
+	ECollisionChannel OcclusionChannel = OcclusionMaterialNode->OcclusionChannel;
+
+	const float OcclusionSphereRadius = OcclusionSphereRadiusReader.Get(OutResult.VariableTable);
+	const FVector3d TraceEndOffset = OcclusionTargetOffsetReader.Get(OutResult.VariableTable);
+
+	const FVector3d TraceStart(OutResult.CameraPose.GetLocation());
+	const FVector3d TraceEnd(OcclusionTarget + TraceEndOffset);
+
+	FCollisionShape SweepShape = FCollisionShape::MakeSphere(OcclusionSphereRadius);
+	// Ignore the player pawn by default.
+	FCollisionQueryParams QueryParams(SCENE_QUERY_STAT(StartOcclusionSweep), false, Pawn);
+	QueryParams.TraceTag = OcclusionTraceTag;
+	QueryParams.OwnerTag = OcclusionTraceOwnerTag;
+
+	OcclusionTraceHandle = World->AsyncSweepByChannel(
+			EAsyncTraceType::Multi, 
+			TraceStart, TraceEnd, FQuat::Identity, 
+			OcclusionChannel, 
+			SweepShape,
+			QueryParams, 
+			FCollisionResponseParams::DefaultResponseParam);
+}
+
+void FOcclusionMaterialCameraNodeEvaluator::HandleOcclusionTraceResult(UWorld* World)
+{
+	// Do some basic validation... right now we just bail out if we can't get the trace result
+	// without figuring out if it's too old, still running, or whatever else. This is because
+	// we're supposed to be running only once a frame, so our trace should have run last frame
+	// and be available now. We'll have to better handle error cases when we start doing multi
+	// evaluations.
+	if (!OcclusionTraceHandle.IsValid())
+	{
+		return;
+	}
+
+	FTraceDatum TraceDatum;
+	if (!World->QueryTraceData(OcclusionTraceHandle, TraceDatum))
+	{
+		return;
+	}
+
+	// Get the list of meshes collected by the occlusion trace, and figure out which ones are
+	// new, and which ones got out of the way.
+	TSet<UStaticMeshComponent*> MeshComponents;
+	for (const FHitResult& Hit : TraceDatum.OutHits)
+	{
+		UStaticMeshComponent* MeshComponent = Cast<UStaticMeshComponent>(Hit.GetComponent());
+		if (MeshComponent)
+		{
+			MeshComponents.Add(MeshComponent);
+		}
+	}
+
+	TSet<UStaticMeshComponent*> NewMeshComponents = MeshComponents.Difference(CurrentlyOccludedMeshComponents);
+	TSet<UStaticMeshComponent*> OldMeshComponents = CurrentlyOccludedMeshComponents.Difference(MeshComponents);
+
+	CurrentlyOccludedMeshComponents = MeshComponents;
+
+	// Apply occlusion material changes to new/old components.
+	ApplyOcclusionMaterial(NewMeshComponents);
+	RemoveOcclusionMaterial(OldMeshComponents);
+
+	OcclusionTraceHandle.Invalidate();
+}
+
+void FOcclusionMaterialCameraNodeEvaluator::ApplyOcclusionMaterial(TSet<UStaticMeshComponent*> MeshComponents)
+{
+	const UOcclusionMaterialCameraNode* OcclusionMaterialNode = GetCameraNodeAs<UOcclusionMaterialCameraNode>();
+	UMaterialInterface* OcclusionTransparencyMaterial = OcclusionMaterialNode->OcclusionTransparencyMaterial;
+	if (!ensureMsgf(OcclusionTransparencyMaterial, TEXT("No occlusion transparency material set!")))
+	{
+		return;
+	}
+
+	for (UStaticMeshComponent* MeshComponent : MeshComponents)
+	{
+		if (!MeshComponent || AppliedMaterialOverrides.Contains(MeshComponent))
+		{
+			continue;
+		}
+
+		FOcclusionMaterialOverrideInfo& MaterialOverride = AppliedMaterialOverrides.Add(MeshComponent);
+		for (int32 MaterialIndex = 0; MaterialIndex < MeshComponent->GetNumMaterials(); ++MaterialIndex)
+		{
+			UMaterialInterface* OriginalMaterial = MeshComponent->GetMaterial(MaterialIndex);
+			UMaterialInterface* OverrideMaterial = MeshComponent->CreateDynamicMaterialInstance(MaterialIndex, OcclusionTransparencyMaterial);
+			MaterialOverride.OriginalMaterials.Add(OriginalMaterial);
+			MaterialOverride.OverrideMaterials.Add(OverrideMaterial);
+			MeshComponent->SetMaterial(MaterialIndex, OverrideMaterial);
+		}
+	}
+}
+
+void FOcclusionMaterialCameraNodeEvaluator::RemoveOcclusionMaterial(TSet<UStaticMeshComponent*> MeshComponents)
+{
+	for (UStaticMeshComponent* MeshComponent : MeshComponents)
+	{
+		if (!MeshComponent)
+		{
+			continue;
+		}
+
+		FOcclusionMaterialOverrideInfo MaterialOverrides;
+		const bool bFoundEntry = AppliedMaterialOverrides.RemoveAndCopyValue(MeshComponent, MaterialOverrides);
+		if (bFoundEntry)
+		{
+			const int32 NumMaterials = MaterialOverrides.OriginalMaterials.Num();
+			for (int32 MaterialIndex = 0; MaterialIndex < NumMaterials; ++MaterialIndex)
+			{
+				MeshComponent->SetMaterial(MaterialIndex, MaterialOverrides.OriginalMaterials[MaterialIndex]);
+			}
+		}
+	}
+}
+
+}  // namespace UE::Cameras
+
+UOcclusionMaterialCameraNode::UOcclusionMaterialCameraNode(const FObjectInitializer& ObjectInit)
+	: Super(ObjectInit)
+{
+	OcclusionSphereRadius.Value = 10.f;
+}
+
+FCameraNodeEvaluatorPtr UOcclusionMaterialCameraNode::OnBuildEvaluator(FCameraNodeEvaluatorBuilder& Builder) const
+{
+	using namespace UE::Cameras;
+	return Builder.BuildEvaluator<FOcclusionMaterialCameraNodeEvaluator>();
+}
+
diff --git a/Engine/Plugins/Cameras/GameplayCameras/Source/GameplayCameras/Public/Core/CameraEvaluationContext.h b/Engine/Plugins/Cameras/GameplayCameras/Source/GameplayCameras/Public/Core/CameraEvaluationContext.h
--- a/Engine/Plugins/Cameras/GameplayCameras/Source/GameplayCameras/Public/Core/CameraEvaluationContext.h
+++ b/Engine/Plugins/Cameras/GameplayCameras/Source/GameplayCameras/Public/Core/CameraEvaluationContext.h
@@ -68,6 +68,12 @@ public:
 	/** Gets the owner of this evaluation context, if any, and if still valid. */
 	UObject* GetOwner() const { return WeakOwner.Get(); }
 
+	/** 
+	 * Gets the world in which this evaluation context runs. 
+	 * This is the owner's world.
+	 */
+	UWorld* GetWorld() const;
+
 	/**
 	 * Gets the player controller (if any) in control of the cameras running inside
 	 * of this evaluation context.
diff --git a/Engine/Plugins/Cameras/GameplayCameras/Source/GameplayCameras/Public/Nodes/Collision/OcclusionMaterialCameraNode.h b/Engine/Plugins/Cameras/GameplayCameras/Source/GameplayCameras/Public/Nodes/Collision/OcclusionMaterialCameraNode.h
new file mode 100644
--- /dev/null
+++ b/Engine/Plugins/Cameras/GameplayCameras/Source/GameplayCameras/Public/Nodes/Collision/OcclusionMaterialCameraNode.h
@@ -0,0 +1,45 @@
+// Copyright Epic Games, Inc. All Rights Reserved.
+
+#pragma once
+
+#include "Core/CameraNode.h"
+#include "Core/CameraParameters.h"
+#include "Engine/EngineTypes.h"
+
+#include "OcclusionMaterialCameraNode.generated.h"
+
+class UMaterialInterface;
+
+UCLASS(MinimalAPI)
+class UOcclusionMaterialCameraNode : public UCameraNode
+{
+	GENERATED_BODY()
+
+public:
+
+	/** Material to apply on objects causing occlusion. */
+	UPROPERTY(EditAnywhere, Category="Occlusion")
+	UMaterialInterface* OcclusionTransparencyMaterial = nullptr;
+
+	/** Radius of the sphere used for occlusion testing. */
+	UPROPERTY(EditAnywhere, Category="Occlusion")
+	FFloatCameraParameter OcclusionSphereRadius;
+
+	/** Collision channel to use for the occlusion testing. */
+	UPROPERTY(EditAnywhere, Category="Occlusion")
+	TEnumAsByte<ECollisionChannel> OcclusionChannel;
+
+	/** World-space offset from the target to the occlusion trace's end. */
+	UPROPERTY(EditAnywhere, Category="Occlusion")
+	FVector3dCameraParameter OcclusionTargetOffset;
+
+public:
+
+	UOcclusionMaterialCameraNode(const FObjectInitializer& ObjectInit);
+
+protected:
+
+	// UCameraNode interface.
+	virtual FCameraNodeEvaluatorPtr OnBuildEvaluator(FCameraNodeEvaluatorBuilder& Builder) const override;
+};
+
