Cameras: improve precision of camera-related math and screen-space framing, and add interpolating damping factor to framing nodes

- Most math is now done with doubles instead of floats.
- Make the reframing damper go to the edge of the unlock zone instead of the center, so that we consume the damper all the way to zero.
- Add a low/high damping factor option on framing nodes.

diff --git a/Engine/Plugins/Cameras/GameplayCameras/Source/GameplayCameras/Private/Core/CameraPose.cpp b/Engine/Plugins/Cameras/GameplayCameras/Source/GameplayCameras/Private/Core/CameraPose.cpp
--- a/Engine/Plugins/Cameras/GameplayCameras/Source/GameplayCameras/Private/Core/CameraPose.cpp
+++ b/Engine/Plugins/Cameras/GameplayCameras/Source/GameplayCameras/Private/Core/CameraPose.cpp
@@ -113,7 +113,7 @@ void FCameraPose::SetTransform(FTransform3d Transform)
 	SetRotation(Transform.GetRotation().Rotator());
 }
 
-float FCameraPose::GetEffectiveFieldOfView() const
+double FCameraPose::GetEffectiveFieldOfView() const
 {
 	checkf((FocalLength > 0.f || FieldOfView > 0.f), TEXT("FocalLength or FieldOfView must have a valid, positive value."));
 
@@ -130,18 +130,18 @@ float FCameraPose::GetEffectiveFieldOfView() const
 	if (FocalLength > 0.f)
 	{
 		// Compute FOV with similar code to UCineCameraComponent...
-		float CropedSensorWidth = SensorWidth * SqueezeFactor;
-		const float AspectRatio = GetSensorAspectRatio();
-		if (AspectRatio > 0.0f)
+		double CropedSensorWidth = SensorWidth * SqueezeFactor;
+		const double AspectRatio = GetSensorAspectRatio();
+		if (AspectRatio > 0.0)
 		{
-			float DesqueezeAspectRatio = SensorWidth * SqueezeFactor / SensorHeight;
+			double DesqueezeAspectRatio = SensorWidth * SqueezeFactor / SensorHeight;
 			if (AspectRatio < DesqueezeAspectRatio)
 			{
 				CropedSensorWidth *= AspectRatio / DesqueezeAspectRatio;
 			}
 		}
 
-		return FMath::RadiansToDegrees(2.f * FMath::Atan(CropedSensorWidth / (2.f * FocalLength)));
+		return FMath::RadiansToDegrees(2.0 * FMath::Atan(CropedSensorWidth / (2.0 * FocalLength)));
 	}
 	else
 	{
@@ -150,9 +150,9 @@ float FCameraPose::GetEffectiveFieldOfView() const
 	}
 }
 
-float FCameraPose::GetSensorAspectRatio() const
+double FCameraPose::GetSensorAspectRatio() const
 {
-	return (SensorHeight > 0.f) ? (SensorWidth / SensorHeight) : 0.f;
+	return (SensorHeight > 0.f) ? (SensorWidth / SensorHeight) : 0.0;
 }
 
 FRay3d FCameraPose::GetAimRay() const
diff --git a/Engine/Plugins/Cameras/GameplayCameras/Source/GameplayCameras/Private/Math/CameraFramingMath.cpp b/Engine/Plugins/Cameras/GameplayCameras/Source/GameplayCameras/Private/Math/CameraFramingMath.cpp
--- a/Engine/Plugins/Cameras/GameplayCameras/Source/GameplayCameras/Private/Math/CameraFramingMath.cpp
+++ b/Engine/Plugins/Cameras/GameplayCameras/Source/GameplayCameras/Private/Math/CameraFramingMath.cpp
@@ -33,7 +33,7 @@ FFramingZoneAngles FCameraFramingMath::GetFramingZoneAngles(const FFramingZone&
 	const float TanHalfHorizontalFOV = FMath::Tan(FMath::DegreesToRadians(FieldsOfView.HorizontalFieldOfView / 2.f));
 	const float TanHalfVerticalFOV = FMath::Tan(FMath::DegreesToRadians(FieldsOfView.VerticalFieldOfView / 2.f));
 
-	const FVector4f BoundFactorsFromCenter = FramingZone.GetNormalizedBounds();
+	const FVector4d BoundFactorsFromCenter = FramingZone.GetNormalizedBounds();
 
 	const double LeftHalfAngleRad = GetBoundAngle(BoundFactorsFromCenter.X, TanHalfHorizontalFOV);
 	const double TopHalfAngleRad = GetBoundAngle(BoundFactorsFromCenter.Y, TanHalfVerticalFOV);
diff --git a/Engine/Plugins/Cameras/GameplayCameras/Source/GameplayCameras/Private/Math/CameraFramingZoneMath.cpp b/Engine/Plugins/Cameras/GameplayCameras/Source/GameplayCameras/Private/Math/CameraFramingZoneMath.cpp
--- a/Engine/Plugins/Cameras/GameplayCameras/Source/GameplayCameras/Private/Math/CameraFramingZoneMath.cpp
+++ b/Engine/Plugins/Cameras/GameplayCameras/Source/GameplayCameras/Private/Math/CameraFramingZoneMath.cpp
@@ -2,6 +2,7 @@
 
 #include "Math/CameraFramingZoneMath.h"
 
+#include "Algo/Sort.h"
 #include "Nodes/Framing/CameraFramingZone.h"
 
 namespace UE::Cameras
@@ -19,18 +20,18 @@ FFramingZone::FFramingZone(const FFramingZoneMargins& FramingZoneMargins)
 {
 	LeftBound = FramingZoneMargins.LeftMargin;
 	TopBound = FramingZoneMargins.TopMargin;
-	RightBound = 1.0f - FramingZoneMargins.RightMargin;
-	BottomBound = 1.0f - FramingZoneMargins.BottomMargin;
+	RightBound = 1.0 - FramingZoneMargins.RightMargin;
+	BottomBound = 1.0 - FramingZoneMargins.BottomMargin;
 
 	ClampBounds();
 }
 
 void FFramingZone::ClampBounds()
 {
-	LeftBound = FMath::Clamp(LeftBound, 0.f, 1.f);
-	TopBound = FMath::Clamp(TopBound, 0.f, 1.f);
-	RightBound = FMath::Clamp(RightBound, 0.f, 1.f);
-	BottomBound = FMath::Clamp(BottomBound, 0.f, 1.f);
+	LeftBound = FMath::Clamp(LeftBound, 0.0, 1.0);
+	TopBound = FMath::Clamp(TopBound, 0.0, 1.0);
+	RightBound = FMath::Clamp(RightBound, 0.0, 1.0);
+	BottomBound = FMath::Clamp(BottomBound, 0.0, 1.0);
 
 	RightBound = FMath::Max(LeftBound, RightBound);
 	BottomBound = FMath::Max(TopBound, BottomBound);
@@ -38,18 +39,18 @@ void FFramingZone::ClampBounds()
 
 void FFramingZone::ClampBounds(const FVector2d& MustContain)
 {
-	LeftBound = FMath::Clamp(LeftBound, 0.f, MustContain.X);
-	TopBound = FMath::Clamp(TopBound, 0.f, MustContain.Y);
-	RightBound = FMath::Clamp(RightBound, MustContain.X, 1.f);
-	BottomBound = FMath::Clamp(BottomBound, MustContain.Y, 1.f);
+	LeftBound = FMath::Clamp(LeftBound, 0.0, MustContain.X);
+	TopBound = FMath::Clamp(TopBound, 0.0, MustContain.Y);
+	RightBound = FMath::Clamp(RightBound, MustContain.X, 1.0);
+	BottomBound = FMath::Clamp(BottomBound, MustContain.Y, 1.0);
 }
 
 void FFramingZone::ClampBounds(const FFramingZone& MustContain)
 {
-	LeftBound = FMath::Clamp(LeftBound, 0.f, MustContain.LeftBound);
-	TopBound = FMath::Clamp(TopBound, 0.f, MustContain.TopBound);
-	RightBound = FMath::Clamp(RightBound, MustContain.RightBound, 1.f);
-	BottomBound = FMath::Clamp(BottomBound, MustContain.BottomBound, 1.f);
+	LeftBound = FMath::Clamp(LeftBound, 0.0, MustContain.LeftBound);
+	TopBound = FMath::Clamp(TopBound, 0.0, MustContain.TopBound);
+	RightBound = FMath::Clamp(RightBound, MustContain.RightBound, 1.0);
+	BottomBound = FMath::Clamp(BottomBound, MustContain.BottomBound, 1.0);
 }
 
 bool FFramingZone::Contains(const FVector2d& Point) const
@@ -58,20 +59,111 @@ bool FFramingZone::Contains(const FVector2d& Point) const
 		Point.Y >= TopBound && Point.Y <= BottomBound;
 }
 
-FVector4f FFramingZone::GetNormalizedBounds() const
+FVector2d FFramingZone::ComputeClosestIntersection(const FVector2d& Origin, const FVector2d& LineDir, bool bLineDirIsNormalized) const
+{
+	// Points along the line are of the form of:
+	//
+	//		P = Orig + Dir*d
+	//
+	// We test this equation against a result that yields an intersection with
+	// one of the unbounded lines of the zone. For instance, to see where it
+	// intersects with the top bound, we have:
+	//
+	//		P.y = TopBound
+	//		P.y = Orig.y + Dir.y*d
+	//		d = (P.y - Orig.y) / Dir.y
+	//
+	//		P.x = Orig.x + Dir.x*d
+	//		P.x = Orig.x + Dir.x*(P.y - Orig.y) / Dir.y
+	//		P.x = Orig.x + (TopBound - Orig.y) * (Dir.x / Dir.y)
+	//
+	// If P.x ends up being between LeftBound and RightBound, we have an intersection
+	// there. Otherwise, it misses the zone. Repeat for all four edges, and pick the
+	// closest intersection.
+
+	TArray<FVector2d, TInlineAllocator<2>> Intersections;
+	const FVector2d Dir = bLineDirIsNormalized ? LineDir : LineDir.GetSafeNormal();
+
+	// Intersection with TopBound/BottomBound.
+	if (Dir.Y != 0)
+	{
+		const double Slope = (Dir.X / Dir.Y);
+		// TopBound
+		{
+			FVector2d Intersection(0, TopBound);
+			Intersection.X = Origin.X + (TopBound - Origin.Y) * Slope;
+			if (Intersection.X >= LeftBound && Intersection.X <= RightBound)
+			{
+				Intersections.Add(Intersection);
+			}
+		}
+		// BottomBound
+		{
+			FVector2d Intersection(0, BottomBound);
+			Intersection.X = Origin.X + (BottomBound - Origin.Y) * Slope;
+			if (Intersection.X >= LeftBound && Intersection.X <= RightBound)
+			{
+				Intersections.Add(Intersection);
+			}
+		}
+	}
+	// Intersection with LeftBound/RightBound.
+	if (Dir.X != 0)
+	{
+		const double Slope = (Dir.Y / Dir.X);
+		// LeftBound
+		{
+			FVector2d Intersection(LeftBound, 0);
+			Intersection.Y = Origin.Y + (LeftBound - Origin.X) * Slope;
+			if (Intersection.Y >= TopBound && Intersection.Y <= BottomBound)
+			{
+				Intersections.Add(Intersection);
+			}
+		}
+		// RightBound
+		{
+			FVector2d Intersection(RightBound, 0);
+			Intersection.Y = Origin.Y + (RightBound - Origin.X) * Slope;
+			if (Intersection.Y >= TopBound && Intersection.Y <= BottomBound)
+			{
+				Intersections.Add(Intersection);
+			}
+		}
+	}
+
+	const int32 NumIntersections = Intersections.Num();
+	ensure(NumIntersections > 0 && NumIntersections <= 2);
+	if (NumIntersections > 0)
+	{
+		if (NumIntersections > 1)
+		{
+			Algo::SortBy(Intersections, [Origin](const FVector2d& Intersection)
+					{
+						return FVector2d::DistSquared(Origin, Intersection);
+					});
+		}
+		return Intersections[0];
+	}
+	else
+	{
+		return FVector2d::ZeroVector;
+	}
+}
+
+FVector4d FFramingZone::GetNormalizedBounds() const
 {
 	// Returned margins are negative if in the left or upper halves, and
 	// positive if in the right or lower halves.
-	return FVector4f(
+	return FVector4d(
 		GetNormalizedBound(LeftBound),
 		GetNormalizedBound(TopBound),
 		GetNormalizedBound(RightBound),
 		GetNormalizedBound(BottomBound));
 }
 
-float FFramingZone::GetNormalizedBound(float Bound)
+double FFramingZone::GetNormalizedBound(double Bound)
 {
-	return (Bound - 0.5f) * 2.f;
+	return (Bound - 0.5) * 2.0;
 }
 
 FVector2d FFramingZone::GetCanvasPosition(const FVector2d& CanvasSize) const
diff --git a/Engine/Plugins/Cameras/GameplayCameras/Source/GameplayCameras/Private/Math/CameraPoseMath.cpp b/Engine/Plugins/Cameras/GameplayCameras/Source/GameplayCameras/Private/Math/CameraPoseMath.cpp
--- a/Engine/Plugins/Cameras/GameplayCameras/Source/GameplayCameras/Private/Math/CameraPoseMath.cpp
+++ b/Engine/Plugins/Cameras/GameplayCameras/Source/GameplayCameras/Private/Math/CameraPoseMath.cpp
@@ -20,36 +20,36 @@ FCameraFieldsOfView FCameraPoseMath::GetEffectiveFieldsOfView(const FCameraPose&
 
 FCameraFieldsOfView FCameraPoseMath::GetEffectiveFieldsOfView(const FCameraPose& CameraPose, const APlayerController* PlayerController)
 {
-	const float AspectRatio = GetEffectiveAspectRatio(CameraPose, PlayerController);
+	const double AspectRatio = GetEffectiveAspectRatio(CameraPose, PlayerController);
 	return GetEffectiveFieldsOfView(CameraPose, AspectRatio);
 }
 
-FCameraFieldsOfView FCameraPoseMath::GetEffectiveFieldsOfView(const FCameraPose& CameraPose, float AspectRatio)
+FCameraFieldsOfView FCameraPoseMath::GetEffectiveFieldsOfView(const FCameraPose& CameraPose, double AspectRatio)
 {
 	if (CameraPose.GetConstrainAspectRatio() || 
-			!ensureMsgf(AspectRatio > 0.f, TEXT("Invalid aspect ratio! Using sensor aspect ratio instead.")))
+			!ensureMsgf(AspectRatio > 0.0, TEXT("Invalid aspect ratio! Using sensor aspect ratio instead.")))
 	{
 		// Aspect ratio is constrained, there will be black bars to enforce it to be equal to
 		// our sensor aspect ratio.
 		return GetEffectiveFieldsOfView(CameraPose, CameraPose.GetSensorAspectRatio());
 	}
 
-	const float HorizontalFOV = CameraPose.GetEffectiveFieldOfView();
+	const double HorizontalFOV = CameraPose.GetEffectiveFieldOfView();
 
 	// Check the sort of aspect ratio axis constraint we have.
 	const EAspectRatioAxisConstraint Constraint(CameraPose.GetAspectRatioAxisConstraint());
 	if (Constraint == AspectRatio_MaintainYFOV ||
-			(Constraint == AspectRatio_MajorAxisFOV && AspectRatio < 1.f))
+			(Constraint == AspectRatio_MajorAxisFOV && AspectRatio < 1.0))
 	{
 		FCameraFieldsOfView FOVs;
 
 		// We need to maintain vertical FOV... the horizontal FOV we have is for our "ideal" aspect ratio,
 		// i.e. our sensor's aspect ratio. Now we need to compute the vertical FOV in this ideal situation
 		// and re-compute the effective horizontal FOV using the effective aspect ratio.
-		const float VerticalFOVRad = 2.f * FMath::Atan(
-				FMath::Tan(FMath::DegreesToRadians(HorizontalFOV / 2.f)) / AspectRatio);
-		const float HorizontalFOVRad = 2.f * FMath::Atan(
-				FMath::Tan(VerticalFOVRad / 2.f) * AspectRatio);
+		const double VerticalFOVRad = 2.0 * FMath::Atan(
+				FMath::Tan(FMath::DegreesToRadians(HorizontalFOV / 2.0)) / AspectRatio);
+		const double HorizontalFOVRad = 2.0 * FMath::Atan(
+				FMath::Tan(VerticalFOVRad / 2.0) * AspectRatio);
 
 		FOVs.HorizontalFieldOfView = FMath::RadiansToDegrees(HorizontalFOVRad);
 		FOVs.VerticalFieldOfView = FMath::RadiansToDegrees(VerticalFOVRad);
@@ -62,14 +62,14 @@ FCameraFieldsOfView FCameraPoseMath::GetEffectiveFieldsOfView(const FCameraPose&
 		FCameraFieldsOfView FOVs;
 
 		FOVs.HorizontalFieldOfView = HorizontalFOV;
-		FOVs.VerticalFieldOfView = FMath::RadiansToDegrees(2.f * FMath::Atan(
-					FMath::Tan(FMath::DegreesToRadians(FOVs.HorizontalFieldOfView / 2.f)) / AspectRatio));
+		FOVs.VerticalFieldOfView = FMath::RadiansToDegrees(2.0 * FMath::Atan(
+					FMath::Tan(FMath::DegreesToRadians(FOVs.HorizontalFieldOfView / 2.0)) / AspectRatio));
 
 		return FOVs;
 	}
 }
 
-float FCameraPoseMath::GetEffectiveAspectRatio(const FCameraPose& CameraPose, const APlayerController* PlayerController)
+double FCameraPoseMath::GetEffectiveAspectRatio(const FCameraPose& CameraPose, const APlayerController* PlayerController)
 {
 	if (CameraPose.GetConstrainAspectRatio() || PlayerController == nullptr)
 	{
@@ -82,7 +82,7 @@ float FCameraPoseMath::GetEffectiveAspectRatio(const FCameraPose& CameraPose, co
 		if (ensureMsgf(ViewportSizeX > 0 && ViewportSizeY > 0,
 				TEXT("Can't get viewport aspect ratio! Using sensor aspect ratio instead.")))
 		{
-			return (float)ViewportSizeX / (float)ViewportSizeY;
+			return (double)ViewportSizeX / (double)ViewportSizeY;
 		}
 		else
 		{
@@ -91,14 +91,14 @@ float FCameraPoseMath::GetEffectiveAspectRatio(const FCameraPose& CameraPose, co
 	}
 }
 
-FMatrix FCameraPoseMath::BuildProjectionMatrix(const FCameraPose& CameraPose, float AspectRatio)
+FMatrix FCameraPoseMath::BuildProjectionMatrix(const FCameraPose& CameraPose, double AspectRatio)
 {
 	const double NearClippingPlane = CameraPose.GetNearClippingPlane() > 0.0f ? 
 		CameraPose.GetNearClippingPlane() : GNearClippingPlane;
 	const double FieldOfView = FMath::Max(CameraPose.GetEffectiveFieldOfView(), 0.001f);
 
 	FMatrix ProjectionMatrix = FReversedZPerspectiveMatrix(
-			FMath::DegreesToRadians(FieldOfView / 2.f),
+			FMath::DegreesToRadians(FieldOfView / 2.0),
 			AspectRatio,
 			1.0f,
 			NearClippingPlane);
@@ -106,7 +106,7 @@ FMatrix FCameraPoseMath::BuildProjectionMatrix(const FCameraPose& CameraPose, fl
 	return ProjectionMatrix;
 }
 
-FMatrix FCameraPoseMath::BuildViewProjectionMatrix(const FCameraPose& CameraPose, float AspectRatio)
+FMatrix FCameraPoseMath::BuildViewProjectionMatrix(const FCameraPose& CameraPose, double AspectRatio)
 {
 	const FTranslationMatrix InverseOrigin(-CameraPose.GetLocation());
 
@@ -123,7 +123,7 @@ FMatrix FCameraPoseMath::BuildViewProjectionMatrix(const FCameraPose& CameraPose
 }
 
 TOptional<FVector2d> FCameraPoseMath::ProjectWorldToScreen(
-		const FCameraPose& CameraPose, float AspectRatio, 
+		const FCameraPose& CameraPose, double AspectRatio, 
 		const FVector3d& WorldLocation, bool bForceLocationInsideFrustum)
 {
 	const FMatrix ViewProjectionMatrix = BuildViewProjectionMatrix(CameraPose, AspectRatio);
@@ -131,7 +131,7 @@ TOptional<FVector2d> FCameraPoseMath::ProjectWorldToScreen(
 }
 
 TOptional<FVector2d> FCameraPoseMath::ProjectCameraToScreen(
-			const FCameraPose& CameraPose, float AspectRatio, 
+			const FCameraPose& CameraPose, double AspectRatio, 
 			const FVector3d& CameraSpaceLocation, bool bForceLocationInsideFrustum)
 {
 	const FMatrix ProjectionMatrix = BuildProjectionMatrix(CameraPose, AspectRatio);
@@ -171,28 +171,28 @@ TOptional<FVector2d> FCameraPoseMath::ProjectToScreen(
 	return FVector2d(ScreenSpaceX, ScreenSpaceY);
 }
 
-FRay3d FCameraPoseMath::UnprojectScreenToCamera(const FCameraPose& CameraPose, float AspectRatio, const FVector2D& ScreenSpacePoint)
+FRay3d FCameraPoseMath::UnprojectScreenToCamera(const FCameraPose& CameraPose, double AspectRatio, const FVector2D& ScreenSpacePoint)
 {
 	const FMatrix ProjectionMatrix = BuildProjectionMatrix(CameraPose, AspectRatio);
 	const FMatrix InvProjectionMatrix = ProjectionMatrix.InverseFast();
 	return UnprojectFromScreen(InvProjectionMatrix, ScreenSpacePoint);
 }
 
-FVector3d FCameraPoseMath::UnprojectScreenToCamera(const FCameraPose& CameraPose, float AspectRatio, const FVector2D& ScreenSpacePoint, double PredictedDistance)
+FVector3d FCameraPoseMath::UnprojectScreenToCamera(const FCameraPose& CameraPose, double AspectRatio, const FVector2D& ScreenSpacePoint, double PredictedDistance)
 {
 	const FRay3d UnprojectedRay = UnprojectScreenToCamera(CameraPose, AspectRatio, ScreenSpacePoint);
 	const FVector3d WorldPoint = UnprojectedRay.PointAt(PredictedDistance);
 	return WorldPoint;
 }
 
-FRay3d FCameraPoseMath::UnprojectScreenToWorld(const FCameraPose& CameraPose, float AspectRatio, const FVector2D& ScreenSpacePoint)
+FRay3d FCameraPoseMath::UnprojectScreenToWorld(const FCameraPose& CameraPose, double AspectRatio, const FVector2D& ScreenSpacePoint)
 {
 	const FMatrix ViewProjectionMatrix = BuildViewProjectionMatrix(CameraPose, AspectRatio);
 	const FMatrix InvViewProjectionMatrix = ViewProjectionMatrix.InverseFast();
 	return UnprojectFromScreen(InvViewProjectionMatrix, ScreenSpacePoint);
 }
 
-FVector3d FCameraPoseMath::UnprojectScreenToWorld(const FCameraPose& CameraPose, float AspectRatio, const FVector2D& ScreenSpacePoint, double PredictedDistance)
+FVector3d FCameraPoseMath::UnprojectScreenToWorld(const FCameraPose& CameraPose, double AspectRatio, const FVector2D& ScreenSpacePoint, double PredictedDistance)
 {
 	const FRay3d UnprojectedRay = UnprojectScreenToWorld(CameraPose, AspectRatio, ScreenSpacePoint);
 	const FVector3d WorldPoint = UnprojectedRay.PointAt(PredictedDistance);
diff --git a/Engine/Plugins/Cameras/GameplayCameras/Source/GameplayCameras/Private/Nodes/Framing/BaseFramingCameraNode.cpp b/Engine/Plugins/Cameras/GameplayCameras/Source/GameplayCameras/Private/Nodes/Framing/BaseFramingCameraNode.cpp
--- a/Engine/Plugins/Cameras/GameplayCameras/Source/GameplayCameras/Private/Nodes/Framing/BaseFramingCameraNode.cpp
+++ b/Engine/Plugins/Cameras/GameplayCameras/Source/GameplayCameras/Private/Nodes/Framing/BaseFramingCameraNode.cpp
@@ -10,6 +10,7 @@
 #include "Debug/CameraDebugBlockBuilder.h"
 #include "Debug/CameraDebugRenderer.h"
 #include "GameplayCameras.h"
+#include "HAL/IConsoleManager.h"
 #include "Math/CameraPoseMath.h"
 #include "Math/ColorList.h"
 
@@ -18,6 +19,12 @@
 namespace UE::Cameras
 {
 
+float GFramingUnlockRadiusEpsilon = 1.e-4;
+static FAutoConsoleVariableRef CVarFramingUnlockRadiusEpsilon(
+	TEXT("GameplayCameras.Framing.UnlockRadiusEpsilon"),
+	GFramingUnlockRadiusEpsilon,
+	TEXT("(Default: 0.0001. The epsilon to determine whether we have reached the framing unlock circle."));
+
 UE_DECLARE_CAMERA_DEBUG_BLOCK_START(GAMEPLAYCAMERAS_API, FBaseFramingCameraDebugBlock)
 	UE_DECLARE_CAMERA_DEBUG_BLOCK_FIELD(FBaseFramingCameraNodeEvaluator::FState, State);
 	UE_DECLARE_CAMERA_DEBUG_BLOCK_FIELD(FBaseFramingCameraNodeEvaluator::FDesired, Desired);
@@ -52,6 +59,7 @@ void FBaseFramingCameraNodeEvaluator::OnInitialize(const FCameraNodeEvaluatorIni
 	Readers.HorizontalFraming.Initialize(BaseFramingNode->HorizontalFraming);
 	Readers.VerticalFraming.Initialize(BaseFramingNode->VerticalFraming);
 	Readers.ReframeDampingFactor.Initialize(BaseFramingNode->ReframeDampingFactor);
+	Readers.LowReframeDampingFactor.Initialize(BaseFramingNode->LowReframeDampingFactor);
 	Readers.ReframeUnlockRadius.Initialize(BaseFramingNode->ReframeUnlockRadius);
 
 	Readers.DeadZoneMargin.Initialize(BaseFramingNode->DeadZone);
@@ -63,12 +71,12 @@ void FBaseFramingCameraNodeEvaluator::UpdateFramingState(const FCameraNodeEvalua
 	// Get screen-space coordinates of the ideal framing point. These are in 0..1 UI space.
 	State.IdealTarget.X = Readers.HorizontalFraming.Get(OutResult.VariableTable);
 	State.IdealTarget.Y = Readers.VerticalFraming.Get(OutResult.VariableTable);
+
+	// Update the damping factors and unlock radius in case they are driven by a variable.
 	State.ReframeDampingFactor = Readers.ReframeDampingFactor.Get(OutResult.VariableTable);
+	State.LowReframeDampingFactor = Readers.LowReframeDampingFactor.Get(OutResult.VariableTable);
 	State.ReframeUnlockRadius = Readers.ReframeUnlockRadius.Get(OutResult.VariableTable);
 
-	// Update the damping factor in case it's driven by a variable.
-	State.ReframeDamper.SetW0(Readers.ReframeDampingFactor.Get(OutResult.VariableTable));
-
 	// Get the effective margins of the framing zones for this frame.
 	const FFramingZoneMargins DeadZoneMargins = Readers.DeadZoneMargin.GetZoneMargins(OutResult.VariableTable);
 	const FFramingZoneMargins SoftZoneMargins = Readers.SoftZoneMargin.GetZoneMargins(OutResult.VariableTable);
@@ -84,11 +92,30 @@ void FBaseFramingCameraNodeEvaluator::UpdateFramingState(const FCameraNodeEvalua
 
 	// Get the target in screen-space.
 	APlayerController* PlayerController = Params.EvaluationContext->GetPlayerController();
-	const float AspectRatio = FCameraPoseMath::GetEffectiveAspectRatio(TempPose, PlayerController);
+	const double AspectRatio = FCameraPoseMath::GetEffectiveAspectRatio(TempPose, PlayerController);
 	const TOptional<FVector2d> ScreenTarget = FCameraPoseMath::ProjectWorldToScreen(TempPose, AspectRatio, TargetLocation, true);
 	State.WorldTarget = TargetLocation;
 	State.ScreenTarget = ScreenTarget.Get(FVector2d(0.5, 0.5));
 
+	// Update the reframe damper's damping factor.
+	if (State.LowReframeDampingFactor <= 0)
+	{
+		State.ReframeDamper.SetW0(State.ReframeDampingFactor);
+	}
+	else
+	{
+		const FVector2d IdealToCurrent(State.ScreenTarget - State.IdealTarget);
+		const double UnlockEdgeToCurrent(IdealToCurrent.Length() - State.ReframeUnlockRadius);
+
+		const FVector2d HardZonePoint = State.SoftZone.ComputeClosestIntersection(State.IdealTarget, IdealToCurrent);
+		const double UnlockEdgeToHardZone = FMath::Max(
+				FVector2d::Distance(State.IdealTarget, HardZonePoint) - State.ReframeUnlockRadius,
+				UE_DOUBLE_SMALL_NUMBER);
+
+		const double Alpha = FMath::Clamp(UnlockEdgeToCurrent / UnlockEdgeToHardZone, 0.0, 1.1);
+		State.ReframeDamper.SetW0(FMath::Lerp(State.LowReframeDampingFactor, State.ReframeDampingFactor, Alpha));
+	}
+
 	// Make sure our framing zones are hierarchically correct: soft zone contains the dead zone, which contains
 	// the ideal target.
 	State.DeadZone.ClampBounds(State.IdealTarget);
@@ -115,7 +142,7 @@ void FBaseFramingCameraNodeEvaluator::UpdateFramingState(const FCameraNodeEvalua
 		State.TargetFramingState = ETargetFramingState::InDeadZone;
 
 		// Even though the target is free to move inside the dead zone, we might still want to continue reframing
-		// it towars the ideal position... if we didn't do that, reframing from the soft zone would stop entirely 
+		// it towards the ideal position... if we didn't do that, reframing from the soft zone would stop entirely 
 		// once we reach the edge of the dead zone, and we would never really ever see the target near the ideal
 		// position. So if we re-enter the dead zone from the soft zone, we keep reframing until we hit a smaller
 		// "unlock reframing" zone defined by the ReframeUnlockRadius.
@@ -129,7 +156,9 @@ void FBaseFramingCameraNodeEvaluator::UpdateFramingState(const FCameraNodeEvalua
 			const double DistanceToIdeal = FVector2d::Distance(
 					FVector2d(State.ScreenTarget.X, (State.ScreenTarget.Y - 0.5) / AspectRatio + 0.5),
 					FVector2d(State.IdealTarget.X, (State.IdealTarget.Y - 0.5) / AspectRatio + 0.5));
-			if (DistanceToIdeal <= State.ReframeUnlockRadius)
+			// Add an epsilon to the comparison to avoid being stuck in reframing mode because of
+			// floating point precision issues.
+			if (DistanceToIdeal <= (State.ReframeUnlockRadius + GFramingUnlockRadiusEpsilon))
 			{
 				State.bIsReframingTarget = false;
 			}
@@ -162,7 +191,11 @@ void FBaseFramingCameraNodeEvaluator::ComputeDesiredState(float DeltaTime)
 	}
 
 	// Move the target towards the ideal position using some damping.
-	const double NewDistanceToGo = State.ReframeDamper.Update(DistanceToGo, DeltaTime);
+	// Remove the radius of the unlock zone from the distance we pass to the damper, 
+	// otherwise the damper won't ever get to smoothly ease out to zero.
+	double DampingDistanceToGo = DistanceToGo - State.ReframeUnlockRadius;
+	const double NewDampedDistanceToGo = State.ReframeDamper.Update(DampingDistanceToGo, DeltaTime);
+	const double NewDistanceToGo = NewDampedDistanceToGo + State.ReframeUnlockRadius;
 
 	// Compute where we want the target this frame.
 	const FVector2d InvReframeDir(IdealToTarget / DistanceToGo);
@@ -288,7 +321,7 @@ void FBaseFramingCameraDebugBlock::OnDebugDraw(const FCameraDebugBlockDrawParams
 		Renderer.AddText(TEXT("[REFRAMING]"));
 	}
 
-	Renderer.AddText(TEXT(" (damping = %0.3f)"), State.ReframeDamper.GetX0());
+	Renderer.AddText(TEXT(" (damping = %0.3f, factor = %0.1f)"), State.ReframeDamper.GetX0(), State.ReframeDamper.GetW0());
 
 	if (Renderer.HasCanvas())
 	{
@@ -369,6 +402,7 @@ UBaseFramingCameraNode::UBaseFramingCameraNode(const FObjectInitializer& ObjectI
 	HorizontalFraming.Value = 0.5f;
 	VerticalFraming.Value = 0.5f;
 	ReframeDampingFactor.Value = 10.f;
+	LowReframeDampingFactor.Value = -1.f;
 	ReframeUnlockRadius.Value = 0.005f;
 }
 
diff --git a/Engine/Plugins/Cameras/GameplayCameras/Source/GameplayCameras/Public/Core/CameraPose.h b/Engine/Plugins/Cameras/GameplayCameras/Source/GameplayCameras/Public/Core/CameraPose.h
--- a/Engine/Plugins/Cameras/GameplayCameras/Source/GameplayCameras/Public/Core/CameraPose.h
+++ b/Engine/Plugins/Cameras/GameplayCameras/Source/GameplayCameras/Public/Core/CameraPose.h
@@ -146,10 +146,10 @@ public:
 	 * The effective field of view can be driven by the FieldOfView property, or
 	 * the FocalLength property in combination with the sensor size.
 	 */
-	float GetEffectiveFieldOfView() const;
+	double GetEffectiveFieldOfView() const;
 
 	/** Gets the aspect ratio of the camera sensor. */
-	float GetSensorAspectRatio() const;
+	double GetSensorAspectRatio() const;
 
 	/** Gets the aiming ray of the camera. */
 	FRay3d GetAimRay() const;
diff --git a/Engine/Plugins/Cameras/GameplayCameras/Source/GameplayCameras/Public/Math/CameraFramingZoneMath.h b/Engine/Plugins/Cameras/GameplayCameras/Source/GameplayCameras/Public/Math/CameraFramingZoneMath.h
--- a/Engine/Plugins/Cameras/GameplayCameras/Source/GameplayCameras/Public/Math/CameraFramingZoneMath.h
+++ b/Engine/Plugins/Cameras/GameplayCameras/Source/GameplayCameras/Public/Math/CameraFramingZoneMath.h
@@ -18,10 +18,10 @@ struct FCameraFieldsOfView;
  */
 struct FFramingZoneMargins
 {
-	float LeftMargin = 0;
-	float TopMargin = 0;
-	float RightMargin = 0;
-	float BottomMargin = 0;
+	double LeftMargin = 0;
+	double TopMargin = 0;
+	double RightMargin = 0;
+	double BottomMargin = 0;
 };
 
 /**
@@ -31,10 +31,10 @@ struct FFramingZoneMargins
  */
 struct FFramingZone
 {
-	float LeftBound = 0;
-	float TopBound = 0;
-	float RightBound = 0;
-	float BottomBound = 0;
+	double LeftBound = 0;
+	double TopBound = 0;
+	double RightBound = 0;
+	double BottomBound = 0;
 
 	/** Builds an empty framing zone. */
 	FFramingZone();
@@ -60,8 +60,14 @@ struct FFramingZone
 	/** Checks whether the given point (in 0..1 UI space) is inside this zone. */
 	bool Contains(const FVector2d& Point) const;
 
+	/** 
+	 * Computes intersections between a line and the zone's box, and returns the one closest
+	 * to the line's origin point.
+	 */
+	FVector2d ComputeClosestIntersection(const FVector2d& Origin, const FVector2d& LineDir, bool bLineDirIsNormalized = false) const;
+
 	/** Gets the inner margins of this zone compared to the screen's center. */
-	FVector4f GetNormalizedBounds() const;
+	FVector4d GetNormalizedBounds() const;
 
 	/** Gets the coordinates of the top-left corner of the zone, in 0..Width/Height canvas units. */
 	FVector2d GetCanvasPosition(const FVector2d& CanvasSize) const;
@@ -72,7 +78,7 @@ struct FFramingZone
 
 private:
 
-	static float GetNormalizedBound(float Bound);
+	static double GetNormalizedBound(double Bound);
 };
 
 FArchive& operator <<(FArchive& Ar, FFramingZone& FramingZone);
diff --git a/Engine/Plugins/Cameras/GameplayCameras/Source/GameplayCameras/Public/Math/CameraPoseMath.h b/Engine/Plugins/Cameras/GameplayCameras/Source/GameplayCameras/Public/Math/CameraPoseMath.h
--- a/Engine/Plugins/Cameras/GameplayCameras/Source/GameplayCameras/Public/Math/CameraPoseMath.h
+++ b/Engine/Plugins/Cameras/GameplayCameras/Source/GameplayCameras/Public/Math/CameraPoseMath.h
@@ -18,8 +18,8 @@ namespace UE::Cameras
  */
 struct FCameraFieldsOfView
 {
-	float HorizontalFieldOfView;
-	float VerticalFieldOfView;
+	double HorizontalFieldOfView;
+	double VerticalFieldOfView;
 };
 
 /**
@@ -36,30 +36,30 @@ public:
 	static FCameraFieldsOfView GetEffectiveFieldsOfView(const FCameraPose& CameraPose, const APlayerController* PlayerController);
 
 	/** Gets both horizontal and vertical effective fields of view, using the given aspect ratio. */
-	static FCameraFieldsOfView GetEffectiveFieldsOfView(const FCameraPose& CameraPose, float AspectRatio);
+	static FCameraFieldsOfView GetEffectiveFieldsOfView(const FCameraPose& CameraPose, double AspectRatio);
 
 	/** Gets the aspect ratio of the viewport associated with the given player controller. */
-	static float GetEffectiveAspectRatio(const FCameraPose& CameraPose, const APlayerController* PlayerController);
+	static double GetEffectiveAspectRatio(const FCameraPose& CameraPose, const APlayerController* PlayerController);
 
 	/** 
 	 * Builds the projection matrix of the given camera pose. 
 	 * This matrix is suitable for projecting camera-space points onto screen-space.
 	 */
-	static FMatrix BuildProjectionMatrix(const FCameraPose& CameraPose, float AspectRatio);
+	static FMatrix BuildProjectionMatrix(const FCameraPose& CameraPose, double AspectRatio);
 	/** 
 	 * Builds the view-projection matrix of the given camera pose.
 	 * This matrix combines the camera transform and the projection matrix, making it
 	 * suitable for projecting world-space points onto screen-space.
 	 */
-	static FMatrix BuildViewProjectionMatrix(const FCameraPose& CameraPose, float AspectRatio);
+	static FMatrix BuildViewProjectionMatrix(const FCameraPose& CameraPose, double AspectRatio);
 
 	/** Projects the given world point onto screen-space. */
 	static TOptional<FVector2d> ProjectWorldToScreen(
-			const FCameraPose& CameraPose, float AspectRatio, 
+			const FCameraPose& CameraPose, double AspectRatio, 
 			const FVector3d& WorldLocation, bool bForceLocationInsideFrustum = false);
 	/** Projects the given camera-local point onto screen-space. */
 	static TOptional<FVector2d> ProjectCameraToScreen(
-			const FCameraPose& CameraPose, float AspectRatio, 
+			const FCameraPose& CameraPose, double AspectRatio, 
 			const FVector3d& CameraSpaceLocation, bool bForceLocationInsideFrustum = false);
 
 	/**
@@ -72,16 +72,16 @@ public:
 
 	/** Unproject a screen-space point into a camera-local ray. */
 	static FRay3d UnprojectScreenToCamera(
-			const FCameraPose& CameraPose, float AspectRatio, const FVector2D& ScreenSpacePoint);
+			const FCameraPose& CameraPose, double AspectRatio, const FVector2D& ScreenSpacePoint);
 	/** Unproject a screen-space point into a camera-local point given an expected distance. */
 	static FVector3d UnprojectScreenToCamera(
-			const FCameraPose& CameraPose, float AspectRatio, const FVector2D& ScreenSpacePoint, double PredictedDistance);
+			const FCameraPose& CameraPose, double AspectRatio, const FVector2D& ScreenSpacePoint, double PredictedDistance);
 	/** Unproject a screen-space point into a world-space ray. */
 	static FRay3d UnprojectScreenToWorld(
-			const FCameraPose& CameraPose, float AspectRatio, const FVector2D& ScreenSpacePoint);
+			const FCameraPose& CameraPose, double AspectRatio, const FVector2D& ScreenSpacePoint);
 	/** Unproject a screen-space point into a world-space point given an expected distance. */
 	static FVector3d UnprojectScreenToWorld(
-			const FCameraPose& CameraPose, float AspectRatio, const FVector2D& ScreenSpacePoint, double PredictedDistance);
+			const FCameraPose& CameraPose, double AspectRatio, const FVector2D& ScreenSpacePoint, double PredictedDistance);
 
 	/** Unprojects the given screen-space point into a ray. */
 	static FRay3d UnprojectFromScreen(
diff --git a/Engine/Plugins/Cameras/GameplayCameras/Source/GameplayCameras/Public/Nodes/Framing/BaseFramingCameraNode.h b/Engine/Plugins/Cameras/GameplayCameras/Source/GameplayCameras/Public/Nodes/Framing/BaseFramingCameraNode.h
--- a/Engine/Plugins/Cameras/GameplayCameras/Source/GameplayCameras/Public/Nodes/Framing/BaseFramingCameraNode.h
+++ b/Engine/Plugins/Cameras/GameplayCameras/Source/GameplayCameras/Public/Nodes/Framing/BaseFramingCameraNode.h
@@ -25,16 +25,24 @@ public:
 
 	/** The ideal horizontal screen-space position of the target. */
 	UPROPERTY(EditAnywhere, Category="Framing Target")
-	FFloatCameraParameter HorizontalFraming;
+	FDoubleCameraParameter HorizontalFraming;
 
 	/** The ideal vertical screen-space position of the target. */
 	UPROPERTY(EditAnywhere, Category="Framing Target")
-	FFloatCameraParameter VerticalFraming;
+	FDoubleCameraParameter VerticalFraming;
 
 	/** The damping factor for how fast the framing recenters on the target. */
 	UPROPERTY(EditAnywhere, Category="Framing Target")
 	FFloatCameraParameter ReframeDampingFactor;
 
+	/** 
+	 * If valid, the recentering damping factor will interpolate between LowReframeDampingFactor 
+	 * and ReframeDampingFactor as the target moves between the ideal target position and the
+	 * boundaries of the hard-zone. If invalid, no interpolation occurs and the damping factor
+	 * is always equal to ReframeDampingFactor. */
+	UPROPERTY(EditAnywhere, Category="Framing Target")
+	FFloatCameraParameter LowReframeDampingFactor;
+
 	/**
 	 * The distance from the ideal framing position at which we can disengage reframing.
 	 * This should be a very small value, but if it is too small the reframing will keep "chasing"
@@ -77,10 +85,10 @@ struct FCameraFramingZoneParameterReader
 {
 public:
 
-	TCameraParameterReader<float> LeftMargin;
-	TCameraParameterReader<float> TopMargin;
-	TCameraParameterReader<float> RightMargin;
-	TCameraParameterReader<float> BottomMargin;
+	TCameraParameterReader<double> LeftMargin;
+	TCameraParameterReader<double> TopMargin;
+	TCameraParameterReader<double> RightMargin;
+	TCameraParameterReader<double> BottomMargin;
 
 public:
 
@@ -159,9 +167,10 @@ protected:
 	/** Utility structure for all the parameter readers we need every frame. */
 	struct FReaders
 	{
-		TCameraParameterReader<float> HorizontalFraming;
-		TCameraParameterReader<float> VerticalFraming;
+		TCameraParameterReader<double> HorizontalFraming;
+		TCameraParameterReader<double> VerticalFraming;
 		TCameraParameterReader<float> ReframeDampingFactor;
+		TCameraParameterReader<float> LowReframeDampingFactor;
 		TCameraParameterReader<float> ReframeUnlockRadius;
 
 		FCameraFramingZoneParameterReader DeadZoneMargin;
@@ -179,6 +188,8 @@ protected:
 		FVector2d IdealTarget;
 		/** Current reframing damping factor. */
 		float ReframeDampingFactor;
+		/** Current low reframing damping factor. */
+		float LowReframeDampingFactor;
 		/** Current reframe unlock radius. */
 		float ReframeUnlockRadius;
 		/** Current coordinates of the dead zone. */
diff --git a/Engine/Plugins/Cameras/GameplayCameras/Source/GameplayCameras/Public/Nodes/Framing/CameraFramingZone.h b/Engine/Plugins/Cameras/GameplayCameras/Source/GameplayCameras/Public/Nodes/Framing/CameraFramingZone.h
--- a/Engine/Plugins/Cameras/GameplayCameras/Source/GameplayCameras/Public/Nodes/Framing/CameraFramingZone.h
+++ b/Engine/Plugins/Cameras/GameplayCameras/Source/GameplayCameras/Public/Nodes/Framing/CameraFramingZone.h
@@ -21,16 +21,16 @@ struct FCameraFramingZone
 public:
 
 	UPROPERTY(EditAnywhere, Category="Framing")
-	FFloatCameraParameter LeftMargin;
+	FDoubleCameraParameter LeftMargin;
 
 	UPROPERTY(EditAnywhere, Category="Framing")
-	FFloatCameraParameter TopMargin;
+	FDoubleCameraParameter TopMargin;
 
 	UPROPERTY(EditAnywhere, Category="Framing")
-	FFloatCameraParameter RightMargin;
+	FDoubleCameraParameter RightMargin;
 
 	UPROPERTY(EditAnywhere, Category="Framing")
-	FFloatCameraParameter BottomMargin;
+	FDoubleCameraParameter BottomMargin;
 
 public:
 
@@ -42,7 +42,7 @@ public:
 		BottomMargin.Value = 0;
 	}
 
-	FCameraFramingZone(float UniformMargin)
+	FCameraFramingZone(double UniformMargin)
 	{
 		LeftMargin.Value = UniformMargin;
 		TopMargin.Value = UniformMargin;
@@ -50,7 +50,7 @@ public:
 		BottomMargin.Value = UniformMargin;
 	}
 
-	FCameraFramingZone(float HorizontalMargin, float VerticalMargin)
+	FCameraFramingZone(double HorizontalMargin, double VerticalMargin)
 	{
 		LeftMargin.Value = HorizontalMargin;
 		TopMargin.Value = VerticalMargin;
@@ -58,7 +58,7 @@ public:
 		BottomMargin.Value = VerticalMargin;
 	}
 
-	FCameraFramingZone(float InLeftMargin, float InTopMargin, float InRightMargin, float InBottomMargin)
+	FCameraFramingZone(double InLeftMargin, double InTopMargin, double InRightMargin, double InBottomMargin)
 	{
 		LeftMargin.Value = InLeftMargin;
 		TopMargin.Value = InTopMargin;
