Cameras: add transition conditions that use gameplay tags

Also of course add gameplay tag container to the camera rig asset. The camera asset will get gameplay tags too later.

diff --git a/Engine/Plugins/Cameras/GameplayCameras/Source/GameplayCameras/GameplayCameras.Build.cs b/Engine/Plugins/Cameras/GameplayCameras/Source/GameplayCameras/GameplayCameras.Build.cs
--- a/Engine/Plugins/Cameras/GameplayCameras/Source/GameplayCameras/GameplayCameras.Build.cs
+++ b/Engine/Plugins/Cameras/GameplayCameras/Source/GameplayCameras/GameplayCameras.Build.cs
@@ -14,6 +14,7 @@ public class GameplayCameras : ModuleRules
 				"Core",
 				"CoreUObject",
 				"Engine",
+				"GameplayTags",
 				"HeadMountedDisplay",
 				"MovieScene",
 				"MovieSceneTracks",
diff --git a/Engine/Plugins/Cameras/GameplayCameras/Source/GameplayCameras/Private/Core/CameraRigAsset.cpp b/Engine/Plugins/Cameras/GameplayCameras/Source/GameplayCameras/Private/Core/CameraRigAsset.cpp
--- a/Engine/Plugins/Cameras/GameplayCameras/Source/GameplayCameras/Private/Core/CameraRigAsset.cpp
+++ b/Engine/Plugins/Cameras/GameplayCameras/Source/GameplayCameras/Private/Core/CameraRigAsset.cpp
@@ -100,6 +100,11 @@ bool FCameraRigInterface::HasInterfaceParameter(const FString& ParameterName) co
 	return FindInterfaceParameterByName(ParameterName) != nullptr;
 }
 
+void UCameraRigAsset::GetOwnedGameplayTags(FGameplayTagContainer& TagContainer) const
+{
+	TagContainer.AppendTags(GameplayTags);
+}
+
 void UCameraRigAsset::BuildCameraRig()
 {
 	using namespace UE::Cameras;
diff --git a/Engine/Plugins/Cameras/GameplayCameras/Source/GameplayCameras/Private/Transitions/GameplayTagTransitionConditions.cpp b/Engine/Plugins/Cameras/GameplayCameras/Source/GameplayCameras/Private/Transitions/GameplayTagTransitionConditions.cpp
new file mode 100644
--- /dev/null
+++ b/Engine/Plugins/Cameras/GameplayCameras/Source/GameplayCameras/Private/Transitions/GameplayTagTransitionConditions.cpp
@@ -0,0 +1,36 @@
+// Copyright Epic Games, Inc. All Rights Reserved.
+
+#include "Transitions/GameplayTagTransitionConditions.h"
+
+#include "Core/CameraRigAsset.h"
+
+#include UE_INLINE_GENERATED_CPP_BY_NAME(GameplayTagTransitionConditions)
+
+bool UPreviousGameplayTagTransitionCondition::OnTransitionMatches(const FCameraRigTransitionConditionMatchParams& Params) const
+{
+	if (Params.FromCameraRig)
+	{
+		FGameplayTagContainer TagContainer;
+		Params.FromCameraRig->GetOwnedGameplayTags(TagContainer);
+		if (TagContainer.MatchesQuery(GameplayTagQuery))
+		{
+			return true;
+		}
+	}
+	return false;
+}
+
+bool UNextGameplayTagTransitionCondition::OnTransitionMatches(const FCameraRigTransitionConditionMatchParams& Params) const
+{
+	if (Params.ToCameraRig)
+	{
+		FGameplayTagContainer TagContainer;
+		Params.ToCameraRig->GetOwnedGameplayTags(TagContainer);
+		if (TagContainer.MatchesQuery(GameplayTagQuery))
+		{
+			return true;
+		}
+	}
+	return false;
+}
+
diff --git a/Engine/Plugins/Cameras/GameplayCameras/Source/GameplayCameras/Public/Core/CameraRigAsset.h b/Engine/Plugins/Cameras/GameplayCameras/Source/GameplayCameras/Public/Core/CameraRigAsset.h
--- a/Engine/Plugins/Cameras/GameplayCameras/Source/GameplayCameras/Public/Core/CameraRigAsset.h
+++ b/Engine/Plugins/Cameras/GameplayCameras/Source/GameplayCameras/Public/Core/CameraRigAsset.h
@@ -7,6 +7,8 @@
 #include "Core/CameraVariableTableFwd.h"
 #include "Core/ObjectTreeGraphObject.h"
 #include "CoreTypes.h"
+#include "GameplayTagAssetInterface.h"
+#include "GameplayTagContainer.h"
 #include "UObject/ObjectPtr.h"
 
 #include "CameraRigAsset.generated.h"
@@ -135,6 +137,7 @@ using FCameraRigPackages = TArray<const UPackage*, TInlineAllocator<4>>;
 UCLASS(MinimalAPI)
 class UCameraRigAsset
 	: public UObject
+	, public IGameplayTagAssetInterface
 	, public IObjectTreeGraphObject
 {
 	GENERATED_BODY()
@@ -151,6 +154,10 @@ public:
 	UPROPERTY(EditAnywhere, Instanced, Category=Common)
 	TObjectPtr<UCameraNode> RootNode;
 
+	/** The gameplay tags on this camera rig. */
+	UPROPERTY(EditAnywhere, Category=GameplayTags)
+	FGameplayTagContainer GameplayTags;
+
 	/** The public data interface of this camera rig. */
 	UPROPERTY()
 	FCameraRigInterface Interface;
@@ -180,6 +187,11 @@ public:
 	 */
 	GAMEPLAYCAMERAS_API void BuildCameraRig();
 
+public:
+
+	// IGameplayTagAssetInterface.
+	virtual void GetOwnedGameplayTags(FGameplayTagContainer& TagContainer) const override;
+
 protected:
 
 	// IObjectTreeGraphObject interface.
diff --git a/Engine/Plugins/Cameras/GameplayCameras/Source/GameplayCameras/Public/Transitions/GameplayTagTransitionConditions.h b/Engine/Plugins/Cameras/GameplayCameras/Source/GameplayCameras/Public/Transitions/GameplayTagTransitionConditions.h
new file mode 100644
--- /dev/null
+++ b/Engine/Plugins/Cameras/GameplayCameras/Source/GameplayCameras/Public/Transitions/GameplayTagTransitionConditions.h
@@ -0,0 +1,53 @@
+// Copyright Epic Games, Inc. All Rights Reserved.
+
+#pragma once
+
+#include "Core/CameraRigTransition.h"
+#include "GameplayTagContainer.h"
+
+#include "GameplayTagTransitionConditions.generated.h"
+
+/**
+ * A transition condition that matches the gameplay tags on the previous
+ * camera rig and asset.
+ */
+UCLASS(MinimalAPI)
+class UPreviousGameplayTagTransitionCondition 
+	: public UCameraRigTransitionCondition
+{
+	GENERATED_BODY()
+
+public:
+
+	/** The gameplay tags to look for on the previous camera rig/asset. */
+	UPROPERTY(EditAnywhere, Category=Transition)
+	FGameplayTagQuery GameplayTagQuery;
+
+protected:
+
+	// UCameraRigTransitionCondition interface.
+	virtual bool OnTransitionMatches(const FCameraRigTransitionConditionMatchParams& Params) const override;
+};
+
+/**
+ * A transition condition that matches the gameplay tags on the next
+ * camera rig and asset.
+ */
+UCLASS(MinimalAPI)
+class UNextGameplayTagTransitionCondition 
+	: public UCameraRigTransitionCondition
+{
+	GENERATED_BODY()
+
+public:
+
+	/** The gameplay tags to look for on the next camera rig/asset. */
+	UPROPERTY(EditAnywhere, Category=Transition)
+	FGameplayTagQuery GameplayTagQuery;
+
+protected:
+
+	// UCameraRigTransitionCondition interface.
+	virtual bool OnTransitionMatches(const FCameraRigTransitionConditionMatchParams& Params) const override;
+};
+
