Cameras: add null check before activating gameplay camera component

diff --git a/Engine/Plugins/Cameras/GameplayCameras/Source/GameplayCameras/Private/GameFramework/GameplayCameraComponent.cpp b/Engine/Plugins/Cameras/GameplayCameras/Source/GameplayCameras/Private/GameFramework/GameplayCameraComponent.cpp
--- a/Engine/Plugins/Cameras/GameplayCameras/Source/GameplayCameras/Private/GameFramework/GameplayCameraComponent.cpp
+++ b/Engine/Plugins/Cameras/GameplayCameras/Source/GameplayCameras/Private/GameFramework/GameplayCameraComponent.cpp
@@ -102,6 +102,11 @@ void UGameplayCameraComponent::ActivateCamera(APlayerController* PlayerControlle
 		return;
 	}
 
+	if (!ensure(Camera))
+	{
+		return;
+	}
+
 	if (!EvaluationContext.IsValid())
 	{
 		EvaluationContext = MakeShared<FGameplayCameraComponentEvaluationContext>();
