Cameras: add Blueprint functions to activate camera rigs in the base/global/visual layers

diff --git a/Engine/Plugins/Cameras/GameplayCameras/Source/GameplayCameras/Private/GameFramework/ActivateCameraRigFunctions.cpp b/Engine/Plugins/Cameras/GameplayCameras/Source/GameplayCameras/Private/GameFramework/ActivateCameraRigFunctions.cpp
new file mode 100644
--- /dev/null
+++ b/Engine/Plugins/Cameras/GameplayCameras/Source/GameplayCameras/Private/GameFramework/ActivateCameraRigFunctions.cpp
@@ -0,0 +1,81 @@
+// Copyright Epic Games, Inc. All Rights Reserved.
+
+#include "GameFramework/ActivateCameraRigFunctions.h"
+
+#include "Camera/PlayerCameraManager.h"
+#include "Core/CameraRigAsset.h"
+#include "Core/RootCameraNode.h"
+#include "Core/CameraSystemEvaluator.h"
+#include "Core/CameraEvaluationContext.h"
+#include "Core/CameraEvaluationContextStack.h"
+#include "GameFramework/GameplayCameraSystemActor.h"
+#include "GameFramework/GameplayCameraSystemComponent.h"
+#include "GameFramework/PlayerController.h"
+#include "GameplayCameras.h"
+#include "Templates/SharedPointer.h"
+
+#include UE_INLINE_GENERATED_CPP_BY_NAME(ActivateCameraRigFunctions)
+
+TSharedPtr<UE::Cameras::FCameraEvaluationContext> UActivateCameraRigFunctions::GlobalContext;
+
+void UActivateCameraRigFunctions::ActivateBaseCameraRig(APlayerController* PlayerController, UCameraRigAsset* CameraRig)
+{
+	ActivateCameraRig(PlayerController, CameraRig, ECameraRigLayer::Base);
+}
+
+void UActivateCameraRigFunctions::ActivateGlobalCameraRig(APlayerController* PlayerController, UCameraRigAsset* CameraRig)
+{
+	ActivateCameraRig(PlayerController, CameraRig, ECameraRigLayer::Global);
+}
+
+void UActivateCameraRigFunctions::ActivateVisualCameraRig(APlayerController* PlayerController, UCameraRigAsset* CameraRig)
+{
+	ActivateCameraRig(PlayerController, CameraRig, ECameraRigLayer::Visual);
+}
+
+void UActivateCameraRigFunctions::ActivateCameraRig(APlayerController* PlayerController, UCameraRigAsset* CameraRig, ECameraRigLayer EvaluationLayer)
+{
+	using namespace UE::Cameras;
+
+	if (!CameraRig)
+	{
+		UE_LOG(LogCameraSystem, Error, TEXT("No camera rig was given to activate!"));
+		return;
+	}
+
+	if (FCameraSystemEvaluator* SystemEvaluator = FindCameraSystemEvaluator(PlayerController))
+	{
+		FActivateCameraRigParams Params;
+		Params.CameraRig = CameraRig;
+		Params.EvaluationContext = EnsureGlobalContext();
+		Params.Evaluator = SystemEvaluator;
+		Params.Layer = EvaluationLayer;
+		SystemEvaluator->GetRootNodeEvaluator()->ActivateCameraRig(Params);
+	}
+}
+
+UE::Cameras::FCameraSystemEvaluator* UActivateCameraRigFunctions::FindCameraSystemEvaluator(APlayerController* PlayerController)
+{
+	if (PlayerController && PlayerController->PlayerCameraManager)
+	{
+		AActor* ViewTarget = PlayerController->PlayerCameraManager->GetViewTarget();
+		if (AGameplayCameraSystemActor* SystemActor = Cast<AGameplayCameraSystemActor>(ViewTarget))
+		{
+			return SystemActor->GetCameraSystemComponent()->GetCameraSystemEvaluator().Get();
+		}
+	}
+	return nullptr;
+}
+
+TSharedPtr<UE::Cameras::FCameraEvaluationContext> UActivateCameraRigFunctions::EnsureGlobalContext()
+{
+	using namespace UE::Cameras;
+
+	if (!GlobalContext.IsValid())
+	{
+		GlobalContext = MakeShared<FCameraEvaluationContext>();
+		GlobalContext->GetInitialResult().bIsValid = true;	
+	}
+	return GlobalContext;
+}
+
diff --git a/Engine/Plugins/Cameras/GameplayCameras/Source/GameplayCameras/Public/Core/RootCameraNode.h b/Engine/Plugins/Cameras/GameplayCameras/Source/GameplayCameras/Public/Core/RootCameraNode.h
--- a/Engine/Plugins/Cameras/GameplayCameras/Source/GameplayCameras/Public/Core/RootCameraNode.h
+++ b/Engine/Plugins/Cameras/GameplayCameras/Source/GameplayCameras/Public/Core/RootCameraNode.h
@@ -12,13 +12,13 @@ class UCameraRigAsset;
 /**
  * Defines evaluation layers for camera rigs.
  */
-UENUM()
-enum class ECameraRigLayer
+UENUM(BlueprintType)
+enum class ECameraRigLayer : uint8
 {
-	Base,
-	Main,
-	Global,
-	Visual,
+	Base UMETA(DisplayName="Base Layer"),
+	Main UMETA(DisplayName="Main Layer"),
+	Global UMETA(DisplayName="Global Layer"),
+	Visual UMETA(DisplayName="Visual Layer"),
 	User0,
 	User1,
 	User2
diff --git a/Engine/Plugins/Cameras/GameplayCameras/Source/GameplayCameras/Public/GameFramework/ActivateCameraRigFunctions.h b/Engine/Plugins/Cameras/GameplayCameras/Source/GameplayCameras/Public/GameFramework/ActivateCameraRigFunctions.h
new file mode 100644
--- /dev/null
+++ b/Engine/Plugins/Cameras/GameplayCameras/Source/GameplayCameras/Public/GameFramework/ActivateCameraRigFunctions.h
@@ -0,0 +1,63 @@
+// Copyright Epic Games, Inc. All Rights Reserved.
+
+#pragma once
+
+#include "Core/CameraEvaluationContext.h"
+#include "Core/RootCameraNode.h"
+#include "CoreTypes.h"
+#include "Kismet/BlueprintFunctionLibrary.h"
+#include "Templates/SharedPointerFwd.h"
+#include "UObject/ObjectMacros.h"
+#include "UObject/ObjectPtr.h"
+
+#include "ActivateCameraRigFunctions.generated.h"
+
+class APlayerController;
+class UCameraRigAsset;
+
+namespace UE::Cameras
+{
+	class FCameraEvaluationContext;
+	class FCameraSystemEvaluator;
+}  // namespace UE::Cameras
+
+/**
+ * Blueprint functions for activating camera rigs in the base/global/visual layers.
+ *
+ * These camera rigs run with a global, shared evaluation context that doesn't provide any
+ * meaningful initial result. They are activated on the camera system found to be running
+ * on the given player controller.
+ */
+UCLASS(MinimalAPI)
+class UActivateCameraRigFunctions : public UBlueprintFunctionLibrary
+{
+	GENERATED_BODY()
+
+public:
+
+	/** Activates the given camera rig in the base layer. */
+	UFUNCTION(BlueprintCallable, Category="Camera")
+	static void ActivateBaseCameraRig(APlayerController* PlayerController, UCameraRigAsset* CameraRig);
+
+	/** Activates the given camera rig in the global layer. */
+	UFUNCTION(BlueprintCallable, Category="Camera")
+	static void ActivateGlobalCameraRig(APlayerController* PlayerController, UCameraRigAsset* CameraRig);
+
+	/** Activates the given camera rig in the visual layer. */
+	UFUNCTION(BlueprintCallable, Category="Camera")
+	static void ActivateVisualCameraRig(APlayerController* PlayerController, UCameraRigAsset* CameraRig);
+
+	/** Activates the given camera rig in the given layer. */
+	UFUNCTION(BlueprintCallable, Category="Camera")
+	static void ActivateCameraRig(APlayerController* PlayerController, UCameraRigAsset* CameraRig, ECameraRigLayer EvaluationLayer);
+
+private:
+
+	static UE::Cameras::FCameraSystemEvaluator* FindCameraSystemEvaluator(APlayerController* PlayerController);
+	static TSharedPtr<UE::Cameras::FCameraEvaluationContext> EnsureGlobalContext();
+
+private:
+
+	static TSharedPtr<UE::Cameras::FCameraEvaluationContext> GlobalContext;
+};
+
