Cameras: add Blueprint camera director

diff --git a/Engine/Plugins/Cameras/GameplayCameras/Source/GameplayCameras/Private/Directors/BlueprintCameraDirector.cpp b/Engine/Plugins/Cameras/GameplayCameras/Source/GameplayCameras/Private/Directors/BlueprintCameraDirector.cpp
new file mode 100644
--- /dev/null
+++ b/Engine/Plugins/Cameras/GameplayCameras/Source/GameplayCameras/Private/Directors/BlueprintCameraDirector.cpp
@@ -0,0 +1,82 @@
+// Copyright Epic Games, Inc. All Rights Reserved.
+
+#include "Directors/BlueprintCameraDirector.h"
+
+#include "Algo/StableSort.h"
+#include "Core/CameraEvaluationContext.h"
+
+#include UE_INLINE_GENERATED_CPP_BY_NAME(BlueprintCameraDirector)
+
+namespace UE::Cameras
+{
+
+class FBlueprintCameraDirectorEvaluator : public FCameraDirectorEvaluator
+{
+	UE_DECLARE_CAMERA_DIRECTOR_EVALUATOR(GAMEPLAYCAMERAS_API, FBlueprintCameraDirectorEvaluator)
+
+protected:
+
+	virtual void OnInitialize(const FCameraDirectorInitializeParams& Params) override;
+	virtual void OnRun(const FCameraDirectorEvaluationParams& Params, FCameraDirectorEvaluationResult& OutResult) override;
+	virtual void OnAddReferencedObjects(FReferenceCollector& Collector) override;
+
+private:
+
+	TObjectPtr<UBlueprintCameraDirectorEvaluator> EvaluatorBlueprint;
+};
+
+UE_DEFINE_CAMERA_DIRECTOR_EVALUATOR(FBlueprintCameraDirectorEvaluator)
+
+void FBlueprintCameraDirectorEvaluator::OnInitialize(const FCameraDirectorInitializeParams& Params)
+{
+	const UBlueprintCameraDirector* Blueprint = GetCameraDirectorAs<UBlueprintCameraDirector>();
+	if (!ensure(Blueprint && Blueprint->CameraDirectorEvaluatorClass))
+	{
+		return;
+	}
+
+	UObject* Outer = Params.OwnerContext->GetOwner();
+	EvaluatorBlueprint = NewObject<UBlueprintCameraDirectorEvaluator>(Outer, Blueprint->CameraDirectorEvaluatorClass);
+}
+
+void FBlueprintCameraDirectorEvaluator::OnRun(const FCameraDirectorEvaluationParams& Params, FCameraDirectorEvaluationResult& OutResult)
+{
+	if (EvaluatorBlueprint)
+	{
+		FBlueprintCameraDirectorEvaluationParams BlueprintParams;
+		BlueprintParams.DeltaTime = Params.DeltaTime;
+		if (Params.OwnerContext)
+		{
+			BlueprintParams.EvaluationContextOwner = Params.OwnerContext->GetOwner();
+		}
+
+		FBlueprintCameraDirectorEvaluationResult BlueprintResult;
+
+		EvaluatorBlueprint->NativeRunCameraDirector(BlueprintParams, BlueprintResult);
+
+		for (UCameraRigAsset* ActiveCameraRig : BlueprintResult.ActiveCameraRigs)
+		{
+			OutResult.Add(Params.OwnerContext, ActiveCameraRig);
+		}
+	}
+}
+
+void FBlueprintCameraDirectorEvaluator::OnAddReferencedObjects(FReferenceCollector& Collector)
+{
+	Collector.AddReferencedObject(EvaluatorBlueprint);
+}
+
+}  // namespace UE::Cameras
+
+void UBlueprintCameraDirectorEvaluator::NativeRunCameraDirector(const FBlueprintCameraDirectorEvaluationParams& Params, FBlueprintCameraDirectorEvaluationResult& OutResult)
+{
+	RunCameraDirector(Params, OutResult.ActiveCameraRigs);
+}
+
+FCameraDirectorEvaluatorPtr UBlueprintCameraDirector::OnBuildEvaluator(FCameraDirectorEvaluatorBuilder& Builder) const
+{
+	using namespace UE::Cameras;
+
+	return Builder.BuildEvaluator<FBlueprintCameraDirectorEvaluator>();
+}
+
diff --git a/Engine/Plugins/Cameras/GameplayCameras/Source/GameplayCameras/Public/Directors/BlueprintCameraDirector.h b/Engine/Plugins/Cameras/GameplayCameras/Source/GameplayCameras/Public/Directors/BlueprintCameraDirector.h
new file mode 100644
--- /dev/null
+++ b/Engine/Plugins/Cameras/GameplayCameras/Source/GameplayCameras/Public/Directors/BlueprintCameraDirector.h
@@ -0,0 +1,85 @@
+// Copyright Epic Games, Inc. All Rights Reserved.
+
+#pragma once
+
+#include "Core/CameraDirector.h"
+#include "Core/CameraDirectorEvaluator.h"
+#include "Templates/SubclassOf.h"
+
+#include "BlueprintCameraDirector.generated.h"
+
+/**
+ * Parameter struct for the Blueprint camera director evaluator.
+ */
+USTRUCT(BlueprintType)
+struct FBlueprintCameraDirectorEvaluationParams
+{
+	GENERATED_BODY()
+
+	/** The elapsed time since the last evaluation. */
+	UPROPERTY(BlueprintReadWrite, Category="Evaluation")
+	float DeltaTime = 0.f;
+
+	/** The owner (if any) of the evaluation context we are running inside of. */
+	UPROPERTY(BlueprintReadWrite, Category = "Evaluation")
+	TObjectPtr<UObject> EvaluationContextOwner;
+};
+
+/**
+ * The evaluation result for the Blueprint camera director evaluator.
+ */
+USTRUCT(BlueprintType)
+struct FBlueprintCameraDirectorEvaluationResult
+{
+	GENERATED_BODY()
+
+	/** The list of camera rigs that should be active this frame. */
+	UPROPERTY(BlueprintReadWrite, Category = "Evaluation")
+	TArray<UCameraRigAsset*> ActiveCameraRigs;
+};
+
+/**
+ * Base class for a Blueprint camera director evaluator.
+ */
+UCLASS(MinimalAPI, Blueprintable)
+class UBlueprintCameraDirectorEvaluator : public UObject
+{
+	GENERATED_BODY()
+
+public:
+	
+	/**
+	 * Override this method in Blueprint to execute the custom logic that determines
+	 * what camera rig(s) should be active every frame.
+	 */
+	UFUNCTION(BlueprintCallable, BlueprintImplementableEvent, Category = "Evaluation")
+	void RunCameraDirector(
+			const FBlueprintCameraDirectorEvaluationParams& Params, 
+			TArray<UCameraRigAsset*>& OutActiveCameraRigs);
+
+	/** Native wrapper for RunCameraDirector. */
+	void NativeRunCameraDirector(
+			const FBlueprintCameraDirectorEvaluationParams& Params, 
+			FBlueprintCameraDirectorEvaluationResult& OutResult);
+};
+
+/**
+ * A camera director that will instantiate the given Blueprint and run it.
+ */
+UCLASS(MinimalAPI, EditInlineNew)
+class UBlueprintCameraDirector : public UCameraDirector
+{
+	GENERATED_BODY()
+
+public:
+
+	/** The blueprint class that we should instantiate and run. */
+	UPROPERTY(EditAnywhere, Category="Evaluation")
+	TSubclassOf<UBlueprintCameraDirectorEvaluator> CameraDirectorEvaluatorClass;
+
+protected:
+
+	// UCameraDirector interface.
+	virtual FCameraDirectorEvaluatorPtr OnBuildEvaluator(FCameraDirectorEvaluatorBuilder& Builder) const override;
+};
+
