Cameras: null fix in rig parameter setter BP node

diff --git a/Engine/Plugins/Cameras/GameplayCameras/Source/GameplayCamerasUncookedOnly/Private/BlueprintGraph/K2Node_SetCameraRigParameters.cpp b/Engine/Plugins/Cameras/GameplayCameras/Source/GameplayCamerasUncookedOnly/Private/BlueprintGraph/K2Node_SetCameraRigParameters.cpp
--- a/Engine/Plugins/Cameras/GameplayCameras/Source/GameplayCamerasUncookedOnly/Private/BlueprintGraph/K2Node_SetCameraRigParameters.cpp
+++ b/Engine/Plugins/Cameras/GameplayCameras/Source/GameplayCamerasUncookedOnly/Private/BlueprintGraph/K2Node_SetCameraRigParameters.cpp
@@ -137,6 +137,12 @@ void UK2Node_SetCameraRigParameters::ExpandNode(class FKismetCompilerContext& Co
 			continue;
 		}
 
+		if (!InterfaceParameter->PrivateVariable)
+		{
+			CompilerContext.MessageLog.Error(*LOCTEXT("ErrorMissingParameterVariable", "SetCameraRigParameters node @@ needs camera rig @@ to be built.").ToString(), this, CameraRig);
+			continue;
+		}
+
 		// Figure out the sort of SetXxxParameter function we want to call for this parameter.
 		FName CallSetParameterFuncName;
 		switch (InterfaceParameter->PrivateVariable->GetVariableType())
