Cameras: auto-restore the correct pin direction of an object tree graph node when rebuilding the graph

diff --git a/Engine/Plugins/Cameras/GameplayCameras/Source/GameplayCamerasEditor/Private/Editors/ObjectTreeGraphSchema.cpp b/Engine/Plugins/Cameras/GameplayCameras/Source/GameplayCamerasEditor/Private/Editors/ObjectTreeGraphSchema.cpp
--- a/Engine/Plugins/Cameras/GameplayCameras/Source/GameplayCamerasEditor/Private/Editors/ObjectTreeGraphSchema.cpp
+++ b/Engine/Plugins/Cameras/GameplayCameras/Source/GameplayCamerasEditor/Private/Editors/ObjectTreeGraphSchema.cpp
@@ -106,6 +106,10 @@ void UObjectTreeGraphSchema::CreateConnections(UObjectTreeGraphNode* InGraphNode
 			UObjectTreeGraphNode* const* ConnectedNode = InCreatedNodes.CreatedNodes.Find(OutConnectedObject);
 			if (ensure(ConnectedNode))
 			{
+				if (Pin->Direction == EGPD_Input)
+				{
+					(*ConnectedNode)->OverrideSelfPinDirection(EGPD_Output);
+				}
 				UEdGraphPin* ConnectedPin = (*ConnectedNode)->GetSelfPin();
 				Pin->MakeLinkTo(ConnectedPin);
 			}
@@ -133,6 +137,10 @@ void UObjectTreeGraphSchema::CreateConnections(UObjectTreeGraphNode* InGraphNode
 				UObjectTreeGraphNode* const* ConnectedNode = InCreatedNodes.CreatedNodes.Find(ConnectedObject);
 				if (ensure(ConnectedNode))
 				{
+					if (Pin->Direction == EGPD_Input)
+					{
+						(*ConnectedNode)->OverrideSelfPinDirection(EGPD_Output);
+					}
 					UEdGraphPin* ConnectedPin = (*ConnectedNode)->GetSelfPin();
 					Pin->MakeLinkTo(ConnectedPin);
 				}
