Cameras: pass the visualized world from the rewind debugger extension to the debug renderer

diff --git a/Engine/Plugins/Cameras/GameplayCameras/Source/GameplayCamerasEditor/Private/Trace/CameraSystemRewindDebuggerExtension.cpp b/Engine/Plugins/Cameras/GameplayCameras/Source/GameplayCamerasEditor/Private/Trace/CameraSystemRewindDebuggerExtension.cpp
--- a/Engine/Plugins/Cameras/GameplayCameras/Source/GameplayCamerasEditor/Private/Trace/CameraSystemRewindDebuggerExtension.cpp
+++ b/Engine/Plugins/Cameras/GameplayCameras/Source/GameplayCamerasEditor/Private/Trace/CameraSystemRewindDebuggerExtension.cpp
@@ -67,6 +67,7 @@ void FCameraSystemRewindDebuggerExtension::Update(float DeltaTime, IRewindDebugg
 		if (FoundFrameData && CurrentTraceTime != LastTraceTime)
 		{
 			LastTraceTime = CurrentTraceTime;
+			VisualizedWorld = RewindDebugger->GetWorldToVisualize();
 
 			DebugBlockStorage.DestroyDebugBlocks();
 
@@ -99,10 +100,9 @@ void FCameraSystemRewindDebuggerExtension::EnsureDebugDrawDelegate(bool bIsRegis
 
 void FCameraSystemRewindDebuggerExtension::DebugDraw(UCanvas* Canvas, APlayerController* PlayController)
 {
-	if (RootDebugBlock)
+	if (RootDebugBlock && VisualizedWorld)
 	{
-		UWorld* World = PlayController->GetWorld();
-		FCameraDebugRenderer CameraDebugRenderer(World, Canvas->Canvas);
+		FCameraDebugRenderer CameraDebugRenderer(VisualizedWorld, Canvas->Canvas);
 		RootDebugBlock->RootDebugDraw(CameraDebugRenderer);
 	}
 }
diff --git a/Engine/Plugins/Cameras/GameplayCameras/Source/GameplayCamerasEditor/Private/Trace/CameraSystemRewindDebuggerExtension.h b/Engine/Plugins/Cameras/GameplayCameras/Source/GameplayCamerasEditor/Private/Trace/CameraSystemRewindDebuggerExtension.h
--- a/Engine/Plugins/Cameras/GameplayCameras/Source/GameplayCamerasEditor/Private/Trace/CameraSystemRewindDebuggerExtension.h
+++ b/Engine/Plugins/Cameras/GameplayCameras/Source/GameplayCamerasEditor/Private/Trace/CameraSystemRewindDebuggerExtension.h
@@ -44,6 +44,8 @@ private:
 	FDelegateHandle DebugDrawDelegateHandle;
 	double LastTraceTime = 0.f;
 
+	UWorld* VisualizedWorld = nullptr;
+
 	FCameraDebugBlockStorage DebugBlockStorage;
 	FRootCameraDebugBlock* RootDebugBlock = nullptr;
 
