Cameras: reset debug block child visit flags before recursing into children

diff --git a/Engine/Plugins/Cameras/GameplayCameras/Source/GameplayCameras/Private/Core/CameraNodeEvaluator.cpp b/Engine/Plugins/Cameras/GameplayCameras/Source/GameplayCameras/Private/Core/CameraNodeEvaluator.cpp
--- a/Engine/Plugins/Cameras/GameplayCameras/Source/GameplayCameras/Private/Core/CameraNodeEvaluator.cpp
+++ b/Engine/Plugins/Cameras/GameplayCameras/Source/GameplayCameras/Private/Core/CameraNodeEvaluator.cpp
@@ -109,6 +109,7 @@ void FCameraNodeEvaluator::BuildDebugBlocks(const FCameraDebugBlockBuildParams&
 
 		// Build debug blocks for children node evaluators.
 		ECameraDebugBlockBuildVisitFlags VisitFlags = Builder.GetVisitFlags();
+		Builder.ResetVisitFlags();
 		if (!EnumHasAnyFlags(VisitFlags, ECameraDebugBlockBuildVisitFlags::SkipChildren))
 		{
 			FCameraNodeEvaluatorChildrenView ChildrenView(GetChildren());
