Cameras: use world context object in camera rig activation BP functions to give the global evaluation context a proper owner

diff --git a/Engine/Plugins/Cameras/GameplayCameras/Source/GameplayCameras/Private/GameFramework/ActivateCameraRigFunctions.cpp b/Engine/Plugins/Cameras/GameplayCameras/Source/GameplayCameras/Private/GameFramework/ActivateCameraRigFunctions.cpp
--- a/Engine/Plugins/Cameras/GameplayCameras/Source/GameplayCameras/Private/GameFramework/ActivateCameraRigFunctions.cpp
+++ b/Engine/Plugins/Cameras/GameplayCameras/Source/GameplayCameras/Private/GameFramework/ActivateCameraRigFunctions.cpp
@@ -18,22 +18,22 @@
 
 TSharedPtr<UE::Cameras::FCameraEvaluationContext> UActivateCameraRigFunctions::GlobalContext;
 
-void UActivateCameraRigFunctions::ActivateBaseCameraRig(APlayerController* PlayerController, UCameraRigAsset* CameraRig)
+void UActivateCameraRigFunctions::ActivateBaseCameraRig(UObject* WorldContextObject, APlayerController* PlayerController, UCameraRigAsset* CameraRig)
 {
-	ActivateCameraRig(PlayerController, CameraRig, ECameraRigLayer::Base);
+	ActivateCameraRig(WorldContextObject, PlayerController, CameraRig, ECameraRigLayer::Base);
 }
 
-void UActivateCameraRigFunctions::ActivateGlobalCameraRig(APlayerController* PlayerController, UCameraRigAsset* CameraRig)
+void UActivateCameraRigFunctions::ActivateGlobalCameraRig(UObject* WorldContextObject, APlayerController* PlayerController, UCameraRigAsset* CameraRig)
 {
-	ActivateCameraRig(PlayerController, CameraRig, ECameraRigLayer::Global);
+	ActivateCameraRig(WorldContextObject, PlayerController, CameraRig, ECameraRigLayer::Global);
 }
 
-void UActivateCameraRigFunctions::ActivateVisualCameraRig(APlayerController* PlayerController, UCameraRigAsset* CameraRig)
+void UActivateCameraRigFunctions::ActivateVisualCameraRig(UObject* WorldContextObject, APlayerController* PlayerController, UCameraRigAsset* CameraRig)
 {
-	ActivateCameraRig(PlayerController, CameraRig, ECameraRigLayer::Visual);
+	ActivateCameraRig(WorldContextObject, PlayerController, CameraRig, ECameraRigLayer::Visual);
 }
 
-void UActivateCameraRigFunctions::ActivateCameraRig(APlayerController* PlayerController, UCameraRigAsset* CameraRig, ECameraRigLayer EvaluationLayer)
+void UActivateCameraRigFunctions::ActivateCameraRig(UObject* WorldContextObject, APlayerController* PlayerController, UCameraRigAsset* CameraRig, ECameraRigLayer EvaluationLayer)
 {
 	using namespace UE::Cameras;
 
@@ -47,7 +47,7 @@ void UActivateCameraRigFunctions::ActivateCameraRig(APlayerController* PlayerCon
 	{
 		FActivateCameraRigParams Params;
 		Params.CameraRig = CameraRig;
-		Params.EvaluationContext = EnsureGlobalContext();
+		Params.EvaluationContext = EnsureGlobalContext(WorldContextObject, PlayerController);
 		Params.Evaluator = SystemEvaluator;
 		Params.Layer = EvaluationLayer;
 		SystemEvaluator->GetRootNodeEvaluator()->ActivateCameraRig(Params);
@@ -67,13 +67,16 @@ UE::Cameras::FCameraSystemEvaluator* UActivateCameraRigFunctions::FindCameraSyst
 	return nullptr;
 }
 
-TSharedPtr<UE::Cameras::FCameraEvaluationContext> UActivateCameraRigFunctions::EnsureGlobalContext()
+TSharedPtr<UE::Cameras::FCameraEvaluationContext> UActivateCameraRigFunctions::EnsureGlobalContext(UObject* WorldContextObject, APlayerController* PlayerController)
 {
 	using namespace UE::Cameras;
 
 	if (!GlobalContext.IsValid())
 	{
-		GlobalContext = MakeShared<FCameraEvaluationContext>();
+		FCameraEvaluationContextInitializeParams InitParams;
+		InitParams.Owner = WorldContextObject;
+		InitParams.PlayerController = PlayerController;
+		GlobalContext = MakeShared<FCameraEvaluationContext>(InitParams);
 		GlobalContext->GetInitialResult().bIsValid = true;	
 	}
 	return GlobalContext;
diff --git a/Engine/Plugins/Cameras/GameplayCameras/Source/GameplayCameras/Public/GameFramework/ActivateCameraRigFunctions.h b/Engine/Plugins/Cameras/GameplayCameras/Source/GameplayCameras/Public/GameFramework/ActivateCameraRigFunctions.h
--- a/Engine/Plugins/Cameras/GameplayCameras/Source/GameplayCameras/Public/GameFramework/ActivateCameraRigFunctions.h
+++ b/Engine/Plugins/Cameras/GameplayCameras/Source/GameplayCameras/Public/GameFramework/ActivateCameraRigFunctions.h
@@ -36,25 +36,25 @@ class UActivateCameraRigFunctions : public UBlueprintFunctionLibrary
 public:
 
 	/** Activates the given camera rig in the base layer. */
-	UFUNCTION(BlueprintCallable, Category="Camera")
-	static void ActivateBaseCameraRig(APlayerController* PlayerController, UCameraRigAsset* CameraRig);
+	UFUNCTION(BlueprintCallable, Category="Camera", meta=(WorldContext="WorldContextObject"))
+	static void ActivateBaseCameraRig(UObject* WorldContextObject, APlayerController* PlayerController, UCameraRigAsset* CameraRig);
 
 	/** Activates the given camera rig in the global layer. */
-	UFUNCTION(BlueprintCallable, Category="Camera")
-	static void ActivateGlobalCameraRig(APlayerController* PlayerController, UCameraRigAsset* CameraRig);
+	UFUNCTION(BlueprintCallable, Category="Camera", meta = (WorldContext = "WorldContextObject"))
+	static void ActivateGlobalCameraRig(UObject* WorldContextObject, APlayerController* PlayerController, UCameraRigAsset* CameraRig);
 
 	/** Activates the given camera rig in the visual layer. */
-	UFUNCTION(BlueprintCallable, Category="Camera")
-	static void ActivateVisualCameraRig(APlayerController* PlayerController, UCameraRigAsset* CameraRig);
+	UFUNCTION(BlueprintCallable, Category="Camera", meta = (WorldContext = "WorldContextObject"))
+	static void ActivateVisualCameraRig(UObject* WorldContextObject, APlayerController* PlayerController, UCameraRigAsset* CameraRig);
 
 	/** Activates the given camera rig in the given layer. */
-	UFUNCTION(BlueprintCallable, Category="Camera")
-	static void ActivateCameraRig(APlayerController* PlayerController, UCameraRigAsset* CameraRig, ECameraRigLayer EvaluationLayer);
+	UFUNCTION(BlueprintCallable, Category="Camera", meta = (WorldContext = "WorldContextObject"))
+	static void ActivateCameraRig(UObject* WorldContextObject, APlayerController* PlayerController, UCameraRigAsset* CameraRig, ECameraRigLayer EvaluationLayer);
 
 private:
 
 	static UE::Cameras::FCameraSystemEvaluator* FindCameraSystemEvaluator(APlayerController* PlayerController);
-	static TSharedPtr<UE::Cameras::FCameraEvaluationContext> EnsureGlobalContext();
+	static TSharedPtr<UE::Cameras::FCameraEvaluationContext> EnsureGlobalContext(UObject* WorldContextObject, APlayerController* PlayerController);
 
 private:
 
