Cameras: Initialize AspectRatioAxisConstraint to AspectRatio_MaintainYFOV

#rnx

diff --git a/Engine/Plugins/Cameras/GameplayCameras/Source/GameplayCameras/Public/Core/CameraPose.h b/Engine/Plugins/Cameras/GameplayCameras/Source/GameplayCameras/Public/Core/CameraPose.h
--- a/Engine/Plugins/Cameras/GameplayCameras/Source/GameplayCameras/Public/Core/CameraPose.h
+++ b/Engine/Plugins/Cameras/GameplayCameras/Source/GameplayCameras/Public/Core/CameraPose.h
@@ -3,6 +3,7 @@
 #pragma once
 
 #include "CoreTypes.h"
+#include "Engine/EngineTypes.h"
 #include "Math/MathFwd.h"
 #include "Math/Transform.h"
 
@@ -243,7 +244,7 @@ private:
 
 	/** If ConstrainAspectRatio is false, how we should compute FieldOfView */
 	UPROPERTY()
-	TEnumAsByte<EAspectRatioAxisConstraint> AspectRatioAxisConstraint;
+	TEnumAsByte<EAspectRatioAxisConstraint> AspectRatioAxisConstraint = EAspectRatioAxisConstraint::AspectRatio_MaintainYFOV;
 
 private:
 	
