Sequencer: Fix deprecated include warning

#rnx

diff --git a/Engine/Plugins/Cameras/GameplayCameras/Source/GameplayCameras/Legacy/Camera/CameraShake.h b/Engine/Plugins/Cameras/GameplayCameras/Source/GameplayCameras/Legacy/Camera/CameraShake.h
--- a/Engine/Plugins/Cameras/GameplayCameras/Source/GameplayCameras/Legacy/Camera/CameraShake.h
+++ b/Engine/Plugins/Cameras/GameplayCameras/Source/GameplayCameras/Legacy/Camera/CameraShake.h
@@ -6,6 +6,6 @@
 
 #include "CoreTypes.h"
 
-COMPILE_WARNING("UCameraShake has been renamed to ULegacyCameraShake, change to #include \"LegacyCameraShake.h\" and rename references")
-#include "LegacyCameraShake.h"
+UE_DEPRECATED_HEADER(5.5, "Use LegacyCameraShake.h instead of CameraShake.h")
 
+#include "LegacyCameraShake.h"
diff --git a/Engine/Plugins/Cameras/GameplayCameras/Source/GameplayCameras/Public/Shakes/MatineeCameraShake.h b/Engine/Plugins/Cameras/GameplayCameras/Source/GameplayCameras/Public/Shakes/MatineeCameraShake.h
--- a/Engine/Plugins/Cameras/GameplayCameras/Source/GameplayCameras/Public/Shakes/MatineeCameraShake.h
+++ b/Engine/Plugins/Cameras/GameplayCameras/Source/GameplayCameras/Public/Shakes/MatineeCameraShake.h
@@ -6,5 +6,6 @@
 
 #include "CoreTypes.h"
 
-COMPILE_WARNING("UMatineeCameraShake has been renamed to ULegacyCameraShake, change to #include \"LegacyCameraShake.h\" and rename references")
+UE_DEPRECATED_HEADER(5.5, "Use LegacyCameraShake.h instead of MatineeCameraShake.h")
+
 #include "LegacyCameraShake.h"
