Cameras: remove node graphs from the GC root set when the editor is closed

Otherwise, they (and their nodes and underlying objects) never get GC'ed.

diff --git a/Engine/Plugins/Cameras/GameplayCameras/Source/GameplayCamerasEditor/Private/Editors/SCameraRigAssetEditor.cpp b/Engine/Plugins/Cameras/GameplayCameras/Source/GameplayCamerasEditor/Private/Editors/SCameraRigAssetEditor.cpp
--- a/Engine/Plugins/Cameras/GameplayCameras/Source/GameplayCamerasEditor/Private/Editors/SCameraRigAssetEditor.cpp
+++ b/Engine/Plugins/Cameras/GameplayCameras/Source/GameplayCamerasEditor/Private/Editors/SCameraRigAssetEditor.cpp
@@ -41,6 +41,21 @@ void SCameraRigAssetEditor::Construct(const FArguments& InArgs)
 	];
 }
 
+SCameraRigAssetEditor::~SCameraRigAssetEditor()
+{
+	if (!GExitPurge)
+	{
+		if (NodeGraph)
+		{
+			NodeGraph->RemoveFromRoot();
+		}
+		if (TransitionGraph)
+		{
+			TransitionGraph->RemoveFromRoot();
+		}
+	}
+}
+
 void SCameraRigAssetEditor::CreateNodeGraphEditor(const FArguments& InArgs)
 {
 	const UGameplayCamerasEditorSettings* Settings = GetDefault<UGameplayCamerasEditorSettings>();
diff --git a/Engine/Plugins/Cameras/GameplayCameras/Source/GameplayCamerasEditor/Public/Editors/SCameraRigAssetEditor.h b/Engine/Plugins/Cameras/GameplayCameras/Source/GameplayCamerasEditor/Public/Editors/SCameraRigAssetEditor.h
--- a/Engine/Plugins/Cameras/GameplayCameras/Source/GameplayCamerasEditor/Public/Editors/SCameraRigAssetEditor.h
+++ b/Engine/Plugins/Cameras/GameplayCameras/Source/GameplayCamerasEditor/Public/Editors/SCameraRigAssetEditor.h
@@ -50,6 +50,7 @@ public:
 	SLATE_END_ARGS()
 
 	void Construct(const FArguments& InArgs);
+	~SCameraRigAssetEditor();
 
 public:
 
