Cameras: add support for ObjectTreeGraphHidden metadata on array properties

diff --git a/Engine/Plugins/Cameras/GameplayCameras/Source/GameplayCamerasEditor/Private/Editors/ObjectTreeGraphConfig.cpp b/Engine/Plugins/Cameras/GameplayCameras/Source/GameplayCamerasEditor/Private/Editors/ObjectTreeGraphConfig.cpp
--- a/Engine/Plugins/Cameras/GameplayCameras/Source/GameplayCamerasEditor/Private/Editors/ObjectTreeGraphConfig.cpp
+++ b/Engine/Plugins/Cameras/GameplayCameras/Source/GameplayCamerasEditor/Private/Editors/ObjectTreeGraphConfig.cpp
@@ -72,13 +72,43 @@ bool FObjectTreeGraphConfig::IsConnectable(FObjectProperty* InObjectProperty) co
 		return false;
 	}
 
+	if (InObjectProperty->GetBoolMetaData(TEXT("ObjectTreeGraphHidden")))
+	{
+		return false;
+	}
+
 	if (!IsConnectable(InObjectProperty->PropertyClass))
 	{
 		return false;
 	}
 
-	const bool bIsHidden = InObjectProperty->GetBoolMetaData(TEXT("ObjectTreeGraphHidden"));
-	return !bIsHidden;
+	return true;
+}
+
+bool FObjectTreeGraphConfig::IsConnectable(FArrayProperty* InArrayProperty) const
+{
+	if (!ensure(InArrayProperty))
+	{
+		return false;
+	}
+
+	if (InArrayProperty->GetBoolMetaData(TEXT("ObjectTreeGraphHidden")))
+	{
+		return false;
+	}
+
+	FObjectProperty* InnerProperty = CastField<FObjectProperty>(InArrayProperty->Inner);
+	if (!InnerProperty)
+	{
+		return false;
+	}
+
+	if (!IsConnectable(InnerProperty->PropertyClass))
+	{
+		return false;
+	}
+
+	return true;
 }
 
 void FObjectTreeGraphConfig::GetConnectableClasses(TArray<UClass*>& OutClasses, bool bPlaceableOnly)
diff --git a/Engine/Plugins/Cameras/GameplayCameras/Source/GameplayCamerasEditor/Private/Editors/ObjectTreeGraphNode.cpp b/Engine/Plugins/Cameras/GameplayCameras/Source/GameplayCamerasEditor/Private/Editors/ObjectTreeGraphNode.cpp
--- a/Engine/Plugins/Cameras/GameplayCameras/Source/GameplayCamerasEditor/Private/Editors/ObjectTreeGraphNode.cpp
+++ b/Engine/Plugins/Cameras/GameplayCameras/Source/GameplayCamerasEditor/Private/Editors/ObjectTreeGraphNode.cpp
@@ -179,12 +179,13 @@ void UObjectTreeGraphNode::AllocateDefaultPins()
 		}
 		else if (FArrayProperty* ArrayProperty = CastField<FArrayProperty>(*PropertyIt))
 		{
-			FObjectProperty* InnerProperty = CastField<FObjectProperty>(ArrayProperty->Inner);
-			if (!InnerProperty || !InnerProperty->PropertyClass || !OuterGraphConfig.IsConnectable(InnerProperty->PropertyClass))
+			if (!OuterGraphConfig.IsConnectable(ArrayProperty))
 			{
 				continue;
 			}
 
+			FObjectProperty* InnerProperty = CastFieldChecked<FObjectProperty>(ArrayProperty->Inner);
+
 			ChildPinType.PinSubCategory = UObjectTreeGraphSchema::PSC_ArrayProperty;
 			ChildPinType.ContainerType = EPinContainerType::Array;
 			UEdGraphPin* ArrayPin = CreatePin(PinDirection, ChildPinType, PropertyName);
@@ -631,7 +632,7 @@ void UObjectTreeGraphNode::OnGraphNodeMoved()
 
 UObjectTreeGraphNode::FNodeContext UObjectTreeGraphNode::GetNodeContext() const
 {
-	UObjectTreeGraph* OuterGraph = GetTypedOuter<UObjectTreeGraph>();
+	UObjectTreeGraph* OuterGraph = CastChecked<UObjectTreeGraph>(GetGraph());
 	const FObjectTreeGraphConfig& OuterGraphConfig = OuterGraph->GetConfig();
 
 	UClass* ObjectClass = Object->GetClass();
diff --git a/Engine/Plugins/Cameras/GameplayCameras/Source/GameplayCamerasEditor/Public/Editors/ObjectTreeGraphConfig.h b/Engine/Plugins/Cameras/GameplayCameras/Source/GameplayCamerasEditor/Public/Editors/ObjectTreeGraphConfig.h
--- a/Engine/Plugins/Cameras/GameplayCameras/Source/GameplayCamerasEditor/Public/Editors/ObjectTreeGraphConfig.h
+++ b/Engine/Plugins/Cameras/GameplayCameras/Source/GameplayCamerasEditor/Public/Editors/ObjectTreeGraphConfig.h
@@ -182,6 +182,14 @@ public:
 	 */
 	bool IsConnectable(FObjectProperty* InObjectProperty) const;
 
+	/**
+	 * Returns whether the given object array property is connectable.
+	 *
+	 * It is connectable if the array's item type is for a connectable class, and if the array
+	 * property doesn't have the ObjectTreeGraphHidden metadata.
+	 */
+	bool IsConnectable(FArrayProperty* InArrayProperty) const;
+
 	/**
 	 * Gets all possible known connectable classes.
 	 *
