Cameras: better manage loose nodes in the camera rig graph editor

Root objects like camera rigs can now hold (for editor builds) a list of all their children objects, including the "loose" ones that would otherwise not be connected to them. A new object interface provides support for that, and the object-tree graph editor uses it to add/remove objects when corresponding nodes are added/removed.

This fixes the problem of loose camera nodes disappearing from the graph, and the problem of deleted camera nodes reappearing in it. In one case it was because these nodes were GC'ed eventually, and in the other case it was because the nodes were kept alive by undo transactions.

diff --git a/Engine/Plugins/Cameras/GameplayCameras/Source/GameplayCameras/Private/Core/CameraRigAsset.cpp b/Engine/Plugins/Cameras/GameplayCameras/Source/GameplayCameras/Private/Core/CameraRigAsset.cpp
--- a/Engine/Plugins/Cameras/GameplayCameras/Source/GameplayCameras/Private/Core/CameraRigAsset.cpp
+++ b/Engine/Plugins/Cameras/GameplayCameras/Source/GameplayCameras/Private/Core/CameraRigAsset.cpp
@@ -204,6 +204,28 @@ UE_CAMERA_VARIABLE_FOR_ALL_TYPES()
 	}
 }
 
+void UCameraRigAsset::PostLoad()
+{
+#if WITH_EDITORONLY_DATA
+	if (AllNodes.IsEmpty())
+	{
+		// Rebuild the AllNodes array from scratch if this is old data.
+		TArray<UObject*> AllObjects;
+		UPackage* Package = GetOutermost();
+		GetObjectsWithPackage(Package, AllObjects);
+		for (UObject* Obj : AllObjects)
+		{
+			if (UCameraNode* CameraNode = Cast<UCameraNode>(Obj))
+			{
+				AllNodes.Add(CameraNode);
+			}
+		}
+	}
+#endif  // WITH_EDITORONLY_DATA
+
+	Super::PostLoad();
+}
+
 void UCameraRigAsset::PreSave(FObjectPreSaveContext ObjectSaveContext)
 {
 #if WITH_EDITOR
@@ -262,4 +284,26 @@ void UCameraRigAsset::OnUpdateGraphNodeCommentText(const FString& NewComment)
 	GraphNodeComment = NewComment;
 }
 
+void UCameraRigAsset::AddConnectableObject(UObject* InObject)
+{
+	// Camera nodes and rig parameters are shown on the graph editor. Handle the former case. The latter are
+	// already added to the interface struct even when not connected.
+	if (UCameraNode* CameraNode = Cast<UCameraNode>(InObject))
+	{
+		Modify();
+		const int32 Index = AllNodes.AddUnique(CameraNode);
+		ensure(Index == AllNodes.Num() - 1);
+	}
+}
+
+void UCameraRigAsset::RemoveConnectableObject(UObject* InObject)
+{
+	if (UCameraNode* CameraNode = Cast<UCameraNode>(InObject))
+	{
+		Modify();
+		const int32 NumRemoved = AllNodes.Remove(CameraNode);
+		ensure(NumRemoved == 1);
+	}
+}
+
 #endif  // WITH_EDITOR
diff --git a/Engine/Plugins/Cameras/GameplayCameras/Source/GameplayCameras/Public/Core/CameraRigAsset.h b/Engine/Plugins/Cameras/GameplayCameras/Source/GameplayCameras/Public/Core/CameraRigAsset.h
--- a/Engine/Plugins/Cameras/GameplayCameras/Source/GameplayCameras/Public/Core/CameraRigAsset.h
+++ b/Engine/Plugins/Cameras/GameplayCameras/Source/GameplayCameras/Public/Core/CameraRigAsset.h
@@ -6,6 +6,7 @@
 #include "Core/CameraRigTransition.h"
 #include "Core/CameraVariableTableFwd.h"
 #include "Core/ObjectTreeGraphObject.h"
+#include "Core/ObjectTreeGraphRootObject.h"
 #include "CoreTypes.h"
 #include "GameplayTagAssetInterface.h"
 #include "GameplayTagContainer.h"
@@ -139,6 +140,7 @@ class UCameraRigAsset
 	: public UObject
 	, public IGameplayTagAssetInterface
 	, public IObjectTreeGraphObject
+	, public IObjectTreeGraphRootObject
 {
 	GENERATED_BODY()
 
@@ -203,22 +205,38 @@ protected:
 	virtual void OnUpdateGraphNodeCommentText(const FString& NewComment) override;
 #endif
 
+#if WITH_EDITOR
+	virtual void AddConnectableObject(UObject* InObject) override;
+	virtual void RemoveConnectableObject(UObject* InObject) override;
+#endif
+
 	// UObject interface
+	virtual void PostLoad() override;
 	virtual void PreSave(FObjectPreSaveContext ObjectSaveContext) override;
 
 private:
 
 #if WITH_EDITORONLY_DATA
 
+	/** Position of the camera rig node in the graph editor. */
 	UPROPERTY()
 	int32 GraphNodePosX = 0;
 
+	/** Position of the camera rig node in the graph editor. */
 	UPROPERTY()
 	int32 GraphNodePosY = 0;
 
+	/** User-written comment in the graph editor. */
 	UPROPERTY()
 	FString GraphNodeComment;
 
+	/** 
+	 * A list of all the camera nodes, including the 'loose' ones that aren't connected
+	 * to the root node, and therefore would be GC'ed if we didn't hold them here.
+	 */
+	UPROPERTY(meta=(ObjectTreeGraphHidden=true))
+	TArray<TObjectPtr<UCameraNode>> AllNodes;
+
 #endif  // WITH_EDITORONLY_DATA
 };
 
diff --git a/Engine/Plugins/Cameras/GameplayCameras/Source/GameplayCameras/Public/Core/ObjectTreeGraphRootObject.h b/Engine/Plugins/Cameras/GameplayCameras/Source/GameplayCameras/Public/Core/ObjectTreeGraphRootObject.h
new file mode 100644
--- /dev/null
+++ b/Engine/Plugins/Cameras/GameplayCameras/Source/GameplayCameras/Public/Core/ObjectTreeGraphRootObject.h
@@ -0,0 +1,31 @@
+// Copyright Epic Games, Inc. All Rights Reserved.
+
+#pragma once
+
+#include "CoreMinimal.h"
+#include "UObject/Interface.h"
+
+#include "ObjectTreeGraphRootObject.generated.h"
+
+UINTERFACE(MinimalAPI)
+class UObjectTreeGraphRootObject : public UInterface
+{
+	GENERATED_BODY()
+};
+
+class IObjectTreeGraphRootObject
+{
+	GENERATED_BODY()
+
+public:
+
+#if WITH_EDITOR
+
+	virtual void AddConnectableObject(UObject* InObject) {}
+	virtual void RemoveConnectableObject(UObject* InObject) {}
+	virtual void GetExtraConnectableObjects(TArray<UObject*>& OutObjects) const {}
+
+#endif  // WITH_EDITOR
+
+};
+
diff --git a/Engine/Plugins/Cameras/GameplayCameras/Source/GameplayCamerasEditor/Private/Editors/CameraNodeGraphSchema.cpp b/Engine/Plugins/Cameras/GameplayCameras/Source/GameplayCamerasEditor/Private/Editors/CameraNodeGraphSchema.cpp
--- a/Engine/Plugins/Cameras/GameplayCameras/Source/GameplayCamerasEditor/Private/Editors/CameraNodeGraphSchema.cpp
+++ b/Engine/Plugins/Cameras/GameplayCameras/Source/GameplayCamerasEditor/Private/Editors/CameraNodeGraphSchema.cpp
@@ -6,6 +6,7 @@
 #include "Core/CameraNode.h"
 #include "EdGraph/EdGraphPin.h"
 #include "Editors/CameraNodeGraphNode.h"
+#include "Editors/CameraRigInterfaceParameterGraphNode.h"
 #include "Editors/ObjectTreeGraph.h"
 
 #include "ScopedTransaction.h"
@@ -223,6 +224,19 @@ bool UCameraNodeGraphSchema::OnBreakSinglePinLink(UEdGraphPin* SourcePin, UEdGra
 	return false;
 }
 
+void UCameraNodeGraphSchema::OnDeleteNodeFromGraph(UObjectTreeGraph* Graph, UEdGraphNode* Node) const
+{
+	Super::OnDeleteNodeFromGraph(Graph, Node);
+
+	UCameraRigAsset* CameraRig = CastChecked<UCameraRigAsset>(Graph->GetRootObject());
+	if (UCameraRigInterfaceParameterGraphNode* RigParameterNode = Cast<UCameraRigInterfaceParameterGraphNode>(Node))
+	{
+		UCameraRigInterfaceParameter* RigParameter = CastChecked<UCameraRigInterfaceParameter>(RigParameterNode->GetObject());
+		const int32 NumRemoved = CameraRig->Interface.InterfaceParameters.Remove(RigParameter);
+		ensure(NumRemoved == 1);
+	}
+}
+
 FCameraNodeGraphSchemaAction_NewInterfaceParameterNode::FCameraNodeGraphSchemaAction_NewInterfaceParameterNode()
 {
 }
diff --git a/Engine/Plugins/Cameras/GameplayCameras/Source/GameplayCamerasEditor/Private/Editors/ObjectTreeGraphSchema.cpp b/Engine/Plugins/Cameras/GameplayCameras/Source/GameplayCamerasEditor/Private/Editors/ObjectTreeGraphSchema.cpp
--- a/Engine/Plugins/Cameras/GameplayCameras/Source/GameplayCamerasEditor/Private/Editors/ObjectTreeGraphSchema.cpp
+++ b/Engine/Plugins/Cameras/GameplayCameras/Source/GameplayCamerasEditor/Private/Editors/ObjectTreeGraphSchema.cpp
@@ -2,10 +2,13 @@
 
 #include "Editors/ObjectTreeGraphSchema.h"
 
+#include "Core/ObjectTreeGraphRootObject.h"
 #include "Editors/ObjectTreeConnectionDrawingPolicy.h"
 #include "Editors/ObjectTreeGraph.h"
 #include "Editors/ObjectTreeGraphNode.h"
+#include "Serialization/ArchiveUObject.h"
 #include "ScopedTransaction.h"
+#include "UObject/FastReferenceCollector.h"
 #include "UObject/UObjectIterator.h"
 
 #include UE_INLINE_GENERATED_CPP_BY_NAME(ObjectTreeGraphSchema)
@@ -19,6 +22,67 @@ const FName UObjectTreeGraphSchema::PSC_ObjectProperty("ObjectProperty");
 const FName UObjectTreeGraphSchema::PSC_ArrayProperty("ArrayProperty");
 const FName UObjectTreeGraphSchema::PSC_ArrayPropertyItem("ArrayPropertyItem");
 
+namespace UE::ObjectTreeGraph
+{
+
+struct FPackageReferenceCollector : public FArchiveUObject
+{
+	FPackageReferenceCollector(UObject* InRootObject, TArray<UObject*>& InOutReferencedObjects)
+		: FArchiveUObject()
+		, RootObject(InRootObject)
+		, PackageScope(InRootObject->GetOutermost())
+		, ReferencedObjects(InOutReferencedObjects)
+	{
+		SetIsPersistent(true);
+		SetIsSaving(true);
+		SetFilterEditorOnly(false);
+
+		ArIsObjectReferenceCollector = true;
+		ArShouldSkipBulkData = true;
+	}
+
+	void CollectReferences()
+	{
+		ObjectsToVisit.Reset();
+		VisitedObjects.Reset();
+
+		ObjectsToVisit.Add(RootObject);
+		while (ObjectsToVisit.Num() > 0)
+		{
+			UObject* CurObj = ObjectsToVisit.Pop(EAllowShrinking::No);
+			VisitedObjects.Add(CurObj);
+			CurObj->Serialize(*this);
+		}
+	}
+
+private:
+
+	virtual FArchive& operator<<(UObject*& ObjRef) override
+	{
+		if (ObjRef != nullptr && ObjRef->IsIn(PackageScope))
+		{
+			if (!VisitedObjects.Contains(ObjRef))
+			{
+				ReferencedObjects.Add(ObjRef);
+				ObjectsToVisit.Add(ObjRef);
+			}
+		}
+		return *this;
+	}
+
+private:
+
+	UObject* RootObject;
+	UPackage* PackageScope;
+
+	TArray<UObject*> ObjectsToVisit;
+	TSet<UObject*> VisitedObjects;
+
+	TArray<UObject*>& ReferencedObjects;
+};
+
+}  // namespace UE::ObjectTreeGraph
+
 UObjectTreeGraphSchema::UObjectTreeGraphSchema(const FObjectInitializer& ObjInit)
 	: Super(ObjInit)
 {
@@ -49,11 +113,25 @@ void UObjectTreeGraphSchema::CreateAllNodes(UObjectTreeGraph* InGraph, EObjectTr
 	}
 
 	TArray<UObject*> AllObjects;
-	UPackage* Package = RootObject->GetOutermost();
-	GetObjectsWithPackage(Package, AllObjects);
+	// Gather up all the objects we need for the graph. Start by all objects that are referenced (directly or
+	// indirectly) by the root object. Our custom reference collector will not collect references that go outside
+	// of the root object's package.
+	{
+		using namespace UE::ObjectTreeGraph;
 
-	ensure(AllObjects.Contains(RootObject));
+		// Make sure the root object itself is in there.
+		AllObjects.Add(RootObject);
 
+		FPackageReferenceCollector Collector(RootObject, AllObjects);
+		Collector.CollectReferences();
+	}
+	// Add any other custom objects the root object may want.
+	if (IObjectTreeGraphRootObject* RootObjectInterface = Cast<IObjectTreeGraphRootObject>(RootObject))
+	{
+		RootObjectInterface->GetExtraConnectableObjects(AllObjects);
+	}
+	
+	// Create all the nodes.
 	FCreatedNodes CreatedNodes;
 	for (UObject* Object : AllObjects)
 	{
@@ -63,6 +141,7 @@ void UObjectTreeGraphSchema::CreateAllNodes(UObjectTreeGraph* InGraph, EObjectTr
 		}
 	}
 
+	// Grab the graph node for the root object.
 	InGraph->RootObjectNode = nullptr;
 	UObjectTreeGraphNode** CreatedRootObjectNode = CreatedNodes.CreatedNodes.Find(RootObject);
 	if (ensure(CreatedRootObjectNode))
@@ -70,6 +149,7 @@ void UObjectTreeGraphSchema::CreateAllNodes(UObjectTreeGraph* InGraph, EObjectTr
 		InGraph->RootObjectNode = *CreatedRootObjectNode;
 	}
 
+	// Create all the connections.
 	for (TPair<UObject*, UObjectTreeGraphNode*> Pair : CreatedNodes.CreatedNodes)
 	{
 		CreateConnections(Pair.Value, CreatedNodes);
@@ -584,6 +664,34 @@ bool UObjectTreeGraphSchema::SupportsDropPinOnNode(UEdGraphNode* InTargetNode, c
 	return Super::SupportsDropPinOnNode(InTargetNode, InSourcePinType, InSourcePinDirection, OutErrorMessage);
 }
 
+bool UObjectTreeGraphSchema::SafeDeleteNodeFromGraph(UEdGraph* Graph, UEdGraphNode* Node) const
+{
+	if (!Graph || !Node)
+	{
+		return false;
+	}
+	
+	const FScopedTransaction Transaction(LOCTEXT("DeleteNode", "Delete Node"));
+
+	BreakNodeLinks(*Node);
+
+	UObjectTreeGraph* ObjectTreeGraph = CastChecked<UObjectTreeGraph>(Graph);
+	OnDeleteNodeFromGraph(ObjectTreeGraph, Node);
+	Node->DestroyNode();
+
+	return true;
+}
+
+void UObjectTreeGraphSchema::OnDeleteNodeFromGraph(UObjectTreeGraph* Graph, UEdGraphNode* Node) const
+{
+	UObjectTreeGraphNode* ObjectNode = Cast<UObjectTreeGraphNode>(Node);
+	IObjectTreeGraphRootObject* RootObjectInterface = Cast<IObjectTreeGraphRootObject>(Graph->GetRootObject());
+	if (RootObjectInterface && ObjectNode)
+	{
+		RootObjectInterface->RemoveConnectableObject(ObjectNode->GetObject());
+	}
+}
+
 void UObjectTreeGraphSchema::ProcessDuplicatedNodes(UObjectTreeGraph* InGraph, const TMap<UEdGraphNode*, UEdGraphNode*>& NodeMap) const
 {
 	UPackage* RootObjectPackage = InGraph->GetRootObject()->GetOutermost();
@@ -670,6 +778,7 @@ UEdGraphNode* FObjectGraphSchemaAction_NewNode::PerformAction(UEdGraph* ParentGr
 
 	const FScopedTransaction Transaction(LOCTEXT("CreateNewNodeAction", "Create New Node"));
 	const UObjectTreeGraphSchema* Schema = CastChecked<UObjectTreeGraphSchema>(ParentGraph->GetSchema());
+	IObjectTreeGraphRootObject* RootObjectInterface = Cast<IObjectTreeGraphRootObject>(ObjectTreeGraph->GetRootObject());
 
 	UObject* NewObject = CreateObject();
 
@@ -677,6 +786,11 @@ UEdGraphNode* FObjectGraphSchemaAction_NewNode::PerformAction(UEdGraph* ParentGr
 	{
 		UObjectTreeGraphNode* NewGraphNode = Schema->CreateObjectNode(ObjectTreeGraph, NewObject);
 
+		if (RootObjectInterface)
+		{
+			RootObjectInterface->AddConnectableObject(NewObject);
+		}
+
 		NewGraphNode->NodePosX = Location.X;
 		NewGraphNode->NodePosY = Location.Y;
 		NewGraphNode->OnGraphNodeMoved();
diff --git a/Engine/Plugins/Cameras/GameplayCameras/Source/GameplayCamerasEditor/Private/Editors/SObjectTreeGraphEditor.cpp b/Engine/Plugins/Cameras/GameplayCameras/Source/GameplayCamerasEditor/Private/Editors/SObjectTreeGraphEditor.cpp
--- a/Engine/Plugins/Cameras/GameplayCameras/Source/GameplayCamerasEditor/Private/Editors/SObjectTreeGraphEditor.cpp
+++ b/Engine/Plugins/Cameras/GameplayCameras/Source/GameplayCamerasEditor/Private/Editors/SObjectTreeGraphEditor.cpp
@@ -408,15 +408,14 @@ void SObjectTreeGraphEditor::ImportNodesFromText(const FVector2D& Location, cons
 
 void SObjectTreeGraphEditor::DeleteNodes(TArrayView<UObjectTreeGraphNode*> NodesToDelete)
 {
-	const UEdGraphSchema* Schema = GraphEditor->GetCurrentGraph()->GetSchema();
+	UEdGraph* CurrentGraph = GraphEditor->GetCurrentGraph();
+	const UEdGraphSchema* Schema = CurrentGraph->GetSchema();
 
 	for (UObjectTreeGraphNode* Node : NodesToDelete)
 	{
 		if (Node)
 		{
-			Node->Modify();
-
-			Schema->BreakNodeLinks(*Node);
+			Schema->SafeDeleteNodeFromGraph(CurrentGraph, Node);
 
 			Node->DestroyNode();
 		}
diff --git a/Engine/Plugins/Cameras/GameplayCameras/Source/GameplayCamerasEditor/Public/Editors/CameraNodeGraphSchema.h b/Engine/Plugins/Cameras/GameplayCameras/Source/GameplayCamerasEditor/Public/Editors/CameraNodeGraphSchema.h
--- a/Engine/Plugins/Cameras/GameplayCameras/Source/GameplayCamerasEditor/Public/Editors/CameraNodeGraphSchema.h
+++ b/Engine/Plugins/Cameras/GameplayCameras/Source/GameplayCamerasEditor/Public/Editors/CameraNodeGraphSchema.h
@@ -31,6 +31,7 @@ protected:
 	virtual bool OnCreateConnection(UEdGraphPin* A, UEdGraphPin* B) const override;
 	virtual bool OnBreakPinLinks(UEdGraphPin& TargetPin, bool bSendsNodeNotification) const override;
 	virtual bool OnBreakSinglePinLink(UEdGraphPin* SourcePin, UEdGraphPin* TargetPin) const override;
+	virtual void OnDeleteNodeFromGraph(UObjectTreeGraph* Graph, UEdGraphNode* Node) const override;
 };
 
 /**
diff --git a/Engine/Plugins/Cameras/GameplayCameras/Source/GameplayCamerasEditor/Public/Editors/ObjectTreeGraphSchema.h b/Engine/Plugins/Cameras/GameplayCameras/Source/GameplayCamerasEditor/Public/Editors/ObjectTreeGraphSchema.h
--- a/Engine/Plugins/Cameras/GameplayCameras/Source/GameplayCamerasEditor/Public/Editors/ObjectTreeGraphSchema.h
+++ b/Engine/Plugins/Cameras/GameplayCameras/Source/GameplayCamerasEditor/Public/Editors/ObjectTreeGraphSchema.h
@@ -63,11 +63,13 @@ public:
 	virtual void BreakPinLinks(UEdGraphPin& TargetPin, bool bSendsNodeNotification) const override;
 	virtual void BreakSinglePinLink(UEdGraphPin* SourcePin, UEdGraphPin* TargetPin) const override;
 	virtual bool SupportsDropPinOnNode(UEdGraphNode* InTargetNode, const FEdGraphPinType& InSourcePinType, EEdGraphPinDirection InSourcePinDirection, FText& OutErrorMessage) const override;
+	virtual bool SafeDeleteNodeFromGraph(UEdGraph* Graph, UEdGraphNode* Node) const override;
 
 	// UObjectTreeGraphSchema interface.
 	virtual bool OnCreateConnection(UEdGraphPin* A, UEdGraphPin* B) const;
 	virtual bool OnBreakPinLinks(UEdGraphPin& TargetPin, bool bSendsNodeNotification) const;
 	virtual bool OnBreakSinglePinLink(UEdGraphPin* SourcePin, UEdGraphPin* TargetPin) const;
+	virtual void OnDeleteNodeFromGraph(UObjectTreeGraph* Graph, UEdGraphNode* Node) const;
 	virtual void FilterGraphContextPlaceableClasses(TArray<UClass*>& InOutClasses) const;
 	virtual void ProcessDuplicatedNodes(UObjectTreeGraph* InGraph, const TMap<UEdGraphNode*, UEdGraphNode*>& NodeMap) const;
 
