Cameras: add better categorization of nodes in the object-tree graph

diff --git a/Engine/Plugins/Cameras/GameplayCameras/Source/GameplayCamerasEditor/Private/Editors/ObjectTreeGraphConfig.cpp b/Engine/Plugins/Cameras/GameplayCameras/Source/GameplayCamerasEditor/Private/Editors/ObjectTreeGraphConfig.cpp
--- a/Engine/Plugins/Cameras/GameplayCameras/Source/GameplayCamerasEditor/Private/Editors/ObjectTreeGraphConfig.cpp
+++ b/Engine/Plugins/Cameras/GameplayCameras/Source/GameplayCamerasEditor/Private/Editors/ObjectTreeGraphConfig.cpp
@@ -20,6 +20,7 @@ FObjectTreeGraphClassConfig::FObjectTreeGraphClassConfig()
 	, _CanCreateNew(true)
 	, _CanDuplicate(true)
 	, _CanDelete(true)
+	, _CreateCategoryMetaData(TEXT("ObjectTreeGraphCategory"))
 {
 }
 
diff --git a/Engine/Plugins/Cameras/GameplayCameras/Source/GameplayCamerasEditor/Private/Editors/ObjectTreeGraphSchema.cpp b/Engine/Plugins/Cameras/GameplayCameras/Source/GameplayCamerasEditor/Private/Editors/ObjectTreeGraphSchema.cpp
--- a/Engine/Plugins/Cameras/GameplayCameras/Source/GameplayCamerasEditor/Private/Editors/ObjectTreeGraphSchema.cpp
+++ b/Engine/Plugins/Cameras/GameplayCameras/Source/GameplayCamerasEditor/Private/Editors/ObjectTreeGraphSchema.cpp
@@ -340,6 +340,8 @@ void UObjectTreeGraphSchema::GetGraphContextActions(FGraphContextMenuBuilder& Co
 
 	FilterGraphContextPlaceableClasses(PossibleObjectClasses);
 
+	const FText MiscellaneousCategoryText = LOCTEXT("MiscellaneousCategory", "Miscellaneous");
+
 	for (UClass* PossibleObjectClass : PossibleObjectClasses)
 	{
 		if (!PossibleObjectClass)
@@ -347,23 +349,46 @@ void UObjectTreeGraphSchema::GetGraphContextActions(FGraphContextMenuBuilder& Co
 			continue;
 		}
 
-		const FString* CategoryName = nullptr;
+		const FText DisplayName = GraphConfig.GetDisplayNameText(PossibleObjectClass);
+
+		TArray<FString> CategoryNames;
+		const FName CreateCategoryMetaData = GraphConfig.GetObjectClassConfig(PossibleObjectClass).CreateCategoryMetaData();
 		for (UClass* CurClass = PossibleObjectClass; CurClass; CurClass = CurClass->GetSuperClass())
 		{
-			CategoryName = CurClass->FindMetaData(TEXT("ObjectTreeGraphCategory"));
-			if (CategoryName)
+			const FString* CategoryNamesMetaData = CurClass->FindMetaData(CreateCategoryMetaData);
+			if (CategoryNamesMetaData)
 			{
+				CategoryNamesMetaData->ParseIntoArray(CategoryNames, TEXT(","), true);
 				break;
 			}
 		}
+		if (CategoryNames.IsEmpty())
+		{
+			CategoryNames.Add(FString());
+		}
+
+		FFormatNamedArguments Arguments;
+		Arguments.Add(TEXT("Name"), DisplayName);
+		const FText ToolTipText = FText::Format(LOCTEXT("NewNodeToolTip", "Adds a {Name} node here"), Arguments);
 
 		checkSlow(PossibleObjectClass);
-		TSharedRef<FObjectGraphSchemaAction_NewNode> Action = MakeShared<FObjectGraphSchemaAction_NewNode>(
-				CategoryName ? FText::FromString(*CategoryName) : FText::GetEmpty(),
-				PossibleObjectClass->GetDisplayNameText(), 
-				PossibleObjectClass->GetDisplayNameText());
-		Action->ObjectClass = PossibleObjectClass;
-		ContextMenuBuilder.AddAction(StaticCastSharedPtr<FEdGraphSchemaAction>(Action.ToSharedPtr()));
+		for (const FString& CategoryName : CategoryNames)
+		{
+			FText CategoryText = MiscellaneousCategoryText;
+			int32 Grouping = -1;
+			if (!CategoryName.IsEmpty())
+			{
+				CategoryText = FText::FromString(CategoryName);
+				Grouping = CategoryName == TEXT("Common") ? 1 : 0;
+			}
+
+			FText KeywordsText(FText::FromString(PossibleObjectClass->GetMetaData(TEXT("Keywords"))));
+
+			TSharedRef<FObjectGraphSchemaAction_NewNode> Action = MakeShared<FObjectGraphSchemaAction_NewNode>(
+					CategoryText, DisplayName, ToolTipText, Grouping, KeywordsText);
+			Action->ObjectClass = PossibleObjectClass;
+			ContextMenuBuilder.AddAction(StaticCastSharedPtr<FEdGraphSchemaAction>(Action.ToSharedPtr()));
+		}
 	}
 
 	// Don't call the base class, we want to control exactly what can be created.
diff --git a/Engine/Plugins/Cameras/GameplayCameras/Source/GameplayCamerasEditor/Public/Editors/ObjectTreeGraphConfig.h b/Engine/Plugins/Cameras/GameplayCameras/Source/GameplayCamerasEditor/Public/Editors/ObjectTreeGraphConfig.h
--- a/Engine/Plugins/Cameras/GameplayCameras/Source/GameplayCamerasEditor/Public/Editors/ObjectTreeGraphConfig.h
+++ b/Engine/Plugins/Cameras/GameplayCameras/Source/GameplayCamerasEditor/Public/Editors/ObjectTreeGraphConfig.h
@@ -69,6 +69,9 @@ public:
 	/** Whether users can delete objects of this class in the graph. */
 	OTGCC_FIELD(bool, CanDelete)
 
+	/** The metadata specifier to look for in order to categorize the 'create node' action for this class. */
+	OTGCC_FIELD(FName, CreateCategoryMetaData)
+
 public:
 
 	FObjectTreeGraphClassConfig();
