Cameras: add proper categorization to camera nodes

diff --git a/Engine/Plugins/Cameras/GameplayCameras/Source/GameplayCameras/Public/Core/BlendCameraNode.h b/Engine/Plugins/Cameras/GameplayCameras/Source/GameplayCameras/Public/Core/BlendCameraNode.h
--- a/Engine/Plugins/Cameras/GameplayCameras/Source/GameplayCameras/Public/Core/BlendCameraNode.h
+++ b/Engine/Plugins/Cameras/GameplayCameras/Source/GameplayCameras/Public/Core/BlendCameraNode.h
@@ -10,7 +10,7 @@
 /**
  * Base class for blend camera nodes.
  */
-UCLASS(MinimalAPI, Hidden, meta=(ObjectTreeGraphCategory="Camera Blends"))
+UCLASS(MinimalAPI, Hidden, meta=(CameraNodeCategories="Blends"))
 class UBlendCameraNode : public UCameraNode
 {
 	GENERATED_BODY()
diff --git a/Engine/Plugins/Cameras/GameplayCameras/Source/GameplayCameras/Public/Core/CameraNode.h b/Engine/Plugins/Cameras/GameplayCameras/Source/GameplayCameras/Public/Core/CameraNode.h
--- a/Engine/Plugins/Cameras/GameplayCameras/Source/GameplayCameras/Public/Core/CameraNode.h
+++ b/Engine/Plugins/Cameras/GameplayCameras/Source/GameplayCameras/Public/Core/CameraNode.h
@@ -19,7 +19,7 @@ using FCameraNodeChildrenView = UE::Cameras::TObjectChildrenView<TObjectPtr<UCam
 /**
  * The base class for a camera node.
  */
-UCLASS(Abstract, DefaultToInstanced, EditInlineNew, meta=(ObjectTreeGraphCategory="Camera Nodes"))
+UCLASS(Abstract, DefaultToInstanced, EditInlineNew, meta=(CameraNodeCategories="Miscellaneous"))
 class GAMEPLAYCAMERAS_API UCameraNode 
 	: public UObject
 	, public IObjectTreeGraphObject
diff --git a/Engine/Plugins/Cameras/GameplayCameras/Source/GameplayCameras/Public/Nodes/Attach/AttachToPlayerPawnCameraNode.h b/Engine/Plugins/Cameras/GameplayCameras/Source/GameplayCameras/Public/Nodes/Attach/AttachToPlayerPawnCameraNode.h
--- a/Engine/Plugins/Cameras/GameplayCameras/Source/GameplayCameras/Public/Nodes/Attach/AttachToPlayerPawnCameraNode.h
+++ b/Engine/Plugins/Cameras/GameplayCameras/Source/GameplayCameras/Public/Nodes/Attach/AttachToPlayerPawnCameraNode.h
@@ -10,7 +10,7 @@
 /**
  * A camera node that moves the camera to the player pawn.
  */
-UCLASS(MinimalAPI)
+UCLASS(MinimalAPI, meta=(CameraNodeCategories="Attachment"))
 class UAttachToPlayerPawnCameraNode : public UCameraNode
 {
 	GENERATED_BODY()
diff --git a/Engine/Plugins/Cameras/GameplayCameras/Source/GameplayCameras/Public/Nodes/Collision/OcclusionMaterialCameraNode.h b/Engine/Plugins/Cameras/GameplayCameras/Source/GameplayCameras/Public/Nodes/Collision/OcclusionMaterialCameraNode.h
--- a/Engine/Plugins/Cameras/GameplayCameras/Source/GameplayCameras/Public/Nodes/Collision/OcclusionMaterialCameraNode.h
+++ b/Engine/Plugins/Cameras/GameplayCameras/Source/GameplayCameras/Public/Nodes/Collision/OcclusionMaterialCameraNode.h
@@ -10,7 +10,7 @@
 
 class UMaterialInterface;
 
-UCLASS(MinimalAPI)
+UCLASS(MinimalAPI, meta=(CameraNodeCategories="Collision"))
 class UOcclusionMaterialCameraNode : public UCameraNode
 {
 	GENERATED_BODY()
diff --git a/Engine/Plugins/Cameras/GameplayCameras/Source/GameplayCameras/Public/Nodes/Common/ArrayCameraNode.h b/Engine/Plugins/Cameras/GameplayCameras/Source/GameplayCameras/Public/Nodes/Common/ArrayCameraNode.h
--- a/Engine/Plugins/Cameras/GameplayCameras/Source/GameplayCameras/Public/Nodes/Common/ArrayCameraNode.h
+++ b/Engine/Plugins/Cameras/GameplayCameras/Source/GameplayCameras/Public/Nodes/Common/ArrayCameraNode.h
@@ -10,7 +10,7 @@
 /**
  * A camera node that runs a list of other camera nodes.
  */
-UCLASS(MinimalAPI)
+UCLASS(MinimalAPI, meta=(CameraNodeCategories="Common,Utility"))
 class UArrayCameraNode : public UCameraNode
 {
 	GENERATED_BODY()
diff --git a/Engine/Plugins/Cameras/GameplayCameras/Source/GameplayCameras/Public/Nodes/Common/BoomArmCameraNode.h b/Engine/Plugins/Cameras/GameplayCameras/Source/GameplayCameras/Public/Nodes/Common/BoomArmCameraNode.h
--- a/Engine/Plugins/Cameras/GameplayCameras/Source/GameplayCameras/Public/Nodes/Common/BoomArmCameraNode.h
+++ b/Engine/Plugins/Cameras/GameplayCameras/Source/GameplayCameras/Public/Nodes/Common/BoomArmCameraNode.h
@@ -10,7 +10,7 @@
 /**
  * A camera node that can rotate the camera in yaw and pitch based on player input.
  */
-UCLASS(MinimalAPI)
+UCLASS(MinimalAPI, meta=(CameraNodeCategories="Common,Transform"))
 class UBoomArmCameraNode : public UCameraNode
 {
 	GENERATED_BODY()
diff --git a/Engine/Plugins/Cameras/GameplayCameras/Source/GameplayCameras/Public/Nodes/Common/DampenPositionCameraNode.h b/Engine/Plugins/Cameras/GameplayCameras/Source/GameplayCameras/Public/Nodes/Common/DampenPositionCameraNode.h
--- a/Engine/Plugins/Cameras/GameplayCameras/Source/GameplayCameras/Public/Nodes/Common/DampenPositionCameraNode.h
+++ b/Engine/Plugins/Cameras/GameplayCameras/Source/GameplayCameras/Public/Nodes/Common/DampenPositionCameraNode.h
@@ -10,7 +10,7 @@
 /**
  * A camera node that offsets the location of the camera.
  */
-UCLASS(MinimalAPI)
+UCLASS(MinimalAPI, meta=(CameraNodeCategories="Common,Transform"))
 class UDampenPositionCameraNode : public UCameraNode
 {
 	GENERATED_BODY()
diff --git a/Engine/Plugins/Cameras/GameplayCameras/Source/GameplayCameras/Public/Nodes/Common/FieldOfViewCameraNode.h b/Engine/Plugins/Cameras/GameplayCameras/Source/GameplayCameras/Public/Nodes/Common/FieldOfViewCameraNode.h
--- a/Engine/Plugins/Cameras/GameplayCameras/Source/GameplayCameras/Public/Nodes/Common/FieldOfViewCameraNode.h
+++ b/Engine/Plugins/Cameras/GameplayCameras/Source/GameplayCameras/Public/Nodes/Common/FieldOfViewCameraNode.h
@@ -11,7 +11,7 @@
 /**
  * A camera node that sets the field of view of the camera.
  */
-UCLASS(MinimalAPI)
+UCLASS(MinimalAPI, meta=(CameraNodeCategories="Common,Lens"))
 class UFieldOfViewCameraNode : public UCameraNode
 {
 	GENERATED_BODY()
diff --git a/Engine/Plugins/Cameras/GameplayCameras/Source/GameplayCameras/Public/Nodes/Common/LensParametersCameraNode.h b/Engine/Plugins/Cameras/GameplayCameras/Source/GameplayCameras/Public/Nodes/Common/LensParametersCameraNode.h
--- a/Engine/Plugins/Cameras/GameplayCameras/Source/GameplayCameras/Public/Nodes/Common/LensParametersCameraNode.h
+++ b/Engine/Plugins/Cameras/GameplayCameras/Source/GameplayCameras/Public/Nodes/Common/LensParametersCameraNode.h
@@ -11,7 +11,7 @@
 /**
  * A camera node that sets parameter values on the camera lens.
  */
-UCLASS(MinimalAPI)
+UCLASS(MinimalAPI, meta=(CameraNodeCategories="Common,Lens"))
 class ULensParametersCameraNode : public UCameraNode
 {
 	GENERATED_BODY()
diff --git a/Engine/Plugins/Cameras/GameplayCameras/Source/GameplayCameras/Public/Nodes/Common/OffsetCameraNode.h b/Engine/Plugins/Cameras/GameplayCameras/Source/GameplayCameras/Public/Nodes/Common/OffsetCameraNode.h
--- a/Engine/Plugins/Cameras/GameplayCameras/Source/GameplayCameras/Public/Nodes/Common/OffsetCameraNode.h
+++ b/Engine/Plugins/Cameras/GameplayCameras/Source/GameplayCameras/Public/Nodes/Common/OffsetCameraNode.h
@@ -11,7 +11,7 @@
 /**
  * A camera node that offsets the location of the camera.
  */
-UCLASS(MinimalAPI)
+UCLASS(MinimalAPI, meta=(CameraNodeCategories="Common,Transform"))
 class UOffsetCameraNode : public UCameraNode
 {
 	GENERATED_BODY()
diff --git a/Engine/Plugins/Cameras/GameplayCameras/Source/GameplayCameras/Public/Nodes/Framing/BaseFramingCameraNode.h b/Engine/Plugins/Cameras/GameplayCameras/Source/GameplayCameras/Public/Nodes/Framing/BaseFramingCameraNode.h
--- a/Engine/Plugins/Cameras/GameplayCameras/Source/GameplayCameras/Public/Nodes/Framing/BaseFramingCameraNode.h
+++ b/Engine/Plugins/Cameras/GameplayCameras/Source/GameplayCameras/Public/Nodes/Framing/BaseFramingCameraNode.h
@@ -16,7 +16,7 @@ class FArchive;
 /**
  * The base class for a standard scren-space framing camera node.
  */
-UCLASS(MinimalAPI, Abstract)
+UCLASS(MinimalAPI, Abstract, meta=(CameraNodeCategories="Framing"))
 class UBaseFramingCameraNode : public UCameraNode
 {
 	GENERATED_BODY()
diff --git a/Engine/Plugins/Cameras/GameplayCameras/Source/GameplayCamerasEditor/Private/Editors/SCameraRigAssetEditor.cpp b/Engine/Plugins/Cameras/GameplayCameras/Source/GameplayCamerasEditor/Private/Editors/SCameraRigAssetEditor.cpp
--- a/Engine/Plugins/Cameras/GameplayCameras/Source/GameplayCamerasEditor/Private/Editors/SCameraRigAssetEditor.cpp
+++ b/Engine/Plugins/Cameras/GameplayCameras/Source/GameplayCamerasEditor/Private/Editors/SCameraRigAssetEditor.cpp
@@ -72,6 +72,7 @@ void SCameraRigAssetEditor::CreateNodeGraphEditor(const FArguments& InArgs)
 		.NodeTitleColor(Settings->CameraRigAssetTitleColor);
 	GraphConfig.ObjectClassConfigs.Emplace(UCameraNode::StaticClass())
 		.StripDisplayNameSuffix(TEXT("Camera Node"))
+		.CreateCategoryMetaData(TEXT("CameraNodeCategories"))
 		.GraphNodeClass(UCameraNodeGraphNode::StaticClass());
 	GraphConfig.ObjectClassConfigs.Emplace(UCameraRigInterfaceParameter::StaticClass())
 		.SelfPinDirection(EGPD_Output)
@@ -120,7 +121,8 @@ void SCameraRigAssetEditor::CreateTransitionGraphEditor(const FArguments& InArgs
 		.StripDisplayNameSuffix(TEXT("Transition Condition"))
 		.NodeTitleColor(Settings->CameraRigTransitionConditionTitleColor);
 	GraphConfig.ObjectClassConfigs.Emplace(UBlendCameraNode::StaticClass())
-		.StripDisplayNameSuffix(TEXT("Camera Node"));
+		.StripDisplayNameSuffix(TEXT("Camera Node"))
+		.CreateCategoryMetaData(TEXT("CameraNodeCategories"));
 
 	TransitionGraph = NewObject<UObjectTreeGraph>(GetTransientPackage(), NAME_None, RF_Transactional);
 	TransitionGraph->Schema = UCameraTransitionGraphSchema::StaticClass();
