Cameras: don't let users create rig parameters from nothing

diff --git a/Engine/Plugins/Cameras/GameplayCameras/Source/GameplayCamerasEditor/Private/Editors/SCameraRigAssetEditor.cpp b/Engine/Plugins/Cameras/GameplayCameras/Source/GameplayCamerasEditor/Private/Editors/SCameraRigAssetEditor.cpp
--- a/Engine/Plugins/Cameras/GameplayCameras/Source/GameplayCamerasEditor/Private/Editors/SCameraRigAssetEditor.cpp
+++ b/Engine/Plugins/Cameras/GameplayCameras/Source/GameplayCamerasEditor/Private/Editors/SCameraRigAssetEditor.cpp
@@ -77,6 +77,8 @@ void SCameraRigAssetEditor::CreateNodeGraphEditor(const FArguments& InArgs)
 	GraphConfig.ObjectClassConfigs.Emplace(UCameraRigInterfaceParameter::StaticClass())
 		.SelfPinDirection(EGPD_Output)
 		.SelfPinName(NAME_None)  // No self pin name, we just want the title
+		.CanCreateNew(false)
+		.CanDuplicate(false)
 		.GraphNodeClass(UCameraRigInterfaceParameterGraphNode::StaticClass());
 
 	NodeGraph = NewObject<UObjectTreeGraph>(GetTransientPackage(), NAME_None, RF_Transactional);
