Cameras: allow creating both enter and exit transitions without dragging a pin

diff --git a/Engine/Plugins/Cameras/GameplayCameras/Source/GameplayCamerasEditor/Private/Editors/CameraTransitionGraphSchema.cpp b/Engine/Plugins/Cameras/GameplayCameras/Source/GameplayCamerasEditor/Private/Editors/CameraTransitionGraphSchema.cpp
--- a/Engine/Plugins/Cameras/GameplayCameras/Source/GameplayCamerasEditor/Private/Editors/CameraTransitionGraphSchema.cpp
+++ b/Engine/Plugins/Cameras/GameplayCameras/Source/GameplayCamerasEditor/Private/Editors/CameraTransitionGraphSchema.cpp
@@ -14,10 +14,13 @@
 
 void UCameraTransitionGraphSchema::GetGraphContextActions(FGraphContextMenuBuilder& ContextMenuBuilder) const
 {
+	// Start with being able to create both enter and exit transitions.
 	bool bCanCreateEnterTransition = true;
-	bool bCanCreateExitTransition = false;
+	bool bCanCreateExitTransition = true;
 	if (const UEdGraphPin* DraggedPin = ContextMenuBuilder.FromPin)
 	{
+		// If we are creating a node from dragging a pin into an empty space, figure out which transition
+		// we can create based on the direction of the dragged pin.
 		bCanCreateEnterTransition = false;
 		bCanCreateExitTransition = false;
 
