Cameras: hide hierarchical properties from DetailsView now that we have a camera rig editor

diff --git a/Engine/Plugins/Cameras/GameplayCameras/Source/GameplayCameras/Public/Core/CameraRigAsset.h b/Engine/Plugins/Cameras/GameplayCameras/Source/GameplayCameras/Public/Core/CameraRigAsset.h
--- a/Engine/Plugins/Cameras/GameplayCameras/Source/GameplayCameras/Public/Core/CameraRigAsset.h
+++ b/Engine/Plugins/Cameras/GameplayCameras/Source/GameplayCameras/Public/Core/CameraRigAsset.h
@@ -113,7 +113,7 @@ struct FCameraRigInterface
 public:
 
 	/** The list of exposed parameters on the camera rig. */
-	UPROPERTY()
+	UPROPERTY(Instanced)
 	TArray<TObjectPtr<UCameraRigInterfaceParameter>> InterfaceParameters;
 
 public:
@@ -153,7 +153,7 @@ public:
 public:
 
 	/** Root camera node. */
-	UPROPERTY(EditAnywhere, Instanced, Category=Common)
+	UPROPERTY(Instanced)
 	TObjectPtr<UCameraNode> RootNode;
 
 	/** The gameplay tags on this camera rig. */
@@ -165,11 +165,11 @@ public:
 	FCameraRigInterface Interface;
 
 	/** List of enter transitions for this camera rig. */
-	UPROPERTY(EditAnywhere, Category=Blending)
+	UPROPERTY(Instanced)
 	TArray<TObjectPtr<UCameraRigTransition>> EnterTransitions;
 
 	/** List of exist transitions for this camera rig. */
-	UPROPERTY(EditAnywhere, Category=Blending)
+	UPROPERTY(Instanced)
 	TArray<TObjectPtr<UCameraRigTransition>> ExitTransitions;
 
 	/** Allocation information for all the nodes and variables in this camera rig. */
diff --git a/Engine/Plugins/Cameras/GameplayCameras/Source/GameplayCameras/Public/Core/CameraRigTransition.h b/Engine/Plugins/Cameras/GameplayCameras/Source/GameplayCameras/Public/Core/CameraRigTransition.h
--- a/Engine/Plugins/Cameras/GameplayCameras/Source/GameplayCameras/Public/Core/CameraRigTransition.h
+++ b/Engine/Plugins/Cameras/GameplayCameras/Source/GameplayCameras/Public/Core/CameraRigTransition.h
@@ -89,11 +89,11 @@ class UCameraRigTransition
 public:
 
 	/** The list of conditions that must pass for this transition to be used. */
-	UPROPERTY(EditAnywhere, Instanced, Category=Common)
+	UPROPERTY(Instanced)
 	TArray<TObjectPtr<UCameraRigTransitionCondition>> Conditions;
 
 	/** The blend to use to blend a given camera rig in or out. */
-	UPROPERTY(EditAnywhere, Instanced, Category=Common)
+	UPROPERTY(Instanced)
 	TObjectPtr<UBlendCameraNode> Blend;
 
 protected:
diff --git a/Engine/Plugins/Cameras/GameplayCameras/Source/GameplayCameras/Public/Nodes/Common/ArrayCameraNode.h b/Engine/Plugins/Cameras/GameplayCameras/Source/GameplayCameras/Public/Nodes/Common/ArrayCameraNode.h
--- a/Engine/Plugins/Cameras/GameplayCameras/Source/GameplayCameras/Public/Nodes/Common/ArrayCameraNode.h
+++ b/Engine/Plugins/Cameras/GameplayCameras/Source/GameplayCameras/Public/Nodes/Common/ArrayCameraNode.h
@@ -24,7 +24,7 @@ protected:
 public:
 
 	/** The camera nodes to run. */
-	UPROPERTY(EditAnywhere, Category=Common)
+	UPROPERTY()
 	TArray<TObjectPtr<UCameraNode>> Children;
 };
 
