Cameras: use the full result object in the BP interface

diff --git a/Engine/Plugins/Cameras/GameplayCameras/Source/GameplayCameras/Private/Directors/BlueprintCameraDirector.cpp b/Engine/Plugins/Cameras/GameplayCameras/Source/GameplayCameras/Private/Directors/BlueprintCameraDirector.cpp
--- a/Engine/Plugins/Cameras/GameplayCameras/Source/GameplayCameras/Private/Directors/BlueprintCameraDirector.cpp
+++ b/Engine/Plugins/Cameras/GameplayCameras/Source/GameplayCameras/Private/Directors/BlueprintCameraDirector.cpp
@@ -54,6 +54,8 @@ void FBlueprintCameraDirectorEvaluator::OnRun(const FCameraDirectorEvaluationPar
 
 		EvaluatorBlueprint->NativeRunCameraDirector(BlueprintParams, BlueprintResult);
 
+		// The BP interface doesn't specify the evaluation context for the chosen camera rigs: we always automatically
+		// make them run in our own owner context.
 		for (UCameraRigAsset* ActiveCameraRig : BlueprintResult.ActiveCameraRigs)
 		{
 			OutResult.Add(Params.OwnerContext, ActiveCameraRig);
@@ -70,7 +72,7 @@ void FBlueprintCameraDirectorEvaluator::OnAddReferencedObjects(FReferenceCollect
 
 void UBlueprintCameraDirectorEvaluator::NativeRunCameraDirector(const FBlueprintCameraDirectorEvaluationParams& Params, FBlueprintCameraDirectorEvaluationResult& OutResult)
 {
-	RunCameraDirector(Params, OutResult.ActiveCameraRigs);
+	RunCameraDirector(Params, OutResult);
 }
 
 FCameraDirectorEvaluatorPtr UBlueprintCameraDirector::OnBuildEvaluator(FCameraDirectorEvaluatorBuilder& Builder) const
diff --git a/Engine/Plugins/Cameras/GameplayCameras/Source/GameplayCameras/Public/Directors/BlueprintCameraDirector.h b/Engine/Plugins/Cameras/GameplayCameras/Source/GameplayCameras/Public/Directors/BlueprintCameraDirector.h
--- a/Engine/Plugins/Cameras/GameplayCameras/Source/GameplayCameras/Public/Directors/BlueprintCameraDirector.h
+++ b/Engine/Plugins/Cameras/GameplayCameras/Source/GameplayCameras/Public/Directors/BlueprintCameraDirector.h
@@ -55,7 +55,7 @@ public:
 	UFUNCTION(BlueprintCallable, BlueprintImplementableEvent, Category = "Evaluation")
 	void RunCameraDirector(
 			const FBlueprintCameraDirectorEvaluationParams& Params, 
-			TArray<UCameraRigAsset*>& OutActiveCameraRigs);
+			FBlueprintCameraDirectorEvaluationResult& OutResult);
 
 	/** Native wrapper for RunCameraDirector. */
 	void NativeRunCameraDirector(
