Cameras: fix some bugs with connection-handling of rig parameters

diff --git a/Engine/Plugins/Cameras/GameplayCameras/Source/GameplayCamerasEditor/Private/Editors/CameraNodeGraphSchema.cpp b/Engine/Plugins/Cameras/GameplayCameras/Source/GameplayCamerasEditor/Private/Editors/CameraNodeGraphSchema.cpp
--- a/Engine/Plugins/Cameras/GameplayCameras/Source/GameplayCamerasEditor/Private/Editors/CameraNodeGraphSchema.cpp
+++ b/Engine/Plugins/Cameras/GameplayCameras/Source/GameplayCamerasEditor/Private/Editors/CameraNodeGraphSchema.cpp
@@ -174,9 +174,7 @@ bool UCameraNodeGraphSchema::OnBreakPinLinks(UEdGraphPin& TargetPin, bool bSends
 
 			RigParameter->Target = nullptr;
 			RigParameter->TargetPropertyName = NAME_None;
-
-			// Skip our parent, we just want the actual graph-change only.
-			UEdGraphSchema::BreakPinLinks(TargetPin, bSendsNodeNotification);
+			RigParameter->PrivateVariable = nullptr;
 
 			return true;
 		}
@@ -214,9 +212,7 @@ bool UCameraNodeGraphSchema::OnBreakSinglePinLink(UEdGraphPin* SourcePin, UEdGra
 
 		RigParameter->Target = nullptr;
 		RigParameter->TargetPropertyName = NAME_None;
-
-		// Skip our parent, we just want the actual graph-change only.
-		UEdGraphSchema::BreakSinglePinLink(SourcePin, TargetPin);
+		RigParameter->PrivateVariable = nullptr;
 
 		return true;
 	}
@@ -232,6 +228,10 @@ void UCameraNodeGraphSchema::OnDeleteNodeFromGraph(UObjectTreeGraph* Graph, UEdG
 	if (UCameraRigInterfaceParameterGraphNode* RigParameterNode = Cast<UCameraRigInterfaceParameterGraphNode>(Node))
 	{
 		UCameraRigInterfaceParameter* RigParameter = CastChecked<UCameraRigInterfaceParameter>(RigParameterNode->GetObject());
+
+		CameraRig->Modify();
+		RigParameter->Modify();
+
 		const int32 NumRemoved = CameraRig->Interface.InterfaceParameters.Remove(RigParameter);
 		ensure(NumRemoved == 1);
 	}
@@ -254,20 +254,23 @@ UEdGraphNode* FCameraNodeGraphSchemaAction_NewInterfaceParameterNode::PerformAct
 		return nullptr;
 	}
 
-	UCameraRigAsset* CameraRigAsset = Cast<UCameraRigAsset>(ObjectTreeGraph->GetRootObject());
-	if (!ensure(CameraRigAsset))
+	UCameraRigAsset* CameraRig = Cast<UCameraRigAsset>(ObjectTreeGraph->GetRootObject());
+	if (!ensure(CameraRig))
 	{
 		return nullptr;
 	}
 
-	const FScopedTransaction Transaction(LOCTEXT("CreateNewNodeAction", "Create New Node"));
 	const UObjectTreeGraphSchema* Schema = CastChecked<UObjectTreeGraphSchema>(ParentGraph->GetSchema());
 
-	UCameraRigInterfaceParameter* NewInterfaceParameter = NewObject<UCameraRigInterfaceParameter>(CameraRigAsset, NAME_None, RF_Transactional);
+	const FScopedTransaction Transaction(LOCTEXT("CreateNewNodeAction", "Create New Node"));
+
+	CameraRig->Modify();
+
+	UCameraRigInterfaceParameter* NewInterfaceParameter = NewObject<UCameraRigInterfaceParameter>(CameraRig, NAME_None, RF_Transactional);
 	NewInterfaceParameter->Target = Target;
 	NewInterfaceParameter->TargetPropertyName = TargetPropertyName;
 	NewInterfaceParameter->InterfaceParameterName = TargetPropertyName.ToString();
-	CameraRigAsset->Interface.InterfaceParameters.Add(NewInterfaceParameter);
+	CameraRig->Interface.InterfaceParameters.Add(NewInterfaceParameter);
 
 	UObjectTreeGraphNode* NewGraphNode = Schema->CreateObjectNode(ObjectTreeGraph, NewInterfaceParameter);
 	NewGraphNode->NodePosX = Location.X;
diff --git a/Engine/Plugins/Cameras/GameplayCameras/Source/GameplayCamerasEditor/Private/Editors/ObjectTreeGraphSchema.cpp b/Engine/Plugins/Cameras/GameplayCameras/Source/GameplayCamerasEditor/Private/Editors/ObjectTreeGraphSchema.cpp
--- a/Engine/Plugins/Cameras/GameplayCameras/Source/GameplayCamerasEditor/Private/Editors/ObjectTreeGraphSchema.cpp
+++ b/Engine/Plugins/Cameras/GameplayCameras/Source/GameplayCamerasEditor/Private/Editors/ObjectTreeGraphSchema.cpp
@@ -576,6 +576,7 @@ void UObjectTreeGraphSchema::BreakPinLinks(UEdGraphPin& TargetPin, bool bSendsNo
 	const bool bHandled = OnBreakPinLinks(TargetPin, bSendsNodeNotification);
 	if (bHandled)
 	{
+		Super::BreakPinLinks(TargetPin, bSendsNodeNotification);
 		return;
 	}
 
@@ -632,6 +633,7 @@ void UObjectTreeGraphSchema::BreakSinglePinLink(UEdGraphPin* SourcePin, UEdGraph
 	const bool bHandled = OnBreakSinglePinLink(SourcePin, TargetPin);
 	if (bHandled)
 	{
+		Super::BreakSinglePinLink(SourcePin, TargetPin);
 		return;
 	}
 
