Cameras: make a proper camera rig builder class

Move BuildCameraRig into a utility class and heavily improve its implementation in order to support edge cases, prevent stale private variables being left behind, and more.

diff --git a/Engine/Plugins/Cameras/GameplayCameras/Source/GameplayCameras/Private/Core/CameraRigAllocationInfoBuilder.cpp b/Engine/Plugins/Cameras/GameplayCameras/Source/GameplayCameras/Private/Core/CameraRigAllocationInfoBuilder.cpp
deleted file mode 100644
diff --git a/Engine/Plugins/Cameras/GameplayCameras/Source/GameplayCameras/Private/Core/CameraRigAsset.cpp b/Engine/Plugins/Cameras/GameplayCameras/Source/GameplayCameras/Private/Core/CameraRigAsset.cpp
--- a/Engine/Plugins/Cameras/GameplayCameras/Source/GameplayCameras/Private/Core/CameraRigAsset.cpp
+++ b/Engine/Plugins/Cameras/GameplayCameras/Source/GameplayCameras/Private/Core/CameraRigAsset.cpp
@@ -4,72 +4,13 @@
 
 #include "Core/CameraNode.h"
 #include "Core/CameraParameters.h"
-#include "Core/CameraRigAllocationInfoBuilder.h"
+#include "Core/CameraRigAssetBuilder.h"
 #include "Core/CameraVariableAssets.h"
 #include "Templates/UnrealTypeTraits.h"
 #include "UObject/ObjectSaveContext.h"
 
 #include UE_INLINE_GENERATED_CPP_BY_NAME(CameraRigAsset)
 
-namespace UE::Cameras::Internal
-{
-
-template<typename VariableAssetType, typename ValueType>
-void SetPrivateVariableDefaultValue(VariableAssetType* PrivateVariable, typename TCallTraits<ValueType>::ParamType Value)
-{
-	PrivateVariable->DefaultValue = Value;
-}
-
-template<>
-void SetPrivateVariableDefaultValue<UBooleanCameraVariable, bool>(UBooleanCameraVariable* PrivateVariable, bool bValue)
-{
-	PrivateVariable->bDefaultValue = bValue;
-}
-
-template<typename CameraParameterType>
-UCameraVariableAsset* CreatePrivateVariable(UCameraRigAsset* CameraRigAsset, UCameraRigInterfaceParameter* InterfaceParameter, const FString& VariableName, CameraParameterType* CameraParameter)
-{
-	using ValueType = typename CameraParameterType::ValueType;
-	using VariableAssetType = typename CameraParameterType::VariableAssetType;
-
-	VariableAssetType* PrivateVariable = nullptr;
-
-	if (CameraParameter->Variable != nullptr)
-	{
-		UObject* VariableOuter = CameraParameter->Variable->GetOuter();
-		if (VariableOuter != CameraRigAsset)
-		{
-			// If this parameter is driven by a user-defined variable, emit a warning, since we will 
-			// replace that driving variable with our private variable.
-			UE_LOG(LogCameraSystem, Warning, 
-					TEXT("Camera node parameter '%s.%s' in camera rig '%s' is both exposed and driven by a variable! That variable will be replaced by the rig parameter variable."), 
-					*InterfaceParameter->Target->GetName(), 
-					*InterfaceParameter->TargetPropertyName.ToString(),
-					*CameraRigAsset->GetPathName());
-		}
-		else
-		{
-			PrivateVariable = CameraParameter->Variable;
-		}
-	}
-
-	if (PrivateVariable == nullptr)
-	{
-		PrivateVariable = NewObject<VariableAssetType>(CameraRigAsset, *VariableName);
-	}
-
-	PrivateVariable->bIsPrivate = true;
-	PrivateVariable->bAutoReset = false;
-	SetPrivateVariableDefaultValue<VariableAssetType, ValueType>(PrivateVariable, CameraParameter->Value);
-
-	InterfaceParameter->PrivateVariable = PrivateVariable;
-	CameraParameter->Variable = PrivateVariable;
-
-	return PrivateVariable;
-}
-
-}  // namespace UE::Cameras::Internal
-
 #if WITH_EDITOR
 
 void UCameraRigInterfaceParameter::GetGraphNodePosition(int32& NodePosX, int32& NodePosY) const
@@ -108,100 +49,9 @@ void UCameraRigAsset::GetOwnedGameplayTags(FGameplayTagContainer& TagContainer)
 void UCameraRigAsset::BuildCameraRig()
 {
 	using namespace UE::Cameras;
-	using namespace UE::Cameras::Internal;
-
-	Modify();
-
-	// Build allocation info.
-	FCameraRigAllocationInfoBuilder CameraRigBuilder;
-	CameraRigBuilder.BuildAllocationInfo(this, AllocationInfo);
-	BuildStatus = ECameraRigBuildStatus::Clean;
-
-	// Build interface parameters.
-	TSet<FString> UsedInterfaceParameterNames;
-	for (UCameraRigInterfaceParameter* InterfaceParameter : Interface.InterfaceParameters)
-	{
-		// Do some basic validation.
-		if (!InterfaceParameter || !InterfaceParameter->Target)
-		{
-			UE_LOG(LogCameraSystem, Error, 
-					TEXT("Invalid interface parameter target in camera rig: '%s'."),
-					*GetPathName());
-			continue;
-		}
-		if (InterfaceParameter->TargetPropertyName.IsNone())
-		{
-			UE_LOG(LogCameraSystem, Error, 
-					TEXT("Invalid interface parameter target property name in camera rig: '%s'."), 
-					*GetPathName());
-			continue;
-		}
-		if (InterfaceParameter->InterfaceParameterName.IsEmpty())
-		{
-			UE_LOG(LogCameraSystem, Error, 
-					TEXT("Invalid interface parameter name in camera rig '%s'."), 
-					*GetPathName());
-			continue;
-		}
 
-		if (UsedInterfaceParameterNames.Contains(InterfaceParameter->InterfaceParameterName))
-		{
-			UE_LOG(LogCameraSystem, Error, 
-					TEXT("Multiple interface parameters named '%s' in camera rig '%s'. Ignoring duplicates."), 
-					*InterfaceParameter->InterfaceParameterName, 
-					*GetPathName());
-			continue;
-		}
-		UsedInterfaceParameterNames.Add(InterfaceParameter->InterfaceParameterName);
-
-		// Get the target camera node property and check that it is a camera parameter struct.
-		UClass* TargetClass = InterfaceParameter->Target->GetClass();
-		FProperty* TargetProperty = TargetClass->FindPropertyByName(InterfaceParameter->TargetPropertyName);
-		if (!TargetProperty)
-		{
-			UE_LOG(LogCameraSystem, Error, 
-					TEXT("Invalid interface parameter in camera rig '%s': no property '%s' on camera node '%s'."), 
-					*GetPathName(), 
-					*InterfaceParameter->InterfaceParameterName, 
-					*InterfaceParameter->Target->GetName());
-			continue;
-		}
-
-		FStructProperty* TargetStructProperty = CastField<FStructProperty>(TargetProperty);
-		if (!TargetStructProperty)
-		{
-			UE_LOG(LogCameraSystem, Error, 
-					TEXT("Invalid interface parameter in camera rig '%s': property '%s' on camera node '%s' is not a camera parameter."),
-					*GetPathName(), 
-					*InterfaceParameter->InterfaceParameterName, 
-					*InterfaceParameter->Target->GetName());
-			continue;
-		}
-
-		// Get the type of the camera parameter by matching the struct against all the types we support,
-		// and create a private camera variable asset to drive its value.
-		FString PrivateVariableName(FString::Format(TEXT("Override_{0}_{1}"), { GetName(), InterfaceParameter->InterfaceParameterName }));
-#define UE_CAMERA_VARIABLE_FOR_TYPE(ValueType, ValueName)\
-		if (TargetStructProperty->Struct == F##ValueName##CameraParameter::StaticStruct())\
-		{\
-			F##ValueName##CameraParameter CameraParameter;\
-			TargetStructProperty->GetValue_InContainer(InterfaceParameter->Target, reinterpret_cast<void*>(&CameraParameter));\
-			CreatePrivateVariable(this, InterfaceParameter, PrivateVariableName, &CameraParameter);\
-			InterfaceParameter->Target->Modify();\
-			TargetStructProperty->SetValue_InContainer(InterfaceParameter->Target, reinterpret_cast<void*>(&CameraParameter));\
-		}\
-		else
-UE_CAMERA_VARIABLE_FOR_ALL_TYPES()
-#undef UE_CAMERA_VARIABLE_FOR_TYPE
-		{
-			UE_LOG(LogCameraSystem, Error, 
-					TEXT("Invalid interface parameter in camera rig '%s': property '%s' on camera node '%s' is not a camera parameter."), 
-					*GetPathName(), 
-					*InterfaceParameter->InterfaceParameterName,
-					*InterfaceParameter->Target->GetName());
-			continue;
-		}
-	}
+	FCameraRigAssetBuilder Builder;
+	Builder.BuildCameraRig(this);
 }
 
 void UCameraRigAsset::PostLoad()
diff --git a/Engine/Plugins/Cameras/GameplayCameras/Source/GameplayCameras/Private/Core/CameraRigAssetBuilder.cpp b/Engine/Plugins/Cameras/GameplayCameras/Source/GameplayCameras/Private/Core/CameraRigAssetBuilder.cpp
new file mode 100644
--- /dev/null
+++ b/Engine/Plugins/Cameras/GameplayCameras/Source/GameplayCameras/Private/Core/CameraRigAssetBuilder.cpp
@@ -0,0 +1,517 @@
+// Copyright Epic Games, Inc. All Rights Reserved.
+
+#include "Core/CameraRigAssetBuilder.h"
+
+#include "Core/CameraNode.h"
+#include "Core/CameraNodeEvaluatorStorage.h"
+#include "Core/CameraParameters.h"
+#include "Core/CameraRigAsset.h"
+#include "Core/CameraVariableAssets.h"
+#include "Misc/CString.h"
+
+namespace UE::Cameras
+{
+
+namespace Internal
+{
+
+template<typename VariableAssetType, typename ValueType>
+void SetPrivateVariableDefaultValue(VariableAssetType* PrivateVariable, typename TCallTraits<ValueType>::ParamType Value)
+{
+	if (PrivateVariable->DefaultValue != Value)
+	{
+		PrivateVariable->Modify();
+		PrivateVariable->DefaultValue = Value;
+	}
+}
+
+template<>
+void SetPrivateVariableDefaultValue<UTransform3dCameraVariable, FTransform3d>(UTransform3dCameraVariable* PrivateVariable, const FTransform3d& Value)
+{
+	// Template overload because transforms don't have an operator!=.
+	if (!PrivateVariable->DefaultValue.Equals(Value, 0.f))
+	{
+		PrivateVariable->Modify();
+		PrivateVariable->DefaultValue = Value;
+	}
+}
+
+template<>
+void SetPrivateVariableDefaultValue<UTransform3fCameraVariable, FTransform3f>(UTransform3fCameraVariable* PrivateVariable, const FTransform3f& Value)
+{
+	// Template overload because transforms don't have an operator!=.
+	if (!PrivateVariable->DefaultValue.Equals(Value, 0.f))
+	{
+		PrivateVariable->Modify();
+		PrivateVariable->DefaultValue = Value;
+	}
+}
+
+template<>
+void SetPrivateVariableDefaultValue<UBooleanCameraVariable, bool>(UBooleanCameraVariable* PrivateVariable, bool bValue)
+{
+	// Template overload because boolean variables have a bDefaultValue field, not DefaultValue.
+	if (PrivateVariable->bDefaultValue != bValue)
+	{
+		PrivateVariable->Modify();
+		PrivateVariable->bDefaultValue = bValue;
+	}
+}
+
+struct FPrivateVariableBuilder
+{
+	UCameraRigAsset* CameraRig;
+
+	FPrivateVariableBuilder(FCameraRigAssetBuilder& InOwner)
+		: Owner(InOwner)
+	{
+		CameraRig = Owner.CameraRig;
+	}
+
+	void ReportError()
+	{
+		Owner.bHasErrors = true;
+	}
+
+	template<typename ExpectedVariableAssetType>
+	ExpectedVariableAssetType* FindReusablePrivateVariable(FStructProperty* ForParameterProperty, UCameraNode* ForCameraNode)
+	{
+		UCameraVariableAsset* ReusedVariable = FindReusablePrivateVariable(ForParameterProperty, ForCameraNode);
+		if (ReusedVariable)
+		{
+			return CastChecked<ExpectedVariableAssetType>(ReusedVariable);
+		}
+		return nullptr;
+	}
+
+private:
+
+	UCameraVariableAsset* FindReusablePrivateVariable(FStructProperty* ForParameterProperty, UCameraNode* ForCameraNode)
+	{
+		using FDrivenParameterKey = FCameraRigAssetBuilder::FDrivenParameterKey;
+
+		FDrivenParameterKey ParameterKey{ ForParameterProperty, ForCameraNode };
+		UCameraVariableAsset** FoundItem = Owner.OldDrivenParameters.Find(ParameterKey);
+		if (FoundItem)
+		{
+			// We found an existing variable that was driving this camera node's property.
+			// Re-use it and remove it from the re-use pool.
+			UCameraVariableAsset* ReusedVariable = (*FoundItem);
+			Owner.OldDrivenParameters.Remove(ParameterKey);
+			return ReusedVariable;
+		}
+
+		return nullptr;
+	}
+
+private:
+
+	FCameraRigAssetBuilder& Owner;
+};
+
+template<typename CameraParameterType>
+void SetupPrivateVariable(
+		FPrivateVariableBuilder& Builder, 
+		UCameraRigInterfaceParameter* InterfaceParameter, 
+		FStructProperty* ParameterTargetProperty,
+		CameraParameterType* CameraParameter)
+{
+	using ValueType = typename CameraParameterType::ValueType;
+	using VariableAssetType = typename CameraParameterType::VariableAssetType;
+
+	if (CameraParameter->Variable != nullptr)
+	{
+		// We should have cleared all exposed parameters in GatherOldDrivenParameters, so the only variables
+		// left on camera parameters should be user-defined ones.
+		UObject* VariableOuter = CameraParameter->Variable->GetOuter();
+		if (ensureMsgf(
+					VariableOuter != Builder.CameraRig, 
+					TEXT("Unexpected driving variable found: all exposed parameters should have been cleared before rebuilding.")))
+		{
+			// If this parameter is driven by a user-defined variable, emit an error, and replace that 
+			// driving variable with our private variable.
+			UE_LOG(LogCameraSystem, Error, 
+					TEXT("Camera node parameter '%s.%s' in camera rig '%s' is both exposed and driven by a variable!."), 
+					*InterfaceParameter->Target->GetName(), 
+					*InterfaceParameter->TargetPropertyName.ToString(),
+					*Builder.CameraRig->GetPathName());
+			Builder.ReportError();
+		}
+	}
+
+	// Start by re-using the camera variable that was already driving this parameter on this node.
+	VariableAssetType* PrivateVariable = Builder.FindReusablePrivateVariable<VariableAssetType>(
+			ParameterTargetProperty, InterfaceParameter->Target);
+	const bool bIsReusedVariable = (PrivateVariable != nullptr);
+
+	const FString VariableName = FString::Format(
+			TEXT("Override_{0}_{1}"), 
+			{ Builder.CameraRig->GetName(), InterfaceParameter->InterfaceParameterName });
+
+	if (PrivateVariable)
+	{
+		// Found a pre-existing variable! Make sure it's got the right name, in case the exposed rig parameter
+		// was renamed. Keeping a good name is mostly to help with debugging.
+		FString OriginalName = PrivateVariable->GetName();
+		OriginalName.RemoveFromStart("REUSABLE_", ESearchCase::CaseSensitive);
+		if (OriginalName != VariableName)
+		{
+			PrivateVariable->Modify();
+		}
+		PrivateVariable->Rename(*VariableName);
+	}
+	else
+	{
+		// Make a new variable.
+		PrivateVariable = NewObject<VariableAssetType>(Builder.CameraRig, FName(*VariableName), RF_Transactional);
+	}
+
+	ensure(PrivateVariable->GetOuter() == Builder.CameraRig);
+
+	PrivateVariable->bIsPrivate = true;
+	PrivateVariable->bAutoReset = false;
+
+	// Set the default value of the variable to be the value in the camera parameter.
+	SetPrivateVariableDefaultValue<VariableAssetType, ValueType>(PrivateVariable, CameraParameter->Value);
+
+	// Set the variable on both the interface parameter and the camera node. Flag them as modified
+	// if we actually changed anything.
+	if (InterfaceParameter->PrivateVariable != PrivateVariable)
+	{
+		InterfaceParameter->Modify();
+	}
+	if (!bIsReusedVariable)
+	{
+		InterfaceParameter->Target->Modify();
+	}
+	InterfaceParameter->PrivateVariable = PrivateVariable;
+	CameraParameter->Variable = PrivateVariable;
+}
+
+}  // namespace Internal
+
+void FCameraRigAssetBuilder::BuildCameraRig(UCameraRigAsset* InCameraRig)
+{
+	if (!ensure(InCameraRig))
+	{
+		return;
+	}
+
+	CameraRig = InCameraRig;
+	bHasErrors = false;
+	bHasWarnings = false;
+	{
+		BuildCameraRigImpl();
+	}
+	UpdateBuildStatus();
+}
+
+void FCameraRigAssetBuilder::BuildCameraRigImpl()
+{
+	if (!CameraRig->RootNode)
+	{
+		UE_LOG(LogCameraSystem, Error, TEXT("Camera rig '%s' has no root node."), *CameraRig->GetPathName());
+		bHasErrors = true;
+		return;
+	}
+
+	FlattenCameraNodeHierarchy();
+
+	GatherOldDrivenParameters();
+	BuildNewDrivenParameters();
+	DiscardUnusedPrivateVariables();
+
+	BuildAllocationInfo();
+}
+
+void FCameraRigAssetBuilder::FlattenCameraNodeHierarchy()
+{
+	// Build a flat list of the camera rig's node hierarchy. It's easier to iterate during
+	// our build process.
+	FlattenedNodes.Reset();
+
+	TArray<UCameraNode*> NodeStack;
+	NodeStack.Add(CameraRig->RootNode);
+	while (!NodeStack.IsEmpty())
+	{
+		UCameraNode* CurrentNode = NodeStack.Pop();
+		FlattenedNodes.Add(CurrentNode);
+
+		FCameraNodeChildrenView CurrentChildren = CurrentNode->GetChildren();
+		for (UCameraNode* Child : ReverseIterate(CurrentChildren))
+		{
+			if (Child)
+			{
+				NodeStack.Add(Child);
+			}
+		}
+	}
+}
+
+void FCameraRigAssetBuilder::GatherOldDrivenParameters()
+{
+	// Keep track of what camera parameters were previously driven by private variables,
+	// and then clear those variables. This is because it's easier to rebuild this from
+	// a blank slate than trying to figure out what changed.
+	//
+	// As we rebuild things in BuildNewDrivenParameters, we compare to the old state to
+	// figure out if we need to flag anything as modified for the current transaction.
+	//
+	// Note that parameters driven by user-defined variables are left alone.
+
+	OldDrivenParameters.Reset();
+
+	for (UCameraNode* CameraNode : FlattenedNodes)
+	{
+		UClass* CameraNodeClass = CameraNode->GetClass();
+		
+		for (TFieldIterator<FProperty> It(CameraNodeClass); It; ++It)
+		{
+			FStructProperty* StructProperty = CastField<FStructProperty>(*It);
+			if (!StructProperty)
+			{
+				continue;
+			}
+
+#define UE_CAMERA_VARIABLE_FOR_TYPE(ValueType, ValueName)\
+			if (StructProperty->Struct == F##ValueName##CameraParameter::StaticStruct())\
+			{\
+				auto* CameraParameterPtr = StructProperty->ContainerPtrToValuePtr<F##ValueName##CameraParameter>(CameraNode);\
+				if (CameraParameterPtr->Variable)\
+				{\
+					UObject* VariableOuter = CameraParameterPtr->Variable->GetOuter();\
+					if (VariableOuter == CameraRig)\
+					{\
+						OldDrivenParameters.Add(\
+								FDrivenParameterKey{ StructProperty, CameraNode },\
+								CameraParameterPtr->Variable);\
+						CameraParameterPtr->Variable = nullptr;\
+					}\
+				}\
+			}\
+			else
+UE_CAMERA_VARIABLE_FOR_ALL_TYPES()
+#undef UE_CAMERA_VARIABLE_FOR_TYPE
+			{
+				// Some other struct property.
+			}
+		}
+	}
+
+	// Temporarily rename all old camera variables, so their names are available to the new
+	// driven parameters.
+	for (TPair<FDrivenParameterKey, UCameraVariableAsset*> Pair : OldDrivenParameters)
+	{
+		TStringBuilder<256> StringBuilder;
+		StringBuilder.Append("REUSABLE_");
+		StringBuilder.Append(Pair.Value->GetName());
+		Pair.Value->Rename(StringBuilder.ToString());
+	}
+}
+
+void FCameraRigAssetBuilder::BuildNewDrivenParameters()
+{
+	using namespace Internal;
+
+	TSet<FString> UsedInterfaceParameterNames;
+
+	const FString CameraRigName = CameraRig->GetName();
+	const FString CameraRigPathName = CameraRig->GetPathName();
+
+	// Look at the new interface parameters and setup the driven camera node parameters with
+	// private camera variables. We have gathered the old ones previously so we can re-use them,
+	// instead of creating new variable assets each time.
+	for (UCameraRigInterfaceParameter* InterfaceParameter : CameraRig->Interface.InterfaceParameters)
+	{
+		// Do some basic validation.
+		if (!InterfaceParameter || !InterfaceParameter->Target)
+		{
+			UE_LOG(LogCameraSystem, Error, 
+					TEXT("Invalid interface parameter target in camera rig: '%s'."),
+					*CameraRigPathName);
+			bHasErrors = true;
+			continue;
+		}
+		if (InterfaceParameter->TargetPropertyName.IsNone())
+		{
+			UE_LOG(LogCameraSystem, Error, 
+					TEXT("Invalid interface parameter target property name in camera rig: '%s'."), 
+					*CameraRigPathName);
+			bHasErrors = true;
+			continue;
+		}
+		if (InterfaceParameter->InterfaceParameterName.IsEmpty())
+		{
+			UE_LOG(LogCameraSystem, Error, 
+					TEXT("Invalid interface parameter name in camera rig '%s'."), 
+					*CameraRigPathName);
+			bHasErrors = true;
+			continue;
+		}
+
+		// Check duplicate parameter names.
+		if (UsedInterfaceParameterNames.Contains(InterfaceParameter->InterfaceParameterName))
+		{
+			UE_LOG(LogCameraSystem, Error, 
+					TEXT("Multiple interface parameters named '%s' in camera rig '%s'. Ignoring duplicates."), 
+					*InterfaceParameter->InterfaceParameterName, 
+					*CameraRigPathName);
+			bHasErrors = true;
+			continue;
+		}
+		UsedInterfaceParameterNames.Add(InterfaceParameter->InterfaceParameterName);
+
+		// Get the target camera node property and check that it is a camera parameter struct.
+		UCameraNode* Target = InterfaceParameter->Target;
+		UClass* TargetClass = Target->GetClass();
+		FProperty* TargetProperty = TargetClass->FindPropertyByName(InterfaceParameter->TargetPropertyName);
+		if (!TargetProperty)
+		{
+			UE_LOG(LogCameraSystem, Error, 
+					TEXT("Invalid interface parameter in camera rig '%s': no property '%s' on camera node '%s'."), 
+					*CameraRigPathName, 
+					*InterfaceParameter->InterfaceParameterName, 
+					*Target->GetName());
+			bHasErrors = true;
+			continue;
+		}
+
+		FStructProperty* TargetStructProperty = CastField<FStructProperty>(TargetProperty);
+		if (!TargetStructProperty)
+		{
+			UE_LOG(LogCameraSystem, Error, 
+					TEXT("Invalid interface parameter in camera rig '%s': property '%s' on camera node '%s' is not a camera parameter."),
+					*CameraRigPathName, 
+					*InterfaceParameter->InterfaceParameterName, 
+					*Target->GetName());
+			bHasErrors = true;
+			continue;
+		}
+
+		// Get the type of the camera parameter by matching the struct against all the types we support,
+		// and create a private camera variable asset to drive its value.
+		FPrivateVariableBuilder PrivateVariableBuilder(*this);
+#define UE_CAMERA_VARIABLE_FOR_TYPE(ValueType, ValueName)\
+		if (TargetStructProperty->Struct == F##ValueName##CameraParameter::StaticStruct())\
+		{\
+			auto* CameraParameterPtr = TargetStructProperty->ContainerPtrToValuePtr<F##ValueName##CameraParameter>(Target);\
+			SetupPrivateVariable(PrivateVariableBuilder, InterfaceParameter, TargetStructProperty, CameraParameterPtr);\
+		}\
+		else
+UE_CAMERA_VARIABLE_FOR_ALL_TYPES()
+#undef UE_CAMERA_VARIABLE_FOR_TYPE
+		{
+			UE_LOG(LogCameraSystem, Error, 
+					TEXT("Invalid interface parameter in camera rig '%s': property '%s' on camera node '%s' is not a camera parameter."), 
+					*CameraRigPathName, 
+					*InterfaceParameter->InterfaceParameterName,
+					*Target->GetName());
+			bHasErrors = true;
+			continue;
+		}
+	}
+}
+
+void FCameraRigAssetBuilder::DiscardUnusedPrivateVariables()
+{
+	// Now that we've rebuilt all exposed parameters, anything left from the old list 
+	// must be discarded.
+	for (TPair<FDrivenParameterKey, UCameraVariableAsset*> Pair : OldDrivenParameters)
+	{
+		// We null'ed the driving variable in GatherOldDrivenParameters. Now it's time
+		// to flag the camera node as modified.
+		UCameraNode* Target = Pair.Key.Value;
+		Target->Modify();
+
+		// Trash the old camera variable. This helps with debugging.
+		UCameraVariableAsset* VariableToDiscard = Pair.Value;
+		TStringBuilder<256> StringBuilder;
+		StringBuilder.Append("TRASH_");
+		StringBuilder.Append(VariableToDiscard->GetName());
+		VariableToDiscard->Rename(StringBuilder.ToString());
+	}
+
+	OldDrivenParameters.Reset();
+}
+
+void FCameraRigAssetBuilder::BuildAllocationInfo()
+{
+	AllocationInfo = FCameraRigAllocationInfo();
+
+	// Build a mock tree of evaluators.
+	FCameraNodeEvaluatorTreeBuildParams BuildParams;
+	BuildParams.RootCameraNode = CameraRig->RootNode;
+	BuildParams.bInitialize = false;  // Don't initialize, we just want to know how much room they take.
+	FCameraNodeEvaluatorStorage Storage;
+	Storage.BuildEvaluatorTree(BuildParams);
+
+	// Get the size of the evaluators' allocation.
+	Storage.GetAllocationInfo(AllocationInfo.EvaluatorInfo);
+
+	// Compute the allocation info for camera variables.
+	for (UCameraNode* CameraNode : FlattenedNodes)
+	{
+		BuildAllocationInfo(CameraNode);
+	}
+
+	// Set it on the camera rig asset.
+	CameraRig->AllocationInfo = AllocationInfo;
+}
+
+void FCameraRigAssetBuilder::BuildAllocationInfo(const UCameraNode* CameraNode)
+{
+	// Look for properties that are camera parameters, and gather what camera variables they reference. 
+	// This is for both exposed rig parameters (which we just built in BuildNewDrivenParameters) and 
+	// for parameters driven by user-defined variables.
+	UClass* CameraNodeClass = CameraNode->GetClass();
+	for (TFieldIterator<FProperty> It(CameraNodeClass); It; ++It)
+	{
+		FStructProperty* StructProperty = CastField<FStructProperty>(*It);
+		if (!StructProperty)
+		{
+			continue;
+		}
+
+#define UE_CAMERA_VARIABLE_FOR_TYPE(ValueType, ValueName)\
+		if (StructProperty->Struct == F##ValueName##CameraParameter::StaticStruct())\
+		{\
+			auto* CameraParameterPtr = StructProperty->ContainerPtrToValuePtr<F##ValueName##CameraParameter>(CameraNode);\
+			if (CameraParameterPtr->Variable)\
+			{\
+				FCameraVariableDefinition VariableDefinition = CameraParameterPtr->Variable->GetVariableDefinition();\
+				AllocationInfo.VariableTableInfo.VariableDefinitions.Add(VariableDefinition);\
+			}\
+		}\
+		else
+UE_CAMERA_VARIABLE_FOR_ALL_TYPES()
+#undef UE_CAMERA_VARIABLE_FOR_TYPE
+		{
+			// Some other struct property.
+		}
+	}
+
+	// Let the camera node add any custom variables or extra memory.
+	CameraNode->BuildAllocationInfo(AllocationInfo);
+}
+
+void FCameraRigAssetBuilder::UpdateBuildStatus()
+{
+	ECameraRigBuildStatus BuildStatus = ECameraRigBuildStatus::Clean;
+	if (bHasErrors)
+	{
+		BuildStatus = ECameraRigBuildStatus::WithErrors;
+	}
+	else if (bHasWarnings)
+	{
+		BuildStatus = ECameraRigBuildStatus::CleanWithWarnings;
+	}
+
+	if (CameraRig->BuildStatus != BuildStatus)
+	{
+		CameraRig->Modify();
+		CameraRig->BuildStatus = BuildStatus;
+	}
+}
+
+}  // namespace UE::Cameras
+
diff --git a/Engine/Plugins/Cameras/GameplayCameras/Source/GameplayCameras/Public/Core/CameraRigAllocationInfoBuilder.h b/Engine/Plugins/Cameras/GameplayCameras/Source/GameplayCameras/Public/Core/CameraRigAllocationInfoBuilder.h
deleted file mode 100644
diff --git a/Engine/Plugins/Cameras/GameplayCameras/Source/GameplayCameras/Public/Core/CameraRigAsset.h b/Engine/Plugins/Cameras/GameplayCameras/Source/GameplayCameras/Public/Core/CameraRigAsset.h
--- a/Engine/Plugins/Cameras/GameplayCameras/Source/GameplayCameras/Public/Core/CameraRigAsset.h
+++ b/Engine/Plugins/Cameras/GameplayCameras/Source/GameplayCameras/Public/Core/CameraRigAsset.h
@@ -17,6 +17,8 @@
 class UCameraNode;
 class UCameraVariableAsset;
 
+namespace UE::Cameras { class FCameraRigAssetBuilder; }
+
 /**
  *
  */
@@ -234,9 +236,11 @@ private:
 	 * A list of all the camera nodes, including the 'loose' ones that aren't connected
 	 * to the root node, and therefore would be GC'ed if we didn't hold them here.
 	 */
-	UPROPERTY(meta=(ObjectTreeGraphHidden=true))
+	UPROPERTY(Instanced, meta=(ObjectTreeGraphHidden=true))
 	TArray<TObjectPtr<UCameraNode>> AllNodes;
 
 #endif  // WITH_EDITORONLY_DATA
+
+	friend class UE::Cameras::FCameraRigAssetBuilder;
 };
 
diff --git a/Engine/Plugins/Cameras/GameplayCameras/Source/GameplayCameras/Public/Core/CameraRigAssetBuilder.h b/Engine/Plugins/Cameras/GameplayCameras/Source/GameplayCameras/Public/Core/CameraRigAssetBuilder.h
new file mode 100644
--- /dev/null
+++ b/Engine/Plugins/Cameras/GameplayCameras/Source/GameplayCameras/Public/Core/CameraRigAssetBuilder.h
@@ -0,0 +1,73 @@
+// Copyright Epic Games, Inc. All Rights Reserved.
+
+#pragma once
+
+#include "Containers/Array.h"
+#include "Core/CameraNodeEvaluatorFwd.h"
+#include "Core/CameraRigAsset.h"
+#include "Core/CameraVariableTableFwd.h"
+#include "CoreTypes.h"
+#include "Templates/Tuple.h"
+#include "UObject/WeakObjectPtr.h"
+
+class FStructProperty;
+class UCameraNode;
+class UCameraRigAsset;
+class UCameraVariableAsset;
+struct FCameraVariableTableAllocationInfo;
+
+namespace UE::Cameras
+{
+
+namespace Internal { struct FPrivateVariableBuilder; }
+
+/**
+ * A class that can prepare a camera rig for runtime use.
+ *
+ * This builder class sets up internal camera variables that handle exposed camera
+ * rig parameters, computes the allocation information of the camera rig, and
+ * does various kinds of validation.
+ *
+ * Once the build process is done, the BuildStatus proeprty is set on the camera rig.
+ */
+class FCameraRigAssetBuilder
+{
+public:
+
+	/** Builds the given camera rig. */
+	void BuildCameraRig(UCameraRigAsset* InCameraRig);
+
+private:
+
+	void BuildCameraRigImpl();
+
+	void FlattenCameraNodeHierarchy();
+
+	void GatherOldDrivenParameters();
+	void BuildNewDrivenParameters();
+	void DiscardUnusedPrivateVariables();
+
+	void BuildAllocationInfo();
+	void BuildAllocationInfo(const UCameraNode* CameraNode);
+
+	void UpdateBuildStatus();
+
+private:
+
+	UCameraRigAsset* CameraRig = nullptr;
+
+	TArray<UCameraNode*> FlattenedNodes;
+
+	using FDrivenParameterKey = TTuple<FStructProperty*, UCameraNode*>;
+	TMap<FDrivenParameterKey, UCameraVariableAsset*> OldDrivenParameters;
+
+	FCameraRigAllocationInfo AllocationInfo;
+
+	bool bHasErrors;
+	bool bHasWarnings;
+
+	friend struct Internal::FPrivateVariableBuilder;
+};
+
+}  // namespace UE::Cameras
+
