Cameras: add unit tests for the camera rig builder

diff --git a/Engine/Plugins/Cameras/GameplayCameras/Source/GameplayCameras/Private/Nodes/Common/ArrayCameraNode.cpp b/Engine/Plugins/Cameras/GameplayCameras/Source/GameplayCameras/Private/Nodes/Common/ArrayCameraNode.cpp
--- a/Engine/Plugins/Cameras/GameplayCameras/Source/GameplayCameras/Private/Nodes/Common/ArrayCameraNode.cpp
+++ b/Engine/Plugins/Cameras/GameplayCameras/Source/GameplayCameras/Private/Nodes/Common/ArrayCameraNode.cpp
@@ -26,6 +26,12 @@ private:
 
 UE_DEFINE_CAMERA_NODE_EVALUATOR(FArrayCameraNodeEvaluator)
 
+#if !(UE_BUILD_SHIPPING || UE_BUILD_TEST)
+// For unit tests.
+int32 GArrayCameraNodeEvaluatorSizeof = sizeof(FArrayCameraNodeEvaluator);
+int32 GArrayCameraNodeEvaluatorAlignof = alignof(FArrayCameraNodeEvaluator);
+#endif
+
 FCameraNodeEvaluatorChildrenView FArrayCameraNodeEvaluator::OnGetChildren()
 {
 	return FCameraNodeEvaluatorChildrenView(Children);
diff --git a/Engine/Plugins/Cameras/GameplayCameras/Source/GameplayCameras/Private/Nodes/Common/OffsetCameraNode.cpp b/Engine/Plugins/Cameras/GameplayCameras/Source/GameplayCameras/Private/Nodes/Common/OffsetCameraNode.cpp
--- a/Engine/Plugins/Cameras/GameplayCameras/Source/GameplayCameras/Private/Nodes/Common/OffsetCameraNode.cpp
+++ b/Engine/Plugins/Cameras/GameplayCameras/Source/GameplayCameras/Private/Nodes/Common/OffsetCameraNode.cpp
@@ -27,6 +27,12 @@ private:
 	TCameraParameterReader<FRotator3d> RotationReader;
 };
 
+#if !(UE_BUILD_SHIPPING || UE_BUILD_TEST)
+// For unit tests.
+int32 GOffsetCameraNodeEvaluatorSizeof = sizeof(FOffsetCameraNodeEvaluator);
+int32 GOffsetCameraNodeEvaluatorAlignof = alignof(FOffsetCameraNodeEvaluator);
+#endif
+
 UE_DEFINE_CAMERA_NODE_EVALUATOR(FOffsetCameraNodeEvaluator)
 
 void FOffsetCameraNodeEvaluator::OnInitialize(const FCameraNodeEvaluatorInitializeParams& Params)
diff --git a/Engine/Plugins/Cameras/GameplayCameras/Source/GameplayCameras/Private/Tests/CameraRigAssetBuilderTests.cpp b/Engine/Plugins/Cameras/GameplayCameras/Source/GameplayCameras/Private/Tests/CameraRigAssetBuilderTests.cpp
new file mode 100644
--- /dev/null
+++ b/Engine/Plugins/Cameras/GameplayCameras/Source/GameplayCameras/Private/Tests/CameraRigAssetBuilderTests.cpp
@@ -0,0 +1,158 @@
+// Copyright Epic Games, Inc. All Rights Reserved.
+
+#include "Core/CameraRigAsset.h"
+#include "Misc/AutomationTest.h"
+#include "Nodes/Blends/SmoothBlendCameraNode.h"
+#include "Nodes/Common/OffsetCameraNode.h"
+#include "Nodes/Common/LensParametersCameraNode.h"
+#include "Tests/GameplayCamerasTestBuilder.h"
+
+#define LOCTEXT_NAMESPACE "CameraRigAssetBuilderTests"
+
+IMPLEMENT_SIMPLE_AUTOMATION_TEST(FCameraRigAssetBuilderNullTest, "System.Engine.GameplayCameras.CameraRigAssetBuilder.Null", EAutomationTestFlags::EditorContext | EAutomationTestFlags::EngineFilter)
+bool FCameraRigAssetBuilderNullTest::RunTest(const FString& Parameters)
+{
+	using namespace UE::Cameras::Test;
+
+	UCameraRigAsset* CameraRig = FCameraRigAssetTestBuilder(TEXT("InvalidTest")).Get();
+	UTEST_EQUAL("Dirty status", CameraRig->BuildStatus, ECameraRigBuildStatus::Dirty);
+
+	FStringFormatOrderedArguments ErrorArgs;
+	ErrorArgs.Add(CameraRig->GetPathName());
+	AddExpectedMessage(
+			FString::Format(TEXT("Camera rig '{0}' has no root node."), ErrorArgs),
+			ELogVerbosity::Error,
+			EAutomationExpectedMessageFlags::Exact,
+			1,
+			false);
+	CameraRig->BuildCameraRig();
+	UTEST_EQUAL("Error status", CameraRig->BuildStatus, ECameraRigBuildStatus::WithErrors);
+
+	return true;
+}
+
+namespace UE::Cameras
+{
+	extern int32 GArrayCameraNodeEvaluatorSizeof;
+	extern int32 GArrayCameraNodeEvaluatorAlignof;
+	extern int32 GOffsetCameraNodeEvaluatorSizeof;
+	extern int32 GOffsetCameraNodeEvaluatorAlignof;
+}
+
+IMPLEMENT_SIMPLE_AUTOMATION_TEST(FCameraRigAssetBuilderSimpleAllocationTest, "System.Engine.GameplayCameras.CameraRigAssetBuilder.SimpleAllocation", EAutomationTestFlags::EditorContext | EAutomationTestFlags::EngineFilter)
+bool FCameraRigAssetBuilderSimpleAllocationTest::RunTest(const FString& Parameters)
+{
+	using namespace UE::Cameras;
+	using namespace UE::Cameras::Test;
+
+	UCameraRigAsset* CameraRig = FCameraRigAssetTestBuilder()
+		.MakeRootNode<UArrayCameraNode>()
+			.AddChild<UOffsetCameraNode>(&UArrayCameraNode::Children).Done()
+			.Done()
+		.Get();
+
+	UTEST_EQUAL("No evaluator allocation info", CameraRig->AllocationInfo.EvaluatorInfo.TotalSizeof, 0);
+	CameraRig->BuildCameraRig();
+
+	int32 TotalSizeof = GArrayCameraNodeEvaluatorSizeof;
+	TotalSizeof = Align(TotalSizeof, GOffsetCameraNodeEvaluatorAlignof) + GOffsetCameraNodeEvaluatorSizeof;
+	UTEST_EQUAL("Evaluator allocation info", CameraRig->AllocationInfo.EvaluatorInfo.TotalSizeof, TotalSizeof);
+
+	return true;
+}
+
+IMPLEMENT_SIMPLE_AUTOMATION_TEST(FCameraRigAssetBuilderSimpleParameterTest, "System.Engine.GameplayCameras.CameraRigAssetBuilder.SimpleParameter", EAutomationTestFlags::EditorContext | EAutomationTestFlags::EngineFilter)
+bool FCameraRigAssetBuilderSimpleParameterTest::RunTest(const FString& Parameters)
+{
+	using namespace UE::Cameras::Test;
+
+	UOffsetCameraNode* OffsetNode = nullptr;
+	UCameraRigAsset* CameraRig = FCameraRigAssetTestBuilder(TEXT("SimpleTest"))
+		.MakeRootNode<UArrayCameraNode>()
+			.AddChild<UOffsetCameraNode>(&UArrayCameraNode::Children)
+				.Pin(OffsetNode)
+				.Done()
+			.Done()
+		.ExposeParameter(TEXT("Test"), OffsetNode, GET_MEMBER_NAME_CHECKED(UOffsetCameraNode, TranslationOffset))
+		.Get();
+
+	CameraRig->BuildCameraRig();
+
+	UCameraRigInterfaceParameter* Parameter = CameraRig->Interface.InterfaceParameters[0];
+	UTEST_EQUAL("Test parameter", Parameter->InterfaceParameterName, TEXT("Test"));
+	UTEST_NOT_NULL("Test parameter variable", Parameter->PrivateVariable.Get());
+	UTEST_EQUAL("Test parameter variable name", Parameter->PrivateVariable->GetName(), "Override_SimpleTest_Test");
+	UTEST_EQUAL("Test node parameter", (UCameraVariableAsset*)OffsetNode->TranslationOffset.Variable, Parameter->PrivateVariable.Get());
+
+	return true;
+}
+
+IMPLEMENT_SIMPLE_AUTOMATION_TEST(FCameraRigAssetBuilderReassignParameterTest, "System.Engine.GameplayCameras.CameraRigAssetBuilder.ReassignParameter", EAutomationTestFlags::EditorContext | EAutomationTestFlags::EngineFilter)
+bool FCameraRigAssetBuilderReassignParameterTest::RunTest(const FString& Parameters)
+{
+	using namespace UE::Cameras::Test;
+
+	UOffsetCameraNode* OffsetNode = nullptr;
+	ULensParametersCameraNode* LensParametersNode = nullptr;
+	UCameraRigAsset* CameraRig = FCameraRigAssetTestBuilder(TEXT("SimpleTest"))
+		.MakeRootNode<UArrayCameraNode>()
+			.AddChild<UOffsetCameraNode>(&UArrayCameraNode::Children)
+				.Pin(OffsetNode)
+				.Done()
+			.AddChild<ULensParametersCameraNode>(&UArrayCameraNode::Children)
+				.Pin(LensParametersNode)
+				.Done()
+			.Done()
+		.ExposeParameter(TEXT("Test1"), OffsetNode, GET_MEMBER_NAME_CHECKED(UOffsetCameraNode, TranslationOffset))
+		.ExposeParameter(TEXT("Test2"), LensParametersNode, GET_MEMBER_NAME_CHECKED(ULensParametersCameraNode, FocalLength))
+		.ExposeParameter(TEXT("Test3"), LensParametersNode, GET_MEMBER_NAME_CHECKED(ULensParametersCameraNode, Aperture))
+		.Get();
+
+	UCameraRigInterfaceParameter* Test1Parameter = CameraRig->Interface.InterfaceParameters[0];
+	UCameraRigInterfaceParameter* Test2Parameter = CameraRig->Interface.InterfaceParameters[1];
+	UCameraRigInterfaceParameter* Test3Parameter = CameraRig->Interface.InterfaceParameters[2];
+
+	CameraRig->BuildCameraRig();
+
+	{
+		UTEST_EQUAL_EXPR(Test1Parameter->PrivateVariable->GetName(), "Override_SimpleTest_Test1");
+		UTEST_TRUE_EXPR(Test1Parameter->PrivateVariable->IsA<UVector3dCameraVariable>());
+		UTEST_EQUAL_EXPR((UCameraVariableAsset*)OffsetNode->TranslationOffset.Variable, Test1Parameter->PrivateVariable.Get());
+
+		UTEST_EQUAL_EXPR(Test2Parameter->PrivateVariable->GetName(), "Override_SimpleTest_Test2");
+		UTEST_TRUE_EXPR(Test2Parameter->PrivateVariable->IsA<UFloatCameraVariable>());
+		UTEST_EQUAL_EXPR((UCameraVariableAsset*)LensParametersNode->FocalLength.Variable, Test2Parameter->PrivateVariable.Get());
+
+		UTEST_EQUAL_EXPR(Test3Parameter->PrivateVariable->GetName(), "Override_SimpleTest_Test3");
+		UTEST_TRUE_EXPR(Test3Parameter->PrivateVariable->IsA<UFloatCameraVariable>());
+		UTEST_EQUAL_EXPR((UCameraVariableAsset*)LensParametersNode->Aperture.Variable, Test3Parameter->PrivateVariable.Get());
+	}
+
+	Test1Parameter->Target = LensParametersNode;
+	Test1Parameter->TargetPropertyName = GET_MEMBER_NAME_CHECKED(ULensParametersCameraNode, FocalLength);
+	Test2Parameter->Target = LensParametersNode;
+	Test2Parameter->TargetPropertyName = GET_MEMBER_NAME_CHECKED(ULensParametersCameraNode, Aperture);
+	Test3Parameter->Target = OffsetNode;
+	Test3Parameter->TargetPropertyName = GET_MEMBER_NAME_CHECKED(UOffsetCameraNode, TranslationOffset);
+
+	CameraRig->BuildCameraRig();
+
+	{
+		UTEST_EQUAL_EXPR(Test1Parameter->PrivateVariable->GetName(), "Override_SimpleTest_Test1");
+		UTEST_TRUE_EXPR(Test1Parameter->PrivateVariable->IsA<UFloatCameraVariable>());
+		UTEST_EQUAL_EXPR((UCameraVariableAsset*)LensParametersNode->FocalLength.Variable, Test1Parameter->PrivateVariable.Get());
+
+		UTEST_EQUAL_EXPR(Test2Parameter->PrivateVariable->GetName(), "Override_SimpleTest_Test2");
+		UTEST_TRUE_EXPR(Test2Parameter->PrivateVariable->IsA<UFloatCameraVariable>());
+		UTEST_EQUAL_EXPR((UCameraVariableAsset*)LensParametersNode->Aperture.Variable, Test2Parameter->PrivateVariable.Get());
+
+		UTEST_EQUAL_EXPR(Test3Parameter->PrivateVariable->GetName(), "Override_SimpleTest_Test3");
+		UTEST_TRUE_EXPR(Test3Parameter->PrivateVariable->IsA<UVector3dCameraVariable>());
+		UTEST_EQUAL_EXPR((UCameraVariableAsset*)OffsetNode->TranslationOffset.Variable, Test3Parameter->PrivateVariable.Get());
+	}
+
+	return true;
+}
+
+#undef LOCTEXT_NAMESPACE
+
diff --git a/Engine/Plugins/Cameras/GameplayCameras/Source/GameplayCameras/Private/Tests/GameplayCamerasTestBuilder.cpp b/Engine/Plugins/Cameras/GameplayCameras/Source/GameplayCameras/Private/Tests/GameplayCamerasTestBuilder.cpp
new file mode 100644
--- /dev/null
+++ b/Engine/Plugins/Cameras/GameplayCameras/Source/GameplayCameras/Private/Tests/GameplayCamerasTestBuilder.cpp
@@ -0,0 +1,44 @@
+// Copyright Epic Games, Inc. All Rights Reserved.
+
+#include "Tests/GameplayCamerasTestBuilder.h"
+
+namespace UE::Cameras::Test
+{
+
+FCameraRigAssetTestBuilder::FCameraRigAssetTestBuilder(FName Name, UObject* Outer)
+{
+	if (Outer == nullptr)
+	{
+		Outer = GetTransientPackage();
+	}
+
+	CameraRig = NewObject<UCameraRigAsset>(Outer, Name);
+	TCameraObjectInitializer<UCameraRigAsset>::SetObject(CameraRig);
+}
+
+TCameraRigTransitionTestBuilder<FCameraRigAssetTestBuilder> FCameraRigAssetTestBuilder::AddEnterTransition()
+{
+	TCameraRigTransitionTestBuilder<ThisType> TransitionBuilder(*this, CameraRig);
+	CameraRig->EnterTransitions.Add(TransitionBuilder.Get());
+	return TransitionBuilder;
+}
+
+TCameraRigTransitionTestBuilder<FCameraRigAssetTestBuilder> FCameraRigAssetTestBuilder::AddExitTransition()
+{
+	TCameraRigTransitionTestBuilder<ThisType> TransitionBuilder(*this, CameraRig);
+	CameraRig->ExitTransitions.Add(TransitionBuilder.Get());
+	return TransitionBuilder;
+}
+
+FCameraRigAssetTestBuilder& FCameraRigAssetTestBuilder::ExposeParameter(const FString& ParameterName, UCameraNode* Target, FName TargetPropertyName)
+{
+	UCameraRigInterfaceParameter* InterfaceParameter = NewObject<UCameraRigInterfaceParameter>(CameraRig);
+	InterfaceParameter->InterfaceParameterName = ParameterName;
+	InterfaceParameter->Target = Target;
+	InterfaceParameter->TargetPropertyName = TargetPropertyName;
+	CameraRig->Interface.InterfaceParameters.Add(InterfaceParameter);
+	return *this;
+}
+
+}  // namespace UE::Cameras::Test
+
diff --git a/Engine/Plugins/Cameras/GameplayCameras/Source/GameplayCameras/Private/Tests/GameplayCamerasTestBuilder.h b/Engine/Plugins/Cameras/GameplayCameras/Source/GameplayCameras/Private/Tests/GameplayCamerasTestBuilder.h
new file mode 100644
--- /dev/null
+++ b/Engine/Plugins/Cameras/GameplayCameras/Source/GameplayCameras/Private/Tests/GameplayCamerasTestBuilder.h
@@ -0,0 +1,312 @@
+// Copyright Epic Games, Inc. All Rights Reserved.
+
+#pragma once
+
+#include "Core/BlendCameraNode.h"
+#include "Core/CameraNode.h"
+#include "Core/CameraRigAsset.h"
+#include "Core/CameraRigTransition.h"
+#include "Nodes/Common/ArrayCameraNode.h"
+#include "Templates/PointerIsConvertibleFromTo.h"
+#include "Templates/UnrealTypeTraits.h"
+#include "UObject/Package.h"
+
+#include <type_traits>
+
+namespace UE::Cameras::Test
+{
+
+class FCameraRigAssetTestBuilder;
+
+/**
+ * Template mix-in for adding "go back to parent" support to a builder class.
+ */
+template<typename ParentType>
+struct TScopedConstruction
+{
+	TScopedConstruction(ParentType& InParent)
+		: Parent(InParent)
+	{}
+
+	/** Return the parent builder instance. */
+	ParentType& Done() { return Parent; }
+
+protected:
+
+	ParentType& Parent;
+};
+
+/**
+ * A generic utility class that defines a fluent interface for setting properties and adding items to
+ * array properties on a given object.
+ */
+template<typename ObjectType>
+struct TCameraObjectInitializer
+{
+	/** Sets a value on the given public property (via its member field). */
+	template<typename PropertyType>
+	TCameraObjectInitializer<ObjectType>& Set(PropertyType ObjectType::*Field, typename TCallTraits<PropertyType>::ParamType Value)
+	{
+		PropertyType& FieldPtr = (Object->*Field);
+		FieldPtr = Value;
+	}
+	
+	/** Adds an item to a given public array property (via its member field). */
+	template<typename ItemType>
+	TCameraObjectInitializer<ObjectType>& Add(TArray<ItemType> ObjectType::*Field, typename TCallTraits<ItemType>::ParamType NewItem)
+	{
+		TArray<ItemType>& ArrayPtr = (Object->*Field);
+		ArrayPtr.Add(NewItem);
+	}
+
+protected:
+
+	void SetObject(ObjectType* InObject)
+	{
+		Object = InObject;
+	}
+
+private:
+
+	ObjectType* Object = nullptr;
+};
+
+/**
+ * A builder class for camera nodes.
+ */
+template<
+	typename ParentType,
+	typename NodeType,
+	typename V = std::enable_if_t<TPointerIsConvertibleFromTo<NodeType, UCameraNode>::Value>
+	>
+class TCameraNodeTestBuilder 
+	: public TScopedConstruction<ParentType>
+	, public TCameraObjectInitializer<NodeType>
+{
+public:
+
+	using ThisType = TCameraNodeTestBuilder<ParentType, NodeType, V>;
+
+	/** Creates a new instance of this builder class. */
+	TCameraNodeTestBuilder(ParentType& InParent, UObject* Outer = nullptr)
+		: TScopedConstruction<ParentType>(InParent)
+	{
+		if (Outer == nullptr)
+		{
+			Outer = GetTransientPackage();
+		}
+		CameraNode = NewObject<NodeType>(Outer);
+		TCameraObjectInitializer<NodeType>::SetObject(CameraNode);
+	}
+
+	/** Gets the built camera node. */
+	NodeType* Get() const { return CameraNode; }
+
+	/** Pins the built camera node to a given pointer, for being able to later refer to it. */
+	ThisType& Pin(NodeType*& OutPtr) { OutPtr = CameraNode; return *this; }
+
+	/** Sets the value of a camera parameter field on the camera node. */
+	template<typename ParameterType>
+	ThisType& SetParameter(
+			ParameterType NodeType::*ParameterField,
+			typename TCallTraits<typename ParameterType::ValueType>::ParamType Value)
+	{
+		ParameterType& ParameterRef = (CameraNode->*ParameterField);
+		ParameterRef.Value = Value;
+		return *this;
+	}
+
+	/**
+	 * Adds a child camera node via a public array member field on the camera node.
+	 * Returns a builder for the child. You can go back to the current builder by
+	 * calling Done() on the child builder.
+	 */
+	template<
+		typename ChildNodeType, 
+		typename ArrayItemType,
+		typename = std::enable_if_t<TPointerIsConvertibleFromTo<ChildNodeType, ArrayItemType>::Value>
+		>
+	TCameraNodeTestBuilder<ThisType, ChildNodeType>
+	AddChild(TArray<TObjectPtr<ArrayItemType>> NodeType::*ArrayField)
+	{
+		TCameraNodeTestBuilder<ThisType, ChildNodeType> ChildBuilder(*this, CameraNode->GetOuter());
+		TArray<TObjectPtr<ArrayItemType>>& ArrayRef = (CameraNode->*ArrayField);
+		ArrayRef.Add(ChildBuilder.Get());
+		return ChildBuilder;
+	}
+
+	/** 
+	 * Casting operator that returns a builder for the same camera node, but typed
+	 * around a parent class of the camera node's class. Mostly useful for implicit casting
+	 * when using AddChild().
+	 */
+	template<
+		typename OtherNodeType,
+		typename = std::enable_if_t<TPointerIsConvertibleFromTo<NodeType, OtherNodeType>::Value>
+		>
+	operator TCameraNodeTestBuilder<ParentType, OtherNodeType>() const
+	{
+		return TCameraNodeTestBuilder<ParentType, OtherNodeType>(
+				EForceReuseCameraNode::Yes, 
+				TScopedConstruction<ParentType>::Parent, 
+				CameraNode);
+	}
+
+private:
+
+	enum class EForceReuseCameraNode { Yes };
+
+	TCameraNodeTestBuilder(EForceReuseCameraNode ForceReuse, ParentType& InParent, NodeType* ExistingCameraNode)
+		: TScopedConstruction<ParentType>(InParent)
+		, CameraNode(ExistingCameraNode)
+	{
+		TCameraObjectInitializer<NodeType>::SetObject(CameraNode);
+	}
+
+private:
+
+	NodeType* CameraNode;
+};
+
+/**
+ * Builder class for camera rig transitions.
+ */
+template<typename ParentType>
+class TCameraRigTransitionTestBuilder 
+	: public TScopedConstruction<ParentType>
+	, public TCameraObjectInitializer<UCameraRigTransition>
+{
+public:
+
+	using ThisType = TCameraRigTransitionTestBuilder<ParentType>;
+
+	/** Creates a new instance of this builder class. */
+	TCameraRigTransitionTestBuilder(ParentType& InParent, UObject* Outer = nullptr)
+		: TScopedConstruction<ParentType>(InParent)
+	{
+		if (Outer == nullptr)
+		{
+			Outer = GetTransientPackage();
+		}
+
+		Transition = NewObject<UCameraRigTransition>(Outer);
+		TCameraObjectInitializer<UCameraRigTransition>::SetObject(Transition);
+	}
+
+	/** Gets the built transition object. */
+	UCameraRigTransition* Get() const { return Transition; }
+
+	/** 
+	 * Creates a blend node of the given type, and returns a builder for it.
+	 * You can go back to this transition builder by calling Done() on the blend builder.
+	 */
+	template<
+		typename BlendType,
+		typename = std::enable_if_t<TPointerIsConvertibleFromTo<BlendType, UBlendCameraNode>::Value>
+		>
+	TCameraNodeTestBuilder<ThisType, BlendType> MakeBlend()
+	{
+		TCameraNodeTestBuilder<ThisType, BlendType> BlendBuilder(*this, Transition->GetOuter());
+		Transition->Blend = BlendBuilder.Get();
+		return BlendBuilder;
+	}
+
+	/** Adds a transition condition. */
+	template<
+		typename ConditionType,
+		typename = std::enable_if_t<TPointerIsConvertibleFromTo<ConditionType, UCameraRigTransitionCondition>::Value>
+		>
+	ThisType& AddCondition()
+	{
+		ConditionType* NewCondition = NewObject<ConditionType>(Transition->GetOuter());
+		Transition->Conditions.Add(NewCondition);
+		return *this;
+	}
+
+	/** Adds a transition condition. */
+	template<
+		typename ConditionType,
+		typename = std::enable_if_t<TPointerIsConvertibleFromTo<ConditionType, UCameraRigTransitionCondition>::Value>
+		>
+	ThisType& AddCondition(TFunction<void(ConditionType*)> SetupCallback)
+	{
+		ConditionType* NewCondition = NewObject<ConditionType>(Transition->GetOuter());
+		SetupCallback(NewCondition);
+		Transition->Conditions.Add(NewCondition);
+		return *this;
+	}
+
+private:
+
+	UCameraRigTransition* Transition;
+};
+
+/**
+ * The root builder class for building a camera rig. Follow the fluent interface to construct the
+ * hierarchy of camera nodes, add transitions, etc.
+ *
+ * For instance:
+ *
+ *		UCameraRigAsset* CameraRig = FCameraRigAssetTestBuilder(TEXT("SimpleTest"))
+ *			.MakeRootNode<UArrayCameraNode>()
+ *				.AddChild<UOffsetCameraNode>(&UArrayCameraNode::Children)
+ *					.SetParameter(&UOffsetCameraNode::TranslationOffset, FVector3d{ 1, 0, 0 })
+ *					.Done()
+ *				.AddChild<ULensParametersCameraNode>(&UArrayCameraNode::Children)
+ *					.SetParameter(&ULensParametersCameraNode::FocalLenght, 18.f)
+ *					.Done()
+ *				.Done()
+ *			.AddEnterTransition()
+ *				.MakeBlend<USmoothBlendCameraNode>()
+ *				.Done()
+ *			.Get();
+ */
+class FCameraRigAssetTestBuilder : public TCameraObjectInitializer<UCameraRigAsset>
+{
+public:
+
+	using ThisType = FCameraRigAssetTestBuilder;
+
+	FCameraRigAssetTestBuilder(FName Name = NAME_None, UObject* Outer = nullptr);
+
+	/** Gets the built camera rig. */
+	UCameraRigAsset* Get() { return CameraRig; }
+
+	/**
+	 * Creates a new camera node and sets it as the root node of the rig.
+	 * Returns the builder for the root camera node. You can come back to the rig builder
+	 * by calling Done() on the node builder.
+	 */
+	template<typename NodeType>
+	TCameraNodeTestBuilder<ThisType, NodeType> MakeRootNode()
+	{
+		TCameraNodeTestBuilder<ThisType, NodeType> NodeBuilder(*this, CameraRig);
+		CameraRig->RootNode = NodeBuilder.Get();
+		return NodeBuilder;
+	}
+
+	/**
+	 * Adds a new enter transition and returns a builder for it. You can come back to the
+	 * rig builder by calling Done() on the transition builder.
+	 */
+	TCameraRigTransitionTestBuilder<ThisType> AddEnterTransition();
+	/**
+	 * Adds a new exit transition and returns a builder for it. You can come back to the
+	 * rig builder by calling Done() on the transition builder.
+	 */
+	TCameraRigTransitionTestBuilder<ThisType> AddExitTransition();
+
+	/**
+	 * Creates a new exposed rig parameter and hooks it up to the given camera node's property.
+	 * When building the node hierarchy, you can use the Pin() method on the node builders to
+	 * save a pointer to nodes you need for ExposeParameter().
+	 */
+	FCameraRigAssetTestBuilder& ExposeParameter(const FString& ParameterName, UCameraNode* Target, FName TargetPropertyName);
+
+private:
+
+	UCameraRigAsset* CameraRig;
+};
+
+}  // namespace UE::Cameras::Test
+
