Cameras: add missing return value to test objet builder

diff --git a/Engine/Plugins/Cameras/GameplayCameras/Source/GameplayCameras/Private/Tests/GameplayCamerasTestBuilder.h b/Engine/Plugins/Cameras/GameplayCameras/Source/GameplayCameras/Private/Tests/GameplayCamerasTestBuilder.h
--- a/Engine/Plugins/Cameras/GameplayCameras/Source/GameplayCameras/Private/Tests/GameplayCamerasTestBuilder.h
+++ b/Engine/Plugins/Cameras/GameplayCameras/Source/GameplayCameras/Private/Tests/GameplayCamerasTestBuilder.h
@@ -49,6 +49,7 @@ struct TCameraObjectInitializer
 	{
 		PropertyType& FieldPtr = (Object->*Field);
 		FieldPtr = Value;
+		return *this;
 	}
 	
 	/** Adds an item to a given public array property (via its member field). */
@@ -57,6 +58,7 @@ struct TCameraObjectInitializer
 	{
 		TArray<ItemType>& ArrayPtr = (Object->*Field);
 		ArrayPtr.Add(NewItem);
+		return *this;
 	}
 
 protected:
