Cameras: make all the activate/auto-activate logic consistent

Note that we can't use the existing auto-activation path for actor components because those get activated too early -- we don't get any local player controller at that point. So the auto-activate properties on our components act on BeginPlay.

diff --git a/Engine/Plugins/Cameras/GameplayCameras/Source/GameplayCameras/Private/GameFramework/GameplayCameraActor.cpp b/Engine/Plugins/Cameras/GameplayCameras/Source/GameplayCameras/Private/GameFramework/GameplayCameraActor.cpp
--- a/Engine/Plugins/Cameras/GameplayCameras/Source/GameplayCameras/Private/GameFramework/GameplayCameraActor.cpp
+++ b/Engine/Plugins/Cameras/GameplayCameras/Source/GameplayCameras/Private/GameFramework/GameplayCameraActor.cpp
@@ -15,11 +15,8 @@
 AGameplayCameraActor::AGameplayCameraActor(const FObjectInitializer& ObjectInit)
 	: Super(ObjectInit)
 {
-	SceneComponent = CreateDefaultSubobject<USceneComponent>(TEXT("SceneComponent"));
-	RootComponent = SceneComponent;
-
 	CameraComponent = CreateDefaultSubobject<UGameplayCameraComponent>(TEXT("CameraComponent"));
-	CameraComponent->SetupAttachment(SceneComponent);
+	RootComponent = CameraComponent;
 }
 
 USceneComponent* AGameplayCameraActor::GetDefaultAttachComponent() const
diff --git a/Engine/Plugins/Cameras/GameplayCameras/Source/GameplayCameras/Private/GameFramework/GameplayCameraComponent.cpp b/Engine/Plugins/Cameras/GameplayCameras/Source/GameplayCameras/Private/GameFramework/GameplayCameraComponent.cpp
--- a/Engine/Plugins/Cameras/GameplayCameras/Source/GameplayCameras/Private/GameFramework/GameplayCameraComponent.cpp
+++ b/Engine/Plugins/Cameras/GameplayCameras/Source/GameplayCameras/Private/GameFramework/GameplayCameraComponent.cpp
@@ -10,6 +10,8 @@
 #include "GameFramework/CameraEvaluationResultInterop.h"
 #include "GameFramework/GameplayCameraSystemActor.h"
 #include "GameFramework/GameplayCameraSystemComponent.h"
+#include "GameplayCameras.h"
+#include "Kismet/GameplayStatics.h"
 #include "Logging/MessageLog.h"
 #include "UObject/ConstructorHelpers.h"
 #include "UObject/UnrealType.h"
@@ -37,27 +39,22 @@ UGameplayCameraComponent::UGameplayCameraComponent(const FObjectInitializer& Obj
 
 void UGameplayCameraComponent::ActivateCamera(int32 PlayerIndex)
 {
-	if (ActivatedForPlayerIndex >= 0 && ActivatedForPlayerIndex != PlayerIndex)
+	if (ActivatedForPlayerIndex == PlayerIndex)
 	{
-		DeactivateCamera();
+		return;
 	}
 
-	UWorld* World = GetWorld();
-	if (!ensure(World))
+	if (ActivatedForPlayerIndex >= 0)
 	{
-		return;
+		DeactivateCamera();
 	}
 
-	PlayerIndex = FMath::Max(0, PlayerIndex);
-	if (ensure(PlayerIndex < World->GetNumPlayerControllers()))
+	Activate();
+
+	if (APlayerController* PC = UGameplayStatics::GetPlayerController(this, PlayerIndex))
 	{
-		FConstPlayerControllerIterator It = World->GetPlayerControllerIterator();
-		It += PlayerIndex;
-		if (APlayerController* PlayerController = It->Get())
-		{
-			ActivateCamera(PlayerController);
-			ActivatedForPlayerIndex = PlayerIndex;
-		}
+		ActivateCamera(PC);
+		ActivatedForPlayerIndex = PlayerIndex;
 	}
 }
 
@@ -68,42 +65,37 @@ void UGameplayCameraComponent::DeactivateCamera()
 		return;
 	}
 
-	UWorld* World = GetWorld();
-	if (!ensure(World))
+	if (APlayerController* PC = UGameplayStatics::GetPlayerController(this, ActivatedForPlayerIndex))
 	{
-		return;
-	}
-
-	if (ensure(ActivatedForPlayerIndex < World->GetNumPlayerControllers()))
-	{
-		FConstPlayerControllerIterator It = World->GetPlayerControllerIterator();
-		It += ActivatedForPlayerIndex;
-		if (APlayerController* PlayerController = It->Get())
-		{
-			DeactivateCamera(PlayerController);
-		}
+		DeactivateCamera(PC);
 	}
 
 	ActivatedForPlayerIndex = INDEX_NONE;
+
+	Deactivate();
 }
 
 void UGameplayCameraComponent::ActivateCamera(APlayerController* PlayerController)
 {
 	using namespace UE::Cameras;
 
-	if (!ensure(PlayerController && PlayerController->PlayerCameraManager))
+	if (!ensureMsgf(
+				PlayerController && PlayerController->PlayerCameraManager,
+				TEXT("Can't activate gameplay camera component: invalid player controller!")))
 	{
 		return;
 	}
 	
 	AGameplayCameraSystemActor* CameraSystem = Cast<AGameplayCameraSystemActor>(PlayerController->PlayerCameraManager->GetViewTarget());
-	if (!ensure(CameraSystem))
+	if (!CameraSystem)
 	{
+		UE_LOG(LogCameraSystem, Error, TEXT("Can't activate gameplay camera component: no camera system found on the view target!"));
 		return;
 	}
 
-	if (!ensure(Camera))
+	if (!Camera)
 	{
+		UE_LOG(LogCameraSystem, Error, TEXT("Can't activate gameplay camera component: no camera asset was set!"));
 		return;
 	}
 
@@ -177,6 +169,24 @@ void UGameplayCameraComponent::OnRegister()
 #endif	// WITH_EDITORONLY_DATA
 }
 
+void UGameplayCameraComponent::Deactivate()
+{
+	DeactivateCamera();
+
+	Super::Deactivate();
+}
+
+void UGameplayCameraComponent::BeginPlay()
+{
+	Super::BeginPlay();
+
+	if (AutoActivateForPlayer != EAutoReceiveInput::Disabled && GetNetMode() != NM_DedicatedServer)
+	{
+		const int32 PlayerIndex = AutoActivateForPlayer.GetIntValue() - 1;
+		ActivateCamera(PlayerIndex);
+	}
+}
+
 void UGameplayCameraComponent::TickComponent(float DeltaTime, enum ELevelTick TickType, FActorComponentTickFunction *ThisTickFunction)
 {
 	Super::TickComponent(DeltaTime, TickType, ThisTickFunction);
diff --git a/Engine/Plugins/Cameras/GameplayCameras/Source/GameplayCameras/Private/GameFramework/GameplayCameraSystemActor.cpp b/Engine/Plugins/Cameras/GameplayCameras/Source/GameplayCameras/Private/GameFramework/GameplayCameraSystemActor.cpp
--- a/Engine/Plugins/Cameras/GameplayCameras/Source/GameplayCameras/Private/GameFramework/GameplayCameraSystemActor.cpp
+++ b/Engine/Plugins/Cameras/GameplayCameras/Source/GameplayCameras/Private/GameFramework/GameplayCameraSystemActor.cpp
@@ -6,7 +6,6 @@
 #include "Engine/World.h"
 #include "GameFramework/GameplayCameraSystemComponent.h"
 #include "GameFramework/PlayerController.h"
-#include "Kismet/GameplayStatics.h"
 
 #include UE_INLINE_GENERATED_CPP_BY_NAME(GameplayCameraSystemActor)
 
@@ -19,17 +18,6 @@ AGameplayCameraSystemActor::AGameplayCameraSystemActor(const FObjectInitializer&
 	RootComponent = CameraSystemComponent;
 }
 
-void AGameplayCameraSystemActor::BeginPlay()
-{
-	Super::BeginPlay();
-
-	if (AutoActivateForPlayer != EAutoReceiveInput::Disabled && GetNetMode() != NM_Client)
-	{
-		const int32 PlayerIndex = AutoActivateForPlayer.GetIntValue() - 1;
-		ActivateForPlayer(PlayerIndex);
-	}
-}
-
 void AGameplayCameraSystemActor::BecomeViewTarget(APlayerController* PC)
 {
 	Super::BecomeViewTarget(PC);
@@ -49,14 +37,5 @@ void AGameplayCameraSystemActor::EndViewTarget(APlayerController* PC)
 	Super::EndViewTarget(PC);
 }
 
-void AGameplayCameraSystemActor::ActivateForPlayer(int32 PlayerIndex)
-{
-	APlayerController* PC = UGameplayStatics::GetPlayerController(this, PlayerIndex);
-	if (PC)
-	{
-		PC->SetViewTarget(this);
-	}
-}
-
 #undef LOCTEXT_NAMESPACE
 
diff --git a/Engine/Plugins/Cameras/GameplayCameras/Source/GameplayCameras/Private/GameFramework/GameplayCameraSystemComponent.cpp b/Engine/Plugins/Cameras/GameplayCameras/Source/GameplayCameras/Private/GameFramework/GameplayCameraSystemComponent.cpp
--- a/Engine/Plugins/Cameras/GameplayCameras/Source/GameplayCameras/Private/GameFramework/GameplayCameraSystemComponent.cpp
+++ b/Engine/Plugins/Cameras/GameplayCameras/Source/GameplayCameras/Private/GameFramework/GameplayCameraSystemComponent.cpp
@@ -8,6 +8,7 @@
 #include "Engine/StaticMesh.h"
 #include "Engine/World.h"
 #include "GameFramework/Controller.h"
+#include "Kismet/GameplayStatics.h"
 #include "Logging/MessageLog.h"
 #include "UObject/ConstructorHelpers.h"
 #include "UObject/UnrealType.h"
@@ -89,6 +90,74 @@ void UGameplayCameraSystemComponent::OnRegister()
 #endif	// WITH_EDITORONLY_DATA
 }
 
+void UGameplayCameraSystemComponent::Deactivate()
+{
+	DeactivateCameraSystem();
+
+	Super::Deactivate();
+}
+
+void UGameplayCameraSystemComponent::ActivateCameraSystem(int32 PlayerIndex)
+{
+	if (ActivatedForPlayerIndex == PlayerIndex)
+	{
+		return;
+	}
+
+	if (ActivatedForPlayerIndex >= 0)
+	{
+		DeactivateCameraSystem();
+	}
+
+	AActor* OwningActor = GetOwner();
+	if (!OwningActor)
+	{
+		UE_LOG(LogCameraSystem, Error, TEXT("Can't activate gameplay camera system: no owning actor found!"));
+		return;
+	}
+
+	APlayerController* PC = UGameplayStatics::GetPlayerController(this, PlayerIndex);
+	if (!PC)
+	{
+		UE_LOG(LogCameraSystem, Error, TEXT("Can't activate gameplay camera system: no player controller found!"));
+		return;
+	}
+
+	Activate();
+
+	PC->SetViewTarget(OwningActor);
+	ActivatedForPlayerIndex = PlayerIndex;
+}
+
+void UGameplayCameraSystemComponent::DeactivateCameraSystem(AActor* NextViewTarget)
+{
+	if (ActivatedForPlayerIndex < 0)
+	{
+		return;
+	}
+
+	APlayerController* PC = UGameplayStatics::GetPlayerController(this, ActivatedForPlayerIndex);
+	if (PC)
+	{
+		PC->SetViewTarget(NextViewTarget);
+	}
+
+	ActivatedForPlayerIndex = INDEX_NONE;
+
+	Deactivate();
+}
+
+void UGameplayCameraSystemComponent::BeginPlay()
+{
+	Super::BeginPlay();
+
+	if (AutoActivateForPlayer != EAutoReceiveInput::Disabled && GetNetMode() != NM_DedicatedServer)
+	{
+		const int32 PlayerIndex = AutoActivateForPlayer.GetIntValue() - 1;
+		ActivateCameraSystem(PlayerIndex);
+	}
+}
+
 void UGameplayCameraSystemComponent::TickComponent(float DeltaTime, ELevelTick TickType, FActorComponentTickFunction* ThisTickFunction)
 {
 	Super::TickComponent(DeltaTime, TickType, ThisTickFunction);
diff --git a/Engine/Plugins/Cameras/GameplayCameras/Source/GameplayCameras/Public/GameFramework/GameplayCameraActor.h b/Engine/Plugins/Cameras/GameplayCameras/Source/GameplayCameras/Public/GameFramework/GameplayCameraActor.h
--- a/Engine/Plugins/Cameras/GameplayCameras/Source/GameplayCameras/Public/GameFramework/GameplayCameraActor.h
+++ b/Engine/Plugins/Cameras/GameplayCameras/Source/GameplayCameras/Public/GameFramework/GameplayCameraActor.h
@@ -24,20 +24,18 @@ public:
 
 	AGameplayCameraActor(const FObjectInitializer& ObjectInit);
 
-	virtual USceneComponent* GetDefaultAttachComponent() const override;
-
 public:
 
-	UFUNCTION(BlueprintGetter, Category=Camera)
-	USceneComponent* GetSceneComponent() const { return SceneComponent; }
-
+	/** Gets the camera component. */
 	UFUNCTION(BlueprintGetter, Category=Camera)
 	UGameplayCameraComponent* GetCameraComponent() const { return CameraComponent; }
 
-private:
+public:
+
+	// AActor interface.
+	virtual USceneComponent* GetDefaultAttachComponent() const override;
 
-	UPROPERTY(VisibleAnywhere, Category=Camera, BluePrintGetter="GetSceneComponent")
-	TObjectPtr<USceneComponent> SceneComponent;
+private:
 
 	UPROPERTY(VisibleAnywhere, Category=Camera, BlueprintGetter="GetCameraComponent", meta=(ExposeFunctionCategories="Camera"))
 	TObjectPtr<UGameplayCameraComponent> CameraComponent;
diff --git a/Engine/Plugins/Cameras/GameplayCameras/Source/GameplayCameras/Public/GameFramework/GameplayCameraComponent.h b/Engine/Plugins/Cameras/GameplayCameras/Source/GameplayCameras/Public/GameFramework/GameplayCameraComponent.h
--- a/Engine/Plugins/Cameras/GameplayCameras/Source/GameplayCameras/Public/GameFramework/GameplayCameraComponent.h
+++ b/Engine/Plugins/Cameras/GameplayCameras/Source/GameplayCameras/Public/GameFramework/GameplayCameraComponent.h
@@ -32,12 +32,22 @@ public:
 
 	UGameplayCameraComponent(const FObjectInitializer& ObjectInit);
 
+public:
+
+	/** 
+	 * Activates the camera for the given player.
+	 * This looks up the current player camera manager and/or view target in order to find
+	 * the active camera system for the given player. If found, it adds its own camera asset
+	 * as the active one.
+	 */
 	UFUNCTION(BlueprintCallable, Category=Camera)
 	GAMEPLAYCAMERAS_API void ActivateCamera(int32 PlayerIndex = 0);
 
+	/** Deactivates the camera for the last player it was activated for. */
 	UFUNCTION(BlueprintCallable, Category=Camera)
 	GAMEPLAYCAMERAS_API void DeactivateCamera();
 
+	/** Gets the initial evaluation result for this component's context. */
 	UFUNCTION(BlueprintPure, Category=Camera)
 	GAMEPLAYCAMERAS_API UCameraEvaluationResultInterop* GetInitialResult() const;
 
@@ -45,6 +55,8 @@ public:
 
 	// UActorComponent interface
 	virtual void OnRegister() override;
+	virtual void Deactivate() override;
+	virtual void BeginPlay() override;
 	virtual void TickComponent(float DeltaTime, enum ELevelTick TickType, FActorComponentTickFunction *ThisTickFunction) override;
 	virtual void OnComponentDestroyed(bool bDestroyingHierarchy) override;
 
@@ -61,9 +73,17 @@ private:
 
 public:
 
+	/** The camera asset to run. */
 	UPROPERTY(EditAnywhere, BlueprintReadWrite, Category=Camera)
 	TObjectPtr<UCameraAsset> Camera;
 
+	/**
+	 * If set, auto-activates this component's camera for the given player.
+	 * This is equivalent to calling ActivateCamera on BeginPlay.
+	 */
+	UPROPERTY(EditAnywhere, Category=Camera)
+	TEnumAsByte<EAutoReceiveInput::Type> AutoActivateForPlayer;
+
 protected:
 
 	using FGameplayCameraComponentEvaluationContext = UE::Cameras::FGameplayCameraComponentEvaluationContext;
diff --git a/Engine/Plugins/Cameras/GameplayCameras/Source/GameplayCameras/Public/GameFramework/GameplayCameraSystemActor.h b/Engine/Plugins/Cameras/GameplayCameras/Source/GameplayCameras/Public/GameFramework/GameplayCameraSystemActor.h
--- a/Engine/Plugins/Cameras/GameplayCameras/Source/GameplayCameras/Public/GameFramework/GameplayCameraSystemActor.h
+++ b/Engine/Plugins/Cameras/GameplayCameras/Source/GameplayCameras/Public/GameFramework/GameplayCameraSystemActor.h
@@ -26,16 +26,15 @@ public:
 
 	AGameplayCameraSystemActor(const FObjectInitializer& ObjectInit);
 
-	UFUNCTION(BlueprintCallable, Category=Camera)
-	void ActivateForPlayer(int32 PlayerIndex);
+public:
 
+	/** Gets the camera system component. */
 	UFUNCTION(BlueprintGetter, Category=Camera)
 	UGameplayCameraSystemComponent* GetCameraSystemComponent() const { return CameraSystemComponent; }
 
 public:
 
-	// AActor interface
-	virtual void BeginPlay() override;
+	// AActor interface.
 	virtual void BecomeViewTarget(APlayerController* PC) override;
 	virtual void CalcCamera(float DeltaTime, struct FMinimalViewInfo& OutResult) override;
 	virtual void EndViewTarget(APlayerController* PC) override;
@@ -44,8 +43,5 @@ private:
 
 	UPROPERTY(VisibleAnywhere, Category=Camera, BlueprintGetter="GetCameraSystemComponent", meta=(ExposeFunctionCategories="CameraSystem"))
 	TObjectPtr<UGameplayCameraSystemComponent> CameraSystemComponent;
-
-	UPROPERTY(EditAnywhere, Category=Camera)
-	TEnumAsByte<EAutoReceiveInput::Type> AutoActivateForPlayer;
 };
 
diff --git a/Engine/Plugins/Cameras/GameplayCameras/Source/GameplayCameras/Public/GameFramework/GameplayCameraSystemComponent.h b/Engine/Plugins/Cameras/GameplayCameras/Source/GameplayCameras/Public/GameFramework/GameplayCameraSystemComponent.h
--- a/Engine/Plugins/Cameras/GameplayCameras/Source/GameplayCameras/Public/GameFramework/GameplayCameraSystemComponent.h
+++ b/Engine/Plugins/Cameras/GameplayCameras/Source/GameplayCameras/Public/GameFramework/GameplayCameraSystemComponent.h
@@ -37,14 +37,26 @@ public:
 
 	UGameplayCameraSystemComponent(const FObjectInitializer& ObjectInit);
 
+	/** Gets the camera system evaluator. */
 	TSharedPtr<FCameraSystemEvaluator> GetCameraSystemEvaluator() { return Evaluator; }
 
+	/** Updates the camera system and returns the computed view. */
 	GAMEPLAYCAMERAS_API void GetCameraView(float DeltaTime, FMinimalViewInfo& DesiredView);
 
+	/** Sets this component's actor as the view target for the given player. */
+	UFUNCTION(BlueprintCallable, Category=Camera)
+	void ActivateCameraSystem(int32 PlayerIndex = 0);
+
+	/** Removes this component's actor from being the view target. */
+	UFUNCTION(BlueprintCallable, Category=Camera)
+	void DeactivateCameraSystem(AActor* NextViewTarget = nullptr);
+
 public:
 
 	// UActorComponent interface
 	virtual void OnRegister() override;
+	virtual void Deactivate() override;
+	virtual void BeginPlay() override;
 	virtual void TickComponent(float DeltaTime, ELevelTick TickType, FActorComponentTickFunction* ThisTickFunction) override;
 	virtual void OnComponentDestroyed(bool bDestroyingHierarchy) override;
 
@@ -68,10 +80,21 @@ private:
 	void DebugDraw(UCanvas* Canvas, APlayerController* PlayController);
 #endif  // UE_GAMEPLAY_CAMERAS_DEBUG
 
+public:
+
+	/**
+	 * If set, auto-activates the camera system for the given player.
+	 * This sets this actor as the view target, and is equivalent to calling ActivateCameraSystem on BeginPlay.
+	 */
+	UPROPERTY(EditAnywhere, Category=Camera)
+	TEnumAsByte<EAutoReceiveInput::Type> AutoActivateForPlayer;
+
 private:
 	
 	TSharedPtr<FCameraSystemEvaluator> Evaluator;
 
+	int32 ActivatedForPlayerIndex = INDEX_NONE;
+
 #if UE_GAMEPLAY_CAMERAS_DEBUG
 	FDelegateHandle DebugDrawDelegateHandle;
 #endif  // UE_GAMEPLAY_CAMERAS_DEBUG
