Cameras: change how camera variables are managed by placing them into variable collection assets

This change removes camera variables as standalone assets and instead sets them as objects inside a "camera variable collection" asset, similar to material parameter collections for shaders.

- Change all the asset-tools/asset-factory/etc code to work with the collections instead of the variables.
- Add an editor for the collection.
- Add a details customization for camera parameters, allowing them to pick a variable from a collection.

diff --git a/Engine/Plugins/Cameras/GameplayCameras/Source/GameplayCameras/Private/Core/CameraVariableCollection.cpp b/Engine/Plugins/Cameras/GameplayCameras/Source/GameplayCameras/Private/Core/CameraVariableCollection.cpp
new file mode 100644
--- /dev/null
+++ b/Engine/Plugins/Cameras/GameplayCameras/Source/GameplayCameras/Private/Core/CameraVariableCollection.cpp
@@ -0,0 +1,11 @@
+// Copyright Epic Games, Inc. All Rights Reserved.
+
+#include "Core/CameraVariableCollection.h"
+
+#include UE_INLINE_GENERATED_CPP_BY_NAME(CameraVariableCollection)
+
+UCameraVariableCollection::UCameraVariableCollection(const FObjectInitializer& ObjectInit)
+	: Super(ObjectInit)
+{
+}
+
diff --git a/Engine/Plugins/Cameras/GameplayCameras/Source/GameplayCameras/Public/Core/CameraVariableAssets.h b/Engine/Plugins/Cameras/GameplayCameras/Source/GameplayCameras/Public/Core/CameraVariableAssets.h
--- a/Engine/Plugins/Cameras/GameplayCameras/Source/GameplayCameras/Public/Core/CameraVariableAssets.h
+++ b/Engine/Plugins/Cameras/GameplayCameras/Source/GameplayCameras/Public/Core/CameraVariableAssets.h
@@ -2,9 +2,9 @@
 
 #pragma once
 
+#include "Containers/UnrealString.h"
 #include "Core/CameraVariableTableFwd.h"
 #include "CoreTypes.h"
-#include "Engine/DataAsset.h"
 #include "Math/MathFwd.h"
 
 #include "CameraVariableAssets.generated.h"
@@ -13,7 +13,7 @@
  * The base asset class for all camera variables.
  */
 UCLASS(Abstract)
-class GAMEPLAYCAMERAS_API UCameraVariableAsset : public UDataAsset
+class GAMEPLAYCAMERAS_API UCameraVariableAsset : public UObject
 {
 	GENERATED_BODY()
 
@@ -28,6 +28,10 @@ public:
 	virtual ECameraVariableType GetVariableType() const PURE_VIRTUAL(UCameraVariableAsset::GetVariableType, return ECameraVariableType::Boolean;);
 	virtual const uint8* GetDefaultValuePtr() const PURE_VIRTUAL(UCameraVariableAsset::GetDefaultValuePtr, return nullptr;);
 
+#if WITH_EDITOR
+	virtual FString FormatDefaultValue() const PURE_VIRTUAL(UCameraVariableAsset::FormatDefaultValue, return FString(););
+#endif  // WITH_EDITOR
+
 public:
 
 	// UObject interface
@@ -42,14 +46,20 @@ private:
 
 public:
 
-	/** Whether this variable is private and shouldn't be propagated with evaluation results. */
-	UPROPERTY()
-	bool bIsPrivate = false;
+#if WITH_EDITORONLY_DATA
+	/** The name of the variable. */
+	UPROPERTY(EditAnywhere, Category=Camera)
+	FString DisplayName;
+#endif  // WITH_EDITORONLY_DATA
 
 	/** Whether this variable auto-resets to its default value every frame. */
 	UPROPERTY(EditAnywhere, Category=Camera)
 	bool bAutoReset = false;
 
+	/** Whether this variable is private and shouldn't be propagated with evaluation results. */
+	UPROPERTY()
+	bool bIsPrivate = false;
+
 private:
 
 	UPROPERTY()
@@ -69,6 +79,10 @@ public:
 	virtual ECameraVariableType GetVariableType() const override { return ECameraVariableType::Boolean; }
 	virtual const uint8* GetDefaultValuePtr() const override { return reinterpret_cast<const uint8*>(&bDefaultValue); }
 
+#if WITH_EDITOR
+	virtual FString FormatDefaultValue() const override { return LexToString(bDefaultValue); }
+#endif  // WITH_EDITOR
+
 public:
 
 	/** The default value of this variable. */
@@ -89,6 +103,10 @@ public:
 	virtual ECameraVariableType GetVariableType() const override { return ECameraVariableType::Integer32; }
 	virtual const uint8* GetDefaultValuePtr() const override { return reinterpret_cast<const uint8*>(&DefaultValue); }
 
+#if WITH_EDITOR
+	virtual FString FormatDefaultValue() const override { return LexToString(DefaultValue); }
+#endif  // WITH_EDITOR
+
 public:
 
 	/** The default value of this variable. */
@@ -109,6 +127,10 @@ public:
 	virtual ECameraVariableType GetVariableType() const override { return ECameraVariableType::Float; }
 	virtual const uint8* GetDefaultValuePtr() const override { return reinterpret_cast<const uint8*>(&DefaultValue); }
 
+#if WITH_EDITOR
+	virtual FString FormatDefaultValue() const override { return LexToString(DefaultValue); }
+#endif  // WITH_EDITOR
+
 public:
 
 	/** The default value of this variable. */
@@ -129,6 +151,10 @@ public:
 	virtual ECameraVariableType GetVariableType() const override { return ECameraVariableType::Double; }
 	virtual const uint8* GetDefaultValuePtr() const override { return reinterpret_cast<const uint8*>(&DefaultValue); }
 
+#if WITH_EDITOR
+	virtual FString FormatDefaultValue() const override { return LexToString(DefaultValue); }
+#endif  // WITH_EDITOR
+
 public:
 
 	/** The default value of this variable. */
@@ -149,6 +175,10 @@ public:
 	virtual ECameraVariableType GetVariableType() const override { return ECameraVariableType::Vector2f; }
 	virtual const uint8* GetDefaultValuePtr() const override { return reinterpret_cast<const uint8*>(&DefaultValue); }
 
+#if WITH_EDITOR
+	virtual FString FormatDefaultValue() const override { return DefaultValue.ToString(); }
+#endif  // WITH_EDITOR
+
 public:
 
 	/** The default value of this variable. */
@@ -169,6 +199,10 @@ public:
 	virtual ECameraVariableType GetVariableType() const override { return ECameraVariableType::Vector2d; }
 	virtual const uint8* GetDefaultValuePtr() const override { return reinterpret_cast<const uint8*>(&DefaultValue); }
 
+#if WITH_EDITOR
+	virtual FString FormatDefaultValue() const override { return DefaultValue.ToString(); }
+#endif  // WITH_EDITOR
+
 public:
 
 	/** The default value of this variable. */
@@ -189,6 +223,10 @@ public:
 	virtual ECameraVariableType GetVariableType() const override { return ECameraVariableType::Vector3f; }
 	virtual const uint8* GetDefaultValuePtr() const override { return reinterpret_cast<const uint8*>(&DefaultValue); }
 
+#if WITH_EDITOR
+	virtual FString FormatDefaultValue() const override { return DefaultValue.ToString(); }
+#endif  // WITH_EDITOR
+
 public:
 
 	/** The default value of this variable. */
@@ -209,6 +247,10 @@ public:
 	virtual ECameraVariableType GetVariableType() const override { return ECameraVariableType::Vector3d; }
 	virtual const uint8* GetDefaultValuePtr() const override { return reinterpret_cast<const uint8*>(&DefaultValue); }
 
+#if WITH_EDITOR
+	virtual FString FormatDefaultValue() const override { return DefaultValue.ToString(); }
+#endif  // WITH_EDITOR
+
 public:
 
 	/** The default value of this variable. */
@@ -229,6 +271,10 @@ public:
 	virtual ECameraVariableType GetVariableType() const override { return ECameraVariableType::Vector4f; }
 	virtual const uint8* GetDefaultValuePtr() const override { return reinterpret_cast<const uint8*>(&DefaultValue); }
 
+#if WITH_EDITOR
+	virtual FString FormatDefaultValue() const override { return DefaultValue.ToString(); }
+#endif  // WITH_EDITOR
+
 public:
 
 	/** The default value of this variable. */
@@ -249,6 +295,10 @@ public:
 	virtual ECameraVariableType GetVariableType() const override { return ECameraVariableType::Vector4d; }
 	virtual const uint8* GetDefaultValuePtr() const override { return reinterpret_cast<const uint8*>(&DefaultValue); }
 
+#if WITH_EDITOR
+	virtual FString FormatDefaultValue() const override { return DefaultValue.ToString(); }
+#endif  // WITH_EDITOR
+
 public:
 
 	/** The default value of this variable. */
@@ -269,6 +319,10 @@ public:
 	virtual ECameraVariableType GetVariableType() const override { return ECameraVariableType::Rotator3f; }
 	virtual const uint8* GetDefaultValuePtr() const override { return reinterpret_cast<const uint8*>(&DefaultValue); }
 
+#if WITH_EDITOR
+	virtual FString FormatDefaultValue() const override { return DefaultValue.ToString(); }
+#endif  // WITH_EDITOR
+
 public:
 
 	/** The default value of this variable. */
@@ -289,6 +343,10 @@ public:
 	virtual ECameraVariableType GetVariableType() const override { return ECameraVariableType::Rotator3d; }
 	virtual const uint8* GetDefaultValuePtr() const override { return reinterpret_cast<const uint8*>(&DefaultValue); }
 
+#if WITH_EDITOR
+	virtual FString FormatDefaultValue() const override { return DefaultValue.ToString(); }
+#endif  // WITH_EDITOR
+
 public:
 
 	/** The default value of this variable. */
@@ -309,6 +367,10 @@ public:
 	virtual ECameraVariableType GetVariableType() const override { return ECameraVariableType::Transform3f; }
 	virtual const uint8* GetDefaultValuePtr() const override { return reinterpret_cast<const uint8*>(&DefaultValue); }
 
+#if WITH_EDITOR
+	virtual FString FormatDefaultValue() const override { return DefaultValue.ToString(); }
+#endif  // WITH_EDITOR
+
 public:
 
 	/** The default value of this variable. */
@@ -329,6 +391,10 @@ public:
 	virtual ECameraVariableType GetVariableType() const override { return ECameraVariableType::Transform3d; }
 	virtual const uint8* GetDefaultValuePtr() const override { return reinterpret_cast<const uint8*>(&DefaultValue); }
 
+#if WITH_EDITOR
+	virtual FString FormatDefaultValue() const override { return DefaultValue.ToString(); }
+#endif  // WITH_EDITOR
+
 public:
 
 	/** The default value of this variable. */
diff --git a/Engine/Plugins/Cameras/GameplayCameras/Source/GameplayCameras/Public/Core/CameraVariableCollection.h b/Engine/Plugins/Cameras/GameplayCameras/Source/GameplayCameras/Public/Core/CameraVariableCollection.h
new file mode 100644
--- /dev/null
+++ b/Engine/Plugins/Cameras/GameplayCameras/Source/GameplayCameras/Public/Core/CameraVariableCollection.h
@@ -0,0 +1,29 @@
+// Copyright Epic Games, Inc. All Rights Reserved.
+
+#pragma once
+
+#include "CoreTypes.h"
+
+#include "CameraVariableCollection.generated.h"
+
+class UCameraVariableAsset;
+
+/**
+ * An asset that represents a collection of camera variables.
+ */
+UCLASS(MinimalAPI)
+class UCameraVariableCollection : public UObject
+{
+	GENERATED_BODY()
+
+public:
+
+	UCameraVariableCollection(const FObjectInitializer& ObjectInit);
+
+public:
+
+	/** The variables in this collection. */
+	UPROPERTY(EditAnywhere, Category="Camera")
+	TArray<TObjectPtr<UCameraVariableAsset>> Variables;
+};
+
diff --git a/Engine/Plugins/Cameras/GameplayCameras/Source/GameplayCamerasEditor/Private/AssetTools/AssetDefinition_CameraVariableAssets.cpp b/Engine/Plugins/Cameras/GameplayCameras/Source/GameplayCamerasEditor/Private/AssetTools/AssetDefinition_CameraVariableAssets.cpp
deleted file mode 100644
diff --git a/Engine/Plugins/Cameras/GameplayCameras/Source/GameplayCamerasEditor/Private/AssetTools/AssetDefinition_CameraVariableAssets.h b/Engine/Plugins/Cameras/GameplayCameras/Source/GameplayCamerasEditor/Private/AssetTools/AssetDefinition_CameraVariableAssets.h
deleted file mode 100644
diff --git a/Engine/Plugins/Cameras/GameplayCameras/Source/GameplayCamerasEditor/Private/AssetTools/AssetDefinition_CameraVariableCollection.cpp b/Engine/Plugins/Cameras/GameplayCameras/Source/GameplayCamerasEditor/Private/AssetTools/AssetDefinition_CameraVariableCollection.cpp
new file mode 100644
--- /dev/null
+++ b/Engine/Plugins/Cameras/GameplayCameras/Source/GameplayCamerasEditor/Private/AssetTools/AssetDefinition_CameraVariableCollection.cpp
@@ -0,0 +1,59 @@
+// Copyright Epic Games, Inc. All Rights Reserved.
+
+#include "AssetTools/AssetDefinition_CameraVariableCollection.h"
+
+#include "Core/CameraVariableCollection.h"
+#include "IGameplayCamerasEditorModule.h"
+#include "Toolkits/SimpleAssetEditor.h"
+
+#define LOCTEXT_NAMESPACE "AssetDefinition_CameraVariableCollection"
+
+#include UE_INLINE_GENERATED_CPP_BY_NAME(AssetDefinition_CameraVariableCollection)
+
+TConstArrayView<FAssetCategoryPath> UAssetDefinition_CameraVariableCollection::StaticMenuCategories()
+{
+	static const auto Categories = { FAssetCategoryPath(EAssetCategoryPaths::Gameplay) };
+	return Categories;
+}
+
+FText UAssetDefinition_CameraVariableCollection::GetAssetDisplayName() const
+{
+	return LOCTEXT("AssetDisplayName", "Camera Variable Collection");
+}
+
+FLinearColor UAssetDefinition_CameraVariableCollection::GetAssetColor() const
+{
+	return FLinearColor(FColor(200, 80, 80));
+}
+
+TSoftClassPtr<UObject> UAssetDefinition_CameraVariableCollection::GetAssetClass() const
+{
+	return UCameraVariableCollection::StaticClass();
+}
+
+TConstArrayView<FAssetCategoryPath> UAssetDefinition_CameraVariableCollection::GetAssetCategories() const
+{
+	return StaticMenuCategories();
+}
+
+FAssetOpenSupport UAssetDefinition_CameraVariableCollection::GetAssetOpenSupport(const FAssetOpenSupportArgs& OpenSupportArgs) const
+{
+	return FAssetOpenSupport(
+			OpenSupportArgs.OpenMethod,
+			OpenSupportArgs.OpenMethod == EAssetOpenMethod::Edit,
+			EToolkitMode::Standalone); 
+}
+
+EAssetCommandResult UAssetDefinition_CameraVariableCollection::OpenAssets(const FAssetOpenArgs& OpenArgs) const
+{
+	IGameplayCamerasEditorModule& GameplayCamerasEditorModule = FModuleManager::LoadModuleChecked<IGameplayCamerasEditorModule>("GameplayCamerasEditor");
+	for (UCameraVariableCollection* VariableCollection : OpenArgs.LoadObjects<UCameraVariableCollection>())
+	{
+		GameplayCamerasEditorModule.CreateCameraVariableCollectionEditor(OpenArgs.GetToolkitMode(), OpenArgs.ToolkitHost, VariableCollection);
+	}	
+
+	return EAssetCommandResult::Handled;
+}
+
+#undef LOCTEXT_NAMESPACE
+
diff --git a/Engine/Plugins/Cameras/GameplayCameras/Source/GameplayCamerasEditor/Private/AssetTools/AssetDefinition_CameraVariableCollection.h b/Engine/Plugins/Cameras/GameplayCameras/Source/GameplayCamerasEditor/Private/AssetTools/AssetDefinition_CameraVariableCollection.h
new file mode 100644
--- /dev/null
+++ b/Engine/Plugins/Cameras/GameplayCameras/Source/GameplayCamerasEditor/Private/AssetTools/AssetDefinition_CameraVariableCollection.h
@@ -0,0 +1,26 @@
+// Copyright Epic Games, Inc. All Rights Reserved.
+
+#pragma once
+
+#include "AssetDefinitionDefault.h"
+
+#include "AssetDefinition_CameraVariableCollection.generated.h"
+
+UCLASS(MinimalAPI)
+class UAssetDefinition_CameraVariableCollection : public UAssetDefinitionDefault
+{
+	GENERATED_BODY()
+
+public:
+
+	static TConstArrayView<FAssetCategoryPath> StaticMenuCategories();
+
+	// UAssetDefinition interface
+	virtual FText GetAssetDisplayName() const override;
+	virtual FLinearColor GetAssetColor() const override;
+	virtual TSoftClassPtr<UObject> GetAssetClass() const override;
+	virtual TConstArrayView<FAssetCategoryPath> GetAssetCategories() const override;
+	virtual FAssetOpenSupport GetAssetOpenSupport(const FAssetOpenSupportArgs& OpenSupportArgs) const override;
+	virtual EAssetCommandResult OpenAssets(const FAssetOpenArgs& OpenArgs) const override;
+};
+
diff --git a/Engine/Plugins/Cameras/GameplayCameras/Source/GameplayCamerasEditor/Private/AssetTools/CameraVariableCollectionEditor.cpp b/Engine/Plugins/Cameras/GameplayCameras/Source/GameplayCamerasEditor/Private/AssetTools/CameraVariableCollectionEditor.cpp
new file mode 100644
--- /dev/null
+++ b/Engine/Plugins/Cameras/GameplayCameras/Source/GameplayCamerasEditor/Private/AssetTools/CameraVariableCollectionEditor.cpp
@@ -0,0 +1,27 @@
+// Copyright Epic Games, Inc. All Rights Reserved.
+
+#include "AssetTools/CameraVariableCollectionEditor.h"
+
+#include "Core/CameraVariableCollection.h"
+#include "EditorModeManager.h"
+#include "Toolkits/CameraVariableCollectionEditorToolkit.h"
+
+#include UE_INLINE_GENERATED_CPP_BY_NAME(CameraVariableCollectionEditor)
+
+void UCameraVariableCollectionEditor::Initialize(TObjectPtr<UCameraVariableCollection> InVariableCollection)
+{
+	VariableCollection = InVariableCollection;
+
+	Super::Initialize();
+}
+
+void UCameraVariableCollectionEditor::GetObjectsToEdit(TArray<UObject*>& InObjectsToEdit)
+{
+	InObjectsToEdit.Add(VariableCollection.Get());
+}
+
+TSharedPtr<FBaseAssetToolkit> UCameraVariableCollectionEditor::CreateToolkit()
+{
+	return MakeShared<UE::Cameras::FCameraVariableCollectionEditorToolkit>(this);
+}
+
diff --git a/Engine/Plugins/Cameras/GameplayCameras/Source/GameplayCamerasEditor/Private/Commands/CameraVariableCollectionEditorCommands.cpp b/Engine/Plugins/Cameras/GameplayCameras/Source/GameplayCamerasEditor/Private/Commands/CameraVariableCollectionEditorCommands.cpp
new file mode 100644
--- /dev/null
+++ b/Engine/Plugins/Cameras/GameplayCameras/Source/GameplayCamerasEditor/Private/Commands/CameraVariableCollectionEditorCommands.cpp
@@ -0,0 +1,38 @@
+// Copyright Epic Games, Inc. All Rights Reserved.
+
+#include "Commands/CameraVariableCollectionEditorCommands.h"
+
+#include "Framework/Commands/Commands.h"
+#include "Framework/Commands/InputChord.h"
+#include "Framework/Commands/UICommandInfo.h"
+#include "GenericPlatform/GenericApplication.h"
+#include "Styles/GameplayCamerasEditorStyle.h"
+#include "Styling/AppStyle.h"
+
+#define LOCTEXT_NAMESPACE "CameraVariableCollectionEditorCommands"
+
+namespace UE::Cameras
+{
+
+FCameraVariableCollectionEditorCommands::FCameraVariableCollectionEditorCommands()
+	: TCommands<FCameraVariableCollectionEditorCommands>(
+			"CameraVariableCollectionEditor",
+			LOCTEXT("CameraVariableCollectionEditor", "Camera Variable Collection Editor"),
+			NAME_None,
+			FGameplayCamerasEditorStyle::Get()->GetStyleSetName()
+		)
+{
+}
+
+void FCameraVariableCollectionEditorCommands::RegisterCommands()
+{
+	UI_COMMAND(CreateVariable, "Create Variable", "Adds a new camera variable to the collection",
+			EUserInterfaceActionType::Button, FInputChord());
+	UI_COMMAND(DeleteVariable, "Delete Variable", "Removes a camera variable from the collection",
+			EUserInterfaceActionType::Button, FInputChord());
+}
+
+}  // namespace UE::Cameras
+
+#undef LOCTEXT_NAMESPACE
+
diff --git a/Engine/Plugins/Cameras/GameplayCameras/Source/GameplayCamerasEditor/Private/Customizations/CameraParameterDetailsCustomizations.cpp b/Engine/Plugins/Cameras/GameplayCameras/Source/GameplayCamerasEditor/Private/Customizations/CameraParameterDetailsCustomizations.cpp
new file mode 100644
--- /dev/null
+++ b/Engine/Plugins/Cameras/GameplayCameras/Source/GameplayCamerasEditor/Private/Customizations/CameraParameterDetailsCustomizations.cpp
@@ -0,0 +1,473 @@
+// Copyright Epic Games, Inc. All Rights Reserved.
+
+#include "Customizations/CameraParameterDetailsCustomizations.h"
+
+#include "ContentBrowserModule.h"
+#include "Core/CameraParameters.h"
+#include "Core/CameraVariableAssets.h"
+#include "Core/CameraVariableCollection.h"
+#include "Customizations/MathStructCustomizations.h"
+#include "DetailLayoutBuilder.h"
+#include "Framework/Application/SlateApplication.h"
+#include "Framework/MultiBox/MultiBoxBuilder.h"
+#include "Framework/Views/ITypedTableView.h"
+#include "IContentBrowserSingleton.h"
+#include "IDetailChildrenBuilder.h"
+#include "IDetailPropertyRow.h"
+#include "IPropertyUtilities.h"
+#include "Modules/ModuleManager.h"
+#include "PropertyCustomizationHelpers.h"
+#include "PropertyEditorModule.h"
+#include "ScopedTransaction.h"
+#include "Styles/GameplayCamerasEditorStyle.h"
+#include "Styling/StyleColors.h"
+#include "Widgets/Images/SImage.h"
+#include "Widgets/Input/SButton.h"
+#include "Widgets/Input/SComboButton.h"
+#include "Widgets/SBoxPanel.h"
+#include "Widgets/Text/STextBlock.h"
+#include "Widgets/Views/SListView.h"
+
+#define LOCTEXT_NAMESPACE "CameraParameterDetailsCustomization"
+
+namespace UE::Cameras
+{
+
+DECLARE_DELEGATE_OneParam(FExecuteSetParameterVariable, UCameraVariableAsset*);
+
+class SCameraVariableBrowser : public SCompoundWidget
+{
+public:
+
+	SLATE_BEGIN_ARGS(SCameraVariableBrowser)
+	{}
+		SLATE_ARGUMENT(UClass*, VariableClass)
+		SLATE_ARGUMENT(UCameraVariableAsset*, InitialVariable)
+		SLATE_EVENT(FExecuteSetParameterVariable, RequestSetVariable)
+	SLATE_END_ARGS()
+	
+	void Construct(const FArguments& InArgs)
+	{
+		VariableClass = InArgs._VariableClass;
+		InitialVariable = InArgs._InitialVariable;
+
+		RequestSetVariable = InArgs._RequestSetVariable;
+
+		ChildSlot
+		[
+			SNew(SBorder)
+			.BorderImage(FAppStyle::GetBrush("ToolPanel.GroupBorder"))
+			[
+				SNew(SVerticalBox)
+				+SVerticalBox::Slot()
+				.FillHeight(0.5f)
+				[
+					BuildVariableCollectionAssetPicker()
+				]
+				+SVerticalBox::Slot()
+				.FillHeight(0.5f)
+				[
+					SAssignNew(CameraVariableListView, SListView<UCameraVariableAsset*>)
+					.ListItemsSource(&CameraVariableItemsSource)
+					.OnGenerateRow(this, &SCameraVariableBrowser::OnVariableListGenerateRow)
+					.OnSelectionChanged(this, &SCameraVariableBrowser::OnVariableListSelectionChanged)
+				]
+			]
+		];
+
+		if (InitialVariable)
+		{
+			bSuppressVariableListSelectionChanged = true;
+			CameraVariableListView->SetSelection(InitialVariable);
+			CameraVariableListView->RequestScrollIntoView(InitialVariable);
+			bSuppressVariableListSelectionChanged = false;
+		}
+	}
+
+private:
+
+	TSharedRef<SWidget> BuildVariableCollectionAssetPicker()
+	{
+		FContentBrowserModule& ContentBrowserModule = FModuleManager::Get().LoadModuleChecked<FContentBrowserModule>(TEXT("ContentBrowser"));
+
+		FAssetPickerConfig AssetPickerConfig;
+
+		FARFilter ARFilter;
+		ARFilter.ClassPaths.Add(FTopLevelAssetPath(UCameraVariableCollection::StaticClass()->GetPathName()));
+
+		UCameraVariableCollection* InitialVariableCollection = nullptr;
+		if (InitialVariable)
+		{
+			InitialVariableCollection = InitialVariable->GetTypedOuter<UCameraVariableCollection>();
+		}
+
+		AssetPickerConfig.bAllowDragging = false;
+		AssetPickerConfig.bCanShowClasses = false;
+		AssetPickerConfig.bAllowNullSelection = false;
+		AssetPickerConfig.bShowBottomToolbar = true;
+		AssetPickerConfig.bFocusSearchBoxWhenOpened = true;
+		AssetPickerConfig.SelectionMode = ESelectionMode::Single;
+		AssetPickerConfig.Filter = ARFilter;
+		AssetPickerConfig.SaveSettingsName = TEXT("CameraParameterVariablePropertyPicker");
+		AssetPickerConfig.InitialAssetViewType = EAssetViewType::List;
+		AssetPickerConfig.InitialAssetSelection = FAssetData(InitialVariableCollection);
+		AssetPickerConfig.OnAssetSelected = FOnAssetSelected::CreateSP(this, &SCameraVariableBrowser::OnAssetSelected);
+		AssetPickerConfig.GetCurrentSelectionDelegates.Add(&GetCurrentAssetPickerSelection);
+
+		return ContentBrowserModule.Get().CreateAssetPicker(AssetPickerConfig);
+	}
+
+	void OnAssetSelected(const FAssetData& SelectedAsset)
+	{
+		UpdateVariableListItemsSource();
+	}
+
+	void UpdateVariableListItemsSource()
+	{
+		TArray<FAssetData> SelectedAssets;
+		if (GetCurrentAssetPickerSelection.IsBound())
+		{
+			SelectedAssets = GetCurrentAssetPickerSelection.Execute();
+		}
+
+		CameraVariableItemsSource.Reset();
+		if (!SelectedAssets.IsEmpty())
+		{
+			if (UCameraVariableCollection* VariableCollection = Cast<UCameraVariableCollection>(SelectedAssets[0].GetAsset()))
+			{
+				CameraVariableItemsSource = VariableCollection->Variables.FilterByPredicate(
+						[this](UCameraVariableAsset* Item) { return Item->GetClass() == VariableClass; });
+			}
+		}
+
+		CameraVariableListView->RequestListRefresh();
+	}
+
+	TSharedRef<ITableRow> OnVariableListGenerateRow(UCameraVariableAsset* Item, const TSharedRef<STableViewBase>& OwnerTable)
+	{
+		TSharedRef<FGameplayCamerasEditorStyle> GameplayCamerasStyle = FGameplayCamerasEditorStyle::Get();
+
+		const FText DisplayName = Item->DisplayName.IsEmpty() ?
+			FText::FromName(Item->GetFName()) : FText::FromString(Item->DisplayName);
+
+		return SNew(STableRow<UCameraVariableAsset*>, OwnerTable)
+			[
+				SNew(SHorizontalBox)
+				+SHorizontalBox::Slot()
+				.AutoWidth()
+				[
+					SNew(SImage)
+					.ColorAndOpacity(FSlateColor::UseForeground())
+					.Image(GameplayCamerasStyle->GetBrush("CameraParameter.VariableBrowser"))
+				]
+				+SHorizontalBox::Slot()
+				.FillWidth(1.f)
+				.Padding(4.f, 2.f)
+				[
+					SNew(STextBlock)
+					.Text(DisplayName)
+				]
+			];
+	}
+
+	void OnVariableListSelectionChanged(UCameraVariableAsset* Item, ESelectInfo::Type SelectInfo) const
+	{
+		if (!bSuppressVariableListSelectionChanged && Item)
+		{
+			RequestSetVariable.ExecuteIfBound(Item);
+		}
+		// else, the user should click the "Clear" button to actually remove the variable.
+	}
+
+private:
+
+	TSharedPtr<SListView<UCameraVariableAsset*>> CameraVariableListView;
+	TArray<UCameraVariableAsset*> CameraVariableItemsSource;
+
+	UClass* VariableClass = nullptr;
+	UCameraVariableAsset* InitialVariable = nullptr;
+
+	FGetCurrentSelectionDelegate GetCurrentAssetPickerSelection;
+	bool bSuppressVariableListSelectionChanged = false;
+
+	FExecuteSetParameterVariable RequestSetVariable;
+};
+
+void FCameraParameterDetailsCustomization::Register(FPropertyEditorModule& PropertyEditorModule)
+{
+#define UE_CAMERA_VARIABLE_FOR_TYPE(ValueType, ValueName)\
+	PropertyEditorModule.RegisterCustomPropertyTypeLayout(\
+			F##ValueName##CameraParameter::StaticStruct()->GetFName(),\
+			FOnGetPropertyTypeCustomizationInstance::CreateLambda(\
+				[]{ return MakeShared<F##ValueName##CameraParameterDetailsCustomization>(); }));
+UE_CAMERA_VARIABLE_FOR_ALL_TYPES()
+#undef UE_CAMERA_VARIABLE_FOR_TYPE
+}
+
+void FCameraParameterDetailsCustomization::Unregister(FPropertyEditorModule& PropertyEditorModule)
+{
+#define UE_CAMERA_VARIABLE_FOR_TYPE(ValueType, ValueName)\
+	PropertyEditorModule.UnregisterCustomPropertyTypeLayout(\
+			F##ValueName##CameraParameter::StaticStruct()->GetFName());
+UE_CAMERA_VARIABLE_FOR_ALL_TYPES()
+#undef UE_CAMERA_VARIABLE_FOR_TYPE
+}
+
+void FCameraParameterDetailsCustomization::CustomizeHeader(TSharedRef<IPropertyHandle> PropertyHandle, FDetailWidgetRow& HeaderRow, IPropertyTypeCustomizationUtils& CustomizationUtils)
+{
+	PropertyUtilities = CustomizationUtils.GetPropertyUtilities();
+
+	StructProperty = PropertyHandle;
+	StructProperty->SetOnPropertyResetToDefault(
+			FSimpleDelegate::CreateSP(this, &FCameraParameterDetailsCustomization::OnResetToDefault));
+
+	FName ValuePropertyName, VariablePropertyName;
+	GetValueAndVariablePropertyNames(ValuePropertyName, VariablePropertyName);
+
+	ValueProperty = PropertyHandle->GetChildHandle(ValuePropertyName);
+	VariableProperty = PropertyHandle->GetChildHandle(VariablePropertyName);
+
+	VariableClass = nullptr;
+	if (FObjectProperty* VariableObjectProperty = CastField<FObjectProperty>(VariableProperty->GetProperty()))
+	{
+		VariableClass = VariableObjectProperty->PropertyClass;
+	}
+
+	UpdateVariableInfo();
+
+	const bool bHasText = !VariableInfoText.IsEmpty() || !VariableErrorText.IsEmpty();
+
+	TSharedRef<SWidget> ValueWidget = ValueProperty->CreatePropertyValueWidgetWithCustomization(nullptr);
+	ValueWidget->SetEnabled(!bHasText);
+	// TODO: change SStandaloneCustomizedValueWidget so that it can tell us about the layout requirements
+	//		 of the value widget (min/max desired width/height, H/V alignment, etc.)
+	const float MaxValueWidgetDesiredWidth = bHasText ? 300.f : 0.f;
+
+	TSharedRef<FGameplayCamerasEditorStyle> GameplayCamerasStyle = FGameplayCamerasEditorStyle::Get();
+
+	HeaderRow
+	.NameContent()
+	[
+		StructProperty->CreatePropertyNameWidget()
+	]
+	.ValueContent()
+	.MaxDesiredWidth(MaxValueWidgetDesiredWidth)
+	.HAlign(HAlign_Fill)
+	[
+		SNew(SHorizontalBox)
+		+SHorizontalBox::Slot()
+		.Padding(0)
+		.FillWidth(1.f)
+		[
+			ValueWidget
+		]
+		+SHorizontalBox::Slot()
+		.Padding(0)
+		.AutoWidth()
+		.HAlign(HAlign_Right)
+		[
+			ValueProperty->CreateDefaultPropertyButtonWidgets()
+		]
+	]
+	.ExtensionContent()
+	.HAlign(HAlign_Right)
+	[
+		SNew(SHorizontalBox)
+		+SHorizontalBox::Slot()
+		.Padding(2.f)
+		.AutoWidth()
+		[
+			SAssignNew(VariableBrowserButton, SComboButton)
+			.HasDownArrow(true)
+			.ContentPadding(1.f)
+			.ButtonStyle(FAppStyle::Get(), "SimpleButton")
+			.ToolTipText(LOCTEXT("SetVariable_ToolTip", "Selects a camera variable to drive this parameter"))
+			.IsEnabled(VariableClass != nullptr)
+			.ButtonContent()
+			[
+				SNew(SHorizontalBox)
+				+SHorizontalBox::Slot()
+				.AutoWidth()
+				.Padding(2.f)
+				.VAlign(VAlign_Center)
+				[
+					SNew(SImage)
+					.ColorAndOpacity(FSlateColor::UseForeground())
+					.Image(GameplayCamerasStyle->GetBrush("CameraParameter.VariableBrowser"))
+				]
+				+SHorizontalBox::Slot()
+				.AutoWidth()
+				.Padding(2.f)
+				.VAlign(VAlign_Center)
+				[
+					SNew(SBox)
+					.VAlign(VAlign_Center)
+					.Visibility(!VariableInfoText.IsEmpty() ? EVisibility::Visible : EVisibility::Collapsed)
+					[
+						SNew(STextBlock)
+						.Text(VariableInfoText)
+					]
+				]
+				+SHorizontalBox::Slot()
+				.AutoWidth()
+				.Padding(2.f)
+				.VAlign(VAlign_Center)
+				[
+					SNew(SBox)
+					.VAlign(VAlign_Center)
+					.Visibility(!VariableErrorText.IsEmpty() ? EVisibility::Visible : EVisibility::Collapsed)
+					[
+						SNew(STextBlock)
+						.Text(VariableErrorText)
+						.ColorAndOpacity(FStyleColors::Error)
+					]
+				]
+			]
+			.OnGetMenuContent(this, &FCameraParameterDetailsCustomization::BuildCameraVariableBrowser)
+		]
+	];
+}
+
+void FCameraParameterDetailsCustomization::CustomizeChildren(TSharedRef<IPropertyHandle> PropertyHandle, IDetailChildrenBuilder& ChildBuilder, IPropertyTypeCustomizationUtils& CustomizationUtils)
+{
+}
+
+TSharedRef<SWidget> FCameraParameterDetailsCustomization::BuildCameraVariableBrowser()
+{
+	const bool bShouldCloseWindowAfterMenuSelection = true;
+	const bool bCloseSelfOnly = true;
+	const bool bSearchable = false;
+	FMenuBuilder MenuBuilder(true, nullptr, nullptr, bCloseSelfOnly, &FCoreStyle::Get(), bSearchable);
+
+	MenuBuilder.BeginSection(NAME_None, LOCTEXT("CameraVariableOperations", "Current Parameter"));
+	{
+		MenuBuilder.AddMenuEntry(
+			LOCTEXT("ClearVariable", "Clear"),
+			LOCTEXT("ClearVariable_ToolTip", "Clears the variable from the camera parameter"),
+			FSlateIcon(FAppStyle::GetAppStyleSetName(), "GenericCommands.Delete"),
+			FUIAction(
+				FExecuteAction::CreateSP(this, &FCameraParameterDetailsCustomization::OnClearVariable),
+				FCanExecuteAction::CreateSP(this, &FCameraParameterDetailsCustomization::CanClearVariable))
+			);
+	}
+	MenuBuilder.EndSection();
+
+	MenuBuilder.BeginSection(NAME_None, LOCTEXT("CameraVariableBrowser", "Browse"));
+	{
+		TSharedRef<SWidget> VariableBrowser = SNew(SBox)
+			.MinDesiredWidth(300.f)
+			.MinDesiredHeight(300.f)
+			[
+				SNew(SCameraVariableBrowser)
+				.VariableClass(VariableClass)
+				.InitialVariable(CommonVariable)
+				.RequestSetVariable(
+						FExecuteSetParameterVariable::CreateSP(this, &FCameraParameterDetailsCustomization::OnSetVariable))
+			];
+		MenuBuilder.AddWidget(VariableBrowser, FText(), true, false);
+	}
+	MenuBuilder.EndSection();
+
+	return MenuBuilder.MakeWidget();
+}
+
+void FCameraParameterDetailsCustomization::UpdateVariableInfo()
+{
+	CommonVariable = nullptr;
+	VariableInfoText = FText::GetEmpty();
+	VariableErrorText = FText::GetEmpty();
+
+	UObject* VariableObject = nullptr;
+	FPropertyAccess::Result PropertyAccessResult = VariableProperty->GetValue(VariableObject);
+	if (PropertyAccessResult == FPropertyAccess::Success)
+	{
+		if (VariableObject)
+		{
+			if (UCameraVariableAsset* Variable = Cast<UCameraVariableAsset>(VariableObject))
+			{
+				CommonVariable = Variable;
+				VariableInfoText = Variable->DisplayName.IsEmpty() ?
+					FText::FromName(Variable->GetFName()) :
+					FText::FromString(Variable->DisplayName);
+			}
+			else
+			{
+				VariableErrorText = LOCTEXT("InvalidVariableObject", "Invalid Variable");
+			}
+		}
+		// else: no variable set
+	}
+	else if (PropertyAccessResult == FPropertyAccess::MultipleValues)
+	{
+		VariableInfoText = LOCTEXT("MultipleVariableValues", "Multiple Variables");
+	}
+	else
+	{
+		VariableErrorText = LOCTEXT("ErrorReadingVariable", "Error Reading Variable");
+	}
+}
+
+bool FCameraParameterDetailsCustomization::HasVariableInfoText() const
+{
+	return !VariableInfoText.IsEmpty() || !VariableErrorText.IsEmpty();
+}
+
+bool FCameraParameterDetailsCustomization::CanClearVariable() const
+{
+	return VariableProperty->CanResetToDefault();
+}
+
+void FCameraParameterDetailsCustomization::OnClearVariable()
+{
+	VariableProperty->ResetToDefault();
+	PropertyUtilities->RequestForceRefresh();
+}
+
+void FCameraParameterDetailsCustomization::OnSetVariable(UCameraVariableAsset* InVariable)
+{
+	TArray<void*> RawData;
+	StructProperty->AccessRawData(RawData);
+
+	TArray<UObject*> OuterObjects;
+	StructProperty->GetOuterObjects(OuterObjects);
+
+	check(!OuterObjects.Num() || OuterObjects.Num() == RawData.Num());
+	
+	{
+		FScopedTransaction Transaction(FText::Format(
+					LOCTEXT("SetPropertyValue", "Set {0}"), StructProperty->GetPropertyDisplayName()));
+
+		StructProperty->NotifyPreChange();
+
+		for (int32 ValueIndex = 0; ValueIndex < RawData.Num(); ++ValueIndex)
+		{
+			SetParameterVariable(RawData[ValueIndex], InVariable);
+		}
+
+		StructProperty->NotifyPostChange(EPropertyChangeType::ValueSet);
+	}
+
+	FPropertyChangedEvent ChangeEvent(StructProperty->GetProperty(), EPropertyChangeType::ValueSet, OuterObjects);
+	PropertyUtilities->NotifyFinishedChangingProperties(ChangeEvent);
+
+	PropertyUtilities->RequestForceRefresh();
+	VariableBrowserButton->SetIsOpen(false);
+}
+
+void FCameraParameterDetailsCustomization::OnResetToDefault()
+{
+	PropertyUtilities->RequestForceRefresh();
+}
+
+#define UE_CAMERA_VARIABLE_FOR_TYPE(ValueType, ValueName)\
+void F##ValueName##CameraParameterDetailsCustomization::SetParameterVariable(void* InRawData, UCameraVariableAsset* InVariable)\
+{\
+	F##ValueName##CameraParameter* TypedData = reinterpret_cast<F##ValueName##CameraParameter*>(InRawData);\
+	TypedData->Variable = CastChecked<U##ValueName##CameraVariable>(InVariable);\
+}
+UE_CAMERA_VARIABLE_FOR_ALL_TYPES()
+#undef UE_CAMERA_VARIABLE_FOR_TYPE
+
+}  // namespace UE::Cameras
+
+#undef LOCTEXT_NAMESPACE
diff --git a/Engine/Plugins/Cameras/GameplayCameras/Source/GameplayCamerasEditor/Private/Editors/SCameraVariableCollectionEditor.cpp b/Engine/Plugins/Cameras/GameplayCameras/Source/GameplayCamerasEditor/Private/Editors/SCameraVariableCollectionEditor.cpp
new file mode 100644
--- /dev/null
+++ b/Engine/Plugins/Cameras/GameplayCameras/Source/GameplayCamerasEditor/Private/Editors/SCameraVariableCollectionEditor.cpp
@@ -0,0 +1,283 @@
+// Copyright Epic Games, Inc. All Rights Reserved.
+
+#include "Editors/SCameraVariableCollectionEditor.h"
+
+#include "Commands/CameraVariableCollectionEditorCommands.h"
+#include "Core/CameraVariableAssets.h"
+#include "Core/CameraVariableCollection.h"
+#include "Framework/Views/TableViewMetadata.h"
+#include "IDetailsView.h"
+#include "Styles/GameplayCamerasEditorStyle.h"
+#include "ToolMenus.h"
+#include "UObject/UObjectIterator.h"
+#include "Widgets/Input/SSearchBox.h"
+#include "Widgets/Layout/SBorder.h"
+#include "Widgets/Layout/SSpacer.h"
+#include "Widgets/SBoxPanel.h"
+#include "Widgets/Views/SListView.h"
+
+#define LOCTEXT_NAMESPACE "SCameraVariableCollectionEditor"
+
+namespace UE::Cameras
+{
+
+class SCameraVariableCollectionListRow : public SMultiColumnTableRow<UCameraVariableAsset*>
+{
+public:
+
+	SLATE_BEGIN_ARGS(SCameraVariableCollectionListRow)
+		: _CameraVariable(nullptr)
+	{}
+		/** The camera variable corresponding to this entry. */
+		SLATE_ARGUMENT(UCameraVariableAsset*, CameraVariable)
+		/** Text to highlight if a search is ongoing. */
+		SLATE_ATTRIBUTE(FText, HighlightText)
+	SLATE_END_ARGS()
+
+	void Construct(const FArguments& InArgs, const TSharedRef<STableViewBase>& OwnerTableView);
+
+protected:
+
+	virtual TSharedRef<SWidget> GenerateWidgetForColumn(const FName& ColumnName) override;
+
+private:
+
+	FText GetVariableName() const;
+	FText GetVariableType() const;
+	FText GetDefaultValue() const;
+
+private:
+
+	UCameraVariableAsset* CameraVariable = nullptr;
+
+	TAttribute<FText> HighlightText;
+};
+
+void SCameraVariableCollectionListRow::Construct(const FArguments& InArgs, const TSharedRef<STableViewBase>& OwnerTableView)
+{
+	CameraVariable = InArgs._CameraVariable;
+	HighlightText = InArgs._HighlightText;
+
+	FSuperRowType::Construct(
+		FSuperRowType::FArguments().Padding(1.0f),
+		OwnerTableView);
+}
+
+TSharedRef<SWidget> SCameraVariableCollectionListRow::GenerateWidgetForColumn(const FName& ColumnName)
+{
+	TSharedRef<FGameplayCamerasEditorStyle> CamerasEditorStyle = FGameplayCamerasEditorStyle::Get();
+
+	if (ColumnName == TEXT("VariableName"))
+	{
+		return SNew(SHorizontalBox)
+			+SHorizontalBox::Slot()
+			.AutoWidth()
+			.Padding(8.f)
+			.VAlign(VAlign_Center)
+			[
+				SNew(STextBlock)
+				.TextStyle(CamerasEditorStyle, "CameraVariableCollectionEditor.Entry.Name")
+				.Text(this, &SCameraVariableCollectionListRow::GetVariableName)
+				.HighlightText(HighlightText)
+			];
+	}
+	else if (ColumnName == TEXT("VariableType"))
+	{
+		return SNew(SHorizontalBox)
+			+SHorizontalBox::Slot()
+			.AutoWidth()
+			.Padding(8.f)
+			.VAlign(VAlign_Center)
+			[
+				SNew(STextBlock)
+				.TextStyle(CamerasEditorStyle, "CameraVariableCollectionEditor.Entry.Type")
+				.Text(this, &SCameraVariableCollectionListRow::GetVariableType)
+				.HighlightText(HighlightText)
+			];
+	}
+	else if (ColumnName == TEXT("DefaultValue"))
+	{
+		return SNew(SHorizontalBox)
+			+SHorizontalBox::Slot()
+			.AutoWidth()
+			.Padding(8.f)
+			.VAlign(VAlign_Center)
+			[
+				SNew(STextBlock)
+				.TextStyle(CamerasEditorStyle, "CameraVariableCollectionEditor.Entry.Value")
+				.Text(this, &SCameraVariableCollectionListRow::GetDefaultValue)
+			];
+	}
+	return SNullWidget::NullWidget;
+}
+
+FText SCameraVariableCollectionListRow::GetVariableName() const
+{
+	return CameraVariable->DisplayName.IsEmpty() ?
+		FText::FromName(CameraVariable->GetFName()) :
+		FText::FromString(CameraVariable->DisplayName);
+}
+
+FText SCameraVariableCollectionListRow::GetVariableType() const
+{
+	FString ClassName = CameraVariable->GetClass()->GetName();
+	ClassName.RemoveFromEnd(TEXT("CameraVariable"));
+	return FText::FromString(ClassName);
+}
+
+FText SCameraVariableCollectionListRow::GetDefaultValue() const
+{
+	return FText::FromString(CameraVariable->FormatDefaultValue());
+}
+
+void SCameraVariableCollectionEditor::Construct(const FArguments& InArgs)
+{
+	VariableCollection = InArgs._VariableCollection;
+
+	WeakDetailsView = InArgs._DetailsView;
+
+	SearchTextFilter = MakeShareable(new FEntryTextFilter(
+				FEntryTextFilter::FItemToStringArray::CreateSP(this, &SCameraVariableCollectionEditor::GetEntryStrings)));
+
+	ChildSlot
+	[
+		SNew(SVerticalBox)
+
+		+SVerticalBox::Slot()
+		.AutoHeight()
+		[
+			SNew(SBorder)
+			.BorderImage(FAppStyle::GetBrush("ToolPanel.GroupBorder"))
+			.Padding(8.f)
+			[
+				SAssignNew(SearchBox, SSearchBox)
+				.HintText(LOCTEXT("SearchHint", "Search"))
+				.OnTextChanged(this, &SCameraVariableCollectionEditor::OnSearchTextChanged)
+				.OnTextCommitted(this, &SCameraVariableCollectionEditor::OnSearchTextCommitted)
+			]
+		]
+
+		+SVerticalBox::Slot()
+		.Padding(0, 3.f)
+		[
+			SAssignNew(ListView, SListView<UCameraVariableAsset*>)
+			.ListItemsSource(&FilteredItemSource)
+			.OnGenerateRow(this, &SCameraVariableCollectionEditor::OnListGenerateRow)
+			.OnSelectionChanged(this, &SCameraVariableCollectionEditor::OnListSectionChanged)
+			.HeaderRow
+			(
+				SNew(SHeaderRow)
+
+				+SHeaderRow::Column("VariableName")
+				.DefaultLabel(LOCTEXT("VariableNameColumn", "Variable Name"))
+				.FillWidth(0.4f)
+
+				+SHeaderRow::Column("VariableType")
+				.DefaultLabel(LOCTEXT("VariableTypeColumn", "Variable Type"))
+				.FillWidth(0.3f)
+
+				+SHeaderRow::Column("DefaultValue")
+				.DefaultLabel(LOCTEXT("DefaultValueColumn", "Default Value"))
+				.FillWidth(0.3f)
+			)
+		]
+	];
+
+	bUpdateFilteredItemSource = true;
+
+	SetDetailsViewObject(nullptr);
+}
+
+SCameraVariableCollectionEditor::~SCameraVariableCollectionEditor()
+{
+}
+
+void SCameraVariableCollectionEditor::GetSelectedVariables(TArray<UCameraVariableAsset*>& OutSelection) const
+{
+	ListView->GetSelectedItems(OutSelection);
+}
+
+void SCameraVariableCollectionEditor::RequestListRefresh()
+{
+	bUpdateFilteredItemSource = true;
+}
+
+void SCameraVariableCollectionEditor::Tick(const FGeometry& AllottedGeometry, const double InCurrentTime, const float InDeltaTime)
+{
+	if (bUpdateFilteredItemSource)
+	{
+		UpdateFilteredItemSource();
+		ListView->RequestListRefresh();
+	}
+	bUpdateFilteredItemSource = false;
+
+	SCompoundWidget::Tick(AllottedGeometry, InCurrentTime, InDeltaTime);
+}
+
+void SCameraVariableCollectionEditor::UpdateFilteredItemSource()
+{
+	FilteredItemSource.Reset();
+
+	if (!SearchTextFilter->GetRawFilterText().IsEmpty())
+	{
+		for (UCameraVariableAsset* Item : VariableCollection->Variables)
+		{
+			if (SearchTextFilter->PassesFilter(Item))
+			{
+				FilteredItemSource.Add(Item);
+			}
+		}
+	}
+	else
+	{
+		FilteredItemSource = VariableCollection->Variables;
+	}
+}
+
+void SCameraVariableCollectionEditor::SetDetailsViewObject(UObject* InObject) const
+{
+	if (TSharedPtr<IDetailsView> DetailsView = WeakDetailsView.Pin())
+	{
+		DetailsView->SetObject(InObject);
+	}
+}
+
+TSharedRef<ITableRow> SCameraVariableCollectionEditor::OnListGenerateRow(UCameraVariableAsset* Item, const TSharedRef<STableViewBase>& OwnerTable)
+{
+	return SNew(SCameraVariableCollectionListRow, OwnerTable)
+		.CameraVariable(Item)
+		.HighlightText(this, &SCameraVariableCollectionEditor::GetHighlightText);
+}
+
+void SCameraVariableCollectionEditor::OnListSectionChanged(UCameraVariableAsset* Item, ESelectInfo::Type SelectInfo) const
+{
+	SetDetailsViewObject(Item);
+}
+
+void SCameraVariableCollectionEditor::GetEntryStrings(const UCameraVariableAsset* InItem, TArray<FString>& OutStrings)
+{
+	OutStrings.Add(InItem->GetName());
+	OutStrings.Add(InItem->GetClass()->GetName());
+}
+
+void SCameraVariableCollectionEditor::OnSearchTextChanged(const FText& InFilterText)
+{
+	SearchTextFilter->SetRawFilterText(InFilterText);
+	SearchBox->SetError(SearchTextFilter->GetFilterErrorText());
+
+	bUpdateFilteredItemSource = true;
+}
+
+void SCameraVariableCollectionEditor::OnSearchTextCommitted(const FText& InFilterText, ETextCommit::Type InCommitType)
+{
+	OnSearchTextChanged(InFilterText);
+}
+
+FText SCameraVariableCollectionEditor::GetHighlightText() const
+{
+	return SearchTextFilter->GetRawFilterText();
+}
+
+}  // namespace UE::Cameras
+
+#undef LOCTEXT_NAMESPACE
diff --git a/Engine/Plugins/Cameras/GameplayCameras/Source/GameplayCamerasEditor/Private/Factories/CameraVariableCollectionFactory.cpp b/Engine/Plugins/Cameras/GameplayCameras/Source/GameplayCamerasEditor/Private/Factories/CameraVariableCollectionFactory.cpp
new file mode 100644
--- /dev/null
+++ b/Engine/Plugins/Cameras/GameplayCameras/Source/GameplayCamerasEditor/Private/Factories/CameraVariableCollectionFactory.cpp
@@ -0,0 +1,31 @@
+// Copyright Epic Games, Inc. All Rights Reserved.
+
+#include "Factories/CameraVariableCollectionFactory.h"
+
+#include "Core/CameraVariableCollection.h"
+
+#include UE_INLINE_GENERATED_CPP_BY_NAME(CameraVariableCollectionFactory)
+
+#define LOCTEXT_NAMESPACE "CameraVariableCollectionFactory"
+
+UCameraVariableCollectionFactory::UCameraVariableCollectionFactory(const FObjectInitializer& ObjectInit)
+	: Super(ObjectInit)
+{
+	bCreateNew = true;
+	bEditAfterNew = true;
+	SupportedClass = UCameraVariableCollection::StaticClass();
+}
+
+FText UCameraVariableCollectionFactory::GetDisplayName() const
+{
+	return LOCTEXT("DisplayName", "Camera Variable Collection");
+}
+
+UObject* UCameraVariableCollectionFactory::FactoryCreateNew(UClass* Class, UObject* Parent, FName Name, EObjectFlags Flags, UObject* Context, FFeedbackContext* Warn)
+{
+	UCameraVariableCollection* NewVariableCollection = NewObject<UCameraVariableCollection>(Parent, Class, Name, Flags | RF_Transactional);
+	return NewVariableCollection;
+}
+
+#undef LOCTEXT_NAMESPACE
+
diff --git a/Engine/Plugins/Cameras/GameplayCameras/Source/GameplayCamerasEditor/Private/Factories/CameraVariableCollectionFactory.h b/Engine/Plugins/Cameras/GameplayCameras/Source/GameplayCamerasEditor/Private/Factories/CameraVariableCollectionFactory.h
new file mode 100644
--- /dev/null
+++ b/Engine/Plugins/Cameras/GameplayCameras/Source/GameplayCamerasEditor/Private/Factories/CameraVariableCollectionFactory.h
@@ -0,0 +1,27 @@
+// Copyright Epic Games, Inc. All Rights Reserved.
+
+#pragma once
+
+#include "Factories/Factory.h"
+
+#include "CameraVariableCollectionFactory.generated.h"
+
+class UCameraVariableCollection;
+
+/**
+ * Implements a factory for camera variable collections.
+ */
+UCLASS(hidecategories=Object)
+class UCameraVariableCollectionFactory : public UFactory
+{
+	GENERATED_BODY()
+
+public:
+
+	UCameraVariableCollectionFactory(const FObjectInitializer& ObjectInit);
+
+	// UFactory Interface
+	virtual FText GetDisplayName() const override;
+	virtual UObject* FactoryCreateNew(UClass* Class, UObject* Parent, FName Name, EObjectFlags Flags, UObject* Context, FFeedbackContext* Warn) override;
+};
+
diff --git a/Engine/Plugins/Cameras/GameplayCameras/Source/GameplayCamerasEditor/Private/Factories/CameraVariableFactories.cpp b/Engine/Plugins/Cameras/GameplayCameras/Source/GameplayCamerasEditor/Private/Factories/CameraVariableFactories.cpp
deleted file mode 100644
diff --git a/Engine/Plugins/Cameras/GameplayCameras/Source/GameplayCamerasEditor/Private/Factories/CameraVariableFactories.h b/Engine/Plugins/Cameras/GameplayCameras/Source/GameplayCamerasEditor/Private/Factories/CameraVariableFactories.h
deleted file mode 100644
diff --git a/Engine/Plugins/Cameras/GameplayCameras/Source/GameplayCamerasEditor/Private/GameplayCamerasEditorModule.cpp b/Engine/Plugins/Cameras/GameplayCameras/Source/GameplayCamerasEditor/Private/GameplayCamerasEditorModule.cpp
--- a/Engine/Plugins/Cameras/GameplayCameras/Source/GameplayCamerasEditor/Private/GameplayCamerasEditorModule.cpp
+++ b/Engine/Plugins/Cameras/GameplayCameras/Source/GameplayCamerasEditor/Private/GameplayCamerasEditorModule.cpp
@@ -4,9 +4,12 @@
 
 #include "AssetTools/CameraAssetEditor.h"
 #include "AssetTools/CameraRigAssetEditor.h"
+#include "AssetTools/CameraVariableCollectionEditor.h"
 #include "Commands/CameraAssetEditorCommands.h"
 #include "Commands/CameraRigAssetEditorCommands.h"
+#include "Commands/CameraVariableCollectionEditorCommands.h"
 #include "Commands/GameplayCamerasDebuggerCommands.h"
+#include "Customizations/CameraParameterDetailsCustomizations.h"
 #include "Debug/CameraDebugCategories.h"
 #include "Debugger/SBlendStacksDebugPanel.h"
 #include "Debugger/SCameraNodeTreeDebugPanel.h"
@@ -20,6 +23,7 @@
 #include "ISettingsModule.h"
 #include "Misc/CoreDelegates.h"
 #include "Modules/ModuleManager.h"
+#include "PropertyEditorModule.h"
 #include "Styles/GameplayCamerasEditorStyle.h"
 #include "ToolMenus.h"
 #include "Toolkits/CameraAssetEditorToolkit.h"
@@ -59,6 +63,8 @@ public:
 		RegisterSettings();
 		RegisterCoreDebugCategories();
 		RegisterRewindDebuggerFeatures();
+		RegisterDetailsCustomizations();
+
 		InitializeLiveEditManager();
 
 		UToolMenus::RegisterStartupCallback(FSimpleMulticastDelegate::FDelegate::CreateRaw(
@@ -73,11 +79,14 @@ public:
 
 		FCameraAssetEditorCommands::Unregister();
 		FCameraRigAssetEditorCommands::Unregister();
+		FCameraVariableCollectionEditorCommands::Unregister();
 		FGameplayCamerasDebuggerCommands::Unregister();
 
 		UnregisterSettings();
 		UnregisterCoreDebugCategories();
 		UnregisterRewindDebuggerFeatures();
+		UnregisterDetailsCustomizations();
+
 		TeardownLiveEditManager();
 
 		FCoreDelegates::OnPostEngineInit.RemoveAll(this);
@@ -100,6 +109,14 @@ public:
 		return AssetEditor;
 	}
 
+	virtual UCameraVariableCollectionEditor* CreateCameraVariableCollectionEditor(const EToolkitMode::Type Mode, const TSharedPtr<IToolkitHost>& InitToolkitHost, UCameraVariableCollection* VariableCollection) override
+	{
+		UAssetEditorSubsystem* AssetEditorSubsystem = GEditor->GetEditorSubsystem<UAssetEditorSubsystem>();
+		UCameraVariableCollectionEditor* AssetEditor = NewObject<UCameraVariableCollectionEditor>(AssetEditorSubsystem, NAME_None, RF_Transient);
+		AssetEditor->Initialize(VariableCollection);
+		return AssetEditor;
+	}
+
 	virtual void RegisterDebugCategory(const UE::Cameras::FCameraDebugCategoryInfo& InCategoryInfo) override
 	{
 		if (!ensureMsgf(!InCategoryInfo.Name.IsEmpty(), TEXT("A debug category must at least specify a name!")))
@@ -250,6 +267,7 @@ private:
 
 		FCameraAssetEditorCommands::Register();
 		FCameraRigAssetEditorCommands::Register();	
+		FCameraVariableCollectionEditorCommands::Register();
 		FGameplayCamerasDebuggerCommands::Register();
 	}
 
@@ -279,6 +297,22 @@ private:
 #endif  // UE_GAMEPLAY_CAMERAS_TRACE
 	}
 
+	void RegisterDetailsCustomizations()
+	{
+		using namespace UE::Cameras;
+
+		FPropertyEditorModule& PropertyEditorModule = FModuleManager::GetModuleChecked<FPropertyEditorModule>("PropertyEditor");
+		FCameraParameterDetailsCustomization::Register(PropertyEditorModule);
+	}
+
+	void UnregisterDetailsCustomizations()
+	{
+		using namespace UE::Cameras;
+
+		FPropertyEditorModule& PropertyEditorModule = FModuleManager::GetModuleChecked<FPropertyEditorModule>("PropertyEditor");
+		FCameraParameterDetailsCustomization::Unregister(PropertyEditorModule);
+	}
+
 	void InitializeLiveEditManager()
 	{
 		using namespace UE::Cameras;
diff --git a/Engine/Plugins/Cameras/GameplayCameras/Source/GameplayCamerasEditor/Private/Styles/GameplayCamerasEditorStyle.cpp b/Engine/Plugins/Cameras/GameplayCameras/Source/GameplayCamerasEditor/Private/Styles/GameplayCamerasEditorStyle.cpp
--- a/Engine/Plugins/Cameras/GameplayCameras/Source/GameplayCamerasEditor/Private/Styles/GameplayCamerasEditorStyle.cpp
+++ b/Engine/Plugins/Cameras/GameplayCameras/Source/GameplayCamerasEditor/Private/Styles/GameplayCamerasEditorStyle.cpp
@@ -6,6 +6,7 @@
 #include "Interfaces/IPluginManager.h"
 #include "Misc/Paths.h"
 #include "Styling/AppStyle.h"
+#include "Styling/CoreStyle.h"
 #include "Styling/SlateStyleRegistry.h"
 #include "Styling/SlateStyleMacros.h"
 #include "Styling/SlateTypes.h"
@@ -30,6 +31,8 @@ FGameplayCamerasEditorStyle::FGameplayCamerasEditorStyle()
 	SetContentRoot(ContentDir);
 	SetCoreContentRoot(FPaths::EngineContentDir() / TEXT("Slate"));
 
+	const FTextBlockStyle& NormalText = FAppStyle::Get().GetWidgetStyle<FTextBlockStyle>("NormalText");
+
 	// Camera rig editor icons.
 	Set("CameraRigAssetEditor.Tabs.Toolbox", new IMAGE_BRUSH_SVG("Icons/CameraRig-Toolbox", Icon16x16));
 	Set("CameraRigAssetEditor.Tabs.Search", new CORE_IMAGE_BRUSH_SVG("Starship/Common/search", Icon16x16));
@@ -52,6 +55,17 @@ FGameplayCamerasEditorStyle::FGameplayCamerasEditorStyle()
 	Set("CameraRigAssetEditor.Apply", new IMAGE_BRUSH_SVG("CameraEditor-Apply", Icon20x20));
 	Set("CameraRigAssetEditor.FindInCameraRig", new CORE_IMAGE_BRUSH_SVG("Starship/Common/Search", Icon20x20));
 
+	// Camera parameters icons.
+	Set("CameraParameter.VariableBrowser", new IMAGE_BRUSH_SVG("Icons/CameraParameter-Variable", Icon16x16));
+
+	// Camera variable collection icons.
+	Set("CameraVariableCollectionEditor.CreateVariable", new CORE_IMAGE_BRUSH_SVG("Starship/Common/plus", Icon16x16));
+	Set("CameraVariableCollectionEditor.DeleteVariable", new CORE_IMAGE_BRUSH_SVG("Starship/Common/minus", Icon16x16));
+
+	Set("CameraVariableCollectionEditor.Entry.Name", FTextBlockStyle(NormalText) .SetFont(DEFAULT_FONT("Bold", 12)));
+	Set("CameraVariableCollectionEditor.Entry.Type", FTextBlockStyle(NormalText) .SetFont(DEFAULT_FONT("Italic", 10)));
+	Set("CameraVariableCollectionEditor.Entry.Value", FTextBlockStyle(NormalText) .SetFont(DEFAULT_FONT("Regular", 10)));
+
 	// Debugger tool icons.
 	Set("Debugger.EnableDebugInfo.Icon", new IMAGE_BRUSH("Icons/EnableDebugInfo", Icon16x16));
 	Set("Debugger.DisableDebugInfo.Icon", new IMAGE_BRUSH("Icons/DisableDebugInfo", Icon16x16));
diff --git a/Engine/Plugins/Cameras/GameplayCameras/Source/GameplayCamerasEditor/Private/Toolkits/CameraVariableCollectionEditorToolkit.cpp b/Engine/Plugins/Cameras/GameplayCameras/Source/GameplayCamerasEditor/Private/Toolkits/CameraVariableCollectionEditorToolkit.cpp
new file mode 100644
--- /dev/null
+++ b/Engine/Plugins/Cameras/GameplayCameras/Source/GameplayCamerasEditor/Private/Toolkits/CameraVariableCollectionEditorToolkit.cpp
@@ -0,0 +1,568 @@
+// Copyright Epic Games, Inc. All Rights Reserved.
+
+#include "Toolkits/CameraVariableCollectionEditorToolkit.h"
+
+#include "AssetRegistry/AssetRegistryModule.h"
+#include "AssetTools/CameraVariableCollectionEditor.h"
+#include "Commands/CameraVariableCollectionEditorCommands.h"
+#include "ContentBrowserModule.h"
+#include "Core/CameraVariableAssets.h"
+#include "Core/CameraVariableCollection.h"
+#include "Editors/SCameraVariableCollectionEditor.h"
+#include "Framework/Docking/LayoutExtender.h"
+#include "Framework/Docking/TabManager.h"
+#include "IContentBrowserSingleton.h"
+#include "Misc/FeedbackContext.h"
+#include "Modules/ModuleManager.h"
+#include "ObjectTools.h"
+#include "PropertyEditorModule.h"
+#include "ScopedTransaction.h"
+#include "Styles/GameplayCamerasEditorStyle.h"
+#include "ToolMenuContext.h"
+#include "UObject/UObjectIterator.h"
+#include "Widgets/Docking/SDockTab.h"
+#include "Widgets/Input/SButton.h"
+
+#include UE_INLINE_GENERATED_CPP_BY_NAME(CameraVariableCollectionEditorToolkit)
+
+#define LOCTEXT_NAMESPACE "CameraVariableCollectionEditorToolkit"
+
+namespace UE::Cameras
+{
+
+const FName FCameraVariableCollectionEditorToolkit::VariableCollectionEditorTabId(TEXT("CameraVariableCollectionEditor_VariableCollectionEditor"));
+const FName FCameraVariableCollectionEditorToolkit::DetailsViewTabId(TEXT("CameraVariableCollectionEditor_DetailsView"));
+
+class SDeleteVariableDialog : public SCompoundWidget
+{
+public:
+
+	SLATE_BEGIN_ARGS(SDeleteVariableDialog)
+	{}
+		SLATE_ARGUMENT(TWeakPtr<SWindow>, ParentWindow)
+		SLATE_ARGUMENT(TSet<FName>, ReferencingPackages)
+	SLATE_END_ARGS()
+
+	void Construct(const FArguments& InArgs)
+	{
+		WeakParentWindow = InArgs._ParentWindow;
+
+		ReferencingPackages = InArgs._ReferencingPackages;
+
+		ChildSlot
+		[
+			SNew(SVerticalBox)
+
+			+ SVerticalBox::Slot()
+			.AutoHeight()
+			.Padding(5.0f)
+			[
+				SNew(SBorder)
+				.BorderBackgroundColor(FLinearColor::Green)
+				.BorderImage(FAppStyle::GetBrush("ToolPanel.GroupBorder"))
+				.Visibility(this, &SDeleteVariableDialog::GetNoReferencesVisibility)
+				.Padding(5.0f)
+				[
+					SNew(STextBlock)
+					.Text(LOCTEXT(
+								"VariablesOkToDelete", 
+								"No assets reference the variables being deleted."))
+				]
+			]
+
+			+ SVerticalBox::Slot()
+			.AutoHeight()
+			.Padding(5.0f)
+			[
+				SNew(SBorder)
+				.BorderBackgroundColor(FLinearColor::Red)
+				.BorderImage(FAppStyle::GetBrush("ToolPanel.GroupBorder"))
+				.Visibility(this, &SDeleteVariableDialog::GetReferencesVisiblity)
+				.Padding(5.0f)
+				[
+					SNew(STextBlock)
+					.Text(LOCTEXT(
+								"VariablesPendingDeleteAreInUse", 
+								"Some of the camera variables being deleted are referenced by camera assets."))
+				]
+			]
+
+			+ SVerticalBox::Slot()
+			.FillHeight(1.0f)
+			.Padding(5.0f)
+			[
+				SNew(SBorder)
+				.BorderImage(FAppStyle::GetBrush("ToolPanel.GroupBorder"))
+				.Padding(FMargin(0, 0, 0, 3))
+				.Visibility(this, &SDeleteVariableDialog::GetReferencesVisiblity)
+				[
+					SNew(SVerticalBox)
+
+					+ SVerticalBox::Slot()
+					.AutoHeight()
+					[
+						SNew(SBorder)
+						.BorderImage(FAppStyle::GetBrush("DetailsView.CategoryTop"))
+						.BorderBackgroundColor(FLinearColor(.6, .6, .6, 1.0f))
+						.Padding(3.0f)
+						[
+							SNew(STextBlock)
+							.Text(LOCTEXT(
+										"AssetsReferencingVariablesPendingDelete", 
+										"Assets Referencing the Camera Variables to Delete"))
+							.Font(FAppStyle::GetFontStyle("BoldFont"))
+							.ShadowOffset(FVector2D(1.0f, 1.0f))
+						]
+					]
+
+					+ SVerticalBox::Slot()
+					.FillHeight(1.0f)
+					[
+						SDeleteVariableDialog::BuildReferencerAssetPicker()
+					]
+				]
+			]
+
+			+ SVerticalBox::Slot()
+			.AutoHeight()
+			.Padding(0.f, 4.f)
+			[
+				SNew(SHorizontalBox)
+
+				+ SHorizontalBox::Slot()
+				.FillWidth(1.0f)
+				.Padding(6, 0)
+				[
+					SNew(SBorder)
+					.BorderImage(FAppStyle::GetBrush("NoBorder"))
+					[
+						SNew(SButton)
+						.HAlign(HAlign_Center)
+						.Text(LOCTEXT("Delete", "Delete"))
+						.ToolTipText(LOCTEXT("DeleteTooltipText", "Perform the delete"))
+						.ButtonStyle(FAppStyle::Get(), "FlatButton.Danger")
+						.TextStyle(FAppStyle::Get(), "FlatButton.DefaultTextStyle")
+						.OnClicked(this, &SDeleteVariableDialog::OnDeleteClicked)
+					]
+				]
+
+				+ SHorizontalBox::Slot()
+				.FillWidth(1.0f)
+				.Padding(6, 0)
+				[
+					SNew(SBorder)
+					.BorderImage(FAppStyle::GetBrush("NoBorder"))
+					[
+						SNew(SButton)
+						.HAlign(HAlign_Center)
+						.Text(LOCTEXT("Cancel", "Cancel"))
+						.ToolTipText(LOCTEXT("CancelDeleteTooltipText", "Cancel the delete"))
+						.ButtonStyle(FAppStyle::Get(), "FlatButton.Default")
+						.TextStyle(FAppStyle::Get(), "FlatButton.DefaultTextStyle")
+						.OnClicked(this, &SDeleteVariableDialog::OnCancelClicked)
+					]
+				]
+			]
+		];
+	}
+
+	bool ShouldPerformDelete() const
+	{
+		return bPerformDelete;
+	}
+	
+protected:
+
+	virtual FReply OnKeyDown(const FGeometry& MyGeometry, const FKeyEvent& InKeyEvent) override
+	{
+		if (InKeyEvent.GetKey() == EKeys::Escape)
+		{
+			OnCancelClicked();
+			return FReply::Handled();
+		}
+		return FReply::Unhandled();
+	}
+
+private:
+
+	EVisibility GetNoReferencesVisibility() const
+	{
+		return ReferencingPackages.IsEmpty() ? EVisibility::Visible : EVisibility::Collapsed;
+	}
+
+	EVisibility GetReferencesVisiblity() const
+	{
+		return ReferencingPackages.IsEmpty() ? EVisibility::Collapsed : EVisibility::Visible;
+	}
+
+	TSharedRef<SWidget> BuildReferencerAssetPicker()
+	{
+		FContentBrowserModule& ContentBrowserModule = FModuleManager::Get().LoadModuleChecked<FContentBrowserModule>(TEXT("ContentBrowser"));
+
+		FARFilter ARFilter;
+		ARFilter.PackageNames = ReferencingPackages.Array();
+
+		FAssetPickerConfig AssetPickerConfig;
+		AssetPickerConfig.bAllowDragging = false;
+		AssetPickerConfig.bCanShowClasses = false;
+		AssetPickerConfig.bAllowNullSelection = false;
+		AssetPickerConfig.bShowBottomToolbar = false;
+		AssetPickerConfig.bAutohideSearchBar = true;
+		AssetPickerConfig.Filter = ARFilter;
+		AssetPickerConfig.InitialAssetViewType = EAssetViewType::Tile;
+		AssetPickerConfig.OnAssetsActivated = FOnAssetsActivated::CreateSP(this, &SDeleteVariableDialog::OnAssetsActivated);
+
+		return ContentBrowserModule.Get().CreateAssetPicker(AssetPickerConfig);
+	}
+
+	void OnAssetsActivated(const TArray<FAssetData>& ActivatedAssets, EAssetTypeActivationMethod::Type ActivationMethod)
+	{
+		if (ActivationMethod == EAssetTypeActivationMethod::DoubleClicked || ActivationMethod == EAssetTypeActivationMethod::Opened)
+		{
+			CloseWindow();
+
+			for (const FAssetData& ActivatedAsset : ActivatedAssets)
+			{
+				GEditor->GetEditorSubsystem<UAssetEditorSubsystem>()->OpenEditorForAsset(ActivatedAsset.GetAsset());
+			}
+		}
+	}
+
+	void CloseWindow()
+	{
+		if (TSharedPtr<SWindow> ParentWindow = WeakParentWindow.Pin())
+		{
+			ParentWindow->RequestDestroyWindow();
+		}
+	}
+
+	FReply OnDeleteClicked()
+	{
+		bPerformDelete = true;
+		CloseWindow();
+		return FReply::Handled();
+	}
+
+	FReply OnCancelClicked()
+	{
+		bPerformDelete = false;
+		CloseWindow();
+		return FReply::Handled();
+	}
+
+private:
+
+	TWeakPtr<SWindow> WeakParentWindow;
+
+	TSet<FName> ReferencingPackages;
+
+	bool bPerformDelete = false;
+};
+
+FCameraVariableCollectionEditorToolkit::FCameraVariableCollectionEditorToolkit(UCameraVariableCollectionEditor* InOwningAssetEditor)
+	: FBaseAssetToolkit(InOwningAssetEditor)
+	, VariableCollection(InOwningAssetEditor->GetVariableCollection())
+	, CommandBindings(new FUICommandList())
+{
+	// Override base class default layout.
+	StandaloneDefaultLayout = FTabManager::NewLayout("CameraVariableCollectionEditor_Layout_v1")
+		->AddArea
+		(
+			FTabManager::NewPrimaryArea()->SetOrientation(Orient_Horizontal)
+			->Split
+			(
+				FTabManager::NewStack()
+				->SetSizeCoefficient(0.8f)
+				->SetHideTabWell(true)
+				->AddTab(VariableCollectionEditorTabId, ETabState::OpenedTab)
+			)
+			->Split
+			(
+				FTabManager::NewStack()
+				->SetSizeCoefficient(0.2f)
+				->SetHideTabWell(true)
+				->AddTab(DetailsViewTabId, ETabState::OpenedTab)
+			)
+		);
+}
+
+FCameraVariableCollectionEditorToolkit::~FCameraVariableCollectionEditorToolkit()
+{
+}
+
+void FCameraVariableCollectionEditorToolkit::RegisterTabSpawners(const TSharedRef<class FTabManager>& InTabManager)
+{
+	// Skip FBaseAssetToolkit here because we don't want a viewport tab.
+	FAssetEditorToolkit::RegisterTabSpawners(InTabManager);
+
+	const FName CamerasStyleSetName = FGameplayCamerasEditorStyle::Get()->GetStyleSetName();
+
+	InTabManager->RegisterTabSpawner(VariableCollectionEditorTabId, FOnSpawnTab::CreateSP(this, &FCameraVariableCollectionEditorToolkit::SpawnTab_VariableCollectionEditor))
+		.SetDisplayName(LOCTEXT("VariableCollectionEditor", "Camera Variable Collection"))
+		.SetGroup(AssetEditorTabsCategory.ToSharedRef());
+
+	InTabManager->RegisterTabSpawner(DetailsViewTabId, FOnSpawnTab::CreateSP(this, &FCameraVariableCollectionEditorToolkit::SpawnTab_Details))
+		.SetDisplayName(LOCTEXT("Details", "Details"))
+		.SetGroup(AssetEditorTabsCategory.ToSharedRef())
+		.SetIcon(FSlateIcon(FAppStyle::GetAppStyleSetName(), "LevelEditor.Tabs.Details"));
+}
+
+TSharedRef<SDockTab> FCameraVariableCollectionEditorToolkit::SpawnTab_VariableCollectionEditor(const FSpawnTabArgs& Args)
+{
+	TSharedPtr<SDockTab> VariableCollectionEditorTab = SNew(SDockTab)
+		.Label(LOCTEXT("VariableCollectionEditorTabTitle", "Camera Variable Collection"))
+		[
+			VariableCollectionEditorWidget.ToSharedRef()
+		];
+
+	return VariableCollectionEditorTab.ToSharedRef();
+}
+
+void FCameraVariableCollectionEditorToolkit::UnregisterTabSpawners(const TSharedRef<FTabManager>& InTabManager)
+{
+	// Skip FBaseAssetToolkit here because we don't want a viewport tab.
+	FAssetEditorToolkit::UnregisterTabSpawners(InTabManager);
+
+	InTabManager->UnregisterTabSpawner(VariableCollectionEditorTabId);
+	InTabManager->UnregisterTabSpawner(DetailsViewTabId);
+}
+
+void FCameraVariableCollectionEditorToolkit::CreateWidgets()
+{
+	// Skip FBaseAssetToolkit here because we don't want a viewport tab.
+	// ...no up-call...
+
+	RegisterToolbar();
+	LayoutExtender = MakeShared<FLayoutExtender>();
+
+	FPropertyEditorModule& PropertyEditorModule = FModuleManager::GetModuleChecked<FPropertyEditorModule>("PropertyEditor");
+	FDetailsViewArgs DetailsViewArgs;
+	DetailsViewArgs.NameAreaSettings = FDetailsViewArgs::HideNameArea;
+	DetailsViewArgs.bHideSelectionTip = true;
+	DetailsView = PropertyEditorModule.CreateDetailView(DetailsViewArgs);
+	
+	// Now do our custom stuff.
+
+	// Create the variable collection editor.
+	VariableCollectionEditorWidget = SNew(SCameraVariableCollectionEditor)
+		.DetailsView(DetailsView)
+		.VariableCollection(VariableCollection);
+}
+
+void FCameraVariableCollectionEditorToolkit::RegisterToolbar()
+{
+	FName ParentName;
+	const FName MenuName = GetToolMenuToolbarName(ParentName);
+	UToolMenus* ToolMenus = UToolMenus::Get();
+	if (!ToolMenus->IsMenuRegistered(MenuName))
+	{
+		FToolMenuOwnerScoped ToolMenuOwnerScope(this);
+
+		UToolMenu* ToolbarMenu = UToolMenus::Get()->RegisterMenu(
+				MenuName, ParentName, EMultiBoxType::ToolBar);
+
+		FToolMenuInsert InsertAfterAssetSection("Asset", EToolMenuInsertType::After);
+		const FCameraVariableCollectionEditorCommands& Commands = FCameraVariableCollectionEditorCommands::Get();
+		const FName CamerasStyleSetName = FGameplayCamerasEditorStyle::Get()->GetStyleSetName();
+
+		FToolMenuSection& VariablesSection = ToolbarMenu->AddSection("Variables", TAttribute<FText>(), InsertAfterAssetSection);
+
+		FToolMenuEntry CreateVariableEntry = FToolMenuEntry::InitComboButton(
+				"CreateVariable",
+				FUIAction(),
+				FNewToolMenuDelegate::CreateStatic(&FCameraVariableCollectionEditorToolkit::GenerateAddNewVariableMenu),
+				LOCTEXT("CreateVariableCombo_Label", "Add"),
+				LOCTEXT("CreateVariableCombo_ToolTip", "Add a new camera variable to the collection"),
+				FSlateIcon(CamerasStyleSetName, "CameraVariableCollectionEditor.CreateVariable")
+				);
+		VariablesSection.AddEntry(CreateVariableEntry);
+
+		FToolMenuEntry DeleteVariableEntry = FToolMenuEntry::InitToolBarButton(Commands.DeleteVariable);
+		VariablesSection.AddEntry(DeleteVariableEntry);
+	}
+}
+
+void FCameraVariableCollectionEditorToolkit::GenerateAddNewVariableMenu(UToolMenu* InMenu)
+{
+	UCameraVariableCollectionEditorMenuContext* Context = InMenu->FindContext<UCameraVariableCollectionEditorMenuContext>();
+	if (!ensure(Context))
+	{
+		return;
+	}
+
+	FCameraVariableCollectionEditorToolkit* This = Context->EditorToolkit.Pin().Get();
+	if (!ensure(This))
+	{
+		return;
+	}
+
+	const FCameraVariableCollectionEditorCommands& Commands = FCameraVariableCollectionEditorCommands::Get();
+	FToolMenuSection& VariableTypesSection = InMenu->AddSection("VariableTypes");
+
+	for (TObjectIterator<UClass> It; It; ++It)
+	{
+		UClass* VariableClass = *It;
+		if (VariableClass->IsChildOf<UCameraVariableAsset>() &&
+				!VariableClass->HasAnyClassFlags(CLASS_Abstract))
+		{
+			const FText VariableTypeDisplayName(VariableClass->GetDisplayNameText());
+			VariableTypesSection.AddEntry(FToolMenuEntry::InitMenuEntry(
+						FName(FString::Format(TEXT("AddCameraVariable_{0}"), { VariableClass->GetName() })),
+						TAttribute<FText>(VariableTypeDisplayName),
+						TAttribute<FText>(FText::Format(
+								LOCTEXT("CreateVariableEntry_LabelFmt", "Add a {0} to the collection"), VariableTypeDisplayName)),
+						TAttribute<FSlateIcon>(),
+						FExecuteAction::CreateSP(
+							This, &FCameraVariableCollectionEditorToolkit::OnCreateVariable, 
+							TSubclassOf<UCameraVariableAsset>(VariableClass))
+						));
+		}
+	}
+}
+
+void FCameraVariableCollectionEditorToolkit::InitToolMenuContext(FToolMenuContext& MenuContext)
+{
+	FBaseAssetToolkit::InitToolMenuContext(MenuContext);
+
+	UCameraVariableCollectionEditorMenuContext* Context = NewObject<UCameraVariableCollectionEditorMenuContext>();
+	Context->EditorToolkit = SharedThis(this);
+	MenuContext.AddObject(Context);
+}
+
+void FCameraVariableCollectionEditorToolkit::PostInitAssetEditor()
+{
+	const FCameraVariableCollectionEditorCommands& Commands = FCameraVariableCollectionEditorCommands::Get();
+
+	ToolkitCommands->MapAction(
+		Commands.DeleteVariable,
+		FExecuteAction::CreateSP(this, &FCameraVariableCollectionEditorToolkit::OnDeleteVariable),
+		FCanExecuteAction::CreateSP(this, &FCameraVariableCollectionEditorToolkit::CanDeleteVariable));
+}
+
+FText FCameraVariableCollectionEditorToolkit::GetBaseToolkitName() const
+{
+	return LOCTEXT("AppLabel", "Camera Variable Collection");
+}
+
+FName FCameraVariableCollectionEditorToolkit::GetToolkitFName() const
+{
+	static FName ToolkitName("CameraVariableCollectionEditor");
+	return ToolkitName;
+}
+
+FString FCameraVariableCollectionEditorToolkit::GetWorldCentricTabPrefix() const
+{
+	return LOCTEXT("WorldCentricTabPrefix", "Camera Variable Collection ").ToString();
+}
+
+FLinearColor FCameraVariableCollectionEditorToolkit::GetWorldCentricTabColorScale() const
+{
+	return FLinearColor(0.1f, 0.8f, 0.2f, 0.5f);
+}
+
+void FCameraVariableCollectionEditorToolkit::OnCreateVariable(TSubclassOf<UCameraVariableAsset> InVariableClass)
+{
+	UCameraVariableAsset* NewVariable = NewObject<UCameraVariableAsset>(VariableCollection, InVariableClass.Get());
+	VariableCollection->Variables.Add(NewVariable);
+
+	VariableCollectionEditorWidget->RequestListRefresh();
+}
+
+void FCameraVariableCollectionEditorToolkit::OnDeleteVariable()
+{
+	TArray<UCameraVariableAsset*> Selection;
+	VariableCollectionEditorWidget->GetSelectedVariables(Selection);
+	if (Selection.IsEmpty())
+	{
+		return;
+	}
+
+	GWarn->BeginSlowTask(LOCTEXT("PreDeleteScanning", "Scanning assets before deleting camera variables"), true);
+
+	UPackage* VariableCollectionPackage = VariableCollection->GetOutermost();
+	const FName VariableCollectionPackageName = VariableCollectionPackage->GetFName();
+
+	TSet<FName> AllReferencers;
+
+	{
+		FAssetRegistryModule& AssetRegistryModule = FModuleManager::LoadModuleChecked<FAssetRegistryModule>("AssetRegistry");
+		IAssetRegistry& AssetRegistry = AssetRegistryModule.Get();
+
+		TArray<FName> OnDiskReferencers;
+		AssetRegistry.GetReferencers(VariableCollectionPackageName, OnDiskReferencers);
+		AllReferencers.Append(OnDiskReferencers);
+	}
+
+	for (UCameraVariableAsset* CameraVariable : Selection)
+	{
+		bool bIsReferencedInMemoryByNonUndo = false;
+		bool bIsReferencedInMemoryByUndo = false;
+		FReferencerInformationList MemoryReferences;
+		ObjectTools::GatherObjectReferencersForDeletion(
+				CameraVariable, bIsReferencedInMemoryByNonUndo, bIsReferencedInMemoryByUndo, &MemoryReferences);
+
+		UPackage* TransientPackage = GetTransientPackage();
+		for (const FReferencerInformation& ExternalReference : MemoryReferences.ExternalReferences)
+		{
+			UPackage* ExternalReferencePackage = ExternalReference.Referencer->GetOutermost();
+			if (ExternalReferencePackage != VariableCollectionPackage && ExternalReferencePackage != TransientPackage)
+			{
+				AllReferencers.Add(ExternalReferencePackage->GetFName());
+			}
+		}
+	}
+
+	GWarn->EndSlowTask();
+
+	bool bPerformDelete = false;
+	{
+		TSharedRef<SWindow> DeleteVariableWindow = SNew(SWindow)
+			.Title(LOCTEXT("DeleteVariable", "Delete Variables"))
+			.ClientSize(FVector2D(600, 700));
+
+		TSharedRef<SDeleteVariableDialog> DeleteVariableDialog = SNew(SDeleteVariableDialog)
+			.ParentWindow(DeleteVariableWindow)
+			.ReferencingPackages(AllReferencers);
+		DeleteVariableWindow->SetContent(DeleteVariableDialog);
+
+		GEditor->EditorAddModalWindow(DeleteVariableWindow);
+
+		bPerformDelete = DeleteVariableDialog->ShouldPerformDelete();
+	}
+
+	if (bPerformDelete)
+	{
+		FScopedTransaction DeleteTransaction(LOCTEXT("DeleteVariable", "Delete camera variable"));
+
+		VariableCollection->Modify();
+
+		TArray<UObject*> ObjectsToReplace(Selection);
+		ObjectTools::ForceReplaceReferences(nullptr, ObjectsToReplace);
+
+		TStringBuilder<256> StringBuilder;
+		for (UCameraVariableAsset* VariableToDelete : Selection)
+		{
+			VariableCollection->Variables.Remove(VariableToDelete);
+
+			StringBuilder.Reset();
+			StringBuilder.Append("TRASH_");
+			StringBuilder.Append(VariableToDelete->GetName());
+			VariableToDelete->Rename(StringBuilder.ToString());
+			VariableToDelete->MarkAsGarbage();
+		}
+
+		VariableCollectionEditorWidget->RequestListRefresh();
+	}
+}
+
+bool FCameraVariableCollectionEditorToolkit::CanDeleteVariable()
+{
+	TArray<UCameraVariableAsset*> Selection;
+	VariableCollectionEditorWidget->GetSelectedVariables(Selection);
+	return !Selection.IsEmpty();
+}
+
+}  // namespace UE::Cameras
+
+#undef LOCTEXT_NAMESPACE
+
diff --git a/Engine/Plugins/Cameras/GameplayCameras/Source/GameplayCamerasEditor/Public/AssetTools/CameraVariableCollectionEditor.h b/Engine/Plugins/Cameras/GameplayCameras/Source/GameplayCamerasEditor/Public/AssetTools/CameraVariableCollectionEditor.h
new file mode 100644
--- /dev/null
+++ b/Engine/Plugins/Cameras/GameplayCameras/Source/GameplayCamerasEditor/Public/AssetTools/CameraVariableCollectionEditor.h
@@ -0,0 +1,36 @@
+// Copyright Epic Games, Inc. All Rights Reserved.
+
+#pragma once
+
+#include "Tools/UAssetEditor.h"
+
+#include "CameraVariableCollectionEditor.generated.h"
+
+class FBaseAssetToolkit;
+class UCameraVariableCollection;
+
+/**
+ * Editor for a camera variable collection.
+ */
+UCLASS(Transient)
+class UCameraVariableCollectionEditor : public UAssetEditor
+{
+	GENERATED_BODY()
+
+public:
+
+	void Initialize(TObjectPtr<UCameraVariableCollection> InVariableCollection);
+
+	UCameraVariableCollection* GetVariableCollection() const { return VariableCollection; }
+
+public:
+
+	// UAssetEditor interface
+	virtual void GetObjectsToEdit(TArray<UObject*>& InObjectsToEdit) override;
+	virtual TSharedPtr<FBaseAssetToolkit> CreateToolkit() override;
+
+private:
+
+	TObjectPtr<UCameraVariableCollection> VariableCollection;
+};
+
diff --git a/Engine/Plugins/Cameras/GameplayCameras/Source/GameplayCamerasEditor/Public/Commands/CameraVariableCollectionEditorCommands.h b/Engine/Plugins/Cameras/GameplayCameras/Source/GameplayCamerasEditor/Public/Commands/CameraVariableCollectionEditorCommands.h
new file mode 100644
--- /dev/null
+++ b/Engine/Plugins/Cameras/GameplayCameras/Source/GameplayCamerasEditor/Public/Commands/CameraVariableCollectionEditorCommands.h
@@ -0,0 +1,30 @@
+// Copyright Epic Games, Inc. All Rights Reserved.
+
+#pragma once
+
+#include "CoreTypes.h"
+#include "Framework/Commands/Commands.h"
+
+class FUICommandInfo;
+class UClass;
+struct FInputChord;
+
+namespace UE::Cameras
+{
+
+class FCameraVariableCollectionEditorCommands : public TCommands<FCameraVariableCollectionEditorCommands>
+{
+public:
+
+	FCameraVariableCollectionEditorCommands();
+
+	virtual void RegisterCommands() override;
+	
+public:
+
+	TSharedPtr<FUICommandInfo> CreateVariable;
+	TSharedPtr<FUICommandInfo> DeleteVariable;
+};
+
+}  // namespace UE::Cameras
+
diff --git a/Engine/Plugins/Cameras/GameplayCameras/Source/GameplayCamerasEditor/Public/Customizations/CameraParameterDetailsCustomizations.h b/Engine/Plugins/Cameras/GameplayCameras/Source/GameplayCamerasEditor/Public/Customizations/CameraParameterDetailsCustomizations.h
new file mode 100644
--- /dev/null
+++ b/Engine/Plugins/Cameras/GameplayCameras/Source/GameplayCamerasEditor/Public/Customizations/CameraParameterDetailsCustomizations.h
@@ -0,0 +1,87 @@
+// Copyright Epic Games, Inc. All Rights Reserved.
+
+#pragma once
+
+#include "Core/CameraParameters.h"
+#include "CoreTypes.h"
+#include "IPropertyTypeCustomization.h"
+
+class FPropertyEditorModule;
+class IDetailLayoutBuilder;
+class IPropertyUtilities;
+class SComboButton;
+class SWidget;
+
+namespace UE::Cameras
+{
+
+/**
+ * Base details customization for camera parameters.
+ */
+class FCameraParameterDetailsCustomization : public IPropertyTypeCustomization
+{
+public:
+
+	/** Registers details customizations for all camera parameter types. */
+	static void Register(FPropertyEditorModule& PropertyEditorModule);
+	/** Unregisters details customizations for all camera parameter types. */
+	static void Unregister(FPropertyEditorModule& PropertyEditorModule);
+
+public:
+
+	/** IPropertyTypeCustomization interface */
+	virtual void CustomizeHeader(TSharedRef<IPropertyHandle> PropertyHandle, FDetailWidgetRow& HeaderRow, IPropertyTypeCustomizationUtils& CustomizationUtils) override;
+	virtual void CustomizeChildren(TSharedRef<IPropertyHandle> PropertyHandle, IDetailChildrenBuilder& ChildBuilder, IPropertyTypeCustomizationUtils& CustomizationUtils) override;
+
+protected:
+
+	virtual void GetValueAndVariablePropertyNames(FName& OutValueProperty, FName& OutVariableProperty) = 0;
+	virtual void SetParameterVariable(void* InRawData, UCameraVariableAsset* InVariable) = 0;
+
+private:
+
+	TSharedRef<SWidget> BuildCameraVariableBrowser();
+
+	void UpdateVariableInfo();
+	bool HasVariableInfoText() const;
+
+	bool CanClearVariable() const;
+	void OnClearVariable();
+
+	void OnSetVariable(UCameraVariableAsset* InVariable);
+	void OnResetToDefault();
+
+protected:
+
+	UClass* VariableClass = nullptr;
+
+	UCameraVariableAsset* CommonVariable = nullptr;
+	FText VariableInfoText;
+	FText VariableErrorText;
+
+	TSharedPtr<IPropertyUtilities> PropertyUtilities;
+
+	TSharedPtr<IPropertyHandle> StructProperty;
+	TSharedPtr<IPropertyHandle> ValueProperty;
+	TSharedPtr<IPropertyHandle> VariableProperty;
+
+	TSharedPtr<SComboButton> VariableBrowserButton;
+};
+
+// Create all the individual classes.
+#define UE_CAMERA_VARIABLE_FOR_TYPE(ValueType, ValueName)\
+class F##ValueName##CameraParameterDetailsCustomization : public FCameraParameterDetailsCustomization\
+{\
+protected:\
+	virtual void GetValueAndVariablePropertyNames(FName& OutValueProperty, FName& OutVariableProperty) override\
+	{\
+		OutValueProperty = GET_MEMBER_NAME_CHECKED(F##ValueName##CameraParameter, Value);\
+		OutVariableProperty = GET_MEMBER_NAME_CHECKED(F##ValueName##CameraParameter, Variable);\
+	}\
+	virtual void SetParameterVariable(void* InRawData, UCameraVariableAsset* InVariable) override;\
+};
+UE_CAMERA_VARIABLE_FOR_ALL_TYPES()
+#undef UE_CAMERA_VARIABLE_FOR_TYPE
+
+}  // namespace UE::Cameras
+ 
diff --git a/Engine/Plugins/Cameras/GameplayCameras/Source/GameplayCamerasEditor/Public/Editors/SCameraVariableCollectionEditor.h b/Engine/Plugins/Cameras/GameplayCameras/Source/GameplayCamerasEditor/Public/Editors/SCameraVariableCollectionEditor.h
new file mode 100644
--- /dev/null
+++ b/Engine/Plugins/Cameras/GameplayCameras/Source/GameplayCamerasEditor/Public/Editors/SCameraVariableCollectionEditor.h
@@ -0,0 +1,89 @@
+// Copyright Epic Games, Inc. All Rights Reserved.
+
+#pragma once
+
+#include "CoreTypes.h"
+#include "Misc/TextFilter.h"
+#include "Templates/SharedPointerFwd.h"
+#include "Templates/SubclassOf.h"
+#include "UObject/ObjectPtr.h"
+#include "Widgets/SCompoundWidget.h"
+
+class FAssetEditorToolkit;
+class IDetailsView;
+class ITableRow;
+class SSearchBox;
+class STableViewBase;
+class UCameraVariableAsset;
+class UCameraVariableCollection;
+struct FToolMenuContext;
+template<typename> class SListView;
+
+namespace UE::Cameras
+{
+
+/**
+ * An editor widget for a camera variable collection.
+ */
+class SCameraVariableCollectionEditor : public SCompoundWidget
+{
+public:
+
+	SLATE_BEGIN_ARGS(SCameraVariableCollectionEditor)
+	{}
+		/** The camera variable collection to edit. */
+		SLATE_ARGUMENT(TObjectPtr<UCameraVariableCollection>, VariableCollection)
+		/** The details view to synchronize with the variable list selection. */
+		SLATE_ARGUMENT(TWeakPtr<IDetailsView>, DetailsView)
+		/** The toolkit inside which this editor lives, if any. */
+		SLATE_ARGUMENT(TWeakPtr<FAssetEditorToolkit>, AssetEditorToolkit)
+	SLATE_END_ARGS()
+
+	void Construct(const FArguments& InArgs);
+	~SCameraVariableCollectionEditor();
+
+public:
+
+	/** Gets the selected variables in the list view. */
+	void GetSelectedVariables(TArray<UCameraVariableAsset*>& OutSelection) const;
+
+	/** Requests that the list view be refreshed by the next tick. */
+	void RequestListRefresh();
+
+protected:
+
+	//~ SWidget interface.
+	virtual void Tick(const FGeometry& AllottedGeometry, const double InCurrentTime, const float InDeltaTime) override;
+
+private:
+
+	void UpdateFilteredItemSource();
+	void SetDetailsViewObject(UObject* InObject) const;
+
+	TSharedRef<ITableRow> OnListGenerateRow(UCameraVariableAsset* Item, const TSharedRef<STableViewBase>& OwnerTable);
+	void OnListSectionChanged(UCameraVariableAsset* Item, ESelectInfo::Type SelectInfo) const;
+
+	void GetEntryStrings(const UCameraVariableAsset* InItem, TArray<FString>& OutStrings);
+	void OnSearchTextChanged(const FText& InFilterText);
+	void OnSearchTextCommitted(const FText& InFilterText, ETextCommit::Type InCommitType);
+	FText GetHighlightText() const;
+
+private:
+
+	TObjectPtr<UCameraVariableCollection> VariableCollection;
+
+	TWeakPtr<IDetailsView> WeakDetailsView;
+
+	TSharedPtr<SListView<UCameraVariableAsset*>> ListView;
+
+	TArray<UCameraVariableAsset*> FilteredItemSource;
+
+	using FEntryTextFilter = TTextFilter<const UCameraVariableAsset*>;
+	TSharedPtr<FEntryTextFilter> SearchTextFilter;
+	TSharedPtr<SSearchBox> SearchBox;
+
+	bool bUpdateFilteredItemSource = false;
+};
+
+}  // namespace UE::Cameras
+
diff --git a/Engine/Plugins/Cameras/GameplayCameras/Source/GameplayCamerasEditor/Public/IGameplayCamerasEditorModule.h b/Engine/Plugins/Cameras/GameplayCameras/Source/GameplayCamerasEditor/Public/IGameplayCamerasEditorModule.h
--- a/Engine/Plugins/Cameras/GameplayCameras/Source/GameplayCamerasEditor/Public/IGameplayCamerasEditorModule.h
+++ b/Engine/Plugins/Cameras/GameplayCameras/Source/GameplayCamerasEditor/Public/IGameplayCamerasEditorModule.h
@@ -11,6 +11,8 @@ class UCameraAsset;
 class UCameraAssetEditor;
 class UCameraRigAsset;
 class UCameraRigAssetEditor;
+class UCameraVariableCollection;
+class UCameraVariableCollectionEditor;
 
 DECLARE_DELEGATE_RetVal_OneParam(TSharedRef<SWidget>, FOnCreateDebugCategoryPanel, const FString&);
 
@@ -47,6 +49,9 @@ public:
 	/** Creates an editor for the given camera rig asset */
 	virtual UCameraRigAssetEditor* CreateCameraRigEditor(const EToolkitMode::Type Mode, const TSharedPtr<IToolkitHost>& InitToolkitHost, UCameraRigAsset* CameraRig) = 0;
 
+	/** Creates an editor for the given variable collection */
+	virtual UCameraVariableCollectionEditor* CreateCameraVariableCollectionEditor(const EToolkitMode::Type Mode, const TSharedPtr<IToolkitHost>& InitToolkitHost, UCameraVariableCollection* VariableCollection) = 0;
+
 public:
 
 	using FCameraDebugCategoryInfo = UE::Cameras::FCameraDebugCategoryInfo;
diff --git a/Engine/Plugins/Cameras/GameplayCameras/Source/GameplayCamerasEditor/Public/Toolkits/CameraVariableCollectionEditorToolkit.h b/Engine/Plugins/Cameras/GameplayCameras/Source/GameplayCamerasEditor/Public/Toolkits/CameraVariableCollectionEditorToolkit.h
new file mode 100644
--- /dev/null
+++ b/Engine/Plugins/Cameras/GameplayCameras/Source/GameplayCamerasEditor/Public/Toolkits/CameraVariableCollectionEditorToolkit.h
@@ -0,0 +1,95 @@
+// Copyright Epic Games, Inc. All Rights Reserved.
+
+#pragma once
+
+#include "CoreTypes.h"
+#include "Tools/BaseAssetToolkit.h"
+#include "UObject/GCObject.h"
+
+#include "CameraVariableCollectionEditorToolkit.generated.h"
+
+class UCameraVariableAsset;
+class UCameraVariableCollection;
+class UCameraVariableCollectionEditor;
+
+namespace UE::Cameras
+{
+
+class SCameraVariableCollectionEditor;
+
+/**
+ * Editor toolkit for a camera variable collection.
+ */
+class FCameraVariableCollectionEditorToolkit
+	: public FBaseAssetToolkit
+	, public FGCObject
+{
+public:
+
+	FCameraVariableCollectionEditorToolkit(UCameraVariableCollectionEditor* InOwningAssetEditor);
+	~FCameraVariableCollectionEditorToolkit();
+
+protected:
+
+	// FBaseAssetToolkit interface
+	virtual void RegisterTabSpawners(const TSharedRef<class FTabManager>& InTabManager) override;
+	virtual void UnregisterTabSpawners(const TSharedRef<FTabManager>& InTabManager) override;
+	virtual void CreateWidgets() override;
+	virtual void RegisterToolbar() override;
+	virtual void InitToolMenuContext(FToolMenuContext& MenuContext) override;
+	virtual void PostInitAssetEditor() override;
+
+	// IToolkit interface
+	virtual FText GetBaseToolkitName() const override;
+	virtual FName GetToolkitFName() const override;
+	virtual FString GetWorldCentricTabPrefix() const override;
+	virtual FLinearColor GetWorldCentricTabColorScale() const override;
+
+	// FGCObject interface
+	virtual void AddReferencedObjects(FReferenceCollector& Collector) override
+	{
+		Collector.AddReferencedObject(VariableCollection);
+	}
+	virtual FString GetReferencerName() const override
+	{
+		return TEXT("FCameraVariableCollectionEditorToolkit");
+	}
+
+private:
+
+	TSharedRef<SDockTab> SpawnTab_VariableCollectionEditor(const FSpawnTabArgs& Args);
+
+	static void GenerateAddNewVariableMenu(UToolMenu* InMenu);
+
+	void OnCreateVariable(TSubclassOf<UCameraVariableAsset> InVariableClass);
+
+	void OnDeleteVariable();
+	bool CanDeleteVariable();
+
+private:
+
+	static const FName VariableCollectionEditorTabId;
+	static const FName DetailsViewTabId;
+
+	/** The asset being edited */
+	TObjectPtr<UCameraVariableCollection> VariableCollection;
+
+	/** Command bindings */
+	TSharedRef<FUICommandList> CommandBindings;
+
+	/** Camera variable collection editor widget */
+	TSharedPtr<SCameraVariableCollectionEditor> VariableCollectionEditorWidget;
+};
+
+}  // namespace UE::Cameras
+
+UCLASS()
+class UCameraVariableCollectionEditorMenuContext : public UObject
+{
+	GENERATED_BODY()
+
+public:
+
+	TWeakPtr<UE::Cameras::FCameraVariableCollectionEditorToolkit> EditorToolkit;
+};
+
