Cameras: various minor code fixes

diff --git a/Engine/Plugins/Cameras/GameplayCameras/Source/GameplayCamerasEditor/Private/GameplayCamerasEditorModule.cpp b/Engine/Plugins/Cameras/GameplayCameras/Source/GameplayCamerasEditor/Private/GameplayCamerasEditorModule.cpp
--- a/Engine/Plugins/Cameras/GameplayCameras/Source/GameplayCamerasEditor/Private/GameplayCamerasEditorModule.cpp
+++ b/Engine/Plugins/Cameras/GameplayCameras/Source/GameplayCamerasEditor/Private/GameplayCamerasEditorModule.cpp
@@ -1,6 +1,6 @@
 // Copyright Epic Games, Inc. All Rights Reserved.
 
-#include "GameplayCamerasEditorSettings.h"
+#include "IGameplayCamerasEditorModule.h"
 
 #include "AssetTools/CameraAssetEditor.h"
 #include "AssetTools/CameraRigAssetEditor.h"
@@ -16,6 +16,7 @@
 #include "Debugger/SGameplayCamerasDebugger.h"
 #include "Features/IModularFeatures.h"
 #include "GameplayCameras.h"
+#include "GameplayCamerasEditorSettings.h"
 #include "GameplayCamerasLiveEditManager.h"
 #include "IGameplayCamerasEditorModule.h"
 #include "IGameplayCamerasModule.h"
diff --git a/Engine/Plugins/Cameras/GameplayCameras/Source/GameplayCamerasEditor/Private/Toolkits/CameraRigAssetEditorToolkit.cpp b/Engine/Plugins/Cameras/GameplayCameras/Source/GameplayCamerasEditor/Private/Toolkits/CameraRigAssetEditorToolkit.cpp
--- a/Engine/Plugins/Cameras/GameplayCameras/Source/GameplayCamerasEditor/Private/Toolkits/CameraRigAssetEditorToolkit.cpp
+++ b/Engine/Plugins/Cameras/GameplayCameras/Source/GameplayCamerasEditor/Private/Toolkits/CameraRigAssetEditorToolkit.cpp
@@ -361,8 +361,8 @@ FText FCameraRigAssetEditorToolkit::GetBaseToolkitName() const
 
 FName FCameraRigAssetEditorToolkit::GetToolkitFName() const
 {
-	static FName SequencerName("CameraRigAssetEditor");
-	return SequencerName;
+	static FName ToolkitName("CameraRigAssetEditor");
+	return ToolkitName;
 }
 
 FString FCameraRigAssetEditorToolkit::GetWorldCentricTabPrefix() const
@@ -372,7 +372,7 @@ FString FCameraRigAssetEditorToolkit::GetWorldCentricTabPrefix() const
 
 FLinearColor FCameraRigAssetEditorToolkit::GetWorldCentricTabColorScale() const
 {
-	return FLinearColor(0.7, 0.0f, 0.0f, 0.5f);
+	return FLinearColor(0.7f, 0.0f, 0.0f, 0.5f);
 }
 
 void FCameraRigAssetEditorToolkit::NotifyPostChange(const FPropertyChangedEvent& PropertyChangedEvent, FProperty* PropertyThatChanged)
diff --git a/Engine/Plugins/Cameras/GameplayCameras/Source/GameplayCamerasEditor/Public/Toolkits/CameraRigAssetEditorToolkit.h b/Engine/Plugins/Cameras/GameplayCameras/Source/GameplayCamerasEditor/Public/Toolkits/CameraRigAssetEditorToolkit.h
--- a/Engine/Plugins/Cameras/GameplayCameras/Source/GameplayCamerasEditor/Public/Toolkits/CameraRigAssetEditorToolkit.h
+++ b/Engine/Plugins/Cameras/GameplayCameras/Source/GameplayCamerasEditor/Public/Toolkits/CameraRigAssetEditorToolkit.h
@@ -40,7 +40,7 @@ public:
 protected:
 
 	// FBaseAssetToolkit interface
-	virtual void RegisterTabSpawners(const TSharedRef<class FTabManager>& TabManager) override;
+	virtual void RegisterTabSpawners(const TSharedRef<class FTabManager>& InTabManager) override;
 	virtual void UnregisterTabSpawners(const TSharedRef<FTabManager>& InTabManager) override;
 	virtual void CreateWidgets() override;
 	virtual void RegisterToolbar() override;
