Cameras: fix copy/paste

Previously copy/paste was implemented by copying the graph nodes, and post-processing them by duplicating their underlying objects and patching up their references. This proved to be more complicated than it was worth, and wasn't solid when copy/pasting between different camera rigs.

This new implementation does the opposite: it duplicates the underlying objects, and then recreates brand new nodes from the duplicates. This is much better since the whole concept of the object-tree-graph module is to use the objects as the ground truth, and the graph nodes as something that can always be rebuilt from scratch.

diff --git a/Engine/Plugins/Cameras/GameplayCameras/Source/GameplayCamerasEditor/Private/Editors/ObjectTreeGraphConfig.cpp b/Engine/Plugins/Cameras/GameplayCameras/Source/GameplayCamerasEditor/Private/Editors/ObjectTreeGraphConfig.cpp
--- a/Engine/Plugins/Cameras/GameplayCameras/Source/GameplayCamerasEditor/Private/Editors/ObjectTreeGraphConfig.cpp
+++ b/Engine/Plugins/Cameras/GameplayCameras/Source/GameplayCamerasEditor/Private/Editors/ObjectTreeGraphConfig.cpp
@@ -18,7 +18,6 @@ FObjectTreeGraphClassConfig::FObjectTreeGraphClassConfig()
 	, _DefaultPropertyPinDirection(EGPD_Output)
 	, _NodeTitleUsesObjectName(false)
 	, _CanCreateNew(true)
-	, _CanDuplicate(true)
 	, _CanDelete(true)
 	, _CreateCategoryMetaData(TEXT("ObjectTreeGraphCategory"))
 {
@@ -27,7 +26,6 @@ FObjectTreeGraphClassConfig::FObjectTreeGraphClassConfig()
 FObjectTreeGraphClassConfig& FObjectTreeGraphClassConfig::OnlyAsRoot()
 {
 	_CanCreateNew = false;
-	_CanDuplicate = false;
 	_CanDelete = false;
 	return *this;
 }
diff --git a/Engine/Plugins/Cameras/GameplayCameras/Source/GameplayCamerasEditor/Private/Editors/ObjectTreeGraphNode.cpp b/Engine/Plugins/Cameras/GameplayCameras/Source/GameplayCamerasEditor/Private/Editors/ObjectTreeGraphNode.cpp
--- a/Engine/Plugins/Cameras/GameplayCameras/Source/GameplayCamerasEditor/Private/Editors/ObjectTreeGraphNode.cpp
+++ b/Engine/Plugins/Cameras/GameplayCameras/Source/GameplayCamerasEditor/Private/Editors/ObjectTreeGraphNode.cpp
@@ -33,81 +33,6 @@ void UObjectTreeGraphNode::Initialize(UObject* InObject)
 	}
 }
 
-void UObjectTreeGraphNode::PostDuplicateObject(const TMap<UEdGraphNode*, UEdGraphNode*>& NodeMap)
-{
-	if (!ensure(Object))
-	{
-		return;
-	}
-
-	// Set the correct values on our object.
-	UClass* ObjectClass = Object->GetClass();
-	for (TArray<UEdGraphPin*>::TIterator It(Pins.CreateIterator()); It; ++It)
-	{
-		UEdGraphPin* Pin = *It;
-		if (!ensure(Pin))
-		{
-			continue;
-		}
-		if (Pin->PinType.PinCategory != UObjectTreeGraphSchema::PC_Property)
-		{
-			continue;
-		}
-
-		if (Pin->PinType.PinSubCategory == UObjectTreeGraphSchema::PSC_ObjectProperty)
-		{
-			// Set the underlying object property to the duplicated object pointed to by the connected
-			// duplicated node.
-			FObjectProperty* ObjectProperty = CastFieldChecked<FObjectProperty>(ObjectClass->FindPropertyByName(Pin->GetFName()));
-			if (!Pin->LinkedTo.IsEmpty())
-			{
-				ensure(Pin->LinkedTo.Num() == 1);
-				UObjectTreeGraphNode* LinkedNode = CastChecked<UObjectTreeGraphNode>(Pin->LinkedTo[0]->GetOwningNode());
-				UObject* LinkedObject = LinkedNode->GetObject();
-
-				ObjectProperty->SetValue_InContainer(Object, TObjectPtr<UObject>(LinkedObject));
-			}
-			else
-			{
-				ObjectProperty->ClearValue_InContainer(Object);
-			}
-		}
-		else if (Pin->PinType.PinSubCategory == UObjectTreeGraphSchema::PSC_ArrayPropertyItem)
-		{
-			// Do as above, but in the appropriate place in the value array.
-			// Note that we can have unconnected pins here if we didn't duplicate some of the nodes.
-			// In that case, we need to delete the pin, and remove the null item from the value array.
-			UEdGraphPin* ParentPin = Pin->ParentPin;
-			check(ParentPin);
-			FArrayProperty* ArrayProperty = CastFieldChecked<FArrayProperty>(ObjectClass->FindPropertyByName(ParentPin->GetFName()));
-			FObjectProperty* InnerProperty = CastFieldChecked<FObjectProperty>(ArrayProperty->Inner);
-			FScriptArrayHelper ArrayHelper(ArrayProperty, ArrayProperty->ContainerPtrToValuePtr<void>(Object));
-
-			const int32 PinIndex = ParentPin->SubPins.Find(Pin);
-			check(PinIndex != INDEX_NONE);
-
-			if (!Pin->LinkedTo.IsEmpty())
-			{
-				ensure(Pin->LinkedTo.Num() == 1);
-				UObjectTreeGraphNode* LinkedNode = CastChecked<UObjectTreeGraphNode>(Pin->LinkedTo[0]->GetOwningNode());
-				UObject* LinkedObject = LinkedNode->GetObject();
-
-				InnerProperty->SetObjectPropertyValue(ArrayHelper.GetRawPtr(PinIndex), LinkedObject);
-			}
-			else
-			{
-				It.RemoveCurrent();
-				ParentPin->SubPins.Remove(Pin);
-				Pin->MarkAsGarbage();
-
-				ArrayHelper.RemoveValues(PinIndex, 1);
-			}
-		}
-	}
-
-	RefreshArrayPropertyPinNames();
-}
-
 FText UObjectTreeGraphNode::GetNodeTitle(ENodeTitleType::Type TitleType) const
 {
 	if (Object)
@@ -578,17 +503,6 @@ void UObjectTreeGraphNode::OnRenameNode(const FString& NewName)
 	}
 }
 
-bool UObjectTreeGraphNode::CanDuplicateNode() const
-{
-	const FObjectTreeGraphClassConfig& ObjectClassConfig = GetObjectClassConfig();
-	if (!ObjectClassConfig.CanCreateNew() || !ObjectClassConfig.CanDuplicate())
-	{
-		return false;
-	}
-
-	return Super::CanDuplicateNode();
-}
-
 bool UObjectTreeGraphNode::CanUserDeleteNode() const
 {
 	const FObjectTreeGraphClassConfig& ObjectClassConfig = GetObjectClassConfig();
diff --git a/Engine/Plugins/Cameras/GameplayCameras/Source/GameplayCamerasEditor/Private/Editors/ObjectTreeGraphSchema.cpp b/Engine/Plugins/Cameras/GameplayCameras/Source/GameplayCamerasEditor/Private/Editors/ObjectTreeGraphSchema.cpp
--- a/Engine/Plugins/Cameras/GameplayCameras/Source/GameplayCamerasEditor/Private/Editors/ObjectTreeGraphSchema.cpp
+++ b/Engine/Plugins/Cameras/GameplayCameras/Source/GameplayCamerasEditor/Private/Editors/ObjectTreeGraphSchema.cpp
@@ -6,10 +6,15 @@
 #include "Editors/ObjectTreeConnectionDrawingPolicy.h"
 #include "Editors/ObjectTreeGraph.h"
 #include "Editors/ObjectTreeGraphNode.h"
-#include "Serialization/ArchiveUObject.h"
+#include "Exporters/Exporter.h"
+#include "Factories.h"
+#include "GameplayCameras.h"
+#include "IGameplayCamerasEditorModule.h"
 #include "ScopedTransaction.h"
+#include "Serialization/ArchiveUObject.h"
 #include "UObject/FastReferenceCollector.h"
 #include "UObject/UObjectIterator.h"
+#include "UnrealExporter.h"
 
 #include UE_INLINE_GENERATED_CPP_BY_NAME(ObjectTreeGraphSchema)
 
@@ -81,6 +86,29 @@ private:
 	TArray<UObject*>& ReferencedObjects;
 };
 
+class FObjectTextFactory : public FCustomizableTextObjectFactory
+{
+public:
+	FObjectTextFactory()
+		: FCustomizableTextObjectFactory(GWarn)
+	{
+	}
+
+	// FCustomizableTextObjectFactory interface.
+	virtual bool CanCreateClass(UClass* ObjectClass, bool& bOmitSubObjs) const override
+	{
+		return true;
+	}
+
+	virtual void ProcessConstructedObject(UObject* NewObject) override
+	{
+		check(NewObject);
+		CreatedObjects.Add(NewObject);
+	}
+
+	TArray<UObject*> CreatedObjects;
+};
+
 }  // namespace UE::ObjectTreeGraph
 
 UObjectTreeGraphSchema::UObjectTreeGraphSchema(const FObjectInitializer& ObjInit)
@@ -719,42 +747,145 @@ void UObjectTreeGraphSchema::OnDeleteNodeFromGraph(UObjectTreeGraph* Graph, UEdG
 	}
 }
 
-void UObjectTreeGraphSchema::ProcessDuplicatedNodes(UObjectTreeGraph* InGraph, const TMap<UEdGraphNode*, UEdGraphNode*>& NodeMap) const
+FString UObjectTreeGraphSchema::ExportNodesToText(const FGraphPanelSelectionSet& Nodes, bool bOnlyCanDuplicateNodes, bool bOnlyCanDeleteNodes) const
 {
-	UPackage* RootObjectPackage = InGraph->GetRootObject()->GetOutermost();
+	// Gather up the nodes we need to copy from.
+	TSet<UObject*> ObjectsToExport;
+	TSet<UObject*> OtherNodesToExport;
 
-	// Create object duplicates for all the new nodes.
-	for (TPair<UEdGraphNode*, UEdGraphNode*> Pair : NodeMap)
+	for (FGraphPanelSelectionSet::TConstIterator NodeIt(Nodes); NodeIt; ++NodeIt)
 	{
-		const UObjectTreeGraphNode* OldNode = Cast<UObjectTreeGraphNode>(Pair.Key);
-		UObjectTreeGraphNode* NewNode = Cast<UObjectTreeGraphNode>(Pair.Value);
+		UEdGraphNode* Node = Cast<UEdGraphNode>(*NodeIt);
+		if (Node && 
+				(!bOnlyCanDuplicateNodes || Node->CanDuplicateNode()) &&
+				(!bOnlyCanDeleteNodes || Node->CanUserDeleteNode()))
+		{
+			Node->PrepareForCopying();
 
-		if (OldNode && NewNode)
+			if (UObjectTreeGraphNode* ObjectTreeNode = Cast<UObjectTreeGraphNode>(Node))
+			{
+				ObjectsToExport.Add(ObjectTreeNode->GetObject());
+			}
+			else
+			{
+				OtherNodesToExport.Add(Node);
+			}
+		}
+	}
+
+	if (ObjectsToExport.IsEmpty() && OtherNodesToExport.IsEmpty())
+	{
+		return FString();
+	}
+
+	// Clear the mark state for saving.
+	UnMarkAllObjects(EObjectMark(OBJECTMARK_TagExp | OBJECTMARK_TagImp));
+
+	FStringOutputDevice Archive;
+	const FExportObjectInnerContext Context;
+
+	UObject* LastOuter = nullptr;
+	for (UObject* ObjectToExport : ObjectsToExport)
+	{
+		// The nodes should all be from the same scope.
+		UObject* ThisOuter = ObjectToExport->GetOuter();
+		if (LastOuter != nullptr && ThisOuter != LastOuter)
 		{
-			ensure(OldNode->GetObject() == NewNode->GetObject());
+			UE_LOG(LogCameraSystemEditor, Warning,
+					TEXT("Cannot copy objects from different outers. Only copying from %s"), *LastOuter->GetName());
+			continue;
+		}
+		LastOuter = ThisOuter;
+
+		UExporter::ExportToOutputDevice(
+				&Context,
+				ObjectToExport, 
+				nullptr, // no exporter
+				Archive, 
+				TEXT("copy"), // file type
+				0, // indent
+				PPF_ExportsNotFullyQualified | PPF_Copy | PPF_Delimited, // port flags
+				false, // selected only
+				ThisOuter // export root scope
+				);
+	}
 
-			UObject* NewNodeObject = DuplicateObject<UObject>(OldNode->GetObject(), RootObjectPackage);
-			NewNode->Initialize(NewNodeObject);
+	if (!OtherNodesToExport.IsEmpty())
+	{
+		CopyNonObjectNodes(OtherNodesToExport.Array(), Archive);
+	}
+
+	return Archive;
+}
+
+void UObjectTreeGraphSchema::CopyNonObjectNodes(TArrayView<UObject*> InObjects, FStringOutputDevice& OutDevice) const
+{
+}
+
+void UObjectTreeGraphSchema::ImportNodesFromText(UObjectTreeGraph* InGraph, const FString& TextToImport, TArray<UEdGraphNode*>& OutPastedNodes) const
+{
+	using namespace UE::ObjectTreeGraph;
+
+	TArray<UObject*> ImportedObjects;
+
+	// Import the given text as new objects.
+	UPackage* TempPackage = NewObject<UPackage>(nullptr, TEXT("/Engine/GameplayCamerasEditor/Transient"), RF_Transient);
+	TempPackage->AddToRoot();
+	{
+		FObjectTextFactory Factory;
+		Factory.ProcessBuffer(TempPackage, RF_Transactional, TextToImport);
+		ImportedObjects = Factory.CreatedObjects;
+	}
+	TempPackage->RemoveFromRoot();
+
+	// Finish setting up the new objects: clear the transient flag from the transient package we used above,
+	// move the objects under the our graph root, and optionally add them via the root interface.
+	UObject* GraphRootObject = InGraph->GetRootObject();
+	IObjectTreeGraphRootObject* RootObjectInterface = Cast<IObjectTreeGraphRootObject>(GraphRootObject);
+	if (ensure(GraphRootObject))
+	{
+		for (UObject* Object : ImportedObjects)
+		{
+			Object->ClearFlags(RF_Transient);
+			Object->Rename(nullptr, GraphRootObject);
+
+			if (RootObjectInterface)
+			{
+				RootObjectInterface->AddConnectableObject(Object);
+			}
 		}
 	}
 
-	// Patch-up property references on the new nodes' objects.
-	for (TPair<UEdGraphNode*, UEdGraphNode*> Pair : NodeMap)
+	// Create nodes for all the imported objects.
+	FCreatedNodes CreatedNodes;
+	for (UObject* Object : ImportedObjects)
 	{
-		UObjectTreeGraphNode* NewNode = Cast<UObjectTreeGraphNode>(Pair.Value);
-		if (NewNode)
+		if (UObjectTreeGraphNode* GraphNode = CreateObjectNode(InGraph, Object))
 		{
-			NewNode->PostDuplicateObject(NodeMap);
+			CreatedNodes.CreatedNodes.Add(Object, GraphNode);
 		}
 	}
 
-	// Notify for everything.
-	for (TPair<UEdGraphNode*, UEdGraphNode*> Pair : NodeMap)
+	// Create all the connections.
+	for (TPair<UObject*, UObjectTreeGraphNode*> Pair : CreatedNodes.CreatedNodes)
 	{
-		UObjectTreeGraphNode* NewNode = Cast<UObjectTreeGraphNode>(Pair.Value);
-		InGraph->NotifyNodeChanged(NewNode);
+		CreateConnections(Pair.Value, CreatedNodes);
 	}
-	InGraph->NotifyGraphChanged();
+
+	OnCreateAllNodes(InGraph, CreatedNodes);
+
+	for (const TTuple<UObject*, UObjectTreeGraphNode*>& Pair : CreatedNodes.CreatedNodes)
+	{
+		OutPastedNodes.Add(Pair.Value);
+	}
+}
+
+bool UObjectTreeGraphSchema::CanImportNodesFromText(UObjectTreeGraph* InGraph, const FString& TextToImport) const
+{
+	using namespace UE::ObjectTreeGraph;
+
+	FObjectTextFactory Factory;
+	return Factory.CanCreateObjectsFromText(TextToImport);
 }
 
 const FObjectTreeGraphClassConfig& UObjectTreeGraphSchema::GetObjectClassConfig(const UObjectTreeGraphNode* InNode) const
diff --git a/Engine/Plugins/Cameras/GameplayCameras/Source/GameplayCamerasEditor/Private/Editors/SCameraRigAssetEditor.cpp b/Engine/Plugins/Cameras/GameplayCameras/Source/GameplayCamerasEditor/Private/Editors/SCameraRigAssetEditor.cpp
--- a/Engine/Plugins/Cameras/GameplayCameras/Source/GameplayCamerasEditor/Private/Editors/SCameraRigAssetEditor.cpp
+++ b/Engine/Plugins/Cameras/GameplayCameras/Source/GameplayCamerasEditor/Private/Editors/SCameraRigAssetEditor.cpp
@@ -78,7 +78,6 @@ void SCameraRigAssetEditor::CreateNodeGraphEditor(const FArguments& InArgs)
 		.SelfPinDirection(EGPD_Output)
 		.SelfPinName(NAME_None)  // No self pin name, we just want the title
 		.CanCreateNew(false)
-		.CanDuplicate(false)
 		.GraphNodeClass(UCameraRigInterfaceParameterGraphNode::StaticClass());
 
 	NodeGraph = NewObject<UObjectTreeGraph>(GetTransientPackage(), NAME_None, RF_Transactional);
diff --git a/Engine/Plugins/Cameras/GameplayCameras/Source/GameplayCamerasEditor/Private/Editors/SObjectTreeGraphEditor.cpp b/Engine/Plugins/Cameras/GameplayCameras/Source/GameplayCamerasEditor/Private/Editors/SObjectTreeGraphEditor.cpp
--- a/Engine/Plugins/Cameras/GameplayCameras/Source/GameplayCamerasEditor/Private/Editors/SObjectTreeGraphEditor.cpp
+++ b/Engine/Plugins/Cameras/GameplayCameras/Source/GameplayCamerasEditor/Private/Editors/SObjectTreeGraphEditor.cpp
@@ -5,7 +5,6 @@
 #include "Algo/AnyOf.h"
 #include "EdGraph/EdGraphPin.h"
 #include "EdGraph/EdGraphSchema.h"
-#include "EdGraphUtilities.h"
 #include "Editor.h"
 #include "Editors/ObjectTreeGraph.h"
 #include "Editors/ObjectTreeGraphNode.h"
@@ -317,24 +316,10 @@ void SObjectTreeGraphEditor::OnDoubleClicked()
 
 FString SObjectTreeGraphEditor::ExportNodesToText(const FGraphPanelSelectionSet& Nodes, bool bOnlyCanDuplicateNodes, bool bOnlyCanDeleteNodes)
 {
-	FString ExportedText;
-	TSet<UObject*> NodesToExport;
-
-	for (FGraphPanelSelectionSet::TConstIterator NodeIt(Nodes); NodeIt; ++NodeIt)
-	{
-		UEdGraphNode* Node = Cast<UEdGraphNode>(*NodeIt);
-		if (Node && 
-				(!bOnlyCanDuplicateNodes || Node->CanDuplicateNode()) &&
-				(!bOnlyCanDeleteNodes || Node->CanUserDeleteNode()))
-		{
-			Node->PrepareForCopying();
-			NodesToExport.Add(Node);
-		}
-	}
-
-	FEdGraphUtilities::ExportNodesToText(NodesToExport, ExportedText);
+	UEdGraph* CurrentGraph = GraphEditor->GetCurrentGraph();
+	const UObjectTreeGraphSchema* Schema = CastChecked<UObjectTreeGraphSchema>(CurrentGraph->GetSchema());
 
-	return ExportedText;
+	return Schema->ExportNodesToText(Nodes, bOnlyCanDuplicateNodes, bOnlyCanDeleteNodes);
 }
 
 void SObjectTreeGraphEditor::ImportNodesFromText(const FVector2D& Location, const FString& TextToImport)
@@ -348,32 +333,10 @@ void SObjectTreeGraphEditor::ImportNodesFromText(const FVector2D& Location, cons
 	UPackage* ObjectPackage = Graph->GetRootObject()->GetOutermost();
 	ObjectPackage->Modify();
 
-	// Import the nodes
-	TSet<UEdGraphNode*> PastedNodes;
-	FEdGraphUtilities::ImportNodesFromText(Graph, TextToImport, PastedNodes);
-
-	// Process all the new nodes.
-	TMap<FGuid, UEdGraphNode*> GuidToNode;
-	for (UEdGraphNode* Node : Graph->Nodes)
-	{
-		GuidToNode.Add(Node->NodeGuid, Node);
-	}
-	
-	TMap<UEdGraphNode*, UEdGraphNode*> PastedNodeMap;
-	for (UEdGraphNode* PastedNode : PastedNodes)
-	{
-		const FGuid OldNodeGuid = PastedNode->NodeGuid;
-		PastedNode->CreateNewGuid();
-
-		UEdGraphNode** OldNode = GuidToNode.Find(OldNodeGuid);
-		if (ensure(OldNode))
-		{
-			PastedNodeMap.Add(*OldNode, PastedNode);
-		}
-	}
-	
+	// Import the nodes.
+	TArray<UEdGraphNode*> PastedNodes;
 	const UObjectTreeGraphSchema* Schema = CastChecked<UObjectTreeGraphSchema>(Graph->GetSchema());
-	Schema->ProcessDuplicatedNodes(Graph, PastedNodeMap);
+	Schema->ImportNodesFromText(Graph, TextToImport, PastedNodes);
 
 	// Compute the center of the pasted nodes.
 	FVector2D PastedNodesClusterCenter(FVector2D::ZeroVector);
@@ -406,6 +369,14 @@ void SObjectTreeGraphEditor::ImportNodesFromText(const FVector2D& Location, cons
 	GraphEditor->NotifyGraphChanged();
 }
 
+bool SObjectTreeGraphEditor::CanImportNodesFromText(const FString& TextToImport)
+{
+	UObjectTreeGraph* CurrentGraph = CastChecked<UObjectTreeGraph>(GraphEditor->GetCurrentGraph());
+	const UObjectTreeGraphSchema* Schema = CastChecked<UObjectTreeGraphSchema>(CurrentGraph->GetSchema());
+
+	return Schema->CanImportNodesFromText(CurrentGraph, TextToImport);
+}
+
 void SObjectTreeGraphEditor::DeleteNodes(TArrayView<UObjectTreeGraphNode*> NodesToDelete)
 {
 	UEdGraph* CurrentGraph = GraphEditor->GetCurrentGraph();
@@ -447,6 +418,9 @@ void SObjectTreeGraphEditor::DeleteSelectedNodes()
 	}
 	
 	DeleteNodes(NodesToDelete);
+
+	// Remove deleted nodes from the details view.
+	GraphEditor->ClearSelectionSet();
 }
 
 bool SObjectTreeGraphEditor::CanDeleteSelectedNodes()
@@ -514,7 +488,7 @@ bool SObjectTreeGraphEditor::CanPasteNodes()
 	FString ClipboardContent;
 	FPlatformApplicationMisc::ClipboardPaste(ClipboardContent);
 
-	return FEdGraphUtilities::CanImportNodesFromText(GraphEditor->GetCurrentGraph(), ClipboardContent);
+	return CanImportNodesFromText(ClipboardContent);
 }
 
 void SObjectTreeGraphEditor::DuplicateNodes()
diff --git a/Engine/Plugins/Cameras/GameplayCameras/Source/GameplayCamerasEditor/Private/GameplayCamerasEditorModule.cpp b/Engine/Plugins/Cameras/GameplayCameras/Source/GameplayCamerasEditor/Private/GameplayCamerasEditorModule.cpp
--- a/Engine/Plugins/Cameras/GameplayCameras/Source/GameplayCamerasEditor/Private/GameplayCamerasEditorModule.cpp
+++ b/Engine/Plugins/Cameras/GameplayCameras/Source/GameplayCamerasEditor/Private/GameplayCamerasEditorModule.cpp
@@ -35,6 +35,8 @@
 
 #define LOCTEXT_NAMESPACE "GameplayCamerasEditor"
 
+DEFINE_LOG_CATEGORY(LogCameraSystemEditor);
+
 const FName IGameplayCamerasEditorModule::GameplayCamerasEditorAppIdentifier("GameplayCamerasEditorApp");
 const FName IGameplayCamerasEditorModule::CameraRigAssetEditorToolBarName("CameraRigAssetEditor.ToolBar");
 
diff --git a/Engine/Plugins/Cameras/GameplayCameras/Source/GameplayCamerasEditor/Public/Editors/ObjectTreeGraphConfig.h b/Engine/Plugins/Cameras/GameplayCameras/Source/GameplayCamerasEditor/Public/Editors/ObjectTreeGraphConfig.h
--- a/Engine/Plugins/Cameras/GameplayCameras/Source/GameplayCamerasEditor/Public/Editors/ObjectTreeGraphConfig.h
+++ b/Engine/Plugins/Cameras/GameplayCameras/Source/GameplayCamerasEditor/Public/Editors/ObjectTreeGraphConfig.h
@@ -65,8 +65,6 @@ public:
 	/** Whether users can create new objects of this class in the graph. */
 	OTGCC_FIELD(bool, CanCreateNew)
 	/** Whether users can duplicate objects of this class in the graph. */
-	OTGCC_FIELD(bool, CanDuplicate)
-	/** Whether users can delete objects of this class in the graph. */
 	OTGCC_FIELD(bool, CanDelete)
 
 	/** The metadata specifier to look for in order to categorize the 'create node' action for this class. */
@@ -76,7 +74,7 @@ public:
 
 	FObjectTreeGraphClassConfig();
 
-	/** A shortcut for disabling CanCreateNew, CanDuplicate, and CanDelete. */
+	/** A shortcut for disabling CanCreateNew and CanDelete. */
 	FObjectTreeGraphClassConfig& OnlyAsRoot();
 
 public:
diff --git a/Engine/Plugins/Cameras/GameplayCameras/Source/GameplayCamerasEditor/Public/Editors/ObjectTreeGraphNode.h b/Engine/Plugins/Cameras/GameplayCameras/Source/GameplayCamerasEditor/Public/Editors/ObjectTreeGraphNode.h
--- a/Engine/Plugins/Cameras/GameplayCameras/Source/GameplayCamerasEditor/Public/Editors/ObjectTreeGraphNode.h
+++ b/Engine/Plugins/Cameras/GameplayCameras/Source/GameplayCamerasEditor/Public/Editors/ObjectTreeGraphNode.h
@@ -74,7 +74,6 @@ public:
 	virtual void GetNodeContextMenuActions(class UToolMenu* Menu, class UGraphNodeContextMenuContext* Context) const override;
 	virtual bool GetCanRenameNode() const override;
 	virtual void OnRenameNode(const FString& NewName) override;
-	virtual bool CanDuplicateNode() const override;
 	virtual bool CanUserDeleteNode() const override;
 	virtual bool SupportsCommentBubble() const override;
 	virtual void OnUpdateCommentText(const FString& NewComment) override;
@@ -86,7 +85,6 @@ public:
 public:
 
 	// Internal API.
-	void PostDuplicateObject(const TMap<UEdGraphNode*, UEdGraphNode*>& NodeMap);
 	void CreateNewItemPin(FArrayProperty& InArrayProperty);
 	void CreateNewItemPin(UEdGraphPin* InParentArrayPin);
 	void RemoveItemPin(UEdGraphPin* InItemPin);
diff --git a/Engine/Plugins/Cameras/GameplayCameras/Source/GameplayCamerasEditor/Public/Editors/ObjectTreeGraphSchema.h b/Engine/Plugins/Cameras/GameplayCameras/Source/GameplayCamerasEditor/Public/Editors/ObjectTreeGraphSchema.h
--- a/Engine/Plugins/Cameras/GameplayCameras/Source/GameplayCamerasEditor/Public/Editors/ObjectTreeGraphSchema.h
+++ b/Engine/Plugins/Cameras/GameplayCameras/Source/GameplayCamerasEditor/Public/Editors/ObjectTreeGraphSchema.h
@@ -4,6 +4,7 @@
 
 #include "CoreTypes.h"
 #include "EdGraph/EdGraphSchema.h"
+#include "GraphEditor.h"
 
 #include "ObjectTreeGraphSchema.generated.h"
 
@@ -43,11 +44,21 @@ public:
 	/** Creates a new schema. */
 	UObjectTreeGraphSchema(const FObjectInitializer& ObjInit);
 
+	/** Rebuilds the graph from scratch. */
 	void RebuildGraph(UObjectTreeGraph* InGraph, EObjectTreeGraphBuildSource InSource) const;
 
 	/** Creates an object graph node for the given object. */
 	UObjectTreeGraphNode* CreateObjectNode(UObjectTreeGraph* InGraph, UObject* InObject) const;
 
+	/** Export the given selection into a text suitable for copy/pasting. */
+	FString ExportNodesToText(const FGraphPanelSelectionSet& Nodes, bool bOnlyCanDuplicateNodes, bool bOnlyCanDeleteNodes) const;
+
+	/** Imports the given text into the given graph. */
+	void ImportNodesFromText(UObjectTreeGraph* InGraph, const FString& TextToImport, TArray<UEdGraphNode*>& OutPastedNodes) const;
+
+	/** Checks if the given text is suitable for importing. */
+	bool CanImportNodesFromText(UObjectTreeGraph* InGraph, const FString& TextToImport) const;
+
 public:
 
 	// UEdGraphSchema interface.
@@ -71,7 +82,6 @@ public:
 	virtual bool OnBreakSinglePinLink(UEdGraphPin* SourcePin, UEdGraphPin* TargetPin) const;
 	virtual void OnDeleteNodeFromGraph(UObjectTreeGraph* Graph, UEdGraphNode* Node) const;
 	virtual void FilterGraphContextPlaceableClasses(TArray<UClass*>& InOutClasses) const;
-	virtual void ProcessDuplicatedNodes(UObjectTreeGraph* InGraph, const TMap<UEdGraphNode*, UEdGraphNode*>& NodeMap) const;
 
 protected:
 
@@ -83,6 +93,7 @@ protected:
 	// UObjectTreeGraphSchema interface.
 	virtual void OnCreateAllNodes(UObjectTreeGraph* InGraph, const FCreatedNodes& InCreatedNodes) const;
 	virtual UObjectTreeGraphNode* CreateObjectNodeImpl(UObjectTreeGraph* InGraph, UObject* InObject) const;
+	virtual void CopyNonObjectNodes(TArrayView<UObject*> InObjects, FStringOutputDevice& OutDevice) const;
 
 protected:
 
diff --git a/Engine/Plugins/Cameras/GameplayCameras/Source/GameplayCamerasEditor/Public/Editors/SObjectTreeGraphEditor.h b/Engine/Plugins/Cameras/GameplayCameras/Source/GameplayCamerasEditor/Public/Editors/SObjectTreeGraphEditor.h
--- a/Engine/Plugins/Cameras/GameplayCameras/Source/GameplayCamerasEditor/Public/Editors/SObjectTreeGraphEditor.h
+++ b/Engine/Plugins/Cameras/GameplayCameras/Source/GameplayCamerasEditor/Public/Editors/SObjectTreeGraphEditor.h
@@ -65,6 +65,7 @@ protected:
 
 	FString ExportNodesToText(const FGraphPanelSelectionSet& Nodes, bool bOnlyCanDuplicateNodes, bool bOnlyCanDeleteNodes);
 	void ImportNodesFromText(const FVector2D& Location, const FString& TextToImport);
+	bool CanImportNodesFromText(const FString& TextToImport);
 	void DeleteNodes(TArrayView<UObjectTreeGraphNode*> NodesToDelete);
 
 	void SelectAllNodes();
diff --git a/Engine/Plugins/Cameras/GameplayCameras/Source/GameplayCamerasEditor/Public/IGameplayCamerasEditorModule.h b/Engine/Plugins/Cameras/GameplayCameras/Source/GameplayCamerasEditor/Public/IGameplayCamerasEditorModule.h
--- a/Engine/Plugins/Cameras/GameplayCameras/Source/GameplayCamerasEditor/Public/IGameplayCamerasEditorModule.h
+++ b/Engine/Plugins/Cameras/GameplayCameras/Source/GameplayCamerasEditor/Public/IGameplayCamerasEditorModule.h
@@ -71,3 +71,5 @@ public:
 	virtual void UnregisterDebugCategoryPanel(const FString& InCategoryName) = 0;
 };
 
+DECLARE_LOG_CATEGORY_EXTERN(LogCameraSystemEditor, Log, All);
+
