Cameras: fix moving graph nodes creating many "move nodes" undo transaction on mouse move

diff --git a/Engine/Plugins/Cameras/GameplayCameras/Source/GameplayCameras/Private/Core/CameraNode.cpp b/Engine/Plugins/Cameras/GameplayCameras/Source/GameplayCameras/Private/Core/CameraNode.cpp
--- a/Engine/Plugins/Cameras/GameplayCameras/Source/GameplayCameras/Private/Core/CameraNode.cpp
+++ b/Engine/Plugins/Cameras/GameplayCameras/Source/GameplayCameras/Private/Core/CameraNode.cpp
@@ -32,7 +32,7 @@ void UCameraNode::GetGraphNodePosition(int32& NodePosX, int32& NodePosY) const
 	NodePosY = GraphNodePosY;
 }
 
-void UCameraNode::OnGraphNodeMoved(int32 NodePosX, int32 NodePosY)
+void UCameraNode::OnGraphNodeMoved(int32 NodePosX, int32 NodePosY, bool bMarkDirty)
 {
 	GraphNodePosX = NodePosX;
 	GraphNodePosY = NodePosY;
diff --git a/Engine/Plugins/Cameras/GameplayCameras/Source/GameplayCameras/Private/Core/CameraRigAsset.cpp b/Engine/Plugins/Cameras/GameplayCameras/Source/GameplayCameras/Private/Core/CameraRigAsset.cpp
--- a/Engine/Plugins/Cameras/GameplayCameras/Source/GameplayCameras/Private/Core/CameraRigAsset.cpp
+++ b/Engine/Plugins/Cameras/GameplayCameras/Source/GameplayCameras/Private/Core/CameraRigAsset.cpp
@@ -19,7 +19,7 @@ void UCameraRigInterfaceParameter::GetGraphNodePosition(int32& NodePosX, int32&
 	NodePosY = GraphNodePosY;
 }
 
-void UCameraRigInterfaceParameter::OnGraphNodeMoved(int32 NodePosX, int32 NodePosY)
+void UCameraRigInterfaceParameter::OnGraphNodeMoved(int32 NodePosX, int32 NodePosY, bool bMarkDirty)
 {
 	GraphNodePosX = NodePosX;
 	GraphNodePosY = NodePosY;
@@ -118,7 +118,7 @@ void UCameraRigAsset::GetGraphNodePosition(int32& NodePosX, int32& NodePosY) con
 	NodePosY = GraphNodePosY;
 }
 
-void UCameraRigAsset::OnGraphNodeMoved(int32 NodePosX, int32 NodePosY)
+void UCameraRigAsset::OnGraphNodeMoved(int32 NodePosX, int32 NodePosY, bool bMarkDirty)
 {
 	GraphNodePosX = NodePosX;
 	GraphNodePosY = NodePosY;
diff --git a/Engine/Plugins/Cameras/GameplayCameras/Source/GameplayCameras/Private/Core/CameraRigTransition.cpp b/Engine/Plugins/Cameras/GameplayCameras/Source/GameplayCameras/Private/Core/CameraRigTransition.cpp
--- a/Engine/Plugins/Cameras/GameplayCameras/Source/GameplayCameras/Private/Core/CameraRigTransition.cpp
+++ b/Engine/Plugins/Cameras/GameplayCameras/Source/GameplayCameras/Private/Core/CameraRigTransition.cpp
@@ -17,7 +17,7 @@ void UCameraRigTransitionCondition::GetGraphNodePosition(int32& NodePosX, int32&
 	NodePosY = GraphNodePosY;
 }
 
-void UCameraRigTransitionCondition::OnGraphNodeMoved(int32 NodePosX, int32 NodePosY)
+void UCameraRigTransitionCondition::OnGraphNodeMoved(int32 NodePosX, int32 NodePosY, bool bMarkDirty)
 {
 	GraphNodePosX = NodePosX;
 	GraphNodePosY = NodePosY;
@@ -39,7 +39,7 @@ void UCameraRigTransition::GetGraphNodePosition(int32& NodePosX, int32& NodePosY
 	NodePosY = GraphNodePosY;
 }
 
-void UCameraRigTransition::OnGraphNodeMoved(int32 NodePosX, int32 NodePosY)
+void UCameraRigTransition::OnGraphNodeMoved(int32 NodePosX, int32 NodePosY, bool bMarkDirty)
 {
 	GraphNodePosX = NodePosX;
 	GraphNodePosY = NodePosY;
diff --git a/Engine/Plugins/Cameras/GameplayCameras/Source/GameplayCameras/Public/Core/CameraNode.h b/Engine/Plugins/Cameras/GameplayCameras/Source/GameplayCameras/Public/Core/CameraNode.h
--- a/Engine/Plugins/Cameras/GameplayCameras/Source/GameplayCameras/Public/Core/CameraNode.h
+++ b/Engine/Plugins/Cameras/GameplayCameras/Source/GameplayCameras/Public/Core/CameraNode.h
@@ -55,7 +55,7 @@ protected:
 	// IObjectTreeGraphObject interface.
 #if WITH_EDITOR
 	virtual void GetGraphNodePosition(int32& NodePosX, int32& NodePosY) const override;
-	virtual void OnGraphNodeMoved(int32 NodePosX, int32 NodePosY) override;
+	virtual void OnGraphNodeMoved(int32 NodePosX, int32 NodePosY, bool bMarkDirty) override;
 	virtual EObjectTreeGraphObjectSupportFlags GetSupportFlags() const override { return EObjectTreeGraphObjectSupportFlags::CommentText; }
 	virtual const FString& GetGraphNodeCommentText() const override;
 	virtual void OnUpdateGraphNodeCommentText(const FString& NewComment) override;
diff --git a/Engine/Plugins/Cameras/GameplayCameras/Source/GameplayCameras/Public/Core/CameraRigAsset.h b/Engine/Plugins/Cameras/GameplayCameras/Source/GameplayCameras/Public/Core/CameraRigAsset.h
--- a/Engine/Plugins/Cameras/GameplayCameras/Source/GameplayCameras/Public/Core/CameraRigAsset.h
+++ b/Engine/Plugins/Cameras/GameplayCameras/Source/GameplayCameras/Public/Core/CameraRigAsset.h
@@ -88,7 +88,7 @@ protected:
 	// IObjectTreeGraphObject interface.
 #if WITH_EDITOR
 	virtual void GetGraphNodePosition(int32& NodePosX, int32& NodePosY) const override;
-	virtual void OnGraphNodeMoved(int32 NodePosX, int32 NodePosY) override;
+	virtual void OnGraphNodeMoved(int32 NodePosX, int32 NodePosY, bool bMarkDirty) override;
 #endif
 
 private:
@@ -201,7 +201,7 @@ protected:
 	// IObjectTreeGraphObject interface.
 #if WITH_EDITOR
 	virtual void GetGraphNodePosition(int32& NodePosX, int32& NodePosY) const override;
-	virtual void OnGraphNodeMoved(int32 NodePosX, int32 NodePosY) override;
+	virtual void OnGraphNodeMoved(int32 NodePosX, int32 NodePosY, bool bMarkDirty) override;
 	virtual EObjectTreeGraphObjectSupportFlags GetSupportFlags() const override { return EObjectTreeGraphObjectSupportFlags::CommentText; }
 	virtual const FString& GetGraphNodeCommentText() const override;
 	virtual void OnUpdateGraphNodeCommentText(const FString& NewComment) override;
diff --git a/Engine/Plugins/Cameras/GameplayCameras/Source/GameplayCameras/Public/Core/CameraRigTransition.h b/Engine/Plugins/Cameras/GameplayCameras/Source/GameplayCameras/Public/Core/CameraRigTransition.h
--- a/Engine/Plugins/Cameras/GameplayCameras/Source/GameplayCameras/Public/Core/CameraRigTransition.h
+++ b/Engine/Plugins/Cameras/GameplayCameras/Source/GameplayCameras/Public/Core/CameraRigTransition.h
@@ -54,7 +54,7 @@ protected:
 	// IObjectTreeGraphObject interface.
 #if WITH_EDITOR
 	virtual void GetGraphNodePosition(int32& NodePosX, int32& NodePosY) const override;
-	virtual void OnGraphNodeMoved(int32 NodePosX, int32 NodePosY) override;
+	virtual void OnGraphNodeMoved(int32 NodePosX, int32 NodePosY, bool bMarkDirty) override;
 	virtual EObjectTreeGraphObjectSupportFlags GetSupportFlags() const override { return EObjectTreeGraphObjectSupportFlags::CommentText; }
 	virtual const FString& GetGraphNodeCommentText() const override;
 	virtual void OnUpdateGraphNodeCommentText(const FString& NewComment) override;
@@ -101,7 +101,7 @@ protected:
 	// IObjectTreeGraphObject interface.
 #if WITH_EDITOR
 	virtual void GetGraphNodePosition(int32& NodePosX, int32& NodePosY) const override;
-	virtual void OnGraphNodeMoved(int32 NodePosX, int32 NodePosY) override;
+	virtual void OnGraphNodeMoved(int32 NodePosX, int32 NodePosY, bool bMarkDirty) override;
 	virtual EObjectTreeGraphObjectSupportFlags GetSupportFlags() const override { return EObjectTreeGraphObjectSupportFlags::CommentText; }
 	virtual const FString& GetGraphNodeCommentText() const override;
 	virtual void OnUpdateGraphNodeCommentText(const FString& NewComment) override;
diff --git a/Engine/Plugins/Cameras/GameplayCameras/Source/GameplayCameras/Public/Core/ObjectTreeGraphObject.h b/Engine/Plugins/Cameras/GameplayCameras/Source/GameplayCameras/Public/Core/ObjectTreeGraphObject.h
--- a/Engine/Plugins/Cameras/GameplayCameras/Source/GameplayCameras/Public/Core/ObjectTreeGraphObject.h
+++ b/Engine/Plugins/Cameras/GameplayCameras/Source/GameplayCameras/Public/Core/ObjectTreeGraphObject.h
@@ -47,7 +47,7 @@ public:
 	/** Gets the canvas position for the graph node representing this object. */
 	virtual void GetGraphNodePosition(int32& NodePosX, int32& NodePosY) const {}
 	/** Called to save the canvas position of the graph node representing this object. */
-	virtual void OnGraphNodeMoved(int32 NodePosX, int32 NodePosY) {}
+	virtual void OnGraphNodeMoved(int32 NodePosX, int32 NodePosY, bool bMarkDirty) {}
 
 	// Optional API.
 
diff --git a/Engine/Plugins/Cameras/GameplayCameras/Source/GameplayCamerasEditor/Private/Editors/CameraNodeGraphSchema.cpp b/Engine/Plugins/Cameras/GameplayCameras/Source/GameplayCamerasEditor/Private/Editors/CameraNodeGraphSchema.cpp
--- a/Engine/Plugins/Cameras/GameplayCameras/Source/GameplayCamerasEditor/Private/Editors/CameraNodeGraphSchema.cpp
+++ b/Engine/Plugins/Cameras/GameplayCameras/Source/GameplayCamerasEditor/Private/Editors/CameraNodeGraphSchema.cpp
@@ -275,7 +275,7 @@ UEdGraphNode* FCameraNodeGraphSchemaAction_NewInterfaceParameterNode::PerformAct
 	UObjectTreeGraphNode* NewGraphNode = Schema->CreateObjectNode(ObjectTreeGraph, NewInterfaceParameter);
 	NewGraphNode->NodePosX = Location.X;
 	NewGraphNode->NodePosY = Location.Y;
-	NewGraphNode->OnGraphNodeMoved();
+	NewGraphNode->OnGraphNodeMoved(false);
 	NewGraphNode->AutowireNewNode(FromPin);
 
 	return NewGraphNode;
diff --git a/Engine/Plugins/Cameras/GameplayCameras/Source/GameplayCamerasEditor/Private/Editors/ObjectTreeGraphNode.cpp b/Engine/Plugins/Cameras/GameplayCameras/Source/GameplayCamerasEditor/Private/Editors/ObjectTreeGraphNode.cpp
--- a/Engine/Plugins/Cameras/GameplayCameras/Source/GameplayCamerasEditor/Private/Editors/ObjectTreeGraphNode.cpp
+++ b/Engine/Plugins/Cameras/GameplayCameras/Source/GameplayCamerasEditor/Private/Editors/ObjectTreeGraphNode.cpp
@@ -533,14 +533,13 @@ void UObjectTreeGraphNode::OnUpdateCommentText(const FString& NewComment)
 	}
 }
 
-void UObjectTreeGraphNode::OnGraphNodeMoved()
+void UObjectTreeGraphNode::OnGraphNodeMoved(bool bMarkDirty)
 {
 	IObjectTreeGraphObject* GraphObject = Cast<IObjectTreeGraphObject>(Object);
 	if (GraphObject)
 	{
-		const FScopedTransaction Transaction(LOCTEXT("MoveNode", "Move Node"));
-		Object->Modify();
-		GraphObject->OnGraphNodeMoved(NodePosX, NodePosY);
+		Object->Modify(bMarkDirty);
+		GraphObject->OnGraphNodeMoved(NodePosX, NodePosY, bMarkDirty);
 	}
 }
 
diff --git a/Engine/Plugins/Cameras/GameplayCameras/Source/GameplayCamerasEditor/Private/Editors/ObjectTreeGraphSchema.cpp b/Engine/Plugins/Cameras/GameplayCameras/Source/GameplayCamerasEditor/Private/Editors/ObjectTreeGraphSchema.cpp
--- a/Engine/Plugins/Cameras/GameplayCameras/Source/GameplayCamerasEditor/Private/Editors/ObjectTreeGraphSchema.cpp
+++ b/Engine/Plugins/Cameras/GameplayCameras/Source/GameplayCamerasEditor/Private/Editors/ObjectTreeGraphSchema.cpp
@@ -951,7 +951,7 @@ UEdGraphNode* FObjectGraphSchemaAction_NewNode::PerformAction(UEdGraph* ParentGr
 
 		NewGraphNode->NodePosX = Location.X;
 		NewGraphNode->NodePosY = Location.Y;
-		NewGraphNode->OnGraphNodeMoved();
+		NewGraphNode->OnGraphNodeMoved(false);
 
 		AutoSetupNewNode(NewGraphNode, FromPin);
 
diff --git a/Engine/Plugins/Cameras/GameplayCameras/Source/GameplayCamerasEditor/Private/Editors/SObjectTreeGraphNode.cpp b/Engine/Plugins/Cameras/GameplayCameras/Source/GameplayCamerasEditor/Private/Editors/SObjectTreeGraphNode.cpp
--- a/Engine/Plugins/Cameras/GameplayCameras/Source/GameplayCamerasEditor/Private/Editors/SObjectTreeGraphNode.cpp
+++ b/Engine/Plugins/Cameras/GameplayCameras/Source/GameplayCamerasEditor/Private/Editors/SObjectTreeGraphNode.cpp
@@ -20,7 +20,7 @@ void SObjectTreeGraphNode::MoveTo(const FVector2D& NewPosition, FNodeSet& NodeFi
 
 	if (ObjectGraphNode)
 	{
-		ObjectGraphNode->OnGraphNodeMoved();
+		ObjectGraphNode->OnGraphNodeMoved(bMarkDirty);
 	}
 }
 
diff --git a/Engine/Plugins/Cameras/GameplayCameras/Source/GameplayCamerasEditor/Public/Editors/ObjectTreeGraphNode.h b/Engine/Plugins/Cameras/GameplayCameras/Source/GameplayCamerasEditor/Public/Editors/ObjectTreeGraphNode.h
--- a/Engine/Plugins/Cameras/GameplayCameras/Source/GameplayCamerasEditor/Public/Editors/ObjectTreeGraphNode.h
+++ b/Engine/Plugins/Cameras/GameplayCameras/Source/GameplayCamerasEditor/Public/Editors/ObjectTreeGraphNode.h
@@ -79,7 +79,7 @@ public:
 	virtual void OnUpdateCommentText(const FString& NewComment) override;
 
 	// UObjectTreeGraphNode interface.
-	virtual void OnGraphNodeMoved();
+	virtual void OnGraphNodeMoved(bool bMarkDirty);
 	virtual void OnDoubleClicked() {}
 
 public:
