Cameras: fix incorrect style name for build status button

diff --git a/Engine/Plugins/Cameras/GameplayCameras/Source/GameplayCamerasEditor/Private/Styles/GameplayCamerasEditorStyle.cpp b/Engine/Plugins/Cameras/GameplayCameras/Source/GameplayCamerasEditor/Private/Styles/GameplayCamerasEditorStyle.cpp
--- a/Engine/Plugins/Cameras/GameplayCameras/Source/GameplayCamerasEditor/Private/Styles/GameplayCamerasEditorStyle.cpp
+++ b/Engine/Plugins/Cameras/GameplayCameras/Source/GameplayCamerasEditor/Private/Styles/GameplayCamerasEditorStyle.cpp
@@ -42,7 +42,7 @@ FGameplayCamerasEditorStyle::FGameplayCamerasEditorStyle()
 
 	Set("CameraRigAssetEditor.Build", new IMAGE_BRUSH_SVG("Icons/CameraRig-BuildStatus_Background", Icon20x20));
 	Set("CameraRigAssetEditor.BuildStatus.Background", new IMAGE_BRUSH_SVG("Icons/CameraRig-BuildStatus_Background", Icon20x20));
-	Set("CameraRigAssetEditor.BuildStatus.Overlay.Fail", new IMAGE_BRUSH_SVG("Icons/CameraRig-BuildStatus_Fail_Badge", Icon20x20, FStyleColors::Error));
+	Set("CameraRigAssetEditor.BuildStatus.Overlay.Error", new IMAGE_BRUSH_SVG("Icons/CameraRig-BuildStatus_Fail_Badge", Icon20x20, FStyleColors::Error));
 	Set("CameraRigAssetEditor.BuildStatus.Overlay.Good", new IMAGE_BRUSH_SVG("Icons/CameraRig-BuildStatus_Good_Badge", Icon20x20, FStyleColors::AccentGreen));
 	Set("CameraRigAssetEditor.BuildStatus.Overlay.Unknown", new IMAGE_BRUSH_SVG("Icons/CameraRig-BuildStatus_Unknown_Badge", Icon20x20, FStyleColors::AccentYellow));
 	Set("CameraRigAssetEditor.BuildStatus.Overlay.Warning", new IMAGE_BRUSH_SVG("Icons/CameraRig-BuildStatus_Warning_Badge", Icon20x20, FStyleColors::Warning));
