Cameras: fix pasted nodes not retaining their pasted position

diff --git a/Engine/Plugins/Cameras/GameplayCameras/Source/GameplayCamerasEditor/Private/Editors/SObjectTreeGraphEditor.cpp b/Engine/Plugins/Cameras/GameplayCameras/Source/GameplayCamerasEditor/Private/Editors/SObjectTreeGraphEditor.cpp
--- a/Engine/Plugins/Cameras/GameplayCameras/Source/GameplayCamerasEditor/Private/Editors/SObjectTreeGraphEditor.cpp
+++ b/Engine/Plugins/Cameras/GameplayCameras/Source/GameplayCamerasEditor/Private/Editors/SObjectTreeGraphEditor.cpp
@@ -362,6 +362,13 @@ void SObjectTreeGraphEditor::ImportNodesFromText(const FVector2D& Location, cons
 
 		PastedNode->SnapToGrid(SNodePanel::GetSnapGridSize());
 
+		// Notify object nodes of having been moved so that we save the new position
+		// in the underlying data.
+		if (UObjectTreeGraphNode* PastedObjectNode = Cast<UObjectTreeGraphNode>(PastedNode))
+		{
+			PastedObjectNode->OnGraphNodeMoved(false);
+		}
+
 		GraphEditor->SetNodeSelection(PastedNode, true);
 	}
 
