Cameras: improve how object-tree-graphs keep track of loose objects

This change refactors how the "connectable objects" APIs on IObjectTreeGraphRootObject work. This clarifies the responsibilities between classes between those that affect the underlying data objects and those that only care about the graph nodes. It also makes it possible to have loose objects on multiple graphs related to the same asset, such as the node-tree and transitions graphs for camera rigs.

diff --git a/Engine/Plugins/Cameras/GameplayCameras/Source/GameplayCameras/Private/Core/CameraRigAsset.cpp b/Engine/Plugins/Cameras/GameplayCameras/Source/GameplayCameras/Private/Core/CameraRigAsset.cpp
--- a/Engine/Plugins/Cameras/GameplayCameras/Source/GameplayCameras/Private/Core/CameraRigAsset.cpp
+++ b/Engine/Plugins/Cameras/GameplayCameras/Source/GameplayCameras/Private/Core/CameraRigAsset.cpp
@@ -41,6 +41,9 @@ bool FCameraRigInterface::HasInterfaceParameter(const FString& ParameterName) co
 	return FindInterfaceParameterByName(ParameterName) != nullptr;
 }
 
+const FName UCameraRigAsset::NodeTreeGraphName(TEXT("NodeTree"));
+const FName UCameraRigAsset::TransitionsGraphName(TEXT("Transitions"));
+
 void UCameraRigAsset::GetOwnedGameplayTags(FGameplayTagContainer& TagContainer) const
 {
 	TagContainer.AppendTags(GameplayTags);
@@ -54,28 +57,6 @@ void UCameraRigAsset::BuildCameraRig()
 	Builder.BuildCameraRig(this);
 }
 
-void UCameraRigAsset::PostLoad()
-{
-#if WITH_EDITORONLY_DATA
-	if (AllNodes.IsEmpty())
-	{
-		// Rebuild the AllNodes array from scratch if this is old data.
-		TArray<UObject*> AllObjects;
-		UPackage* Package = GetOutermost();
-		GetObjectsWithPackage(Package, AllObjects);
-		for (UObject* Obj : AllObjects)
-		{
-			if (UCameraNode* CameraNode = Cast<UCameraNode>(Obj))
-			{
-				AllNodes.Add(CameraNode);
-			}
-		}
-	}
-#endif  // WITH_EDITORONLY_DATA
-
-	Super::PostLoad();
-}
-
 void UCameraRigAsset::PreSave(FObjectPreSaveContext ObjectSaveContext)
 {
 #if WITH_EDITOR
@@ -134,24 +115,46 @@ void UCameraRigAsset::OnUpdateGraphNodeCommentText(const FString& NewComment)
 	GraphNodeComment = NewComment;
 }
 
-void UCameraRigAsset::AddConnectableObject(UObject* InObject)
+void UCameraRigAsset::GetConnectableObjects(FName InGraphName, TSet<UObject*>& OutObjects) const
+{
+	if (InGraphName == NodeTreeGraphName)
+	{
+		OutObjects.Append(AllNodeTreeObjects);
+	}
+	else if (InGraphName == TransitionsGraphName)
+	{
+		OutObjects.Append(AllTransitionsObjects);
+	}
+}
+
+void UCameraRigAsset::AddConnectableObject(FName InGraphName, UObject* InObject)
 {
-	// Camera nodes and rig parameters are shown on the graph editor. Handle the former case. The latter are
-	// already added to the interface struct even when not connected.
-	if (UCameraNode* CameraNode = Cast<UCameraNode>(InObject))
+	Modify();
+
+	if (InGraphName == NodeTreeGraphName)
 	{
-		Modify();
-		const int32 Index = AllNodes.AddUnique(CameraNode);
-		ensure(Index == AllNodes.Num() - 1);
+		const int32 Index = AllNodeTreeObjects.AddUnique(InObject);
+		ensure(Index == AllNodeTreeObjects.Num() - 1);
+	}
+	else if (InGraphName == TransitionsGraphName)
+	{
+		const int32 Index = AllTransitionsObjects.AddUnique(InObject);
+		ensure(Index == AllTransitionsObjects.Num() - 1);
 	}
 }
 
-void UCameraRigAsset::RemoveConnectableObject(UObject* InObject)
+void UCameraRigAsset::RemoveConnectableObject(FName InGraphName, UObject* InObject)
 {
-	if (UCameraNode* CameraNode = Cast<UCameraNode>(InObject))
+	Modify();
+
+	if (InGraphName == NodeTreeGraphName)
+	{
+		const int32 NumRemoved = AllNodeTreeObjects.Remove(InObject);
+		ensure(NumRemoved == 1);
+	}
+	else if (InGraphName == TransitionsGraphName)
 	{
-		Modify();
-		const int32 NumRemoved = AllNodes.Remove(CameraNode);
+		const int32 NumRemoved = AllTransitionsObjects.Remove(InObject);
 		ensure(NumRemoved == 1);
 	}
 }
diff --git a/Engine/Plugins/Cameras/GameplayCameras/Source/GameplayCameras/Public/Core/CameraRigAsset.h b/Engine/Plugins/Cameras/GameplayCameras/Source/GameplayCameras/Public/Core/CameraRigAsset.h
--- a/Engine/Plugins/Cameras/GameplayCameras/Source/GameplayCameras/Public/Core/CameraRigAsset.h
+++ b/Engine/Plugins/Cameras/GameplayCameras/Source/GameplayCameras/Public/Core/CameraRigAsset.h
@@ -196,6 +196,12 @@ public:
 	// IGameplayTagAssetInterface.
 	virtual void GetOwnedGameplayTags(FGameplayTagContainer& TagContainer) const override;
 
+public:
+
+	// Graph names for ObjectTreeGraph API.
+	GAMEPLAYCAMERAS_API static const FName NodeTreeGraphName;
+	GAMEPLAYCAMERAS_API static const FName TransitionsGraphName;
+
 protected:
 
 	// IObjectTreeGraphObject interface.
@@ -208,12 +214,12 @@ protected:
 #endif
 
 #if WITH_EDITOR
-	virtual void AddConnectableObject(UObject* InObject) override;
-	virtual void RemoveConnectableObject(UObject* InObject) override;
+	virtual void GetConnectableObjects(FName InGraphName, TSet<UObject*>& OutObjects) const override;
+	virtual void AddConnectableObject(FName InGraphName, UObject* InObject) override;
+	virtual void RemoveConnectableObject(FName InGraphName, UObject* InObject) override;
 #endif
 
 	// UObject interface
-	virtual void PostLoad() override;
 	virtual void PreSave(FObjectPreSaveContext ObjectSaveContext) override;
 
 private:
@@ -237,7 +243,13 @@ private:
 	 * to the root node, and therefore would be GC'ed if we didn't hold them here.
 	 */
 	UPROPERTY(Instanced, meta=(ObjectTreeGraphHidden=true))
-	TArray<TObjectPtr<UCameraNode>> AllNodes;
+	TArray<TObjectPtr<UObject>> AllNodeTreeObjects;
+
+	/**
+	 * Similar to AllNodeTreeObjects, but for the transitions graph.
+	 */
+	UPROPERTY(Instanced, meta=(ObjectTreeGraphHidden=true))
+	TArray<TObjectPtr<UObject>> AllTransitionsObjects;
 
 #endif  // WITH_EDITORONLY_DATA
 
diff --git a/Engine/Plugins/Cameras/GameplayCameras/Source/GameplayCameras/Public/Core/ObjectTreeGraphRootObject.h b/Engine/Plugins/Cameras/GameplayCameras/Source/GameplayCameras/Public/Core/ObjectTreeGraphRootObject.h
--- a/Engine/Plugins/Cameras/GameplayCameras/Source/GameplayCameras/Public/Core/ObjectTreeGraphRootObject.h
+++ b/Engine/Plugins/Cameras/GameplayCameras/Source/GameplayCameras/Public/Core/ObjectTreeGraphRootObject.h
@@ -21,9 +21,9 @@ public:
 
 #if WITH_EDITOR
 
-	virtual void AddConnectableObject(UObject* InObject) {}
-	virtual void RemoveConnectableObject(UObject* InObject) {}
-	virtual void GetExtraConnectableObjects(TArray<UObject*>& OutObjects) const {}
+	virtual void GetConnectableObjects(FName InGraphName, TSet<UObject*>& OutObjects) const {}
+	virtual void AddConnectableObject(FName InGraphName, UObject* InObject) {}
+	virtual void RemoveConnectableObject(FName InGraphName, UObject* InObject) {}
 
 #endif  // WITH_EDITOR
 
diff --git a/Engine/Plugins/Cameras/GameplayCameras/Source/GameplayCamerasEditor/Private/Editors/CameraNodeGraphSchema.cpp b/Engine/Plugins/Cameras/GameplayCameras/Source/GameplayCamerasEditor/Private/Editors/CameraNodeGraphSchema.cpp
--- a/Engine/Plugins/Cameras/GameplayCameras/Source/GameplayCamerasEditor/Private/Editors/CameraNodeGraphSchema.cpp
+++ b/Engine/Plugins/Cameras/GameplayCameras/Source/GameplayCamerasEditor/Private/Editors/CameraNodeGraphSchema.cpp
@@ -21,10 +21,10 @@ void UCameraNodeGraphSchema::OnCreateAllNodes(UObjectTreeGraph* InGraph, const F
 {
 	Super::OnCreateAllNodes(InGraph, InCreatedNodes);
 
-	UCameraRigAsset* CameraRigAsset = Cast<UCameraRigAsset>(InGraph->GetRootObject());
-	if (ensure(CameraRigAsset))
+	UCameraRigAsset* CameraRig = Cast<UCameraRigAsset>(InGraph->GetRootObject());
+	if (ensure(CameraRig))
 	{
-		for (UCameraRigInterfaceParameter* InterfaceParameter : CameraRigAsset->Interface.InterfaceParameters)
+		for (UCameraRigInterfaceParameter* InterfaceParameter : CameraRig->Interface.InterfaceParameters)
 		{
 			UObjectTreeGraphNode* const* InterfaceParameterNode = InCreatedNodes.CreatedNodes.Find(InterfaceParameter);
 			UObjectTreeGraphNode* const* CameraNodeNode = InCreatedNodes.CreatedNodes.Find(InterfaceParameter->Target);
@@ -38,6 +38,40 @@ void UCameraNodeGraphSchema::OnCreateAllNodes(UObjectTreeGraph* InGraph, const F
 	}
 }
 
+void UCameraNodeGraphSchema::OnAddConnectableObject(UObjectTreeGraph* InGraph, UObjectTreeGraphNode* InNewNode) const
+{
+	Super::OnAddConnectableObject(InGraph, InNewNode);
+
+	if (UCameraRigInterfaceParameter* InterfaceParameter = Cast<UCameraRigInterfaceParameter>(InNewNode->GetObject()))
+	{
+		UCameraRigAsset* CameraRig = Cast<UCameraRigAsset>(InGraph->GetRootObject());
+		if (ensure(CameraRig))
+		{
+			CameraRig->Modify();
+
+			const int32 Index = CameraRig->Interface.InterfaceParameters.AddUnique(InterfaceParameter);
+			ensure(Index == CameraRig->Interface.InterfaceParameters.Num() - 1);
+		}
+	}
+}
+
+void UCameraNodeGraphSchema::OnRemoveConnectableObject(UObjectTreeGraph* InGraph, UObjectTreeGraphNode* InRemovedNode) const
+{
+	Super::OnRemoveConnectableObject(InGraph, InRemovedNode);
+
+	if (UCameraRigInterfaceParameter* InterfaceParameter = Cast<UCameraRigInterfaceParameter>(InRemovedNode->GetObject()))
+	{
+		UCameraRigAsset* CameraRig = Cast<UCameraRigAsset>(InGraph->GetRootObject());
+		if (ensure(CameraRig))
+		{
+			CameraRig->Modify();
+
+			const int32 NumRemoved = CameraRig->Interface.InterfaceParameters.Remove(InterfaceParameter);
+			ensure(NumRemoved == 1);
+		}
+	}
+}
+
 void UCameraNodeGraphSchema::GetGraphContextActions(FGraphContextMenuBuilder& ContextMenuBuilder) const
 {
 	// See if we were dragging a camera parameter pin.
@@ -220,23 +254,6 @@ bool UCameraNodeGraphSchema::OnBreakSinglePinLink(UEdGraphPin* SourcePin, UEdGra
 	return false;
 }
 
-void UCameraNodeGraphSchema::OnDeleteNodeFromGraph(UObjectTreeGraph* Graph, UEdGraphNode* Node) const
-{
-	Super::OnDeleteNodeFromGraph(Graph, Node);
-
-	UCameraRigAsset* CameraRig = CastChecked<UCameraRigAsset>(Graph->GetRootObject());
-	if (UCameraRigInterfaceParameterGraphNode* RigParameterNode = Cast<UCameraRigInterfaceParameterGraphNode>(Node))
-	{
-		UCameraRigInterfaceParameter* RigParameter = CastChecked<UCameraRigInterfaceParameter>(RigParameterNode->GetObject());
-
-		CameraRig->Modify();
-		RigParameter->Modify();
-
-		const int32 NumRemoved = CameraRig->Interface.InterfaceParameters.Remove(RigParameter);
-		ensure(NumRemoved == 1);
-	}
-}
-
 FCameraNodeGraphSchemaAction_NewInterfaceParameterNode::FCameraNodeGraphSchemaAction_NewInterfaceParameterNode()
 {
 }
@@ -260,19 +277,20 @@ UEdGraphNode* FCameraNodeGraphSchemaAction_NewInterfaceParameterNode::PerformAct
 		return nullptr;
 	}
 
-	const UObjectTreeGraphSchema* Schema = CastChecked<UObjectTreeGraphSchema>(ParentGraph->GetSchema());
-
 	const FScopedTransaction Transaction(LOCTEXT("CreateNewNodeAction", "Create New Node"));
 
-	CameraRig->Modify();
+	const FName GraphName = ObjectTreeGraph->GetConfig().GraphName;
+	const UObjectTreeGraphSchema* Schema = CastChecked<UObjectTreeGraphSchema>(ParentGraph->GetSchema());
 
 	UCameraRigInterfaceParameter* NewInterfaceParameter = NewObject<UCameraRigInterfaceParameter>(CameraRig, NAME_None, RF_Transactional);
 	NewInterfaceParameter->Target = Target;
 	NewInterfaceParameter->TargetPropertyName = TargetPropertyName;
 	NewInterfaceParameter->InterfaceParameterName = TargetPropertyName.ToString();
-	CameraRig->Interface.InterfaceParameters.Add(NewInterfaceParameter);
 
 	UObjectTreeGraphNode* NewGraphNode = Schema->CreateObjectNode(ObjectTreeGraph, NewInterfaceParameter);
+
+	Schema->AddConnectableObject(ObjectTreeGraph, NewGraphNode);
+
 	NewGraphNode->NodePosX = Location.X;
 	NewGraphNode->NodePosY = Location.Y;
 	NewGraphNode->OnGraphNodeMoved(false);
diff --git a/Engine/Plugins/Cameras/GameplayCameras/Source/GameplayCamerasEditor/Private/Editors/ObjectTreeGraphSchema.cpp b/Engine/Plugins/Cameras/GameplayCameras/Source/GameplayCamerasEditor/Private/Editors/ObjectTreeGraphSchema.cpp
--- a/Engine/Plugins/Cameras/GameplayCameras/Source/GameplayCamerasEditor/Private/Editors/ObjectTreeGraphSchema.cpp
+++ b/Engine/Plugins/Cameras/GameplayCameras/Source/GameplayCamerasEditor/Private/Editors/ObjectTreeGraphSchema.cpp
@@ -140,7 +140,7 @@ void UObjectTreeGraphSchema::CreateAllNodes(UObjectTreeGraph* InGraph, EObjectTr
 		return;
 	}
 
-	TArray<UObject*> AllObjects;
+	TSet<UObject*> AllObjects;
 	// Gather up all the objects we need for the graph. Start by all objects that are referenced (directly or
 	// indirectly) by the root object. Our custom reference collector will not collect references that go outside
 	// of the root object's package.
@@ -150,13 +150,16 @@ void UObjectTreeGraphSchema::CreateAllNodes(UObjectTreeGraph* InGraph, EObjectTr
 		// Make sure the root object itself is in there.
 		AllObjects.Add(RootObject);
 
-		FPackageReferenceCollector Collector(RootObject, AllObjects);
+		TArray<UObject*> ReferencedObjects;
+		FPackageReferenceCollector Collector(RootObject, ReferencedObjects);
 		Collector.CollectReferences();
+		AllObjects.Append(ReferencedObjects);
 	}
 	// Add any other custom objects the root object may want.
 	if (IObjectTreeGraphRootObject* RootObjectInterface = Cast<IObjectTreeGraphRootObject>(RootObject))
 	{
-		RootObjectInterface->GetExtraConnectableObjects(AllObjects);
+		const FName GraphName = InGraph->GetConfig().GraphName;
+		RootObjectInterface->GetConnectableObjects(GraphName, AllObjects);
 	}
 	
 	// Create all the nodes.
@@ -273,10 +276,10 @@ UObjectTreeGraphNode* UObjectTreeGraphSchema::CreateObjectNode(UObjectTreeGraph*
 		return nullptr;
 	}
 
-	return CreateObjectNodeImpl(InGraph, InObject);
+	return OnCreateObjectNode(InGraph, InObject);
 }
 
-UObjectTreeGraphNode* UObjectTreeGraphSchema::CreateObjectNodeImpl(UObjectTreeGraph* InGraph, UObject* InObject) const
+UObjectTreeGraphNode* UObjectTreeGraphSchema::OnCreateObjectNode(UObjectTreeGraph* InGraph, UObject* InObject) const
 {
 	const FObjectTreeGraphConfig& Config = InGraph->GetConfig();
 	const FObjectTreeGraphClassConfig& ClassConfig = Config.GetObjectClassConfig(InObject->GetClass());
@@ -294,6 +297,39 @@ UObjectTreeGraphNode* UObjectTreeGraphSchema::CreateObjectNodeImpl(UObjectTreeGr
 	return NewNode;
 }
 
+void UObjectTreeGraphSchema::AddConnectableObject(UObjectTreeGraph* InGraph, UObjectTreeGraphNode* InNewNode) const
+{
+	UObjectTreeGraphNode* RootObjectNode = InGraph->GetRootObjectNode();
+	IObjectTreeGraphRootObject* RootObjectInterface = Cast<IObjectTreeGraphRootObject>(RootObjectNode->GetObject());
+	if (RootObjectInterface)
+	{
+		const FName GraphName = InGraph->GetConfig().GraphName;
+		RootObjectInterface->AddConnectableObject(GraphName, InNewNode->GetObject());
+	}
+
+	OnAddConnectableObject(InGraph, InNewNode);
+}
+
+void UObjectTreeGraphSchema::OnAddConnectableObject(UObjectTreeGraph* InGraph, UObjectTreeGraphNode* InNewNode) const
+{
+}
+
+void UObjectTreeGraphSchema::RemoveConnectableObject(UObjectTreeGraph* InGraph, UObjectTreeGraphNode* InRemovedNode) const
+{
+	const FName GraphName = InGraph->GetConfig().GraphName;
+	IObjectTreeGraphRootObject* RootObjectInterface = Cast<IObjectTreeGraphRootObject>(InGraph->GetRootObject());
+	if (RootObjectInterface)
+	{
+		RootObjectInterface->RemoveConnectableObject(GraphName, InRemovedNode->GetObject());
+	}
+
+	OnRemoveConnectableObject(InGraph, InRemovedNode);
+}
+
+void UObjectTreeGraphSchema::OnRemoveConnectableObject(UObjectTreeGraph* InGraph, UObjectTreeGraphNode* InRemovedNode) const
+{
+}
+
 void UObjectTreeGraphSchema::GetGraphContextActions(FGraphContextMenuBuilder& ContextMenuBuilder) const
 {
 	const UObjectTreeGraph* Graph = CastChecked<UObjectTreeGraph>(ContextMenuBuilder.CurrentGraph);
@@ -740,10 +776,9 @@ bool UObjectTreeGraphSchema::SafeDeleteNodeFromGraph(UEdGraph* Graph, UEdGraphNo
 void UObjectTreeGraphSchema::OnDeleteNodeFromGraph(UObjectTreeGraph* Graph, UEdGraphNode* Node) const
 {
 	UObjectTreeGraphNode* ObjectNode = Cast<UObjectTreeGraphNode>(Node);
-	IObjectTreeGraphRootObject* RootObjectInterface = Cast<IObjectTreeGraphRootObject>(Graph->GetRootObject());
-	if (RootObjectInterface && ObjectNode)
+	if (ObjectNode)
 	{
-		RootObjectInterface->RemoveConnectableObject(ObjectNode->GetObject());
+		RemoveConnectableObject(Graph, ObjectNode);
 	}
 }
 
@@ -839,30 +874,26 @@ void UObjectTreeGraphSchema::ImportNodesFromText(UObjectTreeGraph* InGraph, cons
 	TempPackage->RemoveFromRoot();
 
 	// Finish setting up the new objects: clear the transient flag from the transient package we used above,
-	// move the objects under the our graph root, and optionally add them via the root interface.
+	// and move the objects under the our graph root.
 	UObject* GraphRootObject = InGraph->GetRootObject();
-	IObjectTreeGraphRootObject* RootObjectInterface = Cast<IObjectTreeGraphRootObject>(GraphRootObject);
 	if (ensure(GraphRootObject))
 	{
 		for (UObject* Object : ImportedObjects)
 		{
 			Object->ClearFlags(RF_Transient);
 			Object->Rename(nullptr, GraphRootObject);
-
-			if (RootObjectInterface)
-			{
-				RootObjectInterface->AddConnectableObject(Object);
-			}
 		}
 	}
 
-	// Create nodes for all the imported objects.
+	// Create nodes for all the imported objects, and add them to the root object if it supports the root interface.
 	FCreatedNodes CreatedNodes;
 	for (UObject* Object : ImportedObjects)
 	{
 		if (UObjectTreeGraphNode* GraphNode = CreateObjectNode(InGraph, Object))
 		{
 			CreatedNodes.CreatedNodes.Add(Object, GraphNode);
+
+			AddConnectableObject(InGraph, GraphNode);
 		}
 	}
 
@@ -934,9 +965,9 @@ UEdGraphNode* FObjectGraphSchemaAction_NewNode::PerformAction(UEdGraph* ParentGr
 		return nullptr;
 	}
 
-	const FScopedTransaction Transaction(LOCTEXT("CreateNewNodeAction", "Create New Node"));
+	const FScopedTransaction Transaction(FText::Format(LOCTEXT("CreateNewNodeAction", "Create {0} Node"), ObjectClass->GetDisplayNameText()));
+
 	const UObjectTreeGraphSchema* Schema = CastChecked<UObjectTreeGraphSchema>(ParentGraph->GetSchema());
-	IObjectTreeGraphRootObject* RootObjectInterface = Cast<IObjectTreeGraphRootObject>(ObjectTreeGraph->GetRootObject());
 
 	UObject* NewObject = CreateObject();
 
@@ -944,10 +975,7 @@ UEdGraphNode* FObjectGraphSchemaAction_NewNode::PerformAction(UEdGraph* ParentGr
 	{
 		UObjectTreeGraphNode* NewGraphNode = Schema->CreateObjectNode(ObjectTreeGraph, NewObject);
 
-		if (RootObjectInterface)
-		{
-			RootObjectInterface->AddConnectableObject(NewObject);
-		}
+		Schema->AddConnectableObject(ObjectTreeGraph, NewGraphNode);
 
 		NewGraphNode->NodePosX = Location.X;
 		NewGraphNode->NodePosY = Location.Y;
diff --git a/Engine/Plugins/Cameras/GameplayCameras/Source/GameplayCamerasEditor/Private/Editors/SCameraRigAssetEditor.cpp b/Engine/Plugins/Cameras/GameplayCameras/Source/GameplayCamerasEditor/Private/Editors/SCameraRigAssetEditor.cpp
--- a/Engine/Plugins/Cameras/GameplayCameras/Source/GameplayCamerasEditor/Private/Editors/SCameraRigAssetEditor.cpp
+++ b/Engine/Plugins/Cameras/GameplayCameras/Source/GameplayCamerasEditor/Private/Editors/SCameraRigAssetEditor.cpp
@@ -61,6 +61,7 @@ void SCameraRigAssetEditor::CreateNodeGraphEditor(const FArguments& InArgs)
 	const UGameplayCamerasEditorSettings* Settings = GetDefault<UGameplayCamerasEditorSettings>();
 
 	FObjectTreeGraphConfig GraphConfig;
+	GraphConfig.GraphName = UCameraRigAsset::NodeTreeGraphName;
 	GraphConfig.ConnectableObjectClasses.Add(UCameraRigAsset::StaticClass());
 	GraphConfig.ConnectableObjectClasses.Add(UCameraNode::StaticClass());
 	GraphConfig.ConnectableObjectClasses.Add(UCameraRigInterfaceParameter::StaticClass());
@@ -102,6 +103,7 @@ void SCameraRigAssetEditor::CreateTransitionGraphEditor(const FArguments& InArgs
 	const UGameplayCamerasEditorSettings* Settings = GetDefault<UGameplayCamerasEditorSettings>();
 
 	FObjectTreeGraphConfig GraphConfig;
+	GraphConfig.GraphName = UCameraRigAsset::TransitionsGraphName;
 	GraphConfig.ConnectableObjectClasses.Add(UCameraRigAsset::StaticClass());
 	GraphConfig.ConnectableObjectClasses.Add(UCameraRigTransition::StaticClass());
 	GraphConfig.ConnectableObjectClasses.Add(UCameraRigTransitionCondition::StaticClass());
diff --git a/Engine/Plugins/Cameras/GameplayCameras/Source/GameplayCamerasEditor/Public/Editors/CameraNodeGraphSchema.h b/Engine/Plugins/Cameras/GameplayCameras/Source/GameplayCamerasEditor/Public/Editors/CameraNodeGraphSchema.h
--- a/Engine/Plugins/Cameras/GameplayCameras/Source/GameplayCamerasEditor/Public/Editors/CameraNodeGraphSchema.h
+++ b/Engine/Plugins/Cameras/GameplayCameras/Source/GameplayCamerasEditor/Public/Editors/CameraNodeGraphSchema.h
@@ -28,10 +28,11 @@ protected:
 
 	// UObjectTreeGraphSchema interface.
 	virtual void OnCreateAllNodes(UObjectTreeGraph* InGraph, const FCreatedNodes& InCreatedNodes) const override;
+	virtual void OnAddConnectableObject(UObjectTreeGraph* InGraph, UObjectTreeGraphNode* InNewNode) const override;
+	virtual void OnRemoveConnectableObject(UObjectTreeGraph* InGraph, UObjectTreeGraphNode* InRemovedNode) const override;
 	virtual bool OnCreateConnection(UEdGraphPin* A, UEdGraphPin* B) const override;
 	virtual bool OnBreakPinLinks(UEdGraphPin& TargetPin, bool bSendsNodeNotification) const override;
 	virtual bool OnBreakSinglePinLink(UEdGraphPin* SourcePin, UEdGraphPin* TargetPin) const override;
-	virtual void OnDeleteNodeFromGraph(UObjectTreeGraph* Graph, UEdGraphNode* Node) const override;
 };
 
 /**
diff --git a/Engine/Plugins/Cameras/GameplayCameras/Source/GameplayCamerasEditor/Public/Editors/ObjectTreeGraphConfig.h b/Engine/Plugins/Cameras/GameplayCameras/Source/GameplayCamerasEditor/Public/Editors/ObjectTreeGraphConfig.h
--- a/Engine/Plugins/Cameras/GameplayCameras/Source/GameplayCamerasEditor/Public/Editors/ObjectTreeGraphConfig.h
+++ b/Engine/Plugins/Cameras/GameplayCameras/Source/GameplayCamerasEditor/Public/Editors/ObjectTreeGraphConfig.h
@@ -130,6 +130,11 @@ struct FObjectTreeGraphConfig
 {
 public:
 
+	/**
+	 * The name of the graph, passed to some APIs like IObjectTreeGraphRootObject.
+	 */
+	FName GraphName;
+
 	/**
 	 * The list of connectable object classes.
 	 *
diff --git a/Engine/Plugins/Cameras/GameplayCameras/Source/GameplayCamerasEditor/Public/Editors/ObjectTreeGraphSchema.h b/Engine/Plugins/Cameras/GameplayCameras/Source/GameplayCamerasEditor/Public/Editors/ObjectTreeGraphSchema.h
--- a/Engine/Plugins/Cameras/GameplayCameras/Source/GameplayCamerasEditor/Public/Editors/ObjectTreeGraphSchema.h
+++ b/Engine/Plugins/Cameras/GameplayCameras/Source/GameplayCamerasEditor/Public/Editors/ObjectTreeGraphSchema.h
@@ -8,6 +8,7 @@
 
 #include "ObjectTreeGraphSchema.generated.h"
 
+class IObjectTreeGraphRootObject;
 class UEdGraph;
 class UObjectTreeGraph;
 class UObjectTreeGraphNode;
@@ -50,6 +51,12 @@ public:
 	/** Creates an object graph node for the given object. */
 	UObjectTreeGraphNode* CreateObjectNode(UObjectTreeGraph* InGraph, UObject* InObject) const;
 
+	/** Adds an object to the underlying data after it has been added to the graph. */
+	void AddConnectableObject(UObjectTreeGraph* InGraph, UObjectTreeGraphNode* InNewNode) const;
+
+	/** Removes an object from the underlying data after it has been removed from the graph. */
+	void RemoveConnectableObject(UObjectTreeGraph* InGraph, UObjectTreeGraphNode* InRemovedNode) const;
+
 	/** Export the given selection into a text suitable for copy/pasting. */
 	FString ExportNodesToText(const FGraphPanelSelectionSet& Nodes, bool bOnlyCanDuplicateNodes, bool bOnlyCanDeleteNodes) const;
 
@@ -76,13 +83,6 @@ public:
 	virtual bool SupportsDropPinOnNode(UEdGraphNode* InTargetNode, const FEdGraphPinType& InSourcePinType, EEdGraphPinDirection InSourcePinDirection, FText& OutErrorMessage) const override;
 	virtual bool SafeDeleteNodeFromGraph(UEdGraph* Graph, UEdGraphNode* Node) const override;
 
-	// UObjectTreeGraphSchema interface.
-	virtual bool OnCreateConnection(UEdGraphPin* A, UEdGraphPin* B) const;
-	virtual bool OnBreakPinLinks(UEdGraphPin& TargetPin, bool bSendsNodeNotification) const;
-	virtual bool OnBreakSinglePinLink(UEdGraphPin* SourcePin, UEdGraphPin* TargetPin) const;
-	virtual void OnDeleteNodeFromGraph(UObjectTreeGraph* Graph, UEdGraphNode* Node) const;
-	virtual void FilterGraphContextPlaceableClasses(TArray<UClass*>& InOutClasses) const;
-
 protected:
 
 	struct FCreatedNodes
@@ -92,8 +92,15 @@ protected:
 
 	// UObjectTreeGraphSchema interface.
 	virtual void OnCreateAllNodes(UObjectTreeGraph* InGraph, const FCreatedNodes& InCreatedNodes) const;
-	virtual UObjectTreeGraphNode* CreateObjectNodeImpl(UObjectTreeGraph* InGraph, UObject* InObject) const;
+	virtual UObjectTreeGraphNode* OnCreateObjectNode(UObjectTreeGraph* InGraph, UObject* InObject) const;
+	virtual void OnAddConnectableObject(UObjectTreeGraph* InGraph, UObjectTreeGraphNode* InNewNode) const;
+	virtual void OnRemoveConnectableObject(UObjectTreeGraph* InGraph, UObjectTreeGraphNode* InRemovedNode) const;
 	virtual void CopyNonObjectNodes(TArrayView<UObject*> InObjects, FStringOutputDevice& OutDevice) const;
+	virtual bool OnCreateConnection(UEdGraphPin* A, UEdGraphPin* B) const;
+	virtual bool OnBreakPinLinks(UEdGraphPin& TargetPin, bool bSendsNodeNotification) const;
+	virtual bool OnBreakSinglePinLink(UEdGraphPin* SourcePin, UEdGraphPin* TargetPin) const;
+	virtual void OnDeleteNodeFromGraph(UObjectTreeGraph* Graph, UEdGraphNode* Node) const;
+	virtual void FilterGraphContextPlaceableClasses(TArray<UClass*>& InOutClasses) const;
 
 protected:
 
