Cameras: don't crash when trying to build an evaluator for a null camera node

diff --git a/Engine/Plugins/Cameras/GameplayCameras/Source/GameplayCameras/Private/Core/CameraNodeEvaluator.cpp b/Engine/Plugins/Cameras/GameplayCameras/Source/GameplayCameras/Private/Core/CameraNodeEvaluator.cpp
--- a/Engine/Plugins/Cameras/GameplayCameras/Source/GameplayCameras/Private/Core/CameraNodeEvaluator.cpp
+++ b/Engine/Plugins/Cameras/GameplayCameras/Source/GameplayCameras/Private/Core/CameraNodeEvaluator.cpp
@@ -23,9 +23,13 @@ void FCameraNodeEvaluationResult::Reset()
 
 FCameraNodeEvaluator* FCameraNodeEvaluatorBuildParams::BuildEvaluator(const UCameraNode* InNode) const
 {
-	FCameraNodeEvaluator* NewEvaluator = InNode->BuildEvaluator(Builder);
-	NewEvaluator->Build(*this);
-	return NewEvaluator;
+	if (InNode)
+	{
+		FCameraNodeEvaluator* NewEvaluator = InNode->BuildEvaluator(Builder);
+		NewEvaluator->Build(*this);
+		return NewEvaluator;
+	}
+	return nullptr;
 }
 
 FCameraNodeEvaluator::FCameraNodeEvaluator()
