Cameras: building a camera rig now populates the messages window in the editor

- Add a new build log object to collect warning/error messages from the camera rig builder.
- Make it possible to forwards these messages to the console.
- Populate the messages window in the camera rig editor with the build messages.
- Implement logic to jump to a graph node corresponding to an object mentioned in a message.

diff --git a/Engine/Plugins/Cameras/GameplayCameras/Source/GameplayCameras/Private/Core/CameraRigAsset.cpp b/Engine/Plugins/Cameras/GameplayCameras/Source/GameplayCameras/Private/Core/CameraRigAsset.cpp
--- a/Engine/Plugins/Cameras/GameplayCameras/Source/GameplayCameras/Private/Core/CameraRigAsset.cpp
+++ b/Engine/Plugins/Cameras/GameplayCameras/Source/GameplayCameras/Private/Core/CameraRigAsset.cpp
@@ -53,7 +53,16 @@ void UCameraRigAsset::BuildCameraRig()
 {
 	using namespace UE::Cameras;
 
-	FCameraRigAssetBuilder Builder;
+	FCameraRigAssetBuildLog BuildLog;
+	BuildLog.SetForwardMessagesToLogging(true);
+	BuildCameraRig(BuildLog);
+}
+
+void UCameraRigAsset::BuildCameraRig(UE::Cameras::FCameraRigAssetBuildLog& InBuildLog)
+{
+	using namespace UE::Cameras;
+
+	FCameraRigAssetBuilder Builder(InBuildLog);
 	Builder.BuildCameraRig(this);
 }
 
diff --git a/Engine/Plugins/Cameras/GameplayCameras/Source/GameplayCameras/Private/Core/CameraRigAssetBuilder.cpp b/Engine/Plugins/Cameras/GameplayCameras/Source/GameplayCameras/Private/Core/CameraRigAssetBuilder.cpp
--- a/Engine/Plugins/Cameras/GameplayCameras/Source/GameplayCameras/Private/Core/CameraRigAssetBuilder.cpp
+++ b/Engine/Plugins/Cameras/GameplayCameras/Source/GameplayCameras/Private/Core/CameraRigAssetBuilder.cpp
@@ -8,6 +8,9 @@
 #include "Core/CameraRigAsset.h"
 #include "Core/CameraVariableAssets.h"
 #include "Misc/CString.h"
+#include "Misc/UObjectToken.h"
+
+#define LOCTEXT_NAMESPACE "CameraRigAssetBuilder"
 
 namespace UE::Cameras
 {
@@ -68,8 +71,14 @@ struct FPrivateVariableBuilder
 		CameraRig = Owner.CameraRig;
 	}
 
-	void ReportError()
+	void ReportError(FText&& ErrorMessage)
+	{
+		ReportError(nullptr, MoveTemp(ErrorMessage));
+	}
+
+	void ReportError(UObject* Object, FText&& ErrorMessage)
 	{
+		Owner.BuildLog.AddMessage(EMessageSeverity::Error, MoveTemp(ErrorMessage));
 		Owner.bHasErrors = true;
 	}
 
@@ -130,12 +139,14 @@ void SetupPrivateVariable(
 		{
 			// If this parameter is driven by a user-defined variable, emit an error, and replace that 
 			// driving variable with our private variable.
-			UE_LOG(LogCameraSystem, Error, 
-					TEXT("Camera node parameter '%s.%s' in camera rig '%s' is both exposed and driven by a variable!."), 
-					*InterfaceParameter->Target->GetName(), 
-					*InterfaceParameter->TargetPropertyName.ToString(),
-					*Builder.CameraRig->GetPathName());
-			Builder.ReportError();
+			Builder.ReportError(
+					InterfaceParameter->Target,
+					FText::Format(
+						LOCTEXT(
+							"CameraParameterDrivenTwice", 
+							"Camera node parameter '{0}.{1}' is both exposed and driven by a variable!"),
+						FText::FromName(InterfaceParameter->Target->GetFName()), 
+						FText::FromName(InterfaceParameter->TargetPropertyName)));
 		}
 	}
 
@@ -190,6 +201,79 @@ void SetupPrivateVariable(
 
 }  // namespace Internal
 
+FString FCameraRigAssetBuildLogMessage::ToString() const
+{
+	TStringBuilder<256> StringBuilder;
+	if (Object)
+	{
+		StringBuilder.Append(Object->GetName());
+		StringBuilder.Append(TEXT(": "));
+	}
+	StringBuilder.Append(Text.ToString());
+	return StringBuilder.ToString();
+}
+
+void FCameraRigAssetBuildLogMessage::SendToLogging(const FString& InLoggingPrefix) const
+{
+#define UE_LOG_FORWARD_CAMERA_RIG_BULID_LOG_MESSAGE(Verbosity)\
+	UE_LOG(LogCameraSystem, Verbosity, TEXT("%s%s"), *InLoggingPrefix, *ToString());
+
+	switch (Severity)
+	{
+	case EMessageSeverity::Error:
+		UE_LOG_FORWARD_CAMERA_RIG_BULID_LOG_MESSAGE(Error);
+		break;
+	case EMessageSeverity::PerformanceWarning:
+	case EMessageSeverity::Warning:
+		UE_LOG_FORWARD_CAMERA_RIG_BULID_LOG_MESSAGE(Warning);
+		break;
+	case EMessageSeverity::Info:
+	default:
+		UE_LOG_FORWARD_CAMERA_RIG_BULID_LOG_MESSAGE(Log);
+		break;
+	};
+
+#undef UE_LOG_FORWARD_CAMERA_RIG_BULID_LOG_MESSAGE
+}
+
+void FCameraRigAssetBuildLog::SetLoggingPrefix(const FString& InPrefix)
+{
+	if (InPrefix.IsEmpty())
+	{
+		LoggingPrefix.Empty();
+	}
+	else
+	{
+		LoggingPrefix = InPrefix + FString(": ");
+	}
+}
+
+void FCameraRigAssetBuildLog::SetForwardMessagesToLogging(bool bInForwardToLogging)
+{
+	bForwardToLogging = bInForwardToLogging;
+}
+
+void FCameraRigAssetBuildLog::AddMessage(EMessageSeverity::Type InSeverity, FText&& InText)
+{
+	AddMessage(InSeverity, nullptr, MoveTemp(InText));
+}
+
+void FCameraRigAssetBuildLog::AddMessage(EMessageSeverity::Type InSeverity, UObject* InObject, FText&& InText)
+{
+	Messages.Add(FCameraRigAssetBuildLogMessage{ InSeverity, InObject, MoveTemp(InText) });
+
+	if (bForwardToLogging)
+	{
+		const FCameraRigAssetBuildLogMessage& LastMessage = Messages.Last();
+		LastMessage.SendToLogging(LoggingPrefix);
+	}
+}
+
+FCameraRigAssetBuilder::FCameraRigAssetBuilder(FCameraRigAssetBuildLog& InBuildLog)
+	: BuildLog(InBuildLog)
+{
+}
+
 void FCameraRigAssetBuilder::BuildCameraRig(UCameraRigAsset* InCameraRig)
 {
 	if (!ensure(InCameraRig))
@@ -200,9 +284,11 @@ void FCameraRigAssetBuilder::BuildCameraRig(UCameraRigAsset* InCameraRig)
 	CameraRig = InCameraRig;
 	bHasErrors = false;
 	bHasWarnings = false;
+	BuildLog.SetLoggingPrefix(InCameraRig->GetPathName() + TEXT(": "));
 	{
 		BuildCameraRigImpl();
 	}
+	BuildLog.SetLoggingPrefix(FString());
 	UpdateBuildStatus();
 }
 
@@ -210,7 +296,7 @@ void FCameraRigAssetBuilder::BuildCameraRigImpl()
 {
 	if (!CameraRig->RootNode)
 	{
-		UE_LOG(LogCameraSystem, Error, TEXT("Camera rig '%s' has no root node."), *CameraRig->GetPathName());
+		BuildLog.AddMessage(EMessageSeverity::Error, LOCTEXT("MissingRootNode", "Camera rig has no root node set."));
 		bHasErrors = true;
 		return;
 	}
@@ -324,27 +410,40 @@ void FCameraRigAssetBuilder::BuildNewDrivenParameters()
 	for (UCameraRigInterfaceParameter* InterfaceParameter : CameraRig->Interface.InterfaceParameters)
 	{
 		// Do some basic validation.
-		if (!InterfaceParameter || !InterfaceParameter->Target)
+		if (!InterfaceParameter)
 		{
-			UE_LOG(LogCameraSystem, Error, 
-					TEXT("Invalid interface parameter target in camera rig: '%s'."),
-					*CameraRigPathName);
+			BuildLog.AddMessage(EMessageSeverity::Error,
+					LOCTEXT("InvalidInterfaceParameter", "Invalid interface parameter or target."));
 			bHasErrors = true;
 			continue;
 		}
+		if (!InterfaceParameter->Target)
+		{
+			BuildLog.AddMessage(EMessageSeverity::Warning,
+					InterfaceParameter,
+					LOCTEXT(
+						"DisconnectedInterfaceParameter", 
+						"Interface parameter isn't connected: setting overrides for it will not do anything."));
+			bHasWarnings = true;
+			continue;
+		}
 		if (InterfaceParameter->TargetPropertyName.IsNone())
 		{
-			UE_LOG(LogCameraSystem, Error, 
-					TEXT("Invalid interface parameter target property name in camera rig: '%s'."), 
-					*CameraRigPathName);
+			BuildLog.AddMessage(EMessageSeverity::Error,
+					InterfaceParameter,
+					LOCTEXT(
+						"InvalidInterfaceParameterTargetPropertyName", 
+						"Invalid interface parameter target property name."));
 			bHasErrors = true;
 			continue;
 		}
 		if (InterfaceParameter->InterfaceParameterName.IsEmpty())
 		{
-			UE_LOG(LogCameraSystem, Error, 
-					TEXT("Invalid interface parameter name in camera rig '%s'."), 
-					*CameraRigPathName);
+			BuildLog.AddMessage(EMessageSeverity::Error,
+					InterfaceParameter,
+					LOCTEXT(
+						"InvalidInterfaceParameterName",
+						"Invalid interface parameter name."));
 			bHasErrors = true;
 			continue;
 		}
@@ -352,10 +451,12 @@ void FCameraRigAssetBuilder::BuildNewDrivenParameters()
 		// Check duplicate parameter names.
 		if (UsedInterfaceParameterNames.Contains(InterfaceParameter->InterfaceParameterName))
 		{
-			UE_LOG(LogCameraSystem, Error, 
-					TEXT("Multiple interface parameters named '%s' in camera rig '%s'. Ignoring duplicates."), 
-					*InterfaceParameter->InterfaceParameterName, 
-					*CameraRigPathName);
+			BuildLog.AddMessage(EMessageSeverity::Error,
+					InterfaceParameter,
+					FText::Format(LOCTEXT(
+						"InterfaceParameterNameCollision",
+						"Multiple interface parameters named '{0}'. Ignoring duplicates."),
+						FText::FromString(InterfaceParameter->InterfaceParameterName)));
 			bHasErrors = true;
 			continue;
 		}
@@ -367,11 +468,14 @@ void FCameraRigAssetBuilder::BuildNewDrivenParameters()
 		FProperty* TargetProperty = TargetClass->FindPropertyByName(InterfaceParameter->TargetPropertyName);
 		if (!TargetProperty)
 		{
-			UE_LOG(LogCameraSystem, Error, 
-					TEXT("Invalid interface parameter in camera rig '%s': no property '%s' on camera node '%s'."), 
-					*CameraRigPathName, 
-					*InterfaceParameter->InterfaceParameterName, 
-					*Target->GetName());
+			BuildLog.AddMessage(EMessageSeverity::Error,
+					Target,
+					FText::Format(LOCTEXT(
+						"InvalidInterfaceParameterTargetProperty",
+						"Invalid interface parameter '{0}', driving property '{1}' on '{2}', but no such property found."),
+						FText::FromString(InterfaceParameter->InterfaceParameterName), 
+						FText::FromName(InterfaceParameter->TargetPropertyName),
+						FText::FromName(Target->GetFName())));
 			bHasErrors = true;
 			continue;
 		}
@@ -379,11 +483,14 @@ void FCameraRigAssetBuilder::BuildNewDrivenParameters()
 		FStructProperty* TargetStructProperty = CastField<FStructProperty>(TargetProperty);
 		if (!TargetStructProperty)
 		{
-			UE_LOG(LogCameraSystem, Error, 
-					TEXT("Invalid interface parameter in camera rig '%s': property '%s' on camera node '%s' is not a camera parameter."),
-					*CameraRigPathName, 
-					*InterfaceParameter->InterfaceParameterName, 
-					*Target->GetName());
+			BuildLog.AddMessage(EMessageSeverity::Error,
+					Target,
+					FText::Format(LOCTEXT(
+						"InvalidCameraNodeProperty",
+						"Invalid interface parameter '{0}', driving property '{1}' on '{2}', but it's not a camera parameter."),
+						FText::FromString(InterfaceParameter->InterfaceParameterName), 
+						FText::FromName(InterfaceParameter->TargetPropertyName),
+						FText::FromName(Target->GetFName())));
 			bHasErrors = true;
 			continue;
 		}
@@ -401,11 +508,13 @@ void FCameraRigAssetBuilder::BuildNewDrivenParameters()
 UE_CAMERA_VARIABLE_FOR_ALL_TYPES()
 #undef UE_CAMERA_VARIABLE_FOR_TYPE
 		{
-			UE_LOG(LogCameraSystem, Error, 
-					TEXT("Invalid interface parameter in camera rig '%s': property '%s' on camera node '%s' is not a camera parameter."), 
-					*CameraRigPathName, 
-					*InterfaceParameter->InterfaceParameterName,
-					*Target->GetName());
+			BuildLog.AddMessage(EMessageSeverity::Error,
+					FText::Format(LOCTEXT(
+						"InvalidCameraNodeProperty",
+						"Invalid interface parameter '{0}', driving property '{1}' on '{2}', but it's not a camera parameter."),
+						FText::FromString(InterfaceParameter->InterfaceParameterName), 
+						FText::FromName(InterfaceParameter->TargetPropertyName),
+						FText::FromName(Target->GetFName())));
 			bHasErrors = true;
 			continue;
 		}
@@ -515,3 +624,5 @@ void FCameraRigAssetBuilder::UpdateBuildStatus()
 
 }  // namespace UE::Cameras
 
+#undef LOCTEXT_NAMESPACE
+
diff --git a/Engine/Plugins/Cameras/GameplayCameras/Source/GameplayCameras/Public/Core/CameraRigAsset.h b/Engine/Plugins/Cameras/GameplayCameras/Source/GameplayCameras/Public/Core/CameraRigAsset.h
--- a/Engine/Plugins/Cameras/GameplayCameras/Source/GameplayCameras/Public/Core/CameraRigAsset.h
+++ b/Engine/Plugins/Cameras/GameplayCameras/Source/GameplayCameras/Public/Core/CameraRigAsset.h
@@ -17,7 +17,11 @@
 class UCameraNode;
 class UCameraVariableAsset;
 
-namespace UE::Cameras { class FCameraRigAssetBuilder; }
+namespace UE::Cameras
+{
+	class FCameraRigAssetBuilder;
+	class FCameraRigAssetBuildLog;
+}
 
 /**
  *
@@ -191,6 +195,11 @@ public:
 	 */
 	GAMEPLAYCAMERAS_API void BuildCameraRig();
 
+	/**
+	 * Builds this camera rig, similar to BuildCameraRig() but using a given build log.
+	 */
+	GAMEPLAYCAMERAS_API void BuildCameraRig(UE::Cameras::FCameraRigAssetBuildLog& InBuildLog);
+
 public:
 
 	// IGameplayTagAssetInterface.
diff --git a/Engine/Plugins/Cameras/GameplayCameras/Source/GameplayCameras/Public/Core/CameraRigAssetBuilder.h b/Engine/Plugins/Cameras/GameplayCameras/Source/GameplayCameras/Public/Core/CameraRigAssetBuilder.h
--- a/Engine/Plugins/Cameras/GameplayCameras/Source/GameplayCameras/Public/Core/CameraRigAssetBuilder.h
+++ b/Engine/Plugins/Cameras/GameplayCameras/Source/GameplayCameras/Public/Core/CameraRigAssetBuilder.h
@@ -7,6 +7,7 @@
 #include "Core/CameraRigAsset.h"
 #include "Core/CameraVariableTableFwd.h"
 #include "CoreTypes.h"
+#include "Logging/TokenizedMessage.h"
 #include "Templates/Tuple.h"
 #include "UObject/WeakObjectPtr.h"
 
@@ -21,6 +22,59 @@ namespace UE::Cameras
 
 namespace Internal { struct FPrivateVariableBuilder; }
 
+/**
+ * A message emitted by the camera rig asset builder.
+ */
+struct GAMEPLAYCAMERAS_API FCameraRigAssetBuildLogMessage
+{
+	/** Severity of the message. */
+	EMessageSeverity::Type Severity = EMessageSeverity::Info;
+	/** An optional object that the message relates to. */
+	UObject* Object = nullptr;
+	/** The actual message. */
+	FText Text;
+
+	/** Generates a plain string representation of this message. */
+	FString ToString() const;
+	/** Sends a string version of this message to the LogCameraSystem console log. */
+	void SendToLogging(const FString& InLoggingPrefix) const;
+};
+
+/**
+ * Build log, populated when building a camera rig asset.
+ */
+class GAMEPLAYCAMERAS_API FCameraRigAssetBuildLog
+{
+public:
+
+	/**
+	 * Sets a string that will be prefixed to all messages sent to the console.
+	 * Only useful when IsForwardingMessagesToLogging is true.
+	 * This is generally set to the name of the camera rig asset being built.
+	 */
+	void SetLoggingPrefix(const FString& InPrefix);
+
+	/** Returns whether build messages are sent to the console. */
+	bool IsForwardingMessagesToLogging() const { return bForwardToLogging; }
+	/** Sets whether build messages are sent to the console. */
+	void SetForwardMessagesToLogging(bool bInForwardToLogging);
+
+	/** Adds a new message. */
+	void AddMessage(EMessageSeverity::Type InSeverity, FText&& InText);
+	/** Adds a new message. */
+	void AddMessage(EMessageSeverity::Type InSeverity, UObject* InObject, FText&& InText);
+
+	/** Gets the list of received messages so far. */
+	TArrayView<const FCameraRigAssetBuildLogMessage> GetMessages() const { return Messages; }
+
+private:
+
+	TArray<FCameraRigAssetBuildLogMessage> Messages;
+
+	FString LoggingPrefix;
+	bool bForwardToLogging = true;
+};
+
 /**
  * A class that can prepare a camera rig for runtime use.
  *
@@ -34,6 +88,9 @@ class FCameraRigAssetBuilder
 {
 public:
 
+	/** Creates a new camera rig builder. */
+	FCameraRigAssetBuilder(FCameraRigAssetBuildLog& InBuildLog);
+
 	/** Builds the given camera rig. */
 	void BuildCameraRig(UCameraRigAsset* InCameraRig);
 
@@ -54,6 +111,8 @@ private:
 
 private:
 
+	FCameraRigAssetBuildLog& BuildLog;
+
 	UCameraRigAsset* CameraRig = nullptr;
 
 	TArray<UCameraNode*> FlattenedNodes;
diff --git a/Engine/Plugins/Cameras/GameplayCameras/Source/GameplayCamerasEditor/Private/Editors/ObjectTreeGraph.cpp b/Engine/Plugins/Cameras/GameplayCameras/Source/GameplayCamerasEditor/Private/Editors/ObjectTreeGraph.cpp
--- a/Engine/Plugins/Cameras/GameplayCameras/Source/GameplayCamerasEditor/Private/Editors/ObjectTreeGraph.cpp
+++ b/Engine/Plugins/Cameras/GameplayCameras/Source/GameplayCamerasEditor/Private/Editors/ObjectTreeGraph.cpp
@@ -28,6 +28,21 @@ void UObjectTreeGraph::Initialize(TObjectPtr<UObject> InRootObject, const FObjec
 	}
 }
 
+UObjectTreeGraphNode* UObjectTreeGraph::FindObjectNode(UObject* InObject) const
+{
+	for (UEdGraphNode* Node : Nodes)
+	{
+		if (UObjectTreeGraphNode* ObjectNode = Cast<UObjectTreeGraphNode>(Node))
+		{
+			if (ObjectNode->GetObject() == InObject)
+			{
+				return ObjectNode;
+			}
+		}
+	}
+	return nullptr;
+}
+
 const FObjectTreeGraphConfig& UObjectTreeGraph::GetConfig() const
 {
 	return Config;
diff --git a/Engine/Plugins/Cameras/GameplayCameras/Source/GameplayCamerasEditor/Private/Editors/SCameraRigAssetEditor.cpp b/Engine/Plugins/Cameras/GameplayCameras/Source/GameplayCamerasEditor/Private/Editors/SCameraRigAssetEditor.cpp
--- a/Engine/Plugins/Cameras/GameplayCameras/Source/GameplayCamerasEditor/Private/Editors/SCameraRigAssetEditor.cpp
+++ b/Engine/Plugins/Cameras/GameplayCameras/Source/GameplayCamerasEditor/Private/Editors/SCameraRigAssetEditor.cpp
@@ -255,6 +255,23 @@ void SCameraRigAssetEditor::JumpToNode(UEdGraphNode* InGraphNode)
 	}
 }
 
+bool SCameraRigAssetEditor::FindAndJumpToObjectNode(UObject* InObject)
+{
+	if (UObjectTreeGraphNode* NodeGraphObjectNode = NodeGraph->FindObjectNode(InObject))
+	{
+		SetEditorMode(ECameraRigAssetEditorMode::NodeGraph);
+		NodeGraphEditor->JumpToNode(NodeGraphObjectNode);
+		return true;
+	}
+	if (UObjectTreeGraphNode* TransitionGraphObjectNode = TransitionGraph->FindObjectNode(InObject))
+	{
+		SetEditorMode(ECameraRigAssetEditorMode::TransitionGraph);
+		TransitionGraphEditor->JumpToNode(TransitionGraphObjectNode);
+		return true;
+	}
+	return false;
+}
+
 FText SCameraRigAssetEditor::GetCameraRigAssetName() const
 {
 	return FText::FromString(CameraRigAsset->GetName());
diff --git a/Engine/Plugins/Cameras/GameplayCameras/Source/GameplayCamerasEditor/Private/Toolkits/CameraRigAssetEditorToolkit.cpp b/Engine/Plugins/Cameras/GameplayCameras/Source/GameplayCamerasEditor/Private/Toolkits/CameraRigAssetEditorToolkit.cpp
--- a/Engine/Plugins/Cameras/GameplayCameras/Source/GameplayCamerasEditor/Private/Toolkits/CameraRigAssetEditorToolkit.cpp
+++ b/Engine/Plugins/Cameras/GameplayCameras/Source/GameplayCamerasEditor/Private/Toolkits/CameraRigAssetEditorToolkit.cpp
@@ -6,6 +6,7 @@
 #include "Commands/CameraRigAssetEditorCommands.h"
 #include "Core/CameraNode.h"
 #include "Core/CameraRigAsset.h"
+#include "Core/CameraRigAssetBuilder.h"
 #include "EditorModeManager.h"
 #include "Editors/ObjectTreeGraph.h"
 #include "Editors/ObjectTreeGraphConfig.h"
@@ -20,6 +21,7 @@
 #include "IMessageLogListing.h"
 #include "MessageLogInitializationOptions.h"
 #include "MessageLogModule.h"
+#include "Misc/UObjectToken.h"
 #include "Modules/ModuleManager.h"
 #include "PropertyEditorModule.h"
 #include "Styles/GameplayCamerasEditorStyle.h"
@@ -217,6 +219,7 @@ void FCameraRigAssetEditorToolkit::CreateWidgets()
 	LogOptions.bAllowClear = false;
 	LogOptions.MaxPageCount = 1;
 	MessageListing = MessageLogModule.CreateLogListing("CameraRigAssetEditorStats", LogOptions);
+	MessageListing->OnMessageTokenClicked().AddSP(this, &FCameraRigAssetEditorToolkit::OnMessageTokenClicked);
 
 	MessagesWidget = MessageLogModule.CreateLogListingWidget(MessageListing.ToSharedRef());
 }
@@ -380,11 +383,30 @@ void FCameraRigAssetEditorToolkit::NotifyPostChange(const FPropertyChangedEvent&
 	CameraRigAsset->BuildStatus = ECameraRigBuildStatus::Dirty;
 }
 
+void FCameraRigAssetEditorToolkit::OnMessageTokenClicked(const TSharedRef<IMessageToken>& InMessageToken)
+{
+	if (InMessageToken->GetType() == EMessageToken::Object)
+	{
+		const TSharedRef<FUObjectToken> ObjectToken = StaticCastSharedRef<FUObjectToken>(InMessageToken);
+		if (UObject* Object = ObjectToken->GetObject().Get())
+		{
+			CameraRigEditorWidget->FindAndJumpToObjectNode(Object);
+		}
+	}
+}
+
 void FCameraRigAssetEditorToolkit::OnBuild()
 {
 	using namespace UE::Cameras;
 
-	CameraRigAsset->BuildCameraRig();
+	FCameraRigAssetBuildLog BuildLog;
+	CameraRigAsset->BuildCameraRig(BuildLog);
+	GetBuildMessageLogListing(BuildLog);
+
+	if (CameraRigAsset->BuildStatus != ECameraRigBuildStatus::Clean)
+	{
+		TabManager->TryInvokeTab(MessagesTabId);
+	}
 
 	FCameraRigPackages BuiltPackages;
 	CameraRigAsset->GatherPackages(BuiltPackages);
@@ -395,6 +417,37 @@ void FCameraRigAssetEditorToolkit::OnBuild()
 	}
 }
 
+void FCameraRigAssetEditorToolkit::GetBuildMessageLogListing(FCameraRigAssetBuildLog& InBuildLog)
+{
+	if (MessageListing.IsValid())
+	{
+		MessageListing->ClearMessages();
+	}
+
+	for (const FCameraRigAssetBuildLogMessage& Message : InBuildLog.GetMessages())
+	{
+		TSharedRef<FTokenizedMessage> TokenizedMessage = FTokenizedMessage::Create(Message.Severity);
+
+		if (Message.Object)
+		{
+			TSharedRef<FUObjectToken> ObjectToken = FUObjectToken::Create(
+					Message.Object, FText::FromName(Message.Object->GetFName()));
+
+			// Suppress default activation callback that opens the content browser. We handle that
+			// sort of stuff in a more general fashion in OnMessageTokenClicked.
+			static auto DummyActivation = [](const TSharedRef<class IMessageToken>&) {};
+			ObjectToken->OnMessageTokenActivated(FOnMessageTokenActivated::CreateLambda(DummyActivation));
+			
+			TokenizedMessage->AddToken(ObjectToken);
+		}
+
+		TokenizedMessage->AddToken(FTextToken::Create(Message.Text));
+
+		MessageListing->AddMessage(TokenizedMessage);
+	}
+}
+
+
 void FCameraRigAssetEditorToolkit::OnFindInCameraRig()
 {
 	TabManager->TryInvokeTab(SearchTabId);
diff --git a/Engine/Plugins/Cameras/GameplayCameras/Source/GameplayCamerasEditor/Public/Editors/ObjectTreeGraph.h b/Engine/Plugins/Cameras/GameplayCameras/Source/GameplayCamerasEditor/Public/Editors/ObjectTreeGraph.h
--- a/Engine/Plugins/Cameras/GameplayCameras/Source/GameplayCamerasEditor/Public/Editors/ObjectTreeGraph.h
+++ b/Engine/Plugins/Cameras/GameplayCameras/Source/GameplayCamerasEditor/Public/Editors/ObjectTreeGraph.h
@@ -37,6 +37,10 @@ public:
 	UObject* GetRootObject() const { return WeakRootObject.Get(); }
 	/** Gets the root object's graph node. */
 	UObjectTreeGraphNode* GetRootObjectNode() const { return RootObjectNode; }
+
+	/** Finds a node for the given object. */
+	UObjectTreeGraphNode* FindObjectNode(UObject* InObject) const;
+
 	/** Gets the graph configuration. */
 	const FObjectTreeGraphConfig& GetConfig() const;
 
diff --git a/Engine/Plugins/Cameras/GameplayCameras/Source/GameplayCamerasEditor/Public/Editors/SCameraRigAssetEditor.h b/Engine/Plugins/Cameras/GameplayCameras/Source/GameplayCamerasEditor/Public/Editors/SCameraRigAssetEditor.h
--- a/Engine/Plugins/Cameras/GameplayCameras/Source/GameplayCamerasEditor/Public/Editors/SCameraRigAssetEditor.h
+++ b/Engine/Plugins/Cameras/GameplayCameras/Source/GameplayCamerasEditor/Public/Editors/SCameraRigAssetEditor.h
@@ -73,6 +73,8 @@ public:
 	void FocusHome();
 	/** Jumps the current graph to the given node. */
 	void JumpToNode(UEdGraphNode* InGraphNode);
+	/** Finds a node for the given object and, if so, jumps to it. */
+	bool FindAndJumpToObjectNode(UObject* InObject);
 
 protected:
 
diff --git a/Engine/Plugins/Cameras/GameplayCameras/Source/GameplayCamerasEditor/Public/Toolkits/CameraRigAssetEditorToolkit.h b/Engine/Plugins/Cameras/GameplayCameras/Source/GameplayCamerasEditor/Public/Toolkits/CameraRigAssetEditorToolkit.h
--- a/Engine/Plugins/Cameras/GameplayCameras/Source/GameplayCamerasEditor/Public/Toolkits/CameraRigAssetEditorToolkit.h
+++ b/Engine/Plugins/Cameras/GameplayCameras/Source/GameplayCamerasEditor/Public/Toolkits/CameraRigAssetEditorToolkit.h
@@ -21,6 +21,7 @@ class UEdGraphNode;
 namespace UE::Cameras
 {
 
+class FCameraRigAssetBuildLog;
 class IGameplayCamerasLiveEditManager;
 class SCameraRigAssetEditor;
 
@@ -76,6 +77,9 @@ private:
 	FSlateIcon GetBuildButtonIcon() const;
 	FText GetBuildButtonTooltip() const;
 
+	void OnMessageTokenClicked(const TSharedRef<IMessageToken>& InMessageToken);
+	void GetBuildMessageLogListing(FCameraRigAssetBuildLog& InBuildLog);
+
 	void OnBuild();
 	void OnFindInCameraRig();
 
