Cameras: dirty the build status of a camera rig when nodes are moved, pasted, deleted, etc.

diff --git a/Engine/Plugins/Cameras/GameplayCameras/Source/GameplayCamerasEditor/Private/Editors/SCameraRigAssetEditor.cpp b/Engine/Plugins/Cameras/GameplayCameras/Source/GameplayCamerasEditor/Private/Editors/SCameraRigAssetEditor.cpp
--- a/Engine/Plugins/Cameras/GameplayCameras/Source/GameplayCamerasEditor/Private/Editors/SCameraRigAssetEditor.cpp
+++ b/Engine/Plugins/Cameras/GameplayCameras/Source/GameplayCamerasEditor/Private/Editors/SCameraRigAssetEditor.cpp
@@ -48,10 +48,22 @@ SCameraRigAssetEditor::~SCameraRigAssetEditor()
 		if (NodeGraph)
 		{
 			NodeGraph->RemoveFromRoot();
+
+			if (NodeGraphChangedHandle.IsValid())
+			{
+				NodeGraph->RemoveOnGraphChangedHandler(NodeGraphChangedHandle);
+				NodeGraphChangedHandle.Reset();
+			}
 		}
 		if (TransitionGraph)
 		{
 			TransitionGraph->RemoveFromRoot();
+
+			if (TransitionGraphChangedHandle.IsValid())
+			{
+				TransitionGraph->RemoveOnGraphChangedHandler(TransitionGraphChangedHandle);
+				TransitionGraphChangedHandle.Reset();
+			}
 		}
 	}
 }
@@ -87,6 +99,9 @@ void SCameraRigAssetEditor::CreateNodeGraphEditor(const FArguments& InArgs)
 	NodeGraph->Initialize(CameraRigAsset, GraphConfig);
 	NodeGraph->RebuildGraph(EObjectTreeGraphBuildSource::RootObjectPackage);
 
+	NodeGraphChangedHandle = NodeGraph->AddOnGraphChangedHandler(
+			FOnGraphChanged::FDelegate::CreateSP(this, &SCameraRigAssetEditor::OnGraphChanged));
+
 	FGraphAppearanceInfo Appearance;
 	Appearance.CornerText = LOCTEXT("CameraRigGraphText", "CAMERA RIG");
 
@@ -133,6 +148,9 @@ void SCameraRigAssetEditor::CreateTransitionGraphEditor(const FArguments& InArgs
 	TransitionGraph->Initialize(CameraRigAsset, GraphConfig);
 	TransitionGraph->RebuildGraph(EObjectTreeGraphBuildSource::RootObjectPackage);
 
+	TransitionGraphChangedHandle = TransitionGraph->AddOnGraphChangedHandler(
+			FOnGraphChanged::FDelegate::CreateSP(this, &SCameraRigAssetEditor::OnGraphChanged));
+
 	FGraphAppearanceInfo Appearance;
 	Appearance.CornerText = LOCTEXT("TransitionGraphText", "TRANSITIONS");
 
@@ -277,6 +295,29 @@ FText SCameraRigAssetEditor::GetCameraRigAssetName() const
 	return FText::FromString(CameraRigAsset->GetName());
 }
 
+void SCameraRigAssetEditor::OnGraphChanged(const FEdGraphEditAction& InEditAction)
+{
+	OnAnyGraphChanged.Broadcast(InEditAction);
+}
+
+FDelegateHandle SCameraRigAssetEditor::AddOnAnyGraphChanged(FOnGraphChanged::FDelegate InAddDelegate)
+{
+	return OnAnyGraphChanged.Add(InAddDelegate);
+}
+
+void SCameraRigAssetEditor::RemoveOnAnyGraphChanged(FDelegateHandle InDelegateHandle)
+{
+	if (InDelegateHandle.IsValid())
+	{
+		OnAnyGraphChanged.Remove(InDelegateHandle);
+	}
+}
+
+void SCameraRigAssetEditor::RemoveOnAnyGraphChanged(const void* InUserObject)
+{
+	OnAnyGraphChanged.RemoveAll(InUserObject);
+}
+
 }  // namespace UE::Cameras
 
 #undef LOCTEXT_NAMESPACE
diff --git a/Engine/Plugins/Cameras/GameplayCameras/Source/GameplayCamerasEditor/Private/Toolkits/CameraRigAssetEditorToolkit.cpp b/Engine/Plugins/Cameras/GameplayCameras/Source/GameplayCamerasEditor/Private/Toolkits/CameraRigAssetEditorToolkit.cpp
--- a/Engine/Plugins/Cameras/GameplayCameras/Source/GameplayCamerasEditor/Private/Toolkits/CameraRigAssetEditorToolkit.cpp
+++ b/Engine/Plugins/Cameras/GameplayCameras/Source/GameplayCamerasEditor/Private/Toolkits/CameraRigAssetEditorToolkit.cpp
@@ -85,6 +85,10 @@ FCameraRigAssetEditorToolkit::FCameraRigAssetEditorToolkit(UCameraRigAssetEditor
 
 FCameraRigAssetEditorToolkit::~FCameraRigAssetEditorToolkit()
 {
+	if (CameraRigEditorWidget)
+	{
+		CameraRigEditorWidget->RemoveOnAnyGraphChanged(this);
+	}
 }
 
 void FCameraRigAssetEditorToolkit::RegisterTabSpawners(const TSharedRef<FTabManager>& InTabManager)
@@ -197,6 +201,8 @@ void FCameraRigAssetEditorToolkit::CreateWidgets()
 	CameraRigEditorWidget = SNew(SCameraRigAssetEditor)
 		.DetailsView(DetailsView)
 		.CameraRigAsset(CameraRigAsset);
+	CameraRigEditorWidget->AddOnAnyGraphChanged(FOnGraphChanged::FDelegate::CreateSP(
+				this, &FCameraRigAssetEditorToolkit::OnAnyGraphChanged));
 
 	// Create the toolbox, default to the rig editor items.
 	ToolboxWidget = SNew(SObjectTreeGraphToolbox)
@@ -383,6 +389,11 @@ void FCameraRigAssetEditorToolkit::NotifyPostChange(const FPropertyChangedEvent&
 	CameraRigAsset->BuildStatus = ECameraRigBuildStatus::Dirty;
 }
 
+void FCameraRigAssetEditorToolkit::OnAnyGraphChanged(const FEdGraphEditAction& InEditAction)
+{
+	CameraRigAsset->BuildStatus = ECameraRigBuildStatus::Dirty;
+}
+
 void FCameraRigAssetEditorToolkit::OnMessageTokenClicked(const TSharedRef<IMessageToken>& InMessageToken)
 {
 	if (InMessageToken->GetType() == EMessageToken::Object)
diff --git a/Engine/Plugins/Cameras/GameplayCameras/Source/GameplayCamerasEditor/Public/Editors/SCameraRigAssetEditor.h b/Engine/Plugins/Cameras/GameplayCameras/Source/GameplayCamerasEditor/Public/Editors/SCameraRigAssetEditor.h
--- a/Engine/Plugins/Cameras/GameplayCameras/Source/GameplayCamerasEditor/Public/Editors/SCameraRigAssetEditor.h
+++ b/Engine/Plugins/Cameras/GameplayCameras/Source/GameplayCamerasEditor/Public/Editors/SCameraRigAssetEditor.h
@@ -76,11 +76,20 @@ public:
 	/** Finds a node for the given object and, if so, jumps to it. */
 	bool FindAndJumpToObjectNode(UObject* InObject);
 
+	/** Adds a callback that will be invoked when a graph editor is changed. */
+	FDelegateHandle AddOnAnyGraphChanged(FOnGraphChanged::FDelegate InAddDelegate);
+	/** Removes a previous added callback. */
+	void RemoveOnAnyGraphChanged(FDelegateHandle InDelegateHandle);
+	/** Removes a previous added callback. */
+	void RemoveOnAnyGraphChanged(const void* InUserObject);
+
 protected:
 
 	void CreateNodeGraphEditor(const FArguments& InArgs);
 	void CreateTransitionGraphEditor(const FArguments& InArgs);
 
+	void OnGraphChanged(const FEdGraphEditAction& InEditAction);
+
 	FText GetCameraRigAssetName() const;
 
 private:
@@ -103,6 +112,13 @@ private:
 
 	/** The mode for the currently shown graph editor */
 	ECameraRigAssetEditorMode CurrentMode;
+
+	/** Handles for listening to changes in the graphs editors */
+	FDelegateHandle NodeGraphChangedHandle;
+	FDelegateHandle TransitionGraphChangedHandle;
+
+	/* Forwarding delegate for changes in any of the graphs */
+	FOnGraphChanged OnAnyGraphChanged;
 };
 
 }  // namespace UE::Cameras
diff --git a/Engine/Plugins/Cameras/GameplayCameras/Source/GameplayCamerasEditor/Public/Toolkits/CameraRigAssetEditorToolkit.h b/Engine/Plugins/Cameras/GameplayCameras/Source/GameplayCamerasEditor/Public/Toolkits/CameraRigAssetEditorToolkit.h
--- a/Engine/Plugins/Cameras/GameplayCameras/Source/GameplayCamerasEditor/Public/Toolkits/CameraRigAssetEditorToolkit.h
+++ b/Engine/Plugins/Cameras/GameplayCameras/Source/GameplayCamerasEditor/Public/Toolkits/CameraRigAssetEditorToolkit.h
@@ -17,6 +17,7 @@ class SWidget;
 class UCameraRigAsset;
 class UCameraRigAssetEditor;
 class UEdGraphNode;
+struct FEdGraphEditAction;
 
 namespace UE::Cameras
 {
@@ -77,6 +78,8 @@ private:
 	FSlateIcon GetBuildButtonIcon() const;
 	FText GetBuildButtonTooltip() const;
 
+	void OnAnyGraphChanged(const FEdGraphEditAction& InEditAction);
+
 	void OnMessageTokenClicked(const TSharedRef<IMessageToken>& InMessageToken);
 	void GetBuildMessageLogListing(FCameraRigAssetBuildLog& InBuildLog);
 
