Cameras: fix localization text collision

diff --git a/Engine/Plugins/Cameras/GameplayCameras/Source/GameplayCamerasEditor/Private/Toolkits/CameraVariableCollectionEditorToolkit.cpp b/Engine/Plugins/Cameras/GameplayCameras/Source/GameplayCamerasEditor/Private/Toolkits/CameraVariableCollectionEditorToolkit.cpp
--- a/Engine/Plugins/Cameras/GameplayCameras/Source/GameplayCamerasEditor/Private/Toolkits/CameraVariableCollectionEditorToolkit.cpp
+++ b/Engine/Plugins/Cameras/GameplayCameras/Source/GameplayCamerasEditor/Private/Toolkits/CameraVariableCollectionEditorToolkit.cpp
@@ -517,7 +517,7 @@ void FCameraVariableCollectionEditorToolkit::OnDeleteVariable()
 	bool bPerformDelete = false;
 	{
 		TSharedRef<SWindow> DeleteVariableWindow = SNew(SWindow)
-			.Title(LOCTEXT("DeleteVariable", "Delete Variables"))
+			.Title(LOCTEXT("DeleteVariablesWindowTitle", "Delete Variables"))
 			.ClientSize(FVector2D(600, 700));
 
 		TSharedRef<SDeleteVariableDialog> DeleteVariableDialog = SNew(SDeleteVariableDialog)
