Cameras: add rigs on the camera asset, and add a camera asset builder

Camera assets now have a list of camera rigs they can work with (although the editor-side code isn't ready for that yet).
Building a camera asset means, for now, building all its camera rigs. This requires moving all the build log classes and build status API into a more generic place.
Some of this work, like the IHasCameraBuildStatus interface, is in preparation of editor work.

diff --git a/Engine/Plugins/Cameras/GameplayCameras/Source/GameplayCameras/Private/Core/CameraAsset.cpp b/Engine/Plugins/Cameras/GameplayCameras/Source/GameplayCameras/Private/Core/CameraAsset.cpp
--- a/Engine/Plugins/Cameras/GameplayCameras/Source/GameplayCameras/Private/Core/CameraAsset.cpp
+++ b/Engine/Plugins/Cameras/GameplayCameras/Source/GameplayCameras/Private/Core/CameraAsset.cpp
@@ -2,5 +2,30 @@
 
 #include "Core/CameraAsset.h"
 
+#include "Core/CameraAssetBuilder.h"
+#include "Core/CameraBuildLog.h"
+
 #include UE_INLINE_GENERATED_CPP_BY_NAME(CameraAsset)
 
+void UCameraAsset::BuildCamera()
+{
+	using namespace UE::Cameras;
+
+	FCameraBuildLog BuildLog;
+	BuildLog.SetForwardMessagesToLogging(true);
+	BuildCamera(BuildLog);
+}
+
+void UCameraAsset::BuildCamera(UE::Cameras::FCameraBuildLog& InBuildLog)
+{
+	using namespace UE::Cameras;
+
+	FCameraAssetBuilder Builder(InBuildLog);
+	Builder.BuildCamera(this);
+}
+
+void UCameraAsset::DirtyBuildStatus()
+{
+	BuildStatus = ECameraBuildStatus::Dirty;
+}
+
diff --git a/Engine/Plugins/Cameras/GameplayCameras/Source/GameplayCameras/Private/Core/CameraAssetBuilder.cpp b/Engine/Plugins/Cameras/GameplayCameras/Source/GameplayCameras/Private/Core/CameraAssetBuilder.cpp
new file mode 100644
--- /dev/null
+++ b/Engine/Plugins/Cameras/GameplayCameras/Source/GameplayCameras/Private/Core/CameraAssetBuilder.cpp
@@ -0,0 +1,87 @@
+// Copyright Epic Games, Inc. All Rights Reserved.
+
+#include "Core/CameraAssetBuilder.h"
+
+#include "Core/CameraNode.h"
+#include "Core/CameraNodeEvaluatorStorage.h"
+#include "Core/CameraParameters.h"
+#include "Core/CameraAsset.h"
+#include "Core/CameraRigAsset.h"
+#include "Core/CameraRigAssetBuilder.h"
+#include "Core/CameraVariableAssets.h"
+#include "Logging/TokenizedMessage.h"
+
+#define LOCTEXT_NAMESPACE "CameraAssetBuilder"
+
+namespace UE::Cameras
+{
+
+FCameraAssetBuilder::FCameraAssetBuilder(FCameraBuildLog& InBuildLog)
+	: BuildLog(InBuildLog)
+{
+}
+
+void FCameraAssetBuilder::BuildCamera(UCameraAsset* InCameraAsset)
+{
+	if (!ensure(InCameraAsset))
+	{
+		return;
+	}
+
+	CameraAsset = InCameraAsset;
+	bHasErrors = false;
+	bHasWarnings = false;
+	BuildLog.SetLoggingPrefix(InCameraAsset->GetPathName() + TEXT(": "));
+	{
+		BuildCameraImpl();
+	}
+	BuildLog.SetLoggingPrefix(FString());
+	UpdateBuildStatus();
+}
+
+void FCameraAssetBuilder::BuildCameraImpl()
+{
+	if (!CameraAsset->CameraDirector)
+	{
+		BuildLog.AddMessage(EMessageSeverity::Error, LOCTEXT("MissingDirector", "Camera has no director set."));
+		bHasErrors = true;
+	}
+
+	if (CameraAsset->CameraRigs.IsEmpty())
+	{
+		BuildLog.AddMessage(EMessageSeverity::Warning, LOCTEXT("MissingRigs", "Camera has no camera rigs defined."));
+		bHasWarnings = true;
+	}
+
+	for (UCameraRigAsset* CameraRig : CameraAsset->CameraRigs)
+	{
+		FCameraRigAssetBuilder CameraRigBuilder(BuildLog);
+		CameraRigBuilder.BuildCameraRig(CameraRig);
+		bHasErrors |= CameraRigBuilder.LastBuildHadErrors();
+		bHasWarnings |= CameraRigBuilder.LastBuildHadWarnings();
+	}
+}
+
+void FCameraAssetBuilder::UpdateBuildStatus()
+{
+	ECameraBuildStatus BuildStatus = ECameraBuildStatus::Clean;
+	if (bHasErrors)
+	{
+		BuildStatus = ECameraBuildStatus::WithErrors;
+	}
+	else if (bHasWarnings)
+	{
+		BuildStatus = ECameraBuildStatus::CleanWithWarnings;
+	}
+
+	if (CameraAsset->BuildStatus != BuildStatus)
+	{
+		CameraAsset->Modify();
+		CameraAsset->BuildStatus = BuildStatus;
+	}
+}
+
+}  // namespace UE::Cameras
+
+#undef LOCTEXT_NAMESPACE
+
diff --git a/Engine/Plugins/Cameras/GameplayCameras/Source/GameplayCameras/Private/Core/CameraBuildLog.cpp b/Engine/Plugins/Cameras/GameplayCameras/Source/GameplayCameras/Private/Core/CameraBuildLog.cpp
new file mode 100644
--- /dev/null
+++ b/Engine/Plugins/Cameras/GameplayCameras/Source/GameplayCameras/Private/Core/CameraBuildLog.cpp
@@ -0,0 +1,81 @@
+// Copyright Epic Games, Inc. All Rights Reserved.
+
+#include "Core/CameraBuildLog.h"
+
+#include "GameplayCameras.h"
+#include "Misc/CString.h"
+#include "Misc/UObjectToken.h"
+
+namespace UE::Cameras
+{
+
+FString FCameraBuildLogMessage::ToString() const
+{
+	TStringBuilder<256> StringBuilder;
+	if (Object)
+	{
+		StringBuilder.Append(Object->GetName());
+		StringBuilder.Append(TEXT(": "));
+	}
+	StringBuilder.Append(Text.ToString());
+	return StringBuilder.ToString();
+}
+
+void FCameraBuildLogMessage::SendToLogging(const FString& InLoggingPrefix) const
+{
+#define UE_LOG_FORWARD_CAMERA_RIG_BULID_LOG_MESSAGE(Verbosity)\
+	UE_LOG(LogCameraSystem, Verbosity, TEXT("%s%s"), *InLoggingPrefix, *ToString());
+
+	switch (Severity)
+	{
+	case EMessageSeverity::Error:
+		UE_LOG_FORWARD_CAMERA_RIG_BULID_LOG_MESSAGE(Error);
+		break;
+	case EMessageSeverity::PerformanceWarning:
+	case EMessageSeverity::Warning:
+		UE_LOG_FORWARD_CAMERA_RIG_BULID_LOG_MESSAGE(Warning);
+		break;
+	case EMessageSeverity::Info:
+	default:
+		UE_LOG_FORWARD_CAMERA_RIG_BULID_LOG_MESSAGE(Log);
+		break;
+	};
+
+#undef UE_LOG_FORWARD_CAMERA_RIG_BULID_LOG_MESSAGE
+}
+
+void FCameraBuildLog::SetLoggingPrefix(const FString& InPrefix)
+{
+	if (InPrefix.IsEmpty())
+	{
+		LoggingPrefix.Empty();
+	}
+	else
+	{
+		LoggingPrefix = InPrefix + FString(": ");
+	}
+}
+
+void FCameraBuildLog::SetForwardMessagesToLogging(bool bInForwardToLogging)
+{
+	bForwardToLogging = bInForwardToLogging;
+}
+
+void FCameraBuildLog::AddMessage(EMessageSeverity::Type InSeverity, FText&& InText)
+{
+	AddMessage(InSeverity, nullptr, MoveTemp(InText));
+}
+
+void FCameraBuildLog::AddMessage(EMessageSeverity::Type InSeverity, UObject* InObject, FText&& InText)
+{
+	Messages.Add(FCameraBuildLogMessage{ InSeverity, InObject, MoveTemp(InText) });
+
+	if (bForwardToLogging)
+	{
+		const FCameraBuildLogMessage& LastMessage = Messages.Last();
+		LastMessage.SendToLogging(LoggingPrefix);
+	}
+}
+
+}  // namespace UE::Cameras
+
diff --git a/Engine/Plugins/Cameras/GameplayCameras/Source/GameplayCameras/Private/Core/CameraRigAsset.cpp b/Engine/Plugins/Cameras/GameplayCameras/Source/GameplayCameras/Private/Core/CameraRigAsset.cpp
--- a/Engine/Plugins/Cameras/GameplayCameras/Source/GameplayCameras/Private/Core/CameraRigAsset.cpp
+++ b/Engine/Plugins/Cameras/GameplayCameras/Source/GameplayCameras/Private/Core/CameraRigAsset.cpp
@@ -2,6 +2,7 @@
 
 #include "Core/CameraRigAsset.h"
 
+#include "Core/CameraBuildLog.h"
 #include "Core/CameraNode.h"
 #include "Core/CameraParameters.h"
 #include "Core/CameraRigAssetBuilder.h"
@@ -53,12 +54,12 @@ void UCameraRigAsset::BuildCameraRig()
 {
 	using namespace UE::Cameras;
 
-	FCameraRigAssetBuildLog BuildLog;
+	FCameraBuildLog BuildLog;
 	BuildLog.SetForwardMessagesToLogging(true);
 	BuildCameraRig(BuildLog);
 }
 
-void UCameraRigAsset::BuildCameraRig(UE::Cameras::FCameraRigAssetBuildLog& InBuildLog)
+void UCameraRigAsset::BuildCameraRig(UE::Cameras::FCameraBuildLog& InBuildLog)
 {
 	using namespace UE::Cameras;
 
@@ -66,6 +67,11 @@ void UCameraRigAsset::BuildCameraRig(UE::Cameras::FCameraRigAssetBuildLog& InBui
 	Builder.BuildCameraRig(this);
 }
 
+void UCameraRigAsset::DirtyBuildStatus()
+{
+	BuildStatus = ECameraBuildStatus::Dirty;
+}
+
 void UCameraRigAsset::PreSave(FObjectPreSaveContext ObjectSaveContext)
 {
 #if WITH_EDITOR
diff --git a/Engine/Plugins/Cameras/GameplayCameras/Source/GameplayCameras/Private/Core/CameraRigAssetBuilder.cpp b/Engine/Plugins/Cameras/GameplayCameras/Source/GameplayCameras/Private/Core/CameraRigAssetBuilder.cpp
--- a/Engine/Plugins/Cameras/GameplayCameras/Source/GameplayCameras/Private/Core/CameraRigAssetBuilder.cpp
+++ b/Engine/Plugins/Cameras/GameplayCameras/Source/GameplayCameras/Private/Core/CameraRigAssetBuilder.cpp
@@ -7,8 +7,6 @@
 #include "Core/CameraParameters.h"
 #include "Core/CameraRigAsset.h"
 #include "Core/CameraVariableAssets.h"
-#include "Misc/CString.h"
-#include "Misc/UObjectToken.h"
 
 #define LOCTEXT_NAMESPACE "CameraRigAssetBuilder"
 
@@ -201,75 +199,7 @@ void SetupPrivateVariable(
 
 }  // namespace Internal
 
-FString FCameraRigAssetBuildLogMessage::ToString() const
-{
-	TStringBuilder<256> StringBuilder;
-	if (Object)
-	{
-		StringBuilder.Append(Object->GetName());
-		StringBuilder.Append(TEXT(": "));
-	}
-	StringBuilder.Append(Text.ToString());
-	return StringBuilder.ToString();
-}
-
-void FCameraRigAssetBuildLogMessage::SendToLogging(const FString& InLoggingPrefix) const
-{
-#define UE_LOG_FORWARD_CAMERA_RIG_BULID_LOG_MESSAGE(Verbosity)\
-	UE_LOG(LogCameraSystem, Verbosity, TEXT("%s%s"), *InLoggingPrefix, *ToString());
-
-	switch (Severity)
-	{
-	case EMessageSeverity::Error:
-		UE_LOG_FORWARD_CAMERA_RIG_BULID_LOG_MESSAGE(Error);
-		break;
-	case EMessageSeverity::PerformanceWarning:
-	case EMessageSeverity::Warning:
-		UE_LOG_FORWARD_CAMERA_RIG_BULID_LOG_MESSAGE(Warning);
-		break;
-	case EMessageSeverity::Info:
-	default:
-		UE_LOG_FORWARD_CAMERA_RIG_BULID_LOG_MESSAGE(Log);
-		break;
-	};
-
-#undef UE_LOG_FORWARD_CAMERA_RIG_BULID_LOG_MESSAGE
-}
-
-void FCameraRigAssetBuildLog::SetLoggingPrefix(const FString& InPrefix)
-{
-	if (InPrefix.IsEmpty())
-	{
-		LoggingPrefix.Empty();
-	}
-	else
-	{
-		LoggingPrefix = InPrefix + FString(": ");
-	}
-}
-
-void FCameraRigAssetBuildLog::SetForwardMessagesToLogging(bool bInForwardToLogging)
-{
-	bForwardToLogging = bInForwardToLogging;
-}
-
-void FCameraRigAssetBuildLog::AddMessage(EMessageSeverity::Type InSeverity, FText&& InText)
-{
-	AddMessage(InSeverity, nullptr, MoveTemp(InText));
-}
-
-void FCameraRigAssetBuildLog::AddMessage(EMessageSeverity::Type InSeverity, UObject* InObject, FText&& InText)
-{
-	Messages.Add(FCameraRigAssetBuildLogMessage{ InSeverity, InObject, MoveTemp(InText) });
-
-	if (bForwardToLogging)
-	{
-		const FCameraRigAssetBuildLogMessage& LastMessage = Messages.Last();
-		LastMessage.SendToLogging(LoggingPrefix);
-	}
-}
-
-FCameraRigAssetBuilder::FCameraRigAssetBuilder(FCameraRigAssetBuildLog& InBuildLog)
+FCameraRigAssetBuilder::FCameraRigAssetBuilder(FCameraBuildLog& InBuildLog)
 	: BuildLog(InBuildLog)
 {
 }
@@ -296,7 +226,7 @@ void FCameraRigAssetBuilder::BuildCameraRigImpl()
 {
 	if (!CameraRig->RootNode)
 	{
-		BuildLog.AddMessage(EMessageSeverity::Error, LOCTEXT("MissingRootNode", "Camera rig has no root node set."));
+		BuildLog.AddMessage(EMessageSeverity::Error, CameraRig, LOCTEXT("MissingRootNode", "Camera rig has no root node set."));
 		bHasErrors = true;
 		return;
 	}
@@ -413,6 +343,7 @@ void FCameraRigAssetBuilder::BuildNewDrivenParameters()
 		if (!InterfaceParameter)
 		{
 			BuildLog.AddMessage(EMessageSeverity::Error,
+					CameraRig,
 					LOCTEXT("InvalidInterfaceParameter", "Invalid interface parameter or target."));
 			bHasErrors = true;
 			continue;
@@ -509,6 +440,7 @@ UE_CAMERA_VARIABLE_FOR_ALL_TYPES()
 #undef UE_CAMERA_VARIABLE_FOR_TYPE
 		{
 			BuildLog.AddMessage(EMessageSeverity::Error,
+					InterfaceParameter,
 					FText::Format(LOCTEXT(
 						"InvalidCameraNodeProperty",
 						"Invalid interface parameter '{0}', driving property '{1}' on '{2}', but it's not a camera parameter."),
@@ -605,14 +537,14 @@ UE_CAMERA_VARIABLE_FOR_ALL_TYPES()
 
 void FCameraRigAssetBuilder::UpdateBuildStatus()
 {
-	ECameraRigBuildStatus BuildStatus = ECameraRigBuildStatus::Clean;
+	ECameraBuildStatus BuildStatus = ECameraBuildStatus::Clean;
 	if (bHasErrors)
 	{
-		BuildStatus = ECameraRigBuildStatus::WithErrors;
+		BuildStatus = ECameraBuildStatus::WithErrors;
 	}
 	else if (bHasWarnings)
 	{
-		BuildStatus = ECameraRigBuildStatus::CleanWithWarnings;
+		BuildStatus = ECameraBuildStatus::CleanWithWarnings;
 	}
 
 	if (CameraRig->BuildStatus != BuildStatus)
diff --git a/Engine/Plugins/Cameras/GameplayCameras/Source/GameplayCameras/Private/Tests/CameraRigAssetBuilderTests.cpp b/Engine/Plugins/Cameras/GameplayCameras/Source/GameplayCameras/Private/Tests/CameraRigAssetBuilderTests.cpp
--- a/Engine/Plugins/Cameras/GameplayCameras/Source/GameplayCameras/Private/Tests/CameraRigAssetBuilderTests.cpp
+++ b/Engine/Plugins/Cameras/GameplayCameras/Source/GameplayCameras/Private/Tests/CameraRigAssetBuilderTests.cpp
@@ -15,7 +15,7 @@ bool FCameraRigAssetBuilderNullTest::RunTest(const FString& Parameters)
 	using namespace UE::Cameras::Test;
 
 	UCameraRigAsset* CameraRig = FCameraRigAssetTestBuilder(TEXT("InvalidTest")).Get();
-	UTEST_EQUAL("Dirty status", CameraRig->BuildStatus, ECameraRigBuildStatus::Dirty);
+	UTEST_EQUAL("Dirty status", CameraRig->BuildStatus, ECameraBuildStatus::Dirty);
 
 	FStringFormatOrderedArguments ErrorArgs;
 	ErrorArgs.Add(CameraRig->GetPathName());
@@ -26,7 +26,7 @@ bool FCameraRigAssetBuilderNullTest::RunTest(const FString& Parameters)
 			1,
 			false);
 	CameraRig->BuildCameraRig();
-	UTEST_EQUAL("Error status", CameraRig->BuildStatus, ECameraRigBuildStatus::WithErrors);
+	UTEST_EQUAL("Error status", CameraRig->BuildStatus, ECameraBuildStatus::WithErrors);
 
 	return true;
 }
diff --git a/Engine/Plugins/Cameras/GameplayCameras/Source/GameplayCameras/Public/Core/CameraAsset.h b/Engine/Plugins/Cameras/GameplayCameras/Source/GameplayCameras/Public/Core/CameraAsset.h
--- a/Engine/Plugins/Cameras/GameplayCameras/Source/GameplayCameras/Public/Core/CameraAsset.h
+++ b/Engine/Plugins/Cameras/GameplayCameras/Source/GameplayCameras/Public/Core/CameraAsset.h
@@ -2,6 +2,7 @@
 
 #pragma once
 
+#include "Core/CameraBuildStatus.h"
 #include "Core/CameraRigTransition.h"
 #include "CoreTypes.h"
 #include "UObject/ObjectPtr.h"
@@ -9,28 +10,60 @@
 #include "CameraAsset.generated.h"
 
 class UCameraDirector;
+class UCameraRigAsset;
+
+namespace UE::Cameras { class FCameraBuildLog; }
 
 /**
  * A complete camera asset.
  */
 UCLASS(MinimalAPI)
-class UCameraAsset : public UObject
+class UCameraAsset 
+	: public UObject
+	, public IHasCameraBuildStatus
 {
 	GENERATED_BODY()
 
 public:
 
 	/** The camera director to use in this camera. */
-	UPROPERTY(Instanced, EditAnywhere, Category=Director)
+	UPROPERTY(Instanced)
 	TObjectPtr<UCameraDirector> CameraDirector;
 
+	/** The list of camera rigs used by this camera. */
+	UPROPERTY()
+	TArray<TObjectPtr<UCameraRigAsset>> CameraRigs;
+
 	/** A list of default enter transitions for all the camera rigs in this asset. */
-	UPROPERTY(EditAnywhere, Category=Blending)
+	UPROPERTY()
 	TArray<TObjectPtr<UCameraRigTransition>> EnterTransitions;
 
 	/** A list of default exit transitions for all the camera rigs in this asset. */
-	UPROPERTY(EditAnywhere, Category=Blending)
+	UPROPERTY()
 	TArray<TObjectPtr<UCameraRigTransition>> ExitTransitions;
-};
 
+public:
+
+	/** The current build state of this camera asset. */
+	UPROPERTY(Transient)
+	ECameraBuildStatus BuildStatus = ECameraBuildStatus::Dirty;
+
+	/**
+	 * Builds and validates this camera, including all its camera rigs.
+	 * Errors and warnings will go to the console.
+	 */
+	GAMEPLAYCAMERAS_API void BuildCamera();
+
+	/**
+	 * Builds and validates this camera, including all its camera rigs.
+	 * Errors and warnings will go to the provided build log.
+	 */
+	GAMEPLAYCAMERAS_API void BuildCamera(UE::Cameras::FCameraBuildLog& InBuildLog);
+
+public:
+
+	// IHasCameraBuildStatus interface.
+	virtual ECameraBuildStatus GetBuildStatus() const override { return BuildStatus; }
+	virtual void DirtyBuildStatus() override;
+};
 
diff --git a/Engine/Plugins/Cameras/GameplayCameras/Source/GameplayCameras/Public/Core/CameraAssetBuilder.h b/Engine/Plugins/Cameras/GameplayCameras/Source/GameplayCameras/Public/Core/CameraAssetBuilder.h
new file mode 100644
--- /dev/null
+++ b/Engine/Plugins/Cameras/GameplayCameras/Source/GameplayCameras/Public/Core/CameraAssetBuilder.h
@@ -0,0 +1,44 @@
+// Copyright Epic Games, Inc. All Rights Reserved.
+
+#pragma once
+
+#include "Core/CameraBuildLog.h"
+#include "CoreTypes.h"
+#include "Templates/Tuple.h"
+
+class UCameraAsset;
+
+namespace UE::Cameras
+{
+
+/**
+ * A class that can prepare a camera asset for runtime use.
+ */
+class FCameraAssetBuilder
+{
+public:
+
+	/** Creates a new camera builder. */
+	FCameraAssetBuilder(FCameraBuildLog& InBuildLog);
+
+	/** Builds the given camera. */
+	void BuildCamera(UCameraAsset* InCameraAsset);
+
+private:
+
+	void BuildCameraImpl();
+
+	void UpdateBuildStatus();
+
+private:
+
+	FCameraBuildLog& BuildLog;
+
+	UCameraAsset* CameraAsset = nullptr;
+
+	bool bHasErrors;
+	bool bHasWarnings;
+};
+
+}  // namespace UE::Cameras
+
diff --git a/Engine/Plugins/Cameras/GameplayCameras/Source/GameplayCameras/Public/Core/CameraBuildLog.h b/Engine/Plugins/Cameras/GameplayCameras/Source/GameplayCameras/Public/Core/CameraBuildLog.h
new file mode 100644
--- /dev/null
+++ b/Engine/Plugins/Cameras/GameplayCameras/Source/GameplayCameras/Public/Core/CameraBuildLog.h
@@ -0,0 +1,66 @@
+// Copyright Epic Games, Inc. All Rights Reserved.
+
+#pragma once
+
+#include "Containers/Array.h"
+#include "CoreTypes.h"
+#include "Logging/TokenizedMessage.h"
+
+namespace UE::Cameras
+{
+
+/**
+ * A message emitted by a camera build process.
+ */
+struct GAMEPLAYCAMERAS_API FCameraBuildLogMessage
+{
+	/** Severity of the message. */
+	EMessageSeverity::Type Severity = EMessageSeverity::Info;
+	/** An optional object that the message relates to. */
+	UObject* Object = nullptr;
+	/** The actual message. */
+	FText Text;
+
+	/** Generates a plain string representation of this message. */
+	FString ToString() const;
+	/** Sends a string version of this message to the LogCameraSystem console log. */
+	void SendToLogging(const FString& InLoggingPrefix) const;
+};
+
+/**
+ * Build log, populated when building a camera rig, or other camera asset.
+ */
+class GAMEPLAYCAMERAS_API FCameraBuildLog
+{
+public:
+
+	/**
+	 * Sets a string that will be prefixed to all messages sent to the console.
+	 * Only useful when IsForwardingMessagesToLogging is true.
+	 * This is generally set to the name of the camera asset being built.
+	 */
+	void SetLoggingPrefix(const FString& InPrefix);
+
+	/** Returns whether build messages are sent to the console. */
+	bool IsForwardingMessagesToLogging() const { return bForwardToLogging; }
+	/** Sets whether build messages are sent to the console. */
+	void SetForwardMessagesToLogging(bool bInForwardToLogging);
+
+	/** Adds a new message. */
+	void AddMessage(EMessageSeverity::Type InSeverity, FText&& InText);
+	/** Adds a new message. */
+	void AddMessage(EMessageSeverity::Type InSeverity, UObject* InObject, FText&& InText);
+
+	/** Gets the list of received messages so far. */
+	TArrayView<const FCameraBuildLogMessage> GetMessages() const { return Messages; }
+
+private:
+
+	TArray<FCameraBuildLogMessage> Messages;
+
+	FString LoggingPrefix;
+	bool bForwardToLogging = true;
+};
+
+}  // namespace UE::Cameras
+
diff --git a/Engine/Plugins/Cameras/GameplayCameras/Source/GameplayCameras/Public/Core/CameraBuildStatus.h b/Engine/Plugins/Cameras/GameplayCameras/Source/GameplayCameras/Public/Core/CameraBuildStatus.h
new file mode 100644
--- /dev/null
+++ b/Engine/Plugins/Cameras/GameplayCameras/Source/GameplayCameras/Public/Core/CameraBuildStatus.h
@@ -0,0 +1,44 @@
+// Copyright Epic Games, Inc. All Rights Reserved.
+
+#pragma once
+
+#include "CoreTypes.h"
+#include "UObject/Interface.h"
+
+#include "CameraBuildStatus.generated.h"
+
+/**
+ * Enumeration that describes if a camera asset needs to be rebuilt.
+ */
+UENUM()
+enum class ECameraBuildStatus : uint8
+{
+	Clean,
+	CleanWithWarnings,
+	WithErrors,
+	Dirty
+};
+
+UINTERFACE(MinimalAPI)
+class UHasCameraBuildStatus : public UInterface
+{
+	GENERATED_BODY()
+};
+
+/**
+ * Interface for an object that can be built and has a build status,
+ * such as a camera asset or camera rig.
+ */
+class IHasCameraBuildStatus
+{
+	GENERATED_BODY()
+
+public:
+
+	/** Gets the build status of the object. */
+	virtual ECameraBuildStatus GetBuildStatus() const PURE_VIRTUAL(IHasCameraBuildStatus::GetBuildStatus, return ECameraBuildStatus::Dirty;);
+
+	/** Dirties the build status of the object. */
+	virtual void DirtyBuildStatus() PURE_VIRTUAL();
+};
+
diff --git a/Engine/Plugins/Cameras/GameplayCameras/Source/GameplayCameras/Public/Core/CameraRigAsset.h b/Engine/Plugins/Cameras/GameplayCameras/Source/GameplayCameras/Public/Core/CameraRigAsset.h
--- a/Engine/Plugins/Cameras/GameplayCameras/Source/GameplayCameras/Public/Core/CameraRigAsset.h
+++ b/Engine/Plugins/Cameras/GameplayCameras/Source/GameplayCameras/Public/Core/CameraRigAsset.h
@@ -2,6 +2,7 @@
 
 #pragma once
 
+#include "Core/CameraBuildStatus.h"
 #include "Core/CameraNodeEvaluatorFwd.h"
 #include "Core/CameraRigTransition.h"
 #include "Core/CameraVariableTableFwd.h"
@@ -19,22 +20,10 @@ class UCameraVariableAsset;
 
 namespace UE::Cameras
 {
+	class FCameraBuildLog;
 	class FCameraRigAssetBuilder;
-	class FCameraRigAssetBuildLog;
 }
 
-/**
- *
- */
-UENUM()
-enum class ECameraRigBuildStatus : uint8
-{
-	Clean,
-	CleanWithWarnings,
-	WithErrors,
-	Dirty
-};
-
 /**
  * Structure describing various allocations needed by a camera node.
  */
@@ -145,6 +134,7 @@ UCLASS(MinimalAPI)
 class UCameraRigAsset
 	: public UObject
 	, public IGameplayTagAssetInterface
+	, public IHasCameraBuildStatus
 	, public IObjectTreeGraphObject
 	, public IObjectTreeGraphRootObject
 {
@@ -186,7 +176,7 @@ public:
 
 	/** The current build state of this camera rig. */
 	UPROPERTY(Transient)
-	ECameraRigBuildStatus BuildStatus = ECameraRigBuildStatus::Dirty;
+	ECameraBuildStatus BuildStatus = ECameraBuildStatus::Dirty;
 
 	/**
 	 * Builds this camera rig.
@@ -198,13 +188,17 @@ public:
 	/**
 	 * Builds this camera rig, similar to BuildCameraRig() but using a given build log.
 	 */
-	GAMEPLAYCAMERAS_API void BuildCameraRig(UE::Cameras::FCameraRigAssetBuildLog& InBuildLog);
+	GAMEPLAYCAMERAS_API void BuildCameraRig(UE::Cameras::FCameraBuildLog& InBuildLog);
 
 public:
 
-	// IGameplayTagAssetInterface.
+	// IGameplayTagAssetInterface interface.
 	virtual void GetOwnedGameplayTags(FGameplayTagContainer& TagContainer) const override;
 
+	// IHasCameraBuildStatus interface.
+	virtual ECameraBuildStatus GetBuildStatus() const override { return BuildStatus; }
+	virtual void DirtyBuildStatus() override;
+
 public:
 
 	// Graph names for ObjectTreeGraph API.
@@ -222,6 +216,7 @@ protected:
 	virtual void OnUpdateGraphNodeCommentText(const FString& NewComment) override;
 #endif
 
+	// IObjectTreeGraphRootObject interface.
 #if WITH_EDITOR
 	virtual void GetConnectableObjects(FName InGraphName, TSet<UObject*>& OutObjects) const override;
 	virtual void AddConnectableObject(FName InGraphName, UObject* InObject) override;
diff --git a/Engine/Plugins/Cameras/GameplayCameras/Source/GameplayCameras/Public/Core/CameraRigAssetBuilder.h b/Engine/Plugins/Cameras/GameplayCameras/Source/GameplayCameras/Public/Core/CameraRigAssetBuilder.h
--- a/Engine/Plugins/Cameras/GameplayCameras/Source/GameplayCameras/Public/Core/CameraRigAssetBuilder.h
+++ b/Engine/Plugins/Cameras/GameplayCameras/Source/GameplayCameras/Public/Core/CameraRigAssetBuilder.h
@@ -2,14 +2,12 @@
 
 #pragma once
 
-#include "Containers/Array.h"
+#include "Core/CameraBuildLog.h"
 #include "Core/CameraNodeEvaluatorFwd.h"
 #include "Core/CameraRigAsset.h"
 #include "Core/CameraVariableTableFwd.h"
 #include "CoreTypes.h"
-#include "Logging/TokenizedMessage.h"
 #include "Templates/Tuple.h"
-#include "UObject/WeakObjectPtr.h"
 
 class FStructProperty;
 class UCameraNode;
@@ -22,59 +20,6 @@ namespace UE::Cameras
 
 namespace Internal { struct FPrivateVariableBuilder; }
 
-/**
- * A message emitted by the camera rig asset builder.
- */
-struct GAMEPLAYCAMERAS_API FCameraRigAssetBuildLogMessage
-{
-	/** Severity of the message. */
-	EMessageSeverity::Type Severity = EMessageSeverity::Info;
-	/** An optional object that the message relates to. */
-	UObject* Object = nullptr;
-	/** The actual message. */
-	FText Text;
-
-	/** Generates a plain string representation of this message. */
-	FString ToString() const;
-	/** Sends a string version of this message to the LogCameraSystem console log. */
-	void SendToLogging(const FString& InLoggingPrefix) const;
-};
-
-/**
- * Build log, populated when building a camera rig asset.
- */
-class GAMEPLAYCAMERAS_API FCameraRigAssetBuildLog
-{
-public:
-
-	/**
-	 * Sets a string that will be prefixed to all messages sent to the console.
-	 * Only useful when IsForwardingMessagesToLogging is true.
-	 * This is generally set to the name of the camera rig asset being built.
-	 */
-	void SetLoggingPrefix(const FString& InPrefix);
-
-	/** Returns whether build messages are sent to the console. */
-	bool IsForwardingMessagesToLogging() const { return bForwardToLogging; }
-	/** Sets whether build messages are sent to the console. */
-	void SetForwardMessagesToLogging(bool bInForwardToLogging);
-
-	/** Adds a new message. */
-	void AddMessage(EMessageSeverity::Type InSeverity, FText&& InText);
-	/** Adds a new message. */
-	void AddMessage(EMessageSeverity::Type InSeverity, UObject* InObject, FText&& InText);
-
-	/** Gets the list of received messages so far. */
-	TArrayView<const FCameraRigAssetBuildLogMessage> GetMessages() const { return Messages; }
-
-private:
-
-	TArray<FCameraRigAssetBuildLogMessage> Messages;
-
-	FString LoggingPrefix;
-	bool bForwardToLogging = true;
-};
-
 /**
  * A class that can prepare a camera rig for runtime use.
  *
@@ -82,18 +27,21 @@ private:
  * rig parameters, computes the allocation information of the camera rig, and
  * does various kinds of validation.
  *
- * Once the build process is done, the BuildStatus proeprty is set on the camera rig.
+ * Once the build process is done, the BuildStatus property is set on the camera rig.
  */
 class FCameraRigAssetBuilder
 {
 public:
 
 	/** Creates a new camera rig builder. */
-	FCameraRigAssetBuilder(FCameraRigAssetBuildLog& InBuildLog);
+	FCameraRigAssetBuilder(FCameraBuildLog& InBuildLog);
 
 	/** Builds the given camera rig. */
 	void BuildCameraRig(UCameraRigAsset* InCameraRig);
 
+	bool LastBuildHadErrors() const { return bHasErrors; }
+	bool LastBuildHadWarnings() const { return bHasWarnings; }
+
 private:
 
 	void BuildCameraRigImpl();
@@ -111,7 +59,7 @@ private:
 
 private:
 
-	FCameraRigAssetBuildLog& BuildLog;
+	FCameraBuildLog& BuildLog;
 
 	UCameraRigAsset* CameraRig = nullptr;
 
diff --git a/Engine/Plugins/Cameras/GameplayCameras/Source/GameplayCamerasEditor/Private/Toolkits/CameraRigAssetEditorToolkit.cpp b/Engine/Plugins/Cameras/GameplayCameras/Source/GameplayCamerasEditor/Private/Toolkits/CameraRigAssetEditorToolkit.cpp
--- a/Engine/Plugins/Cameras/GameplayCameras/Source/GameplayCamerasEditor/Private/Toolkits/CameraRigAssetEditorToolkit.cpp
+++ b/Engine/Plugins/Cameras/GameplayCameras/Source/GameplayCamerasEditor/Private/Toolkits/CameraRigAssetEditorToolkit.cpp
@@ -298,13 +298,13 @@ FSlateIcon FCameraRigAssetEditorToolkit::GetBuildButtonIcon() const
 	switch (CameraRigAsset->BuildStatus)
 	{
 		default:
-		case ECameraRigBuildStatus::Dirty:
+		case ECameraBuildStatus::Dirty:
 			return FSlateIcon(CamerasStyleSetName, BuildStatusBackground, NAME_None, BuildStatusUnknown);
-		case ECameraRigBuildStatus::WithErrors:
+		case ECameraBuildStatus::WithErrors:
 			return FSlateIcon(CamerasStyleSetName, BuildStatusBackground, NAME_None, BuildStatusError);
-		case ECameraRigBuildStatus::Clean:
+		case ECameraBuildStatus::Clean:
 			return FSlateIcon(CamerasStyleSetName, BuildStatusBackground, NAME_None, BuildStatusGood);
-		case ECameraRigBuildStatus::CleanWithWarnings:
+		case ECameraBuildStatus::CleanWithWarnings:
 			return FSlateIcon(CamerasStyleSetName, BuildStatusBackground, NAME_None, BuildStatusWarning);
 	}
 }
@@ -319,13 +319,13 @@ FText FCameraRigAssetEditorToolkit::GetBuildButtonTooltip() const
 	switch (CameraRigAsset->BuildStatus)
 	{
 		default:
-		case ECameraRigBuildStatus::Dirty:
+		case ECameraBuildStatus::Dirty:
 			return LOCTEXT("BuildButtonStatusDirty", "Dirty or unknown, should rebuild");
-		case ECameraRigBuildStatus::WithErrors:
+		case ECameraBuildStatus::WithErrors:
 			return LOCTEXT("BuildButtonStatusWithErrors", "There were errors during the build, see the log window for details");
-		case ECameraRigBuildStatus::Clean:
+		case ECameraBuildStatus::Clean:
 			return LOCTEXT("BuildButtonStatusClean", "Good to go");
-		case ECameraRigBuildStatus::CleanWithWarnings:
+		case ECameraBuildStatus::CleanWithWarnings:
 			return LOCTEXT("BuildButtonStatusCleanWithWarnings", "There were warnings during the build, see the log window for details");
 	}
 }
@@ -387,12 +387,12 @@ FLinearColor FCameraRigAssetEditorToolkit::GetWorldCentricTabColorScale() const
 
 void FCameraRigAssetEditorToolkit::NotifyPostChange(const FPropertyChangedEvent& PropertyChangedEvent, FProperty* PropertyThatChanged)
 {
-	CameraRigAsset->BuildStatus = ECameraRigBuildStatus::Dirty;
+	CameraRigAsset->BuildStatus = ECameraBuildStatus::Dirty;
 }
 
 void FCameraRigAssetEditorToolkit::OnAnyGraphChanged(const FEdGraphEditAction& InEditAction)
 {
-	CameraRigAsset->BuildStatus = ECameraRigBuildStatus::Dirty;
+	CameraRigAsset->BuildStatus = ECameraBuildStatus::Dirty;
 }
 
 void FCameraRigAssetEditorToolkit::OnMessageTokenClicked(const TSharedRef<IMessageToken>& InMessageToken)
@@ -411,11 +411,11 @@ void FCameraRigAssetEditorToolkit::OnBuild()
 {
 	using namespace UE::Cameras;
 
-	FCameraRigAssetBuildLog BuildLog;
+	FCameraBuildLog BuildLog;
 	CameraRigAsset->BuildCameraRig(BuildLog);
 	GetBuildMessageLogListing(BuildLog);
 
-	if (CameraRigAsset->BuildStatus != ECameraRigBuildStatus::Clean)
+	if (CameraRigAsset->BuildStatus != ECameraBuildStatus::Clean)
 	{
 		TabManager->TryInvokeTab(MessagesTabId);
 	}
@@ -429,14 +429,14 @@ void FCameraRigAssetEditorToolkit::OnBuild()
 	}
 }
 
-void FCameraRigAssetEditorToolkit::GetBuildMessageLogListing(FCameraRigAssetBuildLog& InBuildLog)
+void FCameraRigAssetEditorToolkit::GetBuildMessageLogListing(FCameraBuildLog& InBuildLog)
 {
 	if (MessageListing.IsValid())
 	{
 		MessageListing->ClearMessages();
 	}
 
-	for (const FCameraRigAssetBuildLogMessage& Message : InBuildLog.GetMessages())
+	for (const FCameraBuildLogMessage& Message : InBuildLog.GetMessages())
 	{
 		TSharedRef<FTokenizedMessage> TokenizedMessage = FTokenizedMessage::Create(Message.Severity);
 
diff --git a/Engine/Plugins/Cameras/GameplayCameras/Source/GameplayCamerasEditor/Public/Toolkits/CameraRigAssetEditorToolkit.h b/Engine/Plugins/Cameras/GameplayCameras/Source/GameplayCamerasEditor/Public/Toolkits/CameraRigAssetEditorToolkit.h
--- a/Engine/Plugins/Cameras/GameplayCameras/Source/GameplayCamerasEditor/Public/Toolkits/CameraRigAssetEditorToolkit.h
+++ b/Engine/Plugins/Cameras/GameplayCameras/Source/GameplayCamerasEditor/Public/Toolkits/CameraRigAssetEditorToolkit.h
@@ -22,7 +22,7 @@ struct FEdGraphEditAction;
 namespace UE::Cameras
 {
 
-class FCameraRigAssetBuildLog;
+class FCameraBuildLog;
 class IGameplayCamerasLiveEditManager;
 class SCameraRigAssetEditor;
 
@@ -81,7 +81,7 @@ private:
 	void OnAnyGraphChanged(const FEdGraphEditAction& InEditAction);
 
 	void OnMessageTokenClicked(const TSharedRef<IMessageToken>& InMessageToken);
-	void GetBuildMessageLogListing(FCameraRigAssetBuildLog& InBuildLog);
+	void GetBuildMessageLogListing(FCameraBuildLog& InBuildLog);
 
 	void OnBuild();
 	void OnFindInCameraRig();
