Cameras: big UI change to implement the first pass of a camera asset editor

- Add the camera asset editor, which has 3 editor modes: director, rigs, transitions.
- To support this, implement a base utility toolkit class that can switch between editor modes. It's similar to FWorkflowCentricApplication in concept, but it has some implementation differences that make it work for our particular case.
- For the rigs and transitions modes, we don't want to reimplement similar logic to what was already in the standalone camera rig editor. So this change also refactors the camera rig editor in order to create a "toolkit base class" we can re-use in both the actual camera rig editor toolkit, and in the editor modes of the camera asset editor.
- Also for better code re-use, the build log and build button logic are now in toolkit utility classes.
- Camera rigs now have a display name.
- Added a new transition editor, used for now with the camera asset's shared transitions, but later useful for global transitions.
- Make the node graph and transition graph more friendly to not having any data, such as when no camera rig is open in the camera asset editor.
- Add more options to the object tree graph auto-collection of objects to solve a problem with camera asset shared transition graphs also getting loose objects from camera rigs' graphs.
- A few other fixes.

diff --git a/Engine/Plugins/Cameras/GameplayCameras/Source/GameplayCameras/Private/Core/CameraAsset.cpp b/Engine/Plugins/Cameras/GameplayCameras/Source/GameplayCameras/Private/Core/CameraAsset.cpp
--- a/Engine/Plugins/Cameras/GameplayCameras/Source/GameplayCameras/Private/Core/CameraAsset.cpp
+++ b/Engine/Plugins/Cameras/GameplayCameras/Source/GameplayCameras/Private/Core/CameraAsset.cpp
@@ -7,6 +7,8 @@
 
 #include UE_INLINE_GENERATED_CPP_BY_NAME(CameraAsset)
 
+const FName UCameraAsset::SharedTransitionsGraphName("SharedTransitions");
+
 void UCameraAsset::BuildCamera()
 {
 	using namespace UE::Cameras;
@@ -29,3 +31,50 @@ void UCameraAsset::DirtyBuildStatus()
 	BuildStatus = ECameraBuildStatus::Dirty;
 }
 
+#if WITH_EDITOR
+
+void UCameraAsset::GetGraphNodePosition(int32& NodePosX, int32& NodePosY) const
+{
+	NodePosX = GraphNodePosX;
+	NodePosY = GraphNodePosY;
+}
+
+void UCameraAsset::OnGraphNodeMoved(int32 NodePosX, int32 NodePosY, bool bMarkDirty)
+{
+	GraphNodePosX = NodePosX;
+	GraphNodePosY = NodePosY;
+}
+
+const FString& UCameraAsset::GetGraphNodeCommentText() const
+{
+	return GraphNodeComment;
+}
+
+void UCameraAsset::OnUpdateGraphNodeCommentText(const FString& NewComment)
+{
+	GraphNodeComment = NewComment;
+}
+
+void UCameraAsset::GetConnectableObjects(FName InGraphName, TSet<UObject*>& OutObjects) const
+{
+	OutObjects.Append(AllSharedTransitionsObjects);
+}
+
+void UCameraAsset::AddConnectableObject(FName InGraphName, UObject* InObject)
+{
+	Modify();
+
+	const int32 Index = AllSharedTransitionsObjects.AddUnique(InObject);
+	ensure(Index == AllSharedTransitionsObjects.Num() - 1);
+}
+
+void UCameraAsset::RemoveConnectableObject(FName InGraphName, UObject* InObject)
+{
+	Modify();
+
+	const int32 NumRemoved = AllSharedTransitionsObjects.Remove(InObject);
+	ensure(NumRemoved == 1);
+}
+
+#endif  // WITH_EDITOR
+
diff --git a/Engine/Plugins/Cameras/GameplayCameras/Source/GameplayCameras/Private/Core/CameraRigAsset.cpp b/Engine/Plugins/Cameras/GameplayCameras/Source/GameplayCameras/Private/Core/CameraRigAsset.cpp
--- a/Engine/Plugins/Cameras/GameplayCameras/Source/GameplayCameras/Private/Core/CameraRigAsset.cpp
+++ b/Engine/Plugins/Cameras/GameplayCameras/Source/GameplayCameras/Private/Core/CameraRigAsset.cpp
@@ -50,6 +50,15 @@ void UCameraRigAsset::GetOwnedGameplayTags(FGameplayTagContainer& TagContainer)
 	TagContainer.AppendTags(GameplayTags);
 }
 
+FString UCameraRigAsset::GetDisplayName() const
+{
+	if (!Interface.DisplayName.IsEmpty())
+	{
+		return Interface.DisplayName;
+	}
+	return GetName();
+}
+
 void UCameraRigAsset::BuildCameraRig()
 {
 	using namespace UE::Cameras;
@@ -175,3 +184,4 @@ void UCameraRigAsset::RemoveConnectableObject(FName InGraphName, UObject* InObje
 }
 
 #endif  // WITH_EDITOR
+
diff --git a/Engine/Plugins/Cameras/GameplayCameras/Source/GameplayCameras/Public/Core/CameraAsset.h b/Engine/Plugins/Cameras/GameplayCameras/Source/GameplayCameras/Public/Core/CameraAsset.h
--- a/Engine/Plugins/Cameras/GameplayCameras/Source/GameplayCameras/Public/Core/CameraAsset.h
+++ b/Engine/Plugins/Cameras/GameplayCameras/Source/GameplayCameras/Public/Core/CameraAsset.h
@@ -4,6 +4,8 @@
 
 #include "Core/CameraBuildStatus.h"
 #include "Core/CameraRigTransition.h"
+#include "Core/ObjectTreeGraphObject.h"
+#include "Core/ObjectTreeGraphRootObject.h"
 #include "CoreTypes.h"
 #include "UObject/ObjectPtr.h"
 
@@ -21,6 +23,8 @@ UCLASS(MinimalAPI)
 class UCameraAsset 
 	: public UObject
 	, public IHasCameraBuildStatus
+	, public IObjectTreeGraphObject
+	, public IObjectTreeGraphRootObject
 {
 	GENERATED_BODY()
 
@@ -62,8 +66,53 @@ public:
 
 public:
 
+	// Graph names for ObjectTreeGraph API.
+	GAMEPLAYCAMERAS_API static const FName SharedTransitionsGraphName;
+
 	// IHasCameraBuildStatus interface.
 	virtual ECameraBuildStatus GetBuildStatus() const override { return BuildStatus; }
 	virtual void DirtyBuildStatus() override;
+
+protected:
+
+	// IObjectTreeGraphObject interface.
+#if WITH_EDITOR
+	virtual void GetGraphNodePosition(int32& NodePosX, int32& NodePosY) const override;
+	virtual void OnGraphNodeMoved(int32 NodePosX, int32 NodePosY, bool bMarkDirty) override;
+	virtual EObjectTreeGraphObjectSupportFlags GetSupportFlags() const override { return EObjectTreeGraphObjectSupportFlags::CommentText; }
+	virtual const FString& GetGraphNodeCommentText() const override;
+	virtual void OnUpdateGraphNodeCommentText(const FString& NewComment) override;
+#endif
+
+	// IObjectTreeGraphRootObject interface.
+#if WITH_EDITOR
+	virtual void GetConnectableObjects(FName InGraphName, TSet<UObject*>& OutObjects) const override;
+	virtual void AddConnectableObject(FName InGraphName, UObject* InObject) override;
+	virtual void RemoveConnectableObject(FName InGraphName, UObject* InObject) override;
+#endif
+
+private:
+
+#if WITH_EDITORONLY_DATA
+
+	/** Position of the camera node in the transition editor. */
+	UPROPERTY()
+	int32 GraphNodePosX = 0;
+
+	/** Position of the camera node in the graph editor. */
+	UPROPERTY()
+	int32 GraphNodePosY = 0;
+
+	/** User-written comment in the graph editor. */
+	UPROPERTY()
+	FString GraphNodeComment;
+
+	/**
+	 * Similar to AllNodeTreeObjects, but for the transitions graph.
+	 */
+	UPROPERTY(Instanced, meta=(ObjectTreeGraphHidden=true))
+	TArray<TObjectPtr<UObject>> AllSharedTransitionsObjects;
+
+#endif  // WITH_EDITORONLY_DATA
 };
 
diff --git a/Engine/Plugins/Cameras/GameplayCameras/Source/GameplayCameras/Public/Core/CameraRigAsset.h b/Engine/Plugins/Cameras/GameplayCameras/Source/GameplayCameras/Public/Core/CameraRigAsset.h
--- a/Engine/Plugins/Cameras/GameplayCameras/Source/GameplayCameras/Public/Core/CameraRigAsset.h
+++ b/Engine/Plugins/Cameras/GameplayCameras/Source/GameplayCameras/Public/Core/CameraRigAsset.h
@@ -107,6 +107,9 @@ struct FCameraRigInterface
 
 public:
 
+	UPROPERTY()
+	FString DisplayName;
+
 	/** The list of exposed parameters on the camera rig. */
 	UPROPERTY(Instanced)
 	TArray<TObjectPtr<UCameraRigInterfaceParameter>> InterfaceParameters;
@@ -172,6 +175,12 @@ public:
 	UPROPERTY()
 	FCameraRigAllocationInfo AllocationInfo;
 
+	/**
+	 * Gets the display name of this camera rig.
+	 * This is either the display name set on the interface object, or its internal name.
+	 */
+	GAMEPLAYCAMERAS_API FString GetDisplayName() const;
+
 public:
 
 	/** The current build state of this camera rig. */
diff --git a/Engine/Plugins/Cameras/GameplayCameras/Source/GameplayCamerasEditor/Private/AssetTools/CameraRigAssetEditor.cpp b/Engine/Plugins/Cameras/GameplayCameras/Source/GameplayCamerasEditor/Private/AssetTools/CameraRigAssetEditor.cpp
--- a/Engine/Plugins/Cameras/GameplayCameras/Source/GameplayCamerasEditor/Private/AssetTools/CameraRigAssetEditor.cpp
+++ b/Engine/Plugins/Cameras/GameplayCameras/Source/GameplayCamerasEditor/Private/AssetTools/CameraRigAssetEditor.cpp
@@ -22,6 +22,9 @@ void UCameraRigAssetEditor::GetObjectsToEdit(TArray<UObject*>& InObjectsToEdit)
 
 TSharedPtr<FBaseAssetToolkit> UCameraRigAssetEditor::CreateToolkit()
 {
-	return MakeShared<UE::Cameras::FCameraRigAssetEditorToolkit>(this);
+	using namespace UE::Cameras;
+	TSharedPtr<FCameraRigAssetEditorToolkit> Toolkit = MakeShared<FCameraRigAssetEditorToolkit>(this);
+	Toolkit->SetCameraRigAsset(CameraRigAsset);
+	return Toolkit;
 }
 
diff --git a/Engine/Plugins/Cameras/GameplayCameras/Source/GameplayCamerasEditor/Private/AssetTools/CameraRigTransitionEditor.cpp b/Engine/Plugins/Cameras/GameplayCameras/Source/GameplayCamerasEditor/Private/AssetTools/CameraRigTransitionEditor.cpp
new file mode 100644
--- /dev/null
+++ b/Engine/Plugins/Cameras/GameplayCameras/Source/GameplayCamerasEditor/Private/AssetTools/CameraRigTransitionEditor.cpp
@@ -0,0 +1,27 @@
+// Copyright Epic Games, Inc. All Rights Reserved.
+
+#include "AssetTools/CameraRigTransitionEditor.h"
+
+#include "Core/CameraRigTransition.h"
+#include "EditorModeManager.h"
+#include "Toolkits/CameraRigTransitionEditorToolkit.h"
+
+#include UE_INLINE_GENERATED_CPP_BY_NAME(CameraRigTransitionEditor)
+
+void UCameraRigTransitionEditor::Initialize(TObjectPtr<UObject> InTransitionOwner)
+{
+	TransitionOwner = InTransitionOwner;
+
+	Super::Initialize();
+}
+
+void UCameraRigTransitionEditor::GetObjectsToEdit(TArray<UObject*>& InObjectsToEdit)
+{
+	InObjectsToEdit.Add(TransitionOwner.Get());
+}
+
+TSharedPtr<FBaseAssetToolkit> UCameraRigTransitionEditor::CreateToolkit()
+{
+	return MakeShared<UE::Cameras::FCameraRigTransitionEditorToolkit>(this);
+}
+
diff --git a/Engine/Plugins/Cameras/GameplayCameras/Source/GameplayCamerasEditor/Private/Commands/CameraAssetEditorCommands.cpp b/Engine/Plugins/Cameras/GameplayCameras/Source/GameplayCamerasEditor/Private/Commands/CameraAssetEditorCommands.cpp
--- a/Engine/Plugins/Cameras/GameplayCameras/Source/GameplayCamerasEditor/Private/Commands/CameraAssetEditorCommands.cpp
+++ b/Engine/Plugins/Cameras/GameplayCameras/Source/GameplayCamerasEditor/Private/Commands/CameraAssetEditorCommands.cpp
@@ -5,6 +5,7 @@
 #include "Framework/Commands/Commands.h"
 #include "Framework/Commands/InputChord.h"
 #include "Framework/Commands/UICommandInfo.h"
+#include "Styles/GameplayCamerasEditorStyle.h"
 #include "Styling/AppStyle.h"
 
 #define LOCTEXT_NAMESPACE "CameraAssetEditorCommands"
@@ -14,10 +15,10 @@ namespace UE::Cameras
 
 FCameraAssetEditorCommands::FCameraAssetEditorCommands()
 	: TCommands<FCameraAssetEditorCommands>(
-			"GameplayCameras_CameraAssetEditor",
-			NSLOCTEXT("Contexts", "GameplayCameras_CameraAssetEditor", "Camera Asset Editor"),
+			"CameraAssetEditor",
+			LOCTEXT("CameraAssetEditor", "Camera Asset Editor"),
 			NAME_None,
-			FAppStyle::GetAppStyleSetName()
+			FGameplayCamerasEditorStyle::Get()->GetStyleSetName()
 		)
 {
 }
@@ -26,6 +27,28 @@ void FCameraAssetEditorCommands::RegisterCommands()
 {
 	UI_COMMAND(Build, "Build", "Builds the asset and refreshes it in PIE",
 			EUserInterfaceActionType::Button, FInputChord(EKeys::F7));
+
+	UI_COMMAND(ShowCameraDirector, "Director", "Shows the camera director editor",
+			EUserInterfaceActionType::ToggleButton, FInputChord());
+	UI_COMMAND(ShowCameraRigs, "Rigs", "Shows the camera rigs",
+			EUserInterfaceActionType::ToggleButton, FInputChord());
+	UI_COMMAND(ShowSharedTransitions, "Shared Transitions", "Shows the shared transitions",
+			EUserInterfaceActionType::ToggleButton, FInputChord());
+
+	UI_COMMAND(EditCameraRig, "Edit Camera Rig", "Edits the selected camera rig",
+			EUserInterfaceActionType::Button, FInputChord());
+
+	UI_COMMAND(AddCameraRig, "Add Camera Rig", "Adds a new camera rig",
+			EUserInterfaceActionType::Button, FInputChord());
+	UI_COMMAND(DeleteCameraRig, "Delete Camera Rig", "Deletes the selected camera rig(s)",
+			EUserInterfaceActionType::Button, FInputChord());
+	UI_COMMAND(RenameCameraRig, "Rename Camera Rig", "Renames the selected camera rig",
+			EUserInterfaceActionType::Button, FInputChord());
+
+	UI_COMMAND(ShowMessages, "Messages", "Shows the message log for this camera asset",
+			EUserInterfaceActionType::ToggleButton, FInputChord());
+	UI_COMMAND(FindInCamera, "Search", "Searches for nodes in this camera asset",
+			EUserInterfaceActionType::Button, FInputChord());
 }
 
 }  // namespace UE::Cameras
diff --git a/Engine/Plugins/Cameras/GameplayCameras/Source/GameplayCamerasEditor/Private/Commands/CameraRigTransitionEditorCommands.cpp b/Engine/Plugins/Cameras/GameplayCameras/Source/GameplayCamerasEditor/Private/Commands/CameraRigTransitionEditorCommands.cpp
new file mode 100644
--- /dev/null
+++ b/Engine/Plugins/Cameras/GameplayCameras/Source/GameplayCamerasEditor/Private/Commands/CameraRigTransitionEditorCommands.cpp
@@ -0,0 +1,38 @@
+// Copyright Epic Games, Inc. All Rights Reserved.
+
+#include "Commands/CameraRigTransitionEditorCommands.h"
+
+#include "Framework/Commands/Commands.h"
+#include "Framework/Commands/InputChord.h"
+#include "Framework/Commands/UICommandInfo.h"
+#include "GenericPlatform/GenericApplication.h"
+#include "Styles/GameplayCamerasEditorStyle.h"
+#include "Styling/AppStyle.h"
+
+#define LOCTEXT_NAMESPACE "CameraRigTransitionEditorCommands"
+
+namespace UE::Cameras
+{
+
+FCameraRigTransitionEditorCommands::FCameraRigTransitionEditorCommands()
+	: TCommands<FCameraRigTransitionEditorCommands>(
+			"CameraRigTransitionEditor",
+			LOCTEXT("CameraRigTransitionEditor", "Camera Rig Transition Editor"),
+			NAME_None,
+			FGameplayCamerasEditorStyle::Get()->GetStyleSetName()
+		)
+{
+}
+
+void FCameraRigTransitionEditorCommands::RegisterCommands()
+{
+	UI_COMMAND(FindInTransitions, "Search", "Searches for nodes in these transitions",
+			EUserInterfaceActionType::Button, FInputChord(EKeys::F, EModifierKey::Control));
+	UI_COMMAND(FocusHome, "Home", "Focuses the graph canvas back on the root node",
+			EUserInterfaceActionType::Button, FInputChord());
+}
+
+}  // namespace UE::Cameras
+
+#undef LOCTEXT_NAMESPACE
+
diff --git a/Engine/Plugins/Cameras/GameplayCameras/Source/GameplayCamerasEditor/Private/Editors/CameraNodeGraphSchema.cpp b/Engine/Plugins/Cameras/GameplayCameras/Source/GameplayCamerasEditor/Private/Editors/CameraNodeGraphSchema.cpp
--- a/Engine/Plugins/Cameras/GameplayCameras/Source/GameplayCamerasEditor/Private/Editors/CameraNodeGraphSchema.cpp
+++ b/Engine/Plugins/Cameras/GameplayCameras/Source/GameplayCamerasEditor/Private/Editors/CameraNodeGraphSchema.cpp
@@ -2,12 +2,16 @@
 
 #include "Editors/CameraNodeGraphSchema.h"
 
-#include "Core/CameraRigAsset.h"
+#include "Core/BlendCameraNode.h"
 #include "Core/CameraNode.h"
+#include "Core/CameraRigAsset.h"
 #include "EdGraph/EdGraphPin.h"
 #include "Editors/CameraNodeGraphNode.h"
 #include "Editors/CameraRigInterfaceParameterGraphNode.h"
 #include "Editors/ObjectTreeGraph.h"
+#include "Editors/ObjectTreeGraphConfig.h"
+#include "Editors/ObjectTreeGraphNode.h"
+#include "GameplayCamerasEditorSettings.h"
 
 #include "ScopedTransaction.h"
 
@@ -17,6 +21,44 @@
 
 const FName UCameraNodeGraphSchema::PC_CameraParameter("CameraParameter");
 
+FObjectTreeGraphConfig UCameraNodeGraphSchema::BuildGraphConfig()
+{
+	const UGameplayCamerasEditorSettings* Settings = GetDefault<UGameplayCamerasEditorSettings>();
+
+	FObjectTreeGraphConfig GraphConfig;
+	GraphConfig.GraphName = UCameraRigAsset::NodeTreeGraphName;
+	GraphConfig.ConnectableObjectClasses.Add(UCameraRigAsset::StaticClass());
+	GraphConfig.ConnectableObjectClasses.Add(UCameraNode::StaticClass());
+	GraphConfig.ConnectableObjectClasses.Add(UCameraRigInterfaceParameter::StaticClass());
+	GraphConfig.NonConnectableObjectClasses.Add(UBlendCameraNode::StaticClass());
+	GraphConfig.GraphDisplayInfo.PlainName = LOCTEXT("NodeGraphPlainName", "CameraNodes");
+	GraphConfig.GraphDisplayInfo.DisplayName = LOCTEXT("NodeGraphDisplayName", "Camera Nodes");
+	GraphConfig.ObjectClassConfigs.Emplace(UCameraRigAsset::StaticClass())
+		.OnlyAsRoot()
+		.HasSelfPin(false)
+		.NodeTitleUsesObjectName(true)
+		.NodeTitleColor(Settings->CameraRigAssetTitleColor);
+	GraphConfig.ObjectClassConfigs.Emplace(UCameraNode::StaticClass())
+		.StripDisplayNameSuffix(TEXT("Camera Node"))
+		.CreateCategoryMetaData(TEXT("CameraNodeCategories"))
+		.GraphNodeClass(UCameraNodeGraphNode::StaticClass());
+	GraphConfig.ObjectClassConfigs.Emplace(UCameraRigInterfaceParameter::StaticClass())
+		.SelfPinDirection(EGPD_Output)
+		.SelfPinName(NAME_None)  // No self pin name, we just want the title
+		.CanCreateNew(false)
+		.GraphNodeClass(UCameraRigInterfaceParameterGraphNode::StaticClass());
+	GraphConfig.OnFormatObjectDisplayName = FOnFormatObjectDisplayName::CreateLambda(
+			[](const UObject* Object, FText& InOutDisplayNameText)
+			{
+				if (const UCameraRigAsset* CameraRigAsset = Cast<UCameraRigAsset>(Object))
+				{
+					InOutDisplayNameText = FText::FromString(CameraRigAsset->GetDisplayName());
+				}
+			});
+
+	return GraphConfig;
+}
+
 void UCameraNodeGraphSchema::OnCreateAllNodes(UObjectTreeGraph* InGraph, const FCreatedNodes& InCreatedNodes) const
 {
 	Super::OnCreateAllNodes(InGraph, InCreatedNodes);
diff --git a/Engine/Plugins/Cameras/GameplayCameras/Source/GameplayCamerasEditor/Private/Editors/CameraRigTransitionGraphSchema.cpp b/Engine/Plugins/Cameras/GameplayCameras/Source/GameplayCamerasEditor/Private/Editors/CameraRigTransitionGraphSchema.cpp
new file mode 100644
--- /dev/null
+++ b/Engine/Plugins/Cameras/GameplayCameras/Source/GameplayCamerasEditor/Private/Editors/CameraRigTransitionGraphSchema.cpp
@@ -0,0 +1,138 @@
+// Copyright Epic Games, Inc. All Rights Reserved.
+
+#include "Editors/CameraRigTransitionGraphSchema.h"
+
+#include "Core/BlendCameraNode.h"
+#include "Core/CameraRigAsset.h"
+#include "EdGraph/EdGraphPin.h"
+#include "Editors/ObjectTreeGraphConfig.h"
+#include "Editors/ObjectTreeGraphNode.h"
+#include "GameplayCamerasEditorSettings.h"
+
+#include "ScopedTransaction.h"
+
+#include UE_INLINE_GENERATED_CPP_BY_NAME(CameraRigTransitionGraphSchema)
+
+#define LOCTEXT_NAMESPACE "CameraRigTransitionGraphSchema"
+
+FObjectTreeGraphConfig UCameraRigTransitionGraphSchema::BuildGraphConfig(const FCameraRigTransitionOwnerInfo& OwnerInfo)
+{
+	using namespace UE::Cameras;
+
+	const UGameplayCamerasEditorSettings* Settings = GetDefault<UGameplayCamerasEditorSettings>();
+
+	FObjectTreeGraphConfig GraphConfig;
+	GraphConfig.GraphName = OwnerInfo.GraphName;
+	GraphConfig.ConnectableObjectClasses.Add(OwnerInfo.TransitionOwnerClass);
+	GraphConfig.ConnectableObjectClasses.Add(UCameraRigTransition::StaticClass());
+	GraphConfig.ConnectableObjectClasses.Add(UCameraRigTransitionCondition::StaticClass());
+	GraphConfig.ConnectableObjectClasses.Add(UBlendCameraNode::StaticClass());
+	GraphConfig.StopAutoCollectAtObjectClasses.Add(UCameraRigAsset::StaticClass());
+	GraphConfig.GraphDisplayInfo.PlainName = FText::FromName(OwnerInfo.GraphName);
+	GraphConfig.GraphDisplayInfo.DisplayName = GraphConfig.GraphDisplayInfo.PlainName;
+	GraphConfig.ObjectClassConfigs.Emplace(OwnerInfo.TransitionOwnerClass)
+		.HasSelfPin(false)
+		.OnlyAsRoot()
+		.SetPropertyPinDirection(OwnerInfo.EnterTransitionsPropertyName, EGPD_Input)
+		.SetPropertyPinDirection(OwnerInfo.ExitTransitionsPropertyName, EGPD_Output)
+		.NodeTitleUsesObjectName(true)
+		.NodeTitleColor(Settings->CameraRigAssetTitleColor);
+	GraphConfig.ObjectClassConfigs.Emplace(UCameraRigTransition::StaticClass())
+		.SetPropertyPinDirection(GET_MEMBER_NAME_STRING_CHECKED(UCameraRigTransition, Conditions), EGPD_Input)
+		.NodeTitleColor(Settings->CameraRigTransitionTitleColor);
+	GraphConfig.ObjectClassConfigs.Emplace(UCameraRigTransitionCondition::StaticClass())
+		.SelfPinDirection(EGPD_Output)
+		.DefaultPropertyPinDirection(EGPD_Input)
+		.StripDisplayNameSuffix(TEXT("Transition Condition"))
+		.NodeTitleColor(Settings->CameraRigTransitionConditionTitleColor);
+	GraphConfig.ObjectClassConfigs.Emplace(UBlendCameraNode::StaticClass())
+		.StripDisplayNameSuffix(TEXT("Camera Node"))
+		.CreateCategoryMetaData(TEXT("CameraNodeCategories"));
+	GraphConfig.OnFormatObjectDisplayName = FOnFormatObjectDisplayName::CreateLambda(
+			[](const UObject* Object, FText& InOutDisplayNameText)
+			{
+				if (const UCameraRigAsset* CameraRigAsset = Cast<UCameraRigAsset>(Object))
+				{
+					InOutDisplayNameText = FText::FromString(CameraRigAsset->GetDisplayName());
+				}
+			});
+
+	return GraphConfig;
+}
+
+void UCameraRigTransitionGraphSchema::GetGraphContextActions(FGraphContextMenuBuilder& ContextMenuBuilder) const
+{
+	// Start with being able to create both enter and exit transitions.
+	bool bCanCreateEnterTransition = true;
+	bool bCanCreateExitTransition = true;
+	if (const UEdGraphPin* DraggedPin = ContextMenuBuilder.FromPin)
+	{
+		// If we are creating a node from dragging a pin into an empty space, figure out which transition
+		// we can create based on the direction of the dragged pin.
+		bCanCreateEnterTransition = false;
+		bCanCreateExitTransition = false;
+
+		UObjectTreeGraphNode* OwningNode = Cast<UObjectTreeGraphNode>(DraggedPin->GetOwningNode());
+		if (OwningNode)
+		{
+			FProperty* DraggedPinProperty = OwningNode->GetPropertyForPin(DraggedPin);
+			if (DraggedPinProperty && DraggedPinProperty->GetOwnerClass()->IsChildOf<UCameraRigAsset>())
+			{
+				bCanCreateEnterTransition = (DraggedPinProperty->GetFName() == GET_MEMBER_NAME_CHECKED(UCameraRigAsset, EnterTransitions));
+				bCanCreateExitTransition = (DraggedPinProperty->GetFName() == GET_MEMBER_NAME_CHECKED(UCameraRigAsset, ExitTransitions));
+			}
+		}
+	}
+
+	if (bCanCreateEnterTransition)
+	{
+		TSharedRef<FCameraRigTransitionGraphSchemaAction_NewTransitionNode> EnterAction = MakeShared<FCameraRigTransitionGraphSchemaAction_NewTransitionNode>(
+				LOCTEXT("TransitionsCategory", "Transitions"),
+				LOCTEXT("EnterTransition", "Enter Transition"),
+				LOCTEXT("EnterTransitionToolTip", "Creates a new enter transition"));
+		EnterAction->TransitionType = FCameraRigTransitionGraphSchemaAction_NewTransitionNode::ETransitionType::Enter;
+		ContextMenuBuilder.AddAction(StaticCastSharedPtr<FEdGraphSchemaAction>(EnterAction.ToSharedPtr()));
+	}
+	
+	if (bCanCreateExitTransition)
+	{
+		TSharedRef<FCameraRigTransitionGraphSchemaAction_NewTransitionNode> ExitAction = MakeShared<FCameraRigTransitionGraphSchemaAction_NewTransitionNode>(
+				LOCTEXT("TransitionsCategory", "Transitions"),
+				LOCTEXT("ExitTransition", "Exit Transition"),
+				LOCTEXT("ExitTransitionToolTip", "Creates a new exit transition"),
+				0);
+		ExitAction->TransitionType = FCameraRigTransitionGraphSchemaAction_NewTransitionNode::ETransitionType::Exit;
+		ContextMenuBuilder.AddAction(StaticCastSharedPtr<FEdGraphSchemaAction>(ExitAction.ToSharedPtr()));
+	}
+
+	Super::GetGraphContextActions(ContextMenuBuilder);
+}
+
+void UCameraRigTransitionGraphSchema::FilterGraphContextPlaceableClasses(TArray<UClass*>& InOutClasses) const
+{
+	InOutClasses.Remove(UCameraRigTransition::StaticClass());
+}
+
+FCameraRigTransitionGraphSchemaAction_NewTransitionNode::FCameraRigTransitionGraphSchemaAction_NewTransitionNode()
+{
+	ObjectClass = UCameraRigTransition::StaticClass();
+}
+
+FCameraRigTransitionGraphSchemaAction_NewTransitionNode::FCameraRigTransitionGraphSchemaAction_NewTransitionNode(FText InNodeCategory, FText InMenuDesc, FText InToolTip, const int32 InGrouping, FText InKeywords)
+	: FObjectGraphSchemaAction_NewNode(InNodeCategory, InMenuDesc, InToolTip, InGrouping, InKeywords)
+{
+	ObjectClass = UCameraRigTransition::StaticClass();
+}
+
+void FCameraRigTransitionGraphSchemaAction_NewTransitionNode::AutoSetupNewNode(UObjectTreeGraphNode* NewNode, UEdGraphPin* FromPin)
+{
+	if (TransitionType == ETransitionType::Enter)
+	{
+		NewNode->OverrideSelfPinDirection(EGPD_Output);
+	}
+
+	FObjectGraphSchemaAction_NewNode::AutoSetupNewNode(NewNode, FromPin);
+}
+
+#undef LOCTEXT_NAMESPACE
+
diff --git a/Engine/Plugins/Cameras/GameplayCameras/Source/GameplayCamerasEditor/Private/Editors/CameraTransitionGraphSchema.cpp b/Engine/Plugins/Cameras/GameplayCameras/Source/GameplayCamerasEditor/Private/Editors/CameraTransitionGraphSchema.cpp
deleted file mode 100644
diff --git a/Engine/Plugins/Cameras/GameplayCameras/Source/GameplayCamerasEditor/Private/Editors/ObjectTreeGraph.cpp b/Engine/Plugins/Cameras/GameplayCameras/Source/GameplayCamerasEditor/Private/Editors/ObjectTreeGraph.cpp
--- a/Engine/Plugins/Cameras/GameplayCameras/Source/GameplayCamerasEditor/Private/Editors/ObjectTreeGraph.cpp
+++ b/Engine/Plugins/Cameras/GameplayCameras/Source/GameplayCamerasEditor/Private/Editors/ObjectTreeGraph.cpp
@@ -16,16 +16,17 @@ UObjectTreeGraph::UObjectTreeGraph(const FObjectInitializer& ObjInit)
 	Schema = UObjectTreeGraphSchema::StaticClass();
 }
 
-void UObjectTreeGraph::Initialize(TObjectPtr<UObject> InRootObject, const FObjectTreeGraphConfig& InConfig)
+void UObjectTreeGraph::Reset(TObjectPtr<UObject> InRootObject, const FObjectTreeGraphConfig& InConfig, EObjectTreeGraphBuildSource InSource)
 {
 	WeakRootObject = InRootObject;
 
 	Config = InConfig;
-	ensureMsgf(Config.ConnectableObjectClasses.Num() > 0, TEXT("No connectable object classes specified... no graph can be created!"));
 	if (!Config.DefaultGraphNodeClass)
 	{
 		Config.DefaultGraphNodeClass = UObjectTreeGraphNode::StaticClass();
 	}
+
+	RebuildGraph(InSource);
 }
 
 UObjectTreeGraphNode* UObjectTreeGraph::FindObjectNode(UObject* InObject) const
diff --git a/Engine/Plugins/Cameras/GameplayCameras/Source/GameplayCamerasEditor/Private/Editors/ObjectTreeGraphSchema.cpp b/Engine/Plugins/Cameras/GameplayCameras/Source/GameplayCamerasEditor/Private/Editors/ObjectTreeGraphSchema.cpp
--- a/Engine/Plugins/Cameras/GameplayCameras/Source/GameplayCamerasEditor/Private/Editors/ObjectTreeGraphSchema.cpp
+++ b/Engine/Plugins/Cameras/GameplayCameras/Source/GameplayCamerasEditor/Private/Editors/ObjectTreeGraphSchema.cpp
@@ -46,6 +46,11 @@ struct FPackageReferenceCollector : public FArchiveUObject
 		ArShouldSkipBulkData = true;
 	}
 
+	void StopAtObjectClasses(TArray<UClass*> InStopAtClasses)
+	{
+		StopAtClasses = TSet<UClass*>(InStopAtClasses);
+	}
+
 	void CollectReferences()
 	{
 		ObjectsToVisit.Reset();
@@ -62,9 +67,22 @@ struct FPackageReferenceCollector : public FArchiveUObject
 
 private:
 
+	bool ShouldStopAt(UObject* Obj)
+	{
+		UClass* ObjClass = Obj->GetClass();
+		for (UClass* StopAtClass : StopAtClasses)
+		{
+			if (ObjClass->IsChildOf(StopAtClass))
+			{
+				return true;
+			}
+		}
+		return false;
+	}
+
 	virtual FArchive& operator<<(UObject*& ObjRef) override
 	{
-		if (ObjRef != nullptr && ObjRef->IsIn(PackageScope))
+		if (ObjRef != nullptr && ObjRef->IsIn(PackageScope) && !ShouldStopAt(ObjRef))
 		{
 			if (!VisitedObjects.Contains(ObjRef))
 			{
@@ -79,6 +97,7 @@ private:
 
 	UObject* RootObject;
 	UPackage* PackageScope;
+	TSet<UClass*> StopAtClasses;
 
 	TArray<UObject*> ObjectsToVisit;
 	TSet<UObject*> VisitedObjects;
@@ -135,15 +154,18 @@ void UObjectTreeGraphSchema::RemoveAllNodes(UObjectTreeGraph* InGraph) const
 void UObjectTreeGraphSchema::CreateAllNodes(UObjectTreeGraph* InGraph, EObjectTreeGraphBuildSource InSource) const
 {
 	UObject* RootObject = InGraph->GetRootObject();
-	if (!ensure(RootObject))
+	if (!RootObject)
 	{
 		return;
 	}
 
+	const FObjectTreeGraphConfig& GraphConfig = InGraph->GetConfig();
+
 	TSet<UObject*> AllObjects;
 	// Gather up all the objects we need for the graph. Start by all objects that are referenced (directly or
 	// indirectly) by the root object. Our custom reference collector will not collect references that go outside
 	// of the root object's package.
+	if (GraphConfig.bAutoCollectInitialObjects)
 	{
 		using namespace UE::ObjectTreeGraph;
 
@@ -152,14 +174,14 @@ void UObjectTreeGraphSchema::CreateAllNodes(UObjectTreeGraph* InGraph, EObjectTr
 
 		TArray<UObject*> ReferencedObjects;
 		FPackageReferenceCollector Collector(RootObject, ReferencedObjects);
+		Collector.StopAtObjectClasses(GraphConfig.StopAutoCollectAtObjectClasses);
 		Collector.CollectReferences();
 		AllObjects.Append(ReferencedObjects);
 	}
 	// Add any other custom objects the root object may want.
 	if (IObjectTreeGraphRootObject* RootObjectInterface = Cast<IObjectTreeGraphRootObject>(RootObject))
 	{
-		const FName GraphName = InGraph->GetConfig().GraphName;
-		RootObjectInterface->GetConnectableObjects(GraphName, AllObjects);
+		RootObjectInterface->GetConnectableObjects(GraphConfig.GraphName, AllObjects);
 	}
 	
 	// Create all the nodes.
@@ -174,10 +196,13 @@ void UObjectTreeGraphSchema::CreateAllNodes(UObjectTreeGraph* InGraph, EObjectTr
 
 	// Grab the graph node for the root object.
 	InGraph->RootObjectNode = nullptr;
-	UObjectTreeGraphNode** CreatedRootObjectNode = CreatedNodes.CreatedNodes.Find(RootObject);
-	if (ensure(CreatedRootObjectNode))
+	if (!AllObjects.IsEmpty())
 	{
-		InGraph->RootObjectNode = *CreatedRootObjectNode;
+		UObjectTreeGraphNode** CreatedRootObjectNode = CreatedNodes.CreatedNodes.Find(RootObject);
+		if (ensure(CreatedRootObjectNode))
+		{
+			InGraph->RootObjectNode = *CreatedRootObjectNode;
+		}
 	}
 
 	// Create all the connections.
@@ -782,6 +807,28 @@ void UObjectTreeGraphSchema::OnDeleteNodeFromGraph(UObjectTreeGraph* Graph, UEdG
 	}
 }
 
+void UObjectTreeGraphSchema::GetGraphDisplayInformation(const UEdGraph& Graph, FGraphDisplayInfo& OutDisplayInfo) const
+{
+	const UObjectTreeGraph* ObjectTreeGraph = CastChecked<const UObjectTreeGraph>(&Graph);
+	const FObjectTreeGraphConfig& GraphConfig = ObjectTreeGraph->GetConfig();
+
+	OutDisplayInfo = GraphConfig.GraphDisplayInfo;
+
+	if (OutDisplayInfo.PlainName.IsEmpty())
+	{
+		OutDisplayInfo.PlainName = FText::FromString(Graph.GetName());
+	}
+	if (OutDisplayInfo.DisplayName.IsEmpty())
+	{
+		OutDisplayInfo.DisplayName = OutDisplayInfo.PlainName;
+	}
+
+	if (GraphConfig.OnGetGraphDisplayInfo.IsBound())
+	{
+		GraphConfig.OnGetGraphDisplayInfo.Execute(ObjectTreeGraph, OutDisplayInfo);
+	}
+}
+
 FString UObjectTreeGraphSchema::ExportNodesToText(const FGraphPanelSelectionSet& Nodes, bool bOnlyCanDuplicateNodes, bool bOnlyCanDeleteNodes) const
 {
 	// Gather up the nodes we need to copy from.
diff --git a/Engine/Plugins/Cameras/GameplayCameras/Source/GameplayCamerasEditor/Private/Editors/SCameraRigAssetEditor.cpp b/Engine/Plugins/Cameras/GameplayCameras/Source/GameplayCamerasEditor/Private/Editors/SCameraRigAssetEditor.cpp
--- a/Engine/Plugins/Cameras/GameplayCameras/Source/GameplayCamerasEditor/Private/Editors/SCameraRigAssetEditor.cpp
+++ b/Engine/Plugins/Cameras/GameplayCameras/Source/GameplayCamerasEditor/Private/Editors/SCameraRigAssetEditor.cpp
@@ -2,19 +2,16 @@
 
 #include "Editors/SCameraRigAssetEditor.h"
 
-#include "Core/BlendCameraNode.h"
 #include "Core/CameraNode.h"
 #include "Core/CameraRigAsset.h"
 #include "Core/CameraRigTransition.h"
 #include "EdGraph/EdGraphPin.h"
-#include "Editors/CameraNodeGraphNode.h"
 #include "Editors/CameraNodeGraphSchema.h"
-#include "Editors/CameraRigInterfaceParameterGraphNode.h"
-#include "Editors/CameraTransitionGraphSchema.h"
+#include "Editors/CameraRigTransitionGraphSchema.h"
 #include "Editors/ObjectTreeGraph.h"
+#include "Editors/ObjectTreeGraphConfig.h"
 #include "Editors/ObjectTreeGraphNode.h"
 #include "Editors/SObjectTreeGraphEditor.h"
-#include "GameplayCamerasEditorSettings.h"
 #include "ObjectEditorUtils.h"
 #include "Widgets/Layout/SBox.h"
 
@@ -26,11 +23,12 @@ namespace UE::Cameras
 void SCameraRigAssetEditor::Construct(const FArguments& InArgs)
 {
 	CameraRigAsset = InArgs._CameraRigAsset;
+	DetailsView = InArgs._DetailsView;
+	AssetEditorToolkit = InArgs._AssetEditorToolkit;
 
 	CurrentMode = ECameraRigAssetEditorMode::NodeGraph;
 
-	CreateNodeGraphEditor(InArgs);
-	CreateTransitionGraphEditor(InArgs);
+	CreateGraphEditors();
 
 	ChildSlot
 	[
@@ -45,108 +43,68 @@ SCameraRigAssetEditor::~SCameraRigAssetEditor()
 {
 	if (!GExitPurge)
 	{
-		if (NodeGraph)
-		{
-			NodeGraph->RemoveFromRoot();
+		DiscardGraphEditors();
+	}
+}
 
-			if (NodeGraphChangedHandle.IsValid())
-			{
-				NodeGraph->RemoveOnGraphChangedHandler(NodeGraphChangedHandle);
-				NodeGraphChangedHandle.Reset();
-			}
-		}
-		if (TransitionGraph)
-		{
-			TransitionGraph->RemoveFromRoot();
+void SCameraRigAssetEditor::SetCameraRigAsset(UCameraRigAsset* InCameraRig)
+{
+	if (CameraRigAsset != InCameraRig)
+	{
+		DiscardGraphEditors();
 
-			if (TransitionGraphChangedHandle.IsValid())
-			{
-				TransitionGraph->RemoveOnGraphChangedHandler(TransitionGraphChangedHandle);
-				TransitionGraphChangedHandle.Reset();
-			}
-		}
+		CameraRigAsset = InCameraRig;
+
+		CreateGraphEditors();
+
+		SetEditorModeImpl(CurrentMode, true);
 	}
 }
 
-void SCameraRigAssetEditor::CreateNodeGraphEditor(const FArguments& InArgs)
+void SCameraRigAssetEditor::CreateGraphEditors()
 {
-	const UGameplayCamerasEditorSettings* Settings = GetDefault<UGameplayCamerasEditorSettings>();
-
-	FObjectTreeGraphConfig GraphConfig;
-	GraphConfig.GraphName = UCameraRigAsset::NodeTreeGraphName;
-	GraphConfig.ConnectableObjectClasses.Add(UCameraRigAsset::StaticClass());
-	GraphConfig.ConnectableObjectClasses.Add(UCameraNode::StaticClass());
-	GraphConfig.ConnectableObjectClasses.Add(UCameraRigInterfaceParameter::StaticClass());
-	GraphConfig.NonConnectableObjectClasses.Add(UBlendCameraNode::StaticClass());
-	GraphConfig.ObjectClassConfigs.Emplace(UCameraRigAsset::StaticClass())
-		.OnlyAsRoot()
-		.HasSelfPin(false)
-		.NodeTitleUsesObjectName(true)
-		.NodeTitleColor(Settings->CameraRigAssetTitleColor);
-	GraphConfig.ObjectClassConfigs.Emplace(UCameraNode::StaticClass())
-		.StripDisplayNameSuffix(TEXT("Camera Node"))
-		.CreateCategoryMetaData(TEXT("CameraNodeCategories"))
-		.GraphNodeClass(UCameraNodeGraphNode::StaticClass());
-	GraphConfig.ObjectClassConfigs.Emplace(UCameraRigInterfaceParameter::StaticClass())
-		.SelfPinDirection(EGPD_Output)
-		.SelfPinName(NAME_None)  // No self pin name, we just want the title
-		.CanCreateNew(false)
-		.GraphNodeClass(UCameraRigInterfaceParameterGraphNode::StaticClass());
+	CreateNodeGraphEditor();
+	CreateTransitionGraphEditor();
+}
+
+void SCameraRigAssetEditor::CreateNodeGraphEditor()
+{
+	FObjectTreeGraphConfig GraphConfig = UCameraNodeGraphSchema::BuildGraphConfig();
 
 	NodeGraph = NewObject<UObjectTreeGraph>(GetTransientPackage(), NAME_None, RF_Transactional);
 	NodeGraph->Schema = UCameraNodeGraphSchema::StaticClass();
 	NodeGraph->AddToRoot();
-	NodeGraph->Initialize(CameraRigAsset, GraphConfig);
-	NodeGraph->RebuildGraph(EObjectTreeGraphBuildSource::RootObjectPackage);
+	NodeGraph->Reset(CameraRigAsset, GraphConfig, EObjectTreeGraphBuildSource::RootObjectPackage);
 
 	NodeGraphChangedHandle = NodeGraph->AddOnGraphChangedHandler(
 			FOnGraphChanged::FDelegate::CreateSP(this, &SCameraRigAssetEditor::OnGraphChanged));
 
 	FGraphAppearanceInfo Appearance;
-	Appearance.CornerText = LOCTEXT("CameraRigGraphText", "CAMERA RIG");
+	Appearance.CornerText = LOCTEXT("CameraRigGraphText", "CAMERA NODES");
 
 	NodeGraphEditor = SNew(SObjectTreeGraphEditor)
 		.Appearance(Appearance)
-		.DetailsView(InArgs._DetailsView)
-		.GraphTitle(this, &SCameraRigAssetEditor::GetCameraRigAssetName)
+		.DetailsView(DetailsView)
+		.GraphTitle(this, &SCameraRigAssetEditor::GetCameraRigAssetName, NodeGraph.Get())
 		.GraphToEdit(NodeGraph)
-		.AssetEditorToolkit(InArgs._AssetEditorToolkit);
+		.AssetEditorToolkit(AssetEditorToolkit);
 }
 
-void SCameraRigAssetEditor::CreateTransitionGraphEditor(const FArguments& InArgs)
+void SCameraRigAssetEditor::CreateTransitionGraphEditor()
 {
-	const UGameplayCamerasEditorSettings* Settings = GetDefault<UGameplayCamerasEditorSettings>();
-
-	FObjectTreeGraphConfig GraphConfig;
-	GraphConfig.GraphName = UCameraRigAsset::TransitionsGraphName;
-	GraphConfig.ConnectableObjectClasses.Add(UCameraRigAsset::StaticClass());
-	GraphConfig.ConnectableObjectClasses.Add(UCameraRigTransition::StaticClass());
-	GraphConfig.ConnectableObjectClasses.Add(UCameraRigTransitionCondition::StaticClass());
-	GraphConfig.ConnectableObjectClasses.Add(UBlendCameraNode::StaticClass());
-	GraphConfig.ObjectClassConfigs.Emplace(UCameraRigAsset::StaticClass())
-		.HasSelfPin(false)
-		.OnlyAsRoot()
-		.SetPropertyPinDirection(GET_MEMBER_NAME_STRING_CHECKED(UCameraRigAsset, EnterTransitions), EGPD_Input)
-		.SetPropertyPinDirection(GET_MEMBER_NAME_STRING_CHECKED(UCameraRigAsset, ExitTransitions), EGPD_Output)
-		.NodeTitleUsesObjectName(true)
-		.NodeTitleColor(Settings->CameraRigAssetTitleColor);
-	GraphConfig.ObjectClassConfigs.Emplace(UCameraRigTransition::StaticClass())
-		.SetPropertyPinDirection(GET_MEMBER_NAME_STRING_CHECKED(UCameraRigTransition, Conditions), EGPD_Input)
-		.NodeTitleColor(Settings->CameraRigTransitionTitleColor);
-	GraphConfig.ObjectClassConfigs.Emplace(UCameraRigTransitionCondition::StaticClass())
-		.SelfPinDirection(EGPD_Output)
-		.DefaultPropertyPinDirection(EGPD_Input)
-		.StripDisplayNameSuffix(TEXT("Transition Condition"))
-		.NodeTitleColor(Settings->CameraRigTransitionConditionTitleColor);
-	GraphConfig.ObjectClassConfigs.Emplace(UBlendCameraNode::StaticClass())
-		.StripDisplayNameSuffix(TEXT("Camera Node"))
-		.CreateCategoryMetaData(TEXT("CameraNodeCategories"));
+	FCameraRigTransitionOwnerInfo OwnerInfo;
+	OwnerInfo.TransitionOwnerClass = UCameraRigAsset::StaticClass();
+	OwnerInfo.GraphName = UCameraRigAsset::TransitionsGraphName;
+	OwnerInfo.EnterTransitionsPropertyName = GET_MEMBER_NAME_CHECKED(UCameraRigAsset, EnterTransitions);
+	OwnerInfo.ExitTransitionsPropertyName = GET_MEMBER_NAME_CHECKED(UCameraRigAsset, ExitTransitions);
+	FObjectTreeGraphConfig GraphConfig = UCameraRigTransitionGraphSchema::BuildGraphConfig(OwnerInfo);
+
+	GraphConfig.GraphDisplayInfo.DisplayName = LOCTEXT("TransitionGraphDisplayName", "Transitions");
 
 	TransitionGraph = NewObject<UObjectTreeGraph>(GetTransientPackage(), NAME_None, RF_Transactional);
-	TransitionGraph->Schema = UCameraTransitionGraphSchema::StaticClass();
+	TransitionGraph->Schema = UCameraRigTransitionGraphSchema::StaticClass();
 	TransitionGraph->AddToRoot();
-	TransitionGraph->Initialize(CameraRigAsset, GraphConfig);
-	TransitionGraph->RebuildGraph(EObjectTreeGraphBuildSource::RootObjectPackage);
+	TransitionGraph->Reset(CameraRigAsset, GraphConfig, EObjectTreeGraphBuildSource::RootObjectPackage);
 
 	TransitionGraphChangedHandle = TransitionGraph->AddOnGraphChangedHandler(
 			FOnGraphChanged::FDelegate::CreateSP(this, &SCameraRigAssetEditor::OnGraphChanged));
@@ -156,10 +114,39 @@ void SCameraRigAssetEditor::CreateTransitionGraphEditor(const FArguments& InArgs
 
 	TransitionGraphEditor = SNew(SObjectTreeGraphEditor)
 		.Appearance(Appearance)
-		.DetailsView(InArgs._DetailsView)
-		.GraphTitle(this, &SCameraRigAssetEditor::GetCameraRigAssetName)
+		.DetailsView(DetailsView)
+		.GraphTitle(this, &SCameraRigAssetEditor::GetCameraRigAssetName, TransitionGraph.Get())
 		.GraphToEdit(TransitionGraph)
-		.AssetEditorToolkit(InArgs._AssetEditorToolkit);
+		.AssetEditorToolkit(AssetEditorToolkit);
+}
+
+void SCameraRigAssetEditor::DiscardGraphEditors()
+{
+	TArray<TTuple<UObjectTreeGraph*, FDelegateHandle>> Graphs 
+	{ 
+		{ NodeGraph, NodeGraphChangedHandle },
+		{ TransitionGraph, TransitionGraphChangedHandle } 
+	};
+	for (auto& Pair : Graphs)
+	{
+		UObjectTreeGraph* Graph(Pair.Key);
+		FDelegateHandle GraphChangedHandle(Pair.Value);
+		if (Graph)
+		{
+			Graph->RemoveFromRoot();
+
+			if (GraphChangedHandle.IsValid())
+			{
+				Graph->RemoveOnGraphChangedHandler(GraphChangedHandle);
+			}
+		}
+	}
+
+	NodeGraphChangedHandle.Reset();
+	TransitionGraphChangedHandle.Reset();
+
+	// WARNING: the graph editors (and their graphs) are still in use as widgets 
+	//			in the layout until they are replaced!
 }
 
 ECameraRigAssetEditorMode SCameraRigAssetEditor::GetEditorMode() const
@@ -174,7 +161,12 @@ bool SCameraRigAssetEditor::IsEditorMode(ECameraRigAssetEditorMode InMode) const
 
 void SCameraRigAssetEditor::SetEditorMode(ECameraRigAssetEditorMode InMode)
 {
-	if (InMode != CurrentMode)
+	SetEditorModeImpl(InMode, false);
+}
+
+void SCameraRigAssetEditor::SetEditorModeImpl(ECameraRigAssetEditorMode InMode, bool bForceSet)
+{
+	if (bForceSet || InMode != CurrentMode)
 	{
 		TSharedPtr<SObjectTreeGraphEditor> CurrentGraphEditor;
 		switch(InMode)
@@ -290,9 +282,14 @@ bool SCameraRigAssetEditor::FindAndJumpToObjectNode(UObject* InObject)
 	return false;
 }
 
-FText SCameraRigAssetEditor::GetCameraRigAssetName() const
+FText SCameraRigAssetEditor::GetCameraRigAssetName(UObjectTreeGraph* ForGraph) const
 {
-	return FText::FromString(CameraRigAsset->GetName());
+	if (CameraRigAsset && ForGraph)
+	{
+		const FObjectTreeGraphConfig& GraphConfig = TransitionGraph->GetConfig();
+		return GraphConfig.GetDisplayNameText(CameraRigAsset);
+	}
+	return LOCTEXT("NoCameraRig", "No Camera Rig");
 }
 
 void SCameraRigAssetEditor::OnGraphChanged(const FEdGraphEditAction& InEditAction)
diff --git a/Engine/Plugins/Cameras/GameplayCameras/Source/GameplayCamerasEditor/Private/Editors/SCameraRigList.cpp b/Engine/Plugins/Cameras/GameplayCameras/Source/GameplayCamerasEditor/Private/Editors/SCameraRigList.cpp
new file mode 100644
--- /dev/null
+++ b/Engine/Plugins/Cameras/GameplayCameras/Source/GameplayCamerasEditor/Private/Editors/SCameraRigList.cpp
@@ -0,0 +1,465 @@
+// Copyright Epic Games, Inc. All Rights Reserved.
+
+#include "Editors/SCameraRigList.h"
+
+#include "Commands/CameraAssetEditorCommands.h"
+#include "Core/CameraAsset.h"
+#include "Framework/Commands/UICommandList.h"
+#include "ScopedTransaction.h"
+#include "Styles/GameplayCamerasEditorStyle.h"
+#include "ToolMenus.h"
+#include "Widgets/Input/SSearchBox.h"
+#include "Widgets/Layout/SBox.h"
+#include "Widgets/SBoxPanel.h"
+#include "Widgets/Text/SInlineEditableTextBlock.h"
+#include "Widgets/Views/SListView.h"
+
+#define LOCTEXT_NAMESPACE "SCameraRigList"
+
+namespace UE::Cameras
+{
+
+void SCameraRigListEntry::Construct(const FArguments& InArgs, const TSharedRef<STableViewBase>& OwnerTable)
+{
+	using FSuperRowType = STableRow<TSharedPtr<FCameraRigListItem>>;
+
+	WeakItem = InArgs._Item;
+
+	ChildSlot
+	.Padding(FMargin(8.f, 2.f, 12.f, 2.f))
+	[
+		SNew(SBox)
+		.Padding(8.f, 4.f)
+		[
+			SAssignNew(EditableTextBlock, SInlineEditableTextBlock)
+			.Text(this, &SCameraRigListEntry::GetDisplayName)
+			.OnTextCommitted(this, &SCameraRigListEntry::OnTextCommitted)
+			.OnVerifyTextChanged(this, &SCameraRigListEntry::OnVerifyTextChanged)
+			.HighlightText(InArgs._HighlightText)
+			.IsSelected(this, &FSuperRowType::IsSelectedExclusively)
+		]
+	];
+
+	if (TSharedPtr<FCameraRigListItem> Item = WeakItem.Pin())
+	{
+		Item->OnRequestRename.BindSP(this, &SCameraRigListEntry::OnRename);
+	}
+
+	TSharedRef<FGameplayCamerasEditorStyle> CamerasStyle = FGameplayCamerasEditorStyle::Get();
+
+	FSuperRowType::ConstructInternal(
+		FSuperRowType::FArguments()
+			.Style(&CamerasStyle->GetWidgetStyle<FTableRowStyle>("CameraAssetEditor.CameraRigsList.RowStyle")),
+		OwnerTable);
+}
+
+FText SCameraRigListEntry::GetDisplayName() const
+{
+	if (TSharedPtr<FCameraRigListItem> Item = WeakItem.Pin())
+	{
+		UCameraRigAsset* CameraRigAsset = Item->CameraRigAsset;
+		return FText::FromString(CameraRigAsset->GetDisplayName());
+	}
+	return FText();
+}
+
+void SCameraRigListEntry::OnRename()
+{
+	EditableTextBlock->EnterEditingMode();
+}
+
+bool SCameraRigListEntry::OnVerifyTextChanged(const FText& Text, FText& OutErrorMessage)
+{
+	TSharedPtr<FCameraRigListItem> Item = WeakItem.Pin();
+	if (!Item)
+	{
+		OutErrorMessage = LOCTEXT("InvalidEntry", "Invalid entry");
+		return false;
+	}
+
+	UCameraRigAsset* CameraRigAsset = Item->CameraRigAsset;
+	UCameraAsset* OwnerCamera = CameraRigAsset->GetTypedOuter<UCameraAsset>();
+	if (!OwnerCamera)
+	{
+		OutErrorMessage = LOCTEXT("InvalidEntry", "Invalid entry");
+		return false;
+	}
+
+	const FString TextString = Text.ToString();
+	TObjectPtr<UCameraRigAsset>* FoundItem = OwnerCamera->CameraRigs.FindByPredicate(
+			[&TextString](const UCameraRigAsset* Item) { return Item->GetDisplayName() == TextString; });
+	if (FoundItem)
+	{
+		OutErrorMessage = LOCTEXT("NamingCollection", "A camera rig already exists with that name");
+		return false;
+	}
+
+	return true;
+}
+
+void SCameraRigListEntry::OnTextCommitted(const FText& Text, ETextCommit::Type CommitType)
+{
+	if (TSharedPtr<FCameraRigListItem> Item = WeakItem.Pin())
+	{
+		FScopedTransaction Transaction(LOCTEXT("RenameCameraRig", "Rename Camera Rig"));
+
+		UCameraRigAsset* CameraRigAsset = Item->CameraRigAsset;
+		CameraRigAsset->Modify();
+
+		// Set the interface name, but also rename the object itself, which helps for debugging.
+		const FString NewDisplayName = Text.ToString();
+		CameraRigAsset->Interface.DisplayName = NewDisplayName;
+		const FName NewObjectName = MakeUniqueObjectName(CameraRigAsset->GetOuter(), UCameraRigAsset::StaticClass(), *NewDisplayName);
+		CameraRigAsset->Rename(*NewObjectName.ToString());
+	}
+}
+
+void SCameraRigList::Construct(const FArguments& InArgs)
+{
+	CameraAsset = InArgs._CameraAsset;
+
+	OnCameraRigListChanged = InArgs._OnCameraRigListChanged;
+	OnRequestEditCameraRig = InArgs._OnRequestEditCameraRig;
+
+	CommandList = MakeShared<FUICommandList>();
+
+	SearchTextFilter = MakeShareable(new FEntryTextFilter(
+		FEntryTextFilter::FItemToStringArray::CreateSP(this, &SCameraRigList::GetEntryStrings)));
+
+	TSharedPtr<SWidget> ToolbarWidget = GenerateToolbar();
+
+	ChildSlot
+	[
+		SNew(SVerticalBox)
+		+SVerticalBox::Slot()
+		.AutoHeight()
+		[
+			SNew(SHorizontalBox)
+			+SHorizontalBox::Slot()
+			.AutoWidth()
+			[
+				ToolbarWidget.ToSharedRef()
+			]
+		]
+		+SVerticalBox::Slot()
+		.AutoHeight()
+		[
+			SNew(SBorder)
+			.BorderImage(FAppStyle::GetBrush("ToolPanel.GroupBorder"))
+			.Padding(8.f)
+			[
+				SAssignNew(SearchBox, SSearchBox)
+				.HintText(LOCTEXT("SearchHint", "Search"))
+				.OnTextChanged(this, &SCameraRigList::OnSearchTextChanged)
+				.OnTextCommitted(this, &SCameraRigList::OnSearchTextCommitted)
+			]
+		]
+		+SVerticalBox::Slot()
+		.Padding(0.f, 3.f)
+		[
+			SAssignNew(ListView, SListView<TSharedPtr<FCameraRigListItem>>)
+				.ListItemsSource(&FilteredItemSource)
+				.OnGenerateRow(this, &SCameraRigList::OnListGenerateItemRow)
+				.OnItemScrolledIntoView(this, &SCameraRigList::OnListItemScrolledIntoView)
+				.OnMouseButtonDoubleClick(this, &SCameraRigList::OnListMouseButtonDoubleClick)
+				.OnContextMenuOpening(this, &SCameraRigList::OnListContextMenuOpening)
+		]
+	];
+
+	const FCameraAssetEditorCommands& Commands = FCameraAssetEditorCommands::Get();
+
+	CommandList->MapAction(
+			Commands.EditCameraRig,
+			FExecuteAction::CreateSP(this, &SCameraRigList::OnEditCameraRig),
+			FCanExecuteAction::CreateSP(this, &SCameraRigList::CanEditCameraRig));
+	CommandList->MapAction(
+			Commands.AddCameraRig,
+			FExecuteAction::CreateSP(this, &SCameraRigList::OnAddCameraRig));
+	CommandList->MapAction(
+			Commands.RenameCameraRig,
+			FExecuteAction::CreateSP(this, &SCameraRigList::OnRenameCameraRig),
+			FCanExecuteAction::CreateSP(this, &SCameraRigList::CanRenameCameraRig));
+	CommandList->MapAction(
+			Commands.DeleteCameraRig,
+			FExecuteAction::CreateSP(this, &SCameraRigList::OnDeleteCameraRig),
+			FCanExecuteAction::CreateSP(this, &SCameraRigList::CanDeleteCameraRig));
+
+	UpdateItemSource();
+	UpdateFilteredItemSource();
+	ListView->RequestListRefresh();
+	if (!FilteredItemSource.IsEmpty())
+	{
+		ListView->SetSelection(FilteredItemSource[0]);
+	}
+}
+
+SCameraRigList::~SCameraRigList()
+{
+}
+
+void SCameraRigList::Tick(const FGeometry& AllottedGeometry, const double InCurrentTime, const float InDeltaTime)
+{
+	if (bUpdateItemSource)
+	{
+		UpdateItemSource();
+	}
+
+	if (bUpdateItemSource || bUpdateFilteredItemSource)
+	{
+		UpdateFilteredItemSource();
+	}
+
+	const bool bRequestListRefresh = bUpdateItemSource || bUpdateFilteredItemSource;
+	bUpdateItemSource = false;
+	bUpdateFilteredItemSource = false;
+
+	if (bRequestListRefresh)
+	{
+		ListView->RequestListRefresh();
+	}
+
+	SCompoundWidget::Tick(AllottedGeometry, InCurrentTime, InDeltaTime);
+}
+
+TSharedPtr<SWidget> SCameraRigList::GenerateToolbar()
+{
+	static const FName ToolbarName("CameraRigList.ToolBar");
+
+	UToolMenus* ToolMenus = UToolMenus::Get();
+
+	if (!ToolMenus->IsMenuRegistered(ToolbarName))
+	{
+		const FCameraAssetEditorCommands& Commands = FCameraAssetEditorCommands::Get();
+
+		UToolMenu* Toolbar = ToolMenus->RegisterMenu(ToolbarName, NAME_None, EMultiBoxType::SlimHorizontalToolBar);
+
+		FToolMenuSection& Section = Toolbar->AddSection("CameraRigs");
+		Section.AddEntry(FToolMenuEntry::InitToolBarButton(
+					Commands.AddCameraRig,
+					LOCTEXT("AddCameraRigButton", "Add")  // Shorter label
+				));
+		Section.AddEntry(FToolMenuEntry::InitToolBarButton(
+					Commands.RenameCameraRig,
+					LOCTEXT("RenameCameraRigButton", "Rename")  // Shorter label
+				));
+		Section.AddEntry(FToolMenuEntry::InitToolBarButton(
+					Commands.DeleteCameraRig,
+					LOCTEXT("DeleteCameraRigButton", "Delete")  // Shorter label
+				));
+	}
+
+	FToolMenuContext MenuContext;
+	MenuContext.AppendCommandList(CommandList);
+	return ToolMenus->GenerateWidget(ToolbarName, MenuContext);
+}
+
+void SCameraRigList::OnEditCameraRig()
+{
+	TArray<TSharedPtr<FCameraRigListItem>> SelectedItems;
+	ListView->GetSelectedItems(SelectedItems);
+	if (!SelectedItems.IsEmpty())
+	{
+		OnRequestEditCameraRig.ExecuteIfBound(SelectedItems[0]->CameraRigAsset);
+	}
+}
+
+bool SCameraRigList::CanEditCameraRig()
+{
+	TArray<TSharedPtr<FCameraRigListItem>> SelectedItems;
+	ListView->GetSelectedItems(SelectedItems);
+	return SelectedItems.Num() == 1;
+}
+
+void SCameraRigList::OnAddCameraRig()
+{
+	FScopedTransaction Transaction(LOCTEXT("AddCameraRig", "Add Camera Rig"));
+
+	CameraAsset->Modify();
+
+	const FName NewCameraRigName = MakeUniqueObjectName(CameraAsset, UCameraRigAsset::StaticClass(), TEXT("NewCameraRig"));
+	UCameraRigAsset* NewCameraRig = NewObject<UCameraRigAsset>(CameraAsset, NewCameraRigName, RF_Transactional);
+	CameraAsset->CameraRigs.Add(NewCameraRig);
+
+	bUpdateItemSource = true;
+}
+
+void SCameraRigList::OnRenameCameraRig()
+{
+	TArray<TSharedPtr<FCameraRigListItem>> SelectedItems;
+	ListView->GetSelectedItems(SelectedItems);
+	if (SelectedItems.Num() > 0)
+	{
+		ListView->RequestScrollIntoView(SelectedItems[0]);
+		DeferredRequestRenameItem = SelectedItems[0];
+	}
+}
+
+bool SCameraRigList::CanRenameCameraRig()
+{
+	TArray<TSharedPtr<FCameraRigListItem>> SelectedItems;
+	ListView->GetSelectedItems(SelectedItems);
+	return SelectedItems.Num() == 1;
+}
+
+void SCameraRigList::OnDeleteCameraRig()
+{
+	TArray<TSharedPtr<FCameraRigListItem>> SelectedItems;
+	ListView->GetSelectedItems(SelectedItems);
+	if (SelectedItems.Num() > 0)
+	{
+		FScopedTransaction Transaction(LOCTEXT("DeleteCameraRigs", "Delete Camera Rigs"));
+
+		CameraAsset->Modify();
+
+		TStringBuilder<256> StringBuilder;
+
+		for (TSharedPtr<FCameraRigListItem> Item : SelectedItems)
+		{
+			UCameraRigAsset* CameraRigAsset = Item->CameraRigAsset;
+			if (CameraRigAsset)
+			{
+				CameraRigAsset->Modify();
+				
+				StringBuilder.Reset();
+				StringBuilder.Append(TEXT("DELETED_"));
+				StringBuilder.Append(CameraRigAsset->GetName());
+				CameraRigAsset->Rename(StringBuilder.ToString());
+
+				const int32 NumRemoved = CameraAsset->CameraRigs.Remove(CameraRigAsset);
+				ensure(NumRemoved == 1);
+			}
+		}
+
+		bUpdateItemSource = true;
+	}
+}
+
+bool SCameraRigList::CanDeleteCameraRig()
+{
+	TArray<TSharedPtr<FCameraRigListItem>> SelectedItems;
+	ListView->GetSelectedItems(SelectedItems);
+	return SelectedItems.Num() > 0;
+}
+
+void SCameraRigList::GetEntryStrings(const TSharedPtr<FCameraRigListItem> InItem, TArray<FString>& OutStrings)
+{
+	if (InItem->CameraRigAsset)
+	{
+		FString DisplayName = InItem->CameraRigAsset->GetDisplayName();
+		OutStrings.Add(DisplayName);
+	}
+}
+
+void SCameraRigList::UpdateItemSource()
+{
+	ItemSource.Reset();
+
+	if (CameraAsset)
+	{
+		for (UCameraRigAsset* CameraRigAsset : CameraAsset->CameraRigs)
+		{
+			TSharedPtr<FCameraRigListItem> Item = MakeShared<FCameraRigListItem>();
+			Item->CameraRigAsset = CameraRigAsset;
+			ItemSource.Add(Item);
+		}
+	}
+
+	OnCameraRigListChanged.ExecuteIfBound(CameraAsset->CameraRigs);
+}
+
+void SCameraRigList::UpdateFilteredItemSource()
+{
+	FilteredItemSource = ItemSource;
+	FilteredItemSource.StableSort([](TSharedPtr<FCameraRigListItem> A, TSharedPtr<FCameraRigListItem> B)
+			{ 
+				return A->CameraRigAsset->GetDisplayName().Compare(B->CameraRigAsset->GetDisplayName()) < 0;
+			});
+
+	if (!SearchTextFilter->GetRawFilterText().IsEmpty())
+	{
+		FilteredItemSource = FilteredItemSource.FilterByPredicate(
+				[this](TSharedPtr<FCameraRigListItem> Item)
+				{
+					return SearchTextFilter->PassesFilter(Item);
+				});
+	}
+}
+
+TSharedRef<ITableRow> SCameraRigList::OnListGenerateItemRow(TSharedPtr<FCameraRigListItem> Item, const TSharedRef<STableViewBase>& OwnerTable)
+{
+	return SNew(SCameraRigListEntry, OwnerTable)
+		.Item(Item)
+		.HighlightText(this, &SCameraRigList::GetHighlightText);
+}
+
+void SCameraRigList::OnListItemScrolledIntoView(TSharedPtr<FCameraRigListItem> Item, const TSharedPtr<ITableRow>& ItemWidget)
+{
+	if (DeferredRequestRenameItem)
+	{
+		DeferredRequestRenameItem->OnRequestRename.ExecuteIfBound();
+		DeferredRequestRenameItem.Reset();
+	}
+}
+
+void SCameraRigList::OnListMouseButtonDoubleClick(TSharedPtr<FCameraRigListItem> Item)
+{
+	if (Item)
+	{
+		OnRequestEditCameraRig.ExecuteIfBound(Item->CameraRigAsset);
+	}
+}
+
+TSharedPtr<SWidget> SCameraRigList::OnListContextMenuOpening()
+{
+	static const FName ContextMenuName("CameraRigList.ContextMenu");
+
+	UToolMenus* ToolMenus = UToolMenus::Get();
+
+	if (!ToolMenus->IsMenuRegistered(ContextMenuName))
+	{
+		const FCameraAssetEditorCommands& Commands = FCameraAssetEditorCommands::Get();
+
+		UToolMenu* ContextMenu = ToolMenus->RegisterMenu(ContextMenuName, NAME_None, EMultiBoxType::Menu);
+
+		FToolMenuSection& Section = ContextMenu->AddSection("Actions");
+		Section.AddEntry(FToolMenuEntry::InitMenuEntry(
+					Commands.EditCameraRig,
+					LOCTEXT("AddCameraRigButton", "Edit")  // Shorter label
+				));
+		Section.AddEntry(FToolMenuEntry::InitMenuEntry(
+					Commands.RenameCameraRig,
+					LOCTEXT("RenameCameraRigButton", "Rename")  // Shorter label
+				));
+		Section.AddEntry(FToolMenuEntry::InitMenuEntry(
+					Commands.DeleteCameraRig,
+					LOCTEXT("DeleteCameraRigButton", "Delete")  // Shorter label
+				));
+	}
+
+	FToolMenuContext MenuContext;
+	MenuContext.AppendCommandList(CommandList);
+	return ToolMenus->GenerateWidget(ContextMenuName, MenuContext);
+}
+
+void SCameraRigList::OnSearchTextChanged(const FText& InFilterText)
+{
+	SearchTextFilter->SetRawFilterText(InFilterText);
+	SearchBox->SetError(SearchTextFilter->GetFilterErrorText());
+
+	bUpdateFilteredItemSource = true;
+}
+
+void SCameraRigList::OnSearchTextCommitted(const FText& InFilterText, ETextCommit::Type InCommitType)
+{
+	OnSearchTextChanged(InFilterText);
+}
+
+FText SCameraRigList::GetHighlightText() const
+{
+	return SearchTextFilter->GetRawFilterText();
+}
+
+}  // namespace UE::Cameras
+
+#undef LOCTEXT_NAMESPACE
+
diff --git a/Engine/Plugins/Cameras/GameplayCameras/Source/GameplayCamerasEditor/Private/Editors/SCameraRigTransitionEditor.cpp b/Engine/Plugins/Cameras/GameplayCameras/Source/GameplayCamerasEditor/Private/Editors/SCameraRigTransitionEditor.cpp
new file mode 100644
--- /dev/null
+++ b/Engine/Plugins/Cameras/GameplayCameras/Source/GameplayCamerasEditor/Private/Editors/SCameraRigTransitionEditor.cpp
@@ -0,0 +1,187 @@
+// Copyright Epic Games, Inc. All Rights Reserved.
+
+#include "Editors/SCameraRigTransitionEditor.h"
+
+#include "Core/CameraAsset.h"
+#include "Core/CameraRigAsset.h"
+#include "Core/CameraRigTransition.h"
+#include "EdGraph/EdGraphPin.h"
+#include "Editors/CameraRigTransitionGraphSchema.h"
+#include "Editors/ObjectTreeGraph.h"
+#include "Editors/ObjectTreeGraphConfig.h"
+#include "Editors/ObjectTreeGraphNode.h"
+#include "Editors/SObjectTreeGraphEditor.h"
+#include "GameplayCamerasEditorSettings.h"
+#include "ObjectEditorUtils.h"
+#include "Widgets/Layout/SBox.h"
+
+#define LOCTEXT_NAMESPACE "SCameraRigTransitionEditor"
+
+namespace UE::Cameras
+{
+
+void SCameraRigTransitionEditor::Construct(const FArguments& InArgs)
+{
+	TransitionOwner = InArgs._TransitionOwner;
+	TransitionOwnerInfo = InArgs._TransitionOwnerInfo;
+	OnBuildTransitionGraphConfig = InArgs._OnBuildTransitionGraphConfig;
+	DetailsView = InArgs._DetailsView;
+	AssetEditorToolkit = InArgs._AssetEditorToolkit;
+	TransitionGraphDisplayInfo = InArgs._TransitionGraphDisplayInfo;
+	TransitionGraphEditorAppearance = InArgs._TransitionGraphEditorAppearance;
+
+	CreateTransitionGraphEditor();
+
+	ChildSlot
+	[
+		SAssignNew(BoxPanel, SBox)
+		[
+			TransitionGraphEditor.ToSharedRef()
+		]
+	];
+}
+
+SCameraRigTransitionEditor::~SCameraRigTransitionEditor()
+{
+	if (!GExitPurge)
+	{
+		DiscardTransitionGraphEditor();
+	}
+}
+
+void SCameraRigTransitionEditor::SetTransitionOwner(UObject* InTransitionOwner)
+{
+	if (TransitionOwner != InTransitionOwner)
+	{
+		DiscardTransitionGraphEditor();
+
+		TransitionOwner = InTransitionOwner;
+
+		CreateTransitionGraphEditor();
+
+		BoxPanel->SetContent(TransitionGraphEditor.ToSharedRef());
+		TransitionGraphEditor->ResyncDetailsView();
+	}
+}
+
+void SCameraRigTransitionEditor::CreateTransitionGraphEditor()
+{
+	FObjectTreeGraphConfig GraphConfig = UCameraRigTransitionGraphSchema::BuildGraphConfig(TransitionOwnerInfo);
+	if (TransitionGraphDisplayInfo.IsSet())
+	{
+		GraphConfig.GraphDisplayInfo = TransitionGraphDisplayInfo.Get();
+	}
+	OnBuildTransitionGraphConfig.ExecuteIfBound(GraphConfig);
+
+	TransitionGraph = NewObject<UObjectTreeGraph>(GetTransientPackage(), NAME_None, RF_Transactional);
+	TransitionGraph->Schema = UCameraRigTransitionGraphSchema::StaticClass();
+	TransitionGraph->AddToRoot();
+	TransitionGraph->Reset(TransitionOwner, GraphConfig, EObjectTreeGraphBuildSource::RootObjectPackage);
+
+	TransitionGraphChangedHandle = TransitionGraph->AddOnGraphChangedHandler(
+			FOnGraphChanged::FDelegate::CreateSP(this, &SCameraRigTransitionEditor::OnGraphChanged));
+
+	FGraphAppearanceInfo Appearance;
+	Appearance.CornerText = LOCTEXT("TransitionGraphText", "TRANSITIONS");
+	if (TransitionGraphEditorAppearance.IsSet())
+	{
+		Appearance = TransitionGraphEditorAppearance.Get();
+	}
+
+	TransitionGraphEditor = SNew(SObjectTreeGraphEditor)
+		.Appearance(Appearance)
+		.DetailsView(DetailsView)
+		.GraphTitle(this, &SCameraRigTransitionEditor::GetTransitionOwnerName)
+		.GraphToEdit(TransitionGraph)
+		.AssetEditorToolkit(AssetEditorToolkit);
+}
+
+void SCameraRigTransitionEditor::DiscardTransitionGraphEditor()
+{
+	if (TransitionGraph)
+	{
+		TransitionGraph->RemoveFromRoot();
+
+		if (TransitionGraphChangedHandle.IsValid())
+		{
+			TransitionGraph->RemoveOnGraphChangedHandler(TransitionGraphChangedHandle);
+		}
+	}
+
+	TransitionGraph = nullptr;
+	TransitionGraphChangedHandle.Reset();
+
+	// WARNING: the graph editor (and its graph) is still in use as a widget in the layout 
+	//			until it is replaced!
+}
+
+void SCameraRigTransitionEditor::OnGraphChanged(const FEdGraphEditAction& InEditAction)
+{
+	OnTransitionGraphChanged.Broadcast(InEditAction);
+}
+
+FDelegateHandle SCameraRigTransitionEditor::AddOnGraphChanged(FOnGraphChanged::FDelegate InAddDelegate)
+{
+	return OnTransitionGraphChanged.Add(InAddDelegate);
+}
+
+void SCameraRigTransitionEditor::RemoveOnGraphChanged(FDelegateHandle InDelegateHandle)
+{
+	OnTransitionGraphChanged.Remove(InDelegateHandle);
+}
+
+void SCameraRigTransitionEditor::RemoveOnGraphChanged(const void* InUserObject)
+{
+	OnTransitionGraphChanged.RemoveAll(InUserObject);
+}
+
+UEdGraph* SCameraRigTransitionEditor::GetTransitionGraph() const
+{
+	return TransitionGraph;
+}
+
+const FObjectTreeGraphConfig& SCameraRigTransitionEditor::GetTransitionGraphConfig() const
+{
+	return TransitionGraph->GetConfig();
+}
+
+void SCameraRigTransitionEditor::FocusHome()
+{
+	if (UObjectTreeGraphNode* RootObjectNode = TransitionGraph->GetRootObjectNode())
+	{
+		JumpToNode(RootObjectNode);
+	}
+}
+
+void SCameraRigTransitionEditor::JumpToNode(UEdGraphNode* InGraphNode)
+{
+	if (InGraphNode)
+	{
+		TransitionGraphEditor->JumpToNode(InGraphNode);
+	}
+}
+
+bool SCameraRigTransitionEditor::FindAndJumpToObjectNode(UObject* InObject)
+{
+	if (UObjectTreeGraphNode* NodeGraphObjectNode = TransitionGraph->FindObjectNode(InObject))
+	{
+		TransitionGraphEditor->JumpToNode(NodeGraphObjectNode);
+		return true;
+	}
+	return false;
+}
+
+FText SCameraRigTransitionEditor::GetTransitionOwnerName() const
+{
+	if (TransitionOwner && TransitionGraph)
+	{
+		const FObjectTreeGraphConfig& GraphConfig = TransitionGraph->GetConfig();
+		return GraphConfig.GetDisplayNameText(TransitionOwner);
+	}
+	return LOCTEXT("NoTransitionOwner", "No Transition Owner");
+}
+
+}  // namespace UE::Cameras
+
+#undef LOCTEXT_NAMESPACE
+
diff --git a/Engine/Plugins/Cameras/GameplayCameras/Source/GameplayCamerasEditor/Private/GameplayCamerasEditorModule.cpp b/Engine/Plugins/Cameras/GameplayCameras/Source/GameplayCamerasEditor/Private/GameplayCamerasEditorModule.cpp
--- a/Engine/Plugins/Cameras/GameplayCameras/Source/GameplayCamerasEditor/Private/GameplayCamerasEditorModule.cpp
+++ b/Engine/Plugins/Cameras/GameplayCameras/Source/GameplayCamerasEditor/Private/GameplayCamerasEditorModule.cpp
@@ -7,6 +7,7 @@
 #include "AssetTools/CameraVariableCollectionEditor.h"
 #include "Commands/CameraAssetEditorCommands.h"
 #include "Commands/CameraRigAssetEditorCommands.h"
+#include "Commands/CameraRigTransitionEditorCommands.h"
 #include "Commands/CameraVariableCollectionEditorCommands.h"
 #include "Commands/GameplayCamerasDebuggerCommands.h"
 #include "Customizations/CameraParameterDetailsCustomizations.h"
@@ -82,6 +83,7 @@ public:
 
 		FCameraAssetEditorCommands::Unregister();
 		FCameraRigAssetEditorCommands::Unregister();
+		FCameraRigTransitionEditorCommands::Unregister();
 		FCameraVariableCollectionEditorCommands::Unregister();
 		FGameplayCamerasDebuggerCommands::Unregister();
 
@@ -270,6 +272,7 @@ private:
 
 		FCameraAssetEditorCommands::Register();
 		FCameraRigAssetEditorCommands::Register();	
+		FCameraRigTransitionEditorCommands::Register();
 		FCameraVariableCollectionEditorCommands::Register();
 		FGameplayCamerasDebuggerCommands::Register();
 	}
diff --git a/Engine/Plugins/Cameras/GameplayCameras/Source/GameplayCamerasEditor/Private/Styles/GameplayCamerasEditorStyle.cpp b/Engine/Plugins/Cameras/GameplayCameras/Source/GameplayCamerasEditor/Private/Styles/GameplayCamerasEditorStyle.cpp
--- a/Engine/Plugins/Cameras/GameplayCameras/Source/GameplayCamerasEditor/Private/Styles/GameplayCamerasEditorStyle.cpp
+++ b/Engine/Plugins/Cameras/GameplayCameras/Source/GameplayCamerasEditor/Private/Styles/GameplayCamerasEditorStyle.cpp
@@ -31,7 +31,26 @@ FGameplayCamerasEditorStyle::FGameplayCamerasEditorStyle()
 	SetContentRoot(ContentDir);
 	SetCoreContentRoot(FPaths::EngineContentDir() / TEXT("Slate"));
 
-	const FTextBlockStyle& NormalText = FAppStyle::Get().GetWidgetStyle<FTextBlockStyle>("NormalText");
+	const ISlateStyle& CoreStyle = FAppStyle::Get();
+	const FTextBlockStyle& NormalText = CoreStyle.GetWidgetStyle<FTextBlockStyle>("NormalText");
+	const FTableRowStyle AlternatingTableRowStyle = CoreStyle.GetWidgetStyle<FTableRowStyle>("TableView.AlternatingRow");
+
+	// Camera asset editor icons.
+	Set("CameraAssetEditor.ShowCameraDirector", new IMAGE_BRUSH_SVG("Icons/CameraEditor-CameraDirector", Icon20x20));
+	Set("CameraAssetEditor.ShowCameraRigs", new IMAGE_BRUSH_SVG("Icons/CameraEditor-CameraRigs", Icon20x20));
+	Set("CameraAssetEditor.ShowSharedTransitions", new IMAGE_BRUSH_SVG("Icons/CameraEditor-SharedTransitions", Icon20x20));
+	
+	Set("CameraAssetEditor.Tabs.Search", new CORE_IMAGE_BRUSH_SVG("Starship/Common/search", Icon16x16));
+	Set("CameraAssetEditor.Tabs.Messages", new CORE_IMAGE_BRUSH_SVG("Starship/Common/OutputLog", Icon16x16));
+	Set("CameraAssetEditor.Tabs.CameraRigs", new IMAGE_BRUSH_SVG("Icons/CameraEditor-CameraRigs", Icon16x16));
+
+	Set("CameraAssetEditor.FindInCamera", new CORE_IMAGE_BRUSH_SVG("Starship/Common/Search", Icon20x20));
+
+	Set("CameraAssetEditor.AddCameraRig", new IMAGE_BRUSH_SVG("Icons/CameraEditor-AddCameraRig", Icon16x16));
+	Set("CameraAssetEditor.DeleteCameraRig", new IMAGE_BRUSH_SVG("Icons/CameraEditor-DeleteCameraRig", Icon16x16));
+	Set("CameraAssetEditor.RenameCameraRig", new IMAGE_BRUSH_SVG("Icons/CameraEditor-RenameCameraRig", Icon16x16));
+
+	Set("CameraAssetEditor.CameraRigsList.RowStyle", AlternatingTableRowStyle);
 
 	// Camera rig editor icons.
 	Set("CameraRigAssetEditor.Tabs.Toolbox", new IMAGE_BRUSH_SVG("Icons/CameraRig-Toolbox", Icon16x16));
diff --git a/Engine/Plugins/Cameras/GameplayCameras/Source/GameplayCamerasEditor/Private/Toolkits/AssetEditorMode.cpp b/Engine/Plugins/Cameras/GameplayCameras/Source/GameplayCamerasEditor/Private/Toolkits/AssetEditorMode.cpp
new file mode 100644
--- /dev/null
+++ b/Engine/Plugins/Cameras/GameplayCameras/Source/GameplayCamerasEditor/Private/Toolkits/AssetEditorMode.cpp
@@ -0,0 +1,40 @@
+// Copyright Epic Games, Inc. All Rights Reserved.
+
+#include "Toolkits/AssetEditorMode.h"
+
+#include "Framework/Docking/LayoutExtender.h"
+#include "Framework/MultiBox/MultiBoxExtender.h"
+
+namespace UE::Cameras
+{
+
+FAssetEditorMode::FAssetEditorMode()
+{
+}
+
+FAssetEditorMode::FAssetEditorMode(FName InModeName)
+	: ModeName(InModeName)
+{
+}
+
+FAssetEditorMode::~FAssetEditorMode()
+{
+}
+
+void FAssetEditorMode::ActivateMode(const FAssetEditorModeActivateParams& InParams)
+{
+	OnActivateMode(InParams);
+}
+
+void FAssetEditorMode::InitToolMenuContext(FToolMenuContext& MenuContext)
+{
+	OnInitToolMenuContext(MenuContext);
+}
+
+void FAssetEditorMode::DeactivateMode(const FAssetEditorModeDeactivateParams& InParams)
+{
+	OnDeactivateMode(InParams);
+}
+
+}  // namespace UE::Cameras
+
diff --git a/Engine/Plugins/Cameras/GameplayCameras/Source/GameplayCamerasEditor/Private/Toolkits/AssetEditorModeManagerToolkit.cpp b/Engine/Plugins/Cameras/GameplayCameras/Source/GameplayCamerasEditor/Private/Toolkits/AssetEditorModeManagerToolkit.cpp
new file mode 100644
--- /dev/null
+++ b/Engine/Plugins/Cameras/GameplayCameras/Source/GameplayCamerasEditor/Private/Toolkits/AssetEditorModeManagerToolkit.cpp
@@ -0,0 +1,148 @@
+// Copyright Epic Games, Inc. All Rights Reserved.
+
+#include "Toolkits/AssetEditorModeManagerToolkit.h"
+
+#include "Framework/Docking/LayoutExtender.h"
+#include "Framework/Docking/TabManager.h"
+#include "Framework/MultiBox/MultiBoxExtender.h"
+#include "Toolkits/AssetEditorMode.h"
+
+namespace UE::Cameras
+{
+
+FAssetEditorModeManagerToolkit::FAssetEditorModeManagerToolkit(UAssetEditor* InOwningAssetEditor)
+	: FBaseAssetToolkit(InOwningAssetEditor)
+{
+}
+
+FAssetEditorModeManagerToolkit::~FAssetEditorModeManagerToolkit()
+{
+}
+
+void FAssetEditorModeManagerToolkit::InitToolMenuContext(FToolMenuContext& MenuContext)
+{
+	FBaseAssetToolkit::InitToolMenuContext(MenuContext);
+
+	if (CurrentEditorMode)
+	{
+		CurrentEditorMode->InitToolMenuContext(MenuContext);
+	}
+}
+
+void FAssetEditorModeManagerToolkit::AddEditorMode(TSharedRef<FAssetEditorMode> InMode)
+{
+	if (!ensureMsgf(
+				!EditorModes.Contains(InMode->GetModeName()),
+				TEXT("An editor mode named '%s' has already been added!"), *InMode->GetModeName().ToString()))
+	{
+		return;
+	}
+
+	EditorModes.Add(InMode->GetModeName(), InMode);
+}
+
+void FAssetEditorModeManagerToolkit::GetEditorModes(TArray<TSharedPtr<FAssetEditorMode>>& OutModes) const
+{
+	for (const TPair<FName, TSharedPtr<FAssetEditorMode>>& Pair : EditorModes)
+	{
+		OutModes.Add(Pair.Value);
+	}
+}
+
+TSharedPtr<FAssetEditorMode> FAssetEditorModeManagerToolkit::GetEditorMode(FName InModeName) const
+{
+	return EditorModes.FindRef(InModeName);
+}
+
+void FAssetEditorModeManagerToolkit::SetEditorMode(FName InModeName)
+{
+	if (CurrentEditorModeName == InModeName)
+	{
+		return;
+	}
+
+	if (!ensure(TabManager))
+	{
+		return;
+	}
+
+	// Remove current editor mode and setup.
+	if (CurrentEditorMode)
+	{
+		// Deactivating should in theory remove all of the mode's tab-spawners.
+		FAssetEditorModeDeactivateParams DeactivateParams;
+		DeactivateParams.Toolkit = SharedThis(this);
+		DeactivateParams.TabManager = TabManager;
+		CurrentEditorMode->DeactivateMode(DeactivateParams);
+
+		if (TSharedPtr<FExtender> OldToolbarExtender = CurrentEditorMode->GetToolbarExtender())
+		{
+			RemoveToolbarExtender(OldToolbarExtender);
+		}
+
+		if (TSharedPtr<FLayoutExtender> OldLayoutExtender = CurrentEditorMode->GetLayoutExtender())
+		{
+			LayoutExtenders.Remove(OldLayoutExtender);
+		}
+	}
+
+	RemoveAllToolbarWidgets();
+
+	// Set the new current editor mode.
+	TSharedPtr<FAssetEditorMode> NewMode = EditorModes.FindRef(InModeName);
+	CurrentEditorMode = NewMode;
+	CurrentEditorModeName = InModeName;
+
+	// Setup the new editor mode.
+	if (NewMode)
+	{
+		FName ParentName;
+		const FName ToolbarMenuName = GetToolMenuToolbarName(ParentName);
+
+		// Activating should in theory add all of the mode's tab-spawners.
+		FAssetEditorModeActivateParams ActivateParams;
+		ActivateParams.Toolkit = SharedThis(this);
+		ActivateParams.TabManager = TabManager;
+		ActivateParams.AssetEditorTabsCategory = AssetEditorTabsCategory;
+		ActivateParams.CommandList = ToolkitCommands;
+		ActivateParams.ToolbarMenuName = ToolbarMenuName;
+		NewMode->ActivateMode(ActivateParams);
+
+		if (TSharedPtr<FLayoutExtender> NewLayoutExtender = NewMode->GetLayoutExtender())
+		{
+			LayoutExtenders.Add(NewLayoutExtender);
+		}
+
+		if (TSharedPtr<FTabManager::FLayout> NewLayout = NewMode->GetDefaultLayout())
+		{
+			RestoreFromLayout(NewLayout.ToSharedRef());
+		}
+
+		if (TSharedPtr<FExtender> NewToolbarExtender = NewMode->GetToolbarExtender())
+		{
+			AddToolbarExtender(NewToolbarExtender);
+		}
+
+		OnEditorToolkitModeActivated();
+	}
+
+	RegenerateMenusAndToolbars();
+}
+
+bool FAssetEditorModeManagerToolkit::CanSetEditorMode(FName InModeName) const
+{
+	return EditorModes.Contains(InModeName);
+}
+
+bool FAssetEditorModeManagerToolkit::IsEditorMode(FName InModeName) const
+{
+	return CurrentEditorModeName == InModeName;
+}
+
+FName FAssetEditorModeManagerToolkit::GetCurrentEditorModeName() const
+{
+	return CurrentEditorModeName;
+}
+
+}  // namespace UE::Cameras
+
diff --git a/Engine/Plugins/Cameras/GameplayCameras/Source/GameplayCamerasEditor/Private/Toolkits/BuildButtonToolkit.cpp b/Engine/Plugins/Cameras/GameplayCameras/Source/GameplayCamerasEditor/Private/Toolkits/BuildButtonToolkit.cpp
new file mode 100644
--- /dev/null
+++ b/Engine/Plugins/Cameras/GameplayCameras/Source/GameplayCamerasEditor/Private/Toolkits/BuildButtonToolkit.cpp
@@ -0,0 +1,92 @@
+// Copyright Epic Games, Inc. All Rights Reserved.
+
+#include "Toolkits/BuildButtonToolkit.h"
+
+#include "Core/CameraBuildStatus.h"
+#include "Styles/GameplayCamerasEditorStyle.h"
+#include "Textures/SlateIcon.h"
+#include "ToolMenuEntry.h"
+
+#define LOCTEXT_NAMESPACE "BuildButtonToolkit"
+
+namespace UE::Cameras
+{
+
+FBuildButtonToolkit::FBuildButtonToolkit()
+{
+}
+
+FBuildButtonToolkit::FBuildButtonToolkit(TScriptInterface<IHasCameraBuildStatus> InTarget)
+	: Target(InTarget)
+{
+}
+
+void FBuildButtonToolkit::SetTarget(TScriptInterface<IHasCameraBuildStatus> InTarget)
+{
+	Target = InTarget;
+}
+
+FToolMenuEntry FBuildButtonToolkit::MakeToolbarButton(TSharedPtr<FUICommandInfo> InCommand)
+{
+	FToolMenuEntry BuildButton = FToolMenuEntry::InitToolBarButton(InCommand);
+	BuildButton.Icon = TAttribute<FSlateIcon>(this, &FBuildButtonToolkit::GetBuildButtonIcon);
+	BuildButton.ToolTip = TAttribute<FText>(this, &FBuildButtonToolkit::GetBuildButtonTooltip);
+	return BuildButton;
+}
+
+FSlateIcon FBuildButtonToolkit::GetBuildButtonIcon() const
+{
+	static const FName BuildStatusBackground("CameraRigAssetEditor.BuildStatus.Background");
+	static const FName BuildStatusError("CameraRigAssetEditor.BuildStatus.Overlay.Error");
+	static const FName BuildStatusGood("CameraRigAssetEditor.BuildStatus.Overlay.Good");
+	static const FName BuildStatusUnknown("CameraRigAssetEditor.BuildStatus.Overlay.Unknown");
+	static const FName BuildStatusWarning("CameraRigAssetEditor.BuildStatus.Overlay.Warning");
+
+	const FName CamerasStyleSetName = FGameplayCamerasEditorStyle::Get()->GetStyleSetName();
+
+	IHasCameraBuildStatus* Interface = Target.GetInterface();
+	if (!Interface)
+	{
+		return FSlateIcon(CamerasStyleSetName, BuildStatusBackground, NAME_None, BuildStatusError);
+	}
+
+	switch (Interface->GetBuildStatus())
+	{
+		default:
+		case ECameraBuildStatus::Dirty:
+			return FSlateIcon(CamerasStyleSetName, BuildStatusBackground, NAME_None, BuildStatusUnknown);
+		case ECameraBuildStatus::WithErrors:
+			return FSlateIcon(CamerasStyleSetName, BuildStatusBackground, NAME_None, BuildStatusError);
+		case ECameraBuildStatus::Clean:
+			return FSlateIcon(CamerasStyleSetName, BuildStatusBackground, NAME_None, BuildStatusGood);
+		case ECameraBuildStatus::CleanWithWarnings:
+			return FSlateIcon(CamerasStyleSetName, BuildStatusBackground, NAME_None, BuildStatusWarning);
+	}
+}
+
+FText FBuildButtonToolkit::GetBuildButtonTooltip() const
+{
+	IHasCameraBuildStatus* Interface = Target.GetInterface();
+	if (!Interface)
+	{
+		return LOCTEXT("BuildButtonStatusNoAsset", "No asset is open");
+	}
+
+	switch (Interface->GetBuildStatus())
+	{
+		default:
+		case ECameraBuildStatus::Dirty:
+			return LOCTEXT("BuildButtonStatusDirty", "Dirty or unknown, should rebuild");
+		case ECameraBuildStatus::WithErrors:
+			return LOCTEXT("BuildButtonStatusWithErrors", "There were errors during the build, see the log window for details");
+		case ECameraBuildStatus::Clean:
+			return LOCTEXT("BuildButtonStatusClean", "Good to go");
+		case ECameraBuildStatus::CleanWithWarnings:
+			return LOCTEXT("BuildButtonStatusCleanWithWarnings", "There were warnings during the build, see the log window for details");
+	}
+}
+
+}  // namespace UE::Cameras
+
+#undef LOCTEXT_NAMESPACE
+
diff --git a/Engine/Plugins/Cameras/GameplayCameras/Source/GameplayCamerasEditor/Private/Toolkits/CameraAssetEditorToolkit.cpp b/Engine/Plugins/Cameras/GameplayCameras/Source/GameplayCamerasEditor/Private/Toolkits/CameraAssetEditorToolkit.cpp
--- a/Engine/Plugins/Cameras/GameplayCameras/Source/GameplayCamerasEditor/Private/Toolkits/CameraAssetEditorToolkit.cpp
+++ b/Engine/Plugins/Cameras/GameplayCameras/Source/GameplayCamerasEditor/Private/Toolkits/CameraAssetEditorToolkit.cpp
@@ -3,57 +3,114 @@
 #include "Toolkits/CameraAssetEditorToolkit.h"
 
 #include "AssetTools/CameraAssetEditor.h"
+#include "Commands/CameraAssetEditorCommands.h"
 #include "Core/CameraAsset.h"
-#include "EditorModeManager.h"
+#include "Core/CameraBuildLog.h"
+#include "Editors/SFindInObjectTreeGraph.h"
 #include "Framework/Docking/LayoutExtender.h"
-#include "IMessageLogListing.h"
-#include "MessageLogInitializationOptions.h"
-#include "MessageLogModule.h"
+#include "Framework/Docking/TabManager.h"
+#include "GameplayCamerasEditorSettings.h"
+#include "IGameplayCamerasLiveEditManager.h"
+#include "IGameplayCamerasModule.h"
 #include "Modules/ModuleManager.h"
 #include "PropertyEditorModule.h"
+#include "Styles/GameplayCamerasEditorStyle.h"
+#include "ToolMenus.h"
+#include "Toolkits/BuildButtonToolkit.h"
+#include "Toolkits/CameraBuildLogToolkit.h"
+#include "Toolkits/CameraDirectorAssetEditorMode.h"
+#include "Toolkits/CameraRigsAssetEditorMode.h"
+#include "Toolkits/CameraSharedTransitionsAssetEditorMode.h"
+#include "Toolkits/StandardToolkitLayout.h"
 #include "Widgets/Docking/SDockTab.h"
 
+#include UE_INLINE_GENERATED_CPP_BY_NAME(CameraAssetEditorToolkit)
+
 #define LOCTEXT_NAMESPACE "CameraAssetEditorToolkit"
 
 namespace UE::Cameras
 {
 
-const FName FCameraAssetEditorToolkit::DetailsViewTabId(TEXT("CameraAssetEditor_DetailsView"));
+const FName FCameraAssetEditorToolkit::SearchTabId(TEXT("CameraAssetEditor_Search"));
+const FName FCameraAssetEditorToolkit::MessagesTabId(TEXT("CameraAssetEditor_Messages"));
 
 FCameraAssetEditorToolkit::FCameraAssetEditorToolkit(UCameraAssetEditor* InOwningAssetEditor)
-	: FBaseAssetToolkit(InOwningAssetEditor)
-{
-	// Override base class default layout.
-	StandaloneDefaultLayout = FTabManager::NewLayout("CameraAssetEditor_Layout")
-		->AddArea
-		(
-			FTabManager::NewPrimaryArea()->SetOrientation(Orient_Vertical)
-			->Split
-			(
-				FTabManager::NewSplitter()->SetOrientation(Orient_Horizontal)->SetSizeCoefficient(0.9f)
-				->Split
-				(
-					FTabManager::NewStack()
-					->AddTab(DetailsViewTabId, ETabState::OpenedTab)
-					->SetForegroundTab(DetailsViewTabId)
-				)
-			)
-		);
+	: FAssetEditorModeManagerToolkit(InOwningAssetEditor)
+	, CameraAsset(InOwningAssetEditor->GetCameraAsset())
+	, StandardLayout(new FStandardToolkitLayout(TEXT("CameraAssetEditor_Layout_v2")))
+	, BuildButtonToolkit(MakeShared<FBuildButtonToolkit>(CameraAsset))
+	, BuildLogToolkit(MakeShared<FCameraBuildLogToolkit>())
+{
+	StandardLayout->AddBottomTab(SearchTabId);
+	StandardLayout->AddBottomTab(MessagesTabId);
+
+	TSharedPtr<FLayoutExtender> NewLayoutExtender = MakeShared<FLayoutExtender>();
+	{
+		NewLayoutExtender->ExtendStack(
+				FStandardToolkitLayout::BottomStackExtensionId,
+				ELayoutExtensionPosition::After,
+				FTabManager::FTab(SearchTabId, ETabState::ClosedTab));
+
+		NewLayoutExtender->ExtendStack(
+				FStandardToolkitLayout::BottomStackExtensionId,
+				ELayoutExtensionPosition::After,
+				FTabManager::FTab(MessagesTabId, ETabState::ClosedTab));
+	}
+	LayoutExtenders.Add(NewLayoutExtender);
 }
 
 FCameraAssetEditorToolkit::~FCameraAssetEditorToolkit()
 {
 }
 
+void FCameraAssetEditorToolkit::AddReferencedObjects(FReferenceCollector& Collector)
+{
+	Collector.AddReferencedObject(CameraAsset);
+}
+
+FString FCameraAssetEditorToolkit::GetReferencerName() const
+{
+	return TEXT("FCameraAssetEditorToolkit");
+}
+
 void FCameraAssetEditorToolkit::RegisterTabSpawners(const TSharedRef<FTabManager>& InTabManager)
 {
 	// Skip FBaseAssetToolkit here because we don't want a viewport tab.
 	FAssetEditorToolkit::RegisterTabSpawners(InTabManager);
 
-	InTabManager->RegisterTabSpawner(DetailsViewTabId, FOnSpawnTab::CreateSP(this, &FCameraAssetEditorToolkit::SpawnTab_Details))
-		.SetDisplayName(LOCTEXT("Details", "Details"))
+	const FName CamerasStyleSetName = FGameplayCamerasEditorStyle::Get()->GetStyleSetName();
+
+	InTabManager->RegisterTabSpawner(SearchTabId, FOnSpawnTab::CreateSP(this, &FCameraAssetEditorToolkit::SpawnTab_Search))
+		.SetDisplayName(LOCTEXT("Search", "Search"))
+		.SetGroup(AssetEditorTabsCategory.ToSharedRef())
+		.SetIcon(FSlateIcon(CamerasStyleSetName, "CameraAssetEditor.Tabs.Search"));
+
+	InTabManager->RegisterTabSpawner(MessagesTabId, FOnSpawnTab::CreateSP(this, &FCameraAssetEditorToolkit::SpawnTab_Messages))
+		.SetDisplayName(LOCTEXT("Messages", "Messages"))
 		.SetGroup(AssetEditorTabsCategory.ToSharedRef())
-		.SetIcon(FSlateIcon(FAppStyle::GetAppStyleSetName(), "LevelEditor.Tabs.Details"));
+		.SetIcon(FSlateIcon(CamerasStyleSetName, "CameraAssetEditor.Tabs.Messages"));
+}
+
+TSharedRef<SDockTab> FCameraAssetEditorToolkit::SpawnTab_Search(const FSpawnTabArgs& Args)
+{
+	TSharedPtr<SDockTab> SearchTab = SNew(SDockTab)
+		.Label(LOCTEXT("SearchTabTitle", "Search"))
+		[
+			SearchWidget.ToSharedRef()
+		];
+
+	return SearchTab.ToSharedRef();
+}
+
+TSharedRef<SDockTab> FCameraAssetEditorToolkit::SpawnTab_Messages(const FSpawnTabArgs& Args)
+{
+	TSharedPtr<SDockTab> MessagesTab = SNew(SDockTab)
+		.Label(LOCTEXT("MessagesTabTitle", "Messages"))
+		[
+			BuildLogToolkit->GetMessagesWidget().ToSharedRef()
+		];
+
+	return MessagesTab.ToSharedRef();
 }
 
 void FCameraAssetEditorToolkit::UnregisterTabSpawners(const TSharedRef<FTabManager>& InTabManager)
@@ -61,7 +118,8 @@ void FCameraAssetEditorToolkit::UnregisterTabSpawners(const TSharedRef<FTabManag
 	// Skip FBaseAssetToolkit here because we don't want a viewport tab.
 	FAssetEditorToolkit::UnregisterTabSpawners(InTabManager);
 
-	InTabManager->UnregisterTabSpawner(DetailsViewTabId);
+	InTabManager->UnregisterTabSpawner(SearchTabId);
+	InTabManager->UnregisterTabSpawner(MessagesTabId);
 }
 
 void FCameraAssetEditorToolkit::CreateWidgets()
@@ -69,31 +127,175 @@ void FCameraAssetEditorToolkit::CreateWidgets()
 	// Skip FBaseAssetToolkit here because we don't want a viewport tab.
 	// ...no up-call...
 
-	// Do most of FBaseAssetToolkit's work except for the viewport.
 	RegisterToolbar();
+	CreateEditorModeManager();
 	LayoutExtender = MakeShared<FLayoutExtender>();
+	// We don't want a details view, but we need to because otherwise it crashes.
+	DetailsView = CreateDetailsView();
+
+	// Do our custom stuff.
 
-	FPropertyEditorModule& PropertyEditorModule = FModuleManager::GetModuleChecked<FPropertyEditorModule>("PropertyEditor");
-	FDetailsViewArgs DetailsViewArgs;
-	DetailsViewArgs.NameAreaSettings = FDetailsViewArgs::HideNameArea;
-	DetailsViewArgs.bHideSelectionTip = true;
-	DetailsView = PropertyEditorModule.CreateDetailView(DetailsViewArgs);
+	// Create the search panel.
+	SearchWidget = SNew(SFindInObjectTreeGraph)
+		.OnJumpToNodeRequested(this, &FCameraAssetEditorToolkit::JumpToNode);
 
-	// Now do our custom stuff.
-	
 	// Create the message log.
-	FMessageLogModule& MessageLogModule = FModuleManager::LoadModuleChecked<FMessageLogModule>("MessageLog");
-	FMessageLogInitializationOptions LogOptions;
-	LogOptions.bShowPages = false;
-	LogOptions.bShowFilters = false;
-	LogOptions.bAllowClear = false;
-	LogOptions.MaxPageCount = 1;
-	StatsListing = MessageLogModule.CreateLogListing("CameraAssetEditorStats", LogOptions);
+	BuildLogToolkit->Initialize("CameraAssetBuildMessages");
+}
+
+void FCameraAssetEditorToolkit::JumpToNode(UEdGraphNode* Node)
+{
+}
+
+void FCameraAssetEditorToolkit::JumpToObject(UObject* Object)
+{
+	TSharedPtr<FCameraRigsAssetEditorMode> CameraRigsMode = StaticCastSharedPtr<FCameraRigsAssetEditorMode>(
+			GetEditorMode(FCameraRigsAssetEditorMode::ModeName));
+	if (CameraRigsMode->JumpToObject(Object))
+	{
+		SetEditorMode(FCameraRigsAssetEditorMode::ModeName);
+		return;
+	}
+
+	TSharedPtr<FCameraSharedTransitionsAssetEditorMode> SharedTransitionsMode = StaticCastSharedPtr<FCameraSharedTransitionsAssetEditorMode>(
+			GetEditorMode(FCameraSharedTransitionsAssetEditorMode::ModeName));
+	if (SharedTransitionsMode->JumpToObject(Object))
+	{
+		SetEditorMode(FCameraSharedTransitionsAssetEditorMode::ModeName);
+		return;
+	}
+}
+
+void FCameraAssetEditorToolkit::RegisterToolbar()
+{
+	FName ParentName;
+	const FName MenuName = GetToolMenuToolbarName(ParentName);
+	UToolMenus* ToolMenus = UToolMenus::Get();
+	if (!ToolMenus->IsMenuRegistered(MenuName))
+	{
+		FToolMenuOwnerScoped ToolMenuOwnerScope(this);
 
-	Stats = MessageLogModule.CreateLogListingWidget(StatsListing.ToSharedRef());}
+		UToolMenu* ToolbarMenu = UToolMenus::Get()->RegisterMenu(
+				MenuName, ParentName, EMultiBoxType::ToolBar);
+
+		FToolMenuInsert InsertAfterAssetSection("Asset", EToolMenuInsertType::After);
+		const FCameraAssetEditorCommands& Commands = FCameraAssetEditorCommands::Get();
+
+		ToolbarMenu->AddDynamicSection("Tools", FNewToolMenuDelegate::CreateLambda(
+				[&Commands](UToolMenu* InMenu)
+				{
+					UCameraAssetEditorMenuContext* Context = InMenu->FindContext<UCameraAssetEditorMenuContext>();
+					FCameraAssetEditorToolkit* This = Context ? Context->Toolkit.Pin().Get() : nullptr;
+					if (!ensure(This))
+					{
+						return;
+					}
+
+					FToolMenuSection& ToolsSection = InMenu->AddSection("Tools");
+					ToolsSection.AddEntry(This->BuildButtonToolkit->MakeToolbarButton(Commands.Build));
+					ToolsSection.AddEntry(FToolMenuEntry::InitToolBarButton(Commands.FindInCamera));
+				}),
+				InsertAfterAssetSection);
+
+		FToolMenuSection& ModesSection = ToolbarMenu->AddSection("EditorModes", TAttribute<FText>(), InsertAfterAssetSection);
+		ModesSection.AddEntry(FToolMenuEntry::InitToolBarButton(Commands.ShowCameraDirector));
+		ModesSection.AddEntry(FToolMenuEntry::InitToolBarButton(Commands.ShowCameraRigs));
+		ModesSection.AddEntry(FToolMenuEntry::InitToolBarButton(Commands.ShowSharedTransitions));
+	}
+}
+
+void FCameraAssetEditorToolkit::InitToolMenuContext(FToolMenuContext& MenuContext)
+{
+	FBaseAssetToolkit::InitToolMenuContext(MenuContext);
+
+	UCameraAssetEditorMenuContext* Context = NewObject<UCameraAssetEditorMenuContext>();
+	Context->Toolkit = SharedThis(this);
+	MenuContext.AddObject(Context);
+}
 
 void FCameraAssetEditorToolkit::PostInitAssetEditor()
 {
+	Settings = GetMutableDefault<UGameplayCamerasEditorSettings>();
+
+	const FName CameraDirectorModeName = FCameraDirectorAssetEditorMode::ModeName;
+	AddEditorMode(MakeShared<FCameraDirectorAssetEditorMode>(CameraAsset));
+
+	const FName CameraRigsModeName = FCameraRigsAssetEditorMode::ModeName;
+	AddEditorMode(MakeShared<FCameraRigsAssetEditorMode>(CameraAsset));
+
+	const FName SharedTransitionsModeName = FCameraSharedTransitionsAssetEditorMode::ModeName;
+	AddEditorMode(MakeShared<FCameraSharedTransitionsAssetEditorMode>(CameraAsset));
+
+	const FCameraAssetEditorCommands& Commands = FCameraAssetEditorCommands::Get();
+	TMap<FName, TSharedPtr<FUICommandInfo>> ModeCommands;
+	ModeCommands.Add(CameraDirectorModeName, Commands.ShowCameraDirector);
+	ModeCommands.Add(CameraRigsModeName, Commands.ShowCameraRigs);
+	ModeCommands.Add(SharedTransitionsModeName, Commands.ShowSharedTransitions);
+	for (auto& Pair : ModeCommands)
+	{
+		ToolkitCommands->MapAction(
+			Pair.Value,
+			FExecuteAction::CreateSP(this, &FCameraAssetEditorToolkit::SetEditorMode, Pair.Key),
+			FCanExecuteAction::CreateSP(this, &FCameraAssetEditorToolkit::CanSetEditorMode, Pair.Key),
+			FIsActionChecked::CreateSP(this, &FCameraAssetEditorToolkit::IsEditorMode, Pair.Key));
+	}
+
+	ToolkitCommands->MapAction(
+		Commands.Build,
+		FExecuteAction::CreateSP(this, &FCameraAssetEditorToolkit::OnBuild));
+
+	ToolkitCommands->MapAction(
+		Commands.FindInCamera,
+		FExecuteAction::CreateSP(this, &FCameraAssetEditorToolkit::OnFindInCamera));
+
+	BuildLogToolkit->OnRequestJumpToObject().BindSP(this, &FCameraAssetEditorToolkit::JumpToObject);
+
+	IGameplayCamerasModule& GameplayCamerasModule = FModuleManager::GetModuleChecked<IGameplayCamerasModule>("GameplayCameras");
+	LiveEditManager = GameplayCamerasModule.GetLiveEditManager();
+
+	const FName InitialModeName = !Settings->LastCameraAssetToolkitModeName.IsNone() ? 
+		Settings->LastCameraAssetToolkitModeName : CameraRigsModeName;
+	SetEditorMode(InitialModeName);
+}
+
+void FCameraAssetEditorToolkit::OnEditorToolkitModeActivated()
+{
+	Settings->LastCameraAssetToolkitModeName = GetCurrentEditorModeName();
+	Settings->SaveConfig();
+}
+
+void FCameraAssetEditorToolkit::OnBuild()
+{
+	if (!CameraAsset)
+	{
+		return;
+	}
+
+	FCameraBuildLog BuildLog;
+	CameraAsset->BuildCamera(BuildLog);
+	BuildLogToolkit->PopulateMessageListing(BuildLog);
+
+	if (CameraAsset->BuildStatus != ECameraBuildStatus::Clean)
+	{
+		TabManager->TryInvokeTab(MessagesTabId);
+	}
+
+	for (UCameraRigAsset* CameraRigAsset : CameraAsset->CameraRigs)
+	{
+		FCameraRigPackages BuiltPackages;
+		CameraRigAsset->GatherPackages(BuiltPackages);
+
+		for (const UPackage* BuiltPackage : BuiltPackages)
+		{
+			LiveEditManager->NotifyPostBuildAsset(BuiltPackage);
+		}
+	}
+}
+
+void FCameraAssetEditorToolkit::OnFindInCamera()
+{
+	TabManager->TryInvokeTab(SearchTabId);
+	SearchWidget->FocusSearchEditBox();
 }
 
 FText FCameraAssetEditorToolkit::GetBaseToolkitName() const
diff --git a/Engine/Plugins/Cameras/GameplayCameras/Source/GameplayCamerasEditor/Private/Toolkits/CameraBuildLogToolkit.cpp b/Engine/Plugins/Cameras/GameplayCameras/Source/GameplayCamerasEditor/Private/Toolkits/CameraBuildLogToolkit.cpp
new file mode 100644
--- /dev/null
+++ b/Engine/Plugins/Cameras/GameplayCameras/Source/GameplayCamerasEditor/Private/Toolkits/CameraBuildLogToolkit.cpp
@@ -0,0 +1,85 @@
+// Copyright Epic Games, Inc. All Rights Reserved.
+
+#include "Toolkits/CameraBuildLogToolkit.h"
+
+#include "Core/CameraBuildLog.h"
+#include "IMessageLogListing.h"
+#include "Logging/TokenizedMessage.h"
+#include "MessageLogModule.h"
+#include "Misc/UObjectToken.h"
+#include "Modules/ModuleManager.h"
+
+namespace UE::Cameras
+{
+
+FCameraBuildLogToolkit::FCameraBuildLogToolkit()
+{
+}
+
+void FCameraBuildLogToolkit::Initialize(FName InLogName)
+{
+	if (InLogName.IsNone())
+	{
+		InLogName = "CameraBuildLogMessages";
+	}
+
+	// Create the message log.
+	FMessageLogModule& MessageLogModule = FModuleManager::LoadModuleChecked<FMessageLogModule>("MessageLog");
+	FMessageLogInitializationOptions LogOptions;
+	LogOptions.bShowPages = false;
+	LogOptions.bShowFilters = false;
+	LogOptions.bAllowClear = false;
+	LogOptions.MaxPageCount = 1;
+	MessageListing = MessageLogModule.CreateLogListing(InLogName, LogOptions);
+	MessageListing->OnMessageTokenClicked().AddSP(this, &FCameraBuildLogToolkit::OnMessageTokenClicked);
+
+	// Create the messages widget.
+	MessagesWidget = MessageLogModule.CreateLogListingWidget(MessageListing.ToSharedRef());
+}
+
+void FCameraBuildLogToolkit::PopulateMessageListing(FCameraBuildLog& InBuildLog)
+{
+	if (!ensureMsgf(MessageListing, TEXT("This toolkit wasn't initialized")))
+	{
+		return;
+	}
+
+	MessageListing->ClearMessages();
+
+	for (const FCameraBuildLogMessage& Message : InBuildLog.GetMessages())
+	{
+		TSharedRef<FTokenizedMessage> TokenizedMessage = FTokenizedMessage::Create(Message.Severity);
+
+		if (Message.Object)
+		{
+			TSharedRef<FUObjectToken> ObjectToken = FUObjectToken::Create(
+					Message.Object, FText::FromName(Message.Object->GetFName()));
+
+			// Suppress default activation callback that opens the content browser. We handle that
+			// sort of stuff in a more general fashion in OnMessageTokenClicked.
+			static auto DummyActivation = [](const TSharedRef<class IMessageToken>&) {};
+			ObjectToken->OnMessageTokenActivated(FOnMessageTokenActivated::CreateLambda(DummyActivation));
+			
+			TokenizedMessage->AddToken(ObjectToken);
+		}
+
+		TokenizedMessage->AddToken(FTextToken::Create(Message.Text));
+
+		MessageListing->AddMessage(TokenizedMessage);
+	}
+}
+
+void FCameraBuildLogToolkit::OnMessageTokenClicked(const TSharedRef<IMessageToken>& InMessageToken)
+{
+	if (InMessageToken->GetType() == EMessageToken::Object)
+	{
+		const TSharedRef<FUObjectToken> ObjectToken = StaticCastSharedRef<FUObjectToken>(InMessageToken);
+		if (UObject* Object = ObjectToken->GetObject().Get())
+		{
+			RequestJumpToObject.ExecuteIfBound(Object);
+		}
+	}
+}
+
+}  // namespace UE::Cameras
+
diff --git a/Engine/Plugins/Cameras/GameplayCameras/Source/GameplayCamerasEditor/Private/Toolkits/CameraDirectorAssetEditorMode.cpp b/Engine/Plugins/Cameras/GameplayCameras/Source/GameplayCamerasEditor/Private/Toolkits/CameraDirectorAssetEditorMode.cpp
new file mode 100644
--- /dev/null
+++ b/Engine/Plugins/Cameras/GameplayCameras/Source/GameplayCamerasEditor/Private/Toolkits/CameraDirectorAssetEditorMode.cpp
@@ -0,0 +1,82 @@
+// Copyright Epic Games, Inc. All Rights Reserved.
+
+#include "Toolkits/CameraDirectorAssetEditorMode.h"
+
+#include "Core/CameraAsset.h"
+#include "Core/CameraDirector.h"
+#include "Modules/ModuleManager.h"
+#include "PropertyEditorModule.h"
+#include "ToolMenus.h"
+#include "Toolkits/StandardToolkitLayout.h"
+#include "Widgets/Docking/SDockTab.h"
+
+#define LOCTEXT_NAMESPACE "CameraDirectorAssetEditorMode"
+
+namespace UE::Cameras
+{
+
+FName FCameraDirectorAssetEditorMode::ModeName(TEXT("CameraDirector"));
+
+FName FCameraDirectorAssetEditorMode::DirectorEditorTabId(TEXT("DirectorEditor"));
+
+FCameraDirectorAssetEditorMode::FCameraDirectorAssetEditorMode(UCameraAsset* InCameraAsset)
+	: FAssetEditorMode(ModeName)
+	, CameraAsset(InCameraAsset)
+{
+	StandardLayout = MakeShared<FStandardToolkitLayout>("CameraAssetEditor_Mode_CameraDirector_Layout_v1");
+	{
+		StandardLayout->AddCenterTab(DirectorEditorTabId);
+	}
+
+	DefaultLayout = StandardLayout->GetLayout();
+}
+
+void FCameraDirectorAssetEditorMode::OnActivateMode(const FAssetEditorModeActivateParams& InParams)
+{
+	if (!bInitializedToolkit)
+	{
+		FPropertyEditorModule& PropertyEditorModule = FModuleManager::GetModuleChecked<FPropertyEditorModule>("PropertyEditor");
+		FDetailsViewArgs DetailsViewArgs;
+		DetailsViewArgs.NameAreaSettings = FDetailsViewArgs::HideNameArea;
+		DetailsViewArgs.bHideSelectionTip = true;
+		DetailsView = PropertyEditorModule.CreateDetailView(DetailsViewArgs);
+
+		bInitializedToolkit = true;
+	}
+
+	DetailsView->SetObject(CameraAsset->CameraDirector);
+
+	InParams.TabManager->RegisterTabSpawner(DirectorEditorTabId, FOnSpawnTab::CreateSP(this, &FCameraDirectorAssetEditorMode::SpawnTab_DirectorEditor))
+		.SetDisplayName(LOCTEXT("TransitionEditor", "Camera Transitions"))
+		.SetGroup(InParams.AssetEditorTabsCategory.ToSharedRef());
+}
+
+TSharedRef<SDockTab> FCameraDirectorAssetEditorMode::SpawnTab_DirectorEditor(const FSpawnTabArgs& Args)
+{
+	TSharedPtr<SDockTab> DirectorEditorTab = SNew(SDockTab)
+		.Label(LOCTEXT("TransitionEditorTabTitle", "Camera Transitions"))
+		[
+			DetailsView.ToSharedRef()
+		];
+
+	return DirectorEditorTab.ToSharedRef();
+}
+
+void FCameraDirectorAssetEditorMode::OnDeactivateMode(const FAssetEditorModeDeactivateParams& InParams)
+{
+	InParams.TabManager->UnregisterTabSpawner(DirectorEditorTabId);
+}
+
+bool FCameraDirectorAssetEditorMode::JumpToObject(UObject* InObject)
+{
+	if (InObject == CameraAsset->CameraDirector)
+	{
+		return true;
+	}
+	return false;
+}
+
+}  // namespace UE::Cameras
+
+#undef LOCTEXT_NAMESPACE
+
diff --git a/Engine/Plugins/Cameras/GameplayCameras/Source/GameplayCamerasEditor/Private/Toolkits/CameraRigAssetEditorToolkit.cpp b/Engine/Plugins/Cameras/GameplayCameras/Source/GameplayCamerasEditor/Private/Toolkits/CameraRigAssetEditorToolkit.cpp
--- a/Engine/Plugins/Cameras/GameplayCameras/Source/GameplayCamerasEditor/Private/Toolkits/CameraRigAssetEditorToolkit.cpp
+++ b/Engine/Plugins/Cameras/GameplayCameras/Source/GameplayCamerasEditor/Private/Toolkits/CameraRigAssetEditorToolkit.cpp
@@ -2,29 +2,21 @@
 
 #include "Toolkits/CameraRigAssetEditorToolkit.h"
 
-#include "AssetTools/CameraRigAssetEditor.h"
 #include "Commands/CameraRigAssetEditorCommands.h"
-#include "Core/CameraNode.h"
+#include "Core/CameraBuildLog.h"
 #include "Core/CameraRigAsset.h"
-#include "Core/CameraRigAssetBuilder.h"
-#include "EditorModeManager.h"
-#include "Editors/ObjectTreeGraph.h"
-#include "Editors/ObjectTreeGraphConfig.h"
 #include "Editors/SCameraRigAssetEditor.h"
 #include "Editors/SFindInObjectTreeGraph.h"
-#include "Editors/SObjectTreeGraphEditor.h"
-#include "Editors/SObjectTreeGraphToolbox.h"
 #include "Framework/Docking/LayoutExtender.h"
 #include "Framework/Docking/TabManager.h"
 #include "IGameplayCamerasLiveEditManager.h"
 #include "IGameplayCamerasModule.h"
-#include "IMessageLogListing.h"
-#include "MessageLogInitializationOptions.h"
-#include "MessageLogModule.h"
-#include "Misc/UObjectToken.h"
 #include "Modules/ModuleManager.h"
-#include "PropertyEditorModule.h"
 #include "Styles/GameplayCamerasEditorStyle.h"
+#include "Toolkits/BuildButtonToolkit.h"
+#include "Toolkits/CameraBuildLogToolkit.h"
+#include "Toolkits/CameraRigAssetEditorToolkitBase.h"
+#include "Toolkits/StandardToolkitLayout.h"
 #include "ToolMenus.h"
 #include "Widgets/Docking/SDockTab.h"
 
@@ -35,78 +27,39 @@
 namespace UE::Cameras
 {
 
-const FName FCameraRigAssetEditorToolkit::ToolboxTabId(TEXT("CameraRigAssetEditor_Toolbox"));
-const FName FCameraRigAssetEditorToolkit::CameraRigEditorTabId(TEXT("CameraRigAssetEditor_CameraRigEditor"));
 const FName FCameraRigAssetEditorToolkit::SearchTabId(TEXT("CameraRigAssetEditor_Search"));
 const FName FCameraRigAssetEditorToolkit::MessagesTabId(TEXT("CameraRigAssetEditor_Messages"));
-const FName FCameraRigAssetEditorToolkit::DetailsViewTabId(TEXT("CameraRigAssetEditor_DetailsView"));
 
-FCameraRigAssetEditorToolkit::FCameraRigAssetEditorToolkit(UCameraRigAssetEditor* InOwningAssetEditor)
+FCameraRigAssetEditorToolkit::FCameraRigAssetEditorToolkit(UAssetEditor* InOwningAssetEditor)
 	: FBaseAssetToolkit(InOwningAssetEditor)
-	, CameraRigAsset(InOwningAssetEditor->GetCameraRigAsset())
-	, CommandBindings(new FUICommandList())
 {
+	Impl = MakeShared<FCameraRigAssetEditorToolkitBase>(TEXT("CameraRigAssetEditor_Layout_v6"));
+	BuildButtonToolkit = MakeShared<FBuildButtonToolkit>();
+	BuildLogToolkit = MakeShared<FCameraBuildLogToolkit>();
+
 	// Override base class default layout.
-	StandaloneDefaultLayout = FTabManager::NewLayout("CameraRigAssetEditor_Layout_v5")
-		->AddArea
-		(
-			FTabManager::NewPrimaryArea()->SetOrientation(Orient_Vertical)
-			->Split
-			(
-				FTabManager::NewSplitter()->SetOrientation(Orient_Horizontal)
-				->Split
-				(
-					FTabManager::NewStack()
-					->SetSizeCoefficient(0.2f)
-					->AddTab(ToolboxTabId, ETabState::OpenedTab)
-				)
-				->Split
-				(
-					FTabManager::NewStack()
-					->SetSizeCoefficient(0.6f)
-					->AddTab(CameraRigEditorTabId, ETabState::OpenedTab)
-					->SetHideTabWell(true)
-				)
-				->Split
-				(
-					FTabManager::NewStack()
-					->SetSizeCoefficient(0.2f)
-					->AddTab(DetailsViewTabId, ETabState::OpenedTab)
-				)
-			)
-			->Split
-			(
-				FTabManager::NewStack()
-				->SetSizeCoefficient(0.2f)
-				->AddTab(SearchTabId, ETabState::ClosedTab)
-				->AddTab(MessagesTabId, ETabState::ClosedTab)
-			)
-		);
+	TSharedPtr<FStandardToolkitLayout> StandardLayout = Impl->GetStandardLayout();
+	{
+		StandardLayout->AddBottomTab(SearchTabId);
+		StandardLayout->AddBottomTab(MessagesTabId);
+	}
+	StandaloneDefaultLayout = StandardLayout->GetLayout();
 }
 
-FCameraRigAssetEditorToolkit::~FCameraRigAssetEditorToolkit()
+void FCameraRigAssetEditorToolkit::SetCameraRigAsset(UCameraRigAsset* InCameraRig)
 {
-	if (CameraRigEditorWidget)
-	{
-		CameraRigEditorWidget->RemoveOnAnyGraphChanged(this);
-	}
+	Impl->SetCameraRigAsset(InCameraRig);
+	BuildButtonToolkit->SetTarget(InCameraRig);
 }
 
-void FCameraRigAssetEditorToolkit::RegisterTabSpawners(const TSharedRef<FTabManager>& InTabManager)
+void FCameraRigAssetEditorToolkit::RegisterTabSpawners(const TSharedRef<class FTabManager>& InTabManager)
 {
 	// Skip FBaseAssetToolkit here because we don't want a viewport tab.
 	FAssetEditorToolkit::RegisterTabSpawners(InTabManager);
 
-	const FName CamerasStyleSetName = FGameplayCamerasEditorStyle::Get()->GetStyleSetName();
-
-	InTabManager->RegisterTabSpawner(ToolboxTabId, FOnSpawnTab::CreateSP(this, &FCameraRigAssetEditorToolkit::SpawnTab_Toolbox))
-		.SetDisplayName(LOCTEXT("Toolbox", "Toolbox"))
-		.SetGroup(AssetEditorTabsCategory.ToSharedRef())
-		.SetIcon(FSlateIcon(CamerasStyleSetName, "CameraRigAssetEditor.Tabs.Toolbox"));
+	Impl->RegisterTabSpawners(InTabManager, AssetEditorTabsCategory);
 
-	InTabManager->RegisterTabSpawner(CameraRigEditorTabId, FOnSpawnTab::CreateSP(this, &FCameraRigAssetEditorToolkit::SpawnTab_CameraRigEditor))
-		.SetDisplayName(LOCTEXT("CameraRigEditor", "Camera Rig"))
-		.SetGroup(AssetEditorTabsCategory.ToSharedRef());
+	const FName CamerasStyleSetName = FGameplayCamerasEditorStyle::Get()->GetStyleSetName();
 
 	InTabManager->RegisterTabSpawner(SearchTabId, FOnSpawnTab::CreateSP(this, &FCameraRigAssetEditorToolkit::SpawnTab_Search))
 		.SetDisplayName(LOCTEXT("Search", "Search"))
@@ -117,33 +70,6 @@ void FCameraRigAssetEditorToolkit::RegisterTabSpawners(const TSharedRef<FTabMana
 		.SetDisplayName(LOCTEXT("Messages", "Messages"))
 		.SetGroup(AssetEditorTabsCategory.ToSharedRef())
 		.SetIcon(FSlateIcon(CamerasStyleSetName, "CameraRigAssetEditor.Tabs.Messages"));
-
-	InTabManager->RegisterTabSpawner(DetailsViewTabId, FOnSpawnTab::CreateSP(this, &FCameraRigAssetEditorToolkit::SpawnTab_Details))
-		.SetDisplayName(LOCTEXT("Details", "Details"))
-		.SetGroup(AssetEditorTabsCategory.ToSharedRef())
-		.SetIcon(FSlateIcon(FAppStyle::GetAppStyleSetName(), "LevelEditor.Tabs.Details"));
-}
-
-TSharedRef<SDockTab> FCameraRigAssetEditorToolkit::SpawnTab_Toolbox(const FSpawnTabArgs& Args)
-{
-	TSharedPtr<SDockTab> ToolboxTab = SNew(SDockTab)
-		.Label(LOCTEXT("ToolboxTabTitle", "Toolbox"))
-		[
-			ToolboxWidget.ToSharedRef()
-		];
-
-	return ToolboxTab.ToSharedRef();
-}
-
-TSharedRef<SDockTab> FCameraRigAssetEditorToolkit::SpawnTab_CameraRigEditor(const FSpawnTabArgs& Args)
-{
-	TSharedPtr<SDockTab> CameraRigEditorTab = SNew(SDockTab)
-		.Label(LOCTEXT("CameraRigEditorTabTitle", "Camera Rig Editor"))
-		[
-			CameraRigEditorWidget.ToSharedRef()
-		];
-
-	return CameraRigEditorTab.ToSharedRef();
 }
 
 TSharedRef<SDockTab> FCameraRigAssetEditorToolkit::SpawnTab_Search(const FSpawnTabArgs& Args)
@@ -162,7 +88,7 @@ TSharedRef<SDockTab> FCameraRigAssetEditorToolkit::SpawnTab_Messages(const FSpaw
 	TSharedPtr<SDockTab> MessagesTab = SNew(SDockTab)
 		.Label(LOCTEXT("MessagesTabTitle", "Messages"))
 		[
-			MessagesWidget.ToSharedRef()
+			BuildLogToolkit->GetMessagesWidget().ToSharedRef()
 		];
 
 	return MessagesTab.ToSharedRef();
@@ -172,44 +98,33 @@ void FCameraRigAssetEditorToolkit::UnregisterTabSpawners(const TSharedRef<FTabMa
 {
 	// Skip FBaseAssetToolkit here because we don't want a viewport tab.
 	FAssetEditorToolkit::UnregisterTabSpawners(InTabManager);
+	
+	Impl->UnregisterTabSpawners(InTabManager);
 
-	InTabManager->UnregisterTabSpawner(ToolboxTabId);
-	InTabManager->UnregisterTabSpawner(CameraRigEditorTabId);
 	InTabManager->UnregisterTabSpawner(SearchTabId);
 	InTabManager->UnregisterTabSpawner(MessagesTabId);
-	InTabManager->UnregisterTabSpawner(DetailsViewTabId);
 }
 
 void FCameraRigAssetEditorToolkit::CreateWidgets()
 {
-	// Skip FBaseAssetToolkit here because we don't want a viewport tab.
+	// Skip FBaseAssetToolkit here because we don't want a viewport tab, and our base class
+	// has its own details view in order to get a notify hook.
 	// ...no up-call...
 
-	// Do most of FBaseAssetToolkit's work except for the viewport.
 	RegisterToolbar();
+	CreateEditorModeManager();
 	LayoutExtender = MakeShared<FLayoutExtender>();
 
-	FPropertyEditorModule& PropertyEditorModule = FModuleManager::GetModuleChecked<FPropertyEditorModule>("PropertyEditor");
-	FDetailsViewArgs DetailsViewArgs;
-	DetailsViewArgs.NameAreaSettings = FDetailsViewArgs::HideNameArea;
-	DetailsViewArgs.bHideSelectionTip = true;
-	DetailsViewArgs.NotifyHook = this;
-	DetailsView = PropertyEditorModule.CreateDetailView(DetailsViewArgs);
-
 	// Now do our custom stuff.
 
-	// Create the camera rig editor.
-	CameraRigEditorWidget = SNew(SCameraRigAssetEditor)
-		.DetailsView(DetailsView)
-		.CameraRigAsset(CameraRigAsset);
-	CameraRigEditorWidget->AddOnAnyGraphChanged(FOnGraphChanged::FDelegate::CreateSP(
-				this, &FCameraRigAssetEditorToolkit::OnAnyGraphChanged));
+	Impl->CreateWidgets();
 
-	// Create the toolbox, default to the rig editor items.
-	ToolboxWidget = SNew(SObjectTreeGraphToolbox)
-		.GraphConfig(CameraRigEditorWidget->GetFocusedGraphConfig());
+	// We need to set this for our FBaseAssetToolkit parent because otherwise it crashes
+	// unhappily in SetObjectsToEdit.
+	DetailsView = Impl->GetDetailsView();
 
 	// Create the search panel.
+	TSharedPtr<SCameraRigAssetEditor> CameraRigEditorWidget = Impl->GetCameraRigAssetEditor();
 	TArray<UEdGraph*> CameraRigGraphs;
 	CameraRigEditorWidget->GetGraphs(CameraRigGraphs);
 	SearchWidget = SNew(SFindInObjectTreeGraph)
@@ -219,16 +134,7 @@ void FCameraRigAssetEditorToolkit::CreateWidgets()
 					CameraRigEditorWidget.ToSharedRef(), &SCameraRigAssetEditor::JumpToNode));
 
 	// Create the message log.
-	FMessageLogModule& MessageLogModule = FModuleManager::LoadModuleChecked<FMessageLogModule>("MessageLog");
-	FMessageLogInitializationOptions LogOptions;
-	LogOptions.bShowPages = false;
-	LogOptions.bShowFilters = false;
-	LogOptions.bAllowClear = false;
-	LogOptions.MaxPageCount = 1;
-	MessageListing = MessageLogModule.CreateLogListing("CameraRigAssetEditorStats", LogOptions);
-	MessageListing->OnMessageTokenClicked().AddSP(this, &FCameraRigAssetEditorToolkit::OnMessageTokenClicked);
-
-	MessagesWidget = MessageLogModule.CreateLogListingWidget(MessageListing.ToSharedRef());
+	BuildLogToolkit->Initialize("CameraRigAssetBuildMessages");
 }
 
 void FCameraRigAssetEditorToolkit::RegisterToolbar()
@@ -239,35 +145,30 @@ void FCameraRigAssetEditorToolkit::RegisterToolbar()
 	if (!ToolMenus->IsMenuRegistered(MenuName))
 	{
 		FToolMenuOwnerScoped ToolMenuOwnerScope(this);
+		FToolMenuInsert InsertAfterAssetSection("Asset", EToolMenuInsertType::After);
 
 		UToolMenu* ToolbarMenu = UToolMenus::Get()->RegisterMenu(
 				MenuName, ParentName, EMultiBoxType::ToolBar);
 
-		FToolMenuInsert InsertAfterAssetSection("Asset", EToolMenuInsertType::After);
-		const FCameraRigAssetEditorCommands& Commands = FCameraRigAssetEditorCommands::Get();
-
 		ToolbarMenu->AddDynamicSection("Tools", FNewToolMenuDelegate::CreateLambda(
-				[&Commands](UToolMenu* InMenu)
+				[](UToolMenu* InMenu)
 				{
 					UCameraRigAssetEditorMenuContext* Context = InMenu->FindContext<UCameraRigAssetEditorMenuContext>();
-					FCameraRigAssetEditorToolkit* This = Context->CameraRigAssetEditorToolkit.Pin().Get();
+					FCameraRigAssetEditorToolkit* This = Context ? Context->Toolkit.Pin().Get() : nullptr;
+					if (!ensure(This))
+					{
+						return;
+					}
 
-					FToolMenuSection& ToolsSection = InMenu->AddSection("Tools");
-
-					FToolMenuEntry BuildButton = FToolMenuEntry::InitToolBarButton(Commands.Build);
-					BuildButton.Icon = TAttribute<FSlateIcon>(This, &FCameraRigAssetEditorToolkit::GetBuildButtonIcon);
-					BuildButton.ToolTip = TAttribute<FText>(This, &FCameraRigAssetEditorToolkit::GetBuildButtonTooltip);
-					ToolsSection.AddEntry(BuildButton);
+					const FCameraRigAssetEditorCommands& Commands = FCameraRigAssetEditorCommands::Get();
 
+					FToolMenuSection& ToolsSection = InMenu->AddSection("Tools");
+					ToolsSection.AddEntry(This->BuildButtonToolkit->MakeToolbarButton(Commands.Build));
 					ToolsSection.AddEntry(FToolMenuEntry::InitToolBarButton(Commands.FindInCameraRig));
 				}),
 				InsertAfterAssetSection);
 
-		FToolMenuSection& GraphsSection = ToolbarMenu->AddSection("Graphs", TAttribute<FText>(), InsertAfterAssetSection);
-
-		GraphsSection.AddEntry(FToolMenuEntry::InitToolBarButton(Commands.FocusHome));
-		GraphsSection.AddEntry(FToolMenuEntry::InitToolBarButton(Commands.ShowNodeHierarchy));
-		GraphsSection.AddEntry(FToolMenuEntry::InitToolBarButton(Commands.ShowTransitions));
+		Impl->BuildToolbarMenu(ToolbarMenu);
 	}
 }
 
@@ -276,62 +177,14 @@ void FCameraRigAssetEditorToolkit::InitToolMenuContext(FToolMenuContext& MenuCon
 	FBaseAssetToolkit::InitToolMenuContext(MenuContext);
 
 	UCameraRigAssetEditorMenuContext* Context = NewObject<UCameraRigAssetEditorMenuContext>();
-	Context->CameraRigAssetEditorToolkit = SharedThis(this);
+	Context->Toolkit = SharedThis(this);
 	MenuContext.AddObject(Context);
 }
 
-FSlateIcon FCameraRigAssetEditorToolkit::GetBuildButtonIcon() const
-{
-	static const FName BuildStatusBackground("CameraRigAssetEditor.BuildStatus.Background");
-	static const FName BuildStatusError("CameraRigAssetEditor.BuildStatus.Overlay.Error");
-	static const FName BuildStatusGood("CameraRigAssetEditor.BuildStatus.Overlay.Good");
-	static const FName BuildStatusUnknown("CameraRigAssetEditor.BuildStatus.Overlay.Unknown");
-	static const FName BuildStatusWarning("CameraRigAssetEditor.BuildStatus.Overlay.Warning");
-
-	const FName CamerasStyleSetName = FGameplayCamerasEditorStyle::Get()->GetStyleSetName();
-
-	if (!CameraRigAsset)
-	{
-		return FSlateIcon(CamerasStyleSetName, BuildStatusBackground, NAME_None, BuildStatusError);
-	}
-
-	switch (CameraRigAsset->BuildStatus)
-	{
-		default:
-		case ECameraBuildStatus::Dirty:
-			return FSlateIcon(CamerasStyleSetName, BuildStatusBackground, NAME_None, BuildStatusUnknown);
-		case ECameraBuildStatus::WithErrors:
-			return FSlateIcon(CamerasStyleSetName, BuildStatusBackground, NAME_None, BuildStatusError);
-		case ECameraBuildStatus::Clean:
-			return FSlateIcon(CamerasStyleSetName, BuildStatusBackground, NAME_None, BuildStatusGood);
-		case ECameraBuildStatus::CleanWithWarnings:
-			return FSlateIcon(CamerasStyleSetName, BuildStatusBackground, NAME_None, BuildStatusWarning);
-	}
-}
-
-FText FCameraRigAssetEditorToolkit::GetBuildButtonTooltip() const
-{
-	if (!CameraRigAsset)
-	{
-		return LOCTEXT("BuildButtonStatusNoAsset", "No asset is open");
-	}
-
-	switch (CameraRigAsset->BuildStatus)
-	{
-		default:
-		case ECameraBuildStatus::Dirty:
-			return LOCTEXT("BuildButtonStatusDirty", "Dirty or unknown, should rebuild");
-		case ECameraBuildStatus::WithErrors:
-			return LOCTEXT("BuildButtonStatusWithErrors", "There were errors during the build, see the log window for details");
-		case ECameraBuildStatus::Clean:
-			return LOCTEXT("BuildButtonStatusClean", "Good to go");
-		case ECameraBuildStatus::CleanWithWarnings:
-			return LOCTEXT("BuildButtonStatusCleanWithWarnings", "There were warnings during the build, see the log window for details");
-	}
-}
-
 void FCameraRigAssetEditorToolkit::PostInitAssetEditor()
 {
+	Impl->BindCommands(ToolkitCommands);
+
 	const FCameraRigAssetEditorCommands& Commands = FCameraRigAssetEditorCommands::Get();
 
 	ToolkitCommands->MapAction(
@@ -342,78 +195,27 @@ void FCameraRigAssetEditorToolkit::PostInitAssetEditor()
 		Commands.FindInCameraRig,
 		FExecuteAction::CreateSP(this, &FCameraRigAssetEditorToolkit::OnFindInCameraRig));
 
-	TSharedRef<SCameraRigAssetEditor> CameraRigEditor = CameraRigEditorWidget.ToSharedRef();
-
-	ToolkitCommands->MapAction(
-		Commands.FocusHome,
-		FExecuteAction::CreateSP(CameraRigEditor, &SCameraRigAssetEditor::FocusHome));
-
-	ToolkitCommands->MapAction(
-		Commands.ShowNodeHierarchy,
-		FExecuteAction::CreateSP(CameraRigEditor, &SCameraRigAssetEditor::SetEditorMode, ECameraRigAssetEditorMode::NodeGraph),
-		FCanExecuteAction(),
-		FIsActionChecked::CreateSP(CameraRigEditor, &SCameraRigAssetEditor::IsEditorMode, ECameraRigAssetEditorMode::NodeGraph));
-
-	ToolkitCommands->MapAction(
-		Commands.ShowTransitions,
-		FExecuteAction::CreateSP(CameraRigEditor, &SCameraRigAssetEditor::SetEditorMode, ECameraRigAssetEditorMode::TransitionGraph),
-		FCanExecuteAction(),
-		FIsActionChecked::CreateSP(CameraRigEditor, &SCameraRigAssetEditor::IsEditorMode, ECameraRigAssetEditorMode::TransitionGraph));
+	BuildLogToolkit->OnRequestJumpToObject().BindSPLambda(this, [this](UObject* Object)
+		{
+			TSharedPtr<SCameraRigAssetEditor> CameraRigEditorWidget = Impl->GetCameraRigAssetEditor();
+			CameraRigEditorWidget->FindAndJumpToObjectNode(Object);
+		});
 
 	IGameplayCamerasModule& GameplayCamerasModule = FModuleManager::GetModuleChecked<IGameplayCamerasModule>("GameplayCameras");
 	LiveEditManager = GameplayCamerasModule.GetLiveEditManager();
 }
 
-FText FCameraRigAssetEditorToolkit::GetBaseToolkitName() const
-{
-	return LOCTEXT("AppLabel", "Camera Rig Asset");
-}
-
-FName FCameraRigAssetEditorToolkit::GetToolkitFName() const
-{
-	static FName ToolkitName("CameraRigAssetEditor");
-	return ToolkitName;
-}
-
-FString FCameraRigAssetEditorToolkit::GetWorldCentricTabPrefix() const
-{
-	return LOCTEXT("WorldCentricTabPrefix", "Camera Rig Asset ").ToString();
-}
-
-FLinearColor FCameraRigAssetEditorToolkit::GetWorldCentricTabColorScale() const
-{
-	return FLinearColor(0.7f, 0.0f, 0.0f, 0.5f);
-}
-
-void FCameraRigAssetEditorToolkit::NotifyPostChange(const FPropertyChangedEvent& PropertyChangedEvent, FProperty* PropertyThatChanged)
-{
-	CameraRigAsset->BuildStatus = ECameraBuildStatus::Dirty;
-}
-
-void FCameraRigAssetEditorToolkit::OnAnyGraphChanged(const FEdGraphEditAction& InEditAction)
-{
-	CameraRigAsset->BuildStatus = ECameraBuildStatus::Dirty;
-}
-
-void FCameraRigAssetEditorToolkit::OnMessageTokenClicked(const TSharedRef<IMessageToken>& InMessageToken)
+void FCameraRigAssetEditorToolkit::OnBuild()
 {
-	if (InMessageToken->GetType() == EMessageToken::Object)
+	UCameraRigAsset* CameraRigAsset = Impl->GetCameraRigAsset();
+	if (!CameraRigAsset)
 	{
-		const TSharedRef<FUObjectToken> ObjectToken = StaticCastSharedRef<FUObjectToken>(InMessageToken);
-		if (UObject* Object = ObjectToken->GetObject().Get())
-		{
-			CameraRigEditorWidget->FindAndJumpToObjectNode(Object);
-		}
+		return;
 	}
-}
-
-void FCameraRigAssetEditorToolkit::OnBuild()
-{
-	using namespace UE::Cameras;
 
 	FCameraBuildLog BuildLog;
 	CameraRigAsset->BuildCameraRig(BuildLog);
-	GetBuildMessageLogListing(BuildLog);
+	BuildLogToolkit->PopulateMessageListing(BuildLog);
 
 	if (CameraRigAsset->BuildStatus != ECameraBuildStatus::Clean)
 	{
@@ -429,41 +231,31 @@ void FCameraRigAssetEditorToolkit::OnBuild()
 	}
 }
 
-void FCameraRigAssetEditorToolkit::GetBuildMessageLogListing(FCameraBuildLog& InBuildLog)
+void FCameraRigAssetEditorToolkit::OnFindInCameraRig()
 {
-	if (MessageListing.IsValid())
-	{
-		MessageListing->ClearMessages();
-	}
-
-	for (const FCameraBuildLogMessage& Message : InBuildLog.GetMessages())
-	{
-		TSharedRef<FTokenizedMessage> TokenizedMessage = FTokenizedMessage::Create(Message.Severity);
-
-		if (Message.Object)
-		{
-			TSharedRef<FUObjectToken> ObjectToken = FUObjectToken::Create(
-					Message.Object, FText::FromName(Message.Object->GetFName()));
-
-			// Suppress default activation callback that opens the content browser. We handle that
-			// sort of stuff in a more general fashion in OnMessageTokenClicked.
-			static auto DummyActivation = [](const TSharedRef<class IMessageToken>&) {};
-			ObjectToken->OnMessageTokenActivated(FOnMessageTokenActivated::CreateLambda(DummyActivation));
-			
-			TokenizedMessage->AddToken(ObjectToken);
-		}
+	TabManager->TryInvokeTab(SearchTabId);
+	SearchWidget->FocusSearchEditBox();
+}
 
-		TokenizedMessage->AddToken(FTextToken::Create(Message.Text));
+FText FCameraRigAssetEditorToolkit::GetBaseToolkitName() const
+{
+	return LOCTEXT("AppLabel", "Camera Rig Asset");
+}
 
-		MessageListing->AddMessage(TokenizedMessage);
-	}
+FName FCameraRigAssetEditorToolkit::GetToolkitFName() const
+{
+	static FName ToolkitName("CameraRigAssetEditor");
+	return ToolkitName;
 }
 
+FString FCameraRigAssetEditorToolkit::GetWorldCentricTabPrefix() const
+{
+	return LOCTEXT("WorldCentricTabPrefix", "Camera Rig Asset ").ToString();
+}
 
-void FCameraRigAssetEditorToolkit::OnFindInCameraRig()
+FLinearColor FCameraRigAssetEditorToolkit::GetWorldCentricTabColorScale() const
 {
-	TabManager->TryInvokeTab(SearchTabId);
-	SearchWidget->FocusSearchEditBox();
+	return FLinearColor(0.7f, 0.0f, 0.0f, 0.5f);
 }
 
 }  // namespace UE::Cameras
diff --git a/Engine/Plugins/Cameras/GameplayCameras/Source/GameplayCamerasEditor/Private/Toolkits/CameraRigAssetEditorToolkitBase.cpp b/Engine/Plugins/Cameras/GameplayCameras/Source/GameplayCamerasEditor/Private/Toolkits/CameraRigAssetEditorToolkitBase.cpp
new file mode 100644
--- /dev/null
+++ b/Engine/Plugins/Cameras/GameplayCameras/Source/GameplayCamerasEditor/Private/Toolkits/CameraRigAssetEditorToolkitBase.cpp
@@ -0,0 +1,201 @@
+// Copyright Epic Games, Inc. All Rights Reserved.
+
+#include "Toolkits/CameraRigAssetEditorToolkitBase.h"
+
+#include "Commands/CameraRigAssetEditorCommands.h"
+#include "Core/CameraRigAsset.h"
+#include "Editors/SCameraRigAssetEditor.h"
+#include "Editors/SObjectTreeGraphToolbox.h"
+#include "Framework/Commands/UICommandList.h"
+#include "Framework/Docking/TabManager.h"
+#include "Modules/ModuleManager.h"
+#include "PropertyEditorModule.h"
+#include "Styles/GameplayCamerasEditorStyle.h"
+#include "ToolMenu.h"
+#include "Toolkits/StandardToolkitLayout.h"
+#include "Widgets/Docking/SDockTab.h"
+
+#define LOCTEXT_NAMESPACE "CameraRigAssetEditorToolkitBase"
+
+namespace UE::Cameras
+{
+
+const FName FCameraRigAssetEditorToolkitBase::ToolboxTabId(TEXT("CameraRigAssetEditor_Toolbox"));
+const FName FCameraRigAssetEditorToolkitBase::CameraRigEditorTabId(TEXT("CameraRigAssetEditor_CameraRigEditor"));
+const FName FCameraRigAssetEditorToolkitBase::DetailsViewTabId(TEXT("CameraRigAssetEditor_DetailsView"));
+
+FCameraRigAssetEditorToolkitBase::FCameraRigAssetEditorToolkitBase(FName InLayoutName)
+	: StandardLayout(new FStandardToolkitLayout(InLayoutName))
+	, CommandBindings(new FUICommandList())
+{
+	StandardLayout->AddLeftTab(ToolboxTabId);
+	StandardLayout->AddCenterTab(CameraRigEditorTabId);
+	StandardLayout->AddRightTab(DetailsViewTabId);
+}
+
+FCameraRigAssetEditorToolkitBase::~FCameraRigAssetEditorToolkitBase()
+{
+	if (CameraRigEditorWidget)
+	{
+		CameraRigEditorWidget->RemoveOnAnyGraphChanged(this);
+	}
+}
+
+void FCameraRigAssetEditorToolkitBase::AddReferencedObjects(FReferenceCollector& Collector)
+{
+	Collector.AddReferencedObject(CameraRigAsset);
+}
+
+FString FCameraRigAssetEditorToolkitBase::GetReferencerName() const
+{
+	return TEXT("FCameraRigAssetEditorToolkitBase");
+}
+
+void FCameraRigAssetEditorToolkitBase::SetCameraRigAsset(UCameraRigAsset* InCameraRig)
+{
+	CameraRigAsset = InCameraRig;
+
+	if (CameraRigEditorWidget)
+	{
+		CameraRigEditorWidget->SetCameraRigAsset(InCameraRig);
+	}
+}
+
+void FCameraRigAssetEditorToolkitBase::RegisterTabSpawners(TSharedRef<FTabManager> InTabManager, TSharedPtr<FWorkspaceItem> InAssetEditorTabsCategory)
+{
+	const FName CamerasStyleSetName = FGameplayCamerasEditorStyle::Get()->GetStyleSetName();
+
+	InTabManager->RegisterTabSpawner(ToolboxTabId, FOnSpawnTab::CreateSP(this, &FCameraRigAssetEditorToolkitBase::SpawnTab_Toolbox))
+		.SetDisplayName(LOCTEXT("Toolbox", "Toolbox"))
+		.SetGroup(InAssetEditorTabsCategory.ToSharedRef())
+		.SetIcon(FSlateIcon(CamerasStyleSetName, "CameraRigAssetEditor.Tabs.Toolbox"));
+
+	InTabManager->RegisterTabSpawner(CameraRigEditorTabId, FOnSpawnTab::CreateSP(this, &FCameraRigAssetEditorToolkitBase::SpawnTab_CameraRigEditor))
+		.SetDisplayName(LOCTEXT("CameraRigEditor", "Camera Rig"))
+		.SetGroup(InAssetEditorTabsCategory.ToSharedRef());
+
+	InTabManager->RegisterTabSpawner(DetailsViewTabId, FOnSpawnTab::CreateSP(this, &FCameraRigAssetEditorToolkitBase::SpawnTab_Details))
+		.SetDisplayName(LOCTEXT("Details", "Details"))
+		.SetGroup(InAssetEditorTabsCategory.ToSharedRef())
+		.SetIcon(FSlateIcon(FAppStyle::GetAppStyleSetName(), "LevelEditor.Tabs.Details"));
+}
+
+TSharedRef<SDockTab> FCameraRigAssetEditorToolkitBase::SpawnTab_Toolbox(const FSpawnTabArgs& Args)
+{
+	TSharedPtr<SDockTab> ToolboxTab = SNew(SDockTab)
+		.Label(LOCTEXT("ToolboxTabTitle", "Toolbox"))
+		[
+			ToolboxWidget.ToSharedRef()
+		];
+
+	return ToolboxTab.ToSharedRef();
+}
+
+TSharedRef<SDockTab> FCameraRigAssetEditorToolkitBase::SpawnTab_CameraRigEditor(const FSpawnTabArgs& Args)
+{
+	TSharedPtr<SDockTab> CameraRigEditorTab = SNew(SDockTab)
+		.Label(LOCTEXT("CameraRigEditorTabTitle", "Camera Rig Editor"))
+		[
+			CameraRigEditorWidget.ToSharedRef()
+		];
+
+	return CameraRigEditorTab.ToSharedRef();
+}
+
+TSharedRef<SDockTab> FCameraRigAssetEditorToolkitBase::SpawnTab_Details(const FSpawnTabArgs& Args)
+{
+	TSharedPtr<SDockTab> DetailsTab = SNew(SDockTab)
+		.Label(LOCTEXT("DetailsTitle", "Details"))
+		[
+			DetailsView.ToSharedRef()
+		];
+
+	return DetailsTab.ToSharedRef();
+}
+
+void FCameraRigAssetEditorToolkitBase::UnregisterTabSpawners(TSharedRef<FTabManager> InTabManager)
+{
+	InTabManager->UnregisterTabSpawner(ToolboxTabId);
+	InTabManager->UnregisterTabSpawner(CameraRigEditorTabId);
+	InTabManager->UnregisterTabSpawner(DetailsViewTabId);
+}
+
+void FCameraRigAssetEditorToolkitBase::CreateWidgets()
+{
+	// Create the details view.
+	FPropertyEditorModule& PropertyEditorModule = FModuleManager::GetModuleChecked<FPropertyEditorModule>("PropertyEditor");
+	FDetailsViewArgs DetailsViewArgs;
+	DetailsViewArgs.NameAreaSettings = FDetailsViewArgs::HideNameArea;
+	DetailsViewArgs.bHideSelectionTip = true;
+	DetailsViewArgs.NotifyHook = this;
+	DetailsView = PropertyEditorModule.CreateDetailView(DetailsViewArgs);
+
+	// Create the camera rig editor.
+	CameraRigEditorWidget = SNew(SCameraRigAssetEditor)
+		.DetailsView(DetailsView)
+		.CameraRigAsset(CameraRigAsset);
+	CameraRigEditorWidget->AddOnAnyGraphChanged(FOnGraphChanged::FDelegate::CreateSP(
+				this, &FCameraRigAssetEditorToolkitBase::OnAnyGraphChanged));
+
+	// Create the toolbox, default to the rig editor items.
+	ToolboxWidget = SNew(SObjectTreeGraphToolbox)
+		.GraphConfig(CameraRigEditorWidget->GetFocusedGraphConfig());
+}
+
+void FCameraRigAssetEditorToolkitBase::BuildToolbarMenu(UToolMenu* ToolbarMenu)
+{
+	FToolMenuInsert InsertAfterAssetSection("Asset", EToolMenuInsertType::After);
+	const FCameraRigAssetEditorCommands& Commands = FCameraRigAssetEditorCommands::Get();
+
+	FToolMenuSection& GraphsSection = ToolbarMenu->AddSection("Graphs", TAttribute<FText>(), InsertAfterAssetSection);
+
+	GraphsSection.AddEntry(FToolMenuEntry::InitToolBarButton(Commands.FocusHome));
+	GraphsSection.AddEntry(FToolMenuEntry::InitToolBarButton(Commands.ShowNodeHierarchy));
+	GraphsSection.AddEntry(FToolMenuEntry::InitToolBarButton(Commands.ShowTransitions));
+}
+
+void FCameraRigAssetEditorToolkitBase::BindCommands(TSharedRef<FUICommandList> CommandList)
+{
+	const FCameraRigAssetEditorCommands& Commands = FCameraRigAssetEditorCommands::Get();
+
+	TSharedRef<SCameraRigAssetEditor> CameraRigEditor = CameraRigEditorWidget.ToSharedRef();
+
+	CommandList->MapAction(
+		Commands.FocusHome,
+		FExecuteAction::CreateSP(CameraRigEditor, &SCameraRigAssetEditor::FocusHome));
+
+	CommandList->MapAction(
+		Commands.ShowNodeHierarchy,
+		FExecuteAction::CreateSP(CameraRigEditor, &SCameraRigAssetEditor::SetEditorMode, ECameraRigAssetEditorMode::NodeGraph),
+		FCanExecuteAction(),
+		FIsActionChecked::CreateSP(CameraRigEditor, &SCameraRigAssetEditor::IsEditorMode, ECameraRigAssetEditorMode::NodeGraph));
+
+	CommandList->MapAction(
+		Commands.ShowTransitions,
+		FExecuteAction::CreateSP(CameraRigEditor, &SCameraRigAssetEditor::SetEditorMode, ECameraRigAssetEditorMode::TransitionGraph),
+		FCanExecuteAction(),
+		FIsActionChecked::CreateSP(CameraRigEditor, &SCameraRigAssetEditor::IsEditorMode, ECameraRigAssetEditorMode::TransitionGraph));
+}
+
+void FCameraRigAssetEditorToolkitBase::NotifyPostChange(const FPropertyChangedEvent& PropertyChangedEvent, FProperty* PropertyThatChanged)
+{
+	// Called when something is modified in the details view.
+	if (CameraRigAsset)
+	{
+		CameraRigAsset->BuildStatus = ECameraBuildStatus::Dirty;
+	}
+}
+
+void FCameraRigAssetEditorToolkitBase::OnAnyGraphChanged(const FEdGraphEditAction& InEditAction)
+{
+	// Called when something is modified in the node graph or transition graph.
+	if (CameraRigAsset)
+	{
+		CameraRigAsset->BuildStatus = ECameraBuildStatus::Dirty;
+	}
+}
+
+}  // namespace UE::Cameras
+
+#undef LOCTEXT_NAMESPACE
+
diff --git a/Engine/Plugins/Cameras/GameplayCameras/Source/GameplayCamerasEditor/Private/Toolkits/CameraRigTransitionEditorToolkit.cpp b/Engine/Plugins/Cameras/GameplayCameras/Source/GameplayCamerasEditor/Private/Toolkits/CameraRigTransitionEditorToolkit.cpp
new file mode 100644
--- /dev/null
+++ b/Engine/Plugins/Cameras/GameplayCameras/Source/GameplayCamerasEditor/Private/Toolkits/CameraRigTransitionEditorToolkit.cpp
@@ -0,0 +1,132 @@
+// Copyright Epic Games, Inc. All Rights Reserved.
+
+#include "Toolkits/CameraRigTransitionEditorToolkit.h"
+
+#include "Framework/Docking/LayoutExtender.h"
+#include "Framework/Docking/TabManager.h"
+#include "ToolMenus.h"
+#include "Toolkits/CameraRigTransitionEditorToolkitBase.h"
+#include "Toolkits/StandardToolkitLayout.h"
+#include "Widgets/Docking/SDockTab.h"
+
+#define LOCTEXT_NAMESPACE "CameraRigTransitionEditorToolkit"
+
+namespace UE::Cameras
+{
+
+namespace Internal
+{
+
+class FCameraRigTransitionEditorToolkitBaseImpl : public FCameraRigTransitionEditorToolkitBase
+{
+public:
+	
+	FCameraRigTransitionEditorToolkitBaseImpl()
+		: FCameraRigTransitionEditorToolkitBase(TEXT("CameraRigTransitionEditor_Layout_v2"))
+	{}
+
+protected:
+
+	virtual void GetTransitionOwnerInfo(FCameraRigTransitionOwnerInfo& OutTransitionOwnerInfo) override
+	{
+		// TODO
+	}
+
+	virtual void GetTransitionGraphDisplayInfo(FGraphDisplayInfo& OutGraphDisplayInfo) override
+	{
+		// TODO
+	}
+
+	virtual void GetTransitionGraphAppearanceInfo(FGraphAppearanceInfo& OutGraphAppearanceInfo) override
+	{
+		// TODO
+	}
+};
+
+}  // namespace Internal
+
+FCameraRigTransitionEditorToolkit::FCameraRigTransitionEditorToolkit(UAssetEditor* InOwningAssetEditor)
+	: FBaseAssetToolkit(InOwningAssetEditor)
+{
+	Impl = MakeShared<Internal::FCameraRigTransitionEditorToolkitBaseImpl>();
+
+	// Override base class default layout.
+	StandaloneDefaultLayout = Impl->GetStandardLayout()->GetLayout();
+}
+
+void FCameraRigTransitionEditorToolkit::SetTransitionOwner(UObject* InTransitionOwner)
+{
+	Impl->SetTransitionOwner(InTransitionOwner);
+}
+
+void FCameraRigTransitionEditorToolkit::RegisterTabSpawners(const TSharedRef<FTabManager>& InTabManager)
+{
+	// Skip FBaseAssetToolkit here because we don't want a viewport tab.
+	FAssetEditorToolkit::RegisterTabSpawners(InTabManager);
+	
+	Impl->RegisterTabSpawners(InTabManager, AssetEditorTabsCategory);
+}
+
+void FCameraRigTransitionEditorToolkit::UnregisterTabSpawners(const TSharedRef<FTabManager>& InTabManager)
+{
+	// Skip FBaseAssetToolkit here because we don't want a viewport tab.
+	FAssetEditorToolkit::UnregisterTabSpawners(InTabManager);
+
+	Impl->UnregisterTabSpawners(InTabManager);
+}
+
+void FCameraRigTransitionEditorToolkit::CreateWidgets()
+{
+	// Skip FBaseAssetToolkit here because we don't want a viewport tab, and our base toolkit class
+	// already has its own details view.
+	// ...no up-call...
+
+	RegisterToolbar();
+	CreateEditorModeManager();
+	LayoutExtender = MakeShared<FLayoutExtender>();
+	
+	// Now do our custom stuff.
+
+	Impl->CreateWidgets();
+}
+
+void FCameraRigTransitionEditorToolkit::RegisterToolbar()
+{
+	FName ParentName;
+	const FName MenuName = GetToolMenuToolbarName(ParentName);
+	UToolMenus* ToolMenus = UToolMenus::Get();
+	if (!ToolMenus->IsMenuRegistered(MenuName))
+	{
+		FToolMenuOwnerScoped ToolMenuOwnerScope(this);
+
+		UToolMenu* ToolbarMenu = UToolMenus::Get()->RegisterMenu(
+				MenuName, ParentName, EMultiBoxType::ToolBar);
+		Impl->BuildToolbarMenu(ToolbarMenu);
+	}
+}
+
+FText FCameraRigTransitionEditorToolkit::GetBaseToolkitName() const
+{
+	return LOCTEXT("AppLabel", "Camera Transitions");
+}
+
+FName FCameraRigTransitionEditorToolkit::GetToolkitFName() const
+{
+	static FName ToolkitName("CameraRigTransitionEditor");
+	return ToolkitName;
+}
+
+FString FCameraRigTransitionEditorToolkit::GetWorldCentricTabPrefix() const
+{
+	return LOCTEXT("WorldCentricTabPrefix", "Camera Transitions ").ToString();
+}
+
+FLinearColor FCameraRigTransitionEditorToolkit::GetWorldCentricTabColorScale() const
+{
+	return FLinearColor(0.1f, 0.8f, 0.2f, 0.5f);
+}
+
+}  // namespace UE::Cameras
+
+#undef LOCTEXT_NAMESPACE
+
diff --git a/Engine/Plugins/Cameras/GameplayCameras/Source/GameplayCamerasEditor/Private/Toolkits/CameraRigTransitionEditorToolkitBase.cpp b/Engine/Plugins/Cameras/GameplayCameras/Source/GameplayCamerasEditor/Private/Toolkits/CameraRigTransitionEditorToolkitBase.cpp
new file mode 100644
--- /dev/null
+++ b/Engine/Plugins/Cameras/GameplayCameras/Source/GameplayCamerasEditor/Private/Toolkits/CameraRigTransitionEditorToolkitBase.cpp
@@ -0,0 +1,186 @@
+// Copyright Epic Games, Inc. All Rights Reserved.
+
+#include "Toolkits/CameraRigTransitionEditorToolkitBase.h"
+
+#include "Commands/CameraRigTransitionEditorCommands.h"
+#include "Core/CameraBuildStatus.h"
+#include "Editors/SCameraRigTransitionEditor.h"
+#include "Editors/SObjectTreeGraphToolbox.h"
+#include "Framework/Commands/UICommandList.h"
+#include "Framework/Docking/TabManager.h"
+#include "IDetailsView.h"
+#include "Modules/ModuleManager.h"
+#include "PropertyEditorModule.h"
+#include "Styles/GameplayCamerasEditorStyle.h"
+#include "ToolMenu.h"
+#include "Toolkits/StandardToolkitLayout.h"
+#include "Widgets/Docking/SDockTab.h"
+
+#define LOCTEXT_NAMESPACE "CameraRigTransitionEditorToolkitBase"
+
+namespace UE::Cameras
+{
+
+const FName FCameraRigTransitionEditorToolkitBase::ToolboxTabId(TEXT("CameraRigTransitionEditor_Toolbox"));
+const FName FCameraRigTransitionEditorToolkitBase::TransitionEditorTabId(TEXT("CameraRigTransitionEditor_TransitionEditor"));
+const FName FCameraRigTransitionEditorToolkitBase::DetailsViewTabId(TEXT("CameraRigTransitionEditor_DetailsView"));
+
+FCameraRigTransitionEditorToolkitBase::FCameraRigTransitionEditorToolkitBase(FName InLayoutName)
+	: StandardLayout(new FStandardToolkitLayout(InLayoutName))
+	, CommandBindings(new FUICommandList())
+{
+	StandardLayout->AddLeftTab(ToolboxTabId);
+	StandardLayout->AddCenterTab(TransitionEditorTabId);
+	StandardLayout->AddRightTab(DetailsViewTabId);
+}
+
+FCameraRigTransitionEditorToolkitBase::~FCameraRigTransitionEditorToolkitBase()
+{
+	if (TransitionEditorWidget)
+	{
+		TransitionEditorWidget->RemoveOnGraphChanged(this);
+	}
+}
+
+void FCameraRigTransitionEditorToolkitBase::AddReferencedObjects(FReferenceCollector& Collector)
+{
+	Collector.AddReferencedObject(TransitionOwner);
+}
+
+FString FCameraRigTransitionEditorToolkitBase::GetReferencerName() const
+{
+	return TEXT("FCameraRigTransitionEditorToolkitBase");
+}
+
+void FCameraRigTransitionEditorToolkitBase::SetTransitionOwner(UObject* InTransitionOwner)
+{
+	TransitionOwner = InTransitionOwner;
+
+	if (TransitionEditorWidget)
+	{
+		TransitionEditorWidget->SetTransitionOwner(InTransitionOwner);
+	}
+}
+
+void FCameraRigTransitionEditorToolkitBase::RegisterTabSpawners(TSharedRef<FTabManager> InTabManager, TSharedPtr<FWorkspaceItem> InAssetEditorTabsCategory)
+{
+	const FName CamerasStyleSetName = FGameplayCamerasEditorStyle::Get()->GetStyleSetName();
+
+	InTabManager->RegisterTabSpawner(ToolboxTabId, FOnSpawnTab::CreateSP(this, &FCameraRigTransitionEditorToolkitBase::SpawnTab_Toolbox))
+		.SetDisplayName(LOCTEXT("Toolbox", "Toolbox"))
+		.SetGroup(InAssetEditorTabsCategory.ToSharedRef())
+		.SetIcon(FSlateIcon(CamerasStyleSetName, "CameraRigAssetEditor.Tabs.Toolbox"));
+
+	InTabManager->RegisterTabSpawner(TransitionEditorTabId, FOnSpawnTab::CreateSP(this, &FCameraRigTransitionEditorToolkitBase::SpawnTab_TransitionEditor))
+		.SetDisplayName(LOCTEXT("TransitionEditor", "Camera Transitions"))
+		.SetGroup(InAssetEditorTabsCategory.ToSharedRef());
+
+	InTabManager->RegisterTabSpawner(DetailsViewTabId, FOnSpawnTab::CreateSP(this, &FCameraRigTransitionEditorToolkitBase::SpawnTab_Details))
+		.SetDisplayName(LOCTEXT("Details", "Details"))
+		.SetGroup(InAssetEditorTabsCategory.ToSharedRef())
+		.SetIcon(FSlateIcon(FAppStyle::GetAppStyleSetName(), "LevelEditor.Tabs.Details"));
+}
+
+TSharedRef<SDockTab> FCameraRigTransitionEditorToolkitBase::SpawnTab_Toolbox(const FSpawnTabArgs& Args)
+{
+	TSharedPtr<SDockTab> ToolboxTab = SNew(SDockTab)
+		.Label(LOCTEXT("ToolboxTabTitle", "Toolbox"))
+		[
+			ToolboxWidget.ToSharedRef()
+		];
+
+	return ToolboxTab.ToSharedRef();
+}
+
+TSharedRef<SDockTab> FCameraRigTransitionEditorToolkitBase::SpawnTab_TransitionEditor(const FSpawnTabArgs& Args)
+{
+	TSharedPtr<SDockTab> TransitionEditorTab = SNew(SDockTab)
+		.Label(LOCTEXT("TransitionEditorTabTitle", "Camera Transitions"))
+		[
+			TransitionEditorWidget.ToSharedRef()
+		];
+
+	return TransitionEditorTab.ToSharedRef();
+}
+
+TSharedRef<SDockTab> FCameraRigTransitionEditorToolkitBase::SpawnTab_Details(const FSpawnTabArgs& Args)
+{
+	TSharedPtr<SDockTab> DetailsTab = SNew(SDockTab)
+		.Label(LOCTEXT("BaseDetailsTitle", "Details"))
+		[
+			DetailsView.ToSharedRef()
+		];
+
+	return DetailsTab.ToSharedRef();
+}
+
+void FCameraRigTransitionEditorToolkitBase::UnregisterTabSpawners(TSharedRef<FTabManager> InTabManager)
+{
+	InTabManager->UnregisterTabSpawner(ToolboxTabId);
+	InTabManager->UnregisterTabSpawner(TransitionEditorTabId);
+	InTabManager->UnregisterTabSpawner(DetailsViewTabId);
+}
+
+void FCameraRigTransitionEditorToolkitBase::CreateWidgets()
+{
+	// Create the details view.
+	FPropertyEditorModule& PropertyEditorModule = FModuleManager::GetModuleChecked<FPropertyEditorModule>("PropertyEditor");
+	FDetailsViewArgs DetailsViewArgs;
+	DetailsViewArgs.NameAreaSettings = FDetailsViewArgs::HideNameArea;
+	DetailsViewArgs.bHideSelectionTip = true;
+	DetailsViewArgs.NotifyHook = this;
+	DetailsView = PropertyEditorModule.CreateDetailView(DetailsViewArgs);
+
+	FCameraRigTransitionOwnerInfo TransitionOwnerInfo;
+	GetTransitionOwnerInfo(TransitionOwnerInfo);
+
+	FGraphDisplayInfo GraphDisplayInfo;
+	GetTransitionGraphDisplayInfo(GraphDisplayInfo);
+
+	FGraphAppearanceInfo GraphAppearanceInfo;
+	GetTransitionGraphAppearanceInfo(GraphAppearanceInfo);
+
+	// Create the transition editor.
+	TransitionEditorWidget = SNew(SCameraRigTransitionEditor)
+		.DetailsView(DetailsView)
+		.TransitionOwner(TransitionOwner)
+		.TransitionOwnerInfo(TransitionOwnerInfo)
+		.TransitionGraphDisplayInfo(GraphDisplayInfo)
+		.TransitionGraphEditorAppearance(GraphAppearanceInfo);
+	TransitionEditorWidget->AddOnGraphChanged(FOnGraphChanged::FDelegate::CreateSP(
+				this, &FCameraRigTransitionEditorToolkitBase::OnTransitionGraphChanged));
+
+	// Create the toolbox, default to the rig editor items.
+	ToolboxWidget = SNew(SObjectTreeGraphToolbox)
+		.GraphConfig(TransitionEditorWidget->GetTransitionGraphConfig());
+}
+
+void FCameraRigTransitionEditorToolkitBase::BuildToolbarMenu(UToolMenu* ToolbarMenu)
+{
+	FToolMenuInsert InsertAfterAssetSection("Asset", EToolMenuInsertType::After);
+	const FCameraRigTransitionEditorCommands& Commands = FCameraRigTransitionEditorCommands::Get();
+	const FName CamerasStyleSetName = FGameplayCamerasEditorStyle::Get()->GetStyleSetName();
+}
+
+void FCameraRigTransitionEditorToolkitBase::NotifyPostChange(const FPropertyChangedEvent& PropertyChangedEvent, FProperty* PropertyThatChanged)
+{
+	// Called when something is modified in the details view.
+	if (IHasCameraBuildStatus* Buildable = Cast<IHasCameraBuildStatus>(TransitionOwner))
+	{
+		Buildable->DirtyBuildStatus();
+	}
+}
+
+void FCameraRigTransitionEditorToolkitBase::OnTransitionGraphChanged(const FEdGraphEditAction& InEditAction)
+{
+	// Called when something is modified in the transition graph.
+	if (IHasCameraBuildStatus* Buildable = Cast<IHasCameraBuildStatus>(TransitionOwner))
+	{
+		Buildable->DirtyBuildStatus();
+	}
+}
+
+}  // namespace UE::Cameras
+
+#undef LOCTEXT_NAMESPACE
+
diff --git a/Engine/Plugins/Cameras/GameplayCameras/Source/GameplayCamerasEditor/Private/Toolkits/CameraRigsAssetEditorMode.cpp b/Engine/Plugins/Cameras/GameplayCameras/Source/GameplayCamerasEditor/Private/Toolkits/CameraRigsAssetEditorMode.cpp
new file mode 100644
--- /dev/null
+++ b/Engine/Plugins/Cameras/GameplayCameras/Source/GameplayCamerasEditor/Private/Toolkits/CameraRigsAssetEditorMode.cpp
@@ -0,0 +1,121 @@
+// Copyright Epic Games, Inc. All Rights Reserved.
+
+#include "Toolkits/CameraRigsAssetEditorMode.h"
+
+#include "Editors/SCameraRigAssetEditor.h"
+#include "Framework/Docking/TabManager.h"
+#include "Styles/GameplayCamerasEditorStyle.h"
+#include "ToolMenus.h"
+#include "Toolkits/CameraRigAssetEditorToolkitBase.h"
+#include "Toolkits/StandardToolkitLayout.h"
+#include "Widgets/Docking/SDockTab.h"
+
+#define LOCTEXT_NAMESPACE "CameraRigsAssetEditorMode"
+
+namespace UE::Cameras
+{
+
+FName FCameraRigsAssetEditorMode::ModeName(TEXT("CameraRigs"));
+
+const FName FCameraRigsAssetEditorMode::CameraRigsTabId(TEXT("CameraRigAssetEditor_CameraRigs"));
+
+FCameraRigsAssetEditorMode::FCameraRigsAssetEditorMode(UCameraAsset* InCameraAsset)
+	: FAssetEditorMode(ModeName)
+	, CameraAsset(InCameraAsset)
+{
+	Impl = MakeShared<FCameraRigAssetEditorToolkitBase>(TEXT("CameraAssetEditor_Mode_CameraRigs_v1"));
+
+	TSharedPtr<FStandardToolkitLayout> StandardLayout = Impl->GetStandardLayout();
+	{
+		StandardLayout->AddLeftTab(CameraRigsTabId);
+	}
+	DefaultLayout = StandardLayout->GetLayout();
+}
+
+void FCameraRigsAssetEditorMode::OnActivateMode(const FAssetEditorModeActivateParams& InParams)
+{
+	if (!bInitializedToolkit)
+	{
+		Impl->CreateWidgets();
+
+		CameraRigsListWidget = SNew(SCameraRigList)
+			.CameraAsset(CameraAsset)
+			.OnCameraRigListChanged(this, &FCameraRigsAssetEditorMode::OnCameraRigListChanged)
+			.OnRequestEditCameraRig(this, &FCameraRigsAssetEditorMode::OnCameraRigEditRequested);
+
+		bInitializedToolkit = true;
+	}
+
+	Impl->RegisterTabSpawners(InParams.TabManager.ToSharedRef(), InParams.AssetEditorTabsCategory);
+
+	const FName CamerasStyleSetName = FGameplayCamerasEditorStyle::Get()->GetStyleSetName();
+
+	InParams.TabManager->RegisterTabSpawner(CameraRigsTabId, FOnSpawnTab::CreateSP(this, &FCameraRigsAssetEditorMode::SpawnTab_CameraRigs))
+		.SetDisplayName(LOCTEXT("CameraRigs", "Camera Rigs"))
+		.SetGroup(InParams.AssetEditorTabsCategory.ToSharedRef())
+		.SetIcon(FSlateIcon(CamerasStyleSetName, "CameraAssetEditor.Tabs.CameraRigs"));
+
+	FToolMenuOwnerScoped OwnerScoped(this);
+	UToolMenu* ToolbarMenu = UToolMenus::Get()->ExtendMenu(InParams.ToolbarMenuName);
+	Impl->BuildToolbarMenu(ToolbarMenu);
+
+	Impl->BindCommands(InParams.CommandList.ToSharedRef());
+}
+
+TSharedRef<SDockTab> FCameraRigsAssetEditorMode::SpawnTab_CameraRigs(const FSpawnTabArgs& Args)
+{
+	TSharedPtr<SDockTab> CameraRigsTab = SNew(SDockTab)
+		.Label(LOCTEXT("CameraRigsTitle", "Camera Rigs"))
+		[
+			CameraRigsListWidget.ToSharedRef()
+		];
+
+	return CameraRigsTab.ToSharedRef();
+}
+
+void FCameraRigsAssetEditorMode::OnDeactivateMode(const FAssetEditorModeDeactivateParams& InParams)
+{
+	Impl->UnregisterTabSpawners(InParams.TabManager.ToSharedRef());
+
+	InParams.TabManager->UnregisterTabSpawner(CameraRigsTabId);
+
+	UToolMenus::UnregisterOwner(this);
+}
+
+void FCameraRigsAssetEditorMode::OnCameraRigListChanged(TArrayView<UCameraRigAsset* const> InCameraRigs)
+{
+}
+
+void FCameraRigsAssetEditorMode::OnCameraRigEditRequested(UCameraRigAsset* InCameraRig)
+{
+	Impl->SetCameraRigAsset(InCameraRig);
+}
+
+bool FCameraRigsAssetEditorMode::JumpToObject(UObject* InObject)
+{
+	UCameraRigAsset* FindInCameraRig = nullptr;
+	UObject* CurOuter = InObject;
+	while (CurOuter != nullptr)
+	{
+		if (CurOuter->IsA<UCameraRigAsset>())
+		{
+			FindInCameraRig = Cast<UCameraRigAsset>(CurOuter);
+			break;
+		}
+		CurOuter = CurOuter->GetOuter();
+	}
+	if (!FindInCameraRig)
+	{
+		return false;
+	}
+
+	Impl->SetCameraRigAsset(FindInCameraRig);
+
+	TSharedPtr<SCameraRigAssetEditor> CameraRigAssetEditor = Impl->GetCameraRigAssetEditor();
+	return CameraRigAssetEditor->FindAndJumpToObjectNode(InObject);
+}
+
+}  // namespace UE::Cameras
+
+#undef LOCTEXT_NAMESPACE
+
diff --git a/Engine/Plugins/Cameras/GameplayCameras/Source/GameplayCamerasEditor/Private/Toolkits/CameraSharedTransitionsAssetEditorMode.cpp b/Engine/Plugins/Cameras/GameplayCameras/Source/GameplayCamerasEditor/Private/Toolkits/CameraSharedTransitionsAssetEditorMode.cpp
new file mode 100644
--- /dev/null
+++ b/Engine/Plugins/Cameras/GameplayCameras/Source/GameplayCamerasEditor/Private/Toolkits/CameraSharedTransitionsAssetEditorMode.cpp
@@ -0,0 +1,101 @@
+// Copyright Epic Games, Inc. All Rights Reserved.
+
+#include "Toolkits/CameraSharedTransitionsAssetEditorMode.h"
+
+#include "Core/CameraAsset.h"
+#include "Editors/CameraRigTransitionGraphSchema.h"
+#include "Editors/SCameraRigTransitionEditor.h"
+#include "ToolMenus.h"
+#include "Toolkits/CameraRigTransitionEditorToolkitBase.h"
+#include "Toolkits/StandardToolkitLayout.h"
+
+#define LOCTEXT_NAMESPACE "CameraSharedTransitionsAssetEditorMode"
+
+namespace UE::Cameras
+{
+
+namespace Internal
+{
+
+class FCameraSharedTransitionsAssetEditorModeBaseImpl : public FCameraRigTransitionEditorToolkitBase
+{
+public:
+
+	FCameraSharedTransitionsAssetEditorModeBaseImpl()
+		: FCameraRigTransitionEditorToolkitBase(TEXT("CameraAssetEditor_Mode_SharedTransitions_v1"))
+	{}
+
+protected:
+
+	virtual void GetTransitionOwnerInfo(FCameraRigTransitionOwnerInfo& OutTransitionOwnerInfo) override
+	{
+		OutTransitionOwnerInfo.TransitionOwnerClass = UCameraAsset::StaticClass();
+		OutTransitionOwnerInfo.GraphName = UCameraAsset::SharedTransitionsGraphName;
+		OutTransitionOwnerInfo.EnterTransitionsPropertyName = GET_MEMBER_NAME_CHECKED(UCameraAsset, EnterTransitions);
+		OutTransitionOwnerInfo.ExitTransitionsPropertyName = GET_MEMBER_NAME_CHECKED(UCameraAsset, ExitTransitions);
+	}
+
+	virtual void GetTransitionGraphDisplayInfo(FGraphDisplayInfo& OutGraphDisplayInfo) override
+	{
+		OutGraphDisplayInfo.PlainName = LOCTEXT("SharedTransitionGraphPlainName", "SharedTransitions");
+		OutGraphDisplayInfo.DisplayName = LOCTEXT("SharedTransitionGraphDisplayName", "Shared Transitions");
+	}
+
+	virtual void GetTransitionGraphAppearanceInfo(FGraphAppearanceInfo& OutGraphAppearanceInfo) override
+	{
+		OutGraphAppearanceInfo.CornerText = LOCTEXT("SharedTransitionGraphCornerText", "SHARED TRANSITIONS");
+	}
+};
+
+}  // namespace Internal
+
+FName FCameraSharedTransitionsAssetEditorMode::ModeName(TEXT("SharedTransitions"));
+
+FCameraSharedTransitionsAssetEditorMode::FCameraSharedTransitionsAssetEditorMode(UCameraAsset* InCameraAsset)
+	: FAssetEditorMode(ModeName)
+	, CameraAsset(InCameraAsset)
+{
+	Impl = MakeShared<Internal::FCameraSharedTransitionsAssetEditorModeBaseImpl>();
+
+	Impl->SetTransitionOwner(CameraAsset);
+
+	DefaultLayout = Impl->GetStandardLayout()->GetLayout();
+}
+
+void FCameraSharedTransitionsAssetEditorMode::OnActivateMode(const FAssetEditorModeActivateParams& InParams)
+{
+	if (!bInitializedToolkit)
+	{
+		Impl->CreateWidgets();
+		bInitializedToolkit = true;
+	}
+
+	Impl->RegisterTabSpawners(InParams.TabManager.ToSharedRef(), InParams.AssetEditorTabsCategory);
+
+	FToolMenuOwnerScoped OwnerScoped(this);
+	UToolMenu* ToolbarMenu = UToolMenus::Get()->ExtendMenu(InParams.ToolbarMenuName);
+	Impl->BuildToolbarMenu(ToolbarMenu);
+}
+
+void FCameraSharedTransitionsAssetEditorMode::OnDeactivateMode(const FAssetEditorModeDeactivateParams& InParams)
+{
+	Impl->UnregisterTabSpawners(InParams.TabManager.ToSharedRef());
+
+	UToolMenus::UnregisterOwner(this);
+}
+
+bool FCameraSharedTransitionsAssetEditorMode::JumpToNode(UEdGraphNode* InNode)
+{
+	return false;
+}
+
+bool FCameraSharedTransitionsAssetEditorMode::JumpToObject(UObject* InObject)
+{
+	TSharedPtr<SCameraRigTransitionEditor> TransitionEditor = Impl->GetCameraRigTransitionEditor();
+	return TransitionEditor->FindAndJumpToObjectNode(InObject);
+}
+
+}  // namespace UE::Cameras
+
+#undef LOCTEXT_NAMESPACE
+
diff --git a/Engine/Plugins/Cameras/GameplayCameras/Source/GameplayCamerasEditor/Private/Toolkits/StandardToolkitLayout.cpp b/Engine/Plugins/Cameras/GameplayCameras/Source/GameplayCamerasEditor/Private/Toolkits/StandardToolkitLayout.cpp
new file mode 100644
--- /dev/null
+++ b/Engine/Plugins/Cameras/GameplayCameras/Source/GameplayCamerasEditor/Private/Toolkits/StandardToolkitLayout.cpp
@@ -0,0 +1,108 @@
+// Copyright Epic Games, Inc. All Rights Reserved.
+
+#include "Toolkits/StandardToolkitLayout.h"
+#include "Framework/Docking/TabManager.h"
+
+namespace UE::Cameras
+{
+
+const FName FStandardToolkitLayout::LeftStackExtensionId("LeftStackId");
+const FName FStandardToolkitLayout::CenterStackExtensionId("CenterStackId");
+const FName FStandardToolkitLayout::RightStackExtensionId("RightStackId");
+const FName FStandardToolkitLayout::BottomStackExtensionId("BottomStackId");
+
+FStandardToolkitLayout::FStandardToolkitLayout(FName InLayoutName)
+	: LayoutName(InLayoutName)
+{
+	if (LayoutName.IsNone())
+	{
+		LayoutName = TEXT("StandardCameraEditorToolkit_Layout_v1");
+	}
+
+	BuildLayout();
+}
+
+TSharedPtr<FTabManager::FLayout> FStandardToolkitLayout::GetLayout() const
+{
+	return Layout;
+}
+
+void FStandardToolkitLayout::BuildLayout()
+{
+	if (Layout.IsValid())
+	{
+		return;
+	}
+
+	LeftTabStack = FTabManager::NewStack();
+	CenterTabStack = FTabManager::NewStack();
+	RightTabStack = FTabManager::NewStack();
+	BottomTabStack = FTabManager::NewStack();
+
+	Layout = FTabManager::NewLayout(LayoutName)
+		->AddArea
+		(
+			FTabManager::NewPrimaryArea()
+			->SetOrientation(Orient_Vertical)
+			->Split
+			(
+				FTabManager::NewSplitter()
+				->SetOrientation(Orient_Horizontal)
+				->Split
+				(
+					LeftTabStack
+					->SetExtensionId(LeftStackExtensionId)
+					->SetSizeCoefficient(0.2f)
+				)
+				->Split
+				(
+					CenterTabStack
+					->SetExtensionId(CenterStackExtensionId)
+					->SetSizeCoefficient(0.6f)
+					->SetHideTabWell(true)
+				)
+				->Split
+				(
+					RightTabStack
+					->SetExtensionId(RightStackExtensionId)
+					->SetSizeCoefficient(0.2f)
+				)
+			)
+			->Split
+			(
+				BottomTabStack
+				->SetExtensionId(BottomStackExtensionId)
+				->SetSizeCoefficient(0.2f)
+			)
+		);
+
+	// Add content browsers to the bottom tab stack by default so that
+	// we can dock them in our layout.
+	BottomTabStack->AddTab("ContentBrowserTab1", ETabState::ClosedTab);
+	BottomTabStack->AddTab("ContentBrowserTab2", ETabState::ClosedTab);
+	BottomTabStack->AddTab("ContentBrowserTab3", ETabState::ClosedTab);
+	BottomTabStack->AddTab("ContentBrowserTab4", ETabState::ClosedTab);
+}
+
+void FStandardToolkitLayout::AddLeftTab(FName InTabId, ETabState::Type InTabState)
+{
+	LeftTabStack->AddTab(InTabId, InTabState);
+}
+
+void FStandardToolkitLayout::AddRightTab(FName InTabId, ETabState::Type InTabState)
+{
+	RightTabStack->AddTab(InTabId, InTabState);
+}
+
+void FStandardToolkitLayout::AddBottomTab(FName InTabId, ETabState::Type InTabState)
+{
+	BottomTabStack->AddTab(InTabId, InTabState);
+}
+
+void FStandardToolkitLayout::AddCenterTab(FName InTabId, ETabState::Type InTabState)
+{
+	CenterTabStack->AddTab(InTabId, InTabState);
+}
+
+}  // namespace UE::Cameras
+
diff --git a/Engine/Plugins/Cameras/GameplayCameras/Source/GameplayCamerasEditor/Public/AssetTools/CameraAssetEditor.h b/Engine/Plugins/Cameras/GameplayCameras/Source/GameplayCamerasEditor/Public/AssetTools/CameraAssetEditor.h
--- a/Engine/Plugins/Cameras/GameplayCameras/Source/GameplayCamerasEditor/Public/AssetTools/CameraAssetEditor.h
+++ b/Engine/Plugins/Cameras/GameplayCameras/Source/GameplayCamerasEditor/Public/AssetTools/CameraAssetEditor.h
@@ -21,6 +21,8 @@ public:
 
 	void Initialize(TObjectPtr<UCameraAsset> InCameraAsset);
 
+	UCameraAsset* GetCameraAsset() const { return CameraAsset; }
+
 public:
 
 	// UAssetEditor interface
@@ -29,6 +31,7 @@ public:
 
 private:
 
+	UPROPERTY()
 	TObjectPtr<UCameraAsset> CameraAsset;
 };
 
diff --git a/Engine/Plugins/Cameras/GameplayCameras/Source/GameplayCamerasEditor/Public/AssetTools/CameraRigAssetEditor.h b/Engine/Plugins/Cameras/GameplayCameras/Source/GameplayCamerasEditor/Public/AssetTools/CameraRigAssetEditor.h
--- a/Engine/Plugins/Cameras/GameplayCameras/Source/GameplayCamerasEditor/Public/AssetTools/CameraRigAssetEditor.h
+++ b/Engine/Plugins/Cameras/GameplayCameras/Source/GameplayCamerasEditor/Public/AssetTools/CameraRigAssetEditor.h
@@ -31,6 +31,7 @@ public:
 
 private:
 
+	UPROPERTY()
 	TObjectPtr<UCameraRigAsset> CameraRigAsset;
 };
 
diff --git a/Engine/Plugins/Cameras/GameplayCameras/Source/GameplayCamerasEditor/Public/AssetTools/CameraRigTransitionEditor.h b/Engine/Plugins/Cameras/GameplayCameras/Source/GameplayCamerasEditor/Public/AssetTools/CameraRigTransitionEditor.h
new file mode 100644
--- /dev/null
+++ b/Engine/Plugins/Cameras/GameplayCameras/Source/GameplayCamerasEditor/Public/AssetTools/CameraRigTransitionEditor.h
@@ -0,0 +1,36 @@
+// Copyright Epic Games, Inc. All Rights Reserved.
+
+#pragma once
+
+#include "Tools/UAssetEditor.h"
+
+#include "CameraRigTransitionEditor.generated.h"
+
+class FBaseAssetToolkit;
+
+/**
+ * Editor for camera transitions.
+ */
+UCLASS(Transient)
+class UCameraRigTransitionEditor : public UAssetEditor
+{
+	GENERATED_BODY()
+
+public:
+
+	void Initialize(TObjectPtr<UObject> InTransitionOwner);
+
+	UObject* GetTransitionOwner() const { return TransitionOwner; }
+
+public:
+
+	// UAssetEditor interface
+	virtual void GetObjectsToEdit(TArray<UObject*>& InObjectsToEdit) override;
+	virtual TSharedPtr<FBaseAssetToolkit> CreateToolkit() override;
+
+private:
+
+	UPROPERTY()
+	TObjectPtr<UObject> TransitionOwner;
+};
+
diff --git a/Engine/Plugins/Cameras/GameplayCameras/Source/GameplayCamerasEditor/Public/Commands/CameraAssetEditorCommands.h b/Engine/Plugins/Cameras/GameplayCameras/Source/GameplayCamerasEditor/Public/Commands/CameraAssetEditorCommands.h
--- a/Engine/Plugins/Cameras/GameplayCameras/Source/GameplayCamerasEditor/Public/Commands/CameraAssetEditorCommands.h
+++ b/Engine/Plugins/Cameras/GameplayCameras/Source/GameplayCamerasEditor/Public/Commands/CameraAssetEditorCommands.h
@@ -23,6 +23,19 @@ public:
 public:
 
 	TSharedPtr<FUICommandInfo> Build;
+
+	TSharedPtr<FUICommandInfo> ShowCameraDirector;
+	TSharedPtr<FUICommandInfo> ShowCameraRigs;
+	TSharedPtr<FUICommandInfo> ShowSharedTransitions;
+
+	TSharedPtr<FUICommandInfo> EditCameraRig;
+
+	TSharedPtr<FUICommandInfo> AddCameraRig;
+	TSharedPtr<FUICommandInfo> DeleteCameraRig;
+	TSharedPtr<FUICommandInfo> RenameCameraRig;
+
+	TSharedPtr<FUICommandInfo> ShowMessages;
+	TSharedPtr<FUICommandInfo> FindInCamera;
 };
 
 }  // namespace UE::Cameras
diff --git a/Engine/Plugins/Cameras/GameplayCameras/Source/GameplayCamerasEditor/Public/Commands/CameraRigTransitionEditorCommands.h b/Engine/Plugins/Cameras/GameplayCameras/Source/GameplayCamerasEditor/Public/Commands/CameraRigTransitionEditorCommands.h
new file mode 100644
--- /dev/null
+++ b/Engine/Plugins/Cameras/GameplayCameras/Source/GameplayCamerasEditor/Public/Commands/CameraRigTransitionEditorCommands.h
@@ -0,0 +1,30 @@
+// Copyright Epic Games, Inc. All Rights Reserved.
+
+#pragma once
+
+#include "CoreTypes.h"
+#include "Framework/Commands/Commands.h"
+
+class FUICommandInfo;
+class UClass;
+struct FInputChord;
+
+namespace UE::Cameras
+{
+
+class FCameraRigTransitionEditorCommands : public TCommands<FCameraRigTransitionEditorCommands>
+{
+public:
+
+	FCameraRigTransitionEditorCommands();
+
+	virtual void RegisterCommands() override;
+	
+public:
+
+	TSharedPtr<FUICommandInfo> FindInTransitions;
+	TSharedPtr<FUICommandInfo> FocusHome;
+};
+
+}  // namespace UE::Cameras
+
diff --git a/Engine/Plugins/Cameras/GameplayCameras/Source/GameplayCamerasEditor/Public/Editors/CameraNodeGraphSchema.h b/Engine/Plugins/Cameras/GameplayCameras/Source/GameplayCamerasEditor/Public/Editors/CameraNodeGraphSchema.h
--- a/Engine/Plugins/Cameras/GameplayCameras/Source/GameplayCamerasEditor/Public/Editors/CameraNodeGraphSchema.h
+++ b/Engine/Plugins/Cameras/GameplayCameras/Source/GameplayCamerasEditor/Public/Editors/CameraNodeGraphSchema.h
@@ -7,6 +7,7 @@
 #include "CameraNodeGraphSchema.generated.h"
 
 class UCameraNode;
+struct FObjectTreeGraphConfig;
 
 /**
  * Schema class for camera node graph.
@@ -20,6 +21,8 @@ public:
 
 	static const FName PC_CameraParameter;			// A camera parameter pin.
 
+	static FObjectTreeGraphConfig BuildGraphConfig();
+
 protected:
 
 	// UEdGraphSchema interface.
diff --git a/Engine/Plugins/Cameras/GameplayCameras/Source/GameplayCamerasEditor/Public/Editors/CameraRigTransitionGraphSchema.h b/Engine/Plugins/Cameras/GameplayCameras/Source/GameplayCamerasEditor/Public/Editors/CameraRigTransitionGraphSchema.h
new file mode 100644
--- /dev/null
+++ b/Engine/Plugins/Cameras/GameplayCameras/Source/GameplayCamerasEditor/Public/Editors/CameraRigTransitionGraphSchema.h
@@ -0,0 +1,84 @@
+// Copyright Epic Games, Inc. All Rights Reserved.
+
+#pragma once
+
+#include "Editors/ObjectTreeGraphSchema.h"
+
+#include "CameraRigTransitionGraphSchema.generated.h"
+
+struct FObjectTreeGraphConfig;
+
+/**
+ * Provides information about an object that has enter/exit transitions.
+ */
+struct FCameraRigTransitionOwnerInfo
+{
+	/** The class of the object. */
+	UClass* TransitionOwnerClass = nullptr;
+	/** The name of the transition graph. */
+	FName GraphName;
+	/** The name of the property that returns the array of enter transitions. */
+	FName EnterTransitionsPropertyName;
+	/** The name of the property that returns the array of exit transitions. */
+	FName ExitTransitionsPropertyName;
+};
+
+/**
+ * Schema class for camera transition graph.
+ */
+UCLASS()
+class UCameraRigTransitionGraphSchema : public UObjectTreeGraphSchema
+{
+	GENERATED_BODY()
+
+public:
+
+	static FObjectTreeGraphConfig BuildGraphConfig(const FCameraRigTransitionOwnerInfo& OwnerInfo);
+
+protected:
+
+	// UEdGraphSchema interface.
+	virtual void GetGraphContextActions(FGraphContextMenuBuilder& ContextMenuBuilder) const override;
+
+	// UObjectTreeGraphSchema interface.
+	virtual void FilterGraphContextPlaceableClasses(TArray<UClass*>& InOutClasses) const override;
+};
+
+/**
+ * Graph action to create a new transition node.
+ *
+ * We need a custom action for this because we need to switch the "self" pin according to whether
+ * we want an enter or exit transition.
+ */
+USTRUCT()
+struct FCameraRigTransitionGraphSchemaAction_NewTransitionNode : public FObjectGraphSchemaAction_NewNode
+{
+	GENERATED_BODY()
+
+public:
+
+	enum class ETransitionType
+	{
+		Enter,
+		Exit
+	};
+
+	/** The transition type to create. */
+	ETransitionType TransitionType;
+
+public:
+
+	FCameraRigTransitionGraphSchemaAction_NewTransitionNode();
+	FCameraRigTransitionGraphSchemaAction_NewTransitionNode(FText InNodeCategory, FText InMenuDesc, FText InToolTip, const int32 InGrouping = 0, FText InKeywords = FText());
+
+public:
+
+	// FEdGraphSchemaAction interface.
+	static FName StaticGetTypeId() { static FName Type("FCameraRigTransitionGraphSchemaAction_NewTransitionNode"); return Type; }
+	virtual FName GetTypeId() const override { return StaticGetTypeId(); } 
+
+protected:
+
+	virtual void AutoSetupNewNode(UObjectTreeGraphNode* NewNode, UEdGraphPin* FromPin) override;
+};
+
diff --git a/Engine/Plugins/Cameras/GameplayCameras/Source/GameplayCamerasEditor/Public/Editors/CameraTransitionGraphSchema.h b/Engine/Plugins/Cameras/GameplayCameras/Source/GameplayCamerasEditor/Public/Editors/CameraTransitionGraphSchema.h
deleted file mode 100644
diff --git a/Engine/Plugins/Cameras/GameplayCameras/Source/GameplayCamerasEditor/Public/Editors/ObjectTreeGraph.h b/Engine/Plugins/Cameras/GameplayCameras/Source/GameplayCamerasEditor/Public/Editors/ObjectTreeGraph.h
--- a/Engine/Plugins/Cameras/GameplayCameras/Source/GameplayCamerasEditor/Public/Editors/ObjectTreeGraph.h
+++ b/Engine/Plugins/Cameras/GameplayCameras/Source/GameplayCamerasEditor/Public/Editors/ObjectTreeGraph.h
@@ -31,7 +31,10 @@ public:
 	UObjectTreeGraph(const FObjectInitializer& ObjInit);
 
 	/** Initializes the graph given a root object and a graph configuration. */
-	void Initialize(TObjectPtr<UObject> InRootObject, const FObjectTreeGraphConfig& InConfig);
+	void Reset(
+			TObjectPtr<UObject> InRootObject, 
+			const FObjectTreeGraphConfig& InConfig, 
+			EObjectTreeGraphBuildSource InSource);
 
 	/** Gets the root object. */
 	UObject* GetRootObject() const { return WeakRootObject.Get(); }
@@ -44,7 +47,8 @@ public:
 	/** Gets the graph configuration. */
 	const FObjectTreeGraphConfig& GetConfig() const;
 
-	/** Rebuilds the node graph. */
+public:
+
 	void RebuildGraph(EObjectTreeGraphBuildSource InSource);
 
 private:
diff --git a/Engine/Plugins/Cameras/GameplayCameras/Source/GameplayCamerasEditor/Public/Editors/ObjectTreeGraphConfig.h b/Engine/Plugins/Cameras/GameplayCameras/Source/GameplayCamerasEditor/Public/Editors/ObjectTreeGraphConfig.h
--- a/Engine/Plugins/Cameras/GameplayCameras/Source/GameplayCamerasEditor/Public/Editors/ObjectTreeGraphConfig.h
+++ b/Engine/Plugins/Cameras/GameplayCameras/Source/GameplayCamerasEditor/Public/Editors/ObjectTreeGraphConfig.h
@@ -7,6 +7,7 @@
 #include "Containers/UnrealString.h"
 #include "CoreTypes.h"
 #include "EdGraph/EdGraphNode.h"
+#include "EdGraph/EdGraphSchema.h"
 #include "Misc/Optional.h"
 #include "Templates/SubclassOf.h"
 #include "Templates/UnrealTypeTraits.h"
@@ -16,10 +17,14 @@
 class FText;
 class UClass;
 class UObject;
+class UObjectTreeGraph;
 class UObjectTreeGraphNode;
+struct FObjectTreeGraphConfig;
 
+DECLARE_DELEGATE_OneParam(FOnBuildObjectTreeGraphConfig, FObjectTreeGraphConfig& InOutConfig);
 DECLARE_DELEGATE_RetVal_OneParam(FText, FOnGetObjectClassDisplayName, const UClass*);
 DECLARE_DELEGATE_TwoParams(FOnFormatObjectDisplayName, const UObject*, FText&);
+DECLARE_DELEGATE_TwoParams(FOnGetGraphDisplayInfo, const UObjectTreeGraph*, FGraphDisplayInfo&);
 
 #define OTGCC_FIELD(FieldType, FieldName)\
 	public:\
@@ -164,9 +169,24 @@ public:
 	/** A custom callback to format an object's display name. */
 	FOnFormatObjectDisplayName OnFormatObjectDisplayName;
 
+	/** The graph display information. */
+	FGraphDisplayInfo GraphDisplayInfo;
+
+	/** A custom callback to get the graph display information, to override GraphDisplayInfo. */
+	FOnGetGraphDisplayInfo OnGetGraphDisplayInfo;
+
 	/** Advanced, optional bits of configuration for specific classes and sub-classes of objects. */
 	TMap<UClass*, FObjectTreeGraphClassConfig> ObjectClassConfigs;
 
+	/** Stop auto-collection of initial objects at any object that is of the given types. */
+	TArray<UClass*> StopAutoCollectAtObjectClasses;
+
+	/** 
+	 * Whether to collect all objets in the graph automatically from the network of referenced
+	 * objects from the root.
+	 */
+	bool bAutoCollectInitialObjects = true;
+
 public:
 
 	/** Creates a new graph config. */
diff --git a/Engine/Plugins/Cameras/GameplayCameras/Source/GameplayCamerasEditor/Public/Editors/ObjectTreeGraphSchema.h b/Engine/Plugins/Cameras/GameplayCameras/Source/GameplayCamerasEditor/Public/Editors/ObjectTreeGraphSchema.h
--- a/Engine/Plugins/Cameras/GameplayCameras/Source/GameplayCamerasEditor/Public/Editors/ObjectTreeGraphSchema.h
+++ b/Engine/Plugins/Cameras/GameplayCameras/Source/GameplayCamerasEditor/Public/Editors/ObjectTreeGraphSchema.h
@@ -82,6 +82,7 @@ public:
 	virtual void BreakSinglePinLink(UEdGraphPin* SourcePin, UEdGraphPin* TargetPin) const override;
 	virtual bool SupportsDropPinOnNode(UEdGraphNode* InTargetNode, const FEdGraphPinType& InSourcePinType, EEdGraphPinDirection InSourcePinDirection, FText& OutErrorMessage) const override;
 	virtual bool SafeDeleteNodeFromGraph(UEdGraph* Graph, UEdGraphNode* Node) const override;
+	virtual void GetGraphDisplayInformation(const UEdGraph& Graph, FGraphDisplayInfo& OutDisplayInfo) const override;
 
 protected:
 
diff --git a/Engine/Plugins/Cameras/GameplayCameras/Source/GameplayCamerasEditor/Public/Editors/SCameraRigAssetEditor.h b/Engine/Plugins/Cameras/GameplayCameras/Source/GameplayCamerasEditor/Public/Editors/SCameraRigAssetEditor.h
--- a/Engine/Plugins/Cameras/GameplayCameras/Source/GameplayCamerasEditor/Public/Editors/SCameraRigAssetEditor.h
+++ b/Engine/Plugins/Cameras/GameplayCameras/Source/GameplayCamerasEditor/Public/Editors/SCameraRigAssetEditor.h
@@ -52,6 +52,8 @@ public:
 	void Construct(const FArguments& InArgs);
 	~SCameraRigAssetEditor();
 
+	void SetCameraRigAsset(UCameraRigAsset* InCameraRig);
+
 public:
 
 	/** Gets the current editor mode. */
@@ -85,18 +87,26 @@ public:
 
 protected:
 
-	void CreateNodeGraphEditor(const FArguments& InArgs);
-	void CreateTransitionGraphEditor(const FArguments& InArgs);
+	void CreateGraphEditors();
+	void CreateNodeGraphEditor();
+	void CreateTransitionGraphEditor();
+	void DiscardGraphEditors();
+
+	void SetEditorModeImpl(ECameraRigAssetEditorMode InMode, bool bForceSet);
 
 	void OnGraphChanged(const FEdGraphEditAction& InEditAction);
 
-	FText GetCameraRigAssetName() const;
+	FText GetCameraRigAssetName(UObjectTreeGraph* ForGraph) const;
 
 private:
 
 	/** The asset being edited */
 	TObjectPtr<UCameraRigAsset> CameraRigAsset;
 
+	TSharedPtr<IDetailsView> DetailsView;
+
+	TWeakPtr<FAssetEditorToolkit> AssetEditorToolkit;
+
 	/** The node hierarchy graph */
 	TObjectPtr<UObjectTreeGraph> NodeGraph;
 	/** The node hierarchy graph editor */
diff --git a/Engine/Plugins/Cameras/GameplayCameras/Source/GameplayCamerasEditor/Public/Editors/SCameraRigList.h b/Engine/Plugins/Cameras/GameplayCameras/Source/GameplayCamerasEditor/Public/Editors/SCameraRigList.h
new file mode 100644
--- /dev/null
+++ b/Engine/Plugins/Cameras/GameplayCameras/Source/GameplayCamerasEditor/Public/Editors/SCameraRigList.h
@@ -0,0 +1,142 @@
+// Copyright Epic Games, Inc. All Rights Reserved.
+
+#pragma once
+
+#include "Core/CameraRigAsset.h"
+#include "Misc/TextFilter.h"
+#include "UObject/ObjectPtr.h"
+#include "Widgets/SCompoundWidget.h"
+#include "Widgets/Views/SListView.h"
+#include "Widgets/Views/STableRow.h"
+
+class ITableRow;
+class SInlineEditableTextBlock;
+class SSearchBox;
+class STableViewBase;
+class UCameraAsset;
+class FUICommandList;
+
+namespace UE::Cameras
+{
+
+DECLARE_DELEGATE(FOnRequestRenameCameraRig);
+DECLARE_DELEGATE_OneParam(FOnCameraRigListChanged, TArrayView<UCameraRigAsset* const>);
+DECLARE_DELEGATE_OneParam(FOnRequestEditCameraRig, UCameraRigAsset*);
+
+/**
+ * List item for the list view in SCameraRigList.
+ */
+class FCameraRigListItem : public TSharedFromThis<FCameraRigListItem>
+{
+public:
+
+	UCameraRigAsset* CameraRigAsset = nullptr;
+	FOnRequestRenameCameraRig OnRequestRename;
+};
+
+/**
+ * Item widget for the list view in SCameraRigList.
+ */
+class SCameraRigListEntry : public STableRow<TSharedPtr<FCameraRigListItem>>
+{
+public:
+
+	SLATE_BEGIN_ARGS(SCameraRigListEntry)
+	{}
+		SLATE_ARGUMENT(TWeakPtr<FCameraRigListItem>, Item)
+		SLATE_ATTRIBUTE(FText, HighlightText)
+	SLATE_END_ARGS()
+
+	void Construct(const FArguments& InArgs, const TSharedRef<STableViewBase>& OwnerTable);
+
+private:
+
+	FText GetDisplayName() const;
+	void OnRename();
+	bool OnVerifyTextChanged(const FText& Text, FText& OutErrorMessage);
+	void OnTextCommitted(const FText& Text, ETextCommit::Type CommitType);
+
+private:
+
+	TWeakPtr<FCameraRigListItem> WeakItem;
+	TSharedPtr<SInlineEditableTextBlock> EditableTextBlock;
+};
+
+/**
+ * A panel that shows a list of camera rigs on a camera asset.
+ * Can add, rename, delete, etc. camera rigs in the list.
+ */
+class SCameraRigList : public SCompoundWidget
+{
+public:
+
+	SLATE_BEGIN_ARGS(SCameraRigList)
+		: _CameraAsset(nullptr)
+	{}
+		SLATE_ARGUMENT(UCameraAsset*, CameraAsset)
+		SLATE_EVENT(FOnCameraRigListChanged, OnCameraRigListChanged)
+		SLATE_EVENT(FOnRequestEditCameraRig, OnRequestEditCameraRig)
+	SLATE_END_ARGS()
+
+	void Construct(const FArguments& InArgs);
+	~SCameraRigList();
+
+protected:
+
+	// SWidget interface
+	virtual void Tick(const FGeometry& AllottedGeometry, const double InCurrentTime, const float InDeltaTime) override;
+
+private:
+
+	TSharedPtr<SWidget> GenerateToolbar();
+
+	void OnEditCameraRig();
+	bool CanEditCameraRig();
+
+	void OnAddCameraRig();
+
+	void OnRenameCameraRig();
+	bool CanRenameCameraRig();
+
+	void OnDeleteCameraRig();
+	bool CanDeleteCameraRig();
+
+	void GetEntryStrings(TSharedPtr<FCameraRigListItem> InItem, TArray<FString>& OutStrings);
+
+	void UpdateItemSource();
+	void UpdateFilteredItemSource();
+
+	TSharedRef<ITableRow> OnListGenerateItemRow(TSharedPtr<FCameraRigListItem> Item, const TSharedRef<STableViewBase>& OwnerTable);
+	void OnListItemScrolledIntoView(TSharedPtr<FCameraRigListItem> Item, const TSharedPtr<ITableRow>& ItemWidget);
+	void OnListMouseButtonDoubleClick(TSharedPtr<FCameraRigListItem> Item);
+	TSharedPtr<SWidget> OnListContextMenuOpening();
+
+	void OnSearchTextChanged(const FText& InFilterText);
+	void OnSearchTextCommitted(const FText& InFilterText, ETextCommit::Type InCommitType);
+	FText GetHighlightText() const;
+
+private:
+
+	TObjectPtr<UCameraAsset> CameraAsset;
+
+	FOnCameraRigListChanged OnCameraRigListChanged;
+	FOnRequestEditCameraRig OnRequestEditCameraRig;
+
+	TSharedPtr<FUICommandList> CommandList;
+
+	TSharedPtr<SListView<TSharedPtr<FCameraRigListItem>>> ListView;
+	TArray<TSharedPtr<FCameraRigListItem>> ItemSource;
+
+	using FEntryTextFilter = TTextFilter<const TSharedPtr<FCameraRigListItem>>;
+	TSharedPtr<FEntryTextFilter> SearchTextFilter;
+	TSharedPtr<SSearchBox> SearchBox;
+
+	TArray<TSharedPtr<FCameraRigListItem>> FilteredItemSource;
+	TSharedPtr<FCameraRigListItem> DeferredRequestRenameItem;
+
+	bool bUpdateItemSource = false;
+	bool bUpdateFilteredItemSource = false;
+};
+
+}  // namespace UE::Cameras
+
diff --git a/Engine/Plugins/Cameras/GameplayCameras/Source/GameplayCamerasEditor/Public/Editors/SCameraRigTransitionEditor.h b/Engine/Plugins/Cameras/GameplayCameras/Source/GameplayCamerasEditor/Public/Editors/SCameraRigTransitionEditor.h
new file mode 100644
--- /dev/null
+++ b/Engine/Plugins/Cameras/GameplayCameras/Source/GameplayCamerasEditor/Public/Editors/SCameraRigTransitionEditor.h
@@ -0,0 +1,127 @@
+// Copyright Epic Games, Inc. All Rights Reserved.
+
+#pragma once
+
+#include "CoreTypes.h"
+#include "Editors/CameraRigTransitionGraphSchema.h"
+#include "Editors/ObjectTreeGraphConfig.h"
+#include "GraphEditor.h"
+#include "Templates/SharedPointerFwd.h"
+#include "UObject/ObjectPtr.h"
+#include "Widgets/SCompoundWidget.h"
+
+class FAssetEditorToolkit;
+class IDetailsView;
+class SBox;
+class SObjectTreeGraphEditor;
+class UCameraRigTransition;
+class UEdGraphNode;
+class UObjectTreeGraph;
+struct FObjectTreeGraphConfig;
+
+namespace UE::Cameras
+{
+
+/**
+ * A graph editor for any object that has enter and exit transitions.
+ */
+class SCameraRigTransitionEditor : public SCompoundWidget
+{
+public:
+
+	SLATE_BEGIN_ARGS(SCameraRigTransitionEditor)
+	{}
+		/** The object owning the transitions. */
+		SLATE_ARGUMENT(UObject*, TransitionOwner)
+		/** Information about the transition owner. */
+		SLATE_ARGUMENT(FCameraRigTransitionOwnerInfo, TransitionOwnerInfo)
+		/** Callback to edit the transition graph config before it's used to build the editor. */
+		SLATE_EVENT(FOnBuildObjectTreeGraphConfig, OnBuildTransitionGraphConfig)
+		/** The details view to synchronize with the graph selection. */
+		SLATE_ARGUMENT(TSharedPtr<IDetailsView>, DetailsView)
+		/** The toolkit inside which this editor lives, if any. */
+		SLATE_ARGUMENT(TWeakPtr<FAssetEditorToolkit>, AssetEditorToolkit)
+		/** Display info for the transition graph. */
+		SLATE_ATTRIBUTE(FGraphDisplayInfo, TransitionGraphDisplayInfo)
+		/** Appearance info for the transition graph editor. */
+		SLATE_ATTRIBUTE(FGraphAppearanceInfo, TransitionGraphEditorAppearance)
+		/** Formatting of the transition owner's name. */
+		SLATE_ATTRIBUTE(FOnFormatObjectDisplayName, OnFormatTransitionOwnerName)
+	SLATE_END_ARGS()
+
+	void Construct(const FArguments& InArgs);
+	~SCameraRigTransitionEditor();
+
+	/** Sets the asset being edited for transitions. */
+	void SetTransitionOwner(UObject* InTransitionOwner);
+
+public:
+
+	/** Gets the transition graph. */
+	UEdGraph* GetTransitionGraph() const;
+	/** Gets the transition graph configuration. */
+	const FObjectTreeGraphConfig& GetTransitionGraphConfig() const;
+
+	/** Focuses the current graph to the root object node. */
+	void FocusHome();
+	/** Jumps the current graph to the given node. */
+	void JumpToNode(UEdGraphNode* InGraphNode);
+	/** Finds a node for the given object and, if so, jumps to it. */
+	bool FindAndJumpToObjectNode(UObject* InObject);
+
+	/** Adds a callback that will be invoked when the editor is changed. */
+	FDelegateHandle AddOnGraphChanged(FOnGraphChanged::FDelegate InAddDelegate);
+	/** Removes a previous added callback. */
+	void RemoveOnGraphChanged(FDelegateHandle InDelegateHandle);
+	/** Removes a previous added callback. */
+	void RemoveOnGraphChanged(const void* InUserObject);
+
+protected:
+
+	void CreateTransitionGraphEditor();
+	void DiscardTransitionGraphEditor();
+
+	void OnGraphChanged(const FEdGraphEditAction& InEditAction);
+
+	FText GetTransitionOwnerName() const;
+
+private:
+
+	/** The asset being edited. */
+	UObject* TransitionOwner;
+
+	/** Info about the type of asset we can edit. */
+	FCameraRigTransitionOwnerInfo TransitionOwnerInfo;
+
+	/** Callback for tweaking the graph config. */
+	FOnBuildObjectTreeGraphConfig OnBuildTransitionGraphConfig;
+
+	/** The details view for this editor. */
+	TSharedPtr<IDetailsView> DetailsView;
+
+	/** The owning toolkit. */
+	TWeakPtr<FAssetEditorToolkit> AssetEditorToolkit;
+
+	/** The transition graph. */
+	TObjectPtr<UObjectTreeGraph> TransitionGraph;
+	/** The transition graph editor. */
+	TSharedPtr<SObjectTreeGraphEditor> TransitionGraphEditor;
+
+	/** Appearance info for the transition graph editor. */
+	TAttribute<FGraphDisplayInfo> TransitionGraphDisplayInfo;
+
+	/** Appearance info for the transition graph editor. */
+	TAttribute<FGraphAppearanceInfo> TransitionGraphEditorAppearance;
+
+	/** The panel holding the graph editor. */
+	TSharedPtr<SBox> BoxPanel;
+
+	/** Handle for listening to changes in the graph editor. */
+	FDelegateHandle TransitionGraphChangedHandle;
+
+	/* Forwarding delegate for changes in the graph editor. */
+	FOnGraphChanged OnTransitionGraphChanged;
+};
+
+}  // namespace UE::Cameras
+
diff --git a/Engine/Plugins/Cameras/GameplayCameras/Source/GameplayCamerasEditor/Public/Editors/SFindInObjectTreeGraph.h b/Engine/Plugins/Cameras/GameplayCameras/Source/GameplayCamerasEditor/Public/Editors/SFindInObjectTreeGraph.h
--- a/Engine/Plugins/Cameras/GameplayCameras/Source/GameplayCamerasEditor/Public/Editors/SFindInObjectTreeGraph.h
+++ b/Engine/Plugins/Cameras/GameplayCameras/Source/GameplayCamerasEditor/Public/Editors/SFindInObjectTreeGraph.h
@@ -74,9 +74,9 @@ public:
 		/** The graphs to search. */
 		SLATE_ARGUMENT(TArray<UEdGraph*>, GraphsToSearch)
 		/** The callback to invoke when a search result wants to focus a node. */
-		SLATE_ARGUMENT(FOnJumpToNodeRequested, OnJumpToNodeRequested)
+		SLATE_EVENT(FOnJumpToNodeRequested, OnJumpToNodeRequested)
 		/** The callback to invoke when a search result wants to focus a pin. */
-		SLATE_ARGUMENT(FOnJumpToPinRequested, OnJumpToPinRequested)
+		SLATE_EVENT(FOnJumpToPinRequested, OnJumpToPinRequested)
 	SLATE_END_ARGS()
 
 	void Construct(const FArguments& InArgs);
diff --git a/Engine/Plugins/Cameras/GameplayCameras/Source/GameplayCamerasEditor/Public/GameplayCamerasEditorSettings.h b/Engine/Plugins/Cameras/GameplayCameras/Source/GameplayCamerasEditor/Public/GameplayCamerasEditorSettings.h
--- a/Engine/Plugins/Cameras/GameplayCameras/Source/GameplayCamerasEditor/Public/GameplayCamerasEditorSettings.h
+++ b/Engine/Plugins/Cameras/GameplayCameras/Source/GameplayCamerasEditor/Public/GameplayCamerasEditorSettings.h
@@ -35,5 +35,10 @@ public:
 
 	UPROPERTY(EditAnywhere, Config, Category=NodeTitleColors)
 	FLinearColor CameraBlendNodeTitleColor;
+
+public:
+
+	UPROPERTY()
+	FName LastCameraAssetToolkitModeName;
 };
 
diff --git a/Engine/Plugins/Cameras/GameplayCameras/Source/GameplayCamerasEditor/Public/Toolkits/AssetEditorMode.h b/Engine/Plugins/Cameras/GameplayCameras/Source/GameplayCamerasEditor/Public/Toolkits/AssetEditorMode.h
new file mode 100644
--- /dev/null
+++ b/Engine/Plugins/Cameras/GameplayCameras/Source/GameplayCamerasEditor/Public/Toolkits/AssetEditorMode.h
@@ -0,0 +1,79 @@
+// Copyright Epic Games, Inc. All Rights Reserved.
+
+#pragma once
+
+#include "CoreTypes.h"
+#include "Framework/Docking/TabManager.h"
+
+class FAssetEditorToolkit;
+class FExtender;
+class FLayoutExtender;
+class FUICommandList;
+struct FToolMenuContext;
+
+namespace UE::Cameras
+{
+
+/**
+ * Parameter structure for activating an asset editor mode.
+ */
+struct FAssetEditorModeActivateParams
+{
+	TSharedPtr<FAssetEditorToolkit> Toolkit;
+	TSharedPtr<FTabManager> TabManager;
+	TSharedPtr<FWorkspaceItem> AssetEditorTabsCategory;
+	TSharedPtr<FUICommandList> CommandList;
+	FName ToolbarMenuName;
+};
+
+/**
+ * Parameter structure for deactivating an asset editor mode.
+ */
+struct FAssetEditorModeDeactivateParams
+{
+	TSharedPtr<FAssetEditorToolkit> Toolkit;
+	TSharedPtr<FTabManager> TabManager;
+};
+
+/**
+ * An editor mode inside an FAssetEditorModeManagerToolkit.
+ *
+ * This changes the toolkit's editor to match a desired "editing mode" or "workflow"
+ * in that editor. The layout changes, the tabs change, the toolbars change, etc.
+ */
+class FAssetEditorMode 
+	: public TSharedFromThis<FAssetEditorMode>
+{
+public:
+
+	FAssetEditorMode();
+	FAssetEditorMode(FName InModeName);
+	virtual ~FAssetEditorMode();
+
+	void ActivateMode(const FAssetEditorModeActivateParams& InParams);
+	void DeactivateMode(const FAssetEditorModeDeactivateParams& InParams);
+
+	FName GetModeName() const { return ModeName; }
+	TSharedPtr<FTabManager::FLayout> GetDefaultLayout() const { return DefaultLayout; }
+	TSharedPtr<FExtender> GetToolbarExtender() const { return ToolbarExtender; }
+	TSharedPtr<FLayoutExtender> GetLayoutExtender() const { return LayoutExtender; }
+
+	void InitToolMenuContext(FToolMenuContext& MenuContext);
+
+protected:
+
+	virtual void OnActivateMode(const FAssetEditorModeActivateParams& InParams) {}
+	virtual void OnInitToolMenuContext(FToolMenuContext& MenuContext) {}
+	virtual void OnDeactivateMode(const FAssetEditorModeDeactivateParams& InParams) {}
+
+protected:
+
+	FName ModeName;
+
+	TSharedPtr<FTabManager::FLayout> DefaultLayout;
+	TSharedPtr<FExtender> ToolbarExtender;
+	TSharedPtr<FLayoutExtender> LayoutExtender;
+};
+
+}  // namespace UE::Cameras
+
diff --git a/Engine/Plugins/Cameras/GameplayCameras/Source/GameplayCamerasEditor/Public/Toolkits/AssetEditorModeManagerToolkit.h b/Engine/Plugins/Cameras/GameplayCameras/Source/GameplayCamerasEditor/Public/Toolkits/AssetEditorModeManagerToolkit.h
new file mode 100644
--- /dev/null
+++ b/Engine/Plugins/Cameras/GameplayCameras/Source/GameplayCamerasEditor/Public/Toolkits/AssetEditorModeManagerToolkit.h
@@ -0,0 +1,59 @@
+// Copyright Epic Games, Inc. All Rights Reserved.
+
+#pragma once
+
+#include "Containers/Map.h"
+#include "CoreTypes.h"
+#include "Tools/BaseAssetToolkit.h"
+
+struct FToolMenuContext;
+
+namespace UE::Cameras
+{
+
+class FAssetEditorMode;
+
+/**
+ * An editor toolkit that can manage different "editor modes".
+ *
+ * Similar to FWorkflowCentricApplication but with a few nuances such as making it
+ * possible to retain some common tabs between modes.
+ */
+class FAssetEditorModeManagerToolkit 
+	: public FBaseAssetToolkit
+{
+public:
+
+	FAssetEditorModeManagerToolkit(UAssetEditor* InOwningAssetEditor);
+	~FAssetEditorModeManagerToolkit();
+
+protected:
+
+	// FBaseAssetToolkit interface
+	virtual void InitToolMenuContext(FToolMenuContext& MenuContext) override;
+
+	// FCameraAssetEditorToolkitMode interface
+	virtual void OnEditorToolkitModeActivated() {}
+
+protected:
+
+	void AddEditorMode(TSharedRef<FAssetEditorMode> InMode);
+
+	TSharedPtr<FAssetEditorMode> GetEditorMode(FName InModeName) const;
+	void GetEditorModes(TArray<TSharedPtr<FAssetEditorMode>>& OutModes) const;
+
+	void SetEditorMode(FName InModeName);
+	bool CanSetEditorMode(FName InModeName) const;
+	bool IsEditorMode(FName InModeName) const;
+	FName GetCurrentEditorModeName() const;
+
+private:
+
+	TMap<FName, TSharedPtr<FAssetEditorMode>> EditorModes;
+
+	FName CurrentEditorModeName;
+	TSharedPtr<FAssetEditorMode> CurrentEditorMode;
+};
+
+}  // namespace UE::Cameras
+
diff --git a/Engine/Plugins/Cameras/GameplayCameras/Source/GameplayCamerasEditor/Public/Toolkits/BuildButtonToolkit.h b/Engine/Plugins/Cameras/GameplayCameras/Source/GameplayCamerasEditor/Public/Toolkits/BuildButtonToolkit.h
new file mode 100644
--- /dev/null
+++ b/Engine/Plugins/Cameras/GameplayCameras/Source/GameplayCamerasEditor/Public/Toolkits/BuildButtonToolkit.h
@@ -0,0 +1,43 @@
+// Copyright Epic Games, Inc. All Rights Reserved.
+
+#pragma once
+
+#include "CoreTypes.h"
+#include "Templates/SharedPointer.h"
+#include "UObject/ScriptInterface.h"
+
+class FUICommandInfo;
+class IHasCameraBuildStatus;
+struct FSlateIcon;
+struct FToolMenuEntry;
+
+namespace UE::Cameras
+{
+
+/**
+ * A utility toolkit for have a "Build" toolbar button that features a different
+ * icon based on the target asset's build status.
+ */
+class FBuildButtonToolkit : public TSharedFromThis<FBuildButtonToolkit>
+{
+public:
+
+	FBuildButtonToolkit();
+	FBuildButtonToolkit(TScriptInterface<IHasCameraBuildStatus> InTarget);
+
+	void SetTarget(TScriptInterface<IHasCameraBuildStatus> InTarget);
+
+	FToolMenuEntry MakeToolbarButton(TSharedPtr<FUICommandInfo> InCommand);
+
+private:
+
+	FSlateIcon GetBuildButtonIcon() const;
+	FText GetBuildButtonTooltip() const;
+
+private:
+
+	TScriptInterface<IHasCameraBuildStatus> Target;
+};
+
+}  // namespace UE::Cameras
+
diff --git a/Engine/Plugins/Cameras/GameplayCameras/Source/GameplayCamerasEditor/Public/Toolkits/CameraAssetEditorToolkit.h b/Engine/Plugins/Cameras/GameplayCameras/Source/GameplayCamerasEditor/Public/Toolkits/CameraAssetEditorToolkit.h
--- a/Engine/Plugins/Cameras/GameplayCameras/Source/GameplayCamerasEditor/Public/Toolkits/CameraAssetEditorToolkit.h
+++ b/Engine/Plugins/Cameras/GameplayCameras/Source/GameplayCamerasEditor/Public/Toolkits/CameraAssetEditorToolkit.h
@@ -3,22 +3,32 @@
 #pragma once
 
 #include "CoreTypes.h"
+#include "Toolkits/AssetEditorModeManagerToolkit.h"
 #include "Tools/BaseAssetToolkit.h"
 #include "UObject/GCObject.h"
 
-class IMessageLogListing;
+#include "CameraAssetEditorToolkit.generated.h"
+
+class SFindInObjectTreeGraph;
 class UCameraAsset;
 class UCameraAssetEditor;
-class SWidget;
+class UEdGraph;
+class UEdGraphNode;
+class UGameplayCamerasEditorSettings;
 
 namespace UE::Cameras
 {
 
+class FBuildButtonToolkit;
+class FCameraBuildLogToolkit;
+class FStandardToolkitLayout;
+class IGameplayCamerasLiveEditManager;
+
 /**
  * Editor toolkit for a camera asset.
  */
 class FCameraAssetEditorToolkit 
-	: public FBaseAssetToolkit
+	: public FAssetEditorModeManagerToolkit
 	, public FGCObject
 {
 public:
@@ -32,6 +42,8 @@ protected:
 	virtual void RegisterTabSpawners(const TSharedRef<class FTabManager>& TabManager) override;
 	virtual void UnregisterTabSpawners(const TSharedRef<FTabManager>& InTabManager) override;
 	virtual void CreateWidgets() override;
+	virtual void RegisterToolbar() override;
+	virtual void InitToolMenuContext(FToolMenuContext& MenuContext) override;
 	virtual void PostInitAssetEditor() override;
 
 	// IToolkit interface
@@ -41,28 +53,58 @@ protected:
 	virtual FLinearColor GetWorldCentricTabColorScale() const override;
 
 	// FGCObject interface
-	virtual void AddReferencedObjects(FReferenceCollector& Collector) override
-	{
-		Collector.AddReferencedObject(CameraAsset);
-	}
-	virtual FString GetReferencerName() const override
-	{
-		return TEXT("FCameraAssetEditorToolkit");
-	}
+	virtual void AddReferencedObjects(FReferenceCollector& Collector) override;
+	virtual FString GetReferencerName() const override;
+
+	// FAssetEditorModeManagerToolkit interface
+	virtual void OnEditorToolkitModeActivated() override;
 
 private:
 
-	static const FName DetailsViewTabId;
+	TSharedRef<SDockTab> SpawnTab_Search(const FSpawnTabArgs& Args);
+	TSharedRef<SDockTab> SpawnTab_Messages(const FSpawnTabArgs& Args);
 
-	/** Message log widget */
-	TSharedPtr<SWidget> Stats;
+	void JumpToNode(UEdGraphNode* Node);
+	void JumpToObject(UObject* Object);
 
-	/** Message log listing */
-	TSharedPtr<IMessageLogListing> StatsListing;
+	void OnBuild();
+	void OnFindInCamera();
+
+private:
+
+	static const FName SearchTabId;
+	static const FName MessagesTabId;
 
 	/** The asset being edited */
 	TObjectPtr<UCameraAsset> CameraAsset;
+
+	/** The layout for this toolkit */
+	TSharedPtr<FStandardToolkitLayout> StandardLayout;
+
+	/** The "Build" button */
+	TSharedPtr<FBuildButtonToolkit> BuildButtonToolkit;
+
+	/** The output/log window */
+	TSharedPtr<FCameraBuildLogToolkit> BuildLogToolkit;
+
+	/** The search results window */
+	TSharedPtr<SFindInObjectTreeGraph> SearchWidget;
+
+	TObjectPtr<UGameplayCamerasEditorSettings> Settings;
+
+	/** Live edit manager for updating the assets in the runtime */
+	TSharedPtr<IGameplayCamerasLiveEditManager> LiveEditManager;
 };
 
 }  // namespace UE::Cameras
 
+UCLASS()
+class UCameraAssetEditorMenuContext : public UObject
+{
+	GENERATED_BODY()
+
+public:
+
+	TWeakPtr<UE::Cameras::FCameraAssetEditorToolkit> Toolkit;
+};
+
diff --git a/Engine/Plugins/Cameras/GameplayCameras/Source/GameplayCamerasEditor/Public/Toolkits/CameraBuildLogToolkit.h b/Engine/Plugins/Cameras/GameplayCameras/Source/GameplayCamerasEditor/Public/Toolkits/CameraBuildLogToolkit.h
new file mode 100644
--- /dev/null
+++ b/Engine/Plugins/Cameras/GameplayCameras/Source/GameplayCamerasEditor/Public/Toolkits/CameraBuildLogToolkit.h
@@ -0,0 +1,54 @@
+// Copyright Epic Games, Inc. All Rights Reserved.
+
+#pragma once
+
+#include "CoreTypes.h"
+#include "Delegates/Delegate.h"
+#include "Templates/SharedPointer.h"
+
+class IMessageLogListing;
+class IMessageToken;
+class SWidget;
+class UObject;
+
+namespace UE::Cameras
+{
+
+class FCameraBuildLog;
+
+DECLARE_DELEGATE_OneParam(FRequestJumpToObject, UObject*);
+
+/**
+ * A utility toolkit class that creates and manages a "messages" or "output" tab that can
+ * show the build log of a camera asset or camera rig asset.
+ */
+class FCameraBuildLogToolkit : public TSharedFromThis<FCameraBuildLogToolkit>
+{
+public:
+
+	FCameraBuildLogToolkit();
+
+	void Initialize(FName InLogName);
+
+	TSharedPtr<IMessageLogListing> GetMessageListing() const { return MessageListing; }
+	TSharedPtr<SWidget> GetMessagesWidget() const { return MessagesWidget; }
+
+	FRequestJumpToObject& OnRequestJumpToObject() { return RequestJumpToObject; }
+
+	void PopulateMessageListing(FCameraBuildLog& InBuildLog);
+
+private:
+
+	void OnMessageTokenClicked(const TSharedRef<IMessageToken>& InMessageToken);
+
+private:
+
+	TSharedPtr<IMessageLogListing> MessageListing;
+
+	TSharedPtr<SWidget> MessagesWidget;
+
+	FRequestJumpToObject RequestJumpToObject;
+};
+
+}  // namespace UE::Cameras
+
diff --git a/Engine/Plugins/Cameras/GameplayCameras/Source/GameplayCamerasEditor/Public/Toolkits/CameraDirectorAssetEditorMode.h b/Engine/Plugins/Cameras/GameplayCameras/Source/GameplayCamerasEditor/Public/Toolkits/CameraDirectorAssetEditorMode.h
new file mode 100644
--- /dev/null
+++ b/Engine/Plugins/Cameras/GameplayCameras/Source/GameplayCamerasEditor/Public/Toolkits/CameraDirectorAssetEditorMode.h
@@ -0,0 +1,51 @@
+// Copyright Epic Games, Inc. All Rights Reserved.
+
+#pragma once
+
+#include "CoreTypes.h"
+#include "Toolkits/AssetEditorMode.h"
+
+class IDetailsView;
+class UCameraAsset;
+class UEdGraphNode;
+
+namespace UE::Cameras
+{
+
+class FStandardToolkitLayout;
+
+class FCameraDirectorAssetEditorMode
+	: public FAssetEditorMode
+{
+public:
+
+	static FName ModeName;
+
+	FCameraDirectorAssetEditorMode(UCameraAsset* InCameraAsset);
+
+	bool JumpToObject(UObject* InObject);
+
+protected:
+
+	virtual void OnActivateMode(const FAssetEditorModeActivateParams& InParams) override;
+	virtual void OnDeactivateMode(const FAssetEditorModeDeactivateParams& InParams) override;
+
+private:
+
+	TSharedRef<SDockTab> SpawnTab_DirectorEditor(const FSpawnTabArgs& Args);
+
+private:
+
+	static FName DirectorEditorTabId;
+
+	UCameraAsset* CameraAsset;
+
+	TSharedPtr<FStandardToolkitLayout> StandardLayout;
+
+	TSharedPtr<IDetailsView> DetailsView;
+
+	bool bInitializedToolkit = false;
+};
+
+}  // namespace UE::Cameras
+
diff --git a/Engine/Plugins/Cameras/GameplayCameras/Source/GameplayCamerasEditor/Public/Toolkits/CameraRigAssetEditorToolkit.h b/Engine/Plugins/Cameras/GameplayCameras/Source/GameplayCamerasEditor/Public/Toolkits/CameraRigAssetEditorToolkit.h
--- a/Engine/Plugins/Cameras/GameplayCameras/Source/GameplayCamerasEditor/Public/Toolkits/CameraRigAssetEditorToolkit.h
+++ b/Engine/Plugins/Cameras/GameplayCameras/Source/GameplayCamerasEditor/Public/Toolkits/CameraRigAssetEditorToolkit.h
@@ -3,41 +3,29 @@
 #pragma once
 
 #include "CoreTypes.h"
-#include "Editors/ObjectTreeGraphConfig.h"
-#include "Misc/NotifyHook.h"
 #include "Tools/BaseAssetToolkit.h"
-#include "UObject/GCObject.h"
 
 #include "CameraRigAssetEditorToolkit.generated.h"
 
-class IMessageLogListing;
 class SFindInObjectTreeGraph;
-class SObjectTreeGraphToolbox;
-class SWidget;
 class UCameraRigAsset;
-class UCameraRigAssetEditor;
-class UEdGraphNode;
-struct FEdGraphEditAction;
 
 namespace UE::Cameras
 {
 
-class FCameraBuildLog;
+class FBuildButtonToolkit;
+class FCameraBuildLogToolkit;
+class FCameraRigAssetEditorToolkitBase;
 class IGameplayCamerasLiveEditManager;
-class SCameraRigAssetEditor;
 
-/**
- * Editor toolkit for a camera rig asset.
- */
-class FCameraRigAssetEditorToolkit 
+class FCameraRigAssetEditorToolkit
 	: public FBaseAssetToolkit
-	, public FGCObject
-	, public FNotifyHook
 {
 public:
 
-	FCameraRigAssetEditorToolkit(UCameraRigAssetEditor* InOwningAssetEditor);
-	~FCameraRigAssetEditorToolkit();
+	FCameraRigAssetEditorToolkit(UAssetEditor* InOwningAssetEditor);
+
+	void SetCameraRigAsset(UCameraRigAsset* InCameraRig);
 
 protected:
 
@@ -55,68 +43,29 @@ protected:
 	virtual FString GetWorldCentricTabPrefix() const override;
 	virtual FLinearColor GetWorldCentricTabColorScale() const override;
 
-	// FGCObject interface
-	virtual void AddReferencedObjects(FReferenceCollector& Collector) override
-	{
-		Collector.AddReferencedObject(CameraRigAsset);
-	}
-	virtual FString GetReferencerName() const override
-	{
-		return TEXT("FCameraRigAssetEditorToolkit");
-	}
-
-	// FNotifyHook interface
-	virtual void NotifyPostChange(const FPropertyChangedEvent& PropertyChangedEvent, FProperty* PropertyThatChanged) override;
-
 private:
 
-	TSharedRef<SDockTab> SpawnTab_Toolbox(const FSpawnTabArgs& Args);
-	TSharedRef<SDockTab> SpawnTab_CameraRigEditor(const FSpawnTabArgs& Args);
 	TSharedRef<SDockTab> SpawnTab_Search(const FSpawnTabArgs& Args);
 	TSharedRef<SDockTab> SpawnTab_Messages(const FSpawnTabArgs& Args);
 
-	FSlateIcon GetBuildButtonIcon() const;
-	FText GetBuildButtonTooltip() const;
-
-	void OnAnyGraphChanged(const FEdGraphEditAction& InEditAction);
-
-	void OnMessageTokenClicked(const TSharedRef<IMessageToken>& InMessageToken);
-	void GetBuildMessageLogListing(FCameraBuildLog& InBuildLog);
-
 	void OnBuild();
 	void OnFindInCameraRig();
 
 private:
 
-	static const FName ToolboxTabId;
-	static const FName CameraRigEditorTabId;
 	static const FName SearchTabId;
 	static const FName MessagesTabId;
-	static const FName DetailsViewTabId;
-
-	/** The asset being edited */
-	TObjectPtr<UCameraRigAsset> CameraRigAsset;
-
-	/** Command bindings */
-	TSharedRef<FUICommandList> CommandBindings;
-
-	/** Message log listing */
-	TSharedPtr<IMessageLogListing> MessageListing;
-
-	/** Camera rig editor widget */
-	TSharedPtr<SCameraRigAssetEditor> CameraRigEditorWidget;
 
-	/** Toolbox widget */
-	TSharedPtr<SObjectTreeGraphToolbox> ToolboxWidget;
+	TSharedPtr<FCameraRigAssetEditorToolkitBase> Impl;
 
-	/** Message log widget */
-	TSharedPtr<SWidget> MessagesWidget;
+	TSharedPtr<FBuildButtonToolkit> BuildButtonToolkit;
+	TSharedPtr<FCameraBuildLogToolkit> BuildLogToolkit;
 
 	/** Search widget */
 	TSharedPtr<SFindInObjectTreeGraph> SearchWidget;
 
 	/** Live edit manager for updating the assets in the runtime */
-	TSharedPtr<UE::Cameras::IGameplayCamerasLiveEditManager> LiveEditManager;
+	TSharedPtr<IGameplayCamerasLiveEditManager> LiveEditManager;
 };
 
 }  // namespace UE::Cameras
@@ -128,6 +77,6 @@ class UCameraRigAssetEditorMenuContext : public UObject
 
 public:
 
-	TWeakPtr<UE::Cameras::FCameraRigAssetEditorToolkit> CameraRigAssetEditorToolkit;
+	TWeakPtr<UE::Cameras::FCameraRigAssetEditorToolkit> Toolkit;
 };
 
diff --git a/Engine/Plugins/Cameras/GameplayCameras/Source/GameplayCamerasEditor/Public/Toolkits/CameraRigAssetEditorToolkitBase.h b/Engine/Plugins/Cameras/GameplayCameras/Source/GameplayCamerasEditor/Public/Toolkits/CameraRigAssetEditorToolkitBase.h
new file mode 100644
--- /dev/null
+++ b/Engine/Plugins/Cameras/GameplayCameras/Source/GameplayCamerasEditor/Public/Toolkits/CameraRigAssetEditorToolkitBase.h
@@ -0,0 +1,95 @@
+// Copyright Epic Games, Inc. All Rights Reserved.
+
+#pragma once
+
+#include "CoreTypes.h"
+#include "Misc/NotifyHook.h"
+#include "UObject/GCObject.h"
+
+class FSpawnTabArgs;
+class FTabManager;
+class FUICommandList;
+class FWorkspaceItem;
+class IDetailsView;
+class SDockTab;
+class SObjectTreeGraphToolbox;
+class UCameraRigAsset;
+class UToolMenu;
+struct FEdGraphEditAction;
+
+namespace UE::Cameras
+{
+
+class FStandardToolkitLayout;
+class SCameraRigAssetEditor;
+
+/**
+ * Editor toolkit for a camera rig asset.
+ */
+class FCameraRigAssetEditorToolkitBase
+	: public FGCObject
+	, public FNotifyHook
+	, public TSharedFromThis<FCameraRigAssetEditorToolkitBase>
+{
+public:
+
+	FCameraRigAssetEditorToolkitBase(FName InLayoutName);
+	~FCameraRigAssetEditorToolkitBase();
+
+	UCameraRigAsset* GetCameraRigAsset() const { return CameraRigAsset; }
+	void SetCameraRigAsset(UCameraRigAsset* InCameraRig);
+
+	TSharedPtr<FStandardToolkitLayout> GetStandardLayout() const { return StandardLayout; }
+	TSharedPtr<IDetailsView> GetDetailsView() const { return DetailsView; }
+	TSharedPtr<SCameraRigAssetEditor> GetCameraRigAssetEditor() const { return CameraRigEditorWidget; }
+
+	void RegisterTabSpawners(TSharedRef<FTabManager> InTabManager, TSharedPtr<FWorkspaceItem> InAssetEditorTabsCategory);
+	void UnregisterTabSpawners(TSharedRef<FTabManager> InTabManager);
+	void CreateWidgets();
+	void BuildToolbarMenu(UToolMenu* ToolbarMenu);
+	void BindCommands(TSharedRef<FUICommandList> CommandList);
+
+protected:
+
+	// FGCObject interface
+	virtual void AddReferencedObjects(FReferenceCollector& Collector) override;
+	virtual FString GetReferencerName() const override;
+
+	// FNotifyHook interface
+	virtual void NotifyPostChange(const FPropertyChangedEvent& PropertyChangedEvent, FProperty* PropertyThatChanged) override;
+
+private:
+
+	TSharedRef<SDockTab> SpawnTab_Toolbox(const FSpawnTabArgs& Args);
+	TSharedRef<SDockTab> SpawnTab_CameraRigEditor(const FSpawnTabArgs& Args);
+	TSharedRef<SDockTab> SpawnTab_Details(const FSpawnTabArgs& Args);
+
+	void OnAnyGraphChanged(const FEdGraphEditAction& InEditAction);
+
+private:
+
+	static const FName ToolboxTabId;
+	static const FName CameraRigEditorTabId;
+	static const FName DetailsViewTabId;
+
+	/** The asset being edited */
+	TObjectPtr<UCameraRigAsset> CameraRigAsset;
+
+	/** The layout for this editor toolkit */
+	TSharedPtr<FStandardToolkitLayout> StandardLayout;
+
+	/** Command bindings */
+	TSharedPtr<FUICommandList> CommandBindings;
+
+	/** The details view */
+	TSharedPtr<IDetailsView> DetailsView;
+
+	/** Camera rig editor widget */
+	TSharedPtr<SCameraRigAssetEditor> CameraRigEditorWidget;
+
+	/** Toolbox widget */
+	TSharedPtr<SObjectTreeGraphToolbox> ToolboxWidget;
+};
+
+}  // namespace UE::Cameras
+
diff --git a/Engine/Plugins/Cameras/GameplayCameras/Source/GameplayCamerasEditor/Public/Toolkits/CameraRigTransitionEditorToolkit.h b/Engine/Plugins/Cameras/GameplayCameras/Source/GameplayCamerasEditor/Public/Toolkits/CameraRigTransitionEditorToolkit.h
new file mode 100644
--- /dev/null
+++ b/Engine/Plugins/Cameras/GameplayCameras/Source/GameplayCamerasEditor/Public/Toolkits/CameraRigTransitionEditorToolkit.h
@@ -0,0 +1,49 @@
+// Copyright Epic Games, Inc. All Rights Reserved.
+
+#pragma once
+
+#include "CoreTypes.h"
+#include "Tools/BaseAssetToolkit.h"
+
+class FTabManager;
+class ICameraRigTransitionOwner;
+class UAssetEditor;
+
+namespace UE::Cameras
+{
+
+class FCameraRigTransitionEditorToolkitBase;
+
+/**
+ * Editor toolkit for camera transitions.
+ */
+class FCameraRigTransitionEditorToolkit
+	: public FBaseAssetToolkit
+{
+public:
+
+	FCameraRigTransitionEditorToolkit(UAssetEditor* InOwningAssetEditor);
+
+	void SetTransitionOwner(UObject* InTransitionOwner);
+		
+protected:
+
+	// FBaseAssetToolkit interface
+	virtual void RegisterTabSpawners(const TSharedRef<FTabManager>& InTabManager) override;
+	virtual void UnregisterTabSpawners(const TSharedRef<FTabManager>& InTabManager) override;
+	virtual void CreateWidgets() override;
+	virtual void RegisterToolbar() override;
+
+	// IToolkit interface
+	virtual FText GetBaseToolkitName() const override;
+	virtual FName GetToolkitFName() const override;
+	virtual FString GetWorldCentricTabPrefix() const override;
+	virtual FLinearColor GetWorldCentricTabColorScale() const override;
+
+private:
+
+	TSharedPtr<FCameraRigTransitionEditorToolkitBase> Impl;
+};
+
+}  // namespace UE::Cameras
+
diff --git a/Engine/Plugins/Cameras/GameplayCameras/Source/GameplayCamerasEditor/Public/Toolkits/CameraRigTransitionEditorToolkitBase.h b/Engine/Plugins/Cameras/GameplayCameras/Source/GameplayCamerasEditor/Public/Toolkits/CameraRigTransitionEditorToolkitBase.h
new file mode 100644
--- /dev/null
+++ b/Engine/Plugins/Cameras/GameplayCameras/Source/GameplayCamerasEditor/Public/Toolkits/CameraRigTransitionEditorToolkitBase.h
@@ -0,0 +1,95 @@
+// Copyright Epic Games, Inc. All Rights Reserved.
+
+#pragma once
+
+#include "CoreTypes.h"
+#include "Editors/CameraRigTransitionGraphSchema.h"
+#include "Misc/NotifyHook.h"
+#include "UObject/GCObject.h"
+#include "UObject/ScriptInterface.h"
+
+class FSpawnTabArgs;
+class FTabManager;
+class FUICommandList;
+class FWorkspaceItem;
+class IDetailsView;
+class SDockTab;
+class SObjectTreeGraphToolbox;
+class UToolMenu;
+struct FEdGraphEditAction;
+
+namespace UE::Cameras
+{
+
+class FStandardToolkitLayout;
+class SCameraRigTransitionEditor;
+
+class FCameraRigTransitionEditorToolkitBase
+	: public FGCObject
+	, public FNotifyHook
+	, public TSharedFromThis<FCameraRigTransitionEditorToolkitBase>
+{
+public:
+
+	FCameraRigTransitionEditorToolkitBase(FName InLayoutName);
+	~FCameraRigTransitionEditorToolkitBase();
+
+	UObject* GetTransitionOwner() const { return TransitionOwner; }
+	void SetTransitionOwner(UObject* InTransitionOwner);
+
+	TSharedPtr<FStandardToolkitLayout> GetStandardLayout() const { return StandardLayout; }
+	TSharedPtr<SCameraRigTransitionEditor> GetCameraRigTransitionEditor() const { return TransitionEditorWidget; }
+
+	void RegisterTabSpawners(TSharedRef<FTabManager> InTabManager, TSharedPtr<FWorkspaceItem> InAssetEditorTabsCategory);
+	void UnregisterTabSpawners(TSharedRef<FTabManager> InTabManager);
+	void CreateWidgets();
+	void BuildToolbarMenu(UToolMenu* ToolbarMenu);
+
+protected:
+
+	// FGCObject interface
+	virtual void AddReferencedObjects(FReferenceCollector& Collector) override;
+	virtual FString GetReferencerName() const override;
+
+	// FNotifyHook interface
+	virtual void NotifyPostChange(const FPropertyChangedEvent& PropertyChangedEvent, FProperty* PropertyThatChanged) override;
+
+	// FCameraRigTransitionEditorToolkitBase interface
+	virtual void GetTransitionOwnerInfo(FCameraRigTransitionOwnerInfo& OutTransitionOwnerInfo) {}
+	virtual void GetTransitionGraphDisplayInfo(FGraphDisplayInfo& OutGraphDisplayInfo) {}
+	virtual void GetTransitionGraphAppearanceInfo(FGraphAppearanceInfo& OutGraphAppearanceInfo) {}
+
+private:
+
+	TSharedRef<SDockTab> SpawnTab_Toolbox(const FSpawnTabArgs& Args);
+	TSharedRef<SDockTab> SpawnTab_TransitionEditor(const FSpawnTabArgs& Args);
+	TSharedRef<SDockTab> SpawnTab_Details(const FSpawnTabArgs& Args);
+
+	void OnTransitionGraphChanged(const FEdGraphEditAction& InEditAction);
+
+private:
+
+	static const FName ToolboxTabId;
+	static const FName TransitionEditorTabId;
+	static const FName DetailsViewTabId;
+
+	/** The object being edited */
+	TObjectPtr<UObject> TransitionOwner;
+
+	/** The layout for this editor */
+	TSharedPtr<FStandardToolkitLayout> StandardLayout;
+
+	/** Command bindings */
+	TSharedPtr<FUICommandList> CommandBindings;
+
+	TSharedPtr<IDetailsView> DetailsView;
+
+	/** Camera transition editor widget */
+	TSharedPtr<SCameraRigTransitionEditor> TransitionEditorWidget;
+
+	/** Toolbox widget */
+	TSharedPtr<SObjectTreeGraphToolbox> ToolboxWidget;
+};
+
+}  // namespace UE::Cameras
+
diff --git a/Engine/Plugins/Cameras/GameplayCameras/Source/GameplayCamerasEditor/Public/Toolkits/CameraRigsAssetEditorMode.h b/Engine/Plugins/Cameras/GameplayCameras/Source/GameplayCamerasEditor/Public/Toolkits/CameraRigsAssetEditorMode.h
new file mode 100644
--- /dev/null
+++ b/Engine/Plugins/Cameras/GameplayCameras/Source/GameplayCamerasEditor/Public/Toolkits/CameraRigsAssetEditorMode.h
@@ -0,0 +1,58 @@
+// Copyright Epic Games, Inc. All Rights Reserved.
+
+#pragma once
+
+#include "CoreTypes.h"
+#include "Toolkits/AssetEditorMode.h"
+#include "Editors/SCameraRigList.h"
+
+class FDocumentTabFactory;
+class FDocumentTracker;
+class UCameraAsset;
+class UCameraRigAsset;
+class UEdGraphNode;
+
+namespace UE::Cameras
+{
+
+class FCameraRigAssetEditorToolkitBase;
+
+class FCameraRigsAssetEditorMode
+	: public FAssetEditorMode
+{
+public:
+
+	static FName ModeName;
+
+	FCameraRigsAssetEditorMode(UCameraAsset* InCameraAsset);
+
+	bool JumpToNode(UEdGraphNode* InNode);
+	bool JumpToObject(UObject* InObject);
+
+protected:
+
+	virtual void OnActivateMode(const FAssetEditorModeActivateParams& InParams) override;
+	virtual void OnDeactivateMode(const FAssetEditorModeDeactivateParams& InParams) override;
+
+private:
+
+	TSharedRef<SDockTab> SpawnTab_CameraRigs(const FSpawnTabArgs& Args);
+
+	void OnCameraRigListChanged(TArrayView<UCameraRigAsset* const> InCameraRigs);
+	void OnCameraRigEditRequested(UCameraRigAsset* InCameraRig);
+
+private:
+
+	static const FName CameraRigsTabId;
+
+	UCameraAsset* CameraAsset;
+
+	TSharedPtr<FCameraRigAssetEditorToolkitBase> Impl;
+
+	TSharedPtr<SCameraRigList> CameraRigsListWidget;
+
+	bool bInitializedToolkit = false;
+};
+
+}  // namespace UE::Cameras
+
diff --git a/Engine/Plugins/Cameras/GameplayCameras/Source/GameplayCamerasEditor/Public/Toolkits/CameraSharedTransitionsAssetEditorMode.h b/Engine/Plugins/Cameras/GameplayCameras/Source/GameplayCamerasEditor/Public/Toolkits/CameraSharedTransitionsAssetEditorMode.h
new file mode 100644
--- /dev/null
+++ b/Engine/Plugins/Cameras/GameplayCameras/Source/GameplayCamerasEditor/Public/Toolkits/CameraSharedTransitionsAssetEditorMode.h
@@ -0,0 +1,43 @@
+// Copyright Epic Games, Inc. All Rights Reserved.
+
+#pragma once
+
+#include "CoreTypes.h"
+#include "Toolkits/AssetEditorMode.h"
+
+class UCameraAsset;
+class UEdGraphNode;
+
+namespace UE::Cameras
+{
+
+class FCameraRigTransitionEditorToolkitBase;
+
+class FCameraSharedTransitionsAssetEditorMode
+	: public FAssetEditorMode
+{
+public:
+
+	static FName ModeName;
+
+	FCameraSharedTransitionsAssetEditorMode(UCameraAsset* InCameraAsset);
+
+	bool JumpToNode(UEdGraphNode* InNode);
+	bool JumpToObject(UObject* InObject);
+
+protected:
+
+	virtual void OnActivateMode(const FAssetEditorModeActivateParams& InParams) override;
+	virtual void OnDeactivateMode(const FAssetEditorModeDeactivateParams& InParams) override;
+
+private:
+
+	UCameraAsset* CameraAsset;
+
+	TSharedPtr<FCameraRigTransitionEditorToolkitBase> Impl;
+
+	bool bInitializedToolkit = false;
+};
+
+}  // namespace UE::Cameras
+
diff --git a/Engine/Plugins/Cameras/GameplayCameras/Source/GameplayCamerasEditor/Public/Toolkits/StandardToolkitLayout.h b/Engine/Plugins/Cameras/GameplayCameras/Source/GameplayCamerasEditor/Public/Toolkits/StandardToolkitLayout.h
new file mode 100644
--- /dev/null
+++ b/Engine/Plugins/Cameras/GameplayCameras/Source/GameplayCamerasEditor/Public/Toolkits/StandardToolkitLayout.h
@@ -0,0 +1,50 @@
+// Copyright Epic Games, Inc. All Rights Reserved.
+
+#pragma once
+
+#include "CoreTypes.h"
+#include "Framework/Docking/TabManager.h"
+
+namespace UE::Cameras
+{
+
+/**
+ * A utility class for building a standard layout with tab stacks on the
+ * left, right, and bottom of a central document area.
+ */
+class FStandardToolkitLayout
+{
+public:
+
+	static const FName LeftStackExtensionId;
+	static const FName CenterStackExtensionId;
+	static const FName RightStackExtensionId;
+	static const FName BottomStackExtensionId;
+
+	FStandardToolkitLayout(FName LayoutName);
+
+	TSharedPtr<FTabManager::FLayout> GetLayout() const;
+
+	void AddLeftTab(FName InTabId, ETabState::Type InTabState = ETabState::OpenedTab);
+	void AddRightTab(FName InTabId, ETabState::Type InTabState = ETabState::OpenedTab);
+	void AddBottomTab(FName InTabId, ETabState::Type InTabState = ETabState::OpenedTab);
+	void AddCenterTab(FName InTabId, ETabState::Type InTabState = ETabState::OpenedTab);
+
+private:
+
+	void BuildLayout();
+
+private:
+
+	FName LayoutName;
+
+	TSharedPtr<FTabManager::FLayout> Layout;
+
+	TSharedPtr<FTabManager::FStack> LeftTabStack;
+	TSharedPtr<FTabManager::FStack> CenterTabStack;
+	TSharedPtr<FTabManager::FStack> RightTabStack;
+	TSharedPtr<FTabManager::FStack> BottomTabStack;
+};
+
+}  // namespace UE::Cameras
+
