Cameras: add a dialog on camera asset creation to choose a camera director

diff --git a/Engine/Plugins/Cameras/GameplayCameras/Source/GameplayCamerasEditor/Private/Factories/CameraAssetFactory.cpp b/Engine/Plugins/Cameras/GameplayCameras/Source/GameplayCamerasEditor/Private/Factories/CameraAssetFactory.cpp
--- a/Engine/Plugins/Cameras/GameplayCameras/Source/GameplayCamerasEditor/Private/Factories/CameraAssetFactory.cpp
+++ b/Engine/Plugins/Cameras/GameplayCameras/Source/GameplayCamerasEditor/Private/Factories/CameraAssetFactory.cpp
@@ -3,11 +3,50 @@
 #include "Factories/CameraAssetFactory.h"
 
 #include "Core/CameraAsset.h"
+#include "Core/CameraDirector.h"
+#include "ClassViewerFilter.h"
+#include "ClassViewerModule.h"
+#include "Directors/BlueprintCameraDirector.h"
+#include "Directors/SingleCameraDirector.h"
+#include "Kismet2/SClassPickerDialog.h"
+#include "Modules/ModuleManager.h"
 
 #include UE_INLINE_GENERATED_CPP_BY_NAME(CameraAssetFactory)
 
 #define LOCTEXT_NAMESPACE "CameraAssetFactory"
 
+namespace UE::Cameras::Private
+{
+
+class FCameraDirectorClassFilter : public IClassViewerFilter
+{
+public:
+	TSet<const UClass*> AllowedClasses;
+	EClassFlags DisallowedClassFlags;
+
+	FCameraDirectorClassFilter()
+	{
+		AllowedClasses.Add(UCameraDirector::StaticClass());
+		DisallowedClassFlags = CLASS_Abstract | CLASS_Deprecated;
+	}
+
+	virtual bool IsClassAllowed(const FClassViewerInitializationOptions& InInitOptions, const UClass* InClass, TSharedRef<FClassViewerFilterFuncs> InFilterFuncs) override
+	{
+		return 
+			!InClass->HasAnyClassFlags(DisallowedClassFlags) &&
+			InFilterFuncs->IfInChildOfClassesSet(AllowedClasses, InClass) != EFilterReturn::Failed;
+	}
+
+	virtual bool IsUnloadedClassAllowed(const FClassViewerInitializationOptions& InInitOptions, const TSharedRef< const IUnloadedBlueprintData > InUnloadedClassData, TSharedRef< FClassViewerFilterFuncs > InFilterFuncs) override
+	{
+		return 
+			!InUnloadedClassData->HasAnyClassFlags(DisallowedClassFlags) && 
+			InFilterFuncs->IfInChildOfClassesSet(AllowedClasses, InUnloadedClassData) != EFilterReturn::Failed;
+	}
+};
+
+}  // namespace UE::Cameras::Private
+
 UCameraAssetFactory::UCameraAssetFactory(const FObjectInitializer& ObjectInitializer)
 	: Super(ObjectInitializer)
 {
@@ -19,12 +58,40 @@ UCameraAssetFactory::UCameraAssetFactory(const FObjectInitializer& ObjectInitial
 UObject* UCameraAssetFactory::FactoryCreateNew(UClass* Class, UObject* Parent, FName Name, EObjectFlags Flags, UObject* Context, FFeedbackContext* Warn)
 {
 	UCameraAsset* NewCameraAsset = NewObject<UCameraAsset>(Parent, Class, Name, Flags | RF_Transactional);
+	if (CameraDirectorClass)
+	{
+		UCameraDirector* NewCameraDirector = NewObject<UCameraDirector>(NewCameraAsset, CameraDirectorClass, NAME_None, Flags | RF_Transactional);
+		NewCameraAsset->CameraDirector = NewCameraDirector;
+	}
 	return NewCameraAsset;
 }
 
 bool UCameraAssetFactory::ConfigureProperties()
 {
-	return true;
+	using namespace UE::Cameras::Private;
+
+	FClassViewerModule& ClassViewerModule = FModuleManager::LoadModuleChecked<FClassViewerModule>("ClassViewer");
+
+	FClassViewerInitializationOptions Options;
+	Options.Mode = EClassViewerMode::ClassPicker;
+	Options.NameTypeToDisplay = EClassViewerNameTypeToDisplay::DisplayName;
+	Options.bShowNoneOption = true;
+
+	Options.ExtraPickerCommonClasses.Add(UBlueprintCameraDirector::StaticClass());
+	Options.ExtraPickerCommonClasses.Add(USingleCameraDirector::StaticClass());
+
+	TSharedPtr<FCameraDirectorClassFilter> Filter = MakeShareable(new FCameraDirectorClassFilter);
+	Options.ClassFilters.Add(Filter.ToSharedRef());
+
+	CameraDirectorClass = nullptr;
+	UClass* ChosenClass = nullptr;
+	const FText TitleText = LOCTEXT("CameraDirectorPicker", "Pick Camera Director Type");
+	const bool bPressedOk = SClassPickerDialog::PickClass(TitleText, Options, ChosenClass, UCameraDirector::StaticClass());
+	if (bPressedOk)
+	{
+		CameraDirectorClass = ChosenClass;
+	}
+	return bPressedOk;
 }
 
 bool UCameraAssetFactory::ShouldShowInNewMenu() const
diff --git a/Engine/Plugins/Cameras/GameplayCameras/Source/GameplayCamerasEditor/Private/Factories/CameraAssetFactory.h b/Engine/Plugins/Cameras/GameplayCameras/Source/GameplayCamerasEditor/Private/Factories/CameraAssetFactory.h
--- a/Engine/Plugins/Cameras/GameplayCameras/Source/GameplayCamerasEditor/Private/Factories/CameraAssetFactory.h
+++ b/Engine/Plugins/Cameras/GameplayCameras/Source/GameplayCamerasEditor/Private/Factories/CameraAssetFactory.h
@@ -7,6 +7,7 @@
 #include "CameraAssetFactory.generated.h"
 
 class UCameraAsset;
+class UCameraDirector;
 
 /**
  * Implements a factory for UCameraAsset objects.
@@ -16,11 +17,19 @@ class UCameraAssetFactory : public UFactory
 {
 	GENERATED_BODY()
 
+public:
+
 	UCameraAssetFactory(const FObjectInitializer& ObjectInitializer);
 
 	// UFactory Interface
 	virtual UObject* FactoryCreateNew(UClass* Class, UObject* Parent, FName Name, EObjectFlags Flags, UObject* Context, FFeedbackContext* Warn) override;
 	virtual bool ConfigureProperties() override;
 	virtual bool ShouldShowInNewMenu() const override;
+
+private:
+
+	/** The type of camera director to create for this asset. */
+	UPROPERTY(EditAnywhere, Category=Camera)
+	TSubclassOf<UCameraDirector> CameraDirectorClass;
 };
 
