Cameras: make Blueprint camera director activate rigs by name instead of by reference

diff --git a/Engine/Plugins/Cameras/GameplayCameras/Source/GameplayCameras/Private/Directors/BlueprintCameraDirector.cpp b/Engine/Plugins/Cameras/GameplayCameras/Source/GameplayCameras/Private/Directors/BlueprintCameraDirector.cpp
--- a/Engine/Plugins/Cameras/GameplayCameras/Source/GameplayCameras/Private/Directors/BlueprintCameraDirector.cpp
+++ b/Engine/Plugins/Cameras/GameplayCameras/Source/GameplayCameras/Private/Directors/BlueprintCameraDirector.cpp
@@ -2,8 +2,10 @@
 
 #include "Directors/BlueprintCameraDirector.h"
 
-#include "Algo/StableSort.h"
+#include "Core/CameraAsset.h"
+#include "Core/CameraRigAsset.h"
 #include "Core/CameraEvaluationContext.h"
+#include "GameplayCameras.h"
 
 #include UE_INLINE_GENERATED_CPP_BY_NAME(BlueprintCameraDirector)
 
@@ -20,6 +22,10 @@ protected:
 	virtual void OnRun(const FCameraDirectorEvaluationParams& Params, FCameraDirectorEvaluationResult& OutResult) override;
 	virtual void OnAddReferencedObjects(FReferenceCollector& Collector) override;
 
+private:
+
+	const UCameraRigAsset* FindCameraRigByName(const UCameraAsset* InCameraAsset, const FString& InCameraRigName);
+
 private:
 
 	TObjectPtr<UBlueprintCameraDirectorEvaluator> EvaluatorBlueprint;
@@ -30,18 +36,32 @@ UE_DEFINE_CAMERA_DIRECTOR_EVALUATOR(FBlueprintCameraDirectorEvaluator)
 void FBlueprintCameraDirectorEvaluator::OnInitialize(const FCameraDirectorInitializeParams& Params)
 {
 	const UBlueprintCameraDirector* Blueprint = GetCameraDirectorAs<UBlueprintCameraDirector>();
-	if (!ensure(Blueprint && Blueprint->CameraDirectorEvaluatorClass))
+	if (!ensure(Blueprint))
+	{
+		return;
+	}
+
+	const UCameraAsset* CameraAsset = Params.OwnerContext->GetCameraAsset();
+	if (!ensure(CameraAsset))
 	{
 		return;
 	}
 
-	UObject* Outer = Params.OwnerContext->GetOwner();
-	EvaluatorBlueprint = NewObject<UBlueprintCameraDirectorEvaluator>(Outer, Blueprint->CameraDirectorEvaluatorClass);
+	if (Blueprint->CameraDirectorEvaluatorClass)
+	{
+		UObject* Outer = Params.OwnerContext->GetOwner();
+		EvaluatorBlueprint = NewObject<UBlueprintCameraDirectorEvaluator>(Outer, Blueprint->CameraDirectorEvaluatorClass);
+	}
+	else
+	{
+		UE_LOG(LogCameraSystem, Error, TEXT("No Blueprint class set on camera director for '%s'."), *CameraAsset->GetPathName());
+	}
 }
 
 void FBlueprintCameraDirectorEvaluator::OnRun(const FCameraDirectorEvaluationParams& Params, FCameraDirectorEvaluationResult& OutResult)
 {
-	if (EvaluatorBlueprint)
+	const UCameraAsset* CameraAsset = Params.OwnerContext->GetCameraAsset();
+	if (EvaluatorBlueprint && CameraAsset)
 	{
 		FBlueprintCameraDirectorEvaluationParams BlueprintParams;
 		BlueprintParams.DeltaTime = Params.DeltaTime;
@@ -56,11 +76,39 @@ void FBlueprintCameraDirectorEvaluator::OnRun(const FCameraDirectorEvaluationPar
 
 		// The BP interface doesn't specify the evaluation context for the chosen camera rigs: we always automatically
 		// make them run in our own owner context.
-		for (UCameraRigAsset* ActiveCameraRig : BlueprintResult.ActiveCameraRigs)
+		for (const FString& ActiveCameraRigName : BlueprintResult.ActiveCameraRigs)
 		{
-			OutResult.Add(Params.OwnerContext, ActiveCameraRig);
+			const UCameraRigAsset* ActiveCameraRig = FindCameraRigByName(CameraAsset, ActiveCameraRigName);
+			if (ActiveCameraRig)
+			{
+				OutResult.Add(Params.OwnerContext, ActiveCameraRig);
+			}
+			else
+			{
+				UE_LOG(
+						LogCameraSystem, 
+						Error, 
+						TEXT("Can't activate camera rig '%s' because no camera rig of that name was found on '%s'."),
+						*ActiveCameraRigName, *CameraAsset->GetPathName());
+			}
 		}
 	}
+	else
+	{
+		UE_LOG(LogCameraSystem, Error, TEXT("Can't run Blueprint camera director, no Blueprint set, or no owner camera set!"));
+	}
+}
+
+const UCameraRigAsset* FBlueprintCameraDirectorEvaluator::FindCameraRigByName(const UCameraAsset* InCameraAsset, const FString& InCameraRigName)
+{
+	for (const UCameraRigAsset* CameraRig : InCameraAsset->CameraRigs)
+	{
+		if (CameraRig->GetName() == InCameraRigName || CameraRig->Interface.DisplayName == InCameraRigName)
+		{
+			return CameraRig;
+		}
+	}
+	return nullptr;
 }
 
 void FBlueprintCameraDirectorEvaluator::OnAddReferencedObjects(FReferenceCollector& Collector)
@@ -70,9 +118,19 @@ void FBlueprintCameraDirectorEvaluator::OnAddReferencedObjects(FReferenceCollect
 
 }  // namespace UE::Cameras
 
+void UBlueprintCameraDirectorEvaluator::ActivateCameraRig(const FString& InCameraRigName)
+{
+	CurrentResult.ActiveCameraRigs.Add(InCameraRigName);
+}
+
 void UBlueprintCameraDirectorEvaluator::NativeRunCameraDirector(const FBlueprintCameraDirectorEvaluationParams& Params, FBlueprintCameraDirectorEvaluationResult& OutResult)
 {
-	RunCameraDirector(Params, OutResult);
+	CurrentResult = OutResult;
+	{
+		// Run the Blueprint logic.
+		RunCameraDirector(Params);
+	}
+	OutResult = CurrentResult;
 }
 
 FCameraDirectorEvaluatorPtr UBlueprintCameraDirector::OnBuildEvaluator(FCameraDirectorEvaluatorBuilder& Builder) const
diff --git a/Engine/Plugins/Cameras/GameplayCameras/Source/GameplayCameras/Public/Directors/BlueprintCameraDirector.h b/Engine/Plugins/Cameras/GameplayCameras/Source/GameplayCameras/Public/Directors/BlueprintCameraDirector.h
--- a/Engine/Plugins/Cameras/GameplayCameras/Source/GameplayCameras/Public/Directors/BlueprintCameraDirector.h
+++ b/Engine/Plugins/Cameras/GameplayCameras/Source/GameplayCameras/Public/Directors/BlueprintCameraDirector.h
@@ -21,7 +21,7 @@ struct FBlueprintCameraDirectorEvaluationParams
 	float DeltaTime = 0.f;
 
 	/** The owner (if any) of the evaluation context we are running inside of. */
-	UPROPERTY(BlueprintReadWrite, Category = "Evaluation")
+	UPROPERTY(BlueprintReadWrite, Category="Evaluation")
 	TObjectPtr<UObject> EvaluationContextOwner;
 };
 
@@ -33,9 +33,9 @@ struct FBlueprintCameraDirectorEvaluationResult
 {
 	GENERATED_BODY()
 
-	/** The list of camera rigs that should be active this frame. */
-	UPROPERTY(BlueprintReadWrite, Category = "Evaluation")
-	TArray<UCameraRigAsset*> ActiveCameraRigs;
+	/** The list of camera rigs (determined by name) that should be active this frame. */
+	UPROPERTY(BlueprintReadWrite, Category="Evaluation")
+	TArray<FString> ActiveCameraRigs;
 };
 
 /**
@@ -52,15 +52,26 @@ public:
 	 * Override this method in Blueprint to execute the custom logic that determines
 	 * what camera rig(s) should be active every frame.
 	 */
-	UFUNCTION(BlueprintCallable, BlueprintImplementableEvent, Category = "Evaluation")
-	void RunCameraDirector(
-			const FBlueprintCameraDirectorEvaluationParams& Params, 
-			FBlueprintCameraDirectorEvaluationResult& OutResult);
+	UFUNCTION(BlueprintCallable, BlueprintImplementableEvent, Category="Evaluation")
+	void RunCameraDirector(const FBlueprintCameraDirectorEvaluationParams& Params);
+
+	/**
+	 * Specifies a camera rig to be active this frame.
+	 */
+	UFUNCTION(BlueprintCallable, Category="Evaluation")
+	void ActivateCameraRig(const FString& InCameraRigName);
 
 	/** Native wrapper for RunCameraDirector. */
 	void NativeRunCameraDirector(
-			const FBlueprintCameraDirectorEvaluationParams& Params, 
+			const FBlueprintCameraDirectorEvaluationParams& Params,
 			FBlueprintCameraDirectorEvaluationResult& OutResult);
+
+protected:
+
+	/** The current camera director evaluation result. */
+	UPROPERTY(BlueprintReadWrite, Category="Evaluation")
+	FBlueprintCameraDirectorEvaluationResult CurrentResult;
+
 };
 
 /**
