Cameras: don't unregister details customizations on module shutdown if the property editor module was shutdown first

diff --git a/Engine/Plugins/Cameras/GameplayCameras/Source/GameplayCamerasEditor/Private/GameplayCamerasEditorModule.cpp b/Engine/Plugins/Cameras/GameplayCameras/Source/GameplayCamerasEditor/Private/GameplayCamerasEditorModule.cpp
--- a/Engine/Plugins/Cameras/GameplayCameras/Source/GameplayCamerasEditor/Private/GameplayCamerasEditorModule.cpp
+++ b/Engine/Plugins/Cameras/GameplayCameras/Source/GameplayCamerasEditor/Private/GameplayCamerasEditorModule.cpp
@@ -315,8 +315,12 @@ private:
 	{
 		using namespace UE::Cameras;
 
-		FPropertyEditorModule& PropertyEditorModule = FModuleManager::GetModuleChecked<FPropertyEditorModule>("PropertyEditor");
-		FCameraParameterDetailsCustomization::Unregister(PropertyEditorModule);
+		FPropertyEditorModule* PropertyEditorModule = FModuleManager::GetModulePtr<FPropertyEditorModule>("PropertyEditor");
+
+		if (PropertyEditorModule)
+		{
+			FCameraParameterDetailsCustomization::Unregister(*PropertyEditorModule);
+		}
 	}
 
 	void InitializeLiveEditManager()
