Cameras: remove the "transition owner info" stuff and instead use subclasses of the transition graph schema.

This serves a similar purpose of customizing transition graphs between camera rig transitions and camera asset shared transitions, but with a bit less code and awkwardness.

diff --git a/Engine/Plugins/Cameras/GameplayCameras/Source/GameplayCamerasEditor/Private/Editors/CameraNodeGraphSchema.cpp b/Engine/Plugins/Cameras/GameplayCameras/Source/GameplayCamerasEditor/Private/Editors/CameraNodeGraphSchema.cpp
--- a/Engine/Plugins/Cameras/GameplayCameras/Source/GameplayCamerasEditor/Private/Editors/CameraNodeGraphSchema.cpp
+++ b/Engine/Plugins/Cameras/GameplayCameras/Source/GameplayCamerasEditor/Private/Editors/CameraNodeGraphSchema.cpp
@@ -21,7 +21,7 @@
 
 const FName UCameraNodeGraphSchema::PC_CameraParameter("CameraParameter");
 
-FObjectTreeGraphConfig UCameraNodeGraphSchema::BuildGraphConfig()
+FObjectTreeGraphConfig UCameraNodeGraphSchema::BuildGraphConfig() const
 {
 	const UGameplayCamerasEditorSettings* Settings = GetDefault<UGameplayCamerasEditorSettings>();
 
diff --git a/Engine/Plugins/Cameras/GameplayCameras/Source/GameplayCamerasEditor/Private/Editors/CameraRigTransitionGraphSchema.cpp b/Engine/Plugins/Cameras/GameplayCameras/Source/GameplayCamerasEditor/Private/Editors/CameraRigTransitionGraphSchema.cpp
--- a/Engine/Plugins/Cameras/GameplayCameras/Source/GameplayCamerasEditor/Private/Editors/CameraRigTransitionGraphSchema.cpp
+++ b/Engine/Plugins/Cameras/GameplayCameras/Source/GameplayCamerasEditor/Private/Editors/CameraRigTransitionGraphSchema.cpp
@@ -2,136 +2,59 @@
 
 #include "Editors/CameraRigTransitionGraphSchema.h"
 
-#include "Core/BlendCameraNode.h"
 #include "Core/CameraRigAsset.h"
-#include "EdGraph/EdGraphPin.h"
 #include "Editors/ObjectTreeGraphConfig.h"
 #include "Editors/ObjectTreeGraphNode.h"
 #include "GameplayCamerasEditorSettings.h"
 
-#include "ScopedTransaction.h"
-
 #include UE_INLINE_GENERATED_CPP_BY_NAME(CameraRigTransitionGraphSchema)
 
 #define LOCTEXT_NAMESPACE "CameraRigTransitionGraphSchema"
 
-FObjectTreeGraphConfig UCameraRigTransitionGraphSchema::BuildGraphConfig(const FCameraRigTransitionOwnerInfo& OwnerInfo)
+void UCameraRigTransitionGraphSchema::OnBuildGraphConfig(FObjectTreeGraphConfig& InOutGraphConfig) const
 {
-	using namespace UE::Cameras;
-
 	const UGameplayCamerasEditorSettings* Settings = GetDefault<UGameplayCamerasEditorSettings>();
 
-	FObjectTreeGraphConfig GraphConfig;
-	GraphConfig.GraphName = OwnerInfo.GraphName;
-	GraphConfig.ConnectableObjectClasses.Add(OwnerInfo.TransitionOwnerClass);
-	GraphConfig.ConnectableObjectClasses.Add(UCameraRigTransition::StaticClass());
-	GraphConfig.ConnectableObjectClasses.Add(UCameraRigTransitionCondition::StaticClass());
-	GraphConfig.ConnectableObjectClasses.Add(UBlendCameraNode::StaticClass());
-	GraphConfig.StopAutoCollectAtObjectClasses.Add(UCameraRigAsset::StaticClass());
-	GraphConfig.GraphDisplayInfo.PlainName = FText::FromName(OwnerInfo.GraphName);
-	GraphConfig.GraphDisplayInfo.DisplayName = GraphConfig.GraphDisplayInfo.PlainName;
-	GraphConfig.ObjectClassConfigs.Emplace(OwnerInfo.TransitionOwnerClass)
+	InOutGraphConfig.GraphName = UCameraRigAsset::TransitionsGraphName;
+	InOutGraphConfig.ConnectableObjectClasses.Add(UCameraRigAsset::StaticClass());
+	InOutGraphConfig.GraphDisplayInfo.PlainName = LOCTEXT("NodeGraphPlainName", "Transitions");
+	InOutGraphConfig.GraphDisplayInfo.DisplayName = LOCTEXT("NodeGraphDisplayName", "Transitions");
+	InOutGraphConfig.ObjectClassConfigs.Emplace(UCameraRigAsset::StaticClass())
 		.HasSelfPin(false)
 		.OnlyAsRoot()
-		.SetPropertyPinDirection(OwnerInfo.EnterTransitionsPropertyName, EGPD_Input)
-		.SetPropertyPinDirection(OwnerInfo.ExitTransitionsPropertyName, EGPD_Output)
+		.SetPropertyPinDirection(GET_MEMBER_NAME_CHECKED(UCameraRigAsset, EnterTransitions), EGPD_Input)
+		.SetPropertyPinDirection(GET_MEMBER_NAME_CHECKED(UCameraRigAsset, ExitTransitions), EGPD_Output)
 		.NodeTitleUsesObjectName(true)
 		.NodeTitleColor(Settings->CameraRigAssetTitleColor);
-	GraphConfig.ObjectClassConfigs.Emplace(UCameraRigTransition::StaticClass())
-		.SetPropertyPinDirection(GET_MEMBER_NAME_STRING_CHECKED(UCameraRigTransition, Conditions), EGPD_Input)
-		.NodeTitleColor(Settings->CameraRigTransitionTitleColor);
-	GraphConfig.ObjectClassConfigs.Emplace(UCameraRigTransitionCondition::StaticClass())
-		.SelfPinDirection(EGPD_Output)
-		.DefaultPropertyPinDirection(EGPD_Input)
-		.StripDisplayNameSuffix(TEXT("Transition Condition"))
-		.NodeTitleColor(Settings->CameraRigTransitionConditionTitleColor);
-	GraphConfig.ObjectClassConfigs.Emplace(UBlendCameraNode::StaticClass())
-		.StripDisplayNameSuffix(TEXT("Camera Node"))
-		.CreateCategoryMetaData(TEXT("CameraNodeCategories"));
-	GraphConfig.OnFormatObjectDisplayName = FOnFormatObjectDisplayName::CreateLambda(
-			[](const UObject* Object, FText& InOutDisplayNameText)
-			{
-				if (const UCameraRigAsset* CameraRigAsset = Cast<UCameraRigAsset>(Object))
-				{
-					InOutDisplayNameText = FText::FromString(CameraRigAsset->GetDisplayName());
-				}
-			});
-
-	return GraphConfig;
 }
 
-void UCameraRigTransitionGraphSchema::GetGraphContextActions(FGraphContextMenuBuilder& ContextMenuBuilder) const
+ETransitionGraphContextActions UCameraRigTransitionGraphSchema::GetTransitionGraphContextActions(FGraphContextMenuBuilder& ContextMenuBuilder) const
 {
-	// Start with being able to create both enter and exit transitions.
-	bool bCanCreateEnterTransition = true;
-	bool bCanCreateExitTransition = true;
+	ETransitionGraphContextActions PossibleActions = ETransitionGraphContextActions::None;
+
 	if (const UEdGraphPin* DraggedPin = ContextMenuBuilder.FromPin)
 	{
 		// If we are creating a node from dragging a pin into an empty space, figure out which transition
 		// we can create based on the direction of the dragged pin.
-		bCanCreateEnterTransition = false;
-		bCanCreateExitTransition = false;
-
 		UObjectTreeGraphNode* OwningNode = Cast<UObjectTreeGraphNode>(DraggedPin->GetOwningNode());
 		if (OwningNode)
 		{
 			FProperty* DraggedPinProperty = OwningNode->GetPropertyForPin(DraggedPin);
 			if (DraggedPinProperty && DraggedPinProperty->GetOwnerClass()->IsChildOf<UCameraRigAsset>())
 			{
-				bCanCreateEnterTransition = (DraggedPinProperty->GetFName() == GET_MEMBER_NAME_CHECKED(UCameraRigAsset, EnterTransitions));
-				bCanCreateExitTransition = (DraggedPinProperty->GetFName() == GET_MEMBER_NAME_CHECKED(UCameraRigAsset, ExitTransitions));
+				if (DraggedPinProperty->GetFName() == GET_MEMBER_NAME_CHECKED(UCameraRigAsset, EnterTransitions))
+				{
+					PossibleActions |= ETransitionGraphContextActions::CreateEnterTransition;
+				}
+				if (DraggedPinProperty->GetFName() == GET_MEMBER_NAME_CHECKED(UCameraRigAsset, ExitTransitions))
+				{
+					PossibleActions |= ETransitionGraphContextActions::CreateExitTransition;
+				}
 			}
 		}
 	}
 
-	if (bCanCreateEnterTransition)
-	{
-		TSharedRef<FCameraRigTransitionGraphSchemaAction_NewTransitionNode> EnterAction = MakeShared<FCameraRigTransitionGraphSchemaAction_NewTransitionNode>(
-				LOCTEXT("TransitionsCategory", "Transitions"),
-				LOCTEXT("EnterTransition", "Enter Transition"),
-				LOCTEXT("EnterTransitionToolTip", "Creates a new enter transition"));
-		EnterAction->TransitionType = FCameraRigTransitionGraphSchemaAction_NewTransitionNode::ETransitionType::Enter;
-		ContextMenuBuilder.AddAction(StaticCastSharedPtr<FEdGraphSchemaAction>(EnterAction.ToSharedPtr()));
-	}
-	
-	if (bCanCreateExitTransition)
-	{
-		TSharedRef<FCameraRigTransitionGraphSchemaAction_NewTransitionNode> ExitAction = MakeShared<FCameraRigTransitionGraphSchemaAction_NewTransitionNode>(
-				LOCTEXT("TransitionsCategory", "Transitions"),
-				LOCTEXT("ExitTransition", "Exit Transition"),
-				LOCTEXT("ExitTransitionToolTip", "Creates a new exit transition"),
-				0);
-		ExitAction->TransitionType = FCameraRigTransitionGraphSchemaAction_NewTransitionNode::ETransitionType::Exit;
-		ContextMenuBuilder.AddAction(StaticCastSharedPtr<FEdGraphSchemaAction>(ExitAction.ToSharedPtr()));
-	}
-
-	Super::GetGraphContextActions(ContextMenuBuilder);
-}
-
-void UCameraRigTransitionGraphSchema::FilterGraphContextPlaceableClasses(TArray<UClass*>& InOutClasses) const
-{
-	InOutClasses.Remove(UCameraRigTransition::StaticClass());
-}
-
-FCameraRigTransitionGraphSchemaAction_NewTransitionNode::FCameraRigTransitionGraphSchemaAction_NewTransitionNode()
-{
-	ObjectClass = UCameraRigTransition::StaticClass();
-}
-
-FCameraRigTransitionGraphSchemaAction_NewTransitionNode::FCameraRigTransitionGraphSchemaAction_NewTransitionNode(FText InNodeCategory, FText InMenuDesc, FText InToolTip, const int32 InGrouping, FText InKeywords)
-	: FObjectGraphSchemaAction_NewNode(InNodeCategory, InMenuDesc, InToolTip, InGrouping, InKeywords)
-{
-	ObjectClass = UCameraRigTransition::StaticClass();
-}
-
-void FCameraRigTransitionGraphSchemaAction_NewTransitionNode::AutoSetupNewNode(UObjectTreeGraphNode* NewNode, UEdGraphPin* FromPin)
-{
-	if (TransitionType == ETransitionType::Enter)
-	{
-		NewNode->OverrideSelfPinDirection(EGPD_Output);
-	}
-
-	FObjectGraphSchemaAction_NewNode::AutoSetupNewNode(NewNode, FromPin);
+	return PossibleActions;
 }
 
 #undef LOCTEXT_NAMESPACE
diff --git a/Engine/Plugins/Cameras/GameplayCameras/Source/GameplayCamerasEditor/Private/Editors/CameraRigTransitionGraphSchemaBase.cpp b/Engine/Plugins/Cameras/GameplayCameras/Source/GameplayCamerasEditor/Private/Editors/CameraRigTransitionGraphSchemaBase.cpp
new file mode 100644
--- /dev/null
+++ b/Engine/Plugins/Cameras/GameplayCameras/Source/GameplayCamerasEditor/Private/Editors/CameraRigTransitionGraphSchemaBase.cpp
@@ -0,0 +1,98 @@
+// Copyright Epic Games, Inc. All Rights Reserved.
+
+#include "Editors/CameraRigTransitionGraphSchemaBase.h"
+
+#include "Core/BlendCameraNode.h"
+#include "Core/CameraRigTransition.h"
+#include "EdGraph/EdGraphPin.h"
+#include "Editors/ObjectTreeGraphConfig.h"
+#include "Editors/ObjectTreeGraphNode.h"
+#include "GameplayCamerasEditorSettings.h"
+
+#include UE_INLINE_GENERATED_CPP_BY_NAME(CameraRigTransitionGraphSchemaBase)
+
+#define LOCTEXT_NAMESPACE "CameraRigTransitionGraphSchemaBase"
+
+FObjectTreeGraphConfig UCameraRigTransitionGraphSchemaBase::BuildGraphConfig() const
+{
+	using namespace UE::Cameras;
+
+	const UGameplayCamerasEditorSettings* Settings = GetDefault<UGameplayCamerasEditorSettings>();
+
+	FObjectTreeGraphConfig GraphConfig;
+	GraphConfig.ConnectableObjectClasses.Add(UCameraRigTransition::StaticClass());
+	GraphConfig.ConnectableObjectClasses.Add(UCameraRigTransitionCondition::StaticClass());
+	GraphConfig.ConnectableObjectClasses.Add(UBlendCameraNode::StaticClass());
+	GraphConfig.ObjectClassConfigs.Emplace(UCameraRigTransition::StaticClass())
+		.SetPropertyPinDirection(GET_MEMBER_NAME_STRING_CHECKED(UCameraRigTransition, Conditions), EGPD_Input)
+		.NodeTitleColor(Settings->CameraRigTransitionTitleColor);
+	GraphConfig.ObjectClassConfigs.Emplace(UCameraRigTransitionCondition::StaticClass())
+		.SelfPinDirection(EGPD_Output)
+		.DefaultPropertyPinDirection(EGPD_Input)
+		.StripDisplayNameSuffix(TEXT("Transition Condition"))
+		.NodeTitleColor(Settings->CameraRigTransitionConditionTitleColor);
+	GraphConfig.ObjectClassConfigs.Emplace(UBlendCameraNode::StaticClass())
+		.StripDisplayNameSuffix(TEXT("Camera Node"))
+		.CreateCategoryMetaData(TEXT("CameraNodeCategories"));
+
+	OnBuildGraphConfig(GraphConfig);
+
+	return GraphConfig;
+}
+
+void UCameraRigTransitionGraphSchemaBase::GetGraphContextActions(FGraphContextMenuBuilder& ContextMenuBuilder) const
+{
+	ETransitionGraphContextActions PossibleActions = GetTransitionGraphContextActions(ContextMenuBuilder);
+
+	if (EnumHasAnyFlags(PossibleActions, ETransitionGraphContextActions::CreateEnterTransition))
+	{
+		TSharedRef<FCameraRigTransitionGraphSchemaAction_NewTransitionNode> EnterAction = MakeShared<FCameraRigTransitionGraphSchemaAction_NewTransitionNode>(
+				LOCTEXT("TransitionsCategory", "Transitions"),
+				LOCTEXT("EnterTransition", "Enter Transition"),
+				LOCTEXT("EnterTransitionToolTip", "Creates a new enter transition"));
+		EnterAction->TransitionType = FCameraRigTransitionGraphSchemaAction_NewTransitionNode::ETransitionType::Enter;
+		ContextMenuBuilder.AddAction(StaticCastSharedPtr<FEdGraphSchemaAction>(EnterAction.ToSharedPtr()));
+	}
+	
+	if (EnumHasAnyFlags(PossibleActions, ETransitionGraphContextActions::CreateExitTransition))
+	{
+		TSharedRef<FCameraRigTransitionGraphSchemaAction_NewTransitionNode> ExitAction = MakeShared<FCameraRigTransitionGraphSchemaAction_NewTransitionNode>(
+				LOCTEXT("TransitionsCategory", "Transitions"),
+				LOCTEXT("ExitTransition", "Exit Transition"),
+				LOCTEXT("ExitTransitionToolTip", "Creates a new exit transition"),
+				0);
+		ExitAction->TransitionType = FCameraRigTransitionGraphSchemaAction_NewTransitionNode::ETransitionType::Exit;
+		ContextMenuBuilder.AddAction(StaticCastSharedPtr<FEdGraphSchemaAction>(ExitAction.ToSharedPtr()));
+	}
+
+	Super::GetGraphContextActions(ContextMenuBuilder);
+}
+
+void UCameraRigTransitionGraphSchemaBase::FilterGraphContextPlaceableClasses(TArray<UClass*>& InOutClasses) const
+{
+	InOutClasses.Remove(UCameraRigTransition::StaticClass());
+}
+
+FCameraRigTransitionGraphSchemaAction_NewTransitionNode::FCameraRigTransitionGraphSchemaAction_NewTransitionNode()
+{
+	ObjectClass = UCameraRigTransition::StaticClass();
+}
+
+FCameraRigTransitionGraphSchemaAction_NewTransitionNode::FCameraRigTransitionGraphSchemaAction_NewTransitionNode(FText InNodeCategory, FText InMenuDesc, FText InToolTip, const int32 InGrouping, FText InKeywords)
+	: FObjectGraphSchemaAction_NewNode(InNodeCategory, InMenuDesc, InToolTip, InGrouping, InKeywords)
+{
+	ObjectClass = UCameraRigTransition::StaticClass();
+}
+
+void FCameraRigTransitionGraphSchemaAction_NewTransitionNode::AutoSetupNewNode(UObjectTreeGraphNode* NewNode, UEdGraphPin* FromPin)
+{
+	if (TransitionType == ETransitionType::Enter)
+	{
+		NewNode->OverrideSelfPinDirection(EGPD_Output);
+	}
+
+	FObjectGraphSchemaAction_NewNode::AutoSetupNewNode(NewNode, FromPin);
+}
+
+#undef LOCTEXT_NAMESPACE
+
diff --git a/Engine/Plugins/Cameras/GameplayCameras/Source/GameplayCamerasEditor/Private/Editors/CameraSharedTransitionGraphSchema.cpp b/Engine/Plugins/Cameras/GameplayCameras/Source/GameplayCamerasEditor/Private/Editors/CameraSharedTransitionGraphSchema.cpp
new file mode 100644
--- /dev/null
+++ b/Engine/Plugins/Cameras/GameplayCameras/Source/GameplayCamerasEditor/Private/Editors/CameraSharedTransitionGraphSchema.cpp
@@ -0,0 +1,63 @@
+// Copyright Epic Games, Inc. All Rights Reserved.
+
+#include "Editors/CameraSharedTransitionGraphSchema.h"
+
+#include "Core/CameraAsset.h"
+#include "Core/CameraRigAsset.h"
+#include "Editors/ObjectTreeGraphConfig.h"
+#include "Editors/ObjectTreeGraphNode.h"
+#include "GameplayCamerasEditorSettings.h"
+
+#include UE_INLINE_GENERATED_CPP_BY_NAME(CameraSharedTransitionGraphSchema)
+
+#define LOCTEXT_NAMESPACE "CameraSharedTransitionGraphSchema"
+
+void UCameraSharedTransitionGraphSchema::OnBuildGraphConfig(FObjectTreeGraphConfig& InOutGraphConfig) const
+{
+	const UGameplayCamerasEditorSettings* Settings = GetDefault<UGameplayCamerasEditorSettings>();
+
+	InOutGraphConfig.GraphName = UCameraAsset::SharedTransitionsGraphName;
+	InOutGraphConfig.ConnectableObjectClasses.Add(UCameraAsset::StaticClass());
+	InOutGraphConfig.GraphDisplayInfo.PlainName = LOCTEXT("NodeGraphPlainName", "SharedTransitions");
+	InOutGraphConfig.GraphDisplayInfo.DisplayName = LOCTEXT("NodeGraphDisplayName", "Shared Transitions");
+	InOutGraphConfig.StopAutoCollectAtObjectClasses.Add(UCameraRigAsset::StaticClass());
+	InOutGraphConfig.ObjectClassConfigs.Emplace(UCameraAsset::StaticClass())
+		.HasSelfPin(false)
+		.OnlyAsRoot()
+		.SetPropertyPinDirection(GET_MEMBER_NAME_CHECKED(UCameraAsset, EnterTransitions), EGPD_Input)
+		.SetPropertyPinDirection(GET_MEMBER_NAME_CHECKED(UCameraAsset, ExitTransitions), EGPD_Output)
+		.NodeTitleUsesObjectName(true)
+		.NodeTitleColor(Settings->CameraAssetTitleColor);
+}
+
+ETransitionGraphContextActions UCameraSharedTransitionGraphSchema::GetTransitionGraphContextActions(FGraphContextMenuBuilder& ContextMenuBuilder) const
+{
+	ETransitionGraphContextActions PossibleActions = ETransitionGraphContextActions::None;
+
+	if (const UEdGraphPin* DraggedPin = ContextMenuBuilder.FromPin)
+	{
+		// If we are creating a node from dragging a pin into an empty space, figure out which transition
+		// we can create based on the direction of the dragged pin.
+		UObjectTreeGraphNode* OwningNode = Cast<UObjectTreeGraphNode>(DraggedPin->GetOwningNode());
+		if (OwningNode)
+		{
+			FProperty* DraggedPinProperty = OwningNode->GetPropertyForPin(DraggedPin);
+			if (DraggedPinProperty && DraggedPinProperty->GetOwnerClass()->IsChildOf<UCameraAsset>())
+			{
+				if (DraggedPinProperty->GetFName() == GET_MEMBER_NAME_CHECKED(UCameraAsset, EnterTransitions))
+				{
+					PossibleActions |= ETransitionGraphContextActions::CreateEnterTransition;
+				}
+				if (DraggedPinProperty->GetFName() == GET_MEMBER_NAME_CHECKED(UCameraAsset, ExitTransitions))
+				{
+					PossibleActions |= ETransitionGraphContextActions::CreateExitTransition;
+				}
+			}
+		}
+	}
+
+	return PossibleActions;
+}
+
+#undef LOCTEXT_NAMESPACE
+
diff --git a/Engine/Plugins/Cameras/GameplayCameras/Source/GameplayCamerasEditor/Private/Editors/SCameraRigAssetEditor.cpp b/Engine/Plugins/Cameras/GameplayCameras/Source/GameplayCamerasEditor/Private/Editors/SCameraRigAssetEditor.cpp
--- a/Engine/Plugins/Cameras/GameplayCameras/Source/GameplayCamerasEditor/Private/Editors/SCameraRigAssetEditor.cpp
+++ b/Engine/Plugins/Cameras/GameplayCameras/Source/GameplayCamerasEditor/Private/Editors/SCameraRigAssetEditor.cpp
@@ -69,10 +69,12 @@ void SCameraRigAssetEditor::CreateGraphEditors()
 
 void SCameraRigAssetEditor::CreateNodeGraphEditor()
 {
-	FObjectTreeGraphConfig GraphConfig = UCameraNodeGraphSchema::BuildGraphConfig();
+	UClass* SchemaClass = UCameraNodeGraphSchema::StaticClass();
+	UCameraNodeGraphSchema* DefaultSchemaObject = SchemaClass->GetDefaultObject<UCameraNodeGraphSchema>();
+	FObjectTreeGraphConfig GraphConfig = DefaultSchemaObject->BuildGraphConfig();
 
 	NodeGraph = NewObject<UObjectTreeGraph>(GetTransientPackage(), NAME_None, RF_Transactional);
-	NodeGraph->Schema = UCameraNodeGraphSchema::StaticClass();
+	NodeGraph->Schema = SchemaClass;
 	NodeGraph->AddToRoot();
 	NodeGraph->Reset(CameraRigAsset, GraphConfig, EObjectTreeGraphBuildSource::RootObjectPackage);
 
@@ -92,17 +94,12 @@ void SCameraRigAssetEditor::CreateNodeGraphEditor()
 
 void SCameraRigAssetEditor::CreateTransitionGraphEditor()
 {
-	FCameraRigTransitionOwnerInfo OwnerInfo;
-	OwnerInfo.TransitionOwnerClass = UCameraRigAsset::StaticClass();
-	OwnerInfo.GraphName = UCameraRigAsset::TransitionsGraphName;
-	OwnerInfo.EnterTransitionsPropertyName = GET_MEMBER_NAME_CHECKED(UCameraRigAsset, EnterTransitions);
-	OwnerInfo.ExitTransitionsPropertyName = GET_MEMBER_NAME_CHECKED(UCameraRigAsset, ExitTransitions);
-	FObjectTreeGraphConfig GraphConfig = UCameraRigTransitionGraphSchema::BuildGraphConfig(OwnerInfo);
-
-	GraphConfig.GraphDisplayInfo.DisplayName = LOCTEXT("TransitionGraphDisplayName", "Transitions");
+	UClass* SchemaClass = UCameraRigTransitionGraphSchema::StaticClass();
+	UCameraRigTransitionGraphSchema* DefaultSchemaObject = SchemaClass->GetDefaultObject<UCameraRigTransitionGraphSchema>();
+	FObjectTreeGraphConfig GraphConfig = DefaultSchemaObject->BuildGraphConfig();
 
 	TransitionGraph = NewObject<UObjectTreeGraph>(GetTransientPackage(), NAME_None, RF_Transactional);
-	TransitionGraph->Schema = UCameraRigTransitionGraphSchema::StaticClass();
+	TransitionGraph->Schema = SchemaClass;
 	TransitionGraph->AddToRoot();
 	TransitionGraph->Reset(CameraRigAsset, GraphConfig, EObjectTreeGraphBuildSource::RootObjectPackage);
 
diff --git a/Engine/Plugins/Cameras/GameplayCameras/Source/GameplayCamerasEditor/Private/Editors/SCameraRigTransitionEditor.cpp b/Engine/Plugins/Cameras/GameplayCameras/Source/GameplayCamerasEditor/Private/Editors/SCameraRigTransitionEditor.cpp
--- a/Engine/Plugins/Cameras/GameplayCameras/Source/GameplayCamerasEditor/Private/Editors/SCameraRigTransitionEditor.cpp
+++ b/Engine/Plugins/Cameras/GameplayCameras/Source/GameplayCamerasEditor/Private/Editors/SCameraRigTransitionEditor.cpp
@@ -23,11 +23,9 @@ namespace UE::Cameras
 void SCameraRigTransitionEditor::Construct(const FArguments& InArgs)
 {
 	TransitionOwner = InArgs._TransitionOwner;
-	TransitionOwnerInfo = InArgs._TransitionOwnerInfo;
-	OnBuildTransitionGraphConfig = InArgs._OnBuildTransitionGraphConfig;
+	TransitionGraphSchemaClass = InArgs._TransitionGraphSchemaClass;
 	DetailsView = InArgs._DetailsView;
 	AssetEditorToolkit = InArgs._AssetEditorToolkit;
-	TransitionGraphDisplayInfo = InArgs._TransitionGraphDisplayInfo;
 	TransitionGraphEditorAppearance = InArgs._TransitionGraphEditorAppearance;
 
 	CreateTransitionGraphEditor();
@@ -66,15 +64,11 @@ void SCameraRigTransitionEditor::SetTransitionOwner(UObject* InTransitionOwner)
 
 void SCameraRigTransitionEditor::CreateTransitionGraphEditor()
 {
-	FObjectTreeGraphConfig GraphConfig = UCameraRigTransitionGraphSchema::BuildGraphConfig(TransitionOwnerInfo);
-	if (TransitionGraphDisplayInfo.IsSet())
-	{
-		GraphConfig.GraphDisplayInfo = TransitionGraphDisplayInfo.Get();
-	}
-	OnBuildTransitionGraphConfig.ExecuteIfBound(GraphConfig);
+	UCameraRigTransitionGraphSchemaBase* DefaultSchemaObject = TransitionGraphSchemaClass->GetDefaultObject<UCameraRigTransitionGraphSchemaBase>();
+	FObjectTreeGraphConfig GraphConfig = DefaultSchemaObject->BuildGraphConfig();
 
 	TransitionGraph = NewObject<UObjectTreeGraph>(GetTransientPackage(), NAME_None, RF_Transactional);
-	TransitionGraph->Schema = UCameraRigTransitionGraphSchema::StaticClass();
+	TransitionGraph->Schema = TransitionGraphSchemaClass;
 	TransitionGraph->AddToRoot();
 	TransitionGraph->Reset(TransitionOwner, GraphConfig, EObjectTreeGraphBuildSource::RootObjectPackage);
 
diff --git a/Engine/Plugins/Cameras/GameplayCameras/Source/GameplayCamerasEditor/Private/GameplayCamerasEditorSettings.cpp b/Engine/Plugins/Cameras/GameplayCameras/Source/GameplayCamerasEditor/Private/GameplayCamerasEditorSettings.cpp
--- a/Engine/Plugins/Cameras/GameplayCameras/Source/GameplayCamerasEditor/Private/GameplayCamerasEditorSettings.cpp
+++ b/Engine/Plugins/Cameras/GameplayCameras/Source/GameplayCamerasEditor/Private/GameplayCamerasEditorSettings.cpp
@@ -8,6 +8,7 @@ UGameplayCamerasEditorSettings::UGameplayCamerasEditorSettings(const FObjectInit
 	: Super(ObjectInitializer)
 {
 	CameraNodeTitleColor = FLinearColor(0.190525f, 0.583898f, 1.0f, 1.0f);           // Light blue
+	CameraAssetTitleColor = FLinearColor(1.0f, 0.170000f, 0.0f, 1.0f);				 // Orange
 	CameraRigAssetTitleColor = FLinearColor(1.0f, 0.170000f, 0.0f, 1.0f);            // Orange
 	CameraRigTransitionTitleColor = FLinearColor(1.0f, 0.65f, 0.4f, 1.0f);           // Beige
 	CameraRigTransitionConditionTitleColor = FLinearColor(0.8f, 0.4f, 0.4f, 1.0f);   // Salmon
diff --git a/Engine/Plugins/Cameras/GameplayCameras/Source/GameplayCamerasEditor/Private/Toolkits/CameraRigTransitionEditorToolkit.cpp b/Engine/Plugins/Cameras/GameplayCameras/Source/GameplayCamerasEditor/Private/Toolkits/CameraRigTransitionEditorToolkit.cpp
--- a/Engine/Plugins/Cameras/GameplayCameras/Source/GameplayCamerasEditor/Private/Toolkits/CameraRigTransitionEditorToolkit.cpp
+++ b/Engine/Plugins/Cameras/GameplayCameras/Source/GameplayCamerasEditor/Private/Toolkits/CameraRigTransitionEditorToolkit.cpp
@@ -24,23 +24,6 @@ public:
 	FCameraRigTransitionEditorToolkitBaseImpl()
 		: FCameraRigTransitionEditorToolkitBase(TEXT("CameraRigTransitionEditor_Layout_v2"))
 	{}
-
-protected:
-
-	virtual void GetTransitionOwnerInfo(FCameraRigTransitionOwnerInfo& OutTransitionOwnerInfo) override
-	{
-		// TODO
-	}
-
-	virtual void GetTransitionGraphDisplayInfo(FGraphDisplayInfo& OutGraphDisplayInfo) override
-	{
-		// TODO
-	}
-
-	virtual void GetTransitionGraphAppearanceInfo(FGraphAppearanceInfo& OutGraphAppearanceInfo) override
-	{
-		// TODO
-	}
 };
 
 }  // namespace Internal
diff --git a/Engine/Plugins/Cameras/GameplayCameras/Source/GameplayCamerasEditor/Private/Toolkits/CameraRigTransitionEditorToolkitBase.cpp b/Engine/Plugins/Cameras/GameplayCameras/Source/GameplayCamerasEditor/Private/Toolkits/CameraRigTransitionEditorToolkitBase.cpp
--- a/Engine/Plugins/Cameras/GameplayCameras/Source/GameplayCamerasEditor/Private/Toolkits/CameraRigTransitionEditorToolkitBase.cpp
+++ b/Engine/Plugins/Cameras/GameplayCameras/Source/GameplayCamerasEditor/Private/Toolkits/CameraRigTransitionEditorToolkitBase.cpp
@@ -4,6 +4,7 @@
 
 #include "Commands/CameraRigTransitionEditorCommands.h"
 #include "Core/CameraBuildStatus.h"
+#include "Editors/CameraRigTransitionGraphSchemaBase.h"
 #include "Editors/SCameraRigTransitionEditor.h"
 #include "Editors/SObjectTreeGraphToolbox.h"
 #include "Framework/Commands/UICommandList.h"
@@ -131,11 +132,7 @@ void FCameraRigTransitionEditorToolkitBase::CreateWidgets()
 	DetailsViewArgs.NotifyHook = this;
 	DetailsView = PropertyEditorModule.CreateDetailView(DetailsViewArgs);
 
-	FCameraRigTransitionOwnerInfo TransitionOwnerInfo;
-	GetTransitionOwnerInfo(TransitionOwnerInfo);
-
-	FGraphDisplayInfo GraphDisplayInfo;
-	GetTransitionGraphDisplayInfo(GraphDisplayInfo);
+	TSubclassOf<UCameraRigTransitionGraphSchemaBase> SchemaClass = GetTransitionGraphSchemaClass();
 
 	FGraphAppearanceInfo GraphAppearanceInfo;
 	GetTransitionGraphAppearanceInfo(GraphAppearanceInfo);
@@ -144,8 +141,7 @@ void FCameraRigTransitionEditorToolkitBase::CreateWidgets()
 	TransitionEditorWidget = SNew(SCameraRigTransitionEditor)
 		.DetailsView(DetailsView)
 		.TransitionOwner(TransitionOwner)
-		.TransitionOwnerInfo(TransitionOwnerInfo)
-		.TransitionGraphDisplayInfo(GraphDisplayInfo)
+		.TransitionGraphSchemaClass(SchemaClass)
 		.TransitionGraphEditorAppearance(GraphAppearanceInfo);
 	TransitionEditorWidget->AddOnGraphChanged(FOnGraphChanged::FDelegate::CreateSP(
 				this, &FCameraRigTransitionEditorToolkitBase::OnTransitionGraphChanged));
@@ -155,6 +151,11 @@ void FCameraRigTransitionEditorToolkitBase::CreateWidgets()
 		.GraphConfig(TransitionEditorWidget->GetTransitionGraphConfig());
 }
 
+TSubclassOf<UCameraRigTransitionGraphSchemaBase> FCameraRigTransitionEditorToolkitBase::GetTransitionGraphSchemaClass()
+{
+	return UCameraRigTransitionGraphSchemaBase::StaticClass();
+}
+
 void FCameraRigTransitionEditorToolkitBase::BuildToolbarMenu(UToolMenu* ToolbarMenu)
 {
 	FToolMenuInsert InsertAfterAssetSection("Asset", EToolMenuInsertType::After);
diff --git a/Engine/Plugins/Cameras/GameplayCameras/Source/GameplayCamerasEditor/Private/Toolkits/CameraSharedTransitionsAssetEditorMode.cpp b/Engine/Plugins/Cameras/GameplayCameras/Source/GameplayCamerasEditor/Private/Toolkits/CameraSharedTransitionsAssetEditorMode.cpp
--- a/Engine/Plugins/Cameras/GameplayCameras/Source/GameplayCamerasEditor/Private/Toolkits/CameraSharedTransitionsAssetEditorMode.cpp
+++ b/Engine/Plugins/Cameras/GameplayCameras/Source/GameplayCamerasEditor/Private/Toolkits/CameraSharedTransitionsAssetEditorMode.cpp
@@ -3,7 +3,7 @@
 #include "Toolkits/CameraSharedTransitionsAssetEditorMode.h"
 
 #include "Core/CameraAsset.h"
-#include "Editors/CameraRigTransitionGraphSchema.h"
+#include "Editors/CameraSharedTransitionGraphSchema.h"
 #include "Editors/SCameraRigTransitionEditor.h"
 #include "ToolMenus.h"
 #include "Toolkits/CameraRigTransitionEditorToolkitBase.h"
@@ -27,18 +27,9 @@ public:
 
 protected:
 
-	virtual void GetTransitionOwnerInfo(FCameraRigTransitionOwnerInfo& OutTransitionOwnerInfo) override
+	virtual TSubclassOf<UCameraRigTransitionGraphSchemaBase> GetTransitionGraphSchemaClass() override
 	{
-		OutTransitionOwnerInfo.TransitionOwnerClass = UCameraAsset::StaticClass();
-		OutTransitionOwnerInfo.GraphName = UCameraAsset::SharedTransitionsGraphName;
-		OutTransitionOwnerInfo.EnterTransitionsPropertyName = GET_MEMBER_NAME_CHECKED(UCameraAsset, EnterTransitions);
-		OutTransitionOwnerInfo.ExitTransitionsPropertyName = GET_MEMBER_NAME_CHECKED(UCameraAsset, ExitTransitions);
-	}
-
-	virtual void GetTransitionGraphDisplayInfo(FGraphDisplayInfo& OutGraphDisplayInfo) override
-	{
-		OutGraphDisplayInfo.PlainName = LOCTEXT("SharedTransitionGraphPlainName", "SharedTransitions");
-		OutGraphDisplayInfo.DisplayName = LOCTEXT("SharedTransitionGraphDisplayName", "Shared Transitions");
+		return UCameraSharedTransitionGraphSchema::StaticClass();
 	}
 
 	virtual void GetTransitionGraphAppearanceInfo(FGraphAppearanceInfo& OutGraphAppearanceInfo) override
diff --git a/Engine/Plugins/Cameras/GameplayCameras/Source/GameplayCamerasEditor/Public/Editors/CameraNodeGraphSchema.h b/Engine/Plugins/Cameras/GameplayCameras/Source/GameplayCamerasEditor/Public/Editors/CameraNodeGraphSchema.h
--- a/Engine/Plugins/Cameras/GameplayCameras/Source/GameplayCamerasEditor/Public/Editors/CameraNodeGraphSchema.h
+++ b/Engine/Plugins/Cameras/GameplayCameras/Source/GameplayCamerasEditor/Public/Editors/CameraNodeGraphSchema.h
@@ -21,7 +21,7 @@ public:
 
 	static const FName PC_CameraParameter;			// A camera parameter pin.
 
-	static FObjectTreeGraphConfig BuildGraphConfig();
+	FObjectTreeGraphConfig BuildGraphConfig() const;
 
 protected:
 
diff --git a/Engine/Plugins/Cameras/GameplayCameras/Source/GameplayCamerasEditor/Public/Editors/CameraRigTransitionGraphSchema.h b/Engine/Plugins/Cameras/GameplayCameras/Source/GameplayCamerasEditor/Public/Editors/CameraRigTransitionGraphSchema.h
--- a/Engine/Plugins/Cameras/GameplayCameras/Source/GameplayCamerasEditor/Public/Editors/CameraRigTransitionGraphSchema.h
+++ b/Engine/Plugins/Cameras/GameplayCameras/Source/GameplayCamerasEditor/Public/Editors/CameraRigTransitionGraphSchema.h
@@ -2,83 +2,22 @@
 
 #pragma once
 
-#include "Editors/ObjectTreeGraphSchema.h"
+#include "Editors/CameraRigTransitionGraphSchemaBase.h"
 
 #include "CameraRigTransitionGraphSchema.generated.h"
 
-struct FObjectTreeGraphConfig;
-
 /**
- * Provides information about an object that has enter/exit transitions.
- */
-struct FCameraRigTransitionOwnerInfo
-{
-	/** The class of the object. */
-	UClass* TransitionOwnerClass = nullptr;
-	/** The name of the transition graph. */
-	FName GraphName;
-	/** The name of the property that returns the array of enter transitions. */
-	FName EnterTransitionsPropertyName;
-	/** The name of the property that returns the array of exit transitions. */
-	FName ExitTransitionsPropertyName;
-};
-
-/**
- * Schema class for camera transition graph.
+ * Schema class for a camera rig's transition graph.
  */
 UCLASS()
-class UCameraRigTransitionGraphSchema : public UObjectTreeGraphSchema
+class UCameraRigTransitionGraphSchema : public UCameraRigTransitionGraphSchemaBase
 {
 	GENERATED_BODY()
 
-public:
-
-	static FObjectTreeGraphConfig BuildGraphConfig(const FCameraRigTransitionOwnerInfo& OwnerInfo);
-
-protected:
-
-	// UEdGraphSchema interface.
-	virtual void GetGraphContextActions(FGraphContextMenuBuilder& ContextMenuBuilder) const override;
-
-	// UObjectTreeGraphSchema interface.
-	virtual void FilterGraphContextPlaceableClasses(TArray<UClass*>& InOutClasses) const override;
-};
-
-/**
- * Graph action to create a new transition node.
- *
- * We need a custom action for this because we need to switch the "self" pin according to whether
- * we want an enter or exit transition.
- */
-USTRUCT()
-struct FCameraRigTransitionGraphSchemaAction_NewTransitionNode : public FObjectGraphSchemaAction_NewNode
-{
-	GENERATED_BODY()
-
-public:
-
-	enum class ETransitionType
-	{
-		Enter,
-		Exit
-	};
-
-	/** The transition type to create. */
-	ETransitionType TransitionType;
-
-public:
-
-	FCameraRigTransitionGraphSchemaAction_NewTransitionNode();
-	FCameraRigTransitionGraphSchemaAction_NewTransitionNode(FText InNodeCategory, FText InMenuDesc, FText InToolTip, const int32 InGrouping = 0, FText InKeywords = FText());
-
-public:
-
-	// FEdGraphSchemaAction interface.
-	static FName StaticGetTypeId() { static FName Type("FCameraRigTransitionGraphSchemaAction_NewTransitionNode"); return Type; }
-	virtual FName GetTypeId() const override { return StaticGetTypeId(); } 
-
 protected:
 
-	virtual void AutoSetupNewNode(UObjectTreeGraphNode* NewNode, UEdGraphPin* FromPin) override;
+	// UCameraRigTransitionGraphSchemaBase interface.
+	virtual void OnBuildGraphConfig(FObjectTreeGraphConfig& InOutGraphConfig) const override;
+	virtual ETransitionGraphContextActions GetTransitionGraphContextActions(FGraphContextMenuBuilder& ContextMenuBuilder) const override;
 };
 
diff --git a/Engine/Plugins/Cameras/GameplayCameras/Source/GameplayCamerasEditor/Public/Editors/CameraRigTransitionGraphSchemaBase.h b/Engine/Plugins/Cameras/GameplayCameras/Source/GameplayCamerasEditor/Public/Editors/CameraRigTransitionGraphSchemaBase.h
new file mode 100644
--- /dev/null
+++ b/Engine/Plugins/Cameras/GameplayCameras/Source/GameplayCamerasEditor/Public/Editors/CameraRigTransitionGraphSchemaBase.h
@@ -0,0 +1,86 @@
+// Copyright Epic Games, Inc. All Rights Reserved.
+
+#pragma once
+
+#include "Editors/ObjectTreeGraphSchema.h"
+
+#include "CameraRigTransitionGraphSchemaBase.generated.h"
+
+struct FObjectTreeGraphConfig;
+
+/**
+ * The list of transition-specific actions that can be performed on a transition graph.
+ */
+enum class ETransitionGraphContextActions
+{
+	None = 0,
+	CreateEnterTransition = 1 << 0,
+	CreateExitTransition = 1 << 1
+};
+ENUM_CLASS_FLAGS(ETransitionGraphContextActions);
+
+/**
+ * Base schema class for camera transition graph.
+ */
+UCLASS()
+class UCameraRigTransitionGraphSchemaBase : public UObjectTreeGraphSchema
+{
+	GENERATED_BODY()
+
+public:
+
+	FObjectTreeGraphConfig BuildGraphConfig() const;
+
+protected:
+
+	// UEdGraphSchema interface.
+	virtual void GetGraphContextActions(FGraphContextMenuBuilder& ContextMenuBuilder) const override;
+
+	// UObjectTreeGraphSchema interface.
+	virtual void FilterGraphContextPlaceableClasses(TArray<UClass*>& InOutClasses) const override;
+
+protected:
+
+	// UCameraRigTransitionGraphSchemaBase interface.
+	virtual void OnBuildGraphConfig(FObjectTreeGraphConfig& InOutGraphConfig) const {}
+	virtual ETransitionGraphContextActions GetTransitionGraphContextActions(FGraphContextMenuBuilder& ContextMenuBuilder) const { return ETransitionGraphContextActions::None; }
+};
+
+/**
+ * Graph action to create a new transition node.
+ *
+ * We need a custom action for this because we need to switch the "self" pin according to whether
+ * we want an enter or exit transition.
+ */
+USTRUCT()
+struct FCameraRigTransitionGraphSchemaAction_NewTransitionNode : public FObjectGraphSchemaAction_NewNode
+{
+	GENERATED_BODY()
+
+public:
+
+	enum class ETransitionType
+	{
+		Enter,
+		Exit
+	};
+
+	/** The transition type to create. */
+	ETransitionType TransitionType;
+
+public:
+
+	FCameraRigTransitionGraphSchemaAction_NewTransitionNode();
+	FCameraRigTransitionGraphSchemaAction_NewTransitionNode(FText InNodeCategory, FText InMenuDesc, FText InToolTip, const int32 InGrouping = 0, FText InKeywords = FText());
+
+public:
+
+	// FEdGraphSchemaAction interface.
+	static FName StaticGetTypeId() { static FName Type("FCameraRigTransitionGraphSchemaAction_NewTransitionNode"); return Type; }
+	virtual FName GetTypeId() const override { return StaticGetTypeId(); } 
+
+protected:
+
+	virtual void AutoSetupNewNode(UObjectTreeGraphNode* NewNode, UEdGraphPin* FromPin) override;
+};
+
diff --git a/Engine/Plugins/Cameras/GameplayCameras/Source/GameplayCamerasEditor/Public/Editors/CameraSharedTransitionGraphSchema.h b/Engine/Plugins/Cameras/GameplayCameras/Source/GameplayCamerasEditor/Public/Editors/CameraSharedTransitionGraphSchema.h
new file mode 100644
--- /dev/null
+++ b/Engine/Plugins/Cameras/GameplayCameras/Source/GameplayCamerasEditor/Public/Editors/CameraSharedTransitionGraphSchema.h
@@ -0,0 +1,23 @@
+// Copyright Epic Games, Inc. All Rights Reserved.
+
+#pragma once
+
+#include "Editors/CameraRigTransitionGraphSchemaBase.h"
+
+#include "CameraSharedTransitionGraphSchema.generated.h"
+
+/**
+ * Schema class for a camera asset's shared transition graph.
+ */
+UCLASS()
+class UCameraSharedTransitionGraphSchema : public UCameraRigTransitionGraphSchemaBase
+{
+	GENERATED_BODY()
+
+protected:
+
+	// UCameraRigTransitionGraphSchemaBase interface.
+	virtual void OnBuildGraphConfig(FObjectTreeGraphConfig& InOutGraphConfig) const override;
+	virtual ETransitionGraphContextActions GetTransitionGraphContextActions(FGraphContextMenuBuilder& ContextMenuBuilder) const override;
+};
+
diff --git a/Engine/Plugins/Cameras/GameplayCameras/Source/GameplayCamerasEditor/Public/Editors/SCameraRigTransitionEditor.h b/Engine/Plugins/Cameras/GameplayCameras/Source/GameplayCamerasEditor/Public/Editors/SCameraRigTransitionEditor.h
--- a/Engine/Plugins/Cameras/GameplayCameras/Source/GameplayCamerasEditor/Public/Editors/SCameraRigTransitionEditor.h
+++ b/Engine/Plugins/Cameras/GameplayCameras/Source/GameplayCamerasEditor/Public/Editors/SCameraRigTransitionEditor.h
@@ -7,6 +7,7 @@
 #include "Editors/ObjectTreeGraphConfig.h"
 #include "GraphEditor.h"
 #include "Templates/SharedPointerFwd.h"
+#include "Templates/SubclassOf.h"
 #include "UObject/ObjectPtr.h"
 #include "Widgets/SCompoundWidget.h"
 
@@ -33,20 +34,14 @@ public:
 	{}
 		/** The object owning the transitions. */
 		SLATE_ARGUMENT(UObject*, TransitionOwner)
-		/** Information about the transition owner. */
-		SLATE_ARGUMENT(FCameraRigTransitionOwnerInfo, TransitionOwnerInfo)
-		/** Callback to edit the transition graph config before it's used to build the editor. */
-		SLATE_EVENT(FOnBuildObjectTreeGraphConfig, OnBuildTransitionGraphConfig)
+		/** The graph schema class to use. */
+		SLATE_ARGUMENT(TSubclassOf<UCameraRigTransitionGraphSchemaBase>, TransitionGraphSchemaClass)
 		/** The details view to synchronize with the graph selection. */
 		SLATE_ARGUMENT(TSharedPtr<IDetailsView>, DetailsView)
 		/** The toolkit inside which this editor lives, if any. */
 		SLATE_ARGUMENT(TWeakPtr<FAssetEditorToolkit>, AssetEditorToolkit)
-		/** Display info for the transition graph. */
-		SLATE_ATTRIBUTE(FGraphDisplayInfo, TransitionGraphDisplayInfo)
 		/** Appearance info for the transition graph editor. */
 		SLATE_ATTRIBUTE(FGraphAppearanceInfo, TransitionGraphEditorAppearance)
-		/** Formatting of the transition owner's name. */
-		SLATE_ATTRIBUTE(FOnFormatObjectDisplayName, OnFormatTransitionOwnerName)
 	SLATE_END_ARGS()
 
 	void Construct(const FArguments& InArgs);
@@ -90,11 +85,8 @@ private:
 	/** The asset being edited. */
 	UObject* TransitionOwner;
 
-	/** Info about the type of asset we can edit. */
-	FCameraRigTransitionOwnerInfo TransitionOwnerInfo;
-
-	/** Callback for tweaking the graph config. */
-	FOnBuildObjectTreeGraphConfig OnBuildTransitionGraphConfig;
+		/** The graph schema class to use. */
+	TSubclassOf<UCameraRigTransitionGraphSchemaBase> TransitionGraphSchemaClass;
 
 	/** The details view for this editor. */
 	TSharedPtr<IDetailsView> DetailsView;
@@ -107,9 +99,6 @@ private:
 	/** The transition graph editor. */
 	TSharedPtr<SObjectTreeGraphEditor> TransitionGraphEditor;
 
-	/** Appearance info for the transition graph editor. */
-	TAttribute<FGraphDisplayInfo> TransitionGraphDisplayInfo;
-
 	/** Appearance info for the transition graph editor. */
 	TAttribute<FGraphAppearanceInfo> TransitionGraphEditorAppearance;
 
diff --git a/Engine/Plugins/Cameras/GameplayCameras/Source/GameplayCamerasEditor/Public/GameplayCamerasEditorSettings.h b/Engine/Plugins/Cameras/GameplayCameras/Source/GameplayCamerasEditor/Public/GameplayCamerasEditorSettings.h
--- a/Engine/Plugins/Cameras/GameplayCameras/Source/GameplayCamerasEditor/Public/GameplayCamerasEditorSettings.h
+++ b/Engine/Plugins/Cameras/GameplayCameras/Source/GameplayCamerasEditor/Public/GameplayCamerasEditorSettings.h
@@ -24,6 +24,9 @@ public:
 	UPROPERTY(EditAnywhere, Config, Category=NodeTitleColors)
 	FLinearColor CameraNodeTitleColor;
 
+	UPROPERTY(EditAnywhere, Config, Category=NodeTitleColors)
+	FLinearColor CameraAssetTitleColor;
+
 	UPROPERTY(EditAnywhere, Config, Category=NodeTitleColors)
 	FLinearColor CameraRigAssetTitleColor;
 
diff --git a/Engine/Plugins/Cameras/GameplayCameras/Source/GameplayCamerasEditor/Public/Toolkits/CameraRigTransitionEditorToolkitBase.h b/Engine/Plugins/Cameras/GameplayCameras/Source/GameplayCamerasEditor/Public/Toolkits/CameraRigTransitionEditorToolkitBase.h
--- a/Engine/Plugins/Cameras/GameplayCameras/Source/GameplayCamerasEditor/Public/Toolkits/CameraRigTransitionEditorToolkitBase.h
+++ b/Engine/Plugins/Cameras/GameplayCameras/Source/GameplayCamerasEditor/Public/Toolkits/CameraRigTransitionEditorToolkitBase.h
@@ -15,6 +15,7 @@ class FWorkspaceItem;
 class IDetailsView;
 class SDockTab;
 class SObjectTreeGraphToolbox;
+class UCameraRigTransitionGraphSchemaBase;
 class UToolMenu;
 struct FEdGraphEditAction;
 
@@ -55,8 +56,7 @@ protected:
 	virtual void NotifyPostChange(const FPropertyChangedEvent& PropertyChangedEvent, FProperty* PropertyThatChanged) override;
 
 	// FCameraRigTransitionEditorToolkitBase interface
-	virtual void GetTransitionOwnerInfo(FCameraRigTransitionOwnerInfo& OutTransitionOwnerInfo) {}
-	virtual void GetTransitionGraphDisplayInfo(FGraphDisplayInfo& OutGraphDisplayInfo) {}
+	virtual TSubclassOf<UCameraRigTransitionGraphSchemaBase> GetTransitionGraphSchemaClass();
 	virtual void GetTransitionGraphAppearanceInfo(FGraphAppearanceInfo& OutGraphAppearanceInfo) {}
 
 private:
