Cameras: add a custom graph editor pin for picking a camera rig name

diff --git a/Engine/Plugins/Cameras/GameplayCameras/Source/GameplayCameras/Public/Directors/BlueprintCameraDirector.h b/Engine/Plugins/Cameras/GameplayCameras/Source/GameplayCameras/Public/Directors/BlueprintCameraDirector.h
--- a/Engine/Plugins/Cameras/GameplayCameras/Source/GameplayCameras/Public/Directors/BlueprintCameraDirector.h
+++ b/Engine/Plugins/Cameras/GameplayCameras/Source/GameplayCameras/Public/Directors/BlueprintCameraDirector.h
@@ -59,7 +59,7 @@ public:
 	 * Specifies a camera rig to be active this frame.
 	 */
 	UFUNCTION(BlueprintCallable, Category="Evaluation")
-	void ActivateCameraRig(const FString& InCameraRigName);
+	void ActivateCameraRig(UPARAM(meta=(UseCameraRigNamePicker="true")) const FString& InCameraRigName);
 
 	/** Native wrapper for RunCameraDirector. */
 	void NativeRunCameraDirector(
diff --git a/Engine/Plugins/Cameras/GameplayCameras/Source/GameplayCamerasEditor/Private/Editors/GameplayCamerasGraphPanelPinFactory.cpp b/Engine/Plugins/Cameras/GameplayCameras/Source/GameplayCamerasEditor/Private/Editors/GameplayCamerasGraphPanelPinFactory.cpp
new file mode 100644
--- /dev/null
+++ b/Engine/Plugins/Cameras/GameplayCameras/Source/GameplayCamerasEditor/Private/Editors/GameplayCamerasGraphPanelPinFactory.cpp
@@ -0,0 +1,46 @@
+// Copyright Epic Games, Inc. All Rights Reserved.
+
+#include "Editors/GameplayCamerasGraphPanelPinFactory.h"
+
+#include "Editors/SCameraRigNameGraphPin.h"
+#include "K2Node_CallFunction.h"
+
+namespace UE::Cameras
+{
+
+TSharedPtr<class SGraphPin> FGameplayCamerasGraphPanelPinFactory::CreatePin(UEdGraphPin* Pin) const
+{
+	if (!Pin)
+	{
+		return nullptr;
+	}
+
+	UK2Node_CallFunction* OwningNode = Cast<UK2Node_CallFunction>(Pin->GetOwningNode());
+	if (!OwningNode)
+	{
+		return nullptr;
+	}
+
+	UClass* BlueprintClass = OwningNode->GetBlueprintClassFromNode();
+	UFunction* ReferencedFunction = OwningNode->FunctionReference.ResolveMember<UFunction>(BlueprintClass);
+	if (!ReferencedFunction)
+	{
+		return nullptr;
+	}
+
+	FProperty* ParameterProperty = ReferencedFunction->FindPropertyByName(Pin->PinName);
+	if (!ParameterProperty)
+	{
+		return nullptr;
+	}
+
+	if (!ParameterProperty->HasMetaData(TEXT("UseCameraRigNamePicker")))
+	{
+		return nullptr;
+	}
+
+	return SNew(SCameraRigNameGraphPin, Pin);
+}
+
+}  // namespace UE::Cameras
+
diff --git a/Engine/Plugins/Cameras/GameplayCameras/Source/GameplayCamerasEditor/Private/Editors/SCameraRigNameGraphPin.cpp b/Engine/Plugins/Cameras/GameplayCameras/Source/GameplayCamerasEditor/Private/Editors/SCameraRigNameGraphPin.cpp
new file mode 100644
--- /dev/null
+++ b/Engine/Plugins/Cameras/GameplayCameras/Source/GameplayCamerasEditor/Private/Editors/SCameraRigNameGraphPin.cpp
@@ -0,0 +1,307 @@
+// Copyright Epic Games, Inc. All Rights Reserved.
+
+#include "Editors/SCameraRigNameGraphPin.h"
+
+#include "Algo/Transform.h"
+#include "AssetRegistry/AssetRegistryModule.h"
+#include "Core/CameraAsset.h"
+#include "Core/CameraRigAsset.h"
+#include "Framework/Views/ITypedTableView.h"
+#include "IContentBrowserSingleton.h"
+#include "Kismet2/BlueprintEditorUtils.h"
+#include "ScopedTransaction.h"
+#include "SGraphPin.h"
+#include "Styles/GameplayCamerasEditorStyle.h"
+#include "Widgets/Images/SImage.h"
+#include "Widgets/Input/SButton.h"
+#include "Widgets/Input/SEditableTextBox.h"
+#include "Widgets/Layout/SBorder.h"
+#include "Widgets/Layout/SBox.h"
+#include "Widgets/SBoxPanel.h"
+#include "Widgets/SNullWidget.h"
+#include "Widgets/Text/STextBlock.h"
+
+#define LOCTEXT_NAMESPACE "SCameraRigNameGraphPin"
+
+namespace UE::Cameras
+{
+
+void SCameraRigNameGraphPin::Construct(const FArguments& InArgs, UEdGraphPin* InGraphPinObj)
+{
+	SGraphPin::Construct(SGraphPin::FArguments(), InGraphPinObj);
+}
+
+TSharedRef<SWidget>	SCameraRigNameGraphPin::GetDefaultValueWidget()
+{
+	if (!GraphPinObj)
+	{
+		return SNullWidget::NullWidget;
+	}
+
+	return SNew(SHorizontalBox)
+		.Visibility(this, &SGraphPin::GetDefaultValueVisibility)
+		+SHorizontalBox::Slot()
+		.AutoWidth()
+		.Padding(2.f)
+		.MaxWidth(200.f)
+		[
+			SAssignNew(CameraRigPickerButton, SComboButton)
+			.ButtonStyle(FAppStyle::Get(), "PropertyEditor.AssetComboStyle")
+			.ContentPadding(FMargin(2.f, 2.f, 2.f, 1.f))
+			.ForegroundColor(this, &SCameraRigNameGraphPin::OnGetComboForeground)
+			.ButtonColorAndOpacity(this, &SCameraRigNameGraphPin::OnGetWidgetBackground)
+			.MenuPlacement(MenuPlacement_BelowAnchor)
+			.IsEnabled(this, &SGraphPin::IsEditingEnabled)
+			.ButtonContent()
+			[
+				SNew(STextBlock)
+				.ColorAndOpacity(this, &SCameraRigNameGraphPin::OnGetComboForeground)
+				.TextStyle(FAppStyle::Get(), "PropertyEditor.AssetClass")
+				.Font(FAppStyle::GetFontStyle("PropertyWindow.NormalFont"))
+				.Text(this, &SCameraRigNameGraphPin::OnGetComboText)
+				.ToolTipText(this, &SCameraRigNameGraphPin::OnGetComboToolTipText)
+			]
+			.OnGetMenuContent(this, &SCameraRigNameGraphPin::OnBuildCameraRigNamePicker)
+		]
+		// Reset button
+		+SHorizontalBox::Slot()
+		.AutoWidth()
+		.Padding(1,0)
+		.VAlign(VAlign_Center)
+		[
+			SNew(SButton)
+			.ButtonStyle(FAppStyle::Get(), "NoBorder")
+			.ButtonColorAndOpacity(this, &SCameraRigNameGraphPin::OnGetWidgetBackground)
+			.OnClicked(this, &SCameraRigNameGraphPin::OnResetButtonClicked)
+			.ContentPadding(1.f)
+			.ToolTipText(LOCTEXT("ResetButtonToolTip", "Reset the name to an empty string."))
+			.IsEnabled(this, &SGraphPin::IsEditingEnabled)
+			[
+				SNew(SImage)
+				.ColorAndOpacity(this, &SCameraRigNameGraphPin::OnGetWidgetForeground)
+				.Image(FAppStyle::GetBrush(TEXT("Icons.CircleArrowLeft")))
+			]
+		];
+}
+
+bool SCameraRigNameGraphPin::DoesWidgetHandleSettingEditingEnabled() const
+{
+	return true;
+}
+
+FSlateColor SCameraRigNameGraphPin::OnGetComboForeground() const
+{
+	float Alpha = (IsHovered() || bOnlyShowDefaultValue) ? ActiveComboAlpha : InactiveComboAlpha;
+	return FSlateColor(FLinearColor(1.f, 1.f, 1.f, Alpha));
+}
+
+FSlateColor SCameraRigNameGraphPin::OnGetWidgetForeground() const
+{
+	float Alpha = (IsHovered() || bOnlyShowDefaultValue) ? ActivePinForegroundAlpha : InactivePinForegroundAlpha;
+	return FSlateColor(FLinearColor(1.f, 1.f, 1.f, Alpha));
+}
+
+FSlateColor SCameraRigNameGraphPin::OnGetWidgetBackground() const
+{
+	float Alpha = (IsHovered() || bOnlyShowDefaultValue) ? ActivePinBackgroundAlpha : InactivePinBackgroundAlpha;
+	return FSlateColor(FLinearColor(1.f, 1.f, 1.f, Alpha));
+}
+
+FText SCameraRigNameGraphPin::GetDefaultComboText() const
+{
+	return LOCTEXT("DefaultComboText", "Select Camera Rig");
+}
+
+FText SCameraRigNameGraphPin::OnGetComboText() const
+{
+	FText Value = GetDefaultComboText();
+	
+	if (GraphPinObj != nullptr)
+	{
+		if (!GraphPinObj->DefaultValue.IsEmpty())
+		{
+			Value = FText::FromString(GraphPinObj->DefaultValue);
+		}
+	}
+	return Value;
+}
+
+FText SCameraRigNameGraphPin::OnGetComboToolTipText() const
+{
+	return LOCTEXT("ComboToolTipText", "The name of the camera rig to activate.");
+}
+
+TSharedRef<SWidget> SCameraRigNameGraphPin::OnBuildCameraRigNamePicker()
+{
+	FAssetRegistryModule& AssetRegistryModule = FModuleManager::LoadModuleChecked<FAssetRegistryModule>("AssetRegistry");
+	FContentBrowserModule& ContentBrowserModule = FModuleManager::LoadModuleChecked<FContentBrowserModule>("ContentBrowser");
+
+	FAssetPickerConfig AssetPickerConfig;
+	AssetPickerConfig.bAllowDragging = false;
+	AssetPickerConfig.bAllowNullSelection = false;
+	AssetPickerConfig.SelectionMode = ESelectionMode::Multi;
+	AssetPickerConfig.InitialAssetViewType = EAssetViewType::List;
+	AssetPickerConfig.Filter.ClassPaths.Add(UCameraAsset::StaticClass()->GetClassPathName());
+	AssetPickerConfig.Filter.bRecursiveClasses = true;
+	AssetPickerConfig.OnAssetSelected = FOnAssetSelected::CreateSP(this, &SCameraRigNameGraphPin::OnPickerAssetSelected);
+	AssetPickerConfig.GetCurrentSelectionDelegates.Add(&GetCurrentAssetPickerSelection);
+
+	if (TSharedPtr<SGraphNode> OwnerNodeWidget = OwnerNodePtr.Pin())
+	{
+		UEdGraphNode* OwnerNode = OwnerNodeWidget->GetNodeObj();
+		UBlueprint* OwnerBlueprint = FBlueprintEditorUtils::FindBlueprintForNode(OwnerNode);
+
+		FARFilter Filter;
+		IAssetRegistry& AssetRegistry = AssetRegistryModule.Get();
+		AssetRegistry.GetReferencers(OwnerBlueprint->GetOutermost()->GetFName(), Filter.PackageNames);
+		TArray<FAssetData> ReferencerAssetData;
+		AssetRegistry.GetAssets(Filter, ReferencerAssetData);
+		if (!ReferencerAssetData.IsEmpty())
+		{
+			//TODO: multi-select all referencers.
+			AssetPickerConfig.InitialAssetSelection = ReferencerAssetData[0];
+		}
+
+		UpdateListItemsSource(ReferencerAssetData);
+	}
+
+	TSharedRef<SWidget> PickerWidget = SNew(SBox)
+		.HeightOverride(600)
+		.WidthOverride(350)
+		[
+			SNew(SBorder)
+			.BorderImage(FAppStyle::GetBrush("Menu.Background"))
+			[
+				SNew(SVerticalBox)
+				+SVerticalBox::Slot()
+				.FillHeight(0.6f)
+				[
+					ContentBrowserModule.Get().CreateAssetPicker(AssetPickerConfig)
+				]
+				+SVerticalBox::Slot()
+				.FillHeight(0.4f)
+				[
+					SAssignNew(CameraRigNameListView, SListView<TSharedPtr<FString>>)
+					.ListItemsSource(&CameraRigsItemsSource)
+					.OnGenerateRow(this, &SCameraRigNameGraphPin::OnCameraRigListGenerateRow)
+					.OnSelectionChanged(this, &SCameraRigNameGraphPin::OnCameraRigListSelectionChanged)
+				]
+			]
+		];
+
+	if (GraphPinObj && !GraphPinObj->DefaultValue.IsEmpty())
+	{
+		TSharedPtr<FString>* FoundItem = CameraRigsItemsSource.FindByPredicate(
+				[this](const TSharedPtr<FString>& Item)
+				{
+					return GraphPinObj->DefaultValue == *Item.Get();
+				});
+		if (FoundItem)
+		{
+			bSuppressCameraRigListSelectionChanged = true;
+			CameraRigNameListView->SetSelection(*FoundItem);
+			CameraRigNameListView->RequestScrollIntoView(*FoundItem);
+			bSuppressCameraRigListSelectionChanged = false;
+		}
+	}
+
+	return PickerWidget;
+}
+
+void SCameraRigNameGraphPin::OnPickerAssetSelected(const FAssetData& AssetData)
+{
+	TArray<FAssetData> SelectedAssets;
+	if (GetCurrentAssetPickerSelection.IsBound())
+	{
+		SelectedAssets = GetCurrentAssetPickerSelection.Execute();
+	}
+
+	UpdateListItemsSource(SelectedAssets);
+	CameraRigNameListView->RequestListRefresh();
+}
+
+FReply SCameraRigNameGraphPin::OnResetButtonClicked()
+{
+	CameraRigPickerButton->SetIsOpen(false);
+	SetCameraRigName(FString());
+	return FReply::Handled();
+}
+
+TSharedRef<ITableRow> SCameraRigNameGraphPin::OnCameraRigListGenerateRow(TSharedPtr<FString> Item, const TSharedRef<STableViewBase>& OwnerTable)
+{
+	TSharedRef<FGameplayCamerasEditorStyle> GameplayCamerasStyle = FGameplayCamerasEditorStyle::Get();
+
+	const FText DisplayName = FText::FromString(*Item);
+
+	return SNew(STableRow<TSharedPtr<FString>>, OwnerTable)
+		[
+			SNew(SHorizontalBox)
+			+SHorizontalBox::Slot()
+			.AutoWidth()
+			[
+				SNew(SImage)
+				.ColorAndOpacity(FSlateColor::UseForeground())
+				.Image(GameplayCamerasStyle->GetBrush("CameraAssetEditor.ShowCameraRigs"))
+			]
+			+SHorizontalBox::Slot()
+			.FillWidth(1.f)
+			.Padding(4.f, 2.f)
+			[
+				SNew(STextBlock)
+				.Text(DisplayName)
+			]
+		];
+}
+
+void SCameraRigNameGraphPin::OnCameraRigListSelectionChanged(TSharedPtr<FString> Item, ESelectInfo::Type SelectInfo)
+{
+	if (!bSuppressCameraRigListSelectionChanged && Item)
+	{
+		CameraRigPickerButton->SetIsOpen(false);
+		SetCameraRigName(*Item.Get());
+	}
+}
+
+void SCameraRigNameGraphPin::UpdateListItemsSource(const TArray<FAssetData>& Assets)
+{
+	// Only show the names that are in common between all selected camera assets.
+	TSet<FString> CommonNames;
+	bool bCommonNamesInitialized = false;
+	for (const FAssetData& SelectedAsset : Assets)
+	{
+		if (const UCameraAsset* CameraAsset = Cast<UCameraAsset>(SelectedAsset.GetAsset()))
+		{
+			TSet<FString> SelectedAssetNames;
+			Algo::Transform(CameraAsset->CameraRigs, SelectedAssetNames, 
+					[](const UCameraRigAsset* Item) { return Item->GetDisplayName(); });
+
+			if (!bCommonNamesInitialized)
+			{
+				CommonNames.Append(SelectedAssetNames);
+				bCommonNamesInitialized = true;
+			}
+			else
+			{
+				CommonNames = CommonNames.Intersect(SelectedAssetNames);
+			}
+		}
+	}
+
+	CameraRigsItemsSource.Reset();
+	for (const FString& CommonName : CommonNames)
+	{
+		CameraRigsItemsSource.Add(MakeShared<FString>(CommonName));
+	}
+}
+
+void SCameraRigNameGraphPin::SetCameraRigName(const FString& InCameraRigName)
+{
+	const FScopedTransaction Transaction(LOCTEXT("ChangeObjectPinValue", "Change Object Pin Value"));
+	GraphPinObj->Modify();
+	GraphPinObj->GetSchema()->TrySetDefaultValue(*GraphPinObj, InCameraRigName);
+}
+
+}  // namespace UE::Cameras
+
+#undef LOCTEXT_NAMESPACE
+
diff --git a/Engine/Plugins/Cameras/GameplayCameras/Source/GameplayCamerasEditor/Private/GameplayCamerasEditorModule.cpp b/Engine/Plugins/Cameras/GameplayCameras/Source/GameplayCamerasEditor/Private/GameplayCamerasEditorModule.cpp
--- a/Engine/Plugins/Cameras/GameplayCameras/Source/GameplayCamerasEditor/Private/GameplayCamerasEditorModule.cpp
+++ b/Engine/Plugins/Cameras/GameplayCameras/Source/GameplayCamerasEditor/Private/GameplayCamerasEditorModule.cpp
@@ -15,6 +15,7 @@
 #include "Debugger/SBlendStacksDebugPanel.h"
 #include "Debugger/SCameraNodeTreeDebugPanel.h"
 #include "Debugger/SGameplayCamerasDebugger.h"
+#include "Editors/GameplayCamerasGraphPanelPinFactory.h"
 #include "Features/IModularFeatures.h"
 #include "GameplayCameras.h"
 #include "GameplayCamerasEditorSettings.h"
@@ -68,6 +69,7 @@ public:
 		RegisterCoreDebugCategories();
 		RegisterRewindDebuggerFeatures();
 		RegisterDetailsCustomizations();
+		RegisterEdGraphUtilities();
 
 		InitializeLiveEditManager();
 
@@ -91,6 +93,7 @@ public:
 		UnregisterCoreDebugCategories();
 		UnregisterRewindDebuggerFeatures();
 		UnregisterDetailsCustomizations();
+		UnregisterEdGraphUtilities();
 
 		TeardownLiveEditManager();
 
@@ -323,6 +326,22 @@ private:
 		}
 	}
 
+	void RegisterEdGraphUtilities()
+	{
+		using namespace UE::Cameras;
+
+		GraphPanelPinFactory = MakeShared<FGameplayCamerasGraphPanelPinFactory>();
+		FEdGraphUtilities::RegisterVisualPinFactory(GraphPanelPinFactory);
+	}
+
+	void UnregisterEdGraphUtilities()
+	{
+		if (GraphPanelPinFactory)
+		{
+			FEdGraphUtilities::UnregisterVisualPinFactory(GraphPanelPinFactory);
+		}
+	}
+
 	void InitializeLiveEditManager()
 	{
 		using namespace UE::Cameras;
@@ -353,6 +372,8 @@ private:
 
 	TSharedPtr<UE::Cameras::FGameplayCamerasLiveEditManager> LiveEditManager;
 
+	TSharedPtr<UE::Cameras::FGameplayCamerasGraphPanelPinFactory> GraphPanelPinFactory;
+
 	TMap<FString, UE::Cameras::FCameraDebugCategoryInfo> DebugCategoryInfos;
 	TMap<FString, FOnCreateDebugCategoryPanel> DebugCategoryPanelCreators;
 
diff --git a/Engine/Plugins/Cameras/GameplayCameras/Source/GameplayCamerasEditor/Public/Editors/GameplayCamerasGraphPanelPinFactory.h b/Engine/Plugins/Cameras/GameplayCameras/Source/GameplayCamerasEditor/Public/Editors/GameplayCamerasGraphPanelPinFactory.h
new file mode 100644
--- /dev/null
+++ b/Engine/Plugins/Cameras/GameplayCameras/Source/GameplayCamerasEditor/Public/Editors/GameplayCamerasGraphPanelPinFactory.h
@@ -0,0 +1,22 @@
+// Copyright Epic Games, Inc. All Rights Reserved.
+
+#pragma once
+
+#include "EdGraphUtilities.h"
+
+namespace UE::Cameras
+{
+
+/**
+ * Graph editor pin factory for camera-specific pin widgets.
+ */
+struct FGameplayCamerasGraphPanelPinFactory : public FGraphPanelPinFactory
+{
+public:
+
+	// FGraphPanelPinFactory interface.
+	virtual TSharedPtr<class SGraphPin> CreatePin(UEdGraphPin* Pin) const override;
+};
+
+}  // namespace UE::Cameras
+
diff --git a/Engine/Plugins/Cameras/GameplayCameras/Source/GameplayCamerasEditor/Public/Editors/SCameraRigNameGraphPin.h b/Engine/Plugins/Cameras/GameplayCameras/Source/GameplayCamerasEditor/Public/Editors/SCameraRigNameGraphPin.h
new file mode 100644
--- /dev/null
+++ b/Engine/Plugins/Cameras/GameplayCameras/Source/GameplayCamerasEditor/Public/Editors/SCameraRigNameGraphPin.h
@@ -0,0 +1,73 @@
+// Copyright Epic Games, Inc. All Rights Reserved.
+
+#pragma once
+
+#include "SGraphPin.h"
+
+#include "ContentBrowserModule.h"
+#include "Widgets/Input/SComboButton.h"
+#include "Widgets/Views/SListView.h"
+
+namespace UE::Cameras
+{
+
+/**
+ * A custom widget for a graph editor pin that shows a camera rig picker dialog.
+ */
+class SCameraRigNameGraphPin : public SGraphPin
+{
+public:
+
+	SLATE_BEGIN_ARGS(SCameraRigNameGraphPin)
+	{}
+	SLATE_END_ARGS()
+
+	void Construct(const FArguments& InArgs, UEdGraphPin* InGraphPinObj);
+
+protected:
+
+	// SGraphPin interface.
+	virtual TSharedRef<SWidget>	GetDefaultValueWidget() override;
+	virtual bool DoesWidgetHandleSettingEditingEnabled() const override;
+
+private:
+
+	static constexpr float ActiveComboAlpha = 1.f;
+	static constexpr float InactiveComboAlpha = 0.6f;
+	static constexpr float ActivePinForegroundAlpha = 1.f;
+	static constexpr float InactivePinForegroundAlpha = 0.15f;
+	static constexpr float ActivePinBackgroundAlpha = 0.8f;
+	static constexpr float InactivePinBackgroundAlpha = 0.4f;
+
+	FSlateColor OnGetComboForeground() const;
+	FSlateColor OnGetWidgetForeground() const;
+	FSlateColor OnGetWidgetBackground() const;
+
+	FText GetDefaultComboText() const;
+	FText OnGetComboText() const;
+	FText OnGetComboToolTipText() const;
+
+	TSharedRef<SWidget> OnBuildCameraRigNamePicker();
+	void OnPickerAssetSelected(const FAssetData& AssetData);
+
+	FReply OnResetButtonClicked();
+
+	TSharedRef<ITableRow> OnCameraRigListGenerateRow(TSharedPtr<FString> Item, const TSharedRef<STableViewBase>& OwnerTable);
+	void OnCameraRigListSelectionChanged(TSharedPtr<FString> Item, ESelectInfo::Type SelectInfo);
+	void UpdateListItemsSource(const TArray<FAssetData>& Assets);
+
+	void SetCameraRigName(const FString& InCameraRigName);
+
+private:
+
+	TSharedPtr<SComboButton> CameraRigPickerButton;
+
+	FGetCurrentSelectionDelegate GetCurrentAssetPickerSelection;
+
+	TSharedPtr<SListView<TSharedPtr<FString>>> CameraRigNameListView;
+	TArray<TSharedPtr<FString>> CameraRigsItemsSource;
+	bool bSuppressCameraRigListSelectionChanged = false;
+};
+
+}  // namespace UE::Cameras
+
