Cameras: fix localization text mismatch

diff --git a/Engine/Plugins/Cameras/GameplayCameras/Source/GameplayCamerasEditor/Private/Editors/SCameraRigList.cpp b/Engine/Plugins/Cameras/GameplayCameras/Source/GameplayCamerasEditor/Private/Editors/SCameraRigList.cpp
--- a/Engine/Plugins/Cameras/GameplayCameras/Source/GameplayCamerasEditor/Private/Editors/SCameraRigList.cpp
+++ b/Engine/Plugins/Cameras/GameplayCameras/Source/GameplayCamerasEditor/Private/Editors/SCameraRigList.cpp
@@ -424,7 +424,7 @@ TSharedPtr<SWidget> SCameraRigList::OnListContextMenuOpening()
 		FToolMenuSection& Section = ContextMenu->AddSection("Actions");
 		Section.AddEntry(FToolMenuEntry::InitMenuEntry(
 					Commands.EditCameraRig,
-					LOCTEXT("AddCameraRigButton", "Edit")  // Shorter label
+					LOCTEXT("AddCameraRigButton", "Add")  // Shorter label
 				));
 		Section.AddEntry(FToolMenuEntry::InitMenuEntry(
 					Commands.RenameCameraRig,
