Cameras: use proper display name for a camera rig in the transitions graph editor

diff --git a/Engine/Plugins/Cameras/GameplayCameras/Source/GameplayCamerasEditor/Private/Editors/CameraRigTransitionGraphSchema.cpp b/Engine/Plugins/Cameras/GameplayCameras/Source/GameplayCamerasEditor/Private/Editors/CameraRigTransitionGraphSchema.cpp
--- a/Engine/Plugins/Cameras/GameplayCameras/Source/GameplayCamerasEditor/Private/Editors/CameraRigTransitionGraphSchema.cpp
+++ b/Engine/Plugins/Cameras/GameplayCameras/Source/GameplayCamerasEditor/Private/Editors/CameraRigTransitionGraphSchema.cpp
@@ -26,6 +26,14 @@ void UCameraRigTransitionGraphSchema::OnBuildGraphConfig(FObjectTreeGraphConfig&
 		.SetPropertyPinDirection(GET_MEMBER_NAME_CHECKED(UCameraRigAsset, ExitTransitions), EGPD_Output)
 		.NodeTitleUsesObjectName(true)
 		.NodeTitleColor(Settings->CameraRigAssetTitleColor);
+	InOutGraphConfig.OnFormatObjectDisplayName = FOnFormatObjectDisplayName::CreateLambda(
+			[](const UObject* Object, FText& InOutDisplayNameText)
+			{
+				if (const UCameraRigAsset* CameraRigAsset = Cast<UCameraRigAsset>(Object))
+				{
+					InOutDisplayNameText = FText::FromString(CameraRigAsset->GetDisplayName());
+				}
+			});
 }
 
 ETransitionGraphContextActions UCameraRigTransitionGraphSchema::GetTransitionGraphContextActions(FGraphContextMenuBuilder& ContextMenuBuilder) const
