Cameras: change object-tree-graph API to always pass the graph name

This is because some nodes, like a camera rig, may have different position coordinates between the node graph and the transition graph. Before this change, moving a camera rig node in one graph would move it in the other.
Also, use `FIntVector2` to cut down on the number of properties.

diff --git a/Engine/Plugins/Cameras/GameplayCameras/Source/GameplayCameras/Private/Core/CameraAsset.cpp b/Engine/Plugins/Cameras/GameplayCameras/Source/GameplayCameras/Private/Core/CameraAsset.cpp
--- a/Engine/Plugins/Cameras/GameplayCameras/Source/GameplayCameras/Private/Core/CameraAsset.cpp
+++ b/Engine/Plugins/Cameras/GameplayCameras/Source/GameplayCameras/Private/Core/CameraAsset.cpp
@@ -33,26 +33,28 @@ void UCameraAsset::DirtyBuildStatus()
 
 #if WITH_EDITOR
 
-void UCameraAsset::GetGraphNodePosition(int32& NodePosX, int32& NodePosY) const
+void UCameraAsset::GetGraphNodePosition(FName InGraphName, int32& NodePosX, int32& NodePosY) const
 {
-	NodePosX = GraphNodePosX;
-	NodePosY = GraphNodePosY;
+	NodePosX = TransitionGraphNodePos.X;
+	NodePosY = TransitionGraphNodePos.Y;
 }
 
-void UCameraAsset::OnGraphNodeMoved(int32 NodePosX, int32 NodePosY, bool bMarkDirty)
+void UCameraAsset::OnGraphNodeMoved(FName InGraphName, int32 NodePosX, int32 NodePosY, bool bMarkDirty)
 {
-	GraphNodePosX = NodePosX;
-	GraphNodePosY = NodePosY;
+	Modify(bMarkDirty);
+
+	TransitionGraphNodePos.X = NodePosX;
+	TransitionGraphNodePos.Y = NodePosY;
 }
 
-const FString& UCameraAsset::GetGraphNodeCommentText() const
+const FString& UCameraAsset::GetGraphNodeCommentText(FName InGraphName) const
 {
-	return GraphNodeComment;
+	return TransitionGraphNodeComment;
 }
 
-void UCameraAsset::OnUpdateGraphNodeCommentText(const FString& NewComment)
+void UCameraAsset::OnUpdateGraphNodeCommentText(FName InGraphName, const FString& NewComment)
 {
-	GraphNodeComment = NewComment;
+	TransitionGraphNodeComment = NewComment;
 }
 
 void UCameraAsset::GetConnectableObjects(FName InGraphName, TSet<UObject*>& OutObjects) const
diff --git a/Engine/Plugins/Cameras/GameplayCameras/Source/GameplayCameras/Private/Core/CameraNode.cpp b/Engine/Plugins/Cameras/GameplayCameras/Source/GameplayCameras/Private/Core/CameraNode.cpp
--- a/Engine/Plugins/Cameras/GameplayCameras/Source/GameplayCameras/Private/Core/CameraNode.cpp
+++ b/Engine/Plugins/Cameras/GameplayCameras/Source/GameplayCameras/Private/Core/CameraNode.cpp
@@ -6,6 +6,23 @@
 
 #include UE_INLINE_GENERATED_CPP_BY_NAME(CameraNode)
 
+void UCameraNode::PostLoad()
+{
+#if WITH_EDITORONLY_DATA
+
+	if (GraphNodePosX_DEPRECATED != 0 || GraphNodePosY_DEPRECATED != 0)
+	{
+		GraphNodePos = FIntVector2(GraphNodePosX_DEPRECATED, GraphNodePosY_DEPRECATED);
+
+		GraphNodePosX_DEPRECATED = 0;
+		GraphNodePosY_DEPRECATED = 0;
+	}
+
+#endif
+
+	Super::PostLoad();
+}
+
 FCameraNodeChildrenView UCameraNode::GetChildren()
 {
 	return OnGetChildren();
@@ -26,24 +43,24 @@ FCameraNodeEvaluatorPtr UCameraNode::BuildEvaluator(FCameraNodeEvaluatorBuilder&
 
 #if WITH_EDITOR
 
-void UCameraNode::GetGraphNodePosition(int32& NodePosX, int32& NodePosY) const
+void UCameraNode::GetGraphNodePosition(FName InGraphName, int32& NodePosX, int32& NodePosY) const
 {
-	NodePosX = GraphNodePosX;
-	NodePosY = GraphNodePosY;
+	NodePosX = GraphNodePos.X;
+	NodePosY = GraphNodePos.Y;
 }
 
-void UCameraNode::OnGraphNodeMoved(int32 NodePosX, int32 NodePosY, bool bMarkDirty)
+void UCameraNode::OnGraphNodeMoved(FName InGraphName, int32 NodePosX, int32 NodePosY, bool bMarkDirty)
 {
-	GraphNodePosX = NodePosX;
-	GraphNodePosY = NodePosY;
+	GraphNodePos.X = NodePosX;
+	GraphNodePos.Y = NodePosY;
 }
 
-const FString& UCameraNode::GetGraphNodeCommentText() const
+const FString& UCameraNode::GetGraphNodeCommentText(FName InGraphName) const
 {
 	return GraphNodeComment;
 }
 
-void UCameraNode::OnUpdateGraphNodeCommentText(const FString& NewComment)
+void UCameraNode::OnUpdateGraphNodeCommentText(FName InGraphName, const FString& NewComment)
 {
 	GraphNodeComment = NewComment;
 }
diff --git a/Engine/Plugins/Cameras/GameplayCameras/Source/GameplayCameras/Private/Core/CameraRigAsset.cpp b/Engine/Plugins/Cameras/GameplayCameras/Source/GameplayCameras/Private/Core/CameraRigAsset.cpp
--- a/Engine/Plugins/Cameras/GameplayCameras/Source/GameplayCameras/Private/Core/CameraRigAsset.cpp
+++ b/Engine/Plugins/Cameras/GameplayCameras/Source/GameplayCameras/Private/Core/CameraRigAsset.cpp
@@ -14,16 +14,16 @@
 
 #if WITH_EDITOR
 
-void UCameraRigInterfaceParameter::GetGraphNodePosition(int32& NodePosX, int32& NodePosY) const
+void UCameraRigInterfaceParameter::GetGraphNodePosition(FName InGraphName, int32& NodePosX, int32& NodePosY) const
 {
-	NodePosX = GraphNodePosX;
-	NodePosY = GraphNodePosY;
+	NodePosX = GraphNodePos.X;
+	NodePosY = GraphNodePos.Y;
 }
 
-void UCameraRigInterfaceParameter::OnGraphNodeMoved(int32 NodePosX, int32 NodePosY, bool bMarkDirty)
+void UCameraRigInterfaceParameter::OnGraphNodeMoved(FName InGraphName, int32 NodePosX, int32 NodePosY, bool bMarkDirty)
 {
-	GraphNodePosX = NodePosX;
-	GraphNodePosY = NodePosY;
+	GraphNodePos.X = NodePosX;
+	GraphNodePos.Y = NodePosY;
 }
 
 #endif
@@ -45,6 +45,23 @@ bool FCameraRigInterface::HasInterfaceParameter(const FString& ParameterName) co
 const FName UCameraRigAsset::NodeTreeGraphName(TEXT("NodeTree"));
 const FName UCameraRigAsset::TransitionsGraphName(TEXT("Transitions"));
 
+void UCameraRigAsset::PostLoad()
+{
+#if WITH_EDITORONLY_DATA
+
+	if (GraphNodePosX_DEPRECATED != 0 || GraphNodePosY_DEPRECATED != 0)
+	{
+		NodeGraphNodePos = FIntVector2(GraphNodePosX_DEPRECATED, GraphNodePosY_DEPRECATED);
+
+		GraphNodePosX_DEPRECATED = 0;
+		GraphNodePosY_DEPRECATED = 0;
+	}
+
+#endif
+
+	Super::PostLoad();
+}
+
 void UCameraRigAsset::GetOwnedGameplayTags(FGameplayTagContainer& TagContainer) const
 {
 	TagContainer.AppendTags(GameplayTags);
@@ -84,11 +101,13 @@ void UCameraRigAsset::DirtyBuildStatus()
 void UCameraRigAsset::PreSave(FObjectPreSaveContext ObjectSaveContext)
 {
 #if WITH_EDITOR
+
 	if (!HasAnyFlags(RF_ClassDefaultObject | RF_ArchetypeObject))
 	{
 		// Build on save.
 		BuildCameraRig();
 	}
+
 #endif
 
 	Super::PreSave(ObjectSaveContext);
@@ -117,26 +136,61 @@ void UCameraRigAsset::GatherPackages(FCameraRigPackages& OutPackages) const
 	}
 }
 
-void UCameraRigAsset::GetGraphNodePosition(int32& NodePosX, int32& NodePosY) const
+void UCameraRigAsset::GetGraphNodePosition(FName InGraphName, int32& NodePosX, int32& NodePosY) const
 {
-	NodePosX = GraphNodePosX;
-	NodePosY = GraphNodePosY;
+	if (InGraphName == NodeTreeGraphName)
+	{
+		NodePosX = NodeGraphNodePos.X;
+		NodePosY = NodeGraphNodePos.Y;
+	}
+	else if (InGraphName == TransitionsGraphName)
+	{
+		NodePosX = TransitionGraphNodePos.X;
+		NodePosY = TransitionGraphNodePos.Y;
+	}
 }
 
-void UCameraRigAsset::OnGraphNodeMoved(int32 NodePosX, int32 NodePosY, bool bMarkDirty)
+void UCameraRigAsset::OnGraphNodeMoved(FName InGraphName, int32 NodePosX, int32 NodePosY, bool bMarkDirty)
 {
-	GraphNodePosX = NodePosX;
-	GraphNodePosY = NodePosY;
+	Modify(bMarkDirty);
+
+	if (InGraphName == NodeTreeGraphName)
+	{
+		NodeGraphNodePos.X = NodePosX;
+		NodeGraphNodePos.Y = NodePosY;
+	}
+	else if (InGraphName == TransitionsGraphName)
+	{
+		TransitionGraphNodePos.X = NodePosX;
+		TransitionGraphNodePos.Y = NodePosY;
+	}
 }
 
-const FString& UCameraRigAsset::GetGraphNodeCommentText() const
+const FString& UCameraRigAsset::GetGraphNodeCommentText(FName InGraphName) const
 {
-	return GraphNodeComment;
+	if (InGraphName == NodeTreeGraphName)
+	{
+		return NodeGraphNodeComment;
+	}
+	else if (InGraphName == TransitionsGraphName)
+	{
+		return TransitionGraphNodeComment;
+	}
+
+	static FString InvalidString;
+	return InvalidString;
 }
 
-void UCameraRigAsset::OnUpdateGraphNodeCommentText(const FString& NewComment)
+void UCameraRigAsset::OnUpdateGraphNodeCommentText(FName InGraphName, const FString& NewComment)
 {
-	GraphNodeComment = NewComment;
+	if (InGraphName == NodeTreeGraphName)
+	{
+		NodeGraphNodeComment = NewComment;
+	}
+	else if (InGraphName == TransitionsGraphName)
+	{
+		TransitionGraphNodeComment = NewComment;
+	}
 }
 
 void UCameraRigAsset::GetConnectableObjects(FName InGraphName, TSet<UObject*>& OutObjects) const
diff --git a/Engine/Plugins/Cameras/GameplayCameras/Source/GameplayCameras/Private/Core/CameraRigTransition.cpp b/Engine/Plugins/Cameras/GameplayCameras/Source/GameplayCameras/Private/Core/CameraRigTransition.cpp
--- a/Engine/Plugins/Cameras/GameplayCameras/Source/GameplayCameras/Private/Core/CameraRigTransition.cpp
+++ b/Engine/Plugins/Cameras/GameplayCameras/Source/GameplayCameras/Private/Core/CameraRigTransition.cpp
@@ -4,6 +4,23 @@
 
 #include UE_INLINE_GENERATED_CPP_BY_NAME(CameraRigTransition)
 
+void UCameraRigTransitionCondition::PostLoad()
+{
+#if WITH_EDITORONLY_DATA
+
+	if (GraphNodePosX_DEPRECATED != 0 || GraphNodePosY_DEPRECATED != 0)
+	{
+		GraphNodePos = FIntVector2(GraphNodePosX_DEPRECATED, GraphNodePosY_DEPRECATED);
+
+		GraphNodePosX_DEPRECATED = 0;
+		GraphNodePosY_DEPRECATED = 0;
+	}
+
+#endif
+
+	Super::PostLoad();
+}
+
 bool UCameraRigTransitionCondition::TransitionMatches(const FCameraRigTransitionConditionMatchParams& Params) const
 {
 	return OnTransitionMatches(Params);
@@ -11,48 +28,74 @@ bool UCameraRigTransitionCondition::TransitionMatches(const FCameraRigTransition
 
 #if WITH_EDITOR
 
-void UCameraRigTransitionCondition::GetGraphNodePosition(int32& NodePosX, int32& NodePosY) const
+void UCameraRigTransitionCondition::GetGraphNodePosition(FName InGraphName, int32& NodePosX, int32& NodePosY) const
 {
-	NodePosX = GraphNodePosX;
-	NodePosY = GraphNodePosY;
+	NodePosX = GraphNodePos.X;
+	NodePosY = GraphNodePos.Y;
 }
 
-void UCameraRigTransitionCondition::OnGraphNodeMoved(int32 NodePosX, int32 NodePosY, bool bMarkDirty)
+void UCameraRigTransitionCondition::OnGraphNodeMoved(FName InGraphName, int32 NodePosX, int32 NodePosY, bool bMarkDirty)
 {
-	GraphNodePosX = NodePosX;
-	GraphNodePosY = NodePosY;
+	Modify(bMarkDirty);
+
+	GraphNodePos.X = NodePosX;
+	GraphNodePos.Y = NodePosY;
 }
 
-const FString& UCameraRigTransitionCondition::GetGraphNodeCommentText() const
+const FString& UCameraRigTransitionCondition::GetGraphNodeCommentText(FName InGraphName) const
 {
 	return GraphNodeComment;
 }
 
-void UCameraRigTransitionCondition::OnUpdateGraphNodeCommentText(const FString& NewComment)
+void UCameraRigTransitionCondition::OnUpdateGraphNodeCommentText(FName InGraphName, const FString& NewComment)
 {
 	GraphNodeComment = NewComment;
 }
 
-void UCameraRigTransition::GetGraphNodePosition(int32& NodePosX, int32& NodePosY) const
+#endif
+
+void UCameraRigTransition::PostLoad()
+{
+#if WITH_EDITORONLY_DATA
+
+	if (GraphNodePosX_DEPRECATED != 0 || GraphNodePosY_DEPRECATED != 0)
+	{
+		GraphNodePos = FIntVector2(GraphNodePosX_DEPRECATED, GraphNodePosY_DEPRECATED);
+
+		GraphNodePosX_DEPRECATED = 0;
+		GraphNodePosY_DEPRECATED = 0;
+	}
+
+#endif
+
+	Super::PostLoad();
+}
+
+#if WITH_EDITOR
+
+void UCameraRigTransition::GetGraphNodePosition(FName InGraphName, int32& NodePosX, int32& NodePosY) const
 {
-	NodePosX = GraphNodePosX;
-	NodePosY = GraphNodePosY;
+	NodePosX = GraphNodePos.X;
+	NodePosY = GraphNodePos.Y;
 }
 
-void UCameraRigTransition::OnGraphNodeMoved(int32 NodePosX, int32 NodePosY, bool bMarkDirty)
+void UCameraRigTransition::OnGraphNodeMoved(FName InGraphName, int32 NodePosX, int32 NodePosY, bool bMarkDirty)
 {
-	GraphNodePosX = NodePosX;
-	GraphNodePosY = NodePosY;
+	Modify(bMarkDirty);
+
+	GraphNodePos.X = NodePosX;
+	GraphNodePos.Y = NodePosY;
 }
 
-const FString& UCameraRigTransition::GetGraphNodeCommentText() const
+const FString& UCameraRigTransition::GetGraphNodeCommentText(FName InGraphName) const
 {
 	return GraphNodeComment;
 }
 
-void UCameraRigTransition::OnUpdateGraphNodeCommentText(const FString& NewComment)
+void UCameraRigTransition::OnUpdateGraphNodeCommentText(FName InGraphName, const FString& NewComment)
 {
 	GraphNodeComment = NewComment;
 }
 
 #endif  // WITH_EDITOR
+
diff --git a/Engine/Plugins/Cameras/GameplayCameras/Source/GameplayCameras/Public/Core/CameraAsset.h b/Engine/Plugins/Cameras/GameplayCameras/Source/GameplayCameras/Public/Core/CameraAsset.h
--- a/Engine/Plugins/Cameras/GameplayCameras/Source/GameplayCameras/Public/Core/CameraAsset.h
+++ b/Engine/Plugins/Cameras/GameplayCameras/Source/GameplayCameras/Public/Core/CameraAsset.h
@@ -77,11 +77,11 @@ protected:
 
 	// IObjectTreeGraphObject interface.
 #if WITH_EDITOR
-	virtual void GetGraphNodePosition(int32& NodePosX, int32& NodePosY) const override;
-	virtual void OnGraphNodeMoved(int32 NodePosX, int32 NodePosY, bool bMarkDirty) override;
-	virtual EObjectTreeGraphObjectSupportFlags GetSupportFlags() const override { return EObjectTreeGraphObjectSupportFlags::CommentText; }
-	virtual const FString& GetGraphNodeCommentText() const override;
-	virtual void OnUpdateGraphNodeCommentText(const FString& NewComment) override;
+	virtual void GetGraphNodePosition(FName InGraphName, int32& NodePosX, int32& NodePosY) const override;
+	virtual void OnGraphNodeMoved(FName InGraphName, int32 NodePosX, int32 NodePosY, bool bMarkDirty) override;
+	virtual EObjectTreeGraphObjectSupportFlags GetSupportFlags(FName InGraphName) const override { return EObjectTreeGraphObjectSupportFlags::CommentText; }
+	virtual const FString& GetGraphNodeCommentText(FName InGraphName) const override;
+	virtual void OnUpdateGraphNodeCommentText(FName InGraphName, const FString& NewComment) override;
 #endif
 
 	// IObjectTreeGraphRootObject interface.
@@ -95,21 +95,15 @@ private:
 
 #if WITH_EDITORONLY_DATA
 
-	/** Position of the camera node in the transition editor. */
+	/** Position of the camera node in the shared transitions graph editor. */
 	UPROPERTY()
-	int32 GraphNodePosX = 0;
+	FIntVector2 TransitionGraphNodePos = FIntVector2::ZeroValue;
 
-	/** Position of the camera node in the graph editor. */
+	/** User-written comment in the transition graph editor. */
 	UPROPERTY()
-	int32 GraphNodePosY = 0;
+	FString TransitionGraphNodeComment;
 
-	/** User-written comment in the graph editor. */
-	UPROPERTY()
-	FString GraphNodeComment;
-
-	/**
-	 * Similar to AllNodeTreeObjects, but for the transitions graph.
-	 */
+	/** All nodes used in the shared transitions graph editor. */
 	UPROPERTY(Instanced, meta=(ObjectTreeGraphHidden=true))
 	TArray<TObjectPtr<UObject>> AllSharedTransitionsObjects;
 
diff --git a/Engine/Plugins/Cameras/GameplayCameras/Source/GameplayCameras/Public/Core/CameraNode.h b/Engine/Plugins/Cameras/GameplayCameras/Source/GameplayCameras/Public/Core/CameraNode.h
--- a/Engine/Plugins/Cameras/GameplayCameras/Source/GameplayCameras/Public/Core/CameraNode.h
+++ b/Engine/Plugins/Cameras/GameplayCameras/Source/GameplayCameras/Public/Core/CameraNode.h
@@ -52,13 +52,16 @@ protected:
 
 protected:
 
+	// UObject interface.
+	virtual void PostLoad() override;
+
 	// IObjectTreeGraphObject interface.
 #if WITH_EDITOR
-	virtual void GetGraphNodePosition(int32& NodePosX, int32& NodePosY) const override;
-	virtual void OnGraphNodeMoved(int32 NodePosX, int32 NodePosY, bool bMarkDirty) override;
-	virtual EObjectTreeGraphObjectSupportFlags GetSupportFlags() const override { return EObjectTreeGraphObjectSupportFlags::CommentText; }
-	virtual const FString& GetGraphNodeCommentText() const override;
-	virtual void OnUpdateGraphNodeCommentText(const FString& NewComment) override;
+	virtual void GetGraphNodePosition(FName InGraphName, int32& NodePosX, int32& NodePosY) const override;
+	virtual void OnGraphNodeMoved(FName InGraphName, int32 NodePosX, int32 NodePosY, bool bMarkDirty) override;
+	virtual EObjectTreeGraphObjectSupportFlags GetSupportFlags(FName InGraphName) const override { return EObjectTreeGraphObjectSupportFlags::CommentText; }
+	virtual const FString& GetGraphNodeCommentText(FName InGraphName) const override;
+	virtual void OnUpdateGraphNodeCommentText(FName InGraphName, const FString& NewComment) override;
 #endif
 
 public:
@@ -69,14 +72,21 @@ public:
 
 #if WITH_EDITORONLY_DATA
 
+	/** Position of the camera node in the node graph editor. */
 	UPROPERTY()
-	int32 GraphNodePosX = 0;
+	FIntVector2 GraphNodePos = FIntVector2::ZeroValue;
 
+	/** User-written comment in the node graph editor. */
 	UPROPERTY()
-	int32 GraphNodePosY = 0;
+	FString GraphNodeComment;
+
+
+	// Deprecated properties.
 
 	UPROPERTY()
-	FString GraphNodeComment;
+	int32 GraphNodePosX_DEPRECATED = 0;
+	UPROPERTY()
+	int32 GraphNodePosY_DEPRECATED = 0;
 
 #endif  // WITH_EDITORONLY_DATA
 };
diff --git a/Engine/Plugins/Cameras/GameplayCameras/Source/GameplayCameras/Public/Core/CameraRigAsset.h b/Engine/Plugins/Cameras/GameplayCameras/Source/GameplayCameras/Public/Core/CameraRigAsset.h
--- a/Engine/Plugins/Cameras/GameplayCameras/Source/GameplayCameras/Public/Core/CameraRigAsset.h
+++ b/Engine/Plugins/Cameras/GameplayCameras/Source/GameplayCameras/Public/Core/CameraRigAsset.h
@@ -80,8 +80,8 @@ protected:
 
 	// IObjectTreeGraphObject interface.
 #if WITH_EDITOR
-	virtual void GetGraphNodePosition(int32& NodePosX, int32& NodePosY) const override;
-	virtual void OnGraphNodeMoved(int32 NodePosX, int32 NodePosY, bool bMarkDirty) override;
+	virtual void GetGraphNodePosition(FName InGraphName, int32& NodePosX, int32& NodePosY) const override;
+	virtual void OnGraphNodeMoved(FName InGraphName, int32 NodePosX, int32 NodePosY, bool bMarkDirty) override;
 #endif
 
 private:
@@ -89,10 +89,7 @@ private:
 #if WITH_EDITORONLY_DATA
 
 	UPROPERTY()
-	int32 GraphNodePosX = 0;
-
-	UPROPERTY()
-	int32 GraphNodePosY = 0;
+	FIntVector2 GraphNodePos = FIntVector2::ZeroValue;
 
 #endif  // WITH_EDITORONLY_DATA
 };
@@ -216,13 +213,16 @@ public:
 
 protected:
 
+	// UObject interface.
+	virtual void PostLoad() override;
+
 	// IObjectTreeGraphObject interface.
 #if WITH_EDITOR
-	virtual void GetGraphNodePosition(int32& NodePosX, int32& NodePosY) const override;
-	virtual void OnGraphNodeMoved(int32 NodePosX, int32 NodePosY, bool bMarkDirty) override;
-	virtual EObjectTreeGraphObjectSupportFlags GetSupportFlags() const override { return EObjectTreeGraphObjectSupportFlags::CommentText; }
-	virtual const FString& GetGraphNodeCommentText() const override;
-	virtual void OnUpdateGraphNodeCommentText(const FString& NewComment) override;
+	virtual void GetGraphNodePosition(FName InGraphName, int32& NodePosX, int32& NodePosY) const override;
+	virtual void OnGraphNodeMoved(FName InGraphName, int32 NodePosX, int32 NodePosY, bool bMarkDirty) override;
+	virtual EObjectTreeGraphObjectSupportFlags GetSupportFlags(FName InGraphName) const override { return EObjectTreeGraphObjectSupportFlags::CommentText; }
+	virtual const FString& GetGraphNodeCommentText(FName InGraphName) const override;
+	virtual void OnUpdateGraphNodeCommentText(FName InGraphName, const FString& NewComment) override;
 #endif
 
 	// IObjectTreeGraphRootObject interface.
@@ -239,17 +239,21 @@ private:
 
 #if WITH_EDITORONLY_DATA
 
-	/** Position of the camera rig node in the graph editor. */
+	/** Position of the camera rig node in the node graph editor. */
+	UPROPERTY()
+	FIntVector2 NodeGraphNodePos = FIntVector2::ZeroValue;
+
+	/** Position of the camera rig node in the transition graph editor. */
 	UPROPERTY()
-	int32 GraphNodePosX = 0;
+	FIntVector2 TransitionGraphNodePos = FIntVector2::ZeroValue;
 
-	/** Position of the camera rig node in the graph editor. */
+	/** User-written comment in the node graph editor. */
 	UPROPERTY()
-	int32 GraphNodePosY = 0;
+	FString NodeGraphNodeComment;
 
-	/** User-written comment in the graph editor. */
+	/** User-written comment in the transition graph editor. */
 	UPROPERTY()
-	FString GraphNodeComment;
+	FString TransitionGraphNodeComment;
 
 	/** 
 	 * A list of all the camera nodes, including the 'loose' ones that aren't connected
@@ -264,6 +268,14 @@ private:
 	UPROPERTY(Instanced, meta=(ObjectTreeGraphHidden=true))
 	TArray<TObjectPtr<UObject>> AllTransitionsObjects;
 
+
+	// Deprecated properties.
+
+	UPROPERTY()
+	int32 GraphNodePosX_DEPRECATED = 0;
+	UPROPERTY()
+	int32 GraphNodePosY_DEPRECATED = 0;
+
 #endif  // WITH_EDITORONLY_DATA
 
 	friend class UE::Cameras::FCameraRigAssetBuilder;
diff --git a/Engine/Plugins/Cameras/GameplayCameras/Source/GameplayCameras/Public/Core/CameraRigTransition.h b/Engine/Plugins/Cameras/GameplayCameras/Source/GameplayCameras/Public/Core/CameraRigTransition.h
--- a/Engine/Plugins/Cameras/GameplayCameras/Source/GameplayCameras/Public/Core/CameraRigTransition.h
+++ b/Engine/Plugins/Cameras/GameplayCameras/Source/GameplayCameras/Public/Core/CameraRigTransition.h
@@ -51,27 +51,37 @@ protected:
 
 protected:
 
+	// UObject interface.
+	virtual void PostLoad() override;
+
 	// IObjectTreeGraphObject interface.
 #if WITH_EDITOR
-	virtual void GetGraphNodePosition(int32& NodePosX, int32& NodePosY) const override;
-	virtual void OnGraphNodeMoved(int32 NodePosX, int32 NodePosY, bool bMarkDirty) override;
-	virtual EObjectTreeGraphObjectSupportFlags GetSupportFlags() const override { return EObjectTreeGraphObjectSupportFlags::CommentText; }
-	virtual const FString& GetGraphNodeCommentText() const override;
-	virtual void OnUpdateGraphNodeCommentText(const FString& NewComment) override;
+	virtual void GetGraphNodePosition(FName InGraphName, int32& NodePosX, int32& NodePosY) const override;
+	virtual void OnGraphNodeMoved(FName InGraphName, int32 NodePosX, int32 NodePosY, bool bMarkDirty) override;
+	virtual EObjectTreeGraphObjectSupportFlags GetSupportFlags(FName InGraphName) const override { return EObjectTreeGraphObjectSupportFlags::CommentText; }
+	virtual const FString& GetGraphNodeCommentText(FName InGraphName) const override;
+	virtual void OnUpdateGraphNodeCommentText(FName InGraphName, const FString& NewComment) override;
 #endif
 
 private:
 
 #if WITH_EDITORONLY_DATA
 
+	/** Position of the transition condition node in the transition graph editor. */
 	UPROPERTY()
-		int32 GraphNodePosX = 0;
+	FIntVector2 GraphNodePos = FIntVector2::ZeroValue;
 
+	/** User-written comment in the transition graph editor. */
 	UPROPERTY()
-		int32 GraphNodePosY = 0;
+	FString GraphNodeComment;
+
 
+	// Deprecated properties.
+
+	UPROPERTY()
+	int32 GraphNodePosX_DEPRECATED = 0;
 	UPROPERTY()
-		FString GraphNodeComment;
+	int32 GraphNodePosY_DEPRECATED = 0;
 
 #endif  // WITH_EDITORONLY_DATA
 };
@@ -98,27 +108,37 @@ public:
 
 protected:
 
+	// UObject interface.
+	virtual void PostLoad() override;
+
 	// IObjectTreeGraphObject interface.
 #if WITH_EDITOR
-	virtual void GetGraphNodePosition(int32& NodePosX, int32& NodePosY) const override;
-	virtual void OnGraphNodeMoved(int32 NodePosX, int32 NodePosY, bool bMarkDirty) override;
-	virtual EObjectTreeGraphObjectSupportFlags GetSupportFlags() const override { return EObjectTreeGraphObjectSupportFlags::CommentText; }
-	virtual const FString& GetGraphNodeCommentText() const override;
-	virtual void OnUpdateGraphNodeCommentText(const FString& NewComment) override;
+	virtual void GetGraphNodePosition(FName InGraphName, int32& NodePosX, int32& NodePosY) const override;
+	virtual void OnGraphNodeMoved(FName InGraphName, int32 NodePosX, int32 NodePosY, bool bMarkDirty) override;
+	virtual EObjectTreeGraphObjectSupportFlags GetSupportFlags(FName InGraphName) const override { return EObjectTreeGraphObjectSupportFlags::CommentText; }
+	virtual const FString& GetGraphNodeCommentText(FName InGraphName) const override;
+	virtual void OnUpdateGraphNodeCommentText(FName InGraphName, const FString& NewComment) override;
 #endif
 
 private:
 
 #if WITH_EDITORONLY_DATA
 
+	/** Position of the transition node in the transition graph editor. */
 	UPROPERTY()
-		int32 GraphNodePosX = 0;
+	FIntVector2 GraphNodePos = FIntVector2::ZeroValue;
 
+	/** User-written comment in the transition graph editor. */
 	UPROPERTY()
-		int32 GraphNodePosY = 0;
+	FString GraphNodeComment;
+
 
+	// Deprecated properties.
+
+	UPROPERTY()
+	int32 GraphNodePosX_DEPRECATED = 0;
 	UPROPERTY()
-		FString GraphNodeComment;
+	int32 GraphNodePosY_DEPRECATED = 0;
 
 #endif  // WITH_EDITORONLY_DATA
 };
diff --git a/Engine/Plugins/Cameras/GameplayCameras/Source/GameplayCameras/Public/Core/ObjectTreeGraphObject.h b/Engine/Plugins/Cameras/GameplayCameras/Source/GameplayCameras/Public/Core/ObjectTreeGraphObject.h
--- a/Engine/Plugins/Cameras/GameplayCameras/Source/GameplayCameras/Public/Core/ObjectTreeGraphObject.h
+++ b/Engine/Plugins/Cameras/GameplayCameras/Source/GameplayCameras/Public/Core/ObjectTreeGraphObject.h
@@ -45,26 +45,26 @@ public:
 	// Minimal API.
 	
 	/** Gets the canvas position for the graph node representing this object. */
-	virtual void GetGraphNodePosition(int32& NodePosX, int32& NodePosY) const {}
+	virtual void GetGraphNodePosition(FName InGraphName, int32& NodePosX, int32& NodePosY) const {}
 	/** Called to save the canvas position of the graph node representing this object. */
-	virtual void OnGraphNodeMoved(int32 NodePosX, int32 NodePosY, bool bMarkDirty) {}
+	virtual void OnGraphNodeMoved(FName InGraphName, int32 NodePosX, int32 NodePosY, bool bMarkDirty) {}
 
 	// Optional API.
 
 	/** Gets optional APIs support flags. */
-	virtual EObjectTreeGraphObjectSupportFlags GetSupportFlags() const { return EObjectTreeGraphObjectSupportFlags::None; }
+	virtual EObjectTreeGraphObjectSupportFlags GetSupportFlags(FName InGraphName) const { return EObjectTreeGraphObjectSupportFlags::None; }
 	/** Gets whether this object supports the given optional API. */
-	bool HasSupportFlags(EObjectTreeGraphObjectSupportFlags InFlags) const { return EnumHasAllFlags(GetSupportFlags(), InFlags); }
+	bool HasSupportFlags(FName InGraphName, EObjectTreeGraphObjectSupportFlags InFlags) const { return EnumHasAllFlags(GetSupportFlags(InGraphName), InFlags); }
 
 	/** Gets the graph node's comment text. */
-	virtual const FString& GetGraphNodeCommentText() const { static const FString EmptyString; return EmptyString; }
+	virtual const FString& GetGraphNodeCommentText(FName InGraphName) const { static const FString EmptyString; return EmptyString; }
 	/** Called to save a new comment text. */
-	virtual void OnUpdateGraphNodeCommentText(const FString& NewComment) {}
+	virtual void OnUpdateGraphNodeCommentText(FName InGraphName, const FString& NewComment) {}
 
 	/** Gets the custom name for the graph node. */
-	virtual const FString& GetGraphNodeName() const { static const FString EmptyString; return EmptyString; }
+	virtual const FString& GetGraphNodeName(FName InGraphName) const { static const FString EmptyString; return EmptyString; }
 	/** Called to save a new custom name. */
-	virtual void OnRenameGraphNode(const FString& NewName) {}
+	virtual void OnRenameGraphNode(FName InGraphName, const FString& NewName) {}
 
 #endif  // WITH_EDITOR
 };
diff --git a/Engine/Plugins/Cameras/GameplayCameras/Source/GameplayCamerasEditor/Private/Editors/ObjectTreeGraphConfig.cpp b/Engine/Plugins/Cameras/GameplayCameras/Source/GameplayCamerasEditor/Private/Editors/ObjectTreeGraphConfig.cpp
--- a/Engine/Plugins/Cameras/GameplayCameras/Source/GameplayCamerasEditor/Private/Editors/ObjectTreeGraphConfig.cpp
+++ b/Engine/Plugins/Cameras/GameplayCameras/Source/GameplayCamerasEditor/Private/Editors/ObjectTreeGraphConfig.cpp
@@ -163,9 +163,9 @@ FText FObjectTreeGraphConfig::GetDisplayNameText(const UObject* InObject) const
 		const FObjectTreeGraphClassConfig& ClassConfig = GetObjectClassConfig(InObject->GetClass());
 
 		const IObjectTreeGraphObject* GraphObject = Cast<IObjectTreeGraphObject>(InObject);
-		if (GraphObject && GraphObject->HasSupportFlags(EObjectTreeGraphObjectSupportFlags::CustomRename))
+		if (GraphObject && GraphObject->HasSupportFlags(GraphName, EObjectTreeGraphObjectSupportFlags::CustomRename))
 		{
-			DisplayNameText = FText::FromString(GraphObject->GetGraphNodeName());
+			DisplayNameText = FText::FromString(GraphObject->GetGraphNodeName(GraphName));
 		}
 
 		if (DisplayNameText.IsEmpty() && ClassConfig.NodeTitleUsesObjectName())
diff --git a/Engine/Plugins/Cameras/GameplayCameras/Source/GameplayCamerasEditor/Private/Editors/ObjectTreeGraphNode.cpp b/Engine/Plugins/Cameras/GameplayCameras/Source/GameplayCamerasEditor/Private/Editors/ObjectTreeGraphNode.cpp
--- a/Engine/Plugins/Cameras/GameplayCameras/Source/GameplayCamerasEditor/Private/Editors/ObjectTreeGraphNode.cpp
+++ b/Engine/Plugins/Cameras/GameplayCameras/Source/GameplayCamerasEditor/Private/Editors/ObjectTreeGraphNode.cpp
@@ -26,10 +26,11 @@ void UObjectTreeGraphNode::Initialize(UObject* InObject)
 
 	Object = InObject;
 
+	const FNodeContext NodeContext = GetNodeContext();
 	IObjectTreeGraphObject* GraphObject = Cast<IObjectTreeGraphObject>(Object);
-	if (GraphObject && GraphObject->HasSupportFlags(EObjectTreeGraphObjectSupportFlags::CommentText))
+	if (GraphObject && GraphObject->HasSupportFlags(NodeContext.GraphConfig.GraphName, EObjectTreeGraphObjectSupportFlags::CommentText))
 	{
-		NodeComment = GraphObject->GetGraphNodeCommentText();
+		NodeComment = GraphObject->GetGraphNodeCommentText(NodeContext.GraphConfig.GraphName);
 	}
 }
 
@@ -419,7 +420,8 @@ void UObjectTreeGraphNode::PostPlacedNewNode()
 	IObjectTreeGraphObject* GraphObject = Cast<IObjectTreeGraphObject>(GetObject());
 	if (GraphObject)
 	{
-		GraphObject->GetGraphNodePosition(NodePosX, NodePosY);
+		const FNodeContext NodeContext = GetNodeContext();
+		GraphObject->GetGraphNodePosition(NodeContext.GraphConfig.GraphName, NodePosX, NodePosY);
 	}
 }
 
@@ -487,8 +489,9 @@ void UObjectTreeGraphNode::GetNodeContextMenuActions(class UToolMenu* Menu, clas
 
 bool UObjectTreeGraphNode::GetCanRenameNode() const
 {
+	const FNodeContext NodeContext = GetNodeContext();
 	IObjectTreeGraphObject* GraphObject = Cast<IObjectTreeGraphObject>(Object);
-	return GraphObject && GraphObject->HasSupportFlags(EObjectTreeGraphObjectSupportFlags::CustomRename);
+	return GraphObject && GraphObject->HasSupportFlags(NodeContext.GraphConfig.GraphName, EObjectTreeGraphObjectSupportFlags::CustomRename);
 }
 
 void UObjectTreeGraphNode::OnRenameNode(const FString& NewName)
@@ -498,8 +501,11 @@ void UObjectTreeGraphNode::OnRenameNode(const FString& NewName)
 	if (IObjectTreeGraphObject* GraphObject = Cast<IObjectTreeGraphObject>(Object))
 	{
 		const FScopedTransaction Transaction(LOCTEXT("RenameNode", "Rename Node"));
+
 		Object->Modify();
-		GraphObject->OnRenameGraphNode(NewName);
+
+		const FNodeContext NodeContext = GetNodeContext();
+		GraphObject->OnRenameGraphNode(NodeContext.GraphConfig.GraphName, NewName);
 	}
 }
 
@@ -516,8 +522,9 @@ bool UObjectTreeGraphNode::CanUserDeleteNode() const
 
 bool UObjectTreeGraphNode::SupportsCommentBubble() const
 {
+	const FNodeContext NodeContext = GetNodeContext();
 	IObjectTreeGraphObject* GraphObject = Cast<IObjectTreeGraphObject>(Object);
-	return GraphObject && GraphObject->HasSupportFlags(EObjectTreeGraphObjectSupportFlags::CommentText);
+	return GraphObject && GraphObject->HasSupportFlags(NodeContext.GraphConfig.GraphName, EObjectTreeGraphObjectSupportFlags::CommentText);
 }
 
 void UObjectTreeGraphNode::OnUpdateCommentText(const FString& NewComment)
@@ -528,8 +535,11 @@ void UObjectTreeGraphNode::OnUpdateCommentText(const FString& NewComment)
 	if (GraphObject)
 	{
 		const FScopedTransaction Transaction(LOCTEXT("UpdateNodeComment", "Update Node Comment"));
+
 		Object->Modify();
-		GraphObject->OnUpdateGraphNodeCommentText(NewComment);
+
+		const FNodeContext NodeContext = GetNodeContext();
+		GraphObject->OnUpdateGraphNodeCommentText(NodeContext.GraphConfig.GraphName, NewComment);
 	}
 }
 
@@ -539,7 +549,9 @@ void UObjectTreeGraphNode::OnGraphNodeMoved(bool bMarkDirty)
 	if (GraphObject)
 	{
 		Object->Modify(bMarkDirty);
-		GraphObject->OnGraphNodeMoved(NodePosX, NodePosY, bMarkDirty);
+
+		const FNodeContext NodeContext = GetNodeContext();
+		GraphObject->OnGraphNodeMoved(NodeContext.GraphConfig.GraphName, NodePosX, NodePosY, bMarkDirty);
 	}
 }
 
