Cameras: auto-populate the blueprint camera director Blueprint event graph

diff --git a/Engine/Plugins/Cameras/GameplayCameras/Source/GameplayCamerasEditor/Private/GameplayCamerasEditorModule.cpp b/Engine/Plugins/Cameras/GameplayCameras/Source/GameplayCamerasEditor/Private/GameplayCamerasEditorModule.cpp
--- a/Engine/Plugins/Cameras/GameplayCameras/Source/GameplayCamerasEditor/Private/GameplayCamerasEditorModule.cpp
+++ b/Engine/Plugins/Cameras/GameplayCameras/Source/GameplayCamerasEditor/Private/GameplayCamerasEditorModule.cpp
@@ -15,6 +15,8 @@
 #include "Debugger/SBlendStacksDebugPanel.h"
 #include "Debugger/SCameraNodeTreeDebugPanel.h"
 #include "Debugger/SGameplayCamerasDebugger.h"
+#include "Directors/BlueprintCameraDirector.h"
+#include "EdGraph/EdGraph.h"
 #include "Editors/GameplayCamerasGraphPanelPinFactory.h"
 #include "Features/IModularFeatures.h"
 #include "GameplayCameras.h"
@@ -24,6 +26,8 @@
 #include "IGameplayCamerasModule.h"
 #include "IRewindDebuggerExtension.h"
 #include "ISettingsModule.h"
+#include "K2Node_Event.h"
+#include "Kismet2/KismetEditorUtilities.h"
 #include "Misc/CoreDelegates.h"
 #include "Modules/ModuleManager.h"
 #include "PropertyEditorModule.h"
@@ -332,6 +336,60 @@ private:
 
 		GraphPanelPinFactory = MakeShared<FGameplayCamerasGraphPanelPinFactory>();
 		FEdGraphUtilities::RegisterVisualPinFactory(GraphPanelPinFactory);
+
+		FKismetEditorUtilities::RegisterAutoGeneratedDefaultEvent(
+				this,
+				UBlueprintCameraDirectorEvaluator::StaticClass(),
+				GET_FUNCTION_NAME_CHECKED(UBlueprintCameraDirectorEvaluator, RunCameraDirector));
+		FKismetEditorUtilities::RegisterOnBlueprintCreatedCallback(
+				this, 
+				UBlueprintCameraDirectorEvaluator::StaticClass(),
+				FKismetEditorUtilities::FOnBlueprintCreated::CreateRaw(
+					this, &FGameplayCamerasEditorModule::OnNewBlueprintCameraDirectorEvaluatorCreated));
+	}
+
+	void OnNewBlueprintCameraDirectorEvaluatorCreated(UBlueprint* InBlueprint)
+	{
+		if (InBlueprint->BlueprintType != BPTYPE_Normal)
+		{
+			return;
+		}
+		TObjectPtr<UEdGraph>* FoundItem = InBlueprint->UbergraphPages.FindByPredicate(
+					[](UEdGraph* Item) { return Item->GetFName() == TEXT("EventGraph"); });
+		if (!FoundItem)
+		{
+			return;
+		}
+
+		TObjectPtr<UEdGraph> EventGraph(*FoundItem);
+
+		TArray<UK2Node_Event*> EventNodes;
+		EventGraph->GetNodesOfClass(EventNodes);
+		if (EventNodes.IsEmpty())
+		{
+			return;
+		}
+
+		const FName EventName = GET_FUNCTION_NAME_CHECKED(UBlueprintCameraDirectorEvaluator, RunCameraDirector);
+		UK2Node_Event** FoundEventNode = EventNodes.FindByPredicate(
+				[&EventName](UK2Node_Event* Item) 
+				{
+					return Item->EventReference.GetMemberName() == EventName;
+				});
+		if (!FoundEventNode)
+		{
+			return;
+		}
+
+		UK2Node_Event* RunEventNode(*FoundEventNode);
+	
+		const FText RunEventNodeCommentText = LOCTEXT(
+				"BlueprintCameraDirector_RunEventComment",
+				"Implement your camera director logic starting from here.\n"
+				"This node is currently disabled, but start dragging off pins to enable it.\n"
+				"Call ActivateCameraRig at least once to declare which camera rig(s) should be active this frame.");
+		RunEventNode->NodeComment = RunEventNodeCommentText.ToString();
+		RunEventNode->bCommentBubbleVisible = true;
 	}
 
 	void UnregisterEdGraphUtilities()
@@ -340,6 +398,8 @@ private:
 		{
 			FEdGraphUtilities::UnregisterVisualPinFactory(GraphPanelPinFactory);
 		}
+
+		FKismetEditorUtilities::UnregisterAutoBlueprintNodeCreation(this);
 	}
 
 	void InitializeLiveEditManager()
