Cameras: select and edit the first camera rig in a camera asset's list when the editor opens

diff --git a/Engine/Plugins/Cameras/GameplayCameras/Source/GameplayCamerasEditor/Private/Editors/SCameraRigList.cpp b/Engine/Plugins/Cameras/GameplayCameras/Source/GameplayCamerasEditor/Private/Editors/SCameraRigList.cpp
--- a/Engine/Plugins/Cameras/GameplayCameras/Source/GameplayCamerasEditor/Private/Editors/SCameraRigList.cpp
+++ b/Engine/Plugins/Cameras/GameplayCameras/Source/GameplayCamerasEditor/Private/Editors/SCameraRigList.cpp
@@ -190,6 +190,7 @@ void SCameraRigList::Construct(const FArguments& InArgs)
 	if (!FilteredItemSource.IsEmpty())
 	{
 		ListView->SetSelection(FilteredItemSource[0]);
+		OnRequestEditCameraRig.ExecuteIfBound(FilteredItemSource[0]->CameraRigAsset);
 	}
 }
 
