Cameras: tweak the camera rig name picker UI and serialize asset picker settings

diff --git a/Engine/Plugins/Cameras/GameplayCameras/Source/GameplayCamerasEditor/Private/Editors/SCameraRigNameGraphPin.cpp b/Engine/Plugins/Cameras/GameplayCameras/Source/GameplayCamerasEditor/Private/Editors/SCameraRigNameGraphPin.cpp
--- a/Engine/Plugins/Cameras/GameplayCameras/Source/GameplayCamerasEditor/Private/Editors/SCameraRigNameGraphPin.cpp
+++ b/Engine/Plugins/Cameras/GameplayCameras/Source/GameplayCamerasEditor/Private/Editors/SCameraRigNameGraphPin.cpp
@@ -143,6 +143,7 @@ TSharedRef<SWidget> SCameraRigNameGraphPin::OnBuildCameraRigNamePicker()
 	AssetPickerConfig.InitialAssetViewType = EAssetViewType::List;
 	AssetPickerConfig.Filter.ClassPaths.Add(UCameraAsset::StaticClass()->GetClassPathName());
 	AssetPickerConfig.Filter.bRecursiveClasses = true;
+	AssetPickerConfig.SaveSettingsName = TEXT("CameraRigNamePicker");
 	AssetPickerConfig.OnAssetSelected = FOnAssetSelected::CreateSP(this, &SCameraRigNameGraphPin::OnPickerAssetSelected);
 	AssetPickerConfig.GetCurrentSelectionDelegates.Add(&GetCurrentAssetPickerSelection);
 
@@ -166,7 +167,7 @@ TSharedRef<SWidget> SCameraRigNameGraphPin::OnBuildCameraRigNamePicker()
 	}
 
 	TSharedRef<SWidget> PickerWidget = SNew(SBox)
-		.HeightOverride(600)
+		.HeightOverride(400)
 		.WidthOverride(350)
 		[
 			SNew(SBorder)
@@ -174,12 +175,12 @@ TSharedRef<SWidget> SCameraRigNameGraphPin::OnBuildCameraRigNamePicker()
 			[
 				SNew(SVerticalBox)
 				+SVerticalBox::Slot()
-				.FillHeight(0.6f)
+				.FillHeight(0.55f)
 				[
 					ContentBrowserModule.Get().CreateAssetPicker(AssetPickerConfig)
 				]
 				+SVerticalBox::Slot()
-				.FillHeight(0.4f)
+				.FillHeight(0.45f)
 				[
 					SAssignNew(CameraRigNameListView, SListView<TSharedPtr<FString>>)
 					.ListItemsSource(&CameraRigsItemsSource)
@@ -234,21 +235,27 @@ TSharedRef<ITableRow> SCameraRigNameGraphPin::OnCameraRigListGenerateRow(TShared
 	const FText DisplayName = FText::FromString(*Item);
 
 	return SNew(STableRow<TSharedPtr<FString>>, OwnerTable)
+		.Padding(FMargin(0.f, 4.f))
 		[
-			SNew(SHorizontalBox)
-			+SHorizontalBox::Slot()
-			.AutoWidth()
-			[
-				SNew(SImage)
-				.ColorAndOpacity(FSlateColor::UseForeground())
-				.Image(GameplayCamerasStyle->GetBrush("CameraAssetEditor.ShowCameraRigs"))
-			]
-			+SHorizontalBox::Slot()
-			.FillWidth(1.f)
-			.Padding(4.f, 2.f)
+			SNew(SBorder)
+			.Padding(FMargin(4.f))
+			.BorderImage(FAppStyle::GetBrush("NoBorder"))
 			[
-				SNew(STextBlock)
-				.Text(DisplayName)
+				SNew(SHorizontalBox)
+				+SHorizontalBox::Slot()
+				.AutoWidth()
+				[
+					SNew(SImage)
+					.ColorAndOpacity(FSlateColor::UseForeground())
+					.Image(GameplayCamerasStyle->GetBrush("CameraAssetEditor.ShowCameraRigs"))
+				]
+				+SHorizontalBox::Slot()
+				.FillWidth(1.f)
+				.Padding(4.f, 2.f)
+				[
+					SNew(STextBlock)
+					.Text(DisplayName)
+				]
 			]
 		];
 }
