Cameras: fix toolbox not updating when switching between camera node graph and transition graph

diff --git a/Engine/Plugins/Cameras/GameplayCameras/Source/GameplayCamerasEditor/Private/Toolkits/CameraRigAssetEditorToolkitBase.cpp b/Engine/Plugins/Cameras/GameplayCameras/Source/GameplayCamerasEditor/Private/Toolkits/CameraRigAssetEditorToolkitBase.cpp
--- a/Engine/Plugins/Cameras/GameplayCameras/Source/GameplayCamerasEditor/Private/Toolkits/CameraRigAssetEditorToolkitBase.cpp
+++ b/Engine/Plugins/Cameras/GameplayCameras/Source/GameplayCamerasEditor/Private/Toolkits/CameraRigAssetEditorToolkitBase.cpp
@@ -166,15 +166,15 @@ void FCameraRigAssetEditorToolkitBase::BindCommands(TSharedRef<FUICommandList> C
 
 	CommandList->MapAction(
 		Commands.ShowNodeHierarchy,
-		FExecuteAction::CreateSP(CameraRigEditor, &SCameraRigAssetEditor::SetEditorMode, ECameraRigAssetEditorMode::NodeGraph),
+		FExecuteAction::CreateSP(this, &FCameraRigAssetEditorToolkitBase::SetCameraRigEditorMode, ECameraRigAssetEditorMode::NodeGraph),
 		FCanExecuteAction(),
-		FIsActionChecked::CreateSP(CameraRigEditor, &SCameraRigAssetEditor::IsEditorMode, ECameraRigAssetEditorMode::NodeGraph));
+		FIsActionChecked::CreateSP(this, &FCameraRigAssetEditorToolkitBase::IsCameraRigEditorMode, ECameraRigAssetEditorMode::NodeGraph));
 
 	CommandList->MapAction(
 		Commands.ShowTransitions,
-		FExecuteAction::CreateSP(CameraRigEditor, &SCameraRigAssetEditor::SetEditorMode, ECameraRigAssetEditorMode::TransitionGraph),
+		FExecuteAction::CreateSP(this, &FCameraRigAssetEditorToolkitBase::SetCameraRigEditorMode, ECameraRigAssetEditorMode::TransitionGraph),
 		FCanExecuteAction(),
-		FIsActionChecked::CreateSP(CameraRigEditor, &SCameraRigAssetEditor::IsEditorMode, ECameraRigAssetEditorMode::TransitionGraph));
+		FIsActionChecked::CreateSP(this, &FCameraRigAssetEditorToolkitBase::IsCameraRigEditorMode, ECameraRigAssetEditorMode::TransitionGraph));
 }
 
 void FCameraRigAssetEditorToolkitBase::NotifyPostChange(const FPropertyChangedEvent& PropertyChangedEvent, FProperty* PropertyThatChanged)
@@ -186,6 +186,17 @@ void FCameraRigAssetEditorToolkitBase::NotifyPostChange(const FPropertyChangedEv
 	}
 }
 
+void FCameraRigAssetEditorToolkitBase::SetCameraRigEditorMode(ECameraRigAssetEditorMode InEditorMode)
+{
+	CameraRigEditorWidget->SetEditorMode(InEditorMode);
+	ToolboxWidget->SetGraphConfig(CameraRigEditorWidget->GetFocusedGraphConfig());
+}
+
+bool FCameraRigAssetEditorToolkitBase::IsCameraRigEditorMode(ECameraRigAssetEditorMode InEditorMode) const
+{
+	return CameraRigEditorWidget->IsEditorMode(InEditorMode);
+}
+
 void FCameraRigAssetEditorToolkitBase::OnAnyGraphChanged(const FEdGraphEditAction& InEditAction)
 {
 	// Called when something is modified in the node graph or transition graph.
diff --git a/Engine/Plugins/Cameras/GameplayCameras/Source/GameplayCamerasEditor/Public/Toolkits/CameraRigAssetEditorToolkitBase.h b/Engine/Plugins/Cameras/GameplayCameras/Source/GameplayCamerasEditor/Public/Toolkits/CameraRigAssetEditorToolkitBase.h
--- a/Engine/Plugins/Cameras/GameplayCameras/Source/GameplayCamerasEditor/Public/Toolkits/CameraRigAssetEditorToolkitBase.h
+++ b/Engine/Plugins/Cameras/GameplayCameras/Source/GameplayCamerasEditor/Public/Toolkits/CameraRigAssetEditorToolkitBase.h
@@ -22,6 +22,7 @@ namespace UE::Cameras
 
 class FStandardToolkitLayout;
 class SCameraRigAssetEditor;
+enum class ECameraRigAssetEditorMode;
 
 /**
  * Editor toolkit for a camera rig asset.
@@ -58,6 +59,11 @@ protected:
 	// FNotifyHook interface
 	virtual void NotifyPostChange(const FPropertyChangedEvent& PropertyChangedEvent, FProperty* PropertyThatChanged) override;
 
+protected:
+
+	void SetCameraRigEditorMode(ECameraRigAssetEditorMode InEditorMode);
+	bool IsCameraRigEditorMode(ECameraRigAssetEditorMode InEditorMode) const;
+
 private:
 
 	TSharedRef<SDockTab> SpawnTab_Toolbox(const FSpawnTabArgs& Args);
