Cameras: don't use UObject RTTI information on shutdown if the UObject system has already been torn down

diff --git a/Engine/Plugins/Cameras/GameplayCameras/Source/GameplayCamerasEditor/Private/Customizations/CameraParameterDetailsCustomizations.cpp b/Engine/Plugins/Cameras/GameplayCameras/Source/GameplayCamerasEditor/Private/Customizations/CameraParameterDetailsCustomizations.cpp
--- a/Engine/Plugins/Cameras/GameplayCameras/Source/GameplayCamerasEditor/Private/Customizations/CameraParameterDetailsCustomizations.cpp
+++ b/Engine/Plugins/Cameras/GameplayCameras/Source/GameplayCamerasEditor/Private/Customizations/CameraParameterDetailsCustomizations.cpp
@@ -206,11 +206,14 @@ UE_CAMERA_VARIABLE_FOR_ALL_TYPES()
 
 void FCameraParameterDetailsCustomization::Unregister(FPropertyEditorModule& PropertyEditorModule)
 {
+	if (UObjectInitialized())
+	{
 #define UE_CAMERA_VARIABLE_FOR_TYPE(ValueType, ValueName)\
-	PropertyEditorModule.UnregisterCustomPropertyTypeLayout(\
-			F##ValueName##CameraParameter::StaticStruct()->GetFName());
+		PropertyEditorModule.UnregisterCustomPropertyTypeLayout(\
+				F##ValueName##CameraParameter::StaticStruct()->GetFName());
 UE_CAMERA_VARIABLE_FOR_ALL_TYPES()
 #undef UE_CAMERA_VARIABLE_FOR_TYPE
+	}
 }
 
 void FCameraParameterDetailsCustomization::CustomizeHeader(TSharedRef<IPropertyHandle> PropertyHandle, FDetailWidgetRow& HeaderRow, IPropertyTypeCustomizationUtils& CustomizationUtils)
