Cameras: fix incorrect schema class being used to make an object tree graph config

diff --git a/Engine/Plugins/Cameras/GameplayCameras/Source/GameplayCamerasEditor/Private/Editors/SCameraRigAssetEditor.cpp b/Engine/Plugins/Cameras/GameplayCameras/Source/GameplayCamerasEditor/Private/Editors/SCameraRigAssetEditor.cpp
--- a/Engine/Plugins/Cameras/GameplayCameras/Source/GameplayCamerasEditor/Private/Editors/SCameraRigAssetEditor.cpp
+++ b/Engine/Plugins/Cameras/GameplayCameras/Source/GameplayCamerasEditor/Private/Editors/SCameraRigAssetEditor.cpp
@@ -70,7 +70,7 @@ void SCameraRigAssetEditor::CreateGraphEditors()
 void SCameraRigAssetEditor::CreateNodeGraphEditor()
 {
 	UClass* SchemaClass = UCameraNodeGraphSchema::StaticClass();
-	UCameraNodeGraphSchema* DefaultSchemaObject = SchemaClass->GetDefaultObject<UCameraNodeGraphSchema>();
+	UCameraNodeGraphSchema* DefaultSchemaObject = Cast<UCameraNodeGraphSchema>(SchemaClass->GetDefaultObject());
 	FObjectTreeGraphConfig GraphConfig = DefaultSchemaObject->BuildGraphConfig();
 
 	NodeGraph = NewObject<UObjectTreeGraph>(GetTransientPackage(), NAME_None, RF_Transactional);
@@ -95,7 +95,7 @@ void SCameraRigAssetEditor::CreateNodeGraphEditor()
 void SCameraRigAssetEditor::CreateTransitionGraphEditor()
 {
 	UClass* SchemaClass = UCameraRigTransitionGraphSchema::StaticClass();
-	UCameraRigTransitionGraphSchema* DefaultSchemaObject = SchemaClass->GetDefaultObject<UCameraRigTransitionGraphSchema>();
+	UCameraRigTransitionGraphSchema* DefaultSchemaObject = Cast<UCameraRigTransitionGraphSchema>(SchemaClass->GetDefaultObject());
 	FObjectTreeGraphConfig GraphConfig = DefaultSchemaObject->BuildGraphConfig();
 
 	TransitionGraph = NewObject<UObjectTreeGraph>(GetTransientPackage(), NAME_None, RF_Transactional);
diff --git a/Engine/Plugins/Cameras/GameplayCameras/Source/GameplayCamerasEditor/Private/Editors/SCameraRigTransitionEditor.cpp b/Engine/Plugins/Cameras/GameplayCameras/Source/GameplayCamerasEditor/Private/Editors/SCameraRigTransitionEditor.cpp
--- a/Engine/Plugins/Cameras/GameplayCameras/Source/GameplayCamerasEditor/Private/Editors/SCameraRigTransitionEditor.cpp
+++ b/Engine/Plugins/Cameras/GameplayCameras/Source/GameplayCamerasEditor/Private/Editors/SCameraRigTransitionEditor.cpp
@@ -64,7 +64,7 @@ void SCameraRigTransitionEditor::SetTransitionOwner(UObject* InTransitionOwner)
 
 void SCameraRigTransitionEditor::CreateTransitionGraphEditor()
 {
-	UCameraRigTransitionGraphSchemaBase* DefaultSchemaObject = TransitionGraphSchemaClass->GetDefaultObject<UCameraRigTransitionGraphSchemaBase>();
+	UCameraRigTransitionGraphSchemaBase* DefaultSchemaObject = Cast<UCameraRigTransitionGraphSchemaBase>(TransitionGraphSchemaClass->GetDefaultObject());
 	FObjectTreeGraphConfig GraphConfig = DefaultSchemaObject->BuildGraphConfig();
 
 	TransitionGraph = NewObject<UObjectTreeGraph>(GetTransientPackage(), NAME_None, RF_Transactional);
