Cameras: make searching object tree graphs go through the objects instead of the graph nodes

This makes it possible to search through object trees that don't have a graph available yet. It avoid having to create the graphs for everything up-front.
Also, make search results properly hierarchical in the tree view.

diff --git a/Engine/Plugins/Cameras/GameplayCameras/Source/GameplayCamerasEditor/Private/Editors/ObjectTreeGraphSearch.cpp b/Engine/Plugins/Cameras/GameplayCameras/Source/GameplayCamerasEditor/Private/Editors/ObjectTreeGraphSearch.cpp
new file mode 100644
--- /dev/null
+++ b/Engine/Plugins/Cameras/GameplayCameras/Source/GameplayCamerasEditor/Private/Editors/ObjectTreeGraphSearch.cpp
@@ -0,0 +1,167 @@
+// Copyright Epic Games, Inc. All Rights Reserved.
+
+#include "Editors/ObjectTreeGraphSearch.h"
+
+#include "Core/ObjectTreeGraphRootObject.h"
+#include "Editors/ObjectTreeGraphConfig.h"
+
+FObjectTreeGraphSearch::FObjectTreeGraphSearch()
+{
+}
+
+void FObjectTreeGraphSearch::AddRootObject(UObject* InObject, const FObjectTreeGraphConfig* InGraphConfig)
+{
+	RootObjectInfos.Add(FRootObjectInfo{ InObject, InGraphConfig });
+}
+
+void FObjectTreeGraphSearch::Search(TArrayView<FString> InTokens, TArray<FObjectTreeGraphSearchResult>& OutResults) const
+{
+	for (const FRootObjectInfo& RootObjectInfo : RootObjectInfos)
+	{
+		SearchRootObject(RootObjectInfo, InTokens, OutResults);
+	}
+}
+
+void FObjectTreeGraphSearch::SearchRootObject(const FRootObjectInfo& InRootObjectInfo, TArrayView<FString> InTokens, TArray<FSearchResult>& OutResults) const
+{
+	UObject* RootObject = InRootObjectInfo.WeakRootObject.Get();
+	if (!RootObject || !InRootObjectInfo.GraphConfig)
+	{
+		return;
+	}
+
+	FSearchState State;
+	State.GraphConfig = InRootObjectInfo.GraphConfig;
+	State.Tokens = InTokens;
+	State.RootObject = RootObject;
+	State.ObjectStack.Add(RootObject);
+
+	if (IObjectTreeGraphRootObject* RootObjectInterface = Cast<IObjectTreeGraphRootObject>(RootObject))
+	{
+		TSet<UObject*> ConnectableObjects;
+		RootObjectInterface->GetConnectableObjects(State.GraphConfig->GraphName, ConnectableObjects);
+		for (UObject* ConnectableObject : ConnectableObjects)
+		{
+			State.ObjectStack.Add(ConnectableObject);
+		}
+	}
+
+	while (State.ObjectStack.Num() > 0)
+	{
+		UObject* CurObject = State.ObjectStack.Pop();
+		if (CurObject && !State.VisitedObjects.Contains(CurObject))
+		{
+			State.VisitedObjects.Add(CurObject);
+
+			SearchObject(CurObject, State);
+		}
+	}
+
+	OutResults.Append(State.Results);
+}
+
+void FObjectTreeGraphSearch::SearchObject(UObject* InObject, FSearchState& InOutState) const
+{
+	UClass* ObjectClass = InObject->GetClass();
+	const FObjectTreeGraphConfig& GraphConfig(*InOutState.GraphConfig);
+
+	const FObjectTreeGraphClassConfig& ObjectClassConfig = GraphConfig.GetObjectClassConfig(ObjectClass);
+
+	if (MatchObject(InObject, InOutState))
+	{
+		InOutState.Results.Add(FSearchResult{ InOutState.RootObject, InOutState.GraphConfig, InObject });
+	}
+
+	for (TFieldIterator<FProperty> PropertyIt(ObjectClass); PropertyIt; ++PropertyIt)
+	{
+		if (FObjectProperty* ObjectProperty = CastField<FObjectProperty>(*PropertyIt))
+		{
+			if (!GraphConfig.IsConnectable(ObjectProperty))
+			{
+				continue;
+			}
+
+			if (MatchObjectProperty(InObject, *PropertyIt, InOutState))
+			{
+				InOutState.Results.Add(
+						FSearchResult{ InOutState.RootObject, InOutState.GraphConfig, InObject, PropertyIt->GetFName() });
+			}
+
+			TObjectPtr<UObject> ConnectedObject;
+			ObjectProperty->GetValue_InContainer(InObject, &ConnectedObject);
+			if (ConnectedObject)
+			{
+				InOutState.ObjectStack.Add(ConnectedObject);
+			}
+		}
+		else if (FArrayProperty* ArrayProperty = CastField<FArrayProperty>(*PropertyIt))
+		{
+			if (!GraphConfig.IsConnectable(ArrayProperty))
+			{
+				continue;
+			}
+
+			if (MatchObjectProperty(InObject, *PropertyIt, InOutState))
+			{
+				InOutState.Results.Add(
+						FSearchResult{ InOutState.RootObject, InOutState.GraphConfig, InObject, PropertyIt->GetFName() });
+			}
+
+			FObjectProperty* InnerProperty = CastFieldChecked<FObjectProperty>(ArrayProperty->Inner);
+			FScriptArrayHelper ArrayHelper(ArrayProperty, ArrayProperty->ContainerPtrToValuePtr<void>(InObject));
+
+			const int32 ArrayNum = ArrayHelper.Num();
+			for (int32 Index = 0; Index < ArrayNum; ++Index)
+			{
+				UObject* ConnectedObject = InnerProperty->GetObjectPropertyValue(ArrayHelper.GetRawPtr(Index));
+				if (ConnectedObject)
+				{
+					InOutState.ObjectStack.Add(ConnectedObject);
+				}
+			}
+		}
+	}
+}
+
+bool FObjectTreeGraphSearch::MatchObject(UObject* InObject, const FSearchState& InState) const
+{
+	const FObjectTreeGraphConfig& GraphConfig(*InState.GraphConfig);
+
+	const FString DisplayNameText = GraphConfig.GetDisplayNameText(InObject).ToString();
+	if (MatchString(DisplayNameText, InState))
+	{
+		return true;
+	}
+
+	return false;
+}
+
+bool FObjectTreeGraphSearch::MatchObjectProperty(UObject* InObject, FProperty* InProperty, const FSearchState& InState) const
+{
+	const FString PropertyName = InProperty->GetName();
+	if (MatchString(PropertyName, InState))
+	{
+		return true;
+	}
+
+	return false;
+}
+
+bool FObjectTreeGraphSearch::MatchString(const FString& InString, const FSearchState& InState) const
+{
+	if (InString.IsEmpty())
+	{
+		return false;
+	}
+
+	for (const FString& Token : InState.Tokens)
+	{
+		if (!InString.Contains(Token))
+		{
+			return false;
+		}
+	}
+
+	return true;
+}
+
diff --git a/Engine/Plugins/Cameras/GameplayCameras/Source/GameplayCamerasEditor/Private/Editors/SCameraRigAssetEditor.cpp b/Engine/Plugins/Cameras/GameplayCameras/Source/GameplayCamerasEditor/Private/Editors/SCameraRigAssetEditor.cpp
--- a/Engine/Plugins/Cameras/GameplayCameras/Source/GameplayCamerasEditor/Private/Editors/SCameraRigAssetEditor.cpp
+++ b/Engine/Plugins/Cameras/GameplayCameras/Source/GameplayCamerasEditor/Private/Editors/SCameraRigAssetEditor.cpp
@@ -237,28 +237,7 @@ void SCameraRigAssetEditor::FocusHome()
 
 	if (Graph && GraphEditor)
 	{
-		if (UObjectTreeGraphNode* RootObjectNode = Graph->GetRootObjectNode())
-		{
-			JumpToNode(RootObjectNode);
-		}
-	}
-}
-
-void SCameraRigAssetEditor::JumpToNode(UEdGraphNode* InGraphNode)
-{
-	if (InGraphNode)
-	{
-		UEdGraph* Graph = InGraphNode->GetGraph();
-		if (Graph == NodeGraph)
-		{
-			SetEditorMode(ECameraRigAssetEditorMode::NodeGraph);
-			NodeGraphEditor->JumpToNode(InGraphNode);
-		}
-		else if (Graph == TransitionGraph)
-		{
-			SetEditorMode(ECameraRigAssetEditorMode::TransitionGraph);
-			TransitionGraphEditor->JumpToNode(InGraphNode);
-		}
+		FindAndJumpToObjectNode(CameraRigAsset);
 	}
 }
 
diff --git a/Engine/Plugins/Cameras/GameplayCameras/Source/GameplayCamerasEditor/Private/Editors/SFindInObjectTreeGraph.cpp b/Engine/Plugins/Cameras/GameplayCameras/Source/GameplayCamerasEditor/Private/Editors/SFindInObjectTreeGraph.cpp
--- a/Engine/Plugins/Cameras/GameplayCameras/Source/GameplayCamerasEditor/Private/Editors/SFindInObjectTreeGraph.cpp
+++ b/Engine/Plugins/Cameras/GameplayCameras/Source/GameplayCamerasEditor/Private/Editors/SFindInObjectTreeGraph.cpp
@@ -2,10 +2,13 @@
 
 #include "Editors/SFindInObjectTreeGraph.h"
 
+#include "Core/ObjectTreeGraphObject.h"
 #include "EdGraph/EdGraph.h"
 #include "EdGraph/EdGraphNode.h"
 #include "EdGraph/EdGraphPin.h"
 #include "EdGraph/EdGraphSchema.h"
+#include "Editors/ObjectTreeGraphConfig.h"
+#include "Editors/ObjectTreeGraphSearch.h"
 #include "Framework/Application/SlateApplication.h"
 #include "Types/SlateEnums.h"
 #include "Widgets/Images/SImage.h"
@@ -20,15 +23,25 @@ FFindInObjectTreeGraphResult::FFindInObjectTreeGraphResult(const FText& InCustom
 {
 }
 
-FFindInObjectTreeGraphResult::FFindInObjectTreeGraphResult(TSharedPtr<FFindInObjectTreeGraphResult>& InParent, UEdGraphNode* InGraphNode)
+FFindInObjectTreeGraphResult::FFindInObjectTreeGraphResult(
+		TSharedPtr<FFindInObjectTreeGraphResult>& InParent, 
+		const FFindInObjectTreeGraphSource& InSource, 
+		UObject* InObject)
 	: Parent(InParent)
-	, GraphNode(InGraphNode)
+	, WeakObject(InObject)
+	, Source(InSource)
 {
 }
 
-FFindInObjectTreeGraphResult::FFindInObjectTreeGraphResult(TSharedPtr<FFindInObjectTreeGraphResult>& InParent, UEdGraphPin* InGraphPin)
+FFindInObjectTreeGraphResult::FFindInObjectTreeGraphResult(
+		TSharedPtr<FFindInObjectTreeGraphResult>& InParent, 
+		const FFindInObjectTreeGraphSource& InSource, 
+		UObject* InObject, 
+		FName InPropertyName)
 	: Parent(InParent)
-	, GraphPin(InGraphPin)
+	, WeakObject(InObject)
+	, PropertyName(InPropertyName)
+	, Source(InSource)
 {
 }
 
@@ -37,13 +50,17 @@ TSharedRef<SWidget>	FFindInObjectTreeGraphResult::GetIcon() const
 	FSlateColor IconColor = FSlateColor::UseForeground();
 	const FSlateBrush* Brush = NULL;
 
-	if (UEdGraphPin* ResolvedPin = GraphPin.Get())
+	UObject* Object = WeakObject.Get();
+
+	if (Object && !PropertyName.IsNone())
 	{
-		if (ResolvedPin->PinType.IsArray())
+		UClass* ObjectClass = Object->GetClass();
+		FProperty* Property = ObjectClass->FindPropertyByName(PropertyName);
+		if (Property->IsA<FArrayProperty>())
 		{
 			Brush = FAppStyle::GetBrush(TEXT("GraphEditor.ArrayPinIcon"));
 		}
-		else if (ResolvedPin->PinType.bIsReference)
+		else if (Property->IsA<FObjectProperty>())
 		{
 			Brush = FAppStyle::GetBrush(TEXT("GraphEditor.RefPinIcon"));
 		}
@@ -51,11 +68,8 @@ TSharedRef<SWidget>	FFindInObjectTreeGraphResult::GetIcon() const
 		{
 			Brush = FAppStyle::GetBrush(TEXT("GraphEditor.PinIcon"));
 		}
-
-		const UEdGraphSchema* Schema = ResolvedPin->GetSchema();
-		IconColor = Schema->GetPinTypeColor(ResolvedPin->PinType);
 	}
-	else if (GraphNode.IsValid())
+	else if (Object)
 	{
 		Brush = FAppStyle::GetBrush(TEXT("GraphEditor.NodeGlyph"));
 	}
@@ -68,63 +82,54 @@ TSharedRef<SWidget>	FFindInObjectTreeGraphResult::GetIcon() const
 
 FText FFindInObjectTreeGraphResult::GetCategory() const
 {
-	if (GraphNode.IsValid())
+	if (WeakObject.IsValid())
 	{
-		return LOCTEXT("NodeCategory", "Node");
-	}
-	else if (GraphPin.Get())
-	{
-		return LOCTEXT("PinCategory", "Pin");
+		if (PropertyName.IsNone())
+		{
+			return LOCTEXT("NodeCategory", "Node");
+		}
+		else
+		{
+			return LOCTEXT("PinCategory", "Pin");
+		}
 	}
 	return FText::GetEmpty();
 }
 
 FText FFindInObjectTreeGraphResult::GetText() const
 {
-	if (UEdGraphNode* ResolvedNode = GraphNode.Get())
+	UObject* Object = WeakObject.Get();
+	if (Object)
 	{
-		const FText NodeFullTitle = ResolvedNode->GetNodeTitle(ENodeTitleType::FullTitle);
-		const FText NodeListViewTitle = ResolvedNode->GetNodeTitle(ENodeTitleType::ListView);
-		if (NodeFullTitle.EqualToCaseIgnored(NodeListViewTitle))
+		if (PropertyName.IsNone())
 		{
-			return NodeFullTitle;
+			const FText DisplayNameText = Source.GraphConfig->GetDisplayNameText(Object);
+			return DisplayNameText;
+		}
+		else
+		{
+			UClass* ObjectClass = Object->GetClass();
+			FProperty* Property = ObjectClass->FindPropertyByName(PropertyName);
+			return Property->GetDisplayNameText();
 		}
-		return FText::Format(LOCTEXT("NodeResultFmt", "{0} - {1}"), NodeListViewTitle, NodeFullTitle);
-	}
-	else if (UEdGraphPin* ResolvedPin = GraphPin.Get())
-	{
-		const UEdGraphSchema* GraphSchema = ResolvedPin->GetSchema();
-		return GraphSchema->GetPinDisplayName(ResolvedPin);
 	}
 	return CustomText;
 }
 
 FText FFindInObjectTreeGraphResult::GetCommentText() const
 {
-	if (UEdGraphNode* ResolvedNode = GraphNode.Get())
+	if (IObjectTreeGraphObject* ObjectInterface = Cast<IObjectTreeGraphObject>(WeakObject.Get()))
 	{
-		return FText::FromString(ResolvedNode->NodeComment);
+		return FText::FromString(ObjectInterface->GetGraphNodeCommentText(Source.GraphConfig->GraphName));
 	}
 	return FText::GetEmpty();
 }
 
 FReply FFindInObjectTreeGraphResult::OnClick(TSharedRef<SFindInObjectTreeGraph> FindInObjectTreeGraph)
 {
-	if (UEdGraphNode* ResolvedNode = GraphNode.Get())
+	if (UObject* Object = WeakObject.Get())
 	{
-		FindInObjectTreeGraph->OnJumpToNodeRequested.ExecuteIfBound(ResolvedNode);
-		return FReply::Handled();
-	}
-	else if (UEdGraphPin* ResolvedPin = GraphPin.Get())
-	{
-		if (FindInObjectTreeGraph->OnJumpToPinRequested.IsBound())
-		{
-			FindInObjectTreeGraph->OnJumpToPinRequested.Execute(ResolvedPin);
-		}
-		else
-		{
-			FindInObjectTreeGraph->OnJumpToNodeRequested.ExecuteIfBound(ResolvedPin->GetOwningNode());
-		}
+		FindInObjectTreeGraph->OnJumpToObjectRequested.ExecuteIfBound(Object, PropertyName);
 		return FReply::Handled();
 	}
 	return FReply::Unhandled();
@@ -132,10 +137,8 @@ FReply FFindInObjectTreeGraphResult::OnClick(TSharedRef<SFindInObjectTreeGraph>
 
 void SFindInObjectTreeGraph::Construct(const FArguments& InArgs)
 {
-	GraphsToSearch = InArgs._GraphsToSearch;
-
-	OnJumpToNodeRequested = InArgs._OnJumpToNodeRequested;
-	OnJumpToPinRequested = InArgs._OnJumpToPinRequested;
+	OnGetRootObjectsToSearch = InArgs._OnGetRootObjectsToSearch;
+	OnJumpToObjectRequested = InArgs._OnJumpToObjectRequested;
 
 	ChildSlot
 	[
@@ -267,112 +270,84 @@ void SFindInObjectTreeGraph::StartSearch()
 
 	Results.Empty();
 	HighlightText = FText::GetEmpty();
+	TArray<FObjectTreeGraphSearchResult> SearchResults;
 	if (Tokens.Num() > 0)
 	{
 		HighlightText = FText::FromString(SearchQuery);
-		for (UEdGraph* Graph : GraphsToSearch)
+
+		TArray<FFindInObjectTreeGraphSource> Sources;
+		OnGetRootObjectsToSearch.ExecuteIfBound(Sources);
+
+		FObjectTreeGraphSearch Searcher;
+		for (const FFindInObjectTreeGraphSource& Source : Sources)
 		{
-			MatchTokensInGraph(Graph, Tokens);
+			Searcher.AddRootObject(Source.RootObject, Source.GraphConfig);
 		}
-	}
-	
-	if (Results.IsEmpty())
-	{
-		Results.Add(MakeShared<FFindInObjectTreeGraphResult>(LOCTEXT("NoResults", "No results found")));
-	}
 
-	ResultTreeView->RequestTreeRefresh();
-	for (FResultPtr Result : Results)
-	{
-		ResultTreeView->SetItemExpansion(Result, true);
+		Searcher.Search(Tokens, SearchResults);
 	}
-}
-
-void SFindInObjectTreeGraph::MatchTokensInGraph(UEdGraph* Graph, TArrayView<FString> Tokens)
-{
-	FResultPtr GraphResult(new FFindInObjectTreeGraphResult(FText::FromName(Graph->GetFName())));
-
-	const UEdGraphSchema* GraphSchema = Graph->GetSchema();
 
-	for (UEdGraphNode* Node : Graph->Nodes)
+	// Convert simple flat search results to hierarchical results, where property results
+	// are always under an object result, object results are always under a graph result, 
+	// and so on. We do this by creating blank results if we need to, although it makes
+	// the assumption that the flat results come ordered (e.g. an object result won't be
+	// found after a property result for that object).
 	{
-		const bool bNodeMatches = GraphNodeMatchesSearchTokens(GraphSchema, Node, Tokens);
-
-		TArray<UEdGraphPin*> MatchingPins;
-		for (UEdGraphPin* Pin : Node->Pins)
+		TMap<UObject*, FResultPtr> RootObjectToWidgetResult;
+		TMap<UObject*, FResultPtr> ObjectToWidgetResult;
+		for (const FObjectTreeGraphSearchResult& SearchResult : SearchResults)
 		{
-			if (GraphPinMatchesSearchTokens(GraphSchema, Pin, Tokens))
+			FFindInObjectTreeGraphSource CurSource{ SearchResult.RootObject, SearchResult.GraphConfig };
+
+			FResultPtr GraphResult;
+			if (SearchResult.RootObject)
 			{
-				MatchingPins.Add(Pin);
+				GraphResult = RootObjectToWidgetResult.FindRef(SearchResult.RootObject);
+				if (!GraphResult)
+				{
+					const FText RootObjectDisplayText = SearchResult.GraphConfig->GetDisplayNameText(SearchResult.RootObject);
+					const FText& GraphDisplayName = SearchResult.GraphConfig->GraphDisplayInfo.DisplayName;
+					const FText GraphResultText = FText::Format(
+							LOCTEXT("GraphResultFmt", "{0}: {1}"), { RootObjectDisplayText, GraphDisplayName });
+					GraphResult = MakeShared<FFindInObjectTreeGraphResult>(GraphResultText);
+					RootObjectToWidgetResult.Add(SearchResult.RootObject, GraphResult);
+					Results.Add(GraphResult);
+				}
 			}
-		}
 
-		if (bNodeMatches || MatchingPins.Num() > 0)
-		{
-			FResultPtr NodeResult(new FFindInObjectTreeGraphResult(GraphResult, Node));
-			Results.Add(NodeResult);
+			FResultPtr ObjectResult;
+			if (SearchResult.Object)
+			{
+				ObjectResult = ObjectToWidgetResult.FindRef(SearchResult.Object);
+				if (!ObjectResult)
+				{
+					ensure(GraphResult);
+					ObjectResult = MakeShared<FFindInObjectTreeGraphResult>(GraphResult, CurSource, SearchResult.Object);
+					ObjectToWidgetResult.Add(SearchResult.Object, ObjectResult);
+					GraphResult->Children.Add(ObjectResult);
+				}
+			}
 
-			for (UEdGraphPin* Pin : MatchingPins)
+			if (!SearchResult.PropertyName.IsNone())
 			{
-				FResultPtr PinResult(new FFindInObjectTreeGraphResult(NodeResult, Pin));
+				ensure(ObjectResult);
+				FResultPtr PropertyResult = MakeShared<FFindInObjectTreeGraphResult>(
+						ObjectResult, CurSource, SearchResult.Object, SearchResult.PropertyName);
+				ObjectResult->Children.Add(PropertyResult);
 			}
 		}
 	}
-
-	for (UEdGraph* SubGraph : Graph->SubGraphs)
-	{
-		MatchTokensInGraph(SubGraph, Tokens);
-	}
-}
-
-bool SFindInObjectTreeGraph::GraphNodeMatchesSearchTokens(const UEdGraphSchema* GraphSchema, UEdGraphNode* GraphNode, TArrayView<FString> Tokens)
-{
-	const FString NodeFullTitle = GraphNode->GetNodeTitle(ENodeTitleType::FullTitle).ToString();
-	if (StringMatchesSearchTokens(NodeFullTitle, Tokens))
-	{
-		return true;
-	}
-
-	const FString NodeListViewTitle = GraphNode->GetNodeTitle(ENodeTitleType::ListView).ToString();
-	if (StringMatchesSearchTokens(NodeListViewTitle, Tokens))
-	{
-		return true;
-	}
-
-	if (StringMatchesSearchTokens(GraphNode->NodeComment, Tokens))
-	{
-		return true;
-	}
-
-	return false;
-}
-
-bool SFindInObjectTreeGraph::GraphPinMatchesSearchTokens(const UEdGraphSchema* GraphSchema, UEdGraphPin* GraphPin, TArrayView<FString> Tokens)
-{
-	const FString PinDisplayName = GraphSchema->GetPinDisplayName(GraphPin).ToString();
-	if (StringMatchesSearchTokens(PinDisplayName, Tokens))
-	{
-		return true;
-	}
-
-	return false;
-}
-
-bool SFindInObjectTreeGraph::StringMatchesSearchTokens(const FString& ComparisonString, TArrayView<FString> Tokens)
-{
-	if (ComparisonString.IsEmpty())
+	
+	if (Results.IsEmpty())
 	{
-		return false;
+		Results.Add(MakeShared<FFindInObjectTreeGraphResult>(LOCTEXT("NoResults", "No results found")));
 	}
 
-	for (const FString& Token : Tokens)
+	ResultTreeView->RequestTreeRefresh();
+	for (FResultPtr Result : Results)
 	{
-		if (!ComparisonString.Contains(Token))
-		{
-			return false;
-		}
+		ResultTreeView->SetItemExpansion(Result, true);
 	}
-	return true;
 }
 
 #undef LOCTEXT_NAMESPACE
diff --git a/Engine/Plugins/Cameras/GameplayCameras/Source/GameplayCamerasEditor/Private/Toolkits/CameraAssetEditorToolkit.cpp b/Engine/Plugins/Cameras/GameplayCameras/Source/GameplayCamerasEditor/Private/Toolkits/CameraAssetEditorToolkit.cpp
--- a/Engine/Plugins/Cameras/GameplayCameras/Source/GameplayCamerasEditor/Private/Toolkits/CameraAssetEditorToolkit.cpp
+++ b/Engine/Plugins/Cameras/GameplayCameras/Source/GameplayCamerasEditor/Private/Toolkits/CameraAssetEditorToolkit.cpp
@@ -6,6 +6,7 @@
 #include "Commands/CameraAssetEditorCommands.h"
 #include "Core/CameraAsset.h"
 #include "Core/CameraBuildLog.h"
+#include "Editors/ObjectTreeGraphConfig.h"
 #include "Editors/SFindInObjectTreeGraph.h"
 #include "Framework/Docking/LayoutExtender.h"
 #include "Framework/Docking/TabManager.h"
@@ -137,35 +138,13 @@ void FCameraAssetEditorToolkit::CreateWidgets()
 
 	// Create the search panel.
 	SearchWidget = SNew(SFindInObjectTreeGraph)
-		.OnJumpToNodeRequested(this, &FCameraAssetEditorToolkit::JumpToNode);
+		.OnGetRootObjectsToSearch(this, &FCameraAssetEditorToolkit::OnGetRootObjectsToSearch)
+		.OnJumpToObjectRequested(this, &FCameraAssetEditorToolkit::OnJumpToObject);
 
 	// Create the message log.
 	BuildLogToolkit->Initialize("CameraAssetBuildMessages");
 }
 
-void FCameraAssetEditorToolkit::JumpToNode(UEdGraphNode* Node)
-{
-}
-
-void FCameraAssetEditorToolkit::JumpToObject(UObject* Object)
-{
-	TSharedPtr<FCameraRigsAssetEditorMode> CameraRigsMode = StaticCastSharedPtr<FCameraRigsAssetEditorMode>(
-			GetEditorMode(FCameraRigsAssetEditorMode::ModeName));
-	if (CameraRigsMode->JumpToObject(Object))
-	{
-		SetEditorMode(FCameraRigsAssetEditorMode::ModeName);
-		return;
-	}
-
-	TSharedPtr<FCameraSharedTransitionsAssetEditorMode> SharedTransitionsMode = StaticCastSharedPtr<FCameraSharedTransitionsAssetEditorMode>(
-			GetEditorMode(FCameraSharedTransitionsAssetEditorMode::ModeName));
-	if (SharedTransitionsMode->JumpToObject(Object))
-	{
-		SetEditorMode(FCameraSharedTransitionsAssetEditorMode::ModeName);
-		return;
-	}
-}
-
 void FCameraAssetEditorToolkit::RegisterToolbar()
 {
 	FName ParentName;
@@ -248,7 +227,7 @@ void FCameraAssetEditorToolkit::PostInitAssetEditor()
 		Commands.FindInCamera,
 		FExecuteAction::CreateSP(this, &FCameraAssetEditorToolkit::OnFindInCamera));
 
-	BuildLogToolkit->OnRequestJumpToObject().BindSP(this, &FCameraAssetEditorToolkit::JumpToObject);
+	BuildLogToolkit->OnRequestJumpToObject().BindSP(this, &FCameraAssetEditorToolkit::OnJumpToObject);
 
 	IGameplayCamerasModule& GameplayCamerasModule = FModuleManager::GetModuleChecked<IGameplayCamerasModule>("GameplayCameras");
 	LiveEditManager = GameplayCamerasModule.GetLiveEditManager();
@@ -298,6 +277,63 @@ void FCameraAssetEditorToolkit::OnFindInCamera()
 	SearchWidget->FocusSearchEditBox();
 }
 
+void FCameraAssetEditorToolkit::OnGetRootObjectsToSearch(TArray<FFindInObjectTreeGraphSource>& OutSources)
+{
+	TSharedPtr<FCameraRigsAssetEditorMode> CameraRigsMode = GetTypedEditorMode<FCameraRigsAssetEditorMode>(
+			FCameraRigsAssetEditorMode::ModeName);
+	CameraRigsMode->OnGetRootObjectsToSearch(OutSources);
+
+	TSharedPtr<FCameraSharedTransitionsAssetEditorMode> SharedTransitionsMode = GetTypedEditorMode<FCameraSharedTransitionsAssetEditorMode>(
+			FCameraSharedTransitionsAssetEditorMode::ModeName);
+	SharedTransitionsMode->OnGetRootObjectsToSearch(OutSources);
+}
+
+void FCameraAssetEditorToolkit::OnJumpToObject(UObject* Object)
+{
+	OnJumpToObject(Object, NAME_None);
+}
+
+void FCameraAssetEditorToolkit::OnJumpToObject(UObject* Object, FName PropertyName)
+{
+	bool bFindInCameraRig = false;
+	bool bFindInSharedTranstions = false;
+	UObject* CurOuter = Object;
+	while (CurOuter != nullptr)
+	{
+		if (CurOuter->IsA<UCameraRigAsset>())
+		{
+			bFindInCameraRig = true;
+			break;
+		}
+		if (CurOuter == CameraAsset)
+		{
+			bFindInSharedTranstions = true;
+			break;
+		}
+
+		CurOuter = CurOuter->GetOuter();
+	}
+	
+	if (bFindInCameraRig)
+	{
+		TSharedPtr<FCameraRigsAssetEditorMode> CameraRigsMode = GetTypedEditorMode<FCameraRigsAssetEditorMode>(
+				FCameraRigsAssetEditorMode::ModeName);
+		SetEditorMode(FCameraRigsAssetEditorMode::ModeName);
+		CameraRigsMode->JumpToObject(Object, PropertyName);
+		return;
+	}
+
+	if (bFindInSharedTranstions)
+	{
+		TSharedPtr<FCameraSharedTransitionsAssetEditorMode> SharedTransitionsMode = GetTypedEditorMode<FCameraSharedTransitionsAssetEditorMode>(
+				FCameraSharedTransitionsAssetEditorMode::ModeName);
+		SetEditorMode(FCameraSharedTransitionsAssetEditorMode::ModeName);
+		SharedTransitionsMode->JumpToObject(Object, PropertyName);
+		return;
+	}
+}
+
+
 FText FCameraAssetEditorToolkit::GetBaseToolkitName() const
 {
 	return LOCTEXT("AppLabel", "Camera Asset");
diff --git a/Engine/Plugins/Cameras/GameplayCameras/Source/GameplayCamerasEditor/Private/Toolkits/CameraRigAssetEditorToolkit.cpp b/Engine/Plugins/Cameras/GameplayCameras/Source/GameplayCamerasEditor/Private/Toolkits/CameraRigAssetEditorToolkit.cpp
--- a/Engine/Plugins/Cameras/GameplayCameras/Source/GameplayCamerasEditor/Private/Toolkits/CameraRigAssetEditorToolkit.cpp
+++ b/Engine/Plugins/Cameras/GameplayCameras/Source/GameplayCamerasEditor/Private/Toolkits/CameraRigAssetEditorToolkit.cpp
@@ -5,6 +5,8 @@
 #include "Commands/CameraRigAssetEditorCommands.h"
 #include "Core/CameraBuildLog.h"
 #include "Core/CameraRigAsset.h"
+#include "Editors/CameraNodeGraphSchema.h"
+#include "Editors/CameraRigTransitionGraphSchema.h"
 #include "Editors/SCameraRigAssetEditor.h"
 #include "Editors/SFindInObjectTreeGraph.h"
 #include "Framework/Docking/LayoutExtender.h"
@@ -44,6 +46,14 @@ FCameraRigAssetEditorToolkit::FCameraRigAssetEditorToolkit(UAssetEditor* InOwnin
 		StandardLayout->AddBottomTab(MessagesTabId);
 	}
 	StandaloneDefaultLayout = StandardLayout->GetLayout();
+
+	UClass* NodeGraphSchemaClass = UCameraNodeGraphSchema::StaticClass();
+	UCameraNodeGraphSchema* DefaultNodeGraphSchema = Cast<UCameraNodeGraphSchema>(NodeGraphSchemaClass->GetDefaultObject());
+	NodeGraphConfig = DefaultNodeGraphSchema->BuildGraphConfig();
+
+	UClass* TransitionSchemaClass = UCameraRigTransitionGraphSchema::StaticClass();
+	UCameraRigTransitionGraphSchema* DefaultTransitionGraphSchema = Cast<UCameraRigTransitionGraphSchema>(TransitionSchemaClass->GetDefaultObject());
+	TransitionGraphConfig = DefaultTransitionGraphSchema->BuildGraphConfig();
 }
 
 void FCameraRigAssetEditorToolkit::SetCameraRigAsset(UCameraRigAsset* InCameraRig)
@@ -128,10 +138,8 @@ void FCameraRigAssetEditorToolkit::CreateWidgets()
 	TArray<UEdGraph*> CameraRigGraphs;
 	CameraRigEditorWidget->GetGraphs(CameraRigGraphs);
 	SearchWidget = SNew(SFindInObjectTreeGraph)
-		.GraphsToSearch(CameraRigGraphs)
-		.OnJumpToNodeRequested(
-				SFindInObjectTreeGraph::FOnJumpToNodeRequested::CreateSP(
-					CameraRigEditorWidget.ToSharedRef(), &SCameraRigAssetEditor::JumpToNode));
+		.OnGetRootObjectsToSearch(this, &FCameraRigAssetEditorToolkit::OnGetRootObjectsToSearch)
+		.OnJumpToObjectRequested(this, &FCameraRigAssetEditorToolkit::OnJumpToObject);
 
 	// Create the message log.
 	BuildLogToolkit->Initialize("CameraRigAssetBuildMessages");
@@ -237,6 +245,19 @@ void FCameraRigAssetEditorToolkit::OnFindInCameraRig()
 	SearchWidget->FocusSearchEditBox();
 }
 
+void FCameraRigAssetEditorToolkit::OnGetRootObjectsToSearch(TArray<FFindInObjectTreeGraphSource>& OutSources)
+{
+	UCameraRigAsset* CameraRig = Impl->GetCameraRigAsset();
+	OutSources.Add(FFindInObjectTreeGraphSource{ CameraRig, &NodeGraphConfig });
+	OutSources.Add(FFindInObjectTreeGraphSource{ CameraRig, &TransitionGraphConfig });
+}
+
+void FCameraRigAssetEditorToolkit::OnJumpToObject(UObject* Object, FName PropertyName)
+{
+	TSharedPtr<SCameraRigAssetEditor> CameraRigEditor = Impl->GetCameraRigAssetEditor();
+	CameraRigEditor->FindAndJumpToObjectNode(Object);
+}
+
 FText FCameraRigAssetEditorToolkit::GetBaseToolkitName() const
 {
 	return LOCTEXT("AppLabel", "Camera Rig Asset");
diff --git a/Engine/Plugins/Cameras/GameplayCameras/Source/GameplayCamerasEditor/Private/Toolkits/CameraRigsAssetEditorMode.cpp b/Engine/Plugins/Cameras/GameplayCameras/Source/GameplayCamerasEditor/Private/Toolkits/CameraRigsAssetEditorMode.cpp
--- a/Engine/Plugins/Cameras/GameplayCameras/Source/GameplayCamerasEditor/Private/Toolkits/CameraRigsAssetEditorMode.cpp
+++ b/Engine/Plugins/Cameras/GameplayCameras/Source/GameplayCamerasEditor/Private/Toolkits/CameraRigsAssetEditorMode.cpp
@@ -2,7 +2,12 @@
 
 #include "Toolkits/CameraRigsAssetEditorMode.h"
 
+#include "Core/CameraAsset.h"
+#include "Core/CameraRigAsset.h"
+#include "Editors/CameraNodeGraphSchema.h"
+#include "Editors/CameraRigTransitionGraphSchema.h"
 #include "Editors/SCameraRigAssetEditor.h"
+#include "Editors/SFindInObjectTreeGraph.h"
 #include "Framework/Docking/TabManager.h"
 #include "Styles/GameplayCamerasEditorStyle.h"
 #include "ToolMenus.h"
@@ -30,6 +35,14 @@ FCameraRigsAssetEditorMode::FCameraRigsAssetEditorMode(UCameraAsset* InCameraAss
 		StandardLayout->AddLeftTab(CameraRigsTabId);
 	}
 	DefaultLayout = StandardLayout->GetLayout();
+
+	UClass* NodeGraphSchemaClass = UCameraNodeGraphSchema::StaticClass();
+	UCameraNodeGraphSchema* DefaultNodeGraphSchema = Cast<UCameraNodeGraphSchema>(NodeGraphSchemaClass->GetDefaultObject());
+	NodeGraphConfig = DefaultNodeGraphSchema->BuildGraphConfig();
+
+	UClass* TransitionSchemaClass = UCameraRigTransitionGraphSchema::StaticClass();
+	UCameraRigTransitionGraphSchema* DefaultTransitionGraphSchema = Cast<UCameraRigTransitionGraphSchema>(TransitionSchemaClass->GetDefaultObject());
+	TransitionGraphConfig = DefaultTransitionGraphSchema->BuildGraphConfig();
 }
 
 void FCameraRigsAssetEditorMode::OnActivateMode(const FAssetEditorModeActivateParams& InParams)
@@ -91,7 +104,16 @@ void FCameraRigsAssetEditorMode::OnCameraRigEditRequested(UCameraRigAsset* InCam
 	Impl->SetCameraRigAsset(InCameraRig);
 }
 
-bool FCameraRigsAssetEditorMode::JumpToObject(UObject* InObject)
+void FCameraRigsAssetEditorMode::OnGetRootObjectsToSearch(TArray<FFindInObjectTreeGraphSource>& OutSources)
+{
+	for (UCameraRigAsset* CameraRig : CameraAsset->CameraRigs)
+	{
+		OutSources.Add(FFindInObjectTreeGraphSource{ CameraRig, &NodeGraphConfig });
+		OutSources.Add(FFindInObjectTreeGraphSource{ CameraRig, &TransitionGraphConfig });
+	}
+}
+
+bool FCameraRigsAssetEditorMode::JumpToObject(UObject* InObject, FName PropertyName)
 {
 	UCameraRigAsset* FindInCameraRig = nullptr;
 	UObject* CurOuter = InObject;
@@ -111,8 +133,11 @@ bool FCameraRigsAssetEditorMode::JumpToObject(UObject* InObject)
 
 	Impl->SetCameraRigAsset(FindInCameraRig);
 
-	TSharedPtr<SCameraRigAssetEditor> CameraRigAssetEditor = Impl->GetCameraRigAssetEditor();
-	return CameraRigAssetEditor->FindAndJumpToObjectNode(InObject);
+	if (TSharedPtr<SCameraRigAssetEditor> CameraRigAssetEditor = Impl->GetCameraRigAssetEditor())
+	{
+		return CameraRigAssetEditor->FindAndJumpToObjectNode(InObject);
+	}
+	return false;
 }
 
 }  // namespace UE::Cameras
diff --git a/Engine/Plugins/Cameras/GameplayCameras/Source/GameplayCamerasEditor/Private/Toolkits/CameraSharedTransitionsAssetEditorMode.cpp b/Engine/Plugins/Cameras/GameplayCameras/Source/GameplayCamerasEditor/Private/Toolkits/CameraSharedTransitionsAssetEditorMode.cpp
--- a/Engine/Plugins/Cameras/GameplayCameras/Source/GameplayCamerasEditor/Private/Toolkits/CameraSharedTransitionsAssetEditorMode.cpp
+++ b/Engine/Plugins/Cameras/GameplayCameras/Source/GameplayCamerasEditor/Private/Toolkits/CameraSharedTransitionsAssetEditorMode.cpp
@@ -5,6 +5,7 @@
 #include "Core/CameraAsset.h"
 #include "Editors/CameraSharedTransitionGraphSchema.h"
 #include "Editors/SCameraRigTransitionEditor.h"
+#include "Editors/SFindInObjectTreeGraph.h"
 #include "ToolMenus.h"
 #include "Toolkits/CameraRigTransitionEditorToolkitBase.h"
 #include "Toolkits/StandardToolkitLayout.h"
@@ -51,6 +52,10 @@ FCameraSharedTransitionsAssetEditorMode::FCameraSharedTransitionsAssetEditorMode
 	Impl->SetTransitionOwner(CameraAsset);
 
 	DefaultLayout = Impl->GetStandardLayout()->GetLayout();
+
+	UClass* SharedTransitionSchemaClass = UCameraSharedTransitionGraphSchema::StaticClass();
+	UCameraSharedTransitionGraphSchema* DefaultTransitionGraphSchema = Cast<UCameraSharedTransitionGraphSchema>(SharedTransitionSchemaClass->GetDefaultObject());
+	TransitionGraphConfig = DefaultTransitionGraphSchema->BuildGraphConfig();
 }
 
 void FCameraSharedTransitionsAssetEditorMode::OnActivateMode(const FAssetEditorModeActivateParams& InParams)
@@ -75,15 +80,18 @@ void FCameraSharedTransitionsAssetEditorMode::OnDeactivateMode(const FAssetEdito
 	UToolMenus::UnregisterOwner(this);
 }
 
-bool FCameraSharedTransitionsAssetEditorMode::JumpToNode(UEdGraphNode* InNode)
+void FCameraSharedTransitionsAssetEditorMode::OnGetRootObjectsToSearch(TArray<FFindInObjectTreeGraphSource>& OutSources)
 {
-	return false;
+	OutSources.Add(FFindInObjectTreeGraphSource{ CameraAsset, &TransitionGraphConfig });
 }
 
-bool FCameraSharedTransitionsAssetEditorMode::JumpToObject(UObject* InObject)
+bool FCameraSharedTransitionsAssetEditorMode::JumpToObject(UObject* InObject, FName PropertyName)
 {
-	TSharedPtr<SCameraRigTransitionEditor> TransitionEditor = Impl->GetCameraRigTransitionEditor();
-	return TransitionEditor->FindAndJumpToObjectNode(InObject);
+	if (TSharedPtr<SCameraRigTransitionEditor> TransitionEditor = Impl->GetCameraRigTransitionEditor())
+	{
+		return TransitionEditor->FindAndJumpToObjectNode(InObject);
+	}
+	return false;
 }
 
 }  // namespace UE::Cameras
diff --git a/Engine/Plugins/Cameras/GameplayCameras/Source/GameplayCamerasEditor/Public/Editors/ObjectTreeGraphSearch.h b/Engine/Plugins/Cameras/GameplayCameras/Source/GameplayCamerasEditor/Public/Editors/ObjectTreeGraphSearch.h
new file mode 100644
--- /dev/null
+++ b/Engine/Plugins/Cameras/GameplayCameras/Source/GameplayCamerasEditor/Public/Editors/ObjectTreeGraphSearch.h
@@ -0,0 +1,82 @@
+// Copyright Epic Games, Inc. All Rights Reserved.
+
+#pragma once
+
+#include "Containers/UnrealString.h"
+#include "UObject/WeakObjectPtr.h"
+
+class UObject;
+struct FObjectTreeGraphConfig;
+
+/** Search result for an object tree graph search. */
+struct FObjectTreeGraphSearchResult
+{
+	/** The root object that the current search result was found in. */
+	UObject* RootObject = nullptr;
+
+	/** The graph config for the object hierarchy this result was found in. */
+	const FObjectTreeGraphConfig* GraphConfig = nullptr;
+
+	/** The object found to match the search. */
+	UObject* Object = nullptr;
+
+	/** 
+	 * The specific object property that matched the search. 
+	 * None if the object itself matches.
+	 */
+	FName PropertyName;
+};
+
+
+/**
+ * A utility class that can search a series of string tokens across an
+ * object tree graph.
+ */
+class FObjectTreeGraphSearch
+{
+public:
+
+	FObjectTreeGraphSearch();
+
+public:
+
+	/** Adds a root object to search through. */
+	void AddRootObject(UObject* InObject, const FObjectTreeGraphConfig* InGraphConfig);
+
+	/** Searchs for the given string tokens. */
+	void Search(TArrayView<FString> InTokens, TArray<FObjectTreeGraphSearchResult>& OutResults) const;
+
+private:
+
+	using FSearchResult = FObjectTreeGraphSearchResult;
+
+	struct FSearchState
+	{
+		const FObjectTreeGraphConfig* GraphConfig = nullptr;
+		TArrayView<FString> Tokens;
+
+		UObject* RootObject = nullptr;
+		TArray<UObject*> ObjectStack;
+		TSet<UObject*> VisitedObjects;
+
+		TArray<FSearchResult> Results;
+	};
+
+	struct FRootObjectInfo
+	{
+		TWeakObjectPtr<> WeakRootObject;
+		const FObjectTreeGraphConfig* GraphConfig = nullptr;
+	};
+
+	void SearchRootObject(const FRootObjectInfo& InRootObjectInfo, TArrayView<FString> InTokens, TArray<FSearchResult>& OutResults) const;
+	void SearchObject(UObject* InObject, FSearchState& InOutState) const;
+
+	bool MatchObject(UObject* InObject, const FSearchState& InState) const;
+	bool MatchObjectProperty(UObject* InObject, FProperty* InProperty, const FSearchState& InState) const;
+	bool MatchString(const FString& InString, const FSearchState& InState) const;
+
+private:
+
+	TArray<FRootObjectInfo> RootObjectInfos;
+};
+
diff --git a/Engine/Plugins/Cameras/GameplayCameras/Source/GameplayCamerasEditor/Public/Editors/SCameraRigAssetEditor.h b/Engine/Plugins/Cameras/GameplayCameras/Source/GameplayCamerasEditor/Public/Editors/SCameraRigAssetEditor.h
--- a/Engine/Plugins/Cameras/GameplayCameras/Source/GameplayCamerasEditor/Public/Editors/SCameraRigAssetEditor.h
+++ b/Engine/Plugins/Cameras/GameplayCameras/Source/GameplayCamerasEditor/Public/Editors/SCameraRigAssetEditor.h
@@ -73,8 +73,6 @@ public:
 
 	/** Focuses the current graph to the root object node. */
 	void FocusHome();
-	/** Jumps the current graph to the given node. */
-	void JumpToNode(UEdGraphNode* InGraphNode);
 	/** Finds a node for the given object and, if so, jumps to it. */
 	bool FindAndJumpToObjectNode(UObject* InObject);
 
diff --git a/Engine/Plugins/Cameras/GameplayCameras/Source/GameplayCamerasEditor/Public/Editors/SFindInObjectTreeGraph.h b/Engine/Plugins/Cameras/GameplayCameras/Source/GameplayCamerasEditor/Public/Editors/SFindInObjectTreeGraph.h
--- a/Engine/Plugins/Cameras/GameplayCameras/Source/GameplayCamerasEditor/Public/Editors/SFindInObjectTreeGraph.h
+++ b/Engine/Plugins/Cameras/GameplayCameras/Source/GameplayCamerasEditor/Public/Editors/SFindInObjectTreeGraph.h
@@ -11,12 +11,23 @@
 #include "Widgets/Views/STableViewBase.h"
 #include "Widgets/Views/STreeView.h"
 
+class SFindInObjectTreeGraph;
+class SSearchBox;
 class UEdGraph;
 class UEdGraphNode;
 class UEdGraphSchema;
-class SSearchBox;
+struct FObjectTreeGraphConfig;
 
-class SFindInObjectTreeGraph;
+/**
+ * Struct describing one source of possible results for a search.
+ */
+struct FFindInObjectTreeGraphSource
+{
+	/** The root object a result was found in. */
+	UObject* RootObject = nullptr;
+	/** The config for the object tree graph a result was found in. */
+	const FObjectTreeGraphConfig* GraphConfig = nullptr;
+};
 
 /** 
  * Structure for a search result inside an object tree graph.
@@ -32,19 +43,26 @@ public:
 
 	/** Custom text for this result. */
 	FText CustomText;
-	/** The graph node that this result refers to. */
-	TWeakObjectPtr<UEdGraphNode> GraphNode;
-	/** The graph pin that this result refers to. */
-	FEdGraphPinReference GraphPin;
+	/** The object that this result refers to. */
+	TWeakObjectPtr<> WeakObject;
+	/** The property name that this result refers to. */
+	FName PropertyName;
 
 public:
 
 	/** Creates a new result with a custom text. */
 	FFindInObjectTreeGraphResult(const FText& InCustomText);
-	/** Creates a new result referring to a graph node, under a parent result. */
-	FFindInObjectTreeGraphResult(TSharedPtr<FFindInObjectTreeGraphResult>& InParent, UEdGraphNode* InGraphNode);
-	/** Creates a new result referring to a graph pin, under a parent result. */
-	FFindInObjectTreeGraphResult(TSharedPtr<FFindInObjectTreeGraphResult>& InParent, UEdGraphPin* InGraphPin);
+	/** Creates a new result referring to an object, under a parent result. */
+	FFindInObjectTreeGraphResult(
+			TSharedPtr<FFindInObjectTreeGraphResult>& InParent, 
+			const FFindInObjectTreeGraphSource& InSource, 
+			UObject* InObject);
+	/** Creates a new result referring to an object's property, under a parent result. */
+	FFindInObjectTreeGraphResult(
+			TSharedPtr<FFindInObjectTreeGraphResult>& InParent, 
+			const FFindInObjectTreeGraphSource& InSource, 
+			UObject* InObject, 
+			FName InPropertyName);
 
 	/** Gets the icon for this result. */
 	TSharedRef<SWidget>	GetIcon() const;
@@ -57,6 +75,10 @@ public:
 
 	/** Go to the graph node, pin, etc. */
 	FReply OnClick(TSharedRef<SFindInObjectTreeGraph> FindInObjectTreeGraph);
+
+private:
+
+	FFindInObjectTreeGraphSource Source;
 };
 
 /**
@@ -66,17 +88,15 @@ class SFindInObjectTreeGraph : public SCompoundWidget
 {
 public:
 
-	DECLARE_DELEGATE_OneParam(FOnJumpToNodeRequested, UEdGraphNode*);
-	DECLARE_DELEGATE_OneParam(FOnJumpToPinRequested, UEdGraphPin*);
+	DECLARE_DELEGATE_OneParam(FOnGetRootObjectsToSearch, TArray<FFindInObjectTreeGraphSource>&);
+	DECLARE_DELEGATE_TwoParams(FOnJumpToObjectRequested, UObject*, FName);
 
 	SLATE_BEGIN_ARGS(SFindInObjectTreeGraph)
 	{}
-		/** The graphs to search. */
-		SLATE_ARGUMENT(TArray<UEdGraph*>, GraphsToSearch)
-		/** The callback to invoke when a search result wants to focus a node. */
-		SLATE_EVENT(FOnJumpToNodeRequested, OnJumpToNodeRequested)
-		/** The callback to invoke when a search result wants to focus a pin. */
-		SLATE_EVENT(FOnJumpToPinRequested, OnJumpToPinRequested)
+		/** The callback to get the graphs to search. */
+		SLATE_EVENT(FOnGetRootObjectsToSearch, OnGetRootObjectsToSearch)
+		/** The callback to invoke when a search result wants to focus an object node or one of its pins. */
+		SLATE_EVENT(FOnJumpToObjectRequested, OnJumpToObjectRequested)
 	SLATE_END_ARGS()
 
 	void Construct(const FArguments& InArgs);
@@ -99,13 +119,12 @@ protected:
 protected:
 
 	void StartSearch();
-	void MatchTokensInGraph(UEdGraph* Graph, TArrayView<FString> Tokens);
-	bool GraphNodeMatchesSearchTokens(const UEdGraphSchema* GraphSchema, UEdGraphNode* GraphNode, TArrayView<FString> Tokens);
-	bool GraphPinMatchesSearchTokens(const UEdGraphSchema* GraphSchema, UEdGraphPin* GraphPin, TArrayView<FString> Tokens);
-	bool StringMatchesSearchTokens(const FString& ComparisonString, TArrayView<FString> Tokens);
 
 protected:
 
+	FOnGetRootObjectsToSearch OnGetRootObjectsToSearch;
+	FOnJumpToObjectRequested OnJumpToObjectRequested;
+
 	TSharedPtr<SSearchBox> SearchBox;
 	TSharedPtr<SResultTreeView> ResultTreeView;
 
@@ -114,11 +133,6 @@ protected:
 
 	FText HighlightText;
 
-	TArray<UEdGraph*> GraphsToSearch;
-
-	FOnJumpToNodeRequested OnJumpToNodeRequested;
-	FOnJumpToPinRequested OnJumpToPinRequested;
-
 	friend struct FFindInObjectTreeGraphResult;
 };
 
diff --git a/Engine/Plugins/Cameras/GameplayCameras/Source/GameplayCamerasEditor/Public/Toolkits/AssetEditorModeManagerToolkit.h b/Engine/Plugins/Cameras/GameplayCameras/Source/GameplayCamerasEditor/Public/Toolkits/AssetEditorModeManagerToolkit.h
--- a/Engine/Plugins/Cameras/GameplayCameras/Source/GameplayCamerasEditor/Public/Toolkits/AssetEditorModeManagerToolkit.h
+++ b/Engine/Plugins/Cameras/GameplayCameras/Source/GameplayCamerasEditor/Public/Toolkits/AssetEditorModeManagerToolkit.h
@@ -47,6 +47,12 @@ protected:
 	bool IsEditorMode(FName InModeName) const;
 	FName GetCurrentEditorModeName() const;
 
+	template<typename EditorModeClass>
+	TSharedPtr<EditorModeClass> GetTypedEditorMode(FName InModeName) const
+	{
+		return StaticCastSharedPtr<EditorModeClass>(GetEditorMode(InModeName));
+	}
+
 private:
 
 	TMap<FName, TSharedPtr<FAssetEditorMode>> EditorModes;
diff --git a/Engine/Plugins/Cameras/GameplayCameras/Source/GameplayCamerasEditor/Public/Toolkits/CameraAssetEditorToolkit.h b/Engine/Plugins/Cameras/GameplayCameras/Source/GameplayCamerasEditor/Public/Toolkits/CameraAssetEditorToolkit.h
--- a/Engine/Plugins/Cameras/GameplayCameras/Source/GameplayCamerasEditor/Public/Toolkits/CameraAssetEditorToolkit.h
+++ b/Engine/Plugins/Cameras/GameplayCameras/Source/GameplayCamerasEditor/Public/Toolkits/CameraAssetEditorToolkit.h
@@ -15,6 +15,7 @@ class UCameraAssetEditor;
 class UEdGraph;
 class UEdGraphNode;
 class UGameplayCamerasEditorSettings;
+struct FFindInObjectTreeGraphSource;
 
 namespace UE::Cameras
 {
@@ -64,12 +65,13 @@ private:
 	TSharedRef<SDockTab> SpawnTab_Search(const FSpawnTabArgs& Args);
 	TSharedRef<SDockTab> SpawnTab_Messages(const FSpawnTabArgs& Args);
 
-	void JumpToNode(UEdGraphNode* Node);
-	void JumpToObject(UObject* Object);
-
 	void OnBuild();
 	void OnFindInCamera();
 
+	void OnGetRootObjectsToSearch(TArray<FFindInObjectTreeGraphSource>& OutSources);
+	void OnJumpToObject(UObject* Object);
+	void OnJumpToObject(UObject* Object, FName PropertyName);
+
 private:
 
 	static const FName SearchTabId;
diff --git a/Engine/Plugins/Cameras/GameplayCameras/Source/GameplayCamerasEditor/Public/Toolkits/CameraRigAssetEditorToolkit.h b/Engine/Plugins/Cameras/GameplayCameras/Source/GameplayCamerasEditor/Public/Toolkits/CameraRigAssetEditorToolkit.h
--- a/Engine/Plugins/Cameras/GameplayCameras/Source/GameplayCamerasEditor/Public/Toolkits/CameraRigAssetEditorToolkit.h
+++ b/Engine/Plugins/Cameras/GameplayCameras/Source/GameplayCamerasEditor/Public/Toolkits/CameraRigAssetEditorToolkit.h
@@ -3,12 +3,14 @@
 #pragma once
 
 #include "CoreTypes.h"
+#include "Editors/ObjectTreeGraphConfig.h"
 #include "Tools/BaseAssetToolkit.h"
 
 #include "CameraRigAssetEditorToolkit.generated.h"
 
 class SFindInObjectTreeGraph;
 class UCameraRigAsset;
+struct FFindInObjectTreeGraphSource;
 
 namespace UE::Cameras
 {
@@ -51,14 +53,25 @@ private:
 	void OnBuild();
 	void OnFindInCameraRig();
 
+	void OnGetRootObjectsToSearch(TArray<FFindInObjectTreeGraphSource>& OutSources);
+	void OnJumpToObject(UObject* Object, FName PropertyName);
+
 private:
 
 	static const FName SearchTabId;
 	static const FName MessagesTabId;
 
+	/** Base implementation */
 	TSharedPtr<FCameraRigAssetEditorToolkitBase> Impl;
 
+	/** Cached config for the node graph */
+	FObjectTreeGraphConfig NodeGraphConfig;
+	/** Cached config for the transition graph */
+	FObjectTreeGraphConfig TransitionGraphConfig;
+
+	/** The build button */
 	TSharedPtr<FBuildButtonToolkit> BuildButtonToolkit;
+	/** The output log */
 	TSharedPtr<FCameraBuildLogToolkit> BuildLogToolkit;
 
 	/** Search widget */
diff --git a/Engine/Plugins/Cameras/GameplayCameras/Source/GameplayCamerasEditor/Public/Toolkits/CameraRigsAssetEditorMode.h b/Engine/Plugins/Cameras/GameplayCameras/Source/GameplayCamerasEditor/Public/Toolkits/CameraRigsAssetEditorMode.h
--- a/Engine/Plugins/Cameras/GameplayCameras/Source/GameplayCamerasEditor/Public/Toolkits/CameraRigsAssetEditorMode.h
+++ b/Engine/Plugins/Cameras/GameplayCameras/Source/GameplayCamerasEditor/Public/Toolkits/CameraRigsAssetEditorMode.h
@@ -3,14 +3,16 @@
 #pragma once
 
 #include "CoreTypes.h"
-#include "Toolkits/AssetEditorMode.h"
+#include "Editors/ObjectTreeGraphConfig.h"
 #include "Editors/SCameraRigList.h"
+#include "Toolkits/AssetEditorMode.h"
 
 class FDocumentTabFactory;
 class FDocumentTracker;
 class UCameraAsset;
 class UCameraRigAsset;
 class UEdGraphNode;
+struct FFindInObjectTreeGraphSource;
 
 namespace UE::Cameras
 {
@@ -26,8 +28,8 @@ public:
 
 	FCameraRigsAssetEditorMode(UCameraAsset* InCameraAsset);
 
-	bool JumpToNode(UEdGraphNode* InNode);
-	bool JumpToObject(UObject* InObject);
+	void OnGetRootObjectsToSearch(TArray<FFindInObjectTreeGraphSource>& OutSources);
+	bool JumpToObject(UObject* InObject, FName PropertyName);
 
 protected:
 
@@ -51,6 +53,9 @@ private:
 
 	TSharedPtr<SCameraRigList> CameraRigsListWidget;
 
+	FObjectTreeGraphConfig NodeGraphConfig;
+	FObjectTreeGraphConfig TransitionGraphConfig;
+
 	bool bInitializedToolkit = false;
 };
 
diff --git a/Engine/Plugins/Cameras/GameplayCameras/Source/GameplayCamerasEditor/Public/Toolkits/CameraSharedTransitionsAssetEditorMode.h b/Engine/Plugins/Cameras/GameplayCameras/Source/GameplayCamerasEditor/Public/Toolkits/CameraSharedTransitionsAssetEditorMode.h
--- a/Engine/Plugins/Cameras/GameplayCameras/Source/GameplayCamerasEditor/Public/Toolkits/CameraSharedTransitionsAssetEditorMode.h
+++ b/Engine/Plugins/Cameras/GameplayCameras/Source/GameplayCamerasEditor/Public/Toolkits/CameraSharedTransitionsAssetEditorMode.h
@@ -3,10 +3,12 @@
 #pragma once
 
 #include "CoreTypes.h"
+#include "Editors/ObjectTreeGraphConfig.h"
 #include "Toolkits/AssetEditorMode.h"
 
 class UCameraAsset;
 class UEdGraphNode;
+struct FFindInObjectTreeGraphSource;
 
 namespace UE::Cameras
 {
@@ -22,8 +24,8 @@ public:
 
 	FCameraSharedTransitionsAssetEditorMode(UCameraAsset* InCameraAsset);
 
-	bool JumpToNode(UEdGraphNode* InNode);
-	bool JumpToObject(UObject* InObject);
+	void OnGetRootObjectsToSearch(TArray<FFindInObjectTreeGraphSource>& OutSources);
+	bool JumpToObject(UObject* InObject, FName PropertyName);
 
 protected:
 
@@ -36,6 +38,8 @@ private:
 
 	TSharedPtr<FCameraRigTransitionEditorToolkitBase> Impl;
 
+	FObjectTreeGraphConfig TransitionGraphConfig;
+
 	bool bInitializedToolkit = false;
 };
 
