PR #11991: Wrong error log fix on Camera animation stop

diff --git a/Engine/Plugins/Cameras/GameplayCameras/Source/GameplayCameras/Private/GameplayCamerasSubsystem.cpp b/Engine/Plugins/Cameras/GameplayCameras/Source/GameplayCameras/Private/GameplayCamerasSubsystem.cpp
--- a/Engine/Plugins/Cameras/GameplayCameras/Source/GameplayCameras/Private/GameplayCamerasSubsystem.cpp
+++ b/Engine/Plugins/Cameras/GameplayCameras/Source/GameplayCameras/Private/GameplayCamerasSubsystem.cpp
@@ -44,7 +44,8 @@ void UGameplayCamerasSubsystem::StopCameraAnimation(APlayerController* PlayerCon
 	UCameraAnimationCameraModifier* CameraModifier = UCameraAnimationCameraModifier::GetCameraAnimationCameraModifierFromPlayerController(PlayerController);
 	if (ensureMsgf(CameraModifier, TEXT("No camera modifier found on the player controller")))
 	{
-		return CameraModifier->StopCameraAnimation(Handle, bImmediate);
+		CameraModifier->StopCameraAnimation(Handle, bImmediate);
+		return;
 	}
 	FFrame::KismetExecutionMessage(TEXT("Can't stop camera animation: no camera animation modifier found"), ELogVerbosity::Error);
 }
@@ -55,6 +56,7 @@ void UGameplayCamerasSubsystem::StopAllCameraAnimationsOf(APlayerController* Pla
 	if (ensureMsgf(CameraModifier, TEXT("No camera modifier found on the player controller")))
 	{
 		CameraModifier->StopAllCameraAnimationsOf(Sequence, bImmediate);
+		return;
 	}
 	FFrame::KismetExecutionMessage(TEXT("Can't stop camera animations: no camera animation modifier found"), ELogVerbosity::Error);
 }
@@ -65,6 +67,7 @@ void UGameplayCamerasSubsystem::StopAllCameraAnimations(APlayerController* Playe
 	if (ensureMsgf(CameraModifier, TEXT("No camera modifier found on the player controller")))
 	{
 		CameraModifier->StopAllCameraAnimations(bImmediate);
+		return;
 	}
 	FFrame::KismetExecutionMessage(TEXT("Can't stop all camera animation: no camera animation modifier found"), ELogVerbosity::Error);
 }
