Cameras: automatically keep track of build warnings and errors in the build log

diff --git a/Engine/Plugins/Cameras/GameplayCameras/Source/GameplayCameras/Private/Core/CameraAssetBuilder.cpp b/Engine/Plugins/Cameras/GameplayCameras/Source/GameplayCameras/Private/Core/CameraAssetBuilder.cpp
--- a/Engine/Plugins/Cameras/GameplayCameras/Source/GameplayCameras/Private/Core/CameraAssetBuilder.cpp
+++ b/Engine/Plugins/Cameras/GameplayCameras/Source/GameplayCameras/Private/Core/CameraAssetBuilder.cpp
@@ -29,8 +29,6 @@ void FCameraAssetBuilder::BuildCamera(UCameraAsset* InCameraAsset)
 	}
 
 	CameraAsset = InCameraAsset;
-	bHasErrors = false;
-	bHasWarnings = false;
 	BuildLog.SetLoggingPrefix(InCameraAsset->GetPathName() + TEXT(": "));
 	{
 		BuildCameraImpl();
@@ -44,32 +42,28 @@ void FCameraAssetBuilder::BuildCameraImpl()
 	if (!CameraAsset->CameraDirector)
 	{
 		BuildLog.AddMessage(EMessageSeverity::Error, LOCTEXT("MissingDirector", "Camera has no director set."));
-		bHasErrors = true;
 	}
 
 	if (CameraAsset->CameraRigs.IsEmpty())
 	{
 		BuildLog.AddMessage(EMessageSeverity::Warning, LOCTEXT("MissingRigs", "Camera has no camera rigs defined."));
-		bHasWarnings = true;
 	}
 
 	for (UCameraRigAsset* CameraRig : CameraAsset->CameraRigs)
 	{
 		FCameraRigAssetBuilder CameraRigBuilder(BuildLog);
 		CameraRigBuilder.BuildCameraRig(CameraRig);
-		bHasErrors |= CameraRigBuilder.LastBuildHadErrors();
-		bHasWarnings |= CameraRigBuilder.LastBuildHadWarnings();
 	}
 }
 
 void FCameraAssetBuilder::UpdateBuildStatus()
 {
 	ECameraBuildStatus BuildStatus = ECameraBuildStatus::Clean;
-	if (bHasErrors)
+	if (BuildLog.HasErrors())
 	{
 		BuildStatus = ECameraBuildStatus::WithErrors;
 	}
-	else if (bHasWarnings)
+	else if (BuildLog.HasWarnings())
 	{
 		BuildStatus = ECameraBuildStatus::CleanWithWarnings;
 	}
diff --git a/Engine/Plugins/Cameras/GameplayCameras/Source/GameplayCameras/Private/Core/CameraBuildLog.cpp b/Engine/Plugins/Cameras/GameplayCameras/Source/GameplayCameras/Private/Core/CameraBuildLog.cpp
--- a/Engine/Plugins/Cameras/GameplayCameras/Source/GameplayCameras/Private/Core/CameraBuildLog.cpp
+++ b/Engine/Plugins/Cameras/GameplayCameras/Source/GameplayCameras/Private/Core/CameraBuildLog.cpp
@@ -3,6 +3,7 @@
 #include "Core/CameraBuildLog.h"
 
 #include "GameplayCameras.h"
+#include "Logging/TokenizedMessage.h"
 #include "Misc/CString.h"
 #include "Misc/UObjectToken.h"
 
@@ -70,6 +71,19 @@ void FCameraBuildLog::AddMessage(EMessageSeverity::Type InSeverity, UObject* InO
 {
 	Messages.Add(FCameraBuildLogMessage{ InSeverity, InObject, MoveTemp(InText) });
 
+	switch (InSeverity)
+	{
+		case EMessageSeverity::Warning:
+		case EMessageSeverity::PerformanceWarning:
+			bHasWarnings = true;
+			break;
+		case EMessageSeverity::Error:
+			bHasErrors = true;
+			break;
+		default:
+			break;
+	}
+
 	if (bForwardToLogging)
 	{
 		const FCameraBuildLogMessage& LastMessage = Messages.Last();
diff --git a/Engine/Plugins/Cameras/GameplayCameras/Source/GameplayCameras/Private/Core/CameraRigAssetBuilder.cpp b/Engine/Plugins/Cameras/GameplayCameras/Source/GameplayCameras/Private/Core/CameraRigAssetBuilder.cpp
--- a/Engine/Plugins/Cameras/GameplayCameras/Source/GameplayCameras/Private/Core/CameraRigAssetBuilder.cpp
+++ b/Engine/Plugins/Cameras/GameplayCameras/Source/GameplayCameras/Private/Core/CameraRigAssetBuilder.cpp
@@ -77,7 +77,6 @@ struct FPrivateVariableBuilder
 	void ReportError(UObject* Object, FText&& ErrorMessage)
 	{
 		Owner.BuildLog.AddMessage(EMessageSeverity::Error, MoveTemp(ErrorMessage));
-		Owner.bHasErrors = true;
 	}
 
 	template<typename ExpectedVariableAssetType>
@@ -212,8 +211,6 @@ void FCameraRigAssetBuilder::BuildCameraRig(UCameraRigAsset* InCameraRig)
 	}
 
 	CameraRig = InCameraRig;
-	bHasErrors = false;
-	bHasWarnings = false;
 	BuildLog.SetLoggingPrefix(InCameraRig->GetPathName() + TEXT(": "));
 	{
 		BuildCameraRigImpl();
@@ -227,7 +224,6 @@ void FCameraRigAssetBuilder::BuildCameraRigImpl()
 	if (!CameraRig->RootNode)
 	{
 		BuildLog.AddMessage(EMessageSeverity::Error, CameraRig, LOCTEXT("MissingRootNode", "Camera rig has no root node set."));
-		bHasErrors = true;
 		return;
 	}
 
@@ -345,7 +341,6 @@ void FCameraRigAssetBuilder::BuildNewDrivenParameters()
 			BuildLog.AddMessage(EMessageSeverity::Error,
 					CameraRig,
 					LOCTEXT("InvalidInterfaceParameter", "Invalid interface parameter or target."));
-			bHasErrors = true;
 			continue;
 		}
 		if (!InterfaceParameter->Target)
@@ -355,7 +350,6 @@ void FCameraRigAssetBuilder::BuildNewDrivenParameters()
 					LOCTEXT(
 						"DisconnectedInterfaceParameter", 
 						"Interface parameter isn't connected: setting overrides for it will not do anything."));
-			bHasWarnings = true;
 			continue;
 		}
 		if (InterfaceParameter->TargetPropertyName.IsNone())
@@ -365,7 +359,6 @@ void FCameraRigAssetBuilder::BuildNewDrivenParameters()
 					LOCTEXT(
 						"InvalidInterfaceParameterTargetPropertyName", 
 						"Invalid interface parameter target property name."));
-			bHasErrors = true;
 			continue;
 		}
 		if (InterfaceParameter->InterfaceParameterName.IsEmpty())
@@ -375,7 +368,6 @@ void FCameraRigAssetBuilder::BuildNewDrivenParameters()
 					LOCTEXT(
 						"InvalidInterfaceParameterName",
 						"Invalid interface parameter name."));
-			bHasErrors = true;
 			continue;
 		}
 
@@ -388,7 +380,6 @@ void FCameraRigAssetBuilder::BuildNewDrivenParameters()
 						"InterfaceParameterNameCollision",
 						"Multiple interface parameters named '{0}'. Ignoring duplicates."),
 						FText::FromString(InterfaceParameter->InterfaceParameterName)));
-			bHasErrors = true;
 			continue;
 		}
 		UsedInterfaceParameterNames.Add(InterfaceParameter->InterfaceParameterName);
@@ -407,7 +398,6 @@ void FCameraRigAssetBuilder::BuildNewDrivenParameters()
 						FText::FromString(InterfaceParameter->InterfaceParameterName), 
 						FText::FromName(InterfaceParameter->TargetPropertyName),
 						FText::FromName(Target->GetFName())));
-			bHasErrors = true;
 			continue;
 		}
 
@@ -422,7 +412,6 @@ void FCameraRigAssetBuilder::BuildNewDrivenParameters()
 						FText::FromString(InterfaceParameter->InterfaceParameterName), 
 						FText::FromName(InterfaceParameter->TargetPropertyName),
 						FText::FromName(Target->GetFName())));
-			bHasErrors = true;
 			continue;
 		}
 
@@ -447,7 +436,6 @@ UE_CAMERA_VARIABLE_FOR_ALL_TYPES()
 						FText::FromString(InterfaceParameter->InterfaceParameterName), 
 						FText::FromName(InterfaceParameter->TargetPropertyName),
 						FText::FromName(Target->GetFName())));
-			bHasErrors = true;
 			continue;
 		}
 	}
@@ -538,11 +526,11 @@ UE_CAMERA_VARIABLE_FOR_ALL_TYPES()
 void FCameraRigAssetBuilder::UpdateBuildStatus()
 {
 	ECameraBuildStatus BuildStatus = ECameraBuildStatus::Clean;
-	if (bHasErrors)
+	if (BuildLog.HasErrors())
 	{
 		BuildStatus = ECameraBuildStatus::WithErrors;
 	}
-	else if (bHasWarnings)
+	else if (BuildLog.HasWarnings())
 	{
 		BuildStatus = ECameraBuildStatus::CleanWithWarnings;
 	}
diff --git a/Engine/Plugins/Cameras/GameplayCameras/Source/GameplayCameras/Public/Core/CameraAssetBuilder.h b/Engine/Plugins/Cameras/GameplayCameras/Source/GameplayCameras/Public/Core/CameraAssetBuilder.h
--- a/Engine/Plugins/Cameras/GameplayCameras/Source/GameplayCameras/Public/Core/CameraAssetBuilder.h
+++ b/Engine/Plugins/Cameras/GameplayCameras/Source/GameplayCameras/Public/Core/CameraAssetBuilder.h
@@ -35,9 +35,6 @@ private:
 	FCameraBuildLog& BuildLog;
 
 	UCameraAsset* CameraAsset = nullptr;
-
-	bool bHasErrors;
-	bool bHasWarnings;
 };
 
 }  // namespace UE::Cameras
diff --git a/Engine/Plugins/Cameras/GameplayCameras/Source/GameplayCameras/Public/Core/CameraBuildLog.h b/Engine/Plugins/Cameras/GameplayCameras/Source/GameplayCameras/Public/Core/CameraBuildLog.h
--- a/Engine/Plugins/Cameras/GameplayCameras/Source/GameplayCameras/Public/Core/CameraBuildLog.h
+++ b/Engine/Plugins/Cameras/GameplayCameras/Source/GameplayCameras/Public/Core/CameraBuildLog.h
@@ -54,12 +54,20 @@ public:
 	/** Gets the list of received messages so far. */
 	TArrayView<const FCameraBuildLogMessage> GetMessages() const { return Messages; }
 
+	/** Returns whether any warning has been logged. */
+	bool HasWarnings() const { return bHasWarnings; }
+	/** Returns whether any error has been logged. */
+	bool HasErrors() const { return bHasErrors; }
+
 private:
 
 	TArray<FCameraBuildLogMessage> Messages;
 
 	FString LoggingPrefix;
 	bool bForwardToLogging = true;
+
+	bool bHasWarnings = false;
+	bool bHasErrors = false;
 };
 
 }  // namespace UE::Cameras
diff --git a/Engine/Plugins/Cameras/GameplayCameras/Source/GameplayCameras/Public/Core/CameraRigAssetBuilder.h b/Engine/Plugins/Cameras/GameplayCameras/Source/GameplayCameras/Public/Core/CameraRigAssetBuilder.h
--- a/Engine/Plugins/Cameras/GameplayCameras/Source/GameplayCameras/Public/Core/CameraRigAssetBuilder.h
+++ b/Engine/Plugins/Cameras/GameplayCameras/Source/GameplayCameras/Public/Core/CameraRigAssetBuilder.h
@@ -39,9 +39,6 @@ public:
 	/** Builds the given camera rig. */
 	void BuildCameraRig(UCameraRigAsset* InCameraRig);
 
-	bool LastBuildHadErrors() const { return bHasErrors; }
-	bool LastBuildHadWarnings() const { return bHasWarnings; }
-
 private:
 
 	void BuildCameraRigImpl();
@@ -70,9 +67,6 @@ private:
 
 	FCameraRigAllocationInfo AllocationInfo;
 
-	bool bHasErrors;
-	bool bHasWarnings;
-
 	friend struct Internal::FPrivateVariableBuilder;
 };
 
