Cameras: un-edit camera rig when deleting it from the rigs list

diff --git a/Engine/Plugins/Cameras/GameplayCameras/Source/GameplayCamerasEditor/Private/Editors/SCameraRigList.cpp b/Engine/Plugins/Cameras/GameplayCameras/Source/GameplayCamerasEditor/Private/Editors/SCameraRigList.cpp
--- a/Engine/Plugins/Cameras/GameplayCameras/Source/GameplayCamerasEditor/Private/Editors/SCameraRigList.cpp
+++ b/Engine/Plugins/Cameras/GameplayCameras/Source/GameplayCamerasEditor/Private/Editors/SCameraRigList.cpp
@@ -120,6 +120,7 @@ void SCameraRigList::Construct(const FArguments& InArgs)
 
 	OnCameraRigListChanged = InArgs._OnCameraRigListChanged;
 	OnRequestEditCameraRig = InArgs._OnRequestEditCameraRig;
+	OnCameraRigDeleted = InArgs._OnCameraRigDeleted;
 
 	CommandList = MakeShared<FUICommandList>();
 
@@ -313,6 +314,7 @@ void SCameraRigList::OnDeleteCameraRig()
 		CameraAsset->Modify();
 
 		TStringBuilder<256> StringBuilder;
+		TArray<UCameraRigAsset*> DeletedCameraRigs;
 
 		for (TSharedPtr<FCameraRigListItem> Item : SelectedItems)
 		{
@@ -328,10 +330,14 @@ void SCameraRigList::OnDeleteCameraRig()
 
 				const int32 NumRemoved = CameraAsset->CameraRigs.Remove(CameraRigAsset);
 				ensure(NumRemoved == 1);
+
+				DeletedCameraRigs.Add(CameraRigAsset);
 			}
 		}
 
 		bUpdateItemSource = true;
+
+		OnCameraRigDeleted.ExecuteIfBound(DeletedCameraRigs);
 	}
 }
 
diff --git a/Engine/Plugins/Cameras/GameplayCameras/Source/GameplayCamerasEditor/Private/Toolkits/CameraRigsAssetEditorMode.cpp b/Engine/Plugins/Cameras/GameplayCameras/Source/GameplayCamerasEditor/Private/Toolkits/CameraRigsAssetEditorMode.cpp
--- a/Engine/Plugins/Cameras/GameplayCameras/Source/GameplayCamerasEditor/Private/Toolkits/CameraRigsAssetEditorMode.cpp
+++ b/Engine/Plugins/Cameras/GameplayCameras/Source/GameplayCamerasEditor/Private/Toolkits/CameraRigsAssetEditorMode.cpp
@@ -54,7 +54,8 @@ void FCameraRigsAssetEditorMode::OnActivateMode(const FAssetEditorModeActivatePa
 		CameraRigsListWidget = SNew(SCameraRigList)
 			.CameraAsset(CameraAsset)
 			.OnCameraRigListChanged(this, &FCameraRigsAssetEditorMode::OnCameraRigListChanged)
-			.OnRequestEditCameraRig(this, &FCameraRigsAssetEditorMode::OnCameraRigEditRequested);
+			.OnRequestEditCameraRig(this, &FCameraRigsAssetEditorMode::OnCameraRigEditRequested)
+			.OnCameraRigDeleted(this, &FCameraRigsAssetEditorMode::OnCameraRigDeleted);
 
 		bInitializedToolkit = true;
 	}
@@ -104,6 +105,14 @@ void FCameraRigsAssetEditorMode::OnCameraRigEditRequested(UCameraRigAsset* InCam
 	Impl->SetCameraRigAsset(InCameraRig);
 }
 
+void FCameraRigsAssetEditorMode::OnCameraRigDeleted(const TArray<UCameraRigAsset*>& InCameraRigs)
+{
+	if (InCameraRigs.Contains(Impl->GetCameraRigAsset()))
+	{
+		Impl->SetCameraRigAsset(nullptr);
+	}
+}
+
 void FCameraRigsAssetEditorMode::OnGetRootObjectsToSearch(TArray<FFindInObjectTreeGraphSource>& OutSources)
 {
 	for (UCameraRigAsset* CameraRig : CameraAsset->CameraRigs)
diff --git a/Engine/Plugins/Cameras/GameplayCameras/Source/GameplayCamerasEditor/Public/Editors/SCameraRigList.h b/Engine/Plugins/Cameras/GameplayCameras/Source/GameplayCamerasEditor/Public/Editors/SCameraRigList.h
--- a/Engine/Plugins/Cameras/GameplayCameras/Source/GameplayCamerasEditor/Public/Editors/SCameraRigList.h
+++ b/Engine/Plugins/Cameras/GameplayCameras/Source/GameplayCamerasEditor/Public/Editors/SCameraRigList.h
@@ -21,7 +21,8 @@ namespace UE::Cameras
 
 DECLARE_DELEGATE(FOnRequestRenameCameraRig);
 DECLARE_DELEGATE_OneParam(FOnCameraRigListChanged, TArrayView<UCameraRigAsset* const>);
-DECLARE_DELEGATE_OneParam(FOnRequestEditCameraRig, UCameraRigAsset*);
+DECLARE_DELEGATE_OneParam(FOnCameraRigEvent, UCameraRigAsset*);
+DECLARE_DELEGATE_OneParam(FOnMultiCameraRigEvent, const TArray<UCameraRigAsset*>&);
 
 /**
  * List item for the list view in SCameraRigList.
@@ -75,7 +76,8 @@ public:
 	{}
 		SLATE_ARGUMENT(UCameraAsset*, CameraAsset)
 		SLATE_EVENT(FOnCameraRigListChanged, OnCameraRigListChanged)
-		SLATE_EVENT(FOnRequestEditCameraRig, OnRequestEditCameraRig)
+		SLATE_EVENT(FOnCameraRigEvent, OnRequestEditCameraRig)
+		SLATE_EVENT(FOnMultiCameraRigEvent, OnCameraRigDeleted)
 	SLATE_END_ARGS()
 
 	void Construct(const FArguments& InArgs);
@@ -120,7 +122,8 @@ private:
 	TObjectPtr<UCameraAsset> CameraAsset;
 
 	FOnCameraRigListChanged OnCameraRigListChanged;
-	FOnRequestEditCameraRig OnRequestEditCameraRig;
+	FOnCameraRigEvent OnRequestEditCameraRig;
+	FOnMultiCameraRigEvent OnCameraRigDeleted;
 
 	TSharedPtr<FUICommandList> CommandList;
 
diff --git a/Engine/Plugins/Cameras/GameplayCameras/Source/GameplayCamerasEditor/Public/Toolkits/CameraRigsAssetEditorMode.h b/Engine/Plugins/Cameras/GameplayCameras/Source/GameplayCamerasEditor/Public/Toolkits/CameraRigsAssetEditorMode.h
--- a/Engine/Plugins/Cameras/GameplayCameras/Source/GameplayCamerasEditor/Public/Toolkits/CameraRigsAssetEditorMode.h
+++ b/Engine/Plugins/Cameras/GameplayCameras/Source/GameplayCamerasEditor/Public/Toolkits/CameraRigsAssetEditorMode.h
@@ -42,6 +42,7 @@ private:
 
 	void OnCameraRigListChanged(TArrayView<UCameraRigAsset* const> InCameraRigs);
 	void OnCameraRigEditRequested(UCameraRigAsset* InCameraRig);
+	void OnCameraRigDeleted(const TArray<UCameraRigAsset*>& InCameraRigs);
 
 private:
 
