Cameras: don't give renamed rigs a space in their UObject name

diff --git a/Engine/Plugins/Cameras/GameplayCameras/Source/GameplayCamerasEditor/Private/Editors/SCameraRigList.cpp b/Engine/Plugins/Cameras/GameplayCameras/Source/GameplayCamerasEditor/Private/Editors/SCameraRigList.cpp
--- a/Engine/Plugins/Cameras/GameplayCameras/Source/GameplayCamerasEditor/Private/Editors/SCameraRigList.cpp
+++ b/Engine/Plugins/Cameras/GameplayCameras/Source/GameplayCamerasEditor/Private/Editors/SCameraRigList.cpp
@@ -8,6 +8,7 @@
 #include "ScopedTransaction.h"
 #include "Styles/GameplayCamerasEditorStyle.h"
 #include "ToolMenus.h"
+#include "UObject/UObjectGlobals.h"
 #include "Widgets/Input/SSearchBox.h"
 #include "Widgets/Layout/SBox.h"
 #include "Widgets/SBoxPanel.h"
@@ -109,7 +110,9 @@ void SCameraRigListEntry::OnTextCommitted(const FText& Text, ETextCommit::Type C
 		// Set the interface name, but also rename the object itself, which helps for debugging.
 		const FString NewDisplayName = Text.ToString();
 		CameraRigAsset->Interface.DisplayName = NewDisplayName;
-		const FName NewObjectName = MakeUniqueObjectName(CameraRigAsset->GetOuter(), UCameraRigAsset::StaticClass(), *NewDisplayName);
+
+		FName NewObjectName = MakeObjectNameFromDisplayLabel(NewDisplayName, CameraRigAsset->GetFName());
+		NewObjectName = MakeUniqueObjectName(CameraRigAsset->GetOuter(), UCameraRigAsset::StaticClass(), NewObjectName);
 		CameraRigAsset->Rename(*NewObjectName.ToString());
 	}
 }
