Cameras: extra camera rig picker from the rig name graph pin code and make it a reusable widget

diff --git a/Engine/Plugins/Cameras/GameplayCameras/Source/GameplayCamerasEditor/Private/Editors/SCameraRigNameGraphPin.cpp b/Engine/Plugins/Cameras/GameplayCameras/Source/GameplayCamerasEditor/Private/Editors/SCameraRigNameGraphPin.cpp
--- a/Engine/Plugins/Cameras/GameplayCameras/Source/GameplayCamerasEditor/Private/Editors/SCameraRigNameGraphPin.cpp
+++ b/Engine/Plugins/Cameras/GameplayCameras/Source/GameplayCamerasEditor/Private/Editors/SCameraRigNameGraphPin.cpp
@@ -2,13 +2,15 @@
 
 #include "Editors/SCameraRigNameGraphPin.h"
 
-#include "Algo/Transform.h"
 #include "AssetRegistry/AssetRegistryModule.h"
 #include "Core/CameraAsset.h"
 #include "Core/CameraRigAsset.h"
+#include "Editors/CameraRigPickerConfig.h"
 #include "Framework/Views/ITypedTableView.h"
 #include "IContentBrowserSingleton.h"
+#include "IGameplayCamerasEditorModule.h"
 #include "Kismet2/BlueprintEditorUtils.h"
+#include "Modules/ModuleManager.h"
 #include "ScopedTransaction.h"
 #include "SGraphPin.h"
 #include "Styles/GameplayCamerasEditorStyle.h"
@@ -16,7 +18,6 @@
 #include "Widgets/Input/SButton.h"
 #include "Widgets/Input/SEditableTextBox.h"
 #include "Widgets/Layout/SBorder.h"
-#include "Widgets/Layout/SBox.h"
 #include "Widgets/SBoxPanel.h"
 #include "Widgets/SNullWidget.h"
 #include "Widgets/Text/STextBlock.h"
@@ -133,22 +134,17 @@ FText SCameraRigNameGraphPin::OnGetComboToolTipText() const
 
 TSharedRef<SWidget> SCameraRigNameGraphPin::OnBuildCameraRigNamePicker()
 {
-	FAssetRegistryModule& AssetRegistryModule = FModuleManager::LoadModuleChecked<FAssetRegistryModule>("AssetRegistry");
-	FContentBrowserModule& ContentBrowserModule = FModuleManager::LoadModuleChecked<FContentBrowserModule>("ContentBrowser");
-
-	FAssetPickerConfig AssetPickerConfig;
-	AssetPickerConfig.bAllowDragging = false;
-	AssetPickerConfig.bAllowNullSelection = false;
-	AssetPickerConfig.SelectionMode = ESelectionMode::Multi;
-	AssetPickerConfig.InitialAssetViewType = EAssetViewType::List;
-	AssetPickerConfig.Filter.ClassPaths.Add(UCameraAsset::StaticClass()->GetClassPathName());
-	AssetPickerConfig.Filter.bRecursiveClasses = true;
-	AssetPickerConfig.SaveSettingsName = TEXT("CameraRigNamePicker");
-	AssetPickerConfig.OnAssetSelected = FOnAssetSelected::CreateSP(this, &SCameraRigNameGraphPin::OnPickerAssetSelected);
-	AssetPickerConfig.GetCurrentSelectionDelegates.Add(&GetCurrentAssetPickerSelection);
+	FCameraRigPickerConfig CameraRigPickerConfig;
+	CameraRigPickerConfig.bCanSelectCameraAsset = true;
+	CameraRigPickerConfig.CameraAssetSelectionMode = ESelectionMode::Multi;
+	CameraRigPickerConfig.CameraAssetViewType = EAssetViewType::List;
+	CameraRigPickerConfig.CameraAssetSaveSettingsName = TEXT("CameraRigNamePicker");
+	CameraRigPickerConfig.OnCameraRigSelected = FOnCameraRigSelected::CreateSP(this, &SCameraRigNameGraphPin::OnPickerAssetSelected);
 
 	if (TSharedPtr<SGraphNode> OwnerNodeWidget = OwnerNodePtr.Pin())
 	{
+		FAssetRegistryModule& AssetRegistryModule = FModuleManager::LoadModuleChecked<FAssetRegistryModule>("AssetRegistry");
+
 		UEdGraphNode* OwnerNode = OwnerNodeWidget->GetNodeObj();
 		UBlueprint* OwnerBlueprint = FBlueprintEditorUtils::FindBlueprintForNode(OwnerNode);
 
@@ -160,65 +156,26 @@ TSharedRef<SWidget> SCameraRigNameGraphPin::OnBuildCameraRigNamePicker()
 		if (!ReferencerAssetData.IsEmpty())
 		{
 			//TODO: multi-select all referencers.
-			AssetPickerConfig.InitialAssetSelection = ReferencerAssetData[0];
+			CameraRigPickerConfig.InitialCameraAssetSelection = ReferencerAssetData[0];
 		}
-
-		UpdateListItemsSource(ReferencerAssetData);
 	}
 
-	TSharedRef<SWidget> PickerWidget = SNew(SBox)
-		.HeightOverride(400)
-		.WidthOverride(350)
-		[
-			SNew(SBorder)
-			.BorderImage(FAppStyle::GetBrush("Menu.Background"))
-			[
-				SNew(SVerticalBox)
-				+SVerticalBox::Slot()
-				.FillHeight(0.55f)
-				[
-					ContentBrowserModule.Get().CreateAssetPicker(AssetPickerConfig)
-				]
-				+SVerticalBox::Slot()
-				.FillHeight(0.45f)
-				[
-					SAssignNew(CameraRigNameListView, SListView<TSharedPtr<FString>>)
-					.ListItemsSource(&CameraRigsItemsSource)
-					.OnGenerateRow(this, &SCameraRigNameGraphPin::OnCameraRigListGenerateRow)
-					.OnSelectionChanged(this, &SCameraRigNameGraphPin::OnCameraRigListSelectionChanged)
-				]
-			]
-		];
-
 	if (GraphPinObj && !GraphPinObj->DefaultValue.IsEmpty())
 	{
-		TSharedPtr<FString>* FoundItem = CameraRigsItemsSource.FindByPredicate(
-				[this](const TSharedPtr<FString>& Item)
-				{
-					return GraphPinObj->DefaultValue == *Item.Get();
-				});
-		if (FoundItem)
-		{
-			bSuppressCameraRigListSelectionChanged = true;
-			CameraRigNameListView->SetSelection(*FoundItem);
-			CameraRigNameListView->RequestScrollIntoView(*FoundItem);
-			bSuppressCameraRigListSelectionChanged = false;
-		}
+		CameraRigPickerConfig.InitialCameraRigSelectionName = GraphPinObj->DefaultValue;
 	}
 
-	return PickerWidget;
+	IGameplayCamerasEditorModule& CamerasEditorModule = FModuleManager::LoadModuleChecked<IGameplayCamerasEditorModule>("GameplayCamerasEditor");
+	return CamerasEditorModule.CreateCameraRigPicker(CameraRigPickerConfig);
 }
 
-void SCameraRigNameGraphPin::OnPickerAssetSelected(const FAssetData& AssetData)
+void SCameraRigNameGraphPin::OnPickerAssetSelected(UCameraRigAsset* SelectedItem)
 {
-	TArray<FAssetData> SelectedAssets;
-	if (GetCurrentAssetPickerSelection.IsBound())
+	if (SelectedItem)
 	{
-		SelectedAssets = GetCurrentAssetPickerSelection.Execute();
+		CameraRigPickerButton->SetIsOpen(false);
+		SetCameraRigName(SelectedItem->GetDisplayName());
 	}
-
-	UpdateListItemsSource(SelectedAssets);
-	CameraRigNameListView->RequestListRefresh();
 }
 
 FReply SCameraRigNameGraphPin::OnResetButtonClicked()
@@ -228,79 +185,6 @@ FReply SCameraRigNameGraphPin::OnResetButtonClicked()
 	return FReply::Handled();
 }
 
-TSharedRef<ITableRow> SCameraRigNameGraphPin::OnCameraRigListGenerateRow(TSharedPtr<FString> Item, const TSharedRef<STableViewBase>& OwnerTable)
-{
-	TSharedRef<FGameplayCamerasEditorStyle> GameplayCamerasStyle = FGameplayCamerasEditorStyle::Get();
-
-	const FText DisplayName = FText::FromString(*Item);
-
-	return SNew(STableRow<TSharedPtr<FString>>, OwnerTable)
-		.Padding(FMargin(0.f, 4.f))
-		[
-			SNew(SBorder)
-			.Padding(FMargin(4.f))
-			.BorderImage(FAppStyle::GetBrush("NoBorder"))
-			[
-				SNew(SHorizontalBox)
-				+SHorizontalBox::Slot()
-				.AutoWidth()
-				[
-					SNew(SImage)
-					.ColorAndOpacity(FSlateColor::UseForeground())
-					.Image(GameplayCamerasStyle->GetBrush("CameraAssetEditor.ShowCameraRigs"))
-				]
-				+SHorizontalBox::Slot()
-				.FillWidth(1.f)
-				.Padding(4.f, 2.f)
-				[
-					SNew(STextBlock)
-					.Text(DisplayName)
-				]
-			]
-		];
-}
-
-void SCameraRigNameGraphPin::OnCameraRigListSelectionChanged(TSharedPtr<FString> Item, ESelectInfo::Type SelectInfo)
-{
-	if (!bSuppressCameraRigListSelectionChanged && Item)
-	{
-		CameraRigPickerButton->SetIsOpen(false);
-		SetCameraRigName(*Item.Get());
-	}
-}
-
-void SCameraRigNameGraphPin::UpdateListItemsSource(const TArray<FAssetData>& Assets)
-{
-	// Only show the names that are in common between all selected camera assets.
-	TSet<FString> CommonNames;
-	bool bCommonNamesInitialized = false;
-	for (const FAssetData& SelectedAsset : Assets)
-	{
-		if (const UCameraAsset* CameraAsset = Cast<UCameraAsset>(SelectedAsset.GetAsset()))
-		{
-			TSet<FString> SelectedAssetNames;
-			Algo::Transform(CameraAsset->CameraRigs, SelectedAssetNames, 
-					[](const UCameraRigAsset* Item) { return Item->GetDisplayName(); });
-
-			if (!bCommonNamesInitialized)
-			{
-				CommonNames.Append(SelectedAssetNames);
-				bCommonNamesInitialized = true;
-			}
-			else
-			{
-				CommonNames = CommonNames.Intersect(SelectedAssetNames);
-			}
-		}
-	}
-
-	CameraRigsItemsSource.Reset();
-	for (const FString& CommonName : CommonNames)
-	{
-		CameraRigsItemsSource.Add(MakeShared<FString>(CommonName));
-	}
-}
-
 void SCameraRigNameGraphPin::SetCameraRigName(const FString& InCameraRigName)
 {
 	const FScopedTransaction Transaction(LOCTEXT("ChangeObjectPinValue", "Change Object Pin Value"));
diff --git a/Engine/Plugins/Cameras/GameplayCameras/Source/GameplayCamerasEditor/Private/Editors/SCameraRigPicker.cpp b/Engine/Plugins/Cameras/GameplayCameras/Source/GameplayCamerasEditor/Private/Editors/SCameraRigPicker.cpp
new file mode 100644
--- /dev/null
+++ b/Engine/Plugins/Cameras/GameplayCameras/Source/GameplayCamerasEditor/Private/Editors/SCameraRigPicker.cpp
@@ -0,0 +1,344 @@
+// Copyright Epic Games, Inc. All Rights Reserved.
+
+#include "Editors/SCameraRigPicker.h"
+
+#include "AssetRegistry/AssetRegistryModule.h"
+#include "ContentBrowserModule.h"
+#include "Core/CameraAsset.h"
+#include "Core/CameraRigAsset.h"
+#include "Editors/CameraRigPickerConfig.h"
+#include "Framework/Application/SlateApplication.h"
+#include "IContentBrowserSingleton.h"
+#include "Layout/WidgetPath.h"
+#include "PropertyHandle.h"
+#include "Styles/GameplayCamerasEditorStyle.h"
+#include "Types/SlateEnums.h"
+#include "Widgets/Input/SSearchBox.h"
+#include "Widgets/Layout/SBox.h"
+
+#define LOCTEXT_NAMESPACE "SCameraRigPicker"
+
+namespace UE::Cameras
+{
+
+void SCameraRigPicker::Construct(const FArguments& InArgs)
+{
+	const FCameraRigPickerConfig& PickerConfig = InArgs._CameraRigPickerConfig;
+
+	SearchTextFilter = MakeShareable(new FTextFilter(
+		FTextFilter::FItemToStringArray::CreateSP(this, &SCameraRigPicker::GetEntryStrings)));
+
+	// Camera asset picker.
+	TSharedRef<SVerticalBox> LayoutBox = SNew(SVerticalBox);
+
+	if (PickerConfig.bCanSelectCameraAsset)
+	{
+		FAssetRegistryModule& AssetRegistryModule = FModuleManager::LoadModuleChecked<FAssetRegistryModule>("AssetRegistry");
+		FContentBrowserModule& ContentBrowserModule = FModuleManager::LoadModuleChecked<FContentBrowserModule>("ContentBrowser");
+
+		FAssetPickerConfig AssetPickerConfig;
+		AssetPickerConfig.bAllowDragging = false;
+		AssetPickerConfig.bAllowNullSelection = false;
+		AssetPickerConfig.Filter.ClassPaths.Add(UCameraAsset::StaticClass()->GetClassPathName());
+		AssetPickerConfig.Filter.bRecursiveClasses = true;
+		AssetPickerConfig.OnAssetSelected = FOnAssetSelected::CreateSP(this, &SCameraRigPicker::OnCameraAssetSelected);
+		AssetPickerConfig.GetCurrentSelectionDelegates.Add(&GetCurrentCameraAssetPickerSelection);
+
+		AssetPickerConfig.SelectionMode = PickerConfig.CameraAssetSelectionMode;
+		AssetPickerConfig.InitialAssetViewType = PickerConfig.CameraAssetViewType;
+		AssetPickerConfig.SaveSettingsName = PickerConfig.CameraAssetSaveSettingsName;
+		AssetPickerConfig.InitialAssetSelection = PickerConfig.InitialCameraAssetSelection;
+
+		LayoutBox->AddSlot()
+		.FillHeight(0.55f)
+		[
+			ContentBrowserModule.Get().CreateAssetPicker(AssetPickerConfig)
+		];
+	}
+
+	// Camera rig picker.
+	const float CamerRigPickerFillHeight = PickerConfig.bCanSelectCameraAsset ? 0.45f : 1.f;
+	LayoutBox->AddSlot()
+	.FillHeight(CamerRigPickerFillHeight)
+	[
+		SNew(SVerticalBox)
+		+SVerticalBox::Slot()
+		.AutoHeight()
+		[
+			SNew(SBorder)
+			.BorderImage(FAppStyle::GetBrush("ToolPanel.GroupBorder"))
+			.Padding(8.f)
+			[
+				SAssignNew(SearchBox, SSearchBox)
+				.HintText(LOCTEXT("SearchHint", "Search"))
+				.OnTextChanged(this, &SCameraRigPicker::OnSearchTextChanged)
+				.OnTextCommitted(this, &SCameraRigPicker::OnSearchTextCommitted)
+				.OnKeyDownHandler(this, &SCameraRigPicker::OnSearchKeyDown)
+			]
+		]
+		+SVerticalBox::Slot()
+		.Padding(0.f, 3.f)
+		[
+			SAssignNew(CameraRigListView, SListView<UCameraRigAsset*>)
+			.ListItemsSource(&CameraRigFilteredItemsSource)
+			.OnGenerateRow(this, &SCameraRigPicker::OnCameraRigListGenerateRow)
+			.OnSelectionChanged(this, &SCameraRigPicker::OnCameraRigListSelectionChanged)
+		]
+	];
+	
+	// Full layout.
+	ChildSlot
+	[
+		SNew(SBox)
+		.HeightOverride(400)
+		.WidthOverride(350)
+		[
+			SNew(SBorder)
+			.BorderImage(FAppStyle::GetBrush("Menu.Background"))
+			[
+				LayoutBox
+			]
+		]
+	];
+
+	// If we have an initially selected camera asset, fill the list of camera rigs immediately.
+	if (PickerConfig.InitialCameraAssetSelection.IsValid())
+	{
+		UpdateCameraRigItemsSource(TArray<FAssetData>{ PickerConfig.InitialCameraAssetSelection });
+	}
+
+	UpdateCameraRigFilteredItemsSource();
+
+	// If we have an initially selected camera rig, select it in the list immediately too.
+	UCameraRigAsset* InitialCameraRigSelection = PickerConfig.InitialCameraRigSelection;
+	if (!PickerConfig.InitialCameraRigSelectionName.IsEmpty() && !InitialCameraRigSelection)
+	{
+		UCameraRigAsset** FoundItem = CameraRigFilteredItemsSource.FindByPredicate(
+				[&PickerConfig](UCameraRigAsset* Item)
+				{
+					return Item->GetDisplayName() == PickerConfig.InitialCameraRigSelectionName;
+				});
+		if (FoundItem)
+		{
+			InitialCameraRigSelection = *FoundItem;
+		}
+	}
+	if (InitialCameraRigSelection)
+	{
+		CameraRigListView->RequestScrollIntoView(InitialCameraRigSelection);
+		CameraRigListView->SetSelection(InitialCameraRigSelection);
+	}
+
+	// If we need to focus the search box, register a time to do that next frame.
+	if (PickerConfig.bFocusCameraRigSearchBoxWhenOpened)
+	{
+		RegisterActiveTimer(0.f, FWidgetActiveTimerDelegate::CreateSP(this, &SCameraRigPicker::FocusCameraRigSearchBox));
+	}
+
+	// Keep track of miscellaneous stuff.
+	OnCameraRigSelected = PickerConfig.OnCameraRigSelected;
+	PropertyToSet = PickerConfig.PropertyToSet;
+}
+
+EActiveTimerReturnType SCameraRigPicker::FocusCameraRigSearchBox(double InCurrentTime, float InDeltaTime)
+{
+	if (SearchBox)
+	{
+		FWidgetPath WidgetToFocusPath;
+		FSlateApplication::Get().GeneratePathToWidgetUnchecked(SearchBox.ToSharedRef(), WidgetToFocusPath);
+		FSlateApplication::Get().SetKeyboardFocus(WidgetToFocusPath, EFocusCause::SetDirectly);
+		WidgetToFocusPath.GetWindow()->SetWidgetToFocusOnActivate(SearchBox);
+		return EActiveTimerReturnType::Stop;
+	}
+
+	return EActiveTimerReturnType::Continue;
+}
+
+void SCameraRigPicker::Tick(const FGeometry& AllottedGeometry, const double InCurrentTime, const float InDeltaTime)
+{
+	if (bUpdateItemsSource)
+	{
+		if (GetCurrentCameraAssetPickerSelection.IsBound())
+		{
+			TArray<FAssetData> SelectedAssets = GetCurrentCameraAssetPickerSelection.Execute();
+			UpdateCameraRigItemsSource(SelectedAssets);
+		}
+	}
+	if (bUpdateFilteredItemsSource || bUpdateItemsSource)
+	{
+		UpdateCameraRigFilteredItemsSource();
+	}
+
+	const bool bRequestListRefresh = bUpdateItemsSource || bUpdateFilteredItemsSource;
+	bUpdateItemsSource = false;
+	bUpdateFilteredItemsSource = false;
+
+	if (bRequestListRefresh)
+	{
+		CameraRigListView->RequestListRefresh();
+	}
+
+	SCompoundWidget::Tick(AllottedGeometry, InCurrentTime, InDeltaTime);
+}
+
+void SCameraRigPicker::OnCameraAssetSelected(const FAssetData& AssetData)
+{
+	bUpdateItemsSource = true;
+}
+
+TSharedRef<ITableRow> SCameraRigPicker::OnCameraRigListGenerateRow(UCameraRigAsset* Item, const TSharedRef<STableViewBase>& OwnerTable)
+{
+	TSharedRef<FGameplayCamerasEditorStyle> GameplayCamerasStyle = FGameplayCamerasEditorStyle::Get();
+
+	return SNew(STableRow<UCameraRigAsset*>, OwnerTable)
+		.Padding(FMargin(2.f))
+		[
+			SNew(SBorder)
+			.BorderImage(FAppStyle::GetBrush("NoBorder"))
+			[
+				SNew(SHorizontalBox)
+				+SHorizontalBox::Slot()
+				.AutoWidth()
+				[
+					SNew(SImage)
+					.ColorAndOpacity(FSlateColor::UseForeground())
+					.Image(GameplayCamerasStyle->GetBrush("CameraAssetEditor.ShowCameraRigs"))
+				]
+				+SHorizontalBox::Slot()
+				.FillWidth(1.f)
+				.Padding(4.f, 2.f)
+				[
+					SNew(STextBlock)
+					.HighlightText(this, &SCameraRigPicker::GetHighlightText)
+					.Text_Lambda([Item]() { return FText::FromString(Item->GetDisplayName()); })
+				]
+			]
+		];
+}
+
+void SCameraRigPicker::OnCameraRigListSelectionChanged(UCameraRigAsset* Item, ESelectInfo::Type SelectInfo)
+{
+	if (SelectInfo != ESelectInfo::Direct)
+	{
+		if (PropertyToSet)
+		{
+			PropertyToSet->SetValue(Item);
+		}
+
+		OnCameraRigSelected.ExecuteIfBound(Item);
+	}
+}
+
+void SCameraRigPicker::UpdateCameraRigItemsSource(const TArray<FAssetData>& Assets)
+{
+	CameraRigItemsSource.Reset();
+
+	for (const FAssetData& SelectedAsset : Assets)
+	{
+		if (const UCameraAsset* CameraAsset = Cast<UCameraAsset>(SelectedAsset.GetAsset()))
+		{
+			CameraRigItemsSource.Append(CameraAsset->CameraRigs);
+		}
+	}
+}
+
+void SCameraRigPicker::UpdateCameraRigFilteredItemsSource()
+{
+	CameraRigFilteredItemsSource = CameraRigItemsSource;
+	CameraRigFilteredItemsSource.StableSort([](UCameraRigAsset& A, UCameraRigAsset& B)
+			{ 
+				return A.GetDisplayName().Compare(B.GetDisplayName()) < 0;
+			});
+
+	if (!SearchTextFilter->GetRawFilterText().IsEmpty())
+	{
+		CameraRigFilteredItemsSource = CameraRigFilteredItemsSource.FilterByPredicate(
+				[this](UCameraRigAsset* Item)
+				{
+					return SearchTextFilter->PassesFilter(Item);
+				});
+	}
+}
+
+void SCameraRigPicker::GetEntryStrings(const UCameraRigAsset* InItem, TArray<FString>& OutStrings)
+{
+	OutStrings.Add(InItem->GetDisplayName());
+}
+
+void SCameraRigPicker::OnSearchTextChanged(const FText& InFilterText)
+{
+	SearchTextFilter->SetRawFilterText(InFilterText);
+	SearchBox->SetError(SearchTextFilter->GetFilterErrorText());
+
+	bUpdateFilteredItemsSource = true;
+}
+
+void SCameraRigPicker::OnSearchTextCommitted(const FText& InFilterText, ETextCommit::Type InCommitType)
+{
+	OnSearchTextChanged(InFilterText);
+
+	if (InCommitType == ETextCommit::OnEnter)
+	{
+		TArray<UCameraRigAsset*> SelectedItems = CameraRigListView->GetSelectedItems();
+		if (SelectedItems.Num() > 0)
+		{
+			OnCameraRigListSelectionChanged(SelectedItems[0], ESelectInfo::OnKeyPress);
+		}
+	}
+}
+
+FReply SCameraRigPicker::OnSearchKeyDown(const FGeometry& MyGeometry, const FKeyEvent& InKeyEvent)
+{
+	int32 SelectionDelta = 0;
+
+	if (InKeyEvent.GetKey() == EKeys::Up)
+	{
+		SelectionDelta = -1;
+	}
+	else if (InKeyEvent.GetKey() == EKeys::Down)
+	{
+		SelectionDelta = +1;
+	}
+
+	if (SelectionDelta != 0 && !CameraRigFilteredItemsSource.IsEmpty())
+	{
+		TArray<UCameraRigAsset*> SelectedItems = CameraRigListView->GetSelectedItems();
+		if (SelectedItems.IsEmpty())
+		{
+			// No items already selected... select the first or last depending on the key pressed.
+			if (SelectionDelta > 0)
+			{
+				CameraRigListView->SetSelection(CameraRigFilteredItemsSource[0]);
+			}
+			else if (SelectionDelta < 0)
+			{
+				CameraRigListView->SetSelection(CameraRigFilteredItemsSource.Last());
+			}
+
+			return FReply::Handled();
+		}
+
+		int32 SelectedIndex = CameraRigFilteredItemsSource.Find(SelectedItems[0]);
+		if (ensure(SelectedIndex >= 0))
+		{
+			// Set the selection to the previous/next item, wrapping around the list.
+			SelectedIndex = (SelectedIndex + SelectionDelta + CameraRigFilteredItemsSource.Num()) % CameraRigFilteredItemsSource.Num();
+			CameraRigListView->RequestScrollIntoView(CameraRigFilteredItemsSource[SelectedIndex]);
+			CameraRigListView->SetSelection(CameraRigFilteredItemsSource[SelectedIndex]);
+
+			return FReply::Handled();
+		}
+	}
+
+	return FReply::Unhandled();
+}
+
+FText SCameraRigPicker::GetHighlightText() const
+{
+	return SearchTextFilter->GetRawFilterText();
+}
+
+}  // namespace UE::Cameras
+
+#undef LOCTEXT_NAMESPACE
+
diff --git a/Engine/Plugins/Cameras/GameplayCameras/Source/GameplayCamerasEditor/Private/Editors/SCameraRigPicker.h b/Engine/Plugins/Cameras/GameplayCameras/Source/GameplayCamerasEditor/Private/Editors/SCameraRigPicker.h
new file mode 100644
--- /dev/null
+++ b/Engine/Plugins/Cameras/GameplayCameras/Source/GameplayCamerasEditor/Private/Editors/SCameraRigPicker.h
@@ -0,0 +1,70 @@
+// Copyright Epic Games, Inc. All Rights Reserved.
+
+#pragma once
+
+#include "ContentBrowserModule.h"
+#include "Editors/CameraRigPickerConfig.h"
+#include "Misc/TextFilter.h"
+#include "Widgets/SCompoundWidget.h"
+#include "Widgets/Views/SListView.h"
+
+class SSearchBox;
+class UCameraRigAsset;
+
+namespace UE::Cameras
+{
+
+class SCameraRigPicker : public SCompoundWidget
+{
+public:
+
+	SLATE_BEGIN_ARGS(SCameraRigPicker)
+	{}
+		SLATE_ARGUMENT(FCameraRigPickerConfig, CameraRigPickerConfig)
+	SLATE_END_ARGS()
+
+	void Construct(const FArguments& InArgs);
+
+protected:
+
+	// SWidget interface
+	virtual void Tick(const FGeometry& AllottedGeometry, const double InCurrentTime, const float InDeltaTime) override;
+
+private:
+
+	EActiveTimerReturnType FocusCameraRigSearchBox(double InCurrentTime, float InDeltaTime);
+
+	void OnCameraAssetSelected(const FAssetData& AssetData);
+
+	TSharedRef<ITableRow> OnCameraRigListGenerateRow(UCameraRigAsset* Item, const TSharedRef<STableViewBase>& OwnerTable);
+	void OnCameraRigListSelectionChanged(UCameraRigAsset* Item, ESelectInfo::Type SelectInfo);
+	void UpdateCameraRigItemsSource(const TArray<FAssetData>& Assets);
+	void UpdateCameraRigFilteredItemsSource();
+
+	void GetEntryStrings(const UCameraRigAsset* InItem, TArray<FString>& OutStrings);
+	void OnSearchTextChanged(const FText& InFilterText);
+	void OnSearchTextCommitted(const FText& InFilterText, ETextCommit::Type InCommitType);
+	FReply OnSearchKeyDown(const FGeometry& MyGeometry, const FKeyEvent& InKeyEvent);
+	FText GetHighlightText() const;
+
+private:
+
+	FGetCurrentSelectionDelegate GetCurrentCameraAssetPickerSelection;
+
+	TSharedPtr<SListView<UCameraRigAsset*>> CameraRigListView;
+	TArray<UCameraRigAsset*> CameraRigItemsSource;
+	TArray<UCameraRigAsset*> CameraRigFilteredItemsSource;
+
+	bool bUpdateItemsSource = false;
+	bool bUpdateFilteredItemsSource = false;
+
+	using FTextFilter = TTextFilter<const UCameraRigAsset*>;
+	TSharedPtr<FTextFilter> SearchTextFilter;
+	TSharedPtr<SSearchBox> SearchBox;
+
+	FOnCameraRigSelected OnCameraRigSelected;
+	TSharedPtr<IPropertyHandle> PropertyToSet;
+};
+
+}  // namespace UE::Cameras
+
diff --git a/Engine/Plugins/Cameras/GameplayCameras/Source/GameplayCamerasEditor/Private/GameplayCamerasEditorModule.cpp b/Engine/Plugins/Cameras/GameplayCameras/Source/GameplayCamerasEditor/Private/GameplayCamerasEditorModule.cpp
--- a/Engine/Plugins/Cameras/GameplayCameras/Source/GameplayCamerasEditor/Private/GameplayCamerasEditorModule.cpp
+++ b/Engine/Plugins/Cameras/GameplayCameras/Source/GameplayCamerasEditor/Private/GameplayCamerasEditorModule.cpp
@@ -18,6 +18,7 @@
 #include "Directors/BlueprintCameraDirector.h"
 #include "EdGraph/EdGraph.h"
 #include "Editors/GameplayCamerasGraphPanelPinFactory.h"
+#include "Editors/SCameraRigPicker.h"
 #include "Features/IModularFeatures.h"
 #include "GameplayCameras.h"
 #include "GameplayCamerasEditorSettings.h"
@@ -129,6 +130,14 @@ public:
 		return AssetEditor;
 	}
 
+	virtual TSharedRef<SWidget> CreateCameraRigPicker(const FCameraRigPickerConfig& InPickerConfig) override
+	{
+		using namespace UE::Cameras;
+
+		return SNew(SCameraRigPicker)
+			.CameraRigPickerConfig(InPickerConfig);
+	}
+
 	virtual void RegisterDebugCategory(const UE::Cameras::FCameraDebugCategoryInfo& InCategoryInfo) override
 	{
 		if (!ensureMsgf(!InCategoryInfo.Name.IsEmpty(), TEXT("A debug category must at least specify a name!")))
diff --git a/Engine/Plugins/Cameras/GameplayCameras/Source/GameplayCamerasEditor/Public/Editors/CameraRigPickerConfig.h b/Engine/Plugins/Cameras/GameplayCameras/Source/GameplayCamerasEditor/Public/Editors/CameraRigPickerConfig.h
new file mode 100644
--- /dev/null
+++ b/Engine/Plugins/Cameras/GameplayCameras/Source/GameplayCamerasEditor/Public/Editors/CameraRigPickerConfig.h
@@ -0,0 +1,66 @@
+// Copyright Epic Games, Inc. All Rights Reserved.
+
+#pragma once
+
+#include "Framework/Views/ITypedTableView.h"
+#include "IContentBrowserSingleton.h"
+#include "UObject/ObjectKey.h"
+
+class IPropertyHandle;
+class UCameraRigAsset;
+
+namespace UE::Cameras
+{
+
+DECLARE_DELEGATE_OneParam(FOnCameraRigSelected, UCameraRigAsset*);
+
+/**
+ * Configuration structure for a camera rig picker widget.
+ * This is a widget that shows:
+ *	- An asset picker for a camera asset, and below it...
+ *	- A list view with the camera rigs inside that camera asset.
+ *
+ * See IGameplayCamerasEditorModule for creating that widget.
+ */
+struct FCameraRigPickerConfig
+{
+	/** The initially selected camera asset, if any. */
+	FAssetData InitialCameraAssetSelection;
+
+	/** 
+	 * The initially selected camera rig, specified as a pointer, if any.
+	 * This shouldn't be set if InitialCameraAssetSelectionName is set.
+	 */
+	UCameraRigAsset* InitialCameraRigSelection;
+
+	/**
+	 * The initially selected camera rig, picked by name, if any.
+	 * This shouldn't be set if InitialCameraAssetSelection is set.
+	 */
+	FString InitialCameraRigSelectionName;
+
+	/** 
+	 * Whether a camera asset can be picked, or whether the InitialCameraAssetSelection
+	 * is the only asset that should be used. If the latter (when the value is false),
+	 * then no camera asset picker is shown. Only the list of rigs is shown.
+	 */
+	bool bCanSelectCameraAsset = false;
+
+	/** Asset picker view type for the camera asset picker. */
+	EAssetViewType::Type CameraAssetViewType = EAssetViewType::List;
+	/** Asset picker selection mode for the camera asset picker. */
+	ESelectionMode::Type CameraAssetSelectionMode = ESelectionMode::Single;
+	/** Asset picker settings name for the camera asset picker. */
+	FString CameraAssetSaveSettingsName;
+
+	/** Whether the camera rig search box should be focused initially. */
+	bool bFocusCameraRigSearchBoxWhenOpened = true;
+
+	/** Callback for when a camera rig has been selected. */
+	FOnCameraRigSelected OnCameraRigSelected;
+
+	TSharedPtr<IPropertyHandle> PropertyToSet;
+};
+
+}  // namespace UE::Cameras
+
diff --git a/Engine/Plugins/Cameras/GameplayCameras/Source/GameplayCamerasEditor/Public/Editors/SCameraRigNameGraphPin.h b/Engine/Plugins/Cameras/GameplayCameras/Source/GameplayCamerasEditor/Public/Editors/SCameraRigNameGraphPin.h
--- a/Engine/Plugins/Cameras/GameplayCameras/Source/GameplayCamerasEditor/Public/Editors/SCameraRigNameGraphPin.h
+++ b/Engine/Plugins/Cameras/GameplayCameras/Source/GameplayCamerasEditor/Public/Editors/SCameraRigNameGraphPin.h
@@ -6,7 +6,8 @@
 
 #include "ContentBrowserModule.h"
 #include "Widgets/Input/SComboButton.h"
-#include "Widgets/Views/SListView.h"
+
+class UCameraRigAsset;
 
 namespace UE::Cameras
 {
@@ -48,25 +49,15 @@ private:
 	FText OnGetComboToolTipText() const;
 
 	TSharedRef<SWidget> OnBuildCameraRigNamePicker();
-	void OnPickerAssetSelected(const FAssetData& AssetData);
+	void OnPickerAssetSelected(UCameraRigAsset* SelectedItem);
 
 	FReply OnResetButtonClicked();
 
-	TSharedRef<ITableRow> OnCameraRigListGenerateRow(TSharedPtr<FString> Item, const TSharedRef<STableViewBase>& OwnerTable);
-	void OnCameraRigListSelectionChanged(TSharedPtr<FString> Item, ESelectInfo::Type SelectInfo);
-	void UpdateListItemsSource(const TArray<FAssetData>& Assets);
-
 	void SetCameraRigName(const FString& InCameraRigName);
 
 private:
 
 	TSharedPtr<SComboButton> CameraRigPickerButton;
-
-	FGetCurrentSelectionDelegate GetCurrentAssetPickerSelection;
-
-	TSharedPtr<SListView<TSharedPtr<FString>>> CameraRigNameListView;
-	TArray<TSharedPtr<FString>> CameraRigsItemsSource;
-	bool bSuppressCameraRigListSelectionChanged = false;
 };
 
 }  // namespace UE::Cameras
diff --git a/Engine/Plugins/Cameras/GameplayCameras/Source/GameplayCamerasEditor/Public/IGameplayCamerasEditorModule.h b/Engine/Plugins/Cameras/GameplayCameras/Source/GameplayCamerasEditor/Public/IGameplayCamerasEditorModule.h
--- a/Engine/Plugins/Cameras/GameplayCameras/Source/GameplayCamerasEditor/Public/IGameplayCamerasEditorModule.h
+++ b/Engine/Plugins/Cameras/GameplayCameras/Source/GameplayCamerasEditor/Public/IGameplayCamerasEditorModule.h
@@ -7,6 +7,7 @@
 #include "Modules/ModuleInterface.h"
 #include "Toolkits/AssetEditorToolkit.h"
 
+class SWidget;
 class UCameraAsset;
 class UCameraAssetEditor;
 class UCameraRigAsset;
@@ -19,6 +20,8 @@ DECLARE_DELEGATE_RetVal_OneParam(TSharedRef<SWidget>, FOnCreateDebugCategoryPane
 namespace UE::Cameras
 {
 
+struct FCameraRigPickerConfig;
+
 struct FCameraDebugCategoryInfo
 {
 	FString Name;
@@ -43,6 +46,8 @@ public:
 
 public:
 
+	using FCameraRigPickerConfig = UE::Cameras::FCameraRigPickerConfig;
+
 	/** Creates an editor for the given camera asset */
 	virtual UCameraAssetEditor* CreateCameraAssetEditor(const EToolkitMode::Type Mode, const TSharedPtr<IToolkitHost>& InitToolkitHost, UCameraAsset* CameraAsset) = 0;
 
@@ -52,6 +57,9 @@ public:
 	/** Creates an editor for the given variable collection */
 	virtual UCameraVariableCollectionEditor* CreateCameraVariableCollectionEditor(const EToolkitMode::Type Mode, const TSharedPtr<IToolkitHost>& InitToolkitHost, UCameraVariableCollection* VariableCollection) = 0;
 
+	/** Creates a new camera rig asset picker widget */
+	virtual TSharedRef<SWidget> CreateCameraRigPicker(const FCameraRigPickerConfig& InPickerConfig) = 0;
+
 public:
 
 	using FCameraDebugCategoryInfo = UE::Cameras::FCameraDebugCategoryInfo;
