Cameras: give the single-rig director a custom rig picker

diff --git a/Engine/Plugins/Cameras/GameplayCameras/Source/GameplayCamerasEditor/Private/Customizations/SingleCameraDirectorDetailsCustomization.cpp b/Engine/Plugins/Cameras/GameplayCameras/Source/GameplayCamerasEditor/Private/Customizations/SingleCameraDirectorDetailsCustomization.cpp
new file mode 100644
--- /dev/null
+++ b/Engine/Plugins/Cameras/GameplayCameras/Source/GameplayCamerasEditor/Private/Customizations/SingleCameraDirectorDetailsCustomization.cpp
@@ -0,0 +1,107 @@
+// Copyright Epic Games, Inc. All Rights Reserved.
+
+#include "Customizations/SingleCameraDirectorDetailsCustomization.h"
+
+#include "Core/CameraAsset.h"
+#include "Core/CameraRigAsset.h"
+#include "DetailCategoryBuilder.h"
+#include "DetailLayoutBuilder.h"
+#include "DetailWidgetRow.h"
+#include "Directors/SingleCameraDirector.h"
+#include "Editors/CameraRigPickerConfig.h"
+#include "IDetailChildrenBuilder.h"
+#include "IGameplayCamerasEditorModule.h"
+#include "Modules/ModuleManager.h"
+#include "PropertyCustomizationHelpers.h"
+#include "Styling/AppStyle.h"
+#include "Widgets/Input/SComboButton.h"
+#include "Widgets/SNullWidget.h"
+
+#define LOCTEXT_NAMESPACE "SingleCameraDirectorDetailsCustomization"
+
+namespace UE::Cameras
+{
+
+TSharedRef<IDetailCustomization> FSingleCameraDirectorDetailsCustomization::MakeInstance()
+{
+	return MakeShared<FSingleCameraDirectorDetailsCustomization>();
+}
+
+void FSingleCameraDirectorDetailsCustomization::CustomizeDetails(IDetailLayoutBuilder& DetailBuilder)
+{
+	TArray<TWeakObjectPtr<USingleCameraDirector>> WeakDirectors = DetailBuilder.GetSelectedObjectsOfType<USingleCameraDirector>();
+	WeakSelectedDirector = WeakDirectors.Num() == 1 ? WeakDirectors[0] : nullptr;
+
+	CameraRigPropertyHandle = DetailBuilder.GetProperty(GET_MEMBER_NAME_CHECKED(USingleCameraDirector, CameraRig));
+
+	IDetailCategoryBuilder& CommonCategory = DetailBuilder.EditCategory(TEXT("Common"));
+	CommonCategory.AddProperty(CameraRigPropertyHandle)
+		.IsEnabled(WeakSelectedDirector.IsValid())
+		.CustomWidget()
+			.NameContent()
+			[
+				CameraRigPropertyHandle->CreatePropertyNameWidget()
+			]
+			.ValueContent()
+			[
+				SAssignNew(ComboButton, SComboButton)
+				.ToolTipText(CameraRigPropertyHandle->GetToolTipText())
+				.ContentPadding(2.f)
+				.ButtonContent()
+				[
+					SNew(STextBlock)
+					.ColorAndOpacity(FSlateColor::UseForeground())
+					.TextStyle(FAppStyle::Get(), "PropertyEditor.AssetClass")
+					.Font(FAppStyle::GetFontStyle("PropertyWindow.NormalFont"))
+					.Text(this, &FSingleCameraDirectorDetailsCustomization::OnGetComboButtonText)
+				]
+				.OnGetMenuContent(this, &FSingleCameraDirectorDetailsCustomization::OnBuildCameraRigPicker)
+			];
+}
+
+FText FSingleCameraDirectorDetailsCustomization::OnGetComboButtonText() const
+{
+	if (USingleCameraDirector* SelectedDirector = WeakSelectedDirector.Get())
+	{
+		FText DisplayText(LOCTEXT("NullCameraRig", "None"));
+		if (SelectedDirector->CameraRig)
+		{
+			DisplayText = FText::FromString(SelectedDirector->CameraRig->GetDisplayName());
+		}
+		return DisplayText;
+	}
+	return LOCTEXT("MultipleCameraRigs", "Multiple Values");
+}
+
+TSharedRef<SWidget> FSingleCameraDirectorDetailsCustomization::OnBuildCameraRigPicker()
+{
+	USingleCameraDirector* SelectedDirector = WeakSelectedDirector.Get();
+	if (!SelectedDirector)
+	{
+		return SNullWidget::NullWidget;
+	}
+
+	IGameplayCamerasEditorModule& CamerasEditorModule = FModuleManager::LoadModuleChecked<IGameplayCamerasEditorModule>("GameplayCamerasEditor");
+
+	UCameraAsset* OuterCameraAsset = SelectedDirector->GetTypedOuter<UCameraAsset>();
+
+	FCameraRigPickerConfig CameraRigPickerConfig;
+	CameraRigPickerConfig.bCanSelectCameraAsset = false;
+	CameraRigPickerConfig.InitialCameraAssetSelection = FAssetData(OuterCameraAsset);
+	CameraRigPickerConfig.OnCameraRigSelected = FOnCameraRigSelected::CreateSP(
+			this, &FSingleCameraDirectorDetailsCustomization::OnCameraRigSelected);
+	CameraRigPickerConfig.PropertyToSet = CameraRigPropertyHandle;
+	CameraRigPickerConfig.InitialCameraRigSelection = SelectedDirector->CameraRig;
+
+	return CamerasEditorModule.CreateCameraRigPicker(CameraRigPickerConfig);
+}
+
+void FSingleCameraDirectorDetailsCustomization::OnCameraRigSelected(UCameraRigAsset* CameraRig)
+{
+	ComboButton->SetIsOpen(false);
+}
+
+}  // namespace UE::Cameras
+
+#undef LOCTEXT_NAMESPACE
+
diff --git a/Engine/Plugins/Cameras/GameplayCameras/Source/GameplayCamerasEditor/Private/GameplayCamerasEditorModule.cpp b/Engine/Plugins/Cameras/GameplayCameras/Source/GameplayCamerasEditor/Private/GameplayCamerasEditorModule.cpp
--- a/Engine/Plugins/Cameras/GameplayCameras/Source/GameplayCamerasEditor/Private/GameplayCamerasEditorModule.cpp
+++ b/Engine/Plugins/Cameras/GameplayCameras/Source/GameplayCamerasEditor/Private/GameplayCamerasEditorModule.cpp
@@ -11,6 +11,7 @@
 #include "Commands/CameraVariableCollectionEditorCommands.h"
 #include "Commands/GameplayCamerasDebuggerCommands.h"
 #include "Customizations/CameraParameterDetailsCustomizations.h"
+#include "Customizations/SingleCameraDirectorDetailsCustomization.h"
 #include "Debug/CameraDebugCategories.h"
 #include "Debugger/SBlendStacksDebugPanel.h"
 #include "Debugger/SCameraNodeTreeDebugPanel.h"
@@ -325,6 +326,8 @@ private:
 
 		FPropertyEditorModule& PropertyEditorModule = FModuleManager::GetModuleChecked<FPropertyEditorModule>("PropertyEditor");
 		FCameraParameterDetailsCustomization::Register(PropertyEditorModule);
+		PropertyEditorModule.RegisterCustomClassLayout("SingleCameraDirector", FOnGetDetailCustomizationInstance::CreateStatic(
+					&FSingleCameraDirectorDetailsCustomization::MakeInstance));
 	}
 
 	void UnregisterDetailsCustomizations()
@@ -336,6 +339,7 @@ private:
 		if (PropertyEditorModule)
 		{
 			FCameraParameterDetailsCustomization::Unregister(*PropertyEditorModule);
+			PropertyEditorModule->UnregisterCustomClassLayout("SingleCameraDirector");
 		}
 	}
 
diff --git a/Engine/Plugins/Cameras/GameplayCameras/Source/GameplayCamerasEditor/Public/Customizations/SingleCameraDirectorDetailsCustomization.h b/Engine/Plugins/Cameras/GameplayCameras/Source/GameplayCamerasEditor/Public/Customizations/SingleCameraDirectorDetailsCustomization.h
new file mode 100644
--- /dev/null
+++ b/Engine/Plugins/Cameras/GameplayCameras/Source/GameplayCamerasEditor/Public/Customizations/SingleCameraDirectorDetailsCustomization.h
@@ -0,0 +1,47 @@
+// Copyright Epic Games, Inc. All Rights Reserved.
+
+#pragma once
+
+#include "IDetailCustomization.h"
+#include "UObject/WeakObjectPtr.h"
+
+class IPropertyHandle;
+class SComboButton;
+class SWidget;
+class UCameraRigAsset;
+class USingleCameraDirector;
+
+namespace UE::Cameras
+{
+
+/**
+ * Details customization for the single-rig camera director.
+ *
+ * It shows a custom camera rig picker that only chooses among the rigs of the parent camera asset.
+ */
+class FSingleCameraDirectorDetailsCustomization : public IDetailCustomization
+{
+public:
+
+	static TSharedRef<IDetailCustomization> MakeInstance();
+
+public:
+
+	/** IDetailCustomization interface */
+	virtual void CustomizeDetails(IDetailLayoutBuilder& DetailBuilder) override;
+
+private:
+
+	TSharedRef<SWidget> OnBuildCameraRigPicker();
+	FText OnGetComboButtonText() const;
+	void OnCameraRigSelected(UCameraRigAsset* CameraRig);
+
+private:
+
+	TSharedPtr<SComboButton> ComboButton;
+	TSharedPtr<IPropertyHandle> CameraRigPropertyHandle;
+	TWeakObjectPtr<USingleCameraDirector> WeakSelectedDirector;
+};
+
+}  // namespace UE::Cameras
+
