Cameras: let camera directors participate in the build stage

diff --git a/Engine/Plugins/Cameras/GameplayCameras/Source/GameplayCameras/Private/Core/CameraAssetBuilder.cpp b/Engine/Plugins/Cameras/GameplayCameras/Source/GameplayCameras/Private/Core/CameraAssetBuilder.cpp
--- a/Engine/Plugins/Cameras/GameplayCameras/Source/GameplayCameras/Private/Core/CameraAssetBuilder.cpp
+++ b/Engine/Plugins/Cameras/GameplayCameras/Source/GameplayCameras/Private/Core/CameraAssetBuilder.cpp
@@ -2,13 +2,11 @@
 
 #include "Core/CameraAssetBuilder.h"
 
-#include "Core/CameraNode.h"
+#include "Core/CameraDirector.h"
 #include "Core/CameraNodeEvaluatorStorage.h"
-#include "Core/CameraParameters.h"
 #include "Core/CameraAsset.h"
 #include "Core/CameraRigAsset.h"
 #include "Core/CameraRigAssetBuilder.h"
-#include "Core/CameraVariableAssets.h"
 #include "Logging/TokenizedMessage.h"
 
 #define LOCTEXT_NAMESPACE "CameraAssetBuilder"
@@ -39,7 +37,11 @@ void FCameraAssetBuilder::BuildCamera(UCameraAsset* InCameraAsset)
 
 void FCameraAssetBuilder::BuildCameraImpl()
 {
-	if (!CameraAsset->CameraDirector)
+	if (CameraAsset->CameraDirector)
+	{
+		CameraAsset->CameraDirector->BuildCameraDirector(BuildLog);
+	}
+	else
 	{
 		BuildLog.AddMessage(EMessageSeverity::Error, LOCTEXT("MissingDirector", "Camera has no director set."));
 	}
diff --git a/Engine/Plugins/Cameras/GameplayCameras/Source/GameplayCameras/Private/Core/CameraDirector.cpp b/Engine/Plugins/Cameras/GameplayCameras/Source/GameplayCameras/Private/Core/CameraDirector.cpp
--- a/Engine/Plugins/Cameras/GameplayCameras/Source/GameplayCameras/Private/Core/CameraDirector.cpp
+++ b/Engine/Plugins/Cameras/GameplayCameras/Source/GameplayCameras/Private/Core/CameraDirector.cpp
@@ -12,3 +12,7 @@ FCameraDirectorEvaluatorPtr UCameraDirector::BuildEvaluator(FCameraDirectorEvalu
 	return NewEvaluator;
 }
 
+void UCameraDirector::BuildCameraDirector(UE::Cameras::FCameraBuildLog& BuildLog)
+{
+	OnBuildCameraDirector(BuildLog);
+}
diff --git a/Engine/Plugins/Cameras/GameplayCameras/Source/GameplayCameras/Private/Directors/SingleCameraDirector.cpp b/Engine/Plugins/Cameras/GameplayCameras/Source/GameplayCameras/Private/Directors/SingleCameraDirector.cpp
--- a/Engine/Plugins/Cameras/GameplayCameras/Source/GameplayCameras/Private/Directors/SingleCameraDirector.cpp
+++ b/Engine/Plugins/Cameras/GameplayCameras/Source/GameplayCameras/Private/Directors/SingleCameraDirector.cpp
@@ -2,8 +2,13 @@
 
 #include "Directors/SingleCameraDirector.h"
 
+#include "Core/CameraBuildLog.h"
+#include "Logging/TokenizedMessage.h"
+
 #include UE_INLINE_GENERATED_CPP_BY_NAME(SingleCameraDirector)
 
+#define LOCTEXT_NAMESPACE "SingleCameraDirector"
+
 namespace UE::Cameras
 {
 
@@ -36,3 +41,13 @@ FCameraDirectorEvaluatorPtr USingleCameraDirector::OnBuildEvaluator(FCameraDirec
 	return Builder.BuildEvaluator<FSingleCameraDirectorEvaluator>();
 }
 
+void USingleCameraDirector::OnBuildCameraDirector(UE::Cameras::FCameraBuildLog& BuildLog)
+{
+	if (!CameraRig)
+	{
+		BuildLog.AddMessage(EMessageSeverity::Error, this, LOCTEXT("MissingCameraRig", "No camera rig is set."));
+	}
+}
+
+#undef LOCTEXT_NAMESPACE
+
diff --git a/Engine/Plugins/Cameras/GameplayCameras/Source/GameplayCameras/Public/Core/CameraDirector.h b/Engine/Plugins/Cameras/GameplayCameras/Source/GameplayCameras/Public/Core/CameraDirector.h
--- a/Engine/Plugins/Cameras/GameplayCameras/Source/GameplayCameras/Public/Core/CameraDirector.h
+++ b/Engine/Plugins/Cameras/GameplayCameras/Source/GameplayCameras/Public/Core/CameraDirector.h
@@ -8,6 +8,8 @@
 
 #include "CameraDirector.generated.h"
 
+namespace UE::Cameras { class FCameraBuildLog; }
+
 /**
  * Base class for a camera director.
  */
@@ -23,9 +25,15 @@ public:
 	/** Build the evaluator for this director. */
 	FCameraDirectorEvaluatorPtr BuildEvaluator(FCameraDirectorEvaluatorBuilder& Builder) const;
 
+	/** Builds and validates this camera director. */
+	void BuildCameraDirector(UE::Cameras::FCameraBuildLog& BuildLog);
+
 protected:
 
 	/** Build the evaluator for this director. */
 	virtual FCameraDirectorEvaluatorPtr OnBuildEvaluator(FCameraDirectorEvaluatorBuilder& Builder) const { return nullptr; }
+
+	/** Builds and validates this camera director. */
+	virtual void OnBuildCameraDirector(UE::Cameras::FCameraBuildLog& BuildLog) {}
 };
 
diff --git a/Engine/Plugins/Cameras/GameplayCameras/Source/GameplayCameras/Public/Directors/SingleCameraDirector.h b/Engine/Plugins/Cameras/GameplayCameras/Source/GameplayCameras/Public/Directors/SingleCameraDirector.h
--- a/Engine/Plugins/Cameras/GameplayCameras/Source/GameplayCameras/Public/Directors/SingleCameraDirector.h
+++ b/Engine/Plugins/Cameras/GameplayCameras/Source/GameplayCameras/Public/Directors/SingleCameraDirector.h
@@ -22,6 +22,7 @@ protected:
 
 	// UCameraDirector interface.
 	virtual FCameraDirectorEvaluatorPtr OnBuildEvaluator(FCameraDirectorEvaluatorBuilder& Builder) const override;
+	virtual void OnBuildCameraDirector(UE::Cameras::FCameraBuildLog& BuildLog) override;
 
 public:
 
