Cameras: add scaffolding step to creating camera assets

- Create one camera rig to get the user started
- The single camera director auto-sets its camera rig to that auto-created rig

diff --git a/Engine/Plugins/Cameras/GameplayCameras/Source/GameplayCameras/Private/Core/CameraDirector.cpp b/Engine/Plugins/Cameras/GameplayCameras/Source/GameplayCameras/Private/Core/CameraDirector.cpp
--- a/Engine/Plugins/Cameras/GameplayCameras/Source/GameplayCameras/Private/Core/CameraDirector.cpp
+++ b/Engine/Plugins/Cameras/GameplayCameras/Source/GameplayCameras/Private/Core/CameraDirector.cpp
@@ -16,3 +16,13 @@ void UCameraDirector::BuildCameraDirector(UE::Cameras::FCameraBuildLog& BuildLog
 {
 	OnBuildCameraDirector(BuildLog);
 }
+
+#if WITH_EDITOR
+
+void UCameraDirector::FactoryCreateAsset(const FCameraDirectorFactoryCreateParams& InParams)
+{
+	OnFactoryCreateAsset(InParams);
+}
+
+#endif
+
diff --git a/Engine/Plugins/Cameras/GameplayCameras/Source/GameplayCameras/Private/Directors/SingleCameraDirector.cpp b/Engine/Plugins/Cameras/GameplayCameras/Source/GameplayCameras/Private/Directors/SingleCameraDirector.cpp
--- a/Engine/Plugins/Cameras/GameplayCameras/Source/GameplayCameras/Private/Directors/SingleCameraDirector.cpp
+++ b/Engine/Plugins/Cameras/GameplayCameras/Source/GameplayCameras/Private/Directors/SingleCameraDirector.cpp
@@ -2,6 +2,7 @@
 
 #include "Directors/SingleCameraDirector.h"
 
+#include "Core/CameraAsset.h"
 #include "Core/CameraBuildLog.h"
 #include "Logging/TokenizedMessage.h"
 
@@ -49,5 +50,21 @@ void USingleCameraDirector::OnBuildCameraDirector(UE::Cameras::FCameraBuildLog&
 	}
 }
 
+#if WITH_EDITOR
+
+void USingleCameraDirector::OnFactoryCreateAsset(const FCameraDirectorFactoryCreateParams& InParams)
+{
+	// Automatically set ourselves to use the first camera rig available.
+	if (UCameraAsset* OuterCameraAsset = GetTypedOuter<UCameraAsset>())
+	{
+		if (OuterCameraAsset->CameraRigs.Num() > 0)
+		{
+			CameraRig = OuterCameraAsset->CameraRigs[0];
+		}
+	}
+}
+
+#endif
+
 #undef LOCTEXT_NAMESPACE
 
diff --git a/Engine/Plugins/Cameras/GameplayCameras/Source/GameplayCameras/Public/Core/CameraDirector.h b/Engine/Plugins/Cameras/GameplayCameras/Source/GameplayCameras/Public/Core/CameraDirector.h
--- a/Engine/Plugins/Cameras/GameplayCameras/Source/GameplayCameras/Public/Core/CameraDirector.h
+++ b/Engine/Plugins/Cameras/GameplayCameras/Source/GameplayCameras/Public/Core/CameraDirector.h
@@ -10,6 +10,18 @@
 
 namespace UE::Cameras { class FCameraBuildLog; }
 
+#if WITH_EDITOR
+
+/**
+ * Parameter struct passed by an asset factory when a new camera asset is created.
+ * This lets a camera director setup data before the editor opens.
+ */
+struct FCameraDirectorFactoryCreateParams
+{
+};
+
+#endif
+
 /**
  * Base class for a camera director.
  */
@@ -28,6 +40,11 @@ public:
 	/** Builds and validates this camera director. */
 	void BuildCameraDirector(UE::Cameras::FCameraBuildLog& BuildLog);
 
+#if WITH_EDITOR
+	/** Called by the asset factories to setup new data before the editor opens. */
+	void FactoryCreateAsset(const FCameraDirectorFactoryCreateParams& InParams);
+#endif
+
 protected:
 
 	/** Build the evaluator for this director. */
@@ -35,5 +52,10 @@ protected:
 
 	/** Builds and validates this camera director. */
 	virtual void OnBuildCameraDirector(UE::Cameras::FCameraBuildLog& BuildLog) {}
+
+#if WITH_EDITOR
+	/** Called by the asset factories to setup new data before the editor opens. */
+	virtual void OnFactoryCreateAsset(const FCameraDirectorFactoryCreateParams& InParams) {}
+#endif
 };
 
diff --git a/Engine/Plugins/Cameras/GameplayCameras/Source/GameplayCameras/Public/Directors/SingleCameraDirector.h b/Engine/Plugins/Cameras/GameplayCameras/Source/GameplayCameras/Public/Directors/SingleCameraDirector.h
--- a/Engine/Plugins/Cameras/GameplayCameras/Source/GameplayCameras/Public/Directors/SingleCameraDirector.h
+++ b/Engine/Plugins/Cameras/GameplayCameras/Source/GameplayCameras/Public/Directors/SingleCameraDirector.h
@@ -23,6 +23,9 @@ protected:
 	// UCameraDirector interface.
 	virtual FCameraDirectorEvaluatorPtr OnBuildEvaluator(FCameraDirectorEvaluatorBuilder& Builder) const override;
 	virtual void OnBuildCameraDirector(UE::Cameras::FCameraBuildLog& BuildLog) override;
+#if WITH_EDITOR
+	virtual void OnFactoryCreateAsset(const FCameraDirectorFactoryCreateParams& InParams) override;
+#endif
 
 public:
 
diff --git a/Engine/Plugins/Cameras/GameplayCameras/Source/GameplayCamerasEditor/Private/Factories/CameraAssetFactory.cpp b/Engine/Plugins/Cameras/GameplayCameras/Source/GameplayCamerasEditor/Private/Factories/CameraAssetFactory.cpp
--- a/Engine/Plugins/Cameras/GameplayCameras/Source/GameplayCamerasEditor/Private/Factories/CameraAssetFactory.cpp
+++ b/Engine/Plugins/Cameras/GameplayCameras/Source/GameplayCamerasEditor/Private/Factories/CameraAssetFactory.cpp
@@ -4,6 +4,7 @@
 
 #include "Core/CameraAsset.h"
 #include "Core/CameraDirector.h"
+#include "Core/CameraRigAsset.h"
 #include "ClassViewerFilter.h"
 #include "ClassViewerModule.h"
 #include "Directors/BlueprintCameraDirector.h"
@@ -58,11 +59,20 @@ UCameraAssetFactory::UCameraAssetFactory(const FObjectInitializer& ObjectInitial
 UObject* UCameraAssetFactory::FactoryCreateNew(UClass* Class, UObject* Parent, FName Name, EObjectFlags Flags, UObject* Context, FFeedbackContext* Warn)
 {
 	UCameraAsset* NewCameraAsset = NewObject<UCameraAsset>(Parent, Class, Name, Flags | RF_Transactional);
+
+	UCameraRigAsset* FirstCameraRig = NewObject<UCameraRigAsset>(NewCameraAsset, TEXT("NewCameraRig"));
+	NewCameraAsset->CameraRigs.Add(FirstCameraRig);
+
 	if (CameraDirectorClass)
 	{
 		UCameraDirector* NewCameraDirector = NewObject<UCameraDirector>(NewCameraAsset, CameraDirectorClass, NAME_None, Flags | RF_Transactional);
 		NewCameraAsset->CameraDirector = NewCameraDirector;
+
+		// Let the camera director do some scaffolding.
+		FCameraDirectorFactoryCreateParams CreateParams;
+		NewCameraDirector->FactoryCreateAsset(CreateParams);
 	}
+
 	return NewCameraAsset;
 }
 
