Cameras: change how the Blueprint director evaluator specifies rigs to activate by introducing proxies and a proxy table

This change removes the previous name-based activation system, and introduces a system that is inspired by the animation system's way of handling a similar problem.

- The user can create "camera rig proxy" assets that stand-in for yet-undefined rigs.
- The Blueprint funtion "ActivateCameraRig" now takes one such proxy.
- The Blueprint camera director has a "proxy table" where the user can map proxies to actual rigs contained in the camera asset.

Of course, this change includes all the boilerplate code necessary to add these new objects (proxies, proxy table, etc), including asset factories, toolkits, details customizations, and more.

diff --git a/Engine/Plugins/Cameras/GameplayCameras/Source/GameplayCameras/Private/Core/CameraRigProxyAsset.cpp b/Engine/Plugins/Cameras/GameplayCameras/Source/GameplayCameras/Private/Core/CameraRigProxyAsset.cpp
new file mode 100644
--- /dev/null
+++ b/Engine/Plugins/Cameras/GameplayCameras/Source/GameplayCameras/Private/Core/CameraRigProxyAsset.cpp
@@ -0,0 +1,26 @@
+// Copyright Epic Games, Inc. All Rights Reserved.
+
+#include "Core/CameraRigProxyAsset.h"
+
+#include "Misc/Guid.h"
+
+UCameraRigProxyAsset::UCameraRigProxyAsset(const FObjectInitializer& ObjectInit)
+	: Super(ObjectInit)
+{
+	Guid = FGuid::NewGuid();
+}
+
+#if WITH_EDITORONLY_DATA
+
+void UCameraRigProxyAsset::PostDuplicate(EDuplicateMode::Type DuplicateMode)
+{
+	Super::PostDuplicate(DuplicateMode);
+
+	if (DuplicateMode == EDuplicateMode::Normal)
+	{
+		Guid = FGuid::NewGuid();
+	}
+}
+
+#endif
+
diff --git a/Engine/Plugins/Cameras/GameplayCameras/Source/GameplayCameras/Private/Core/CameraRigProxyTable.cpp b/Engine/Plugins/Cameras/GameplayCameras/Source/GameplayCameras/Private/Core/CameraRigProxyTable.cpp
new file mode 100644
--- /dev/null
+++ b/Engine/Plugins/Cameras/GameplayCameras/Source/GameplayCameras/Private/Core/CameraRigProxyTable.cpp
@@ -0,0 +1,23 @@
+// Copyright Epic Games, Inc. All Rights Reserved.
+
+#include "Core/CameraRigProxyTable.h"
+
+UCameraRigProxyTable::UCameraRigProxyTable(const FObjectInitializer& ObjectInit)
+	: Super(ObjectInit)
+{
+}
+
+UCameraRigAsset* UCameraRigProxyTable::ResolveProxy(const FCameraRigProxyTableResolveParams& InParams) const
+{
+	ensure(InParams.CameraRigProxy);
+
+	for (const FCameraRigProxyTableEntry& Entry : Entries)
+	{
+		if (InParams.CameraRigProxy == Entry.CameraRigProxy)
+		{
+			return Entry.CameraRig;
+		}
+	}
+	return nullptr;
+}
+
diff --git a/Engine/Plugins/Cameras/GameplayCameras/Source/GameplayCameras/Private/Directors/BlueprintCameraDirector.cpp b/Engine/Plugins/Cameras/GameplayCameras/Source/GameplayCameras/Private/Directors/BlueprintCameraDirector.cpp
--- a/Engine/Plugins/Cameras/GameplayCameras/Source/GameplayCameras/Private/Directors/BlueprintCameraDirector.cpp
+++ b/Engine/Plugins/Cameras/GameplayCameras/Source/GameplayCameras/Private/Directors/BlueprintCameraDirector.cpp
@@ -3,12 +3,17 @@
 #include "Directors/BlueprintCameraDirector.h"
 
 #include "Core/CameraAsset.h"
+#include "Core/CameraBuildLog.h"
 #include "Core/CameraRigAsset.h"
+#include "Core/CameraRigProxyAsset.h"
+#include "Core/CameraRigProxyTable.h"
 #include "Core/CameraEvaluationContext.h"
 #include "GameplayCameras.h"
 
 #include UE_INLINE_GENERATED_CPP_BY_NAME(BlueprintCameraDirector)
 
+#define LOCTEXT_NAMESPACE "BlueprintCameraDirector"
+
 namespace UE::Cameras
 {
 
@@ -24,7 +29,7 @@ protected:
 
 private:
 
-	const UCameraRigAsset* FindCameraRigByName(const UCameraAsset* InCameraAsset, const FString& InCameraRigName);
+	const UCameraRigAsset* FindCameraRigByProxy(const UCameraRigProxyAsset* InProxy);
 
 private:
 
@@ -60,8 +65,7 @@ void FBlueprintCameraDirectorEvaluator::OnInitialize(const FCameraDirectorInitia
 
 void FBlueprintCameraDirectorEvaluator::OnRun(const FCameraDirectorEvaluationParams& Params, FCameraDirectorEvaluationResult& OutResult)
 {
-	const UCameraAsset* CameraAsset = Params.OwnerContext->GetCameraAsset();
-	if (EvaluatorBlueprint && CameraAsset)
+	if (EvaluatorBlueprint)
 	{
 		FBlueprintCameraDirectorEvaluationParams BlueprintParams;
 		BlueprintParams.DeltaTime = Params.DeltaTime;
@@ -76,39 +80,47 @@ void FBlueprintCameraDirectorEvaluator::OnRun(const FCameraDirectorEvaluationPar
 
 		// The BP interface doesn't specify the evaluation context for the chosen camera rigs: we always automatically
 		// make them run in our own owner context.
-		for (const FString& ActiveCameraRigName : BlueprintResult.ActiveCameraRigs)
+		for (const UCameraRigProxyAsset* ActiveCameraRigProxy : BlueprintResult.ActiveCameraRigs)
 		{
-			const UCameraRigAsset* ActiveCameraRig = FindCameraRigByName(CameraAsset, ActiveCameraRigName);
+			const UCameraRigAsset* ActiveCameraRig = FindCameraRigByProxy(ActiveCameraRigProxy);
 			if (ActiveCameraRig)
 			{
 				OutResult.Add(Params.OwnerContext, ActiveCameraRig);
 			}
 			else
 			{
+				const UCameraAsset* CameraAsset = Params.OwnerContext->GetCameraAsset();
 				UE_LOG(
 						LogCameraSystem, 
 						Error, 
-						TEXT("Can't activate camera rig '%s' because no camera rig of that name was found on '%s'."),
-						*ActiveCameraRigName, *CameraAsset->GetPathName());
+						TEXT("No camera rig found mapped to proxy '%s' in camera '%s'."),
+						*ActiveCameraRigProxy->GetPathName(), *CameraAsset->GetPathName());
 			}
 		}
 	}
 	else
 	{
-		UE_LOG(LogCameraSystem, Error, TEXT("Can't run Blueprint camera director, no Blueprint set, or no owner camera set!"));
+		UE_LOG(LogCameraSystem, Error, TEXT("Can't run Blueprint camera director, no Blueprint class was set!"));
 	}
 }
 
-const UCameraRigAsset* FBlueprintCameraDirectorEvaluator::FindCameraRigByName(const UCameraAsset* InCameraAsset, const FString& InCameraRigName)
+const UCameraRigAsset* FBlueprintCameraDirectorEvaluator::FindCameraRigByProxy(const UCameraRigProxyAsset* InProxy)
 {
-	for (const UCameraRigAsset* CameraRig : InCameraAsset->CameraRigs)
+	const UBlueprintCameraDirector* Blueprint = GetCameraDirectorAs<UBlueprintCameraDirector>();
+	if (!ensure(Blueprint))
 	{
-		if (CameraRig->GetName() == InCameraRigName || CameraRig->Interface.DisplayName == InCameraRigName)
-		{
-			return CameraRig;
-		}
+		return nullptr;
 	}
-	return nullptr;
+
+	UCameraRigProxyTable* ProxyTable = Blueprint->CameraRigProxyTable;
+	if (!ensureMsgf(ProxyTable, TEXT("No proxy table set on Blueprint director '%s'."), *Blueprint->GetPathName()))
+	{
+		return nullptr;
+	}
+
+	FCameraRigProxyTableResolveParams ResolveParams;
+	ResolveParams.CameraRigProxy = InProxy;
+	return Blueprint->CameraRigProxyTable->ResolveProxy(ResolveParams);
 }
 
 void FBlueprintCameraDirectorEvaluator::OnAddReferencedObjects(FReferenceCollector& Collector)
@@ -118,9 +130,9 @@ void FBlueprintCameraDirectorEvaluator::OnAddReferencedObjects(FReferenceCollect
 
 }  // namespace UE::Cameras
 
-void UBlueprintCameraDirectorEvaluator::ActivateCameraRig(const FString& InCameraRigName)
+void UBlueprintCameraDirectorEvaluator::ActivateCameraRig(UCameraRigProxyAsset* CameraRigProxy)
 {
-	CurrentResult.ActiveCameraRigs.Add(InCameraRigName);
+	CurrentResult.ActiveCameraRigs.Add(CameraRigProxy);
 }
 
 void UBlueprintCameraDirectorEvaluator::NativeRunCameraDirector(const FBlueprintCameraDirectorEvaluationParams& Params, FBlueprintCameraDirectorEvaluationResult& OutResult)
@@ -140,3 +152,28 @@ FCameraDirectorEvaluatorPtr UBlueprintCameraDirector::OnBuildEvaluator(FCameraDi
 	return Builder.BuildEvaluator<FBlueprintCameraDirectorEvaluator>();
 }
 
+void UBlueprintCameraDirector::OnBuildCameraDirector(UE::Cameras::FCameraBuildLog& BuildLog)
+{
+	if (!CameraDirectorEvaluatorClass)
+	{
+		BuildLog.AddMessage(EMessageSeverity::Error, LOCTEXT("MissingBlueprintClass", "No evaluator Blueprint class is set."));
+		return;
+	}
+
+	// TODO: check that the proxy table is complete.
+}
+
+#if WITH_EDITOR
+
+void UBlueprintCameraDirector::OnFactoryCreateAsset(const FCameraDirectorFactoryCreateParams& InParams)
+{
+	if (!CameraRigProxyTable)
+	{
+		CameraRigProxyTable = NewObject<UCameraRigProxyTable>(this);
+	}
+}
+
+#endif
+
+#undef LOCTEXT_NAMESPACE
+
diff --git a/Engine/Plugins/Cameras/GameplayCameras/Source/GameplayCameras/Public/Core/CameraRigProxyAsset.h b/Engine/Plugins/Cameras/GameplayCameras/Source/GameplayCameras/Public/Core/CameraRigProxyAsset.h
new file mode 100644
--- /dev/null
+++ b/Engine/Plugins/Cameras/GameplayCameras/Source/GameplayCameras/Public/Core/CameraRigProxyAsset.h
@@ -0,0 +1,39 @@
+// Copyright Epic Games, Inc. All Rights Reserved.
+
+#pragma once
+
+#include "CoreTypes.h"
+#include "UObject/Object.h"
+
+#include "CameraRigProxyAsset.generated.h"
+
+/**
+ * A proxy for a camera rig asset.
+ *
+ * This is useful for indicating that a camera rig should be activated in a camera director,
+ * but without hard-referencing a particular camera rig. This way, that camera director can
+ * be used by multiple cameras with their own camera rigs.
+ */
+UCLASS(MinimalAPI)
+class UCameraRigProxyAsset : public UObject
+{
+	GENERATED_BODY()
+
+public:
+
+	UCameraRigProxyAsset(const FObjectInitializer& ObjectInit);
+
+public:
+
+	// UObject interface.
+#if WITH_EDITORONLY_DATA
+	virtual void PostDuplicate(EDuplicateMode::Type DuplicateMode) override;
+#endif
+
+public:
+
+	/** Unique identifier for this camera rig proxy. */
+	UPROPERTY()
+	FGuid Guid;
+};
+
diff --git a/Engine/Plugins/Cameras/GameplayCameras/Source/GameplayCameras/Public/Core/CameraRigProxyTable.h b/Engine/Plugins/Cameras/GameplayCameras/Source/GameplayCameras/Public/Core/CameraRigProxyTable.h
new file mode 100644
--- /dev/null
+++ b/Engine/Plugins/Cameras/GameplayCameras/Source/GameplayCameras/Public/Core/CameraRigProxyTable.h
@@ -0,0 +1,63 @@
+// Copyright Epic Games, Inc. All Rights Reserved.
+
+#pragma once
+
+#include "CoreTypes.h"
+#include "UObject/Object.h"
+
+#include "CameraRigProxyTable.generated.h"
+
+class UCameraRigAsset;
+class UCameraRigProxyAsset;
+
+/**
+ * Parameter structure for resolving a camera rig proxy from a proxy table.
+ */
+struct FCameraRigProxyTableResolveParams
+{
+	/** The camera rig proxy to resolve. */
+	const UCameraRigProxyAsset* CameraRigProxy = nullptr;
+};
+
+/**
+ * An entry in a camera rig proxy table.
+ */
+USTRUCT()
+struct FCameraRigProxyTableEntry
+{
+	GENERATED_BODY()
+
+	/** The camera rig proxy for this table entry. */
+	UPROPERTY(EditAnywhere, Category="Camera")
+	TObjectPtr<UCameraRigProxyAsset> CameraRigProxy;
+
+	/** The actual camera rig that should be mapped to the correspondig proxy. */
+	UPROPERTY(EditAnywhere, Category="Camera")
+	TObjectPtr<UCameraRigAsset> CameraRig;
+};
+
+/**
+ * A table that defines mappings between camera rig proxies and actual camera rigs.
+ */
+UCLASS(MinimalAPI, EditInlineNew)
+class UCameraRigProxyTable : public UObject
+{
+	GENERATED_BODY()
+
+public:
+
+	UCameraRigProxyTable(const FObjectInitializer& ObjectInit);
+
+	/**
+	 * Resolves a given proxy to an actual camera rig.
+	 * Returns nullptr if the given proxy wasn't found, or not mapped to anything in the table.
+	 */
+	UCameraRigAsset* ResolveProxy(const FCameraRigProxyTableResolveParams& InParams) const;
+
+public:
+
+	/** The entries in the table. */
+	UPROPERTY(EditAnywhere, Category="Camera")
+	TArray<FCameraRigProxyTableEntry> Entries;
+};
+
diff --git a/Engine/Plugins/Cameras/GameplayCameras/Source/GameplayCameras/Public/Directors/BlueprintCameraDirector.h b/Engine/Plugins/Cameras/GameplayCameras/Source/GameplayCameras/Public/Directors/BlueprintCameraDirector.h
--- a/Engine/Plugins/Cameras/GameplayCameras/Source/GameplayCameras/Public/Directors/BlueprintCameraDirector.h
+++ b/Engine/Plugins/Cameras/GameplayCameras/Source/GameplayCameras/Public/Directors/BlueprintCameraDirector.h
@@ -8,6 +8,9 @@
 
 #include "BlueprintCameraDirector.generated.h"
 
+class UCameraRigProxyAsset;
+class UCameraRigProxyTable;
+
 /**
  * Parameter struct for the Blueprint camera director evaluator.
  */
@@ -33,9 +36,9 @@ struct FBlueprintCameraDirectorEvaluationResult
 {
 	GENERATED_BODY()
 
-	/** The list of camera rigs (determined by name) that should be active this frame. */
+	/** The list of camera rigs that should be active this frame. */
 	UPROPERTY(BlueprintReadWrite, Category="Evaluation")
-	TArray<FString> ActiveCameraRigs;
+	TArray<UCameraRigProxyAsset*> ActiveCameraRigs;
 };
 
 /**
@@ -59,7 +62,7 @@ public:
 	 * Specifies a camera rig to be active this frame.
 	 */
 	UFUNCTION(BlueprintCallable, Category="Evaluation")
-	void ActivateCameraRig(UPARAM(meta=(UseCameraRigNamePicker="true")) const FString& InCameraRigName);
+	void ActivateCameraRig(UCameraRigProxyAsset* CameraRigProxy);
 
 	/** Native wrapper for RunCameraDirector. */
 	void NativeRunCameraDirector(
@@ -88,9 +91,20 @@ public:
 	UPROPERTY(EditAnywhere, Category="Evaluation")
 	TSubclassOf<UBlueprintCameraDirectorEvaluator> CameraDirectorEvaluatorClass;
 
+	/** 
+	 * The table that maps camera rig proxies (used in the evaluator Blueprint graph)
+	 * to actual camera rigs.
+	 */
+	UPROPERTY(EditAnywhere, Instanced, Category="Evaluation")
+	TObjectPtr<UCameraRigProxyTable> CameraRigProxyTable;
+
 protected:
 
 	// UCameraDirector interface.
 	virtual FCameraDirectorEvaluatorPtr OnBuildEvaluator(FCameraDirectorEvaluatorBuilder& Builder) const override;
+	virtual void OnBuildCameraDirector(UE::Cameras::FCameraBuildLog& BuildLog) override;
+#if WITH_EDITOR
+	virtual void OnFactoryCreateAsset(const FCameraDirectorFactoryCreateParams& InParams) override;
+#endif
 };
 
diff --git a/Engine/Plugins/Cameras/GameplayCameras/Source/GameplayCamerasEditor/Private/AssetTools/AssetDefinition_CameraRigProxyAsset.cpp b/Engine/Plugins/Cameras/GameplayCameras/Source/GameplayCamerasEditor/Private/AssetTools/AssetDefinition_CameraRigProxyAsset.cpp
new file mode 100644
--- /dev/null
+++ b/Engine/Plugins/Cameras/GameplayCameras/Source/GameplayCamerasEditor/Private/AssetTools/AssetDefinition_CameraRigProxyAsset.cpp
@@ -0,0 +1,53 @@
+// Copyright Epic Games, Inc. All Rights Reserved.
+
+#include "AssetTools/AssetDefinition_CameraRigProxyAsset.h"
+
+#include "IGameplayCamerasEditorModule.h"
+#include "Toolkits/IToolkit.h"
+#include "Modules/ModuleManager.h"
+
+#define LOCTEXT_NAMESPACE "AssetDefinition_CameraRigProxyAsset"
+
+#include UE_INLINE_GENERATED_CPP_BY_NAME(AssetDefinition_CameraRigProxyAsset)
+
+FText UAssetDefinition_CameraRigProxyAsset::GetAssetDisplayName() const
+{
+	return LOCTEXT("AssetDisplayName", "Camera Rig Proxy");
+}
+
+FLinearColor UAssetDefinition_CameraRigProxyAsset::GetAssetColor() const
+{
+	return FLinearColor(FColor(200, 80, 80));
+}
+
+TSoftClassPtr<UObject> UAssetDefinition_CameraRigProxyAsset::GetAssetClass() const
+{
+	return UCameraRigProxyAsset::StaticClass();
+}
+
+TConstArrayView<FAssetCategoryPath> UAssetDefinition_CameraRigProxyAsset::GetAssetCategories() const
+{
+	static const auto Categories = { FAssetCategoryPath(EAssetCategoryPaths::Gameplay) };
+	return Categories;
+}
+
+FAssetOpenSupport UAssetDefinition_CameraRigProxyAsset::GetAssetOpenSupport(const FAssetOpenSupportArgs& OpenSupportArgs) const
+{
+	return FAssetOpenSupport(
+			OpenSupportArgs.OpenMethod,
+			OpenSupportArgs.OpenMethod == EAssetOpenMethod::Edit,
+			EToolkitMode::Standalone); 
+}
+
+EAssetCommandResult UAssetDefinition_CameraRigProxyAsset::OpenAssets(const FAssetOpenArgs& OpenArgs) const
+{
+	IGameplayCamerasEditorModule& GameplayCamerasEditorModule = FModuleManager::LoadModuleChecked<IGameplayCamerasEditorModule>("GameplayCamerasEditor");
+	for (UCameraRigProxyAsset* CameraRigProxyAsset : OpenArgs.LoadObjects<UCameraRigProxyAsset>())
+	{
+		GameplayCamerasEditorModule.CreateCameraRigProxyEditor(OpenArgs.GetToolkitMode(), OpenArgs.ToolkitHost, CameraRigProxyAsset);
+	}	
+	return EAssetCommandResult::Handled;
+}
+
+#undef LOCTEXT_NAMESPACE
+
diff --git a/Engine/Plugins/Cameras/GameplayCameras/Source/GameplayCamerasEditor/Private/AssetTools/AssetDefinition_CameraRigProxyAsset.h b/Engine/Plugins/Cameras/GameplayCameras/Source/GameplayCamerasEditor/Private/AssetTools/AssetDefinition_CameraRigProxyAsset.h
new file mode 100644
--- /dev/null
+++ b/Engine/Plugins/Cameras/GameplayCameras/Source/GameplayCamerasEditor/Private/AssetTools/AssetDefinition_CameraRigProxyAsset.h
@@ -0,0 +1,25 @@
+// Copyright Epic Games, Inc. All Rights Reserved.
+
+#pragma once
+
+#include "AssetDefinitionDefault.h"
+#include "Core/CameraRigProxyAsset.h"
+
+#include "AssetDefinition_CameraRigProxyAsset.generated.h"
+
+UCLASS()
+class UAssetDefinition_CameraRigProxyAsset : public UAssetDefinitionDefault
+{
+	GENERATED_BODY()
+
+public:
+
+	// UAssetDefinition interface
+	virtual FText GetAssetDisplayName() const override;
+	virtual FLinearColor GetAssetColor() const override;
+	virtual TSoftClassPtr<UObject> GetAssetClass() const override;
+	virtual TConstArrayView<FAssetCategoryPath> GetAssetCategories() const override;
+	virtual FAssetOpenSupport GetAssetOpenSupport(const FAssetOpenSupportArgs& OpenSupportArgs) const override;
+	virtual EAssetCommandResult OpenAssets(const FAssetOpenArgs& OpenArgs) const override;
+};
+
diff --git a/Engine/Plugins/Cameras/GameplayCameras/Source/GameplayCamerasEditor/Private/AssetTools/CameraRigProxyAssetEditor.cpp b/Engine/Plugins/Cameras/GameplayCameras/Source/GameplayCamerasEditor/Private/AssetTools/CameraRigProxyAssetEditor.cpp
new file mode 100644
--- /dev/null
+++ b/Engine/Plugins/Cameras/GameplayCameras/Source/GameplayCamerasEditor/Private/AssetTools/CameraRigProxyAssetEditor.cpp
@@ -0,0 +1,27 @@
+// Copyright Epic Games, Inc. All Rights Reserved.
+
+#include "AssetTools/CameraRigProxyAssetEditor.h"
+
+#include "Core/CameraRigProxyAsset.h"
+#include "Toolkits/CameraRigProxyAssetEditorToolkit.h"
+
+#include UE_INLINE_GENERATED_CPP_BY_NAME(CameraRigProxyAssetEditor)
+
+void UCameraRigProxyAssetEditor::Initialize(TObjectPtr<UCameraRigProxyAsset> InCameraRigProxyAsset)
+{
+	CameraRigProxyAsset = InCameraRigProxyAsset;
+
+	Super::Initialize();
+}
+
+void UCameraRigProxyAssetEditor::GetObjectsToEdit(TArray<UObject*>& InObjectsToEdit)
+{
+	InObjectsToEdit.Add(CameraRigProxyAsset.Get());
+}
+
+TSharedPtr<FBaseAssetToolkit> UCameraRigProxyAssetEditor::CreateToolkit()
+{
+	using namespace UE::Cameras;
+	return MakeShared<FCameraRigProxyAssetEditorToolkit>(this);
+}
+
diff --git a/Engine/Plugins/Cameras/GameplayCameras/Source/GameplayCamerasEditor/Private/Customizations/CameraProxyTableDetailsCustomization.cpp b/Engine/Plugins/Cameras/GameplayCameras/Source/GameplayCamerasEditor/Private/Customizations/CameraProxyTableDetailsCustomization.cpp
new file mode 100644
--- /dev/null
+++ b/Engine/Plugins/Cameras/GameplayCameras/Source/GameplayCamerasEditor/Private/Customizations/CameraProxyTableDetailsCustomization.cpp
@@ -0,0 +1,135 @@
+// Copyright Epic Games, Inc. All Rights Reserved.
+
+#include "Customizations/CameraProxyTableDetailsCustomization.h"
+
+#include "Core/CameraAsset.h"
+#include "Core/CameraRigAsset.h"
+#include "Core/CameraRigProxyTable.h"
+#include "DetailCategoryBuilder.h"
+#include "DetailWidgetRow.h"
+#include "Editors/CameraRigPickerConfig.h"
+#include "IDetailChildrenBuilder.h"
+#include "IGameplayCamerasEditorModule.h"
+#include "Modules/ModuleManager.h"
+#include "Widgets/Input/SComboButton.h"
+#include "Widgets/SNullWidget.h"
+
+#define LOCTEXT_NAMESPACE "CameraProxyTableDetailsCustomization"
+
+namespace UE::Cameras
+{
+
+TSharedRef<IPropertyTypeCustomization> FCameraProxyTableEntryDetailsCustomization::MakeInstance()
+{
+	return MakeShared<FCameraProxyTableEntryDetailsCustomization>();
+}
+
+void FCameraProxyTableEntryDetailsCustomization::CustomizeHeader(TSharedRef<IPropertyHandle> StructPropertyHandle, FDetailWidgetRow& HeaderRow, IPropertyTypeCustomizationUtils& StructCustomizationUtils)
+{
+	HeaderRow
+		.NameContent()
+		[
+			StructPropertyHandle->CreatePropertyNameWidget()
+		];
+}
+
+void FCameraProxyTableEntryDetailsCustomization::CustomizeChildren(TSharedRef<IPropertyHandle> StructPropertyHandle, IDetailChildrenBuilder& StructBuilder, IPropertyTypeCustomizationUtils& StructCustomizationUtils)
+{
+	TArray<UObject*> OuterObjects;
+	StructPropertyHandle->GetOuterObjects(OuterObjects);
+
+	ProxyTables.Reset();
+	for (UObject* OuterObject : OuterObjects)
+	{
+		ProxyTables.Add(CastChecked<UCameraRigProxyTable>(OuterObject));
+	}
+
+	CameraRigPropertyHandle = StructPropertyHandle->GetChildHandle(GET_MEMBER_NAME_CHECKED(FCameraRigProxyTableEntry, CameraRig));
+	CameraRigProxyPropertyHandle = StructPropertyHandle->GetChildHandle(GET_MEMBER_NAME_CHECKED(FCameraRigProxyTableEntry, CameraRigProxy));
+
+	StructBuilder.AddProperty(CameraRigProxyPropertyHandle.ToSharedRef());
+
+	StructBuilder.AddProperty(CameraRigPropertyHandle.ToSharedRef())
+		.IsEnabled(ProxyTables.Num() == 1)
+		.CustomWidget()
+			.NameContent()
+			[
+				CameraRigPropertyHandle->CreatePropertyNameWidget()
+			]
+			.ValueContent()
+			[
+				SAssignNew(ComboButton, SComboButton)
+				.ToolTipText(CameraRigPropertyHandle->GetToolTipText())
+				.ContentPadding(2.f)
+				.ButtonContent()
+				[
+					SNew(STextBlock)
+					.ColorAndOpacity(FSlateColor::UseForeground())
+					.TextStyle(FAppStyle::Get(), "PropertyEditor.AssetClass")
+					.Font(FAppStyle::GetFontStyle("PropertyWindow.NormalFont"))
+					.Text(this, &FCameraProxyTableEntryDetailsCustomization::OnGetComboButtonText)
+				]
+				.OnGetMenuContent(this, &FCameraProxyTableEntryDetailsCustomization::OnBuildCameraRigPicker)
+			];
+}
+
+FText FCameraProxyTableEntryDetailsCustomization::OnGetComboButtonText() const
+{
+	TArray<void*> RawData;
+	CameraRigPropertyHandle->AccessRawData(RawData);
+
+	if (RawData.Num() == 0)
+	{
+		return LOCTEXT("NoCameraRigs", "None");
+	}
+	else if (RawData.Num() > 1)
+	{
+		return LOCTEXT("MultipleCameraRigs", "Multiple Values");
+	}
+	else
+	{
+		FText DisplayText(LOCTEXT("NullCameraRig", "None"));
+		TObjectPtr<UCameraRigAsset>* CameraRigPtr = (TObjectPtr<UCameraRigAsset>*)RawData[0];
+		if (*CameraRigPtr)
+		{
+			DisplayText = FText::FromString((*CameraRigPtr)->GetDisplayName());
+		}
+		return DisplayText;
+	}
+}
+
+TSharedRef<SWidget> FCameraProxyTableEntryDetailsCustomization::OnBuildCameraRigPicker()
+{
+	TArray<void*> RawData;
+	CameraRigPropertyHandle->AccessRawData(RawData);
+
+	if (RawData.Num() != 1)
+	{
+		return SNullWidget::NullWidget;
+	}
+
+	IGameplayCamerasEditorModule& CamerasEditorModule = FModuleManager::LoadModuleChecked<IGameplayCamerasEditorModule>("GameplayCamerasEditor");
+
+	TObjectPtr<UCameraRigAsset>* CameraRigPtr = (TObjectPtr<UCameraRigAsset>*)RawData[0];
+	UCameraAsset* OuterCameraAsset = ProxyTables[0]->GetTypedOuter<UCameraAsset>();
+
+	FCameraRigPickerConfig CameraRigPickerConfig;
+	CameraRigPickerConfig.bCanSelectCameraAsset = false;
+	CameraRigPickerConfig.InitialCameraAssetSelection = FAssetData(OuterCameraAsset);
+	CameraRigPickerConfig.OnCameraRigSelected = FOnCameraRigSelected::CreateSP(
+			this, &FCameraProxyTableEntryDetailsCustomization::OnCameraRigSelected);
+	CameraRigPickerConfig.PropertyToSet = CameraRigPropertyHandle;
+	CameraRigPickerConfig.InitialCameraRigSelection = (*CameraRigPtr);
+
+	return CamerasEditorModule.CreateCameraRigPicker(CameraRigPickerConfig);
+}
+
+void FCameraProxyTableEntryDetailsCustomization::OnCameraRigSelected(UCameraRigAsset* CameraRig)
+{
+	ComboButton->SetIsOpen(false);
+}
+
+}  // namespace UE::Cameras
+
+#undef LOCTEXT_NAMESPACE
+
diff --git a/Engine/Plugins/Cameras/GameplayCameras/Source/GameplayCamerasEditor/Private/Factories/CameraRigProxyAssetFactory.cpp b/Engine/Plugins/Cameras/GameplayCameras/Source/GameplayCamerasEditor/Private/Factories/CameraRigProxyAssetFactory.cpp
new file mode 100644
--- /dev/null
+++ b/Engine/Plugins/Cameras/GameplayCameras/Source/GameplayCamerasEditor/Private/Factories/CameraRigProxyAssetFactory.cpp
@@ -0,0 +1,36 @@
+// Copyright Epic Games, Inc. All Rights Reserved.
+
+#include "Factories/CameraRigProxyAssetFactory.h"
+
+#include "Core/CameraRigProxyAsset.h"
+
+#include UE_INLINE_GENERATED_CPP_BY_NAME(CameraRigProxyAssetFactory)
+
+#define LOCTEXT_NAMESPACE "CameraRigProxyAssetFactory"
+
+UCameraRigProxyAssetFactory::UCameraRigProxyAssetFactory(const FObjectInitializer& ObjectInitializer)
+	: Super(ObjectInitializer)
+{
+	bCreateNew = true;
+	bEditAfterNew = true;
+	SupportedClass = UCameraRigProxyAsset::StaticClass();
+}
+
+UObject* UCameraRigProxyAssetFactory::FactoryCreateNew(UClass* Class, UObject* Parent, FName Name, EObjectFlags Flags, UObject* Context, FFeedbackContext* Warn)
+{
+	UCameraRigProxyAsset* NewCameraRig = NewObject<UCameraRigProxyAsset>(Parent, Class, Name, Flags | RF_Transactional);
+	return NewCameraRig;
+}
+
+bool UCameraRigProxyAssetFactory::ConfigureProperties()
+{
+	return true;
+}
+
+bool UCameraRigProxyAssetFactory::ShouldShowInNewMenu() const
+{
+	return true;
+}
+
+#undef LOCTEXT_NAMESPACE
+
diff --git a/Engine/Plugins/Cameras/GameplayCameras/Source/GameplayCamerasEditor/Private/Factories/CameraRigProxyAssetFactory.h b/Engine/Plugins/Cameras/GameplayCameras/Source/GameplayCamerasEditor/Private/Factories/CameraRigProxyAssetFactory.h
new file mode 100644
--- /dev/null
+++ b/Engine/Plugins/Cameras/GameplayCameras/Source/GameplayCamerasEditor/Private/Factories/CameraRigProxyAssetFactory.h
@@ -0,0 +1,26 @@
+// Copyright Epic Games, Inc. All Rights Reserved.
+
+#pragma once
+
+#include "Factories/Factory.h"
+
+#include "CameraRigProxyAssetFactory.generated.h"
+
+class UCameraRigProxyAsset;
+
+/**
+ * Implements a factory for UCameraRigProxyAsset objects.
+ */
+UCLASS(hidecategories=Object)
+class UCameraRigProxyAssetFactory : public UFactory
+{
+	GENERATED_BODY()
+
+	UCameraRigProxyAssetFactory(const FObjectInitializer& ObjectInitializer);
+
+	// UFactory Interface
+	virtual UObject* FactoryCreateNew(UClass* Class, UObject* Parent, FName Name, EObjectFlags Flags, UObject* Context, FFeedbackContext* Warn) override;
+	virtual bool ConfigureProperties() override;
+	virtual bool ShouldShowInNewMenu() const override;
+};
+
diff --git a/Engine/Plugins/Cameras/GameplayCameras/Source/GameplayCamerasEditor/Private/GameplayCamerasEditorModule.cpp b/Engine/Plugins/Cameras/GameplayCameras/Source/GameplayCamerasEditor/Private/GameplayCamerasEditorModule.cpp
--- a/Engine/Plugins/Cameras/GameplayCameras/Source/GameplayCamerasEditor/Private/GameplayCamerasEditorModule.cpp
+++ b/Engine/Plugins/Cameras/GameplayCameras/Source/GameplayCamerasEditor/Private/GameplayCamerasEditorModule.cpp
@@ -4,6 +4,7 @@
 
 #include "AssetTools/CameraAssetEditor.h"
 #include "AssetTools/CameraRigAssetEditor.h"
+#include "AssetTools/CameraRigProxyAssetEditor.h"
 #include "AssetTools/CameraVariableCollectionEditor.h"
 #include "Commands/CameraAssetEditorCommands.h"
 #include "Commands/CameraRigAssetEditorCommands.h"
@@ -11,6 +12,7 @@
 #include "Commands/CameraVariableCollectionEditorCommands.h"
 #include "Commands/GameplayCamerasDebuggerCommands.h"
 #include "Customizations/CameraParameterDetailsCustomizations.h"
+#include "Customizations/CameraProxyTableDetailsCustomization.h"
 #include "Customizations/SingleCameraDirectorDetailsCustomization.h"
 #include "Debug/CameraDebugCategories.h"
 #include "Debugger/SBlendStacksDebugPanel.h"
@@ -123,6 +125,14 @@ public:
 		return AssetEditor;
 	}
 
+	virtual UCameraRigProxyAssetEditor* CreateCameraRigProxyEditor(const EToolkitMode::Type Mode, const TSharedPtr<IToolkitHost>& InitToolkitHost, UCameraRigProxyAsset* CameraRigProxy) override
+	{
+		UAssetEditorSubsystem* AssetEditorSubsystem = GEditor->GetEditorSubsystem<UAssetEditorSubsystem>();
+		UCameraRigProxyAssetEditor* AssetEditor = NewObject<UCameraRigProxyAssetEditor>(AssetEditorSubsystem, NAME_None, RF_Transient);
+		AssetEditor->Initialize(CameraRigProxy);
+		return AssetEditor;
+	}
+
 	virtual UCameraVariableCollectionEditor* CreateCameraVariableCollectionEditor(const EToolkitMode::Type Mode, const TSharedPtr<IToolkitHost>& InitToolkitHost, UCameraVariableCollection* VariableCollection) override
 	{
 		UAssetEditorSubsystem* AssetEditorSubsystem = GEditor->GetEditorSubsystem<UAssetEditorSubsystem>();
@@ -326,6 +336,8 @@ private:
 
 		FPropertyEditorModule& PropertyEditorModule = FModuleManager::GetModuleChecked<FPropertyEditorModule>("PropertyEditor");
 		FCameraParameterDetailsCustomization::Register(PropertyEditorModule);
+		PropertyEditorModule.RegisterCustomPropertyTypeLayout("CameraRigProxyTableEntry", FOnGetPropertyTypeCustomizationInstance::CreateStatic(
+					&FCameraProxyTableEntryDetailsCustomization::MakeInstance));
 		PropertyEditorModule.RegisterCustomClassLayout("SingleCameraDirector", FOnGetDetailCustomizationInstance::CreateStatic(
 					&FSingleCameraDirectorDetailsCustomization::MakeInstance));
 	}
@@ -339,6 +351,7 @@ private:
 		if (PropertyEditorModule)
 		{
 			FCameraParameterDetailsCustomization::Unregister(*PropertyEditorModule);
+			PropertyEditorModule->UnregisterCustomPropertyTypeLayout("CameraRigProxyTableEntry");
 			PropertyEditorModule->UnregisterCustomClassLayout("SingleCameraDirector");
 		}
 	}
diff --git a/Engine/Plugins/Cameras/GameplayCameras/Source/GameplayCamerasEditor/Private/Toolkits/CameraRigProxyAssetEditorToolkit.cpp b/Engine/Plugins/Cameras/GameplayCameras/Source/GameplayCamerasEditor/Private/Toolkits/CameraRigProxyAssetEditorToolkit.cpp
new file mode 100644
--- /dev/null
+++ b/Engine/Plugins/Cameras/GameplayCameras/Source/GameplayCamerasEditor/Private/Toolkits/CameraRigProxyAssetEditorToolkit.cpp
@@ -0,0 +1,107 @@
+// Copyright Epic Games, Inc. All Rights Reserved.
+
+#include "Toolkits/CameraRigProxyAssetEditorToolkit.h"
+
+#include "AssetTools/CameraRigProxyAssetEditor.h"
+#include "Core/CameraRigProxyAsset.h"
+#include "Framework/Docking/LayoutExtender.h"
+#include "PropertyEditorModule.h"
+
+#define LOCTEXT_NAMESPACE "CameraRigProxyAssetEditorToolkit"
+
+namespace UE::Cameras
+{
+
+const FName FCameraRigProxyAssetEditorToolkit::DetailsViewTabId(TEXT("CameraRigProxyAssetEditor_DetailsView"));
+
+FCameraRigProxyAssetEditorToolkit::FCameraRigProxyAssetEditorToolkit(UCameraRigProxyAssetEditor* InOwningAssetEditor)
+	: FBaseAssetToolkit(InOwningAssetEditor)
+{
+	// Override base class default layout.
+	StandaloneDefaultLayout = FTabManager::NewLayout("CameraRigProxyAssetEditor_Layout")
+		->AddArea
+		(
+			FTabManager::NewPrimaryArea()->SetOrientation(Orient_Vertical)
+			->Split
+			(
+				FTabManager::NewStack()
+				->AddTab(DetailsViewTabId, ETabState::OpenedTab)
+				->SetForegroundTab(DetailsViewTabId)
+			)
+		);
+}
+
+FCameraRigProxyAssetEditorToolkit::~FCameraRigProxyAssetEditorToolkit()
+{
+}
+
+void FCameraRigProxyAssetEditorToolkit::AddReferencedObjects(FReferenceCollector& Collector)
+{
+	Collector.AddReferencedObject(CameraRigProxyAsset);
+}
+
+FString FCameraRigProxyAssetEditorToolkit::GetReferencerName() const
+{
+	return TEXT("FCameraRigProxyAssetEditorToolkit");
+}
+
+void FCameraRigProxyAssetEditorToolkit::RegisterTabSpawners(const TSharedRef<FTabManager>& InTabManager)
+{
+	// Skip FBaseAssetToolkit here because we don't want a viewport tab.
+	FAssetEditorToolkit::RegisterTabSpawners(InTabManager);
+
+	InTabManager->RegisterTabSpawner(DetailsViewTabId, FOnSpawnTab::CreateSP(this, &FCameraRigProxyAssetEditorToolkit::SpawnTab_Details))
+		.SetDisplayName(LOCTEXT("Details", "Details"))
+		.SetGroup(AssetEditorTabsCategory.ToSharedRef())
+		.SetIcon(FSlateIcon(FAppStyle::GetAppStyleSetName(), "LevelEditor.Tabs.Details"));
+}
+
+void FCameraRigProxyAssetEditorToolkit::UnregisterTabSpawners(const TSharedRef<FTabManager>& InTabManager)
+{
+	// Skip FBaseAssetToolkit here because we don't want a viewport tab.
+	FAssetEditorToolkit::UnregisterTabSpawners(InTabManager);
+
+	InTabManager->UnregisterTabSpawner(DetailsViewTabId);
+}
+
+void FCameraRigProxyAssetEditorToolkit::CreateWidgets()
+{
+	// Skip FBaseAssetToolkit here because we don't want a viewport tab.
+	// ...no up-call...
+
+	// Do most of FBaseAssetToolkit's work except for the viewport.
+	RegisterToolbar();
+	LayoutExtender = MakeShared<FLayoutExtender>();
+
+	DetailsView = CreateDetailsView();
+}
+
+void FCameraRigProxyAssetEditorToolkit::PostInitAssetEditor()
+{
+}
+
+FText FCameraRigProxyAssetEditorToolkit::GetBaseToolkitName() const
+{
+	return LOCTEXT("AppLabel", "Camera Rig Proxy Asset");
+}
+
+FName FCameraRigProxyAssetEditorToolkit::GetToolkitFName() const
+{
+	static FName SequencerName("CameraRigProxyAssetEditor");
+	return SequencerName;
+}
+
+FString FCameraRigProxyAssetEditorToolkit::GetWorldCentricTabPrefix() const
+{
+	return LOCTEXT("WorldCentricTabPrefix", "Camera Rig Proxy Asset ").ToString();
+}
+
+FLinearColor FCameraRigProxyAssetEditorToolkit::GetWorldCentricTabColorScale() const
+{
+	return FLinearColor(0.7, 0.0f, 0.0f, 0.5f);
+}
+
+}  // namespace UE::Cameras
+
+#undef LOCTEXT_NAMESPACE
+
diff --git a/Engine/Plugins/Cameras/GameplayCameras/Source/GameplayCamerasEditor/Public/AssetTools/CameraRigProxyAssetEditor.h b/Engine/Plugins/Cameras/GameplayCameras/Source/GameplayCamerasEditor/Public/AssetTools/CameraRigProxyAssetEditor.h
new file mode 100644
--- /dev/null
+++ b/Engine/Plugins/Cameras/GameplayCameras/Source/GameplayCamerasEditor/Public/AssetTools/CameraRigProxyAssetEditor.h
@@ -0,0 +1,37 @@
+// Copyright Epic Games, Inc. All Rights Reserved.
+
+#pragma once
+
+#include "Tools/UAssetEditor.h"
+
+#include "CameraRigProxyAssetEditor.generated.h"
+
+class FBaseAssetToolkit;
+class UCameraRigProxyAsset;
+
+/**
+ * Editor for a camera rig proxy asset.
+ */
+UCLASS(Transient)
+class UCameraRigProxyAssetEditor : public UAssetEditor
+{
+	GENERATED_BODY()
+
+public:
+
+	void Initialize(TObjectPtr<UCameraRigProxyAsset> InCameraRigProxyAsset);
+
+	UCameraRigProxyAsset* GetCameraRigProxyAsset() const { return CameraRigProxyAsset; }
+
+public:
+
+	// UAssetEditor interface
+	virtual void GetObjectsToEdit(TArray<UObject*>& InObjectsToEdit) override;
+	virtual TSharedPtr<FBaseAssetToolkit> CreateToolkit() override;
+
+private:
+
+	UPROPERTY()
+	TObjectPtr<UCameraRigProxyAsset> CameraRigProxyAsset;
+};
+
diff --git a/Engine/Plugins/Cameras/GameplayCameras/Source/GameplayCamerasEditor/Public/Customizations/CameraProxyTableDetailsCustomization.h b/Engine/Plugins/Cameras/GameplayCameras/Source/GameplayCamerasEditor/Public/Customizations/CameraProxyTableDetailsCustomization.h
new file mode 100644
--- /dev/null
+++ b/Engine/Plugins/Cameras/GameplayCameras/Source/GameplayCamerasEditor/Public/Customizations/CameraProxyTableDetailsCustomization.h
@@ -0,0 +1,51 @@
+// Copyright Epic Games, Inc. All Rights Reserved.
+
+#pragma once
+
+#include "IPropertyTypeCustomization.h"
+
+class FDetailWidgetRow;
+class IDetailChildrenBuilder;
+class IPropertyHandle;
+class SComboButton;
+class SWidget;
+class UCameraRigAsset;
+class UCameraRigProxyTable;
+struct FCameraRigProxyTableEntry;
+
+namespace UE::Cameras
+{
+
+/**
+ * Details customization for an entry in a camera rig proxy table.
+ *
+ * It uses a custom camera rig picker that only chooses among the rigs of the parent camera asset.
+ */
+class FCameraProxyTableEntryDetailsCustomization : public IPropertyTypeCustomization
+{
+public:
+
+	static TSharedRef<IPropertyTypeCustomization> MakeInstance();
+
+public:
+
+	// IPropertyTypeCustomization interface
+	virtual void CustomizeHeader(TSharedRef<IPropertyHandle> StructPropertyHandle, FDetailWidgetRow& HeaderRow, IPropertyTypeCustomizationUtils& StructCustomizationUtils) override;
+	virtual void CustomizeChildren(TSharedRef<IPropertyHandle> StructPropertyHandle, IDetailChildrenBuilder& StructBuilder, IPropertyTypeCustomizationUtils& StructCustomizationUtils) override;
+
+private:
+
+	TSharedRef<SWidget> OnBuildCameraRigPicker();
+	FText OnGetComboButtonText() const;
+	void OnCameraRigSelected(UCameraRigAsset* CameraRig);
+
+private:
+
+	TSharedPtr<SComboButton> ComboButton;
+	TSharedPtr<IPropertyHandle> CameraRigPropertyHandle;
+	TSharedPtr<IPropertyHandle> CameraRigProxyPropertyHandle;
+	TArray<TWeakObjectPtr<UCameraRigProxyTable>> ProxyTables;
+};
+
+}  // namespace UE::Cameras
+
diff --git a/Engine/Plugins/Cameras/GameplayCameras/Source/GameplayCamerasEditor/Public/IGameplayCamerasEditorModule.h b/Engine/Plugins/Cameras/GameplayCameras/Source/GameplayCamerasEditor/Public/IGameplayCamerasEditorModule.h
--- a/Engine/Plugins/Cameras/GameplayCameras/Source/GameplayCamerasEditor/Public/IGameplayCamerasEditorModule.h
+++ b/Engine/Plugins/Cameras/GameplayCameras/Source/GameplayCamerasEditor/Public/IGameplayCamerasEditorModule.h
@@ -12,6 +12,8 @@ class UCameraAsset;
 class UCameraAssetEditor;
 class UCameraRigAsset;
 class UCameraRigAssetEditor;
+class UCameraRigProxyAsset;
+class UCameraRigProxyAssetEditor;
 class UCameraVariableCollection;
 class UCameraVariableCollectionEditor;
 
@@ -54,6 +56,9 @@ public:
 	/** Creates an editor for the given camera rig asset */
 	virtual UCameraRigAssetEditor* CreateCameraRigEditor(const EToolkitMode::Type Mode, const TSharedPtr<IToolkitHost>& InitToolkitHost, UCameraRigAsset* CameraRig) = 0;
 
+	/** Creates an editor for the given camera rig proxy asset */
+	virtual UCameraRigProxyAssetEditor* CreateCameraRigProxyEditor(const EToolkitMode::Type Mode, const TSharedPtr<IToolkitHost>& InitToolkitHost, UCameraRigProxyAsset* CameraRigProxy) = 0;
+
 	/** Creates an editor for the given variable collection */
 	virtual UCameraVariableCollectionEditor* CreateCameraVariableCollectionEditor(const EToolkitMode::Type Mode, const TSharedPtr<IToolkitHost>& InitToolkitHost, UCameraVariableCollection* VariableCollection) = 0;
 
diff --git a/Engine/Plugins/Cameras/GameplayCameras/Source/GameplayCamerasEditor/Public/Toolkits/CameraRigProxyAssetEditorToolkit.h b/Engine/Plugins/Cameras/GameplayCameras/Source/GameplayCamerasEditor/Public/Toolkits/CameraRigProxyAssetEditorToolkit.h
new file mode 100644
--- /dev/null
+++ b/Engine/Plugins/Cameras/GameplayCameras/Source/GameplayCamerasEditor/Public/Toolkits/CameraRigProxyAssetEditorToolkit.h
@@ -0,0 +1,56 @@
+// Copyright Epic Games, Inc. All Rights Reserved.
+
+#pragma once
+
+#include "CoreMinimal.h"
+#include "Tools/BaseAssetToolkit.h"
+#include "UObject/GCObject.h"
+
+class IMessageLogListing;
+class UCameraRigProxyAsset;
+class UCameraRigProxyAssetEditor;
+class SWidget;
+
+namespace UE::Cameras
+{
+
+/**
+ * Editor toolkit for a camera rig proxy asset.
+ */
+class FCameraRigProxyAssetEditorToolkit 
+	: public FBaseAssetToolkit
+	, public FGCObject
+{
+public:
+
+	FCameraRigProxyAssetEditorToolkit(UCameraRigProxyAssetEditor* InOwningAssetEditor);
+	~FCameraRigProxyAssetEditorToolkit();
+
+protected:
+
+	// FBaseAssetToolkit interface
+	virtual void RegisterTabSpawners(const TSharedRef<class FTabManager>& TabManager) override;
+	virtual void UnregisterTabSpawners(const TSharedRef<FTabManager>& InTabManager) override;
+	virtual void CreateWidgets() override;
+	virtual void PostInitAssetEditor() override;
+
+	// IToolkit interface
+	virtual FText GetBaseToolkitName() const override;
+	virtual FName GetToolkitFName() const override;
+	virtual FString GetWorldCentricTabPrefix() const override;
+	virtual FLinearColor GetWorldCentricTabColorScale() const override;
+
+	// FGCObject interface
+	virtual void AddReferencedObjects(FReferenceCollector& Collector) override;
+	virtual FString GetReferencerName() const override;
+
+private:
+
+	static const FName DetailsViewTabId;
+
+	/** The asset being edited */
+	TObjectPtr<UCameraRigProxyAsset> CameraRigProxyAsset;
+};
+
+}  // namespace UE::Cameras
+
