Cameras: fix incorrect copy/pasted editor names

diff --git a/Engine/Plugins/Cameras/GameplayCameras/Source/GameplayCamerasEditor/Private/Toolkits/CameraDirectorAssetEditorMode.cpp b/Engine/Plugins/Cameras/GameplayCameras/Source/GameplayCamerasEditor/Private/Toolkits/CameraDirectorAssetEditorMode.cpp
--- a/Engine/Plugins/Cameras/GameplayCameras/Source/GameplayCamerasEditor/Private/Toolkits/CameraDirectorAssetEditorMode.cpp
+++ b/Engine/Plugins/Cameras/GameplayCameras/Source/GameplayCamerasEditor/Private/Toolkits/CameraDirectorAssetEditorMode.cpp
@@ -47,14 +47,14 @@ void FCameraDirectorAssetEditorMode::OnActivateMode(const FAssetEditorModeActiva
 	DetailsView->SetObject(CameraAsset->CameraDirector);
 
 	InParams.TabManager->RegisterTabSpawner(DirectorEditorTabId, FOnSpawnTab::CreateSP(this, &FCameraDirectorAssetEditorMode::SpawnTab_DirectorEditor))
-		.SetDisplayName(LOCTEXT("TransitionEditor", "Camera Transitions"))
+		.SetDisplayName(LOCTEXT("CameraDirectorEditor", "Camera Director"))
 		.SetGroup(InParams.AssetEditorTabsCategory.ToSharedRef());
 }
 
 TSharedRef<SDockTab> FCameraDirectorAssetEditorMode::SpawnTab_DirectorEditor(const FSpawnTabArgs& Args)
 {
 	TSharedPtr<SDockTab> DirectorEditorTab = SNew(SDockTab)
-		.Label(LOCTEXT("TransitionEditorTabTitle", "Camera Transitions"))
+		.Label(LOCTEXT("CameraDirectorEditorTabTitle", "Camera Director"))
 		[
 			DetailsView.ToSharedRef()
 		];
