Cameras: disable camera rig graph editors when there is no rig asset to edit

diff --git a/Engine/Plugins/Cameras/GameplayCameras/Source/GameplayCamerasEditor/Private/Editors/SCameraRigAssetEditor.cpp b/Engine/Plugins/Cameras/GameplayCameras/Source/GameplayCamerasEditor/Private/Editors/SCameraRigAssetEditor.cpp
--- a/Engine/Plugins/Cameras/GameplayCameras/Source/GameplayCamerasEditor/Private/Editors/SCameraRigAssetEditor.cpp
+++ b/Engine/Plugins/Cameras/GameplayCameras/Source/GameplayCamerasEditor/Private/Editors/SCameraRigAssetEditor.cpp
@@ -88,6 +88,7 @@ void SCameraRigAssetEditor::CreateNodeGraphEditor()
 		.Appearance(Appearance)
 		.DetailsView(DetailsView)
 		.GraphTitle(this, &SCameraRigAssetEditor::GetCameraRigAssetName, NodeGraph.Get())
+		.IsEnabled(this, &SCameraRigAssetEditor::IsGraphEditorEnabled)
 		.GraphToEdit(NodeGraph)
 		.AssetEditorToolkit(AssetEditorToolkit);
 }
@@ -113,6 +114,7 @@ void SCameraRigAssetEditor::CreateTransitionGraphEditor()
 		.Appearance(Appearance)
 		.DetailsView(DetailsView)
 		.GraphTitle(this, &SCameraRigAssetEditor::GetCameraRigAssetName, TransitionGraph.Get())
+		.IsEnabled(this, &SCameraRigAssetEditor::IsGraphEditorEnabled)
 		.GraphToEdit(TransitionGraph)
 		.AssetEditorToolkit(AssetEditorToolkit);
 }
@@ -268,6 +270,11 @@ FText SCameraRigAssetEditor::GetCameraRigAssetName(UObjectTreeGraph* ForGraph) c
 	return LOCTEXT("NoCameraRig", "No Camera Rig");
 }
 
+bool SCameraRigAssetEditor::IsGraphEditorEnabled() const
+{
+	return CameraRigAsset != nullptr;
+}
+
 void SCameraRigAssetEditor::OnGraphChanged(const FEdGraphEditAction& InEditAction)
 {
 	OnAnyGraphChanged.Broadcast(InEditAction);
diff --git a/Engine/Plugins/Cameras/GameplayCameras/Source/GameplayCamerasEditor/Public/Editors/SCameraRigAssetEditor.h b/Engine/Plugins/Cameras/GameplayCameras/Source/GameplayCamerasEditor/Public/Editors/SCameraRigAssetEditor.h
--- a/Engine/Plugins/Cameras/GameplayCameras/Source/GameplayCamerasEditor/Public/Editors/SCameraRigAssetEditor.h
+++ b/Engine/Plugins/Cameras/GameplayCameras/Source/GameplayCamerasEditor/Public/Editors/SCameraRigAssetEditor.h
@@ -95,6 +95,7 @@ protected:
 	void OnGraphChanged(const FEdGraphEditAction& InEditAction);
 
 	FText GetCameraRigAssetName(UObjectTreeGraph* ForGraph) const;
+	bool IsGraphEditorEnabled() const;
 
 private:
 
