Cameras: allow finding objects in the director tab from hyperlinked messages

diff --git a/Engine/Plugins/Cameras/GameplayCameras/Source/GameplayCamerasEditor/Private/Toolkits/CameraAssetEditorToolkit.cpp b/Engine/Plugins/Cameras/GameplayCameras/Source/GameplayCamerasEditor/Private/Toolkits/CameraAssetEditorToolkit.cpp
--- a/Engine/Plugins/Cameras/GameplayCameras/Source/GameplayCamerasEditor/Private/Toolkits/CameraAssetEditorToolkit.cpp
+++ b/Engine/Plugins/Cameras/GameplayCameras/Source/GameplayCamerasEditor/Private/Toolkits/CameraAssetEditorToolkit.cpp
@@ -6,6 +6,7 @@
 #include "Commands/CameraAssetEditorCommands.h"
 #include "Core/CameraAsset.h"
 #include "Core/CameraBuildLog.h"
+#include "Core/CameraDirector.h"
 #include "Editors/ObjectTreeGraphConfig.h"
 #include "Editors/SFindInObjectTreeGraph.h"
 #include "Framework/Docking/LayoutExtender.h"
@@ -295,11 +296,17 @@ void FCameraAssetEditorToolkit::OnJumpToObject(UObject* Object)
 
 void FCameraAssetEditorToolkit::OnJumpToObject(UObject* Object, FName PropertyName)
 {
+	bool bFindInCameraDirector = false;
 	bool bFindInCameraRig = false;
 	bool bFindInSharedTranstions = false;
 	UObject* CurOuter = Object;
 	while (CurOuter != nullptr)
 	{
+		if (CurOuter->IsA<UCameraDirector>())
+		{
+			bFindInCameraDirector = true;
+			break;
+		}
 		if (CurOuter->IsA<UCameraRigAsset>())
 		{
 			bFindInCameraRig = true;
@@ -313,6 +320,15 @@ void FCameraAssetEditorToolkit::OnJumpToObject(UObject* Object, FName PropertyNa
 
 		CurOuter = CurOuter->GetOuter();
 	}
+
+	if (bFindInCameraDirector)
+	{
+		TSharedPtr<FCameraDirectorAssetEditorMode> DirectorMode = GetTypedEditorMode<FCameraDirectorAssetEditorMode>(
+				FCameraDirectorAssetEditorMode::ModeName);
+		SetEditorMode(FCameraDirectorAssetEditorMode::ModeName);
+		DirectorMode->JumpToObject(Object, PropertyName);
+		return;
+	}
 	
 	if (bFindInCameraRig)
 	{
diff --git a/Engine/Plugins/Cameras/GameplayCameras/Source/GameplayCamerasEditor/Private/Toolkits/CameraDirectorAssetEditorMode.cpp b/Engine/Plugins/Cameras/GameplayCameras/Source/GameplayCamerasEditor/Private/Toolkits/CameraDirectorAssetEditorMode.cpp
--- a/Engine/Plugins/Cameras/GameplayCameras/Source/GameplayCamerasEditor/Private/Toolkits/CameraDirectorAssetEditorMode.cpp
+++ b/Engine/Plugins/Cameras/GameplayCameras/Source/GameplayCamerasEditor/Private/Toolkits/CameraDirectorAssetEditorMode.cpp
@@ -67,7 +67,7 @@ void FCameraDirectorAssetEditorMode::OnDeactivateMode(const FAssetEditorModeDeac
 	InParams.TabManager->UnregisterTabSpawner(DirectorEditorTabId);
 }
 
-bool FCameraDirectorAssetEditorMode::JumpToObject(UObject* InObject)
+bool FCameraDirectorAssetEditorMode::JumpToObject(UObject* InObject, FName InPropertyName)
 {
 	if (InObject == CameraAsset->CameraDirector)
 	{
diff --git a/Engine/Plugins/Cameras/GameplayCameras/Source/GameplayCamerasEditor/Public/Toolkits/CameraDirectorAssetEditorMode.h b/Engine/Plugins/Cameras/GameplayCameras/Source/GameplayCamerasEditor/Public/Toolkits/CameraDirectorAssetEditorMode.h
--- a/Engine/Plugins/Cameras/GameplayCameras/Source/GameplayCamerasEditor/Public/Toolkits/CameraDirectorAssetEditorMode.h
+++ b/Engine/Plugins/Cameras/GameplayCameras/Source/GameplayCamerasEditor/Public/Toolkits/CameraDirectorAssetEditorMode.h
@@ -23,7 +23,7 @@ public:
 
 	FCameraDirectorAssetEditorMode(UCameraAsset* InCameraAsset);
 
-	bool JumpToObject(UObject* InObject);
+	bool JumpToObject(UObject* InObject, FName InPropertyName);
 
 protected:
 
