Schema: call pre/post edit callbacks when changing properties in object-tree-graphs

diff --git a/Engine/Plugins/Cameras/GameplayCameras/Source/GameplayCamerasEditor/Private/Editors/ObjectTreeGraphSchema.cpp b/Engine/Plugins/Cameras/GameplayCameras/Source/GameplayCamerasEditor/Private/Editors/ObjectTreeGraphSchema.cpp
--- a/Engine/Plugins/Cameras/GameplayCameras/Source/GameplayCamerasEditor/Private/Editors/ObjectTreeGraphSchema.cpp
+++ b/Engine/Plugins/Cameras/GameplayCameras/Source/GameplayCamerasEditor/Private/Editors/ObjectTreeGraphSchema.cpp
@@ -14,6 +14,7 @@
 #include "Serialization/ArchiveUObject.h"
 #include "UObject/FastReferenceCollector.h"
 #include "UObject/UObjectIterator.h"
+#include "UObject/UnrealType.h"
 #include "UnrealExporter.h"
 
 #include UE_INLINE_GENERATED_CPP_BY_NAME(ObjectTreeGraphSchema)
@@ -556,6 +557,11 @@ const FPinConnectionResponse UObjectTreeGraphSchema::CanCreateConnection(const U
 		return FPinConnectionResponse(CONNECT_RESPONSE_DISALLOW, TEXT("Unsupported source pin"));
 	}
 
+	if (!ObjectA->CanEditChange(PropertyA))
+	{
+		return FPinConnectionResponse(CONNECT_RESPONSE_DISALLOW, TEXT("Property cannot be changed"));
+	}
+
 	if (FObjectProperty* ObjectPropertyA = CastField<FObjectProperty>(PropertyA))
 	{
 		if (ObjectClassB->IsChildOf(ObjectPropertyA->PropertyClass))
@@ -622,12 +628,19 @@ bool UObjectTreeGraphSchema::TryCreateConnection(UEdGraphPin* A, UEdGraphPin* B)
 
 	if (FObjectProperty* ObjectPropertyA = CastField<FObjectProperty>(PropertyA))
 	{
+		ObjectA->PreEditChange(PropertyA);
+
 		ObjectA->Modify();
 
 		ObjectPropertyA->SetValue_InContainer(ObjectA, TObjectPtr<UObject>(ObjectB));
+
+		FPropertyChangedEvent PropertyChangedEvent(PropertyA, EPropertyChangeType::ValueSet);
+		ObjectA->PostEditChangeProperty(PropertyChangedEvent);
 	}
 	else if (FArrayProperty* ArrayPropertyA = CastField<FArrayProperty>(PropertyA))
 	{
+		ObjectA->PreEditChange(PropertyA);
+
 		ObjectA->Modify();
 
 		const int32 Index = NodeA->GetIndexOfArrayPin(A);
@@ -644,6 +657,10 @@ bool UObjectTreeGraphSchema::TryCreateConnection(UEdGraphPin* A, UEdGraphPin* B)
 			NodeA->CreateNewItemPin(*ArrayPropertyA);
 			NodeA->GetGraph()->NotifyNodeChanged(NodeA);
 		}
+
+		FPropertyChangedEvent PropertyChangedEvent(PropertyA);
+		PropertyChangedEvent.ChangeType = bAddNewItemPin ? EPropertyChangeType::ArrayAdd : EPropertyChangeType::ValueSet;
+		ObjectA->PostEditChangeProperty(PropertyChangedEvent);
 	}
 
 	return true;
@@ -686,12 +703,19 @@ void UObjectTreeGraphSchema::BreakPinLinks(UEdGraphPin& TargetPin, bool bSendsNo
 
 	if (FObjectProperty* ObjectProperty = CastField<FObjectProperty>(Property))
 	{
+		OwningObject->PreEditChange(Property);
+
 		OwningObject->Modify();
 
 		ObjectProperty->ClearValue_InContainer(OwningObject);
+
+		FPropertyChangedEvent PropertyChangedEvent(Property, EPropertyChangeType::ValueSet);
+		OwningObject->PostEditChangeProperty(PropertyChangedEvent);
 	}
 	else if (FArrayProperty* ArrayProperty = CastField<FArrayProperty>(Property))
 	{
+		OwningObject->PreEditChange(Property);
+
 		OwningObject->Modify();
 
 		int32 Index = PropertyOwningNode->GetIndexOfArrayPin(PropertyPin);
@@ -701,6 +725,9 @@ void UObjectTreeGraphSchema::BreakPinLinks(UEdGraphPin& TargetPin, bool bSendsNo
 		ArrayHelper.RemoveValues(Index);
 
 		bRemovePropertyPin = true;
+
+		FPropertyChangedEvent PropertyChangedEvent(Property, EPropertyChangeType::ArrayRemove);
+		OwningObject->PostEditChangeProperty(PropertyChangedEvent);
 	}
 
 	Super::BreakPinLinks(TargetPin, bSendsNodeNotification);
