Cameras: add event handlers framework, use it for the camera asset

This will allow editor code to listen to data changes and react to it.

diff --git a/Engine/Plugins/Cameras/GameplayCameras/Source/GameplayCameras/Private/Core/BlendStackCameraNode.cpp b/Engine/Plugins/Cameras/GameplayCameras/Source/GameplayCameras/Private/Core/BlendStackCameraNode.cpp
--- a/Engine/Plugins/Cameras/GameplayCameras/Source/GameplayCameras/Private/Core/BlendStackCameraNode.cpp
+++ b/Engine/Plugins/Cameras/GameplayCameras/Source/GameplayCameras/Private/Core/BlendStackCameraNode.cpp
@@ -349,7 +349,7 @@ const UCameraRigTransition* FBlendStackCameraNodeEvaluator::FindTransition(const
 			if (FromCameraAsset)
 			{
 				TransitionToUse = FindTransition(
-						FromCameraAsset->ExitTransitions,
+						FromCameraAsset->GetExitTransitions(),
 						FromCameraRig, FromCameraAsset, false,
 						ToCameraRig, ToCameraAsset);
 				if (TransitionToUse)
@@ -373,7 +373,7 @@ const UCameraRigTransition* FBlendStackCameraNodeEvaluator::FindTransition(const
 		if (ToCameraAsset)
 		{
 			TransitionToUse = FindTransition(
-					ToCameraAsset->EnterTransitions,
+					ToCameraAsset->GetEnterTransitions(),
 					FromCameraRig, FromCameraAsset, TopEntry.bIsFrozen,
 					ToCameraRig, ToCameraAsset);
 			if (TransitionToUse)
diff --git a/Engine/Plugins/Cameras/GameplayCameras/Source/GameplayCameras/Private/Core/CameraAsset.cpp b/Engine/Plugins/Cameras/GameplayCameras/Source/GameplayCameras/Private/Core/CameraAsset.cpp
--- a/Engine/Plugins/Cameras/GameplayCameras/Source/GameplayCameras/Private/Core/CameraAsset.cpp
+++ b/Engine/Plugins/Cameras/GameplayCameras/Source/GameplayCameras/Private/Core/CameraAsset.cpp
@@ -9,6 +9,135 @@
 
 const FName UCameraAsset::SharedTransitionsGraphName("SharedTransitions");
 
+void UCameraAsset::SetCameraDirector(UCameraDirector* InCameraDirector)
+{
+	using namespace UE::Cameras;
+
+	if (CameraDirector != InCameraDirector)
+	{
+		CameraDirector = InCameraDirector;
+
+		TCameraPropertyChangedEvent<UCameraDirector*> ChangedEvent;
+		ChangedEvent.NewValue = CameraDirector;
+		EventHandlers.Notify(&ICameraAssetEventHandler::OnCameraDirectorChanged, this, ChangedEvent);
+	}
+}
+
+void UCameraAsset::AddCameraRig(UCameraRigAsset* InCameraRig)
+{
+	using namespace UE::Cameras;
+
+	ensure(InCameraRig);
+
+	CameraRigs.Add(InCameraRig);
+
+	TCameraArrayChangedEvent<UCameraRigAsset*> ChangedEvent;
+	ChangedEvent.EventType = ECameraArrayChangedEventType::Add;
+	EventHandlers.Notify(&ICameraAssetEventHandler::OnCameraRigsChanged, this, ChangedEvent);
+}
+
+int32 UCameraAsset::RemoveCameraRig(UCameraRigAsset* InCameraRig)
+{
+	using namespace UE::Cameras;
+
+	const int32 NumRemoved = CameraRigs.Remove(InCameraRig);
+	if (NumRemoved > 0)
+	{
+		TCameraArrayChangedEvent<UCameraRigAsset*> ChangedEvent;
+		ChangedEvent.EventType = ECameraArrayChangedEventType::Remove;
+		EventHandlers.Notify(&ICameraAssetEventHandler::OnCameraRigsChanged, this, ChangedEvent);
+	}
+	return NumRemoved;
+}
+
+void UCameraAsset::AddEnterTransition(UCameraRigTransition* InTransition)
+{
+	using namespace UE::Cameras;
+
+	ensure(InTransition);
+
+	EnterTransitions.Add(InTransition);
+
+	TCameraArrayChangedEvent<UCameraRigTransition*> ChangedEvent;
+	ChangedEvent.EventType = ECameraArrayChangedEventType::Add;
+	EventHandlers.Notify(&ICameraAssetEventHandler::OnEnterTransitionsChanged, this, ChangedEvent);
+}
+
+int32 UCameraAsset::RemoveEnterTransition(UCameraRigTransition* InTransition)
+{
+	using namespace UE::Cameras;
+	
+	const int32 NumRemoved = EnterTransitions.Remove(InTransition);
+	if (NumRemoved > 0)
+	{
+		TCameraArrayChangedEvent<UCameraRigTransition*> ChangedEvent;
+		ChangedEvent.EventType = ECameraArrayChangedEventType::Remove;
+		EventHandlers.Notify(&ICameraAssetEventHandler::OnEnterTransitionsChanged, this, ChangedEvent);
+	}
+	return NumRemoved;
+}
+
+void UCameraAsset::AddExitTransition(UCameraRigTransition* InTransition)
+{
+	using namespace UE::Cameras;
+
+	ensure(InTransition);
+
+	ExitTransitions.Add(InTransition);
+
+	TCameraArrayChangedEvent<UCameraRigTransition*> ChangedEvent;
+	ChangedEvent.EventType = ECameraArrayChangedEventType::Add;
+	EventHandlers.Notify(&ICameraAssetEventHandler::OnExitTransitionsChanged, this, ChangedEvent);
+}
+
+int32 UCameraAsset::RemoveExitTransition(UCameraRigTransition* InTransition)
+{
+	using namespace UE::Cameras;
+
+	const int32 NumRemoved = ExitTransitions.Remove(InTransition);
+	if (NumRemoved > 0)
+	{
+		TCameraArrayChangedEvent<UCameraRigTransition*> ChangedEvent;
+		ChangedEvent.EventType = ECameraArrayChangedEventType::Remove;
+		EventHandlers.Notify(&ICameraAssetEventHandler::OnExitTransitionsChanged, this, ChangedEvent);
+	}
+	return NumRemoved;
+}
+
+#if WITH_EDITOR
+
+void UCameraAsset::PostEditChangeProperty(FPropertyChangedEvent& PropertyChangedEvent)
+{
+	using namespace UE::Cameras;
+
+	const FName PropertyName = PropertyChangedEvent.GetPropertyName();
+	if (PropertyName == GET_MEMBER_NAME_CHECKED(UCameraAsset, CameraDirector))
+	{
+		TCameraPropertyChangedEvent<UCameraDirector*> ChangedEvent;
+		ChangedEvent.NewValue = CameraDirector;
+		EventHandlers.Notify(&ICameraAssetEventHandler::OnCameraDirectorChanged, this, ChangedEvent);
+	}
+	else if (PropertyName == GET_MEMBER_NAME_CHECKED(UCameraAsset, CameraRigs))
+	{
+		TCameraArrayChangedEvent<UCameraRigAsset*> ChangedEvent(PropertyChangedEvent.ChangeType);
+		EventHandlers.Notify(&ICameraAssetEventHandler::OnCameraRigsChanged, this, ChangedEvent);
+	}
+	else if (PropertyName == GET_MEMBER_NAME_CHECKED(UCameraAsset, EnterTransitions))
+	{
+		TCameraArrayChangedEvent<UCameraRigTransition*> ChangedEvent(PropertyChangedEvent.ChangeType);
+		EventHandlers.Notify(&ICameraAssetEventHandler::OnEnterTransitionsChanged, this, ChangedEvent);
+	}
+	else if (PropertyName == GET_MEMBER_NAME_CHECKED(UCameraAsset, ExitTransitions))
+	{
+		TCameraArrayChangedEvent<UCameraRigTransition*> ChangedEvent(PropertyChangedEvent.ChangeType);
+		EventHandlers.Notify(&ICameraAssetEventHandler::OnExitTransitionsChanged, this, ChangedEvent);
+	}
+
+	Super::PostEditChangeProperty(PropertyChangedEvent);
+}
+
+#endif
+
 void UCameraAsset::BuildCamera()
 {
 	using namespace UE::Cameras;
diff --git a/Engine/Plugins/Cameras/GameplayCameras/Source/GameplayCameras/Private/Core/CameraAssetBuilder.cpp b/Engine/Plugins/Cameras/GameplayCameras/Source/GameplayCameras/Private/Core/CameraAssetBuilder.cpp
--- a/Engine/Plugins/Cameras/GameplayCameras/Source/GameplayCameras/Private/Core/CameraAssetBuilder.cpp
+++ b/Engine/Plugins/Cameras/GameplayCameras/Source/GameplayCameras/Private/Core/CameraAssetBuilder.cpp
@@ -37,21 +37,21 @@ void FCameraAssetBuilder::BuildCamera(UCameraAsset* InCameraAsset)
 
 void FCameraAssetBuilder::BuildCameraImpl()
 {
-	if (CameraAsset->CameraDirector)
+	if (UCameraDirector* CameraDirector = CameraAsset->GetCameraDirector())
 	{
-		CameraAsset->CameraDirector->BuildCameraDirector(BuildLog);
+		CameraDirector->BuildCameraDirector(BuildLog);
 	}
 	else
 	{
 		BuildLog.AddMessage(EMessageSeverity::Error, LOCTEXT("MissingDirector", "Camera has no director set."));
 	}
 
-	if (CameraAsset->CameraRigs.IsEmpty())
+	if (CameraAsset->GetCameraRigs().IsEmpty())
 	{
 		BuildLog.AddMessage(EMessageSeverity::Warning, LOCTEXT("MissingRigs", "Camera has no camera rigs defined."));
 	}
 
-	for (UCameraRigAsset* CameraRig : CameraAsset->CameraRigs)
+	for (UCameraRigAsset* CameraRig : CameraAsset->GetCameraRigs())
 	{
 		FCameraRigAssetBuilder CameraRigBuilder(BuildLog);
 		CameraRigBuilder.BuildCameraRig(CameraRig);
@@ -70,10 +70,10 @@ void FCameraAssetBuilder::UpdateBuildStatus()
 		BuildStatus = ECameraBuildStatus::CleanWithWarnings;
 	}
 
-	if (CameraAsset->BuildStatus != BuildStatus)
+	if (CameraAsset->GetBuildStatus() != BuildStatus)
 	{
 		CameraAsset->Modify();
-		CameraAsset->BuildStatus = BuildStatus;
+		CameraAsset->SetBuildStatus(BuildStatus);
 	}
 }
 
diff --git a/Engine/Plugins/Cameras/GameplayCameras/Source/GameplayCameras/Private/Core/CameraEvaluationContext.cpp b/Engine/Plugins/Cameras/GameplayCameras/Source/GameplayCameras/Private/Core/CameraEvaluationContext.cpp
--- a/Engine/Plugins/Cameras/GameplayCameras/Source/GameplayCameras/Private/Core/CameraEvaluationContext.cpp
+++ b/Engine/Plugins/Cameras/GameplayCameras/Source/GameplayCameras/Private/Core/CameraEvaluationContext.cpp
@@ -73,13 +73,13 @@ void FCameraEvaluationContext::AutoCreateDirectorEvaluator()
 			UE_LOG(LogCameraSystem, Error, TEXT("Activating an evaluation context without a camera!"));
 			return;
 		}
-		if (!CameraAsset->CameraDirector)
+		if (!CameraAsset->GetCameraDirector())
 		{
 			UE_LOG(LogCameraSystem, Error, TEXT("Activating an evaluation context without a camera director!"));
 			return;
 		}
 
-		const UCameraDirector* CameraDirector = CameraAsset->CameraDirector;
+		const UCameraDirector* CameraDirector = CameraAsset->GetCameraDirector();
 		FCameraDirectorEvaluatorBuilder DirectorBuilder(DirectorEvaluatorStorage);
 		DirectorEvaluator = CameraDirector->BuildEvaluator(DirectorBuilder);
 
diff --git a/Engine/Plugins/Cameras/GameplayCameras/Source/GameplayCameras/Private/Directors/SingleCameraDirector.cpp b/Engine/Plugins/Cameras/GameplayCameras/Source/GameplayCameras/Private/Directors/SingleCameraDirector.cpp
--- a/Engine/Plugins/Cameras/GameplayCameras/Source/GameplayCameras/Private/Directors/SingleCameraDirector.cpp
+++ b/Engine/Plugins/Cameras/GameplayCameras/Source/GameplayCameras/Private/Directors/SingleCameraDirector.cpp
@@ -57,9 +57,9 @@ void USingleCameraDirector::OnFactoryCreateAsset(const FCameraDirectorFactoryCre
 	// Automatically set ourselves to use the first camera rig available.
 	if (UCameraAsset* OuterCameraAsset = GetTypedOuter<UCameraAsset>())
 	{
-		if (OuterCameraAsset->CameraRigs.Num() > 0)
+		if (OuterCameraAsset->GetCameraRigs().Num() > 0)
 		{
-			CameraRig = OuterCameraAsset->CameraRigs[0];
+			CameraRig = OuterCameraAsset->GetCameraRigs()[0];
 		}
 	}
 }
diff --git a/Engine/Plugins/Cameras/GameplayCameras/Source/GameplayCameras/Public/Core/CameraAsset.h b/Engine/Plugins/Cameras/GameplayCameras/Source/GameplayCameras/Public/Core/CameraAsset.h
--- a/Engine/Plugins/Cameras/GameplayCameras/Source/GameplayCameras/Public/Core/CameraAsset.h
+++ b/Engine/Plugins/Cameras/GameplayCameras/Source/GameplayCameras/Public/Core/CameraAsset.h
@@ -3,6 +3,7 @@
 #pragma once
 
 #include "Core/CameraBuildStatus.h"
+#include "Core/CameraEventHandler.h"
 #include "Core/CameraRigTransition.h"
 #include "Core/ObjectTreeGraphObject.h"
 #include "Core/ObjectTreeGraphRootObject.h"
@@ -11,10 +12,32 @@
 
 #include "CameraAsset.generated.h"
 
+class UCameraAsset;
 class UCameraDirector;
 class UCameraRigAsset;
 
-namespace UE::Cameras { class FCameraBuildLog; }
+namespace UE::Cameras
+{
+	class FCameraBuildLog;
+
+	/**
+	 * Interface for listening to changes on a camera asset.
+	 */
+	class ICameraAssetEventHandler
+	{
+	public:
+		virtual ~ICameraAssetEventHandler() {}
+
+		/** Called when the camera director has been changed. */
+		virtual void OnCameraDirectorChanged(UCameraAsset* InCameraAsset, const TCameraPropertyChangedEvent<UCameraDirector*>& Event) {}
+		/** Changed when the camera rigs have been changed. */
+		virtual void OnCameraRigsChanged(UCameraAsset* InCameraAsset, const TCameraArrayChangedEvent<UCameraRigAsset*>& Event) {}
+		/* Changed when the enter transitions have been changed. */
+		virtual void OnEnterTransitionsChanged(UCameraAsset* InCameraAsset, const TCameraArrayChangedEvent<UCameraRigTransition*>& Event) {}
+		/* Changed when the exit transitions have been changed. */
+		virtual void OnExitTransitionsChanged(UCameraAsset* InCameraAsset, const TCameraArrayChangedEvent<UCameraRigTransition*>& Event) {}
+	};
+}
 
 /**
  * A complete camera asset.
@@ -30,28 +53,34 @@ class UCameraAsset
 
 public:
 
-	/** The camera director to use in this camera. */
-	UPROPERTY(Instanced)
-	TObjectPtr<UCameraDirector> CameraDirector;
-
-	/** The list of camera rigs used by this camera. */
-	UPROPERTY()
-	TArray<TObjectPtr<UCameraRigAsset>> CameraRigs;
-
-	/** A list of default enter transitions for all the camera rigs in this asset. */
-	UPROPERTY()
-	TArray<TObjectPtr<UCameraRigTransition>> EnterTransitions;
-
-	/** A list of default exit transitions for all the camera rigs in this asset. */
-	UPROPERTY()
-	TArray<TObjectPtr<UCameraRigTransition>> ExitTransitions;
+	/** Gets the camera director. */
+	UCameraDirector* GetCameraDirector() const { return CameraDirector; }
+	/** Sets the camera director. */
+	GAMEPLAYCAMERAS_API void SetCameraDirector(UCameraDirector* InCameraDirector);
+
+	/** Gets the camera rigs. */
+	TArrayView<const TObjectPtr<UCameraRigAsset>> GetCameraRigs() const { return CameraRigs; }
+	/** Adds a a camera rig. */
+	GAMEPLAYCAMERAS_API void AddCameraRig(UCameraRigAsset* InCameraRig);
+	/** Removes a camera rig. */
+	GAMEPLAYCAMERAS_API int32 RemoveCameraRig(UCameraRigAsset* InCameraRig);
+
+	/** Gets the enter transitions. */
+	TArrayView<const TObjectPtr<UCameraRigTransition>> GetEnterTransitions() const { return EnterTransitions; }
+	/** Adds an enter transition. */
+	GAMEPLAYCAMERAS_API void AddEnterTransition(UCameraRigTransition* InTransition);
+	/** Removes an enter transition. */
+	GAMEPLAYCAMERAS_API int32 RemoveEnterTransition(UCameraRigTransition* InTransition);
+
+	/** Gets the exit transitions. */
+	TArrayView<const TObjectPtr<UCameraRigTransition>> GetExitTransitions() const { return ExitTransitions; }
+	/** Adds an exit transition. */
+	GAMEPLAYCAMERAS_API void AddExitTransition(UCameraRigTransition* InTransition);
+	/** Removes an exit transition. */
+	GAMEPLAYCAMERAS_API int32 RemoveExitTransition(UCameraRigTransition* InTransition);
 
 public:
 
-	/** The current build state of this camera asset. */
-	UPROPERTY(Transient)
-	ECameraBuildStatus BuildStatus = ECameraBuildStatus::Dirty;
-
 	/**
 	 * Builds and validates this camera, including all its camera rigs.
 	 * Errors and warnings will go to the console.
@@ -73,6 +102,9 @@ public:
 	virtual ECameraBuildStatus GetBuildStatus() const override { return BuildStatus; }
 	virtual void DirtyBuildStatus() override;
 
+	/** Sets the build status. */
+	void SetBuildStatus(ECameraBuildStatus InBuildStatus) { BuildStatus = InBuildStatus; }
+
 protected:
 
 	// IObjectTreeGraphObject interface.
@@ -91,8 +123,38 @@ protected:
 	virtual void RemoveConnectableObject(FName InGraphName, UObject* InObject) override;
 #endif
 
+	// UObject interface.
+#if WITH_EDITOR
+	virtual void PostEditChangeProperty(FPropertyChangedEvent& PropertyChangedEvent) override;
+#endif
+
+public:
+
+	/** Event handlers to be notified of data changes. */
+	UE::Cameras::TCameraEventHandlerContainer<UE::Cameras::ICameraAssetEventHandler> EventHandlers;
+
 private:
 
+	/** The camera director to use in this camera. */
+	UPROPERTY(Instanced)
+	TObjectPtr<UCameraDirector> CameraDirector;
+
+	/** The list of camera rigs used by this camera. */
+	UPROPERTY()
+	TArray<TObjectPtr<UCameraRigAsset>> CameraRigs;
+
+	/** A list of default enter transitions for all the camera rigs in this asset. */
+	UPROPERTY()
+	TArray<TObjectPtr<UCameraRigTransition>> EnterTransitions;
+
+	/** A list of default exit transitions for all the camera rigs in this asset. */
+	UPROPERTY()
+	TArray<TObjectPtr<UCameraRigTransition>> ExitTransitions;
+
+	/** The current build state of this camera asset. */
+	UPROPERTY(Transient)
+	ECameraBuildStatus BuildStatus = ECameraBuildStatus::Dirty;
+
 #if WITH_EDITORONLY_DATA
 
 	/** Position of the camera node in the shared transitions graph editor. */
@@ -107,6 +169,9 @@ private:
 	UPROPERTY(Instanced, meta=(ObjectTreeGraphHidden=true))
 	TArray<TObjectPtr<UObject>> AllSharedTransitionsObjects;
 
+	// Only specified here so that the schema can use GET_MEMBER_NAME_CHECKED...
+	friend class UCameraSharedTransitionGraphSchema;
+
 #endif  // WITH_EDITORONLY_DATA
 };
 
diff --git a/Engine/Plugins/Cameras/GameplayCameras/Source/GameplayCameras/Public/Core/CameraEventHandler.h b/Engine/Plugins/Cameras/GameplayCameras/Source/GameplayCameras/Public/Core/CameraEventHandler.h
new file mode 100644
--- /dev/null
+++ b/Engine/Plugins/Cameras/GameplayCameras/Source/GameplayCameras/Public/Core/CameraEventHandler.h
@@ -0,0 +1,117 @@
+// Copyright Epic Games, Inc. All Rights Reserved.
+
+#pragma once
+
+#include "Containers/ArrayView.h"
+#include "Containers/List.h"
+#include "GameplayCameras.h"
+#include "UObject/UnrealType.h"
+
+namespace UE::Cameras
+{
+
+/** The type of change on an array property.  */
+enum class ECameraArrayChangedEventType
+{
+	Add,
+	Remove,
+	Change
+};
+
+/** Parameter structure for a value property change event. */
+template<typename T>
+struct TCameraPropertyChangedEvent
+{
+	T NewValue;
+};
+
+/** Parameter structure for an array property change event. */
+template<typename T>
+struct TCameraArrayChangedEvent
+{
+	ECameraArrayChangedEventType EventType;
+
+	TCameraArrayChangedEvent() {}
+	TCameraArrayChangedEvent(EPropertyChangeType::Type InPropertyChangeType)
+	{
+		switch (InPropertyChangeType)
+		{
+			case EPropertyChangeType::ArrayAdd:
+				EventType = ECameraArrayChangedEventType::Add;
+				break;
+			case EPropertyChangeType::ArrayRemove:
+				EventType = ECameraArrayChangedEventType::Remove;
+				break;
+			default:
+				EventType = ECameraArrayChangedEventType::Change;
+				break;
+		}
+	}
+};
+
+/**
+ * Wrapper struct, or handle, for a listener of data change events.
+ * The listener class should own one of these for each object they're listening to.
+ * The listener class can start listening by calling Register(HandlerWrapper, this) on the object
+ * to listen to.
+ */
+template<typename HandlerInterface>
+struct TCameraEventHandler : public TLinkedList<HandlerInterface*>
+{
+	TCameraEventHandler()
+	{}
+
+	TCameraEventHandler(HandlerInterface* InInterface)
+		: TLinkedList<HandlerInterface*>(InInterface)
+	{}
+
+	~TCameraEventHandler()
+	{
+		this->Unlink();
+	}
+};
+
+/**
+ * A wrapper struct for a list of listeners waiting to be notified of data changes.
+ * A class that can be listened-to should own one of these, and expose it so that listeners
+ * can register themselves.
+ */
+template<typename HandlerInterface>
+struct TCameraEventHandlerContainer
+{
+#if UE_GAMEPLAY_CAMERAS_EVENT_HANDLERS
+
+	/** Invokes the given callback, with the given parameters, on any registered listener. */
+	template<typename FuncType, typename... ArgTypes>
+	void Notify(FuncType&& Func, ArgTypes&&... Args) const
+	{
+		for (typename TCameraEventHandler<HandlerInterface>::TIterator It(EventHandlers); It; ++It)
+		{
+			Invoke(Func, *It, Forward<ArgTypes>(Args)...);
+		}
+	}
+
+	/** Registers a new listener by linking the given wrapper/handle to the list of listeners. */
+	void Register(TCameraEventHandler<HandlerInterface>& InOutEventHandler, HandlerInterface* InInterface) const
+	{
+		ensure(!InOutEventHandler.IsLinked());
+		InOutEventHandler = TCameraEventHandler<HandlerInterface>(InInterface);
+		InOutEventHandler.LinkHead(EventHandlers);
+	}
+
+private:
+
+	mutable TLinkedList<HandlerInterface*>* EventHandlers = nullptr;
+
+#else
+
+	// Stub for when events are disabled.
+	template<typename ...ArgTypes>
+	void Notify(ArgTypes&&...) const
+	{}
+
+#endif  // UE_GAMEPLAY_CAMERAS_EVENT_HANDLERS
+};
+
+}  // namespace UE::Cameras
+
diff --git a/Engine/Plugins/Cameras/GameplayCameras/Source/GameplayCameras/Public/GameplayCameras.h b/Engine/Plugins/Cameras/GameplayCameras/Source/GameplayCameras/Public/GameplayCameras.h
--- a/Engine/Plugins/Cameras/GameplayCameras/Source/GameplayCameras/Public/GameplayCameras.h
+++ b/Engine/Plugins/Cameras/GameplayCameras/Source/GameplayCameras/Public/GameplayCameras.h
@@ -7,6 +7,10 @@
 #include "HAL/Platform.h"
 #include "Stats/Stats.h"
 
+#ifndef UE_GAMEPLAY_CAMERAS_EVENT_HANDLERS
+	#define UE_GAMEPLAY_CAMERAS_EVENT_HANDLERS WITH_EDITOR
+#endif
+
 #ifndef UE_GAMEPLAY_CAMERAS_DEBUG
 	#if UE_ENABLE_DEBUG_DRAWING
 		#define UE_GAMEPLAY_CAMERAS_DEBUG 1
diff --git a/Engine/Plugins/Cameras/GameplayCameras/Source/GameplayCamerasEditor/Private/Editors/SCameraRigList.cpp b/Engine/Plugins/Cameras/GameplayCameras/Source/GameplayCamerasEditor/Private/Editors/SCameraRigList.cpp
--- a/Engine/Plugins/Cameras/GameplayCameras/Source/GameplayCamerasEditor/Private/Editors/SCameraRigList.cpp
+++ b/Engine/Plugins/Cameras/GameplayCameras/Source/GameplayCamerasEditor/Private/Editors/SCameraRigList.cpp
@@ -87,7 +87,7 @@ bool SCameraRigListEntry::OnVerifyTextChanged(const FText& Text, FText& OutError
 	}
 
 	const FString TextString = Text.ToString();
-	TObjectPtr<UCameraRigAsset>* FoundItem = OwnerCamera->CameraRigs.FindByPredicate(
+	const TObjectPtr<UCameraRigAsset>* FoundItem = OwnerCamera->GetCameraRigs().FindByPredicate(
 			[&TextString](const UCameraRigAsset* Item) { return Item->GetDisplayName() == TextString; });
 	if (FoundItem)
 	{
@@ -283,7 +283,7 @@ void SCameraRigList::OnAddCameraRig()
 
 	const FName NewCameraRigName = MakeUniqueObjectName(CameraAsset, UCameraRigAsset::StaticClass(), TEXT("NewCameraRig"));
 	UCameraRigAsset* NewCameraRig = NewObject<UCameraRigAsset>(CameraAsset, NewCameraRigName, RF_Transactional);
-	CameraAsset->CameraRigs.Add(NewCameraRig);
+	CameraAsset->AddCameraRig(NewCameraRig);
 
 	bUpdateItemSource = true;
 }
@@ -331,7 +331,7 @@ void SCameraRigList::OnDeleteCameraRig()
 				StringBuilder.Append(CameraRigAsset->GetName());
 				CameraRigAsset->Rename(StringBuilder.ToString());
 
-				const int32 NumRemoved = CameraAsset->CameraRigs.Remove(CameraRigAsset);
+				const int32 NumRemoved = CameraAsset->RemoveCameraRig(CameraRigAsset);
 				ensure(NumRemoved == 1);
 
 				DeletedCameraRigs.Add(CameraRigAsset);
@@ -366,7 +366,7 @@ void SCameraRigList::UpdateItemSource()
 
 	if (CameraAsset)
 	{
-		for (UCameraRigAsset* CameraRigAsset : CameraAsset->CameraRigs)
+		for (UCameraRigAsset* CameraRigAsset : CameraAsset->GetCameraRigs())
 		{
 			TSharedPtr<FCameraRigListItem> Item = MakeShared<FCameraRigListItem>();
 			Item->CameraRigAsset = CameraRigAsset;
@@ -374,7 +374,7 @@ void SCameraRigList::UpdateItemSource()
 		}
 	}
 
-	OnCameraRigListChanged.ExecuteIfBound(CameraAsset->CameraRigs);
+	OnCameraRigListChanged.ExecuteIfBound(CameraAsset->GetCameraRigs());
 }
 
 void SCameraRigList::UpdateFilteredItemSource()
diff --git a/Engine/Plugins/Cameras/GameplayCameras/Source/GameplayCamerasEditor/Private/Editors/SCameraRigPicker.cpp b/Engine/Plugins/Cameras/GameplayCameras/Source/GameplayCamerasEditor/Private/Editors/SCameraRigPicker.cpp
--- a/Engine/Plugins/Cameras/GameplayCameras/Source/GameplayCamerasEditor/Private/Editors/SCameraRigPicker.cpp
+++ b/Engine/Plugins/Cameras/GameplayCameras/Source/GameplayCamerasEditor/Private/Editors/SCameraRigPicker.cpp
@@ -237,7 +237,7 @@ void SCameraRigPicker::UpdateCameraRigItemsSource(const TArray<FAssetData>& Asse
 	{
 		if (const UCameraAsset* CameraAsset = Cast<UCameraAsset>(SelectedAsset.GetAsset()))
 		{
-			CameraRigItemsSource.Append(CameraAsset->CameraRigs);
+			CameraRigItemsSource.Append(CameraAsset->GetCameraRigs());
 		}
 	}
 }
diff --git a/Engine/Plugins/Cameras/GameplayCameras/Source/GameplayCamerasEditor/Private/Factories/CameraAssetFactory.cpp b/Engine/Plugins/Cameras/GameplayCameras/Source/GameplayCamerasEditor/Private/Factories/CameraAssetFactory.cpp
--- a/Engine/Plugins/Cameras/GameplayCameras/Source/GameplayCamerasEditor/Private/Factories/CameraAssetFactory.cpp
+++ b/Engine/Plugins/Cameras/GameplayCameras/Source/GameplayCamerasEditor/Private/Factories/CameraAssetFactory.cpp
@@ -61,12 +61,12 @@ UObject* UCameraAssetFactory::FactoryCreateNew(UClass* Class, UObject* Parent, F
 	UCameraAsset* NewCameraAsset = NewObject<UCameraAsset>(Parent, Class, Name, Flags | RF_Transactional);
 
 	UCameraRigAsset* FirstCameraRig = NewObject<UCameraRigAsset>(NewCameraAsset, TEXT("NewCameraRig"));
-	NewCameraAsset->CameraRigs.Add(FirstCameraRig);
+	NewCameraAsset->AddCameraRig(FirstCameraRig);
 
 	if (CameraDirectorClass)
 	{
 		UCameraDirector* NewCameraDirector = NewObject<UCameraDirector>(NewCameraAsset, CameraDirectorClass, NAME_None, Flags | RF_Transactional);
-		NewCameraAsset->CameraDirector = NewCameraDirector;
+		NewCameraAsset->SetCameraDirector(NewCameraDirector);
 
 		// Let the camera director do some scaffolding.
 		FCameraDirectorFactoryCreateParams CreateParams;
diff --git a/Engine/Plugins/Cameras/GameplayCameras/Source/GameplayCamerasEditor/Private/Toolkits/CameraAssetEditorToolkit.cpp b/Engine/Plugins/Cameras/GameplayCameras/Source/GameplayCamerasEditor/Private/Toolkits/CameraAssetEditorToolkit.cpp
--- a/Engine/Plugins/Cameras/GameplayCameras/Source/GameplayCamerasEditor/Private/Toolkits/CameraAssetEditorToolkit.cpp
+++ b/Engine/Plugins/Cameras/GameplayCameras/Source/GameplayCamerasEditor/Private/Toolkits/CameraAssetEditorToolkit.cpp
@@ -255,12 +255,12 @@ void FCameraAssetEditorToolkit::OnBuild()
 	CameraAsset->BuildCamera(BuildLog);
 	BuildLogToolkit->PopulateMessageListing(BuildLog);
 
-	if (CameraAsset->BuildStatus != ECameraBuildStatus::Clean)
+	if (CameraAsset->GetBuildStatus() != ECameraBuildStatus::Clean)
 	{
 		TabManager->TryInvokeTab(MessagesTabId);
 	}
 
-	for (UCameraRigAsset* CameraRigAsset : CameraAsset->CameraRigs)
+	for (UCameraRigAsset* CameraRigAsset : CameraAsset->GetCameraRigs())
 	{
 		FCameraRigPackages BuiltPackages;
 		CameraRigAsset->GatherPackages(BuiltPackages);
diff --git a/Engine/Plugins/Cameras/GameplayCameras/Source/GameplayCamerasEditor/Private/Toolkits/CameraDirectorAssetEditorMode.cpp b/Engine/Plugins/Cameras/GameplayCameras/Source/GameplayCamerasEditor/Private/Toolkits/CameraDirectorAssetEditorMode.cpp
--- a/Engine/Plugins/Cameras/GameplayCameras/Source/GameplayCamerasEditor/Private/Toolkits/CameraDirectorAssetEditorMode.cpp
+++ b/Engine/Plugins/Cameras/GameplayCameras/Source/GameplayCamerasEditor/Private/Toolkits/CameraDirectorAssetEditorMode.cpp
@@ -44,7 +44,7 @@ void FCameraDirectorAssetEditorMode::OnActivateMode(const FAssetEditorModeActiva
 		bInitializedToolkit = true;
 	}
 
-	DetailsView->SetObject(CameraAsset->CameraDirector);
+	DetailsView->SetObject(CameraAsset->GetCameraDirector());
 
 	InParams.TabManager->RegisterTabSpawner(DirectorEditorTabId, FOnSpawnTab::CreateSP(this, &FCameraDirectorAssetEditorMode::SpawnTab_DirectorEditor))
 		.SetDisplayName(LOCTEXT("CameraDirectorEditor", "Camera Director"))
@@ -69,7 +69,7 @@ void FCameraDirectorAssetEditorMode::OnDeactivateMode(const FAssetEditorModeDeac
 
 bool FCameraDirectorAssetEditorMode::JumpToObject(UObject* InObject, FName InPropertyName)
 {
-	if (InObject == CameraAsset->CameraDirector)
+	if (InObject == CameraAsset->GetCameraDirector())
 	{
 		return true;
 	}
diff --git a/Engine/Plugins/Cameras/GameplayCameras/Source/GameplayCamerasEditor/Private/Toolkits/CameraRigsAssetEditorMode.cpp b/Engine/Plugins/Cameras/GameplayCameras/Source/GameplayCamerasEditor/Private/Toolkits/CameraRigsAssetEditorMode.cpp
--- a/Engine/Plugins/Cameras/GameplayCameras/Source/GameplayCamerasEditor/Private/Toolkits/CameraRigsAssetEditorMode.cpp
+++ b/Engine/Plugins/Cameras/GameplayCameras/Source/GameplayCamerasEditor/Private/Toolkits/CameraRigsAssetEditorMode.cpp
@@ -115,7 +115,7 @@ void FCameraRigsAssetEditorMode::OnCameraRigDeleted(const TArray<UCameraRigAsset
 
 void FCameraRigsAssetEditorMode::OnGetRootObjectsToSearch(TArray<FFindInObjectTreeGraphSource>& OutSources)
 {
-	for (UCameraRigAsset* CameraRig : CameraAsset->CameraRigs)
+	for (UCameraRigAsset* CameraRig : CameraAsset->GetCameraRigs())
 	{
 		OutSources.Add(FFindInObjectTreeGraphSource{ CameraRig, &NodeGraphConfig });
 		OutSources.Add(FFindInObjectTreeGraphSource{ CameraRig, &TransitionGraphConfig });
