Cameras: add ability to use custom director editors, make some for the single-rig and blueprint directors

These editors react to camera rigs changing in the camera asset in order to adjust some of their own data.

diff --git a/Engine/Plugins/Cameras/GameplayCameras/Source/GameplayCameras/Public/IGameplayCamerasModule.h b/Engine/Plugins/Cameras/GameplayCameras/Source/GameplayCameras/Public/IGameplayCamerasModule.h
--- a/Engine/Plugins/Cameras/GameplayCameras/Source/GameplayCameras/Public/IGameplayCamerasModule.h
+++ b/Engine/Plugins/Cameras/GameplayCameras/Source/GameplayCameras/Public/IGameplayCamerasModule.h
@@ -23,7 +23,7 @@ public:
 	 *
 	 * @return The ICameraModule instance, loading the module on demand if needed
 	 */
-	static IGameplayCamerasModule& Get();
+	GAMEPLAYCAMERAS_API static IGameplayCamerasModule& Get();
 
 #if WITH_EDITOR
 	using IGameplayCamerasLiveEditManager = UE::Cameras::IGameplayCamerasLiveEditManager;
diff --git a/Engine/Plugins/Cameras/GameplayCameras/Source/GameplayCamerasEditor/Private/GameplayCamerasEditorModule.cpp b/Engine/Plugins/Cameras/GameplayCameras/Source/GameplayCamerasEditor/Private/GameplayCamerasEditorModule.cpp
--- a/Engine/Plugins/Cameras/GameplayCameras/Source/GameplayCamerasEditor/Private/GameplayCamerasEditorModule.cpp
+++ b/Engine/Plugins/Cameras/GameplayCameras/Source/GameplayCamerasEditor/Private/GameplayCamerasEditorModule.cpp
@@ -37,8 +37,10 @@
 #include "PropertyEditorModule.h"
 #include "Styles/GameplayCamerasEditorStyle.h"
 #include "ToolMenus.h"
+#include "Toolkits/BlueprintCameraDirectorAssetEditorMode.h"
 #include "Toolkits/CameraAssetEditorToolkit.h"
 #include "Toolkits/CameraRigAssetEditorToolkit.h"
+#include "Toolkits/SingleCameraDirectorAssetEditorMode.h"
 #include "Trace/CameraSystemRewindDebuggerExtension.h"
 #include "Trace/CameraSystemRewindDebuggerTrack.h"
 #include "Trace/CameraSystemTraceModule.h"
@@ -50,6 +52,11 @@ DEFINE_LOG_CATEGORY(LogCameraSystemEditor);
 const FName IGameplayCamerasEditorModule::GameplayCamerasEditorAppIdentifier("GameplayCamerasEditorApp");
 const FName IGameplayCamerasEditorModule::CameraRigAssetEditorToolBarName("CameraRigAssetEditor.ToolBar");
 
+IGameplayCamerasEditorModule& IGameplayCamerasEditorModule::Get()
+{
+	return FModuleManager::LoadModuleChecked<IGameplayCamerasEditorModule>("GameplayCamerasEditor");
+}
+
 /**
  * Implements the FGameplayCamerasEditor module.
  */
@@ -74,6 +81,7 @@ public:
 		FCoreDelegates::OnEnginePreExit.AddRaw(this, &FGameplayCamerasEditorModule::OnPreExit);
 
 		RegisterSettings();
+		RegisterCameraDirectorEditors();
 		RegisterCoreDebugCategories();
 		RegisterRewindDebuggerFeatures();
 		RegisterDetailsCustomizations();
@@ -98,6 +106,7 @@ public:
 		FGameplayCamerasDebuggerCommands::Unregister();
 
 		UnregisterSettings();
+		UnregisterCameraDirectorEditors();
 		UnregisterCoreDebugCategories();
 		UnregisterRewindDebuggerFeatures();
 		UnregisterDetailsCustomizations();
@@ -149,6 +158,26 @@ public:
 			.CameraRigPickerConfig(InPickerConfig);
 	}
 
+	virtual FDelegateHandle RegisterCameraDirectorEditor(FOnCreateCameraDirectorAssetEditorMode InOnCreateEditor) override
+	{
+		CameraDirectorEditorCreators.Add(InOnCreateEditor);
+		return CameraDirectorEditorCreators.Last().GetHandle();
+	}
+
+	virtual TArrayView<const FOnCreateCameraDirectorAssetEditorMode> GetCameraDirectorEditorCreators() const override
+	{
+		return CameraDirectorEditorCreators;
+	}
+
+	virtual void UnregisterCameraDirectorEditor(FDelegateHandle InHandle)
+	{
+		CameraDirectorEditorCreators.RemoveAll(
+				[=](const FOnCreateCameraDirectorAssetEditorMode& Delegate) 
+				{
+					return Delegate.GetHandle() == InHandle; 
+				});
+	}
+
 	virtual void RegisterDebugCategory(const UE::Cameras::FCameraDebugCategoryInfo& InCategoryInfo) override
 	{
 		if (!ensureMsgf(!InCategoryInfo.Name.IsEmpty(), TEXT("A debug category must at least specify a name!")))
@@ -237,6 +266,22 @@ private:
 		}
 	}
 
+	void RegisterCameraDirectorEditors()
+	{
+		using namespace UE::Cameras;
+
+		BuiltInDirectorCreatorHandles.Add(
+				RegisterCameraDirectorEditor(FOnCreateCameraDirectorAssetEditorMode::CreateStatic(
+						&FSingleCameraDirectorAssetEditorMode::CreateInstance)));
+		BuiltInDirectorCreatorHandles.Add(
+				RegisterCameraDirectorEditor(FOnCreateCameraDirectorAssetEditorMode::CreateStatic(
+						&FBlueprintCameraDirectorAssetEditorMode::CreateInstance)));
+	}
+
+	void UnregisterCameraDirectorEditors()
+	{
+	}
+
 	void RegisterCoreDebugCategories()
 	{
 		using namespace UE::Cameras;
@@ -458,6 +503,9 @@ private:
 
 	TSharedPtr<UE::Cameras::FGameplayCamerasLiveEditManager> LiveEditManager;
 
+	TArray<FOnCreateCameraDirectorAssetEditorMode> CameraDirectorEditorCreators;
+	TArray<FDelegateHandle> BuiltInDirectorCreatorHandles;
+
 	TSharedPtr<UE::Cameras::FGameplayCamerasGraphPanelPinFactory> GraphPanelPinFactory;
 
 	TMap<FString, UE::Cameras::FCameraDebugCategoryInfo> DebugCategoryInfos;
diff --git a/Engine/Plugins/Cameras/GameplayCameras/Source/GameplayCamerasEditor/Private/Toolkits/BlueprintCameraDirectorAssetEditorMode.cpp b/Engine/Plugins/Cameras/GameplayCameras/Source/GameplayCamerasEditor/Private/Toolkits/BlueprintCameraDirectorAssetEditorMode.cpp
new file mode 100644
--- /dev/null
+++ b/Engine/Plugins/Cameras/GameplayCameras/Source/GameplayCamerasEditor/Private/Toolkits/BlueprintCameraDirectorAssetEditorMode.cpp
@@ -0,0 +1,51 @@
+// Copyright Epic Games, Inc. All Rights Reserved.
+
+#include "Toolkits/BlueprintCameraDirectorAssetEditorMode.h"
+
+#include "Core/CameraRigProxyTable.h"
+#include "Directors/BlueprintCameraDirector.h"
+
+namespace UE::Cameras
+{
+
+TSharedPtr<FCameraDirectorAssetEditorMode> FBlueprintCameraDirectorAssetEditorMode::CreateInstance(UCameraAsset* InCameraAsset)
+{
+	UCameraDirector* CameraDirector = InCameraAsset->GetCameraDirector();
+	if (Cast<UBlueprintCameraDirector>(CameraDirector))
+	{
+		return MakeShared<FBlueprintCameraDirectorAssetEditorMode>(InCameraAsset);
+	}
+	return nullptr;
+}
+
+FBlueprintCameraDirectorAssetEditorMode::FBlueprintCameraDirectorAssetEditorMode(UCameraAsset* InCameraAsset)
+	: FCameraDirectorAssetEditorMode(InCameraAsset)
+{
+	if (InCameraAsset)
+	{
+		InCameraAsset->EventHandlers.Register(EventHandler, this);
+	}
+}
+
+void FBlueprintCameraDirectorAssetEditorMode::OnCameraRigsChanged(UCameraAsset* InCameraAsset, const TCameraArrayChangedEvent<UCameraRigAsset*>& Event)
+{
+	UBlueprintCameraDirector* CameraDirector = Cast<UBlueprintCameraDirector>(InCameraAsset->GetCameraDirector());
+	if (!ensure(CameraDirector))
+	{
+		return;
+	}
+
+	TSet<UCameraRigAsset*> CameraRigs(InCameraAsset->GetCameraRigs());
+
+	for (FCameraRigProxyTableEntry& Entry : CameraDirector->CameraRigProxyTable->Entries)
+	{
+		if (!CameraRigs.Contains(Entry.CameraRig))
+		{
+			CameraDirector->CameraRigProxyTable->Modify();
+			Entry.CameraRig = nullptr;
+		}
+	}
+}
+
+}  // namespace UE::Cameras
+
diff --git a/Engine/Plugins/Cameras/GameplayCameras/Source/GameplayCamerasEditor/Private/Toolkits/CameraAssetEditorToolkit.cpp b/Engine/Plugins/Cameras/GameplayCameras/Source/GameplayCamerasEditor/Private/Toolkits/CameraAssetEditorToolkit.cpp
--- a/Engine/Plugins/Cameras/GameplayCameras/Source/GameplayCamerasEditor/Private/Toolkits/CameraAssetEditorToolkit.cpp
+++ b/Engine/Plugins/Cameras/GameplayCameras/Source/GameplayCamerasEditor/Private/Toolkits/CameraAssetEditorToolkit.cpp
@@ -12,9 +12,9 @@
 #include "Framework/Docking/LayoutExtender.h"
 #include "Framework/Docking/TabManager.h"
 #include "GameplayCamerasEditorSettings.h"
+#include "IGameplayCamerasEditorModule.h"
 #include "IGameplayCamerasLiveEditManager.h"
 #include "IGameplayCamerasModule.h"
-#include "Modules/ModuleManager.h"
 #include "PropertyEditorModule.h"
 #include "Styles/GameplayCamerasEditorStyle.h"
 #include "ToolMenus.h"
@@ -197,8 +197,25 @@ void FCameraAssetEditorToolkit::PostInitAssetEditor()
 {
 	Settings = GetMutableDefault<UGameplayCamerasEditorSettings>();
 
+	IGameplayCamerasEditorModule& GameplayCamerasEditorModule = IGameplayCamerasEditorModule::Get();
+
 	const FName CameraDirectorModeName = FCameraDirectorAssetEditorMode::ModeName;
-	AddEditorMode(MakeShared<FCameraDirectorAssetEditorMode>(CameraAsset));
+	{
+		TSharedPtr<FCameraDirectorAssetEditorMode> CameraDirectorEditor;
+		for (const FOnCreateCameraDirectorAssetEditorMode& EditorCreator : GameplayCamerasEditorModule.GetCameraDirectorEditorCreators())
+		{
+			CameraDirectorEditor = EditorCreator.Execute(CameraAsset);
+			if (CameraDirectorEditor)
+			{
+				break;
+			}
+		}
+		if (!CameraDirectorEditor)
+		{
+			CameraDirectorEditor = MakeShared<FCameraDirectorAssetEditorMode>(CameraAsset);
+		}
+		AddEditorMode(CameraDirectorEditor.ToSharedRef());
+	}
 
 	const FName CameraRigsModeName = FCameraRigsAssetEditorMode::ModeName;
 	AddEditorMode(MakeShared<FCameraRigsAssetEditorMode>(CameraAsset));
@@ -230,7 +247,7 @@ void FCameraAssetEditorToolkit::PostInitAssetEditor()
 
 	BuildLogToolkit->OnRequestJumpToObject().BindSP(this, &FCameraAssetEditorToolkit::OnJumpToObject);
 
-	IGameplayCamerasModule& GameplayCamerasModule = FModuleManager::GetModuleChecked<IGameplayCamerasModule>("GameplayCameras");
+	IGameplayCamerasModule& GameplayCamerasModule = IGameplayCamerasModule::Get();
 	LiveEditManager = GameplayCamerasModule.GetLiveEditManager();
 
 	const FName InitialModeName = !Settings->LastCameraAssetToolkitModeName.IsNone() ? 
diff --git a/Engine/Plugins/Cameras/GameplayCameras/Source/GameplayCamerasEditor/Private/Toolkits/SingleCameraDirectorAssetEditorMode.cpp b/Engine/Plugins/Cameras/GameplayCameras/Source/GameplayCamerasEditor/Private/Toolkits/SingleCameraDirectorAssetEditorMode.cpp
new file mode 100644
--- /dev/null
+++ b/Engine/Plugins/Cameras/GameplayCameras/Source/GameplayCamerasEditor/Private/Toolkits/SingleCameraDirectorAssetEditorMode.cpp
@@ -0,0 +1,45 @@
+// Copyright Epic Games, Inc. All Rights Reserved.
+
+#include "Toolkits/SingleCameraDirectorAssetEditorMode.h"
+
+#include "Directors/SingleCameraDirector.h"
+
+namespace UE::Cameras
+{
+
+TSharedPtr<FCameraDirectorAssetEditorMode> FSingleCameraDirectorAssetEditorMode::CreateInstance(UCameraAsset* InCameraAsset)
+{
+	UCameraDirector* CameraDirector = InCameraAsset->GetCameraDirector();
+	if (Cast<USingleCameraDirector>(CameraDirector))
+	{
+		return MakeShared<FSingleCameraDirectorAssetEditorMode>(InCameraAsset);
+	}
+	return nullptr;
+}
+
+FSingleCameraDirectorAssetEditorMode::FSingleCameraDirectorAssetEditorMode(UCameraAsset* InCameraAsset)
+	: FCameraDirectorAssetEditorMode(InCameraAsset)
+{
+	if (InCameraAsset)
+	{
+		InCameraAsset->EventHandlers.Register(EventHandler, this);
+	}
+}
+
+void FSingleCameraDirectorAssetEditorMode::OnCameraRigsChanged(UCameraAsset* InCameraAsset, const TCameraArrayChangedEvent<UCameraRigAsset*>& Event)
+{
+	USingleCameraDirector* CameraDirector = Cast<USingleCameraDirector>(InCameraAsset->GetCameraDirector());
+	if (!ensure(CameraDirector))
+	{
+		return;
+	}
+
+	if (!InCameraAsset->GetCameraRigs().Contains(CameraDirector->CameraRig))
+	{
+		CameraDirector->Modify();
+		CameraDirector->CameraRig = nullptr;
+	}
+}
+
+}  // namespace UE::Cameras
+
diff --git a/Engine/Plugins/Cameras/GameplayCameras/Source/GameplayCamerasEditor/Public/IGameplayCamerasEditorModule.h b/Engine/Plugins/Cameras/GameplayCameras/Source/GameplayCamerasEditor/Public/IGameplayCamerasEditorModule.h
--- a/Engine/Plugins/Cameras/GameplayCameras/Source/GameplayCamerasEditor/Public/IGameplayCamerasEditorModule.h
+++ b/Engine/Plugins/Cameras/GameplayCameras/Source/GameplayCamerasEditor/Public/IGameplayCamerasEditorModule.h
@@ -17,11 +17,10 @@ class UCameraRigProxyAssetEditor;
 class UCameraVariableCollection;
 class UCameraVariableCollectionEditor;
 
-DECLARE_DELEGATE_RetVal_OneParam(TSharedRef<SWidget>, FOnCreateDebugCategoryPanel, const FString&);
-
 namespace UE::Cameras
 {
 
+class FCameraDirectorAssetEditorMode;
 struct FCameraRigPickerConfig;
 
 struct FCameraDebugCategoryInfo
@@ -34,6 +33,9 @@ struct FCameraDebugCategoryInfo
 
 }  // namespace UE::Cameras
 
+DECLARE_DELEGATE_RetVal_OneParam(TSharedPtr<UE::Cameras::FCameraDirectorAssetEditorMode>, FOnCreateCameraDirectorAssetEditorMode, UCameraAsset*);
+DECLARE_DELEGATE_RetVal_OneParam(TSharedRef<SWidget>, FOnCreateDebugCategoryPanel, const FString&);
+
 /**
  * The gameplay cameras editor module.
  */
@@ -44,6 +46,8 @@ public:
 	static const FName GameplayCamerasEditorAppIdentifier;
 	static const FName CameraRigAssetEditorToolBarName;
 
+	GAMEPLAYCAMERASEDITOR_API static IGameplayCamerasEditorModule& Get();
+
 	virtual ~IGameplayCamerasEditorModule() = default;
 
 public:
@@ -65,6 +69,15 @@ public:
 	/** Creates a new camera rig asset picker widget */
 	virtual TSharedRef<SWidget> CreateCameraRigPicker(const FCameraRigPickerConfig& InPickerConfig) = 0;
 
+public:
+
+	/** Registers a new camera director editor creator. */
+	virtual FDelegateHandle RegisterCameraDirectorEditor(FOnCreateCameraDirectorAssetEditorMode InOnCreateEditor) = 0;
+	/** Gets the registered camera director editor creators. */
+	virtual TArrayView<const FOnCreateCameraDirectorAssetEditorMode> GetCameraDirectorEditorCreators() const = 0;
+	/** Unregisters a camera director editor creator. */
+	virtual void UnregisterCameraDirectorEditor(FDelegateHandle InHandle) = 0;
+
 public:
 
 	using FCameraDebugCategoryInfo = UE::Cameras::FCameraDebugCategoryInfo;
diff --git a/Engine/Plugins/Cameras/GameplayCameras/Source/GameplayCamerasEditor/Public/Toolkits/BlueprintCameraDirectorAssetEditorMode.h b/Engine/Plugins/Cameras/GameplayCameras/Source/GameplayCamerasEditor/Public/Toolkits/BlueprintCameraDirectorAssetEditorMode.h
new file mode 100644
--- /dev/null
+++ b/Engine/Plugins/Cameras/GameplayCameras/Source/GameplayCamerasEditor/Public/Toolkits/BlueprintCameraDirectorAssetEditorMode.h
@@ -0,0 +1,31 @@
+// Copyright Epic Games, Inc. All Rights Reserved.
+
+#pragma once
+
+#include "Core/CameraAsset.h"
+#include "Toolkits/CameraDirectorAssetEditorMode.h"
+
+namespace UE::Cameras
+{
+
+class FBlueprintCameraDirectorAssetEditorMode
+	: public FCameraDirectorAssetEditorMode
+	, public ICameraAssetEventHandler
+{
+public:
+
+	static TSharedPtr<FCameraDirectorAssetEditorMode> CreateInstance(UCameraAsset* InCameraAsset);
+
+	FBlueprintCameraDirectorAssetEditorMode(UCameraAsset* InCameraAsset);
+
+protected:
+
+	virtual void OnCameraRigsChanged(UCameraAsset* InCameraAsset, const TCameraArrayChangedEvent<UCameraRigAsset*>& Event) override;
+
+private:
+
+	TCameraEventHandler<ICameraAssetEventHandler> EventHandler;
+};
+
+}  // namespace UE::Cameras
+
diff --git a/Engine/Plugins/Cameras/GameplayCameras/Source/GameplayCamerasEditor/Public/Toolkits/SingleCameraDirectorAssetEditorMode.h b/Engine/Plugins/Cameras/GameplayCameras/Source/GameplayCamerasEditor/Public/Toolkits/SingleCameraDirectorAssetEditorMode.h
new file mode 100644
--- /dev/null
+++ b/Engine/Plugins/Cameras/GameplayCameras/Source/GameplayCamerasEditor/Public/Toolkits/SingleCameraDirectorAssetEditorMode.h
@@ -0,0 +1,31 @@
+// Copyright Epic Games, Inc. All Rights Reserved.
+
+#pragma once
+
+#include "Core/CameraAsset.h"
+#include "Toolkits/CameraDirectorAssetEditorMode.h"
+
+namespace UE::Cameras
+{
+
+class FSingleCameraDirectorAssetEditorMode
+	: public FCameraDirectorAssetEditorMode
+	, public ICameraAssetEventHandler
+{
+public:
+
+	static TSharedPtr<FCameraDirectorAssetEditorMode> CreateInstance(UCameraAsset* InCameraAsset);
+
+	FSingleCameraDirectorAssetEditorMode(UCameraAsset* InCameraAsset);
+
+protected:
+
+	virtual void OnCameraRigsChanged(UCameraAsset* InCameraAsset, const TCameraArrayChangedEvent<UCameraRigAsset*>& Event) override;
+
+private:
+
+	TCameraEventHandler<ICameraAssetEventHandler> EventHandler;
+};
+
+}  // namespace UE::Cameras
+
