Cameras: add placeholder asset icons

diff --git a/Engine/Plugins/Cameras/GameplayCameras/Source/GameplayCamerasEditor/Private/Styles/GameplayCamerasEditorStyle.cpp b/Engine/Plugins/Cameras/GameplayCameras/Source/GameplayCamerasEditor/Private/Styles/GameplayCamerasEditorStyle.cpp
--- a/Engine/Plugins/Cameras/GameplayCameras/Source/GameplayCamerasEditor/Private/Styles/GameplayCamerasEditorStyle.cpp
+++ b/Engine/Plugins/Cameras/GameplayCameras/Source/GameplayCamerasEditor/Private/Styles/GameplayCamerasEditorStyle.cpp
@@ -35,6 +35,16 @@ FGameplayCamerasEditorStyle::FGameplayCamerasEditorStyle()
 	const FTextBlockStyle& NormalText = CoreStyle.GetWidgetStyle<FTextBlockStyle>("NormalText");
 	const FTableRowStyle AlternatingTableRowStyle = CoreStyle.GetWidgetStyle<FTableRowStyle>("TableView.AlternatingRow");
 
+	// Camera assets.
+	Set("ClassIcon.CameraAsset", new IMAGE_BRUSH_SVG("Icons/ContentBrowser-CameraKit", Icon16x16));
+	Set("ClassThumbnail.CameraAsset", new IMAGE_BRUSH_SVG("Icons/ContentBrowser-CameraKit", Icon64x64));
+	Set("ClassIcon.CameraRigAsset", new IMAGE_BRUSH_SVG("Icons/ContentBrowser-CameraRig", Icon16x16));
+	Set("ClassThumbnail.CameraRigAsset", new IMAGE_BRUSH_SVG("Icons/ContentBrowser-CameraRig", Icon64x64));
+	Set("ClassIcon.CameraRigProxyAsset", new IMAGE_BRUSH_SVG("Icons/ContentBrowser-CameraRigProxy", Icon16x16));
+	Set("ClassThumbnail.CameraRigProxyAsset", new IMAGE_BRUSH_SVG("Icons/ContentBrowser-CameraRigProxy", Icon64x64));
+	Set("ClassIcon.CameraVariableCollection", new IMAGE_BRUSH_SVG("Icons/ContentBrowser-CameraVariableCollection", Icon16x16));
+	Set("ClassThumbnail.CameraVariableCollection", new IMAGE_BRUSH_SVG("Icons/ContentBrowser-CameraVariableCollection", Icon64x64));
+
 	// Camera asset editor icons.
 	Set("CameraAssetEditor.ShowCameraDirector", new IMAGE_BRUSH_SVG("Icons/CameraEditor-CameraDirector", Icon20x20));
 	Set("CameraAssetEditor.ShowCameraRigs", new IMAGE_BRUSH_SVG("Icons/CameraEditor-CameraRigs", Icon20x20));
