Cameras: store object tree graph property pin direction as metadata on properties

diff --git a/Engine/Plugins/Cameras/GameplayCameras/Source/GameplayCameras/Public/Core/CameraAsset.h b/Engine/Plugins/Cameras/GameplayCameras/Source/GameplayCameras/Public/Core/CameraAsset.h
--- a/Engine/Plugins/Cameras/GameplayCameras/Source/GameplayCameras/Public/Core/CameraAsset.h
+++ b/Engine/Plugins/Cameras/GameplayCameras/Source/GameplayCameras/Public/Core/CameraAsset.h
@@ -144,11 +144,11 @@ private:
 	TArray<TObjectPtr<UCameraRigAsset>> CameraRigs;
 
 	/** A list of default enter transitions for all the camera rigs in this asset. */
-	UPROPERTY()
+	UPROPERTY(meta=(ObjectTreeGraphPinDirection=Input))
 	TArray<TObjectPtr<UCameraRigTransition>> EnterTransitions;
 
 	/** A list of default exit transitions for all the camera rigs in this asset. */
-	UPROPERTY()
+	UPROPERTY(meta=(ObjectTreeGraphPinDirection=Output))
 	TArray<TObjectPtr<UCameraRigTransition>> ExitTransitions;
 
 	/** The current build state of this camera asset. */
diff --git a/Engine/Plugins/Cameras/GameplayCameras/Source/GameplayCameras/Public/Core/CameraRigAsset.h b/Engine/Plugins/Cameras/GameplayCameras/Source/GameplayCameras/Public/Core/CameraRigAsset.h
--- a/Engine/Plugins/Cameras/GameplayCameras/Source/GameplayCameras/Public/Core/CameraRigAsset.h
+++ b/Engine/Plugins/Cameras/GameplayCameras/Source/GameplayCameras/Public/Core/CameraRigAsset.h
@@ -161,11 +161,11 @@ public:
 	FCameraRigInterface Interface;
 
 	/** List of enter transitions for this camera rig. */
-	UPROPERTY(Instanced)
+	UPROPERTY(Instanced, meta=(ObjectTreeGraphPinDirection=Input))
 	TArray<TObjectPtr<UCameraRigTransition>> EnterTransitions;
 
 	/** List of exist transitions for this camera rig. */
-	UPROPERTY(Instanced)
+	UPROPERTY(Instanced, meta=(ObjectTreeGraphPinDirection=Output))
 	TArray<TObjectPtr<UCameraRigTransition>> ExitTransitions;
 
 	/** Allocation information for all the nodes and variables in this camera rig. */
diff --git a/Engine/Plugins/Cameras/GameplayCameras/Source/GameplayCameras/Public/Core/CameraRigTransition.h b/Engine/Plugins/Cameras/GameplayCameras/Source/GameplayCameras/Public/Core/CameraRigTransition.h
--- a/Engine/Plugins/Cameras/GameplayCameras/Source/GameplayCameras/Public/Core/CameraRigTransition.h
+++ b/Engine/Plugins/Cameras/GameplayCameras/Source/GameplayCameras/Public/Core/CameraRigTransition.h
@@ -99,7 +99,7 @@ class UCameraRigTransition
 public:
 
 	/** The list of conditions that must pass for this transition to be used. */
-	UPROPERTY(Instanced)
+	UPROPERTY(Instanced, meta=(ObjectTreeGraphPinDirection=Input))
 	TArray<TObjectPtr<UCameraRigTransitionCondition>> Conditions;
 
 	/** The blend to use to blend a given camera rig in or out. */
diff --git a/Engine/Plugins/Cameras/GameplayCameras/Source/GameplayCamerasEditor/Private/Editors/CameraRigTransitionGraphSchema.cpp b/Engine/Plugins/Cameras/GameplayCameras/Source/GameplayCamerasEditor/Private/Editors/CameraRigTransitionGraphSchema.cpp
--- a/Engine/Plugins/Cameras/GameplayCameras/Source/GameplayCamerasEditor/Private/Editors/CameraRigTransitionGraphSchema.cpp
+++ b/Engine/Plugins/Cameras/GameplayCameras/Source/GameplayCamerasEditor/Private/Editors/CameraRigTransitionGraphSchema.cpp
@@ -22,8 +22,6 @@ void UCameraRigTransitionGraphSchema::OnBuildGraphConfig(FObjectTreeGraphConfig&
 	InOutGraphConfig.ObjectClassConfigs.Emplace(UCameraRigAsset::StaticClass())
 		.HasSelfPin(false)
 		.OnlyAsRoot()
-		.SetPropertyPinDirection(GET_MEMBER_NAME_CHECKED(UCameraRigAsset, EnterTransitions), EGPD_Input)
-		.SetPropertyPinDirection(GET_MEMBER_NAME_CHECKED(UCameraRigAsset, ExitTransitions), EGPD_Output)
 		.NodeTitleUsesObjectName(true)
 		.NodeTitleColor(Settings->CameraRigAssetTitleColor);
 	InOutGraphConfig.OnFormatObjectDisplayName = FOnFormatObjectDisplayName::CreateLambda(
diff --git a/Engine/Plugins/Cameras/GameplayCameras/Source/GameplayCamerasEditor/Private/Editors/CameraRigTransitionGraphSchemaBase.cpp b/Engine/Plugins/Cameras/GameplayCameras/Source/GameplayCamerasEditor/Private/Editors/CameraRigTransitionGraphSchemaBase.cpp
--- a/Engine/Plugins/Cameras/GameplayCameras/Source/GameplayCamerasEditor/Private/Editors/CameraRigTransitionGraphSchemaBase.cpp
+++ b/Engine/Plugins/Cameras/GameplayCameras/Source/GameplayCamerasEditor/Private/Editors/CameraRigTransitionGraphSchemaBase.cpp
@@ -24,7 +24,6 @@ FObjectTreeGraphConfig UCameraRigTransitionGraphSchemaBase::BuildGraphConfig() c
 	GraphConfig.ConnectableObjectClasses.Add(UCameraRigTransitionCondition::StaticClass());
 	GraphConfig.ConnectableObjectClasses.Add(UBlendCameraNode::StaticClass());
 	GraphConfig.ObjectClassConfigs.Emplace(UCameraRigTransition::StaticClass())
-		.SetPropertyPinDirection(GET_MEMBER_NAME_STRING_CHECKED(UCameraRigTransition, Conditions), EGPD_Input)
 		.NodeTitleColor(Settings->CameraRigTransitionTitleColor);
 	GraphConfig.ObjectClassConfigs.Emplace(UCameraRigTransitionCondition::StaticClass())
 		.SelfPinDirection(EGPD_Output)
diff --git a/Engine/Plugins/Cameras/GameplayCameras/Source/GameplayCamerasEditor/Private/Editors/CameraSharedTransitionGraphSchema.cpp b/Engine/Plugins/Cameras/GameplayCameras/Source/GameplayCamerasEditor/Private/Editors/CameraSharedTransitionGraphSchema.cpp
--- a/Engine/Plugins/Cameras/GameplayCameras/Source/GameplayCamerasEditor/Private/Editors/CameraSharedTransitionGraphSchema.cpp
+++ b/Engine/Plugins/Cameras/GameplayCameras/Source/GameplayCamerasEditor/Private/Editors/CameraSharedTransitionGraphSchema.cpp
@@ -24,8 +24,6 @@ void UCameraSharedTransitionGraphSchema::OnBuildGraphConfig(FObjectTreeGraphConf
 	InOutGraphConfig.ObjectClassConfigs.Emplace(UCameraAsset::StaticClass())
 		.HasSelfPin(false)
 		.OnlyAsRoot()
-		.SetPropertyPinDirection(GET_MEMBER_NAME_CHECKED(UCameraAsset, EnterTransitions), EGPD_Input)
-		.SetPropertyPinDirection(GET_MEMBER_NAME_CHECKED(UCameraAsset, ExitTransitions), EGPD_Output)
 		.NodeTitleUsesObjectName(true)
 		.NodeTitleColor(Settings->CameraAssetTitleColor);
 }
diff --git a/Engine/Plugins/Cameras/GameplayCameras/Source/GameplayCamerasEditor/Private/Editors/ObjectTreeGraphConfig.cpp b/Engine/Plugins/Cameras/GameplayCameras/Source/GameplayCamerasEditor/Private/Editors/ObjectTreeGraphConfig.cpp
--- a/Engine/Plugins/Cameras/GameplayCameras/Source/GameplayCamerasEditor/Private/Editors/ObjectTreeGraphConfig.cpp
+++ b/Engine/Plugins/Cameras/GameplayCameras/Source/GameplayCamerasEditor/Private/Editors/ObjectTreeGraphConfig.cpp
@@ -228,5 +228,28 @@ void FObjectTreeGraphConfig::FormatDisplayNameText(const UObject* InObject, cons
 	OnFormatObjectDisplayName.ExecuteIfBound(InObject, InOutDisplayNameText);
 }
 
+EEdGraphPinDirection FObjectTreeGraphConfig::GetPropertyPinDirection(const UClass* InObjectClass, const FName& InPropertyName) const
+{
+	const FObjectTreeGraphClassConfig& ClassConfig = GetObjectClassConfig(InObjectClass);
+	TOptional<EEdGraphPinDirection> PinDirectionOverride = ClassConfig.GetPropertyPinDirectionOverride(InPropertyName);
+	if (PinDirectionOverride.IsSet())
+	{
+		return PinDirectionOverride.GetValue();
+	}
+
+	const FProperty* Property = InObjectClass->FindPropertyByName(InPropertyName);
+	const FString& CustomDirection = Property->GetMetaData(TEXT("ObjectTreeGraphPinDirection"));
+	if (CustomDirection == TEXT("Input"))
+	{
+		return EGPD_Input;
+	}
+	else if (CustomDirection == TEXT("Output"))
+	{
+		return EGPD_Output;
+	}
+
+	return ClassConfig.DefaultPropertyPinDirection();
+}
+
 #undef LOCTEXT_NAMESPACE
 
diff --git a/Engine/Plugins/Cameras/GameplayCameras/Source/GameplayCamerasEditor/Private/Editors/ObjectTreeGraphNode.cpp b/Engine/Plugins/Cameras/GameplayCameras/Source/GameplayCamerasEditor/Private/Editors/ObjectTreeGraphNode.cpp
--- a/Engine/Plugins/Cameras/GameplayCameras/Source/GameplayCamerasEditor/Private/Editors/ObjectTreeGraphNode.cpp
+++ b/Engine/Plugins/Cameras/GameplayCameras/Source/GameplayCamerasEditor/Private/Editors/ObjectTreeGraphNode.cpp
@@ -85,7 +85,7 @@ void UObjectTreeGraphNode::AllocateDefaultPins()
 	{
 		const FName PropertyName = PropertyIt->GetFName();
 
-		const EEdGraphPinDirection PinDirection = ObjectClassConfig.GetPropertyPinDirection(PropertyName);
+		const EEdGraphPinDirection PinDirection = OuterGraphConfig.GetPropertyPinDirection(NodeContext.ObjectClass, PropertyName);
 
 		FEdGraphPinType ChildPinType;
 		ChildPinType.PinCategory = UObjectTreeGraphSchema::PC_Property;
@@ -167,7 +167,7 @@ void UObjectTreeGraphNode::CreateNewItemPin(FArrayProperty& InArrayProperty)
 
 void UObjectTreeGraphNode::CreateNewItemPin(UEdGraphPin* InParentArrayPin)
 {
-	const FObjectTreeGraphClassConfig& ObjectClassConfig = GetObjectClassConfig();
+	const FNodeContext NodeContext = GetNodeContext();
 
 	const FName PropertyName = InParentArrayPin->GetFName();
 	const int32 NewIndex = InParentArrayPin->SubPins.Num();
@@ -176,7 +176,7 @@ void UObjectTreeGraphNode::CreateNewItemPin(UEdGraphPin* InParentArrayPin)
 	ChildPinType.PinCategory = UObjectTreeGraphSchema::PC_Property;
 	ChildPinType.PinSubCategory = UObjectTreeGraphSchema::PSC_ArrayPropertyItem;
 
-	const EEdGraphPinDirection PinDirection = ObjectClassConfig.GetPropertyPinDirection(PropertyName);
+	const EEdGraphPinDirection PinDirection = NodeContext.GraphConfig.GetPropertyPinDirection(NodeContext.ObjectClass, PropertyName);
 
 	UEdGraphPin* ChildPin = CreatePin(PinDirection, ChildPinType, PropertyName);
 	ChildPin->PinFriendlyName = FText::Format(LOCTEXT("ArrayPinFriendlyNameFmt", "{0} {1}"), FText::FromName(PropertyName), NewIndex);
diff --git a/Engine/Plugins/Cameras/GameplayCameras/Source/GameplayCamerasEditor/Public/Editors/ObjectTreeGraphConfig.h b/Engine/Plugins/Cameras/GameplayCameras/Source/GameplayCamerasEditor/Public/Editors/ObjectTreeGraphConfig.h
--- a/Engine/Plugins/Cameras/GameplayCameras/Source/GameplayCamerasEditor/Public/Editors/ObjectTreeGraphConfig.h
+++ b/Engine/Plugins/Cameras/GameplayCameras/Source/GameplayCamerasEditor/Public/Editors/ObjectTreeGraphConfig.h
@@ -102,29 +102,29 @@ public:
 	}
 
 	/** Gets the custom property pin directions for given named properties. */
-	const TMap<FName, EEdGraphPinDirection>& PropertyPinDirections() const { return _PropertyPinDirections; }
+	const TMap<FName, EEdGraphPinDirection>& PropertyPinDirectionOverrides() const { return _PropertyPinDirectionOverrides; }
 
 	/** Adds a new custom property pin direction for a given named property. */
-	FObjectTreeGraphClassConfig& SetPropertyPinDirection(const FName& InPropertyName, EEdGraphPinDirection InDirection)
+	FObjectTreeGraphClassConfig& SetPropertyPinDirectionOverride(const FName& InPropertyName, EEdGraphPinDirection InDirection)
 	{
-		_PropertyPinDirections.Add(InPropertyName, InDirection);
+		_PropertyPinDirectionOverrides.Add(InPropertyName, InDirection);
 		return *this;
 	}
 
 	/** Gets the custom property pin direction for a given named property. */
-	EEdGraphPinDirection GetPropertyPinDirection(const FName& InPropertyName) const
+	TOptional<EEdGraphPinDirection> GetPropertyPinDirectionOverride(const FName& InPropertyName) const
 	{
-		if (const EEdGraphPinDirection* PinDirection = _PropertyPinDirections.Find(InPropertyName))
+		if (const EEdGraphPinDirection* PinDirection = _PropertyPinDirectionOverrides.Find(InPropertyName))
 		{
 			return *PinDirection;
 		}
-		return _DefaultPropertyPinDirection;
+		return TOptional<EEdGraphPinDirection>();
 	}
 
 private:
 
 	TArray<FString> _StripDisplayNameSuffixes;
-	TMap<FName, EEdGraphPinDirection> _PropertyPinDirections;
+	TMap<FName, EEdGraphPinDirection> _PropertyPinDirectionOverrides;
 };
 
 /**
@@ -231,6 +231,9 @@ public:
 	/** Computes the display name of the given object class. */
 	FText GetDisplayNameText(const UClass* InClass) const;
 
+	/** Gets the custom property pin direction for a given named property. */
+	EEdGraphPinDirection GetPropertyPinDirection(const UClass* InObjectClass, const FName& InPropertyName) const;
+
 private:
 
 	FText GetDisplayNameText(const UClass* InClass, const FObjectTreeGraphClassConfig& InClassConfig) const;
