Cameras: add a camera rig build context object to let camera nodes to more general validation and log messages

diff --git a/Engine/Plugins/Cameras/GameplayCameras/Source/GameplayCameras/Private/Core/CameraBuildLog.cpp b/Engine/Plugins/Cameras/GameplayCameras/Source/GameplayCameras/Private/Core/CameraBuildLog.cpp
--- a/Engine/Plugins/Cameras/GameplayCameras/Source/GameplayCameras/Private/Core/CameraBuildLog.cpp
+++ b/Engine/Plugins/Cameras/GameplayCameras/Source/GameplayCameras/Private/Core/CameraBuildLog.cpp
@@ -67,7 +67,7 @@ void FCameraBuildLog::AddMessage(EMessageSeverity::Type InSeverity, FText&& InTe
 	AddMessage(InSeverity, nullptr, MoveTemp(InText));
 }
 
-void FCameraBuildLog::AddMessage(EMessageSeverity::Type InSeverity, UObject* InObject, FText&& InText)
+void FCameraBuildLog::AddMessage(EMessageSeverity::Type InSeverity, const UObject* InObject, FText&& InText)
 {
 	Messages.Add(FCameraBuildLogMessage{ InSeverity, InObject, MoveTemp(InText) });
 
diff --git a/Engine/Plugins/Cameras/GameplayCameras/Source/GameplayCameras/Private/Core/CameraNode.cpp b/Engine/Plugins/Cameras/GameplayCameras/Source/GameplayCameras/Private/Core/CameraNode.cpp
--- a/Engine/Plugins/Cameras/GameplayCameras/Source/GameplayCameras/Private/Core/CameraNode.cpp
+++ b/Engine/Plugins/Cameras/GameplayCameras/Source/GameplayCameras/Private/Core/CameraNode.cpp
@@ -28,9 +28,9 @@ FCameraNodeChildrenView UCameraNode::GetChildren()
 	return OnGetChildren();
 }
 
-void UCameraNode::BuildAllocationInfo(FCameraRigAllocationInfo& AllocationInfo) const
+void UCameraNode::Build(FCameraRigBuildContext& BuildContext)
 {
-	OnBuildAllocationInfo(AllocationInfo);
+	OnBuild(BuildContext);
 }
 
 FCameraNodeEvaluatorPtr UCameraNode::BuildEvaluator(FCameraNodeEvaluatorBuilder& Builder) const
diff --git a/Engine/Plugins/Cameras/GameplayCameras/Source/GameplayCameras/Private/Core/CameraRigAssetBuilder.cpp b/Engine/Plugins/Cameras/GameplayCameras/Source/GameplayCameras/Private/Core/CameraRigAssetBuilder.cpp
--- a/Engine/Plugins/Cameras/GameplayCameras/Source/GameplayCameras/Private/Core/CameraRigAssetBuilder.cpp
+++ b/Engine/Plugins/Cameras/GameplayCameras/Source/GameplayCameras/Private/Core/CameraRigAssetBuilder.cpp
@@ -5,6 +5,7 @@
 #include "Core/CameraNode.h"
 #include "Core/CameraNodeEvaluatorStorage.h"
 #include "Core/CameraParameters.h"
+#include "Core/CameraRigBuildContext.h"
 #include "Core/CameraRigAsset.h"
 #include "Core/CameraVariableAssets.h"
 
@@ -487,7 +488,7 @@ void FCameraRigAssetBuilder::BuildAllocationInfo()
 	CameraRig->AllocationInfo = AllocationInfo;
 }
 
-void FCameraRigAssetBuilder::BuildAllocationInfo(const UCameraNode* CameraNode)
+void FCameraRigAssetBuilder::BuildAllocationInfo(UCameraNode* CameraNode)
 {
 	// Look for properties that are camera parameters, and gather what camera variables they reference. 
 	// This is for both exposed rig parameters (which we just built in BuildNewDrivenParameters) and 
@@ -520,7 +521,8 @@ UE_CAMERA_VARIABLE_FOR_ALL_TYPES()
 	}
 
 	// Let the camera node add any custom variables or extra memory.
-	CameraNode->BuildAllocationInfo(AllocationInfo);
+	FCameraRigBuildContext BuildContext(AllocationInfo, BuildLog);
+	CameraNode->Build(BuildContext);
 }
 
 void FCameraRigAssetBuilder::UpdateBuildStatus()
diff --git a/Engine/Plugins/Cameras/GameplayCameras/Source/GameplayCameras/Public/Core/CameraBuildLog.h b/Engine/Plugins/Cameras/GameplayCameras/Source/GameplayCameras/Public/Core/CameraBuildLog.h
--- a/Engine/Plugins/Cameras/GameplayCameras/Source/GameplayCameras/Public/Core/CameraBuildLog.h
+++ b/Engine/Plugins/Cameras/GameplayCameras/Source/GameplayCameras/Public/Core/CameraBuildLog.h
@@ -17,7 +17,7 @@ struct GAMEPLAYCAMERAS_API FCameraBuildLogMessage
 	/** Severity of the message. */
 	EMessageSeverity::Type Severity = EMessageSeverity::Info;
 	/** An optional object that the message relates to. */
-	UObject* Object = nullptr;
+	const UObject* Object = nullptr;
 	/** The actual message. */
 	FText Text;
 
@@ -49,7 +49,7 @@ public:
 	/** Adds a new message. */
 	void AddMessage(EMessageSeverity::Type InSeverity, FText&& InText);
 	/** Adds a new message. */
-	void AddMessage(EMessageSeverity::Type InSeverity, UObject* InObject, FText&& InText);
+	void AddMessage(EMessageSeverity::Type InSeverity, const UObject* InObject, FText&& InText);
 
 	/** Gets the list of received messages so far. */
 	TArrayView<const FCameraBuildLogMessage> GetMessages() const { return Messages; }
diff --git a/Engine/Plugins/Cameras/GameplayCameras/Source/GameplayCameras/Public/Core/CameraNode.h b/Engine/Plugins/Cameras/GameplayCameras/Source/GameplayCameras/Public/Core/CameraNode.h
--- a/Engine/Plugins/Cameras/GameplayCameras/Source/GameplayCameras/Public/Core/CameraNode.h
+++ b/Engine/Plugins/Cameras/GameplayCameras/Source/GameplayCameras/Public/Core/CameraNode.h
@@ -3,7 +3,6 @@
 #pragma once
 
 #include "Core/CameraNodeEvaluatorBuilder.h"
-#include "Core/CameraVariableTableFwd.h"
 #include "Core/ObjectChildrenView.h"
 #include "Core/ObjectTreeGraphObject.h"
 #include "CoreTypes.h"
@@ -11,7 +10,10 @@
 
 #include "CameraNode.generated.h"
 
-struct FCameraRigAllocationInfo;
+namespace UE::Cameras
+{
+	struct FCameraRigBuildContext;
+}
 
 /** View on a camera node's children. */
 using FCameraNodeChildrenView = UE::Cameras::TObjectChildrenView<TObjectPtr<UCameraNode>>;
@@ -28,13 +30,14 @@ class GAMEPLAYCAMERAS_API UCameraNode
 
 public:
 	
+	using FCameraRigBuildContext = UE::Cameras::FCameraRigBuildContext;
 	using FCameraNodeEvaluatorBuilder = UE::Cameras::FCameraNodeEvaluatorBuilder;
 
 	/** Get the list of children under this node. */
 	FCameraNodeChildrenView GetChildren();
 
 	/** Gets optional info about this node's required allocations at runtime. */
-	void BuildAllocationInfo(FCameraRigAllocationInfo& AllocationInfo) const;
+	void Build(FCameraRigBuildContext& BuildContext);
 
 	/** Builds the evaluator for this node. */
 	FCameraNodeEvaluatorPtr BuildEvaluator(FCameraNodeEvaluatorBuilder& Builder) const;
@@ -45,7 +48,7 @@ protected:
 	virtual FCameraNodeChildrenView OnGetChildren() { return FCameraNodeChildrenView(); }
 
 	/** Gets optional info about this node's required allocations at runtime. */
-	virtual void OnBuildAllocationInfo(FCameraRigAllocationInfo& AllocationInfo) const {}
+	virtual void OnBuild(FCameraRigBuildContext& BuildContext) {}
 
 	/** Builds the evaluator for this node. */
 	virtual FCameraNodeEvaluatorPtr OnBuildEvaluator(FCameraNodeEvaluatorBuilder& Builder) const { return nullptr; }
diff --git a/Engine/Plugins/Cameras/GameplayCameras/Source/GameplayCameras/Public/Core/CameraRigAssetBuilder.h b/Engine/Plugins/Cameras/GameplayCameras/Source/GameplayCameras/Public/Core/CameraRigAssetBuilder.h
--- a/Engine/Plugins/Cameras/GameplayCameras/Source/GameplayCameras/Public/Core/CameraRigAssetBuilder.h
+++ b/Engine/Plugins/Cameras/GameplayCameras/Source/GameplayCameras/Public/Core/CameraRigAssetBuilder.h
@@ -50,7 +50,7 @@ private:
 	void DiscardUnusedPrivateVariables();
 
 	void BuildAllocationInfo();
-	void BuildAllocationInfo(const UCameraNode* CameraNode);
+	void BuildAllocationInfo(UCameraNode* CameraNode);
 
 	void UpdateBuildStatus();
 
diff --git a/Engine/Plugins/Cameras/GameplayCameras/Source/GameplayCameras/Public/Core/CameraRigBuildContext.h b/Engine/Plugins/Cameras/GameplayCameras/Source/GameplayCameras/Public/Core/CameraRigBuildContext.h
new file mode 100644
--- /dev/null
+++ b/Engine/Plugins/Cameras/GameplayCameras/Source/GameplayCameras/Public/Core/CameraRigBuildContext.h
@@ -0,0 +1,33 @@
+// Copyright Epic Games, Inc. All Rights Reserved.
+
+#pragma once
+
+#include "UObject/ObjectPtr.h"
+
+class UCameraRigAsset;
+struct FCameraRigAllocationInfo;
+
+namespace UE::Cameras
+{
+
+class FCameraBuildLog;
+
+/**
+ * Camera rig asset build context.
+ */
+struct FCameraRigBuildContext
+{
+	FCameraRigBuildContext(
+			FCameraRigAllocationInfo& InAllocationInfo, FCameraBuildLog& InBuildLog)
+		: AllocationInfo(InAllocationInfo)
+		, BuildLog(InBuildLog)
+	{}
+
+	/** The allocation information to be determined. */
+	FCameraRigAllocationInfo& AllocationInfo;
+	/** The build log for emitting messages. */
+	FCameraBuildLog& BuildLog;
+};
+
+}  // namespace UE::Cameras
+
