Cameras: fix possible crash while hot-reloading camera rigs

diff --git a/Engine/Plugins/Cameras/GameplayCameras/Source/GameplayCameras/Private/Core/CameraRigAsset.cpp b/Engine/Plugins/Cameras/GameplayCameras/Source/GameplayCameras/Private/Core/CameraRigAsset.cpp
--- a/Engine/Plugins/Cameras/GameplayCameras/Source/GameplayCameras/Private/Core/CameraRigAsset.cpp
+++ b/Engine/Plugins/Cameras/GameplayCameras/Source/GameplayCameras/Private/Core/CameraRigAsset.cpp
@@ -131,7 +131,10 @@ void UCameraRigAsset::GatherPackages(FCameraRigPackages& OutPackages) const
 		FCameraNodeChildrenView CurrentChildren = CurrentNode->GetChildren();
 		for (UCameraNode* CurrentChild : ReverseIterate(CurrentChildren))
 		{
-			NodeStack.Add(CurrentChild);
+			if (CurrentChild)
+			{
+				NodeStack.Add(CurrentChild);
+			}
 		}
 	}
 }
