Cameras: give Guids to camera rigs and variables, use the Guid for variable IDs

diff --git a/Engine/Plugins/Cameras/GameplayCameras/Source/GameplayCameras/Private/Core/CameraRigAsset.cpp b/Engine/Plugins/Cameras/GameplayCameras/Source/GameplayCameras/Private/Core/CameraRigAsset.cpp
--- a/Engine/Plugins/Cameras/GameplayCameras/Source/GameplayCameras/Private/Core/CameraRigAsset.cpp
+++ b/Engine/Plugins/Cameras/GameplayCameras/Source/GameplayCameras/Private/Core/CameraRigAsset.cpp
@@ -59,9 +59,35 @@ void UCameraRigAsset::PostLoad()
 
 #endif
 
+	if (!Guid.IsValid())
+	{
+		Guid = FGuid::NewGuid();
+	}
+
 	Super::PostLoad();
 }
 
+void UCameraRigAsset::PostInitProperties()
+{
+	Super::PostInitProperties();
+
+	if (!HasAnyFlags(RF_ClassDefaultObject | RF_ArchetypeObject | RF_NeedLoad | RF_WasLoaded) && 
+			!Guid.IsValid())
+	{
+		Guid = FGuid::NewGuid();
+	}
+}
+
+void UCameraRigAsset::PostDuplicate(EDuplicateMode::Type DuplicateMode)
+{
+	Super::PostDuplicate(DuplicateMode);
+
+	if (DuplicateMode == EDuplicateMode::Normal)
+	{
+		Guid = FGuid::NewGuid();
+	}
+}
+
 void UCameraRigAsset::GetOwnedGameplayTags(FGameplayTagContainer& TagContainer) const
 {
 	TagContainer.AppendTags(GameplayTags);
diff --git a/Engine/Plugins/Cameras/GameplayCameras/Source/GameplayCameras/Private/Core/CameraVariableAssets.cpp b/Engine/Plugins/Cameras/GameplayCameras/Source/GameplayCameras/Private/Core/CameraVariableAssets.cpp
--- a/Engine/Plugins/Cameras/GameplayCameras/Source/GameplayCameras/Private/Core/CameraVariableAssets.cpp
+++ b/Engine/Plugins/Cameras/GameplayCameras/Source/GameplayCameras/Private/Core/CameraVariableAssets.cpp
@@ -9,9 +9,10 @@ UCameraVariableAsset::UCameraVariableAsset(const FObjectInitializer& ObjectInit)
 {
 }
 
-void UCameraVariableAsset::RegenerateVariableID()
+FCameraVariableID  UCameraVariableAsset::GetVariableID() const
 {
-	VariableID = FCameraVariableID::FromHashValue(GetTypeHash(GetFullName()));
+	ensure(Guid.IsValid());
+	return FCameraVariableID::FromHashValue(GetTypeHash(Guid));
 }
 
 FCameraVariableDefinition UCameraVariableAsset::GetVariableDefinition() const
@@ -25,33 +26,55 @@ FCameraVariableDefinition UCameraVariableAsset::GetVariableDefinition() const
 	return VariableDefinition;
 }
 
-void UCameraVariableAsset::Serialize(FArchive& Ar)
+void UCameraVariableAsset::PostLoad()
 {
-	Super::Serialize(Ar);
-
-#if WITH_EDITORONLY_DATA
-	if ((Ar.IsLoading() && !VariableID.IsValid()) || Ar.IsSaving())
+	if (!Guid.IsValid())
 	{
-		RegenerateVariableID();
+		Guid = FGuid::NewGuid();
 	}
-#endif
+
+	Super::PostLoad();
 }
 
 void UCameraVariableAsset::PostInitProperties()
 {
-	RegenerateVariableID();
 	Super::PostInitProperties();
-}
 
-void UCameraVariableAsset::PostRename(UObject* OldOuter, const FName OldName)
-{
-	RegenerateVariableID();
-	Super::PostRename(OldOuter, OldName);
+	if (!HasAnyFlags(RF_ClassDefaultObject | RF_ArchetypeObject | RF_NeedLoad | RF_WasLoaded) && 
+			!Guid.IsValid())
+	{
+		Guid = FGuid::NewGuid();
+	}
 }
 
 void UCameraVariableAsset::PostDuplicate(EDuplicateMode::Type DuplicateMode)
 {
-	RegenerateVariableID();
 	Super::PostDuplicate(DuplicateMode);
+
+	if (DuplicateMode == EDuplicateMode::Normal)
+	{
+		Guid = FGuid::NewGuid();
+	}
+}
+
+#if WITH_EDITOR
+
+FString UCameraVariableAsset::GetDisplayName() const
+{
+	if (!DisplayName.IsEmpty())
+	{
+		return DisplayName;
+	}
+	return GetName();
+}
+
+FText UCameraVariableAsset::GetDisplayText() const
+{
+	if (!DisplayName.IsEmpty())
+	{
+		return FText::FromString(DisplayName);
+	}
+	return FText::FromName(GetFName());
 }
 
+#endif  // WITH_EDITOR
diff --git a/Engine/Plugins/Cameras/GameplayCameras/Source/GameplayCameras/Public/Core/CameraRigAsset.h b/Engine/Plugins/Cameras/GameplayCameras/Source/GameplayCameras/Public/Core/CameraRigAsset.h
--- a/Engine/Plugins/Cameras/GameplayCameras/Source/GameplayCameras/Public/Core/CameraRigAsset.h
+++ b/Engine/Plugins/Cameras/GameplayCameras/Source/GameplayCameras/Public/Core/CameraRigAsset.h
@@ -172,6 +172,9 @@ public:
 	UPROPERTY()
 	FCameraRigAllocationInfo AllocationInfo;
 
+	/** Gets the camera rig's unique ID. */
+	const FGuid& GetGuid() const { return Guid; }
+
 	/**
 	 * Gets the display name of this camera rig.
 	 * This is either the display name set on the interface object, or its internal name.
@@ -213,9 +216,6 @@ public:
 
 protected:
 
-	// UObject interface.
-	virtual void PostLoad() override;
-
 	// IObjectTreeGraphObject interface.
 #if WITH_EDITOR
 	virtual void GetGraphNodePosition(FName InGraphName, int32& NodePosX, int32& NodePosY) const override;
@@ -233,10 +233,16 @@ protected:
 #endif
 
 	// UObject interface
+	virtual void PostLoad() override;
+	virtual void PostInitProperties() override;
+	virtual void PostDuplicate(EDuplicateMode::Type DuplicateMode) override;
 	virtual void PreSave(FObjectPreSaveContext ObjectSaveContext) override;
 
 private:
 
+	UPROPERTY()
+	FGuid Guid;
+
 #if WITH_EDITORONLY_DATA
 
 	/** Position of the camera rig node in the node graph editor. */
diff --git a/Engine/Plugins/Cameras/GameplayCameras/Source/GameplayCameras/Public/Core/CameraVariableAssets.h b/Engine/Plugins/Cameras/GameplayCameras/Source/GameplayCameras/Public/Core/CameraVariableAssets.h
--- a/Engine/Plugins/Cameras/GameplayCameras/Source/GameplayCameras/Public/Core/CameraVariableAssets.h
+++ b/Engine/Plugins/Cameras/GameplayCameras/Source/GameplayCameras/Public/Core/CameraVariableAssets.h
@@ -21,7 +21,7 @@ public:
 
 	UCameraVariableAsset(const FObjectInitializer& ObjectInit);
 
-	FCameraVariableID GetVariableID() const { return VariableID; }
+	FCameraVariableID GetVariableID() const;
 
 	FCameraVariableDefinition GetVariableDefinition() const;
 
@@ -29,21 +29,21 @@ public:
 	virtual const uint8* GetDefaultValuePtr() const PURE_VIRTUAL(UCameraVariableAsset::GetDefaultValuePtr, return nullptr;);
 
 #if WITH_EDITOR
+	const FGuid& GetGuid() const { return Guid; }
+
+	FString GetDisplayName() const;
+	FText GetDisplayText() const;
+
 	virtual FString FormatDefaultValue() const PURE_VIRTUAL(UCameraVariableAsset::FormatDefaultValue, return FString(););
 #endif  // WITH_EDITOR
 
 public:
 
 	// UObject interface
-	virtual void Serialize(FArchive& Ar) override;
+	virtual void PostLoad() override;
 	virtual void PostInitProperties() override;
-	virtual void PostRename(UObject* OldOuter, const FName OldName) override;
 	virtual void PostDuplicate(EDuplicateMode::Type DuplicateMode) override;
 
-private:
-
-	void RegenerateVariableID();
-
 public:
 
 #if WITH_EDITORONLY_DATA
@@ -63,7 +63,7 @@ public:
 private:
 
 	UPROPERTY()
-	FCameraVariableID VariableID;
+	FGuid Guid;
 };
 
 /** Boolean camera variable. */
