Cameras: large change to add support for pre-blended parameters, including player input

- Node evaluators can now opt-in (via flags) to a new evaluation phase. This phase lets them place parameter values in the variable table, and they will be blended together so that OnRun uses the blended values.
- The variable table now has the concept of "input" values so that they can be isolated and blended during this new phase without affecting the other ones.
- Blend nodes must implement an additional method to blend those parameters separately.
- Creation of node evaluators is now separate from initializing them. Initialization passes a few new pieces of information such as the last active camera rig in a blend stack.
- Add new camera nodes for input slots and simple player input processing. These nodes can use custom variables, but also "built-in" variables, described in a new singleton class of "well-known" values.

diff --git a/Engine/Plugins/Cameras/GameplayCameras/GameplayCameras.uplugin b/Engine/Plugins/Cameras/GameplayCameras/GameplayCameras.uplugin
--- a/Engine/Plugins/Cameras/GameplayCameras/GameplayCameras.uplugin
+++ b/Engine/Plugins/Cameras/GameplayCameras/GameplayCameras.uplugin
@@ -16,6 +16,10 @@
 	"Installed" : false,
 	"Plugins":
 	[
+		{
+			"Name": "EnhancedInput",
+			"Enabled": true
+		},
 		{
 			"Name": "TemplateSequence",
 			"Enabled": true
diff --git a/Engine/Plugins/Cameras/GameplayCameras/Source/GameplayCameras/GameplayCameras.Build.cs b/Engine/Plugins/Cameras/GameplayCameras/Source/GameplayCameras/GameplayCameras.Build.cs
--- a/Engine/Plugins/Cameras/GameplayCameras/Source/GameplayCameras/GameplayCameras.Build.cs
+++ b/Engine/Plugins/Cameras/GameplayCameras/Source/GameplayCameras/GameplayCameras.Build.cs
@@ -14,6 +14,7 @@ public class GameplayCameras : ModuleRules
 				"Core",
 				"CoreUObject",
 				"Engine",
+				"EnhancedInput",
 				"GameplayTags",
 				"HeadMountedDisplay",
 				"MovieScene",
diff --git a/Engine/Plugins/Cameras/GameplayCameras/Source/GameplayCameras/Private/Core/BlendCameraNode.cpp b/Engine/Plugins/Cameras/GameplayCameras/Source/GameplayCameras/Private/Core/BlendCameraNode.cpp
--- a/Engine/Plugins/Cameras/GameplayCameras/Source/GameplayCameras/Private/Core/BlendCameraNode.cpp
+++ b/Engine/Plugins/Cameras/GameplayCameras/Source/GameplayCameras/Private/Core/BlendCameraNode.cpp
@@ -9,6 +9,11 @@ namespace UE::Cameras
 
 UE_DEFINE_CAMERA_NODE_EVALUATOR(FBlendCameraNodeEvaluator)
 
+void FBlendCameraNodeEvaluator::BlendParameters(const FCameraNodePreBlendParams& Params, FCameraNodePreBlendResult& OutResult)
+{
+	OnBlendParameters(Params, OutResult);
+}
+
 void FBlendCameraNodeEvaluator::BlendResults(const FCameraNodeBlendParams& Params, FCameraNodeBlendResult& OutResult)
 {
 	OnBlendResults(Params, OutResult);
diff --git a/Engine/Plugins/Cameras/GameplayCameras/Source/GameplayCameras/Private/Core/BlendStackCameraNode.cpp b/Engine/Plugins/Cameras/GameplayCameras/Source/GameplayCameras/Private/Core/BlendStackCameraNode.cpp
--- a/Engine/Plugins/Cameras/GameplayCameras/Source/GameplayCameras/Private/Core/BlendStackCameraNode.cpp
+++ b/Engine/Plugins/Cameras/GameplayCameras/Source/GameplayCameras/Private/Core/BlendStackCameraNode.cpp
@@ -92,18 +92,16 @@ void FBlendStackCameraNodeEvaluator::Push(const FBlendStackCameraPushParams& Par
 
 	// Make the new stack entry, and use its storage buffer to build the tree of evaluators.
 	FCameraRigEntry NewEntry;
-
-	FCameraNodeEvaluatorTreeBuildParams BuildParams;
-	BuildParams.RootCameraNode = EntryRootNode;
-	BuildParams.Evaluator = Params.Evaluator;
-	BuildParams.EvaluationContext = Params.EvaluationContext;
-	FCameraNodeEvaluator* RootEvaluator = NewEntry.EvaluatorStorage.BuildEvaluatorTree(BuildParams);
-
-	NewEntry.EvaluationContext = Params.EvaluationContext;
-	NewEntry.CameraRig = Params.CameraRig;
-	NewEntry.RootNode = EntryRootNode;
-	NewEntry.RootEvaluator = RootEvaluator->CastThisChecked<FBlendStackRootCameraNodeEvaluator>();
-	NewEntry.bIsFirstFrame = true;
+	const bool bInitialized = InitializeEntry(
+			NewEntry, 
+			Params.CameraRig,
+			Params.Evaluator,
+			Params.EvaluationContext,
+			EntryRootNode);
+	if (!bInitialized)
+	{
+		return;
+	}
 
 #if WITH_EDITOR
 	IGameplayCamerasModule& GameplayCamerasModule = FModuleManager::GetModuleChecked<IGameplayCamerasModule>("GameplayCameras");
@@ -125,6 +123,89 @@ void FBlendStackCameraNodeEvaluator::Push(const FBlendStackCameraPushParams& Par
 	Entries.Add(MoveTemp(NewEntry));
 }
 
+bool FBlendStackCameraNodeEvaluator::InitializeEntry(
+		FCameraRigEntry& NewEntry, 
+		const UCameraRigAsset* CameraRig,
+		FCameraSystemEvaluator* Evaluator,
+		TSharedPtr<const FCameraEvaluationContext> EvaluationContext,
+		UBlendStackRootCameraNode* EntryRootNode)
+{
+	// Generate the hierarchy of node evaluators inside our storage buffer.
+	FCameraNodeEvaluatorTreeBuildParams BuildParams;
+	BuildParams.RootCameraNode = EntryRootNode;
+	BuildParams.AllocationInfo = &CameraRig->AllocationInfo.EvaluatorInfo;
+	FCameraNodeEvaluator* RootEvaluator = NewEntry.EvaluatorStorage.BuildEvaluatorTree(BuildParams);
+	if (!ensureMsgf(RootEvaluator, TEXT("No root evaluator was created for new camera rig!")))
+	{
+		return false;
+	}
+
+	// Initialize the node evaluators.
+	FCameraNodeEvaluatorInitializeParams InitParams;
+	InitParams.Evaluator = Evaluator;
+	InitParams.EvaluationContext = EvaluationContext;
+	InitParams.LastActiveCameraRig = GetActiveCameraRigEvaluationInfo();
+	RootEvaluator->Initialize(InitParams);
+
+	// Gather blended parameter evaluators.
+	NewEntry.ParameterEvaluators.Reset();
+	GatherEntryParameterEvaluators(RootEvaluator, NewEntry.ParameterEvaluators);
+
+	// Allocate variables in the variable table.
+	NewEntry.Result.VariableTable.Initialize(CameraRig->AllocationInfo.VariableTableInfo);
+
+	// Wrap up!
+	NewEntry.EvaluationContext = EvaluationContext;
+	NewEntry.CameraRig = CameraRig;
+	NewEntry.RootNode = EntryRootNode;
+	NewEntry.RootEvaluator = RootEvaluator->CastThisChecked<FBlendStackRootCameraNodeEvaluator>();
+	NewEntry.bIsFirstFrame = true;
+
+	return true;
+}
+
+void FBlendStackCameraNodeEvaluator::GatherEntryParameterEvaluators(FCameraNodeEvaluator* RootEvaluator, TArray<FCameraNodeEvaluator*>& OutParameterEvaluators)
+{
+	TArray<FCameraNodeEvaluator*> EvaluatorStack;
+	EvaluatorStack.Add(RootEvaluator);
+	while (!EvaluatorStack.IsEmpty())
+	{
+		FCameraNodeEvaluator* CurEvaluator = EvaluatorStack.Pop();
+
+		if (EnumHasAnyFlags(CurEvaluator->GetNodeEvaluatorFlags(), ECameraNodeEvaluatorFlags::NeedsParameterUpdate))
+		{
+			OutParameterEvaluators.Add(CurEvaluator);
+		}
+		else
+		{
+			FCameraNodeEvaluatorChildrenView CurChildren(CurEvaluator->GetChildren());
+			for (FCameraNodeEvaluator* Child : ReverseIterate(CurChildren))
+			{
+				if (Child)
+				{
+					EvaluatorStack.Add(Child);
+				}
+			}
+		}
+	}
+}
+
+TOptional<FCameraRigEvaluationInfo> FBlendStackCameraNodeEvaluator::GetActiveCameraRigEvaluationInfo() const
+{
+	if (Entries.Num() > 0)
+	{
+		const FCameraRigEntry& ActiveEntry = Entries[0];
+		FCameraRigEvaluationInfo Info(
+				ActiveEntry.CameraRig, 
+				ActiveEntry.EvaluationContext.Pin(),
+				ActiveEntry.Result,
+				ActiveEntry.RootEvaluator);
+		Info.bIsFrozen = ActiveEntry.bIsFrozen;
+		return Info;
+	}
+	return TOptional<FCameraRigEvaluationInfo>();
+}
+
 FCameraNodeEvaluatorChildrenView FBlendStackCameraNodeEvaluator::OnGetChildren()
 {
 	FCameraNodeEvaluatorChildrenView View;
@@ -147,41 +228,46 @@ void FBlendStackCameraNodeEvaluator::OnRun(const FCameraNodeEvaluationParams& Pa
 {
 	const UBlendStackCameraNode* BlendStackNode = GetCameraNodeAs<UBlendStackCameraNode>();
 
-	// Start by evaluating all the root nodes in the stack.
-	for (FCameraRigEntry& Entry : Entries)
+	// Filter out invalid entries and go through other basic validation.
+	struct FValidEntry
 	{
+		FCameraRigEntry& Entry;
+		TSharedPtr<const FCameraEvaluationContext> Context;
+		int32 EntryIndex;
+	};
+
+	TArray<FValidEntry> ValidEntries;
+
+	for (int32 Index = 0; Index < Entries.Num(); ++Index)
+	{
+		FCameraRigEntry& Entry(Entries[Index]);
 		TSharedPtr<const FCameraEvaluationContext> CurContext = Entry.EvaluationContext.Pin();
-		if (UNLIKELY(!CurContext.IsValid()))
+
+		if (!Entry.bIsFrozen)
 		{
-			Entry.Result.bIsValid = false;
+			// Check that we still have a valid context.
+			if (UNLIKELY(!CurContext.IsValid()))
+			{
+				Entry.Result.bIsValid = false;
 
 #if UE_GAMEPLAY_CAMERAS_TRACE
-			if (Entry.bLogWarnings)
-			{
-				UE_LOG(LogCameraSystem, Warning,
-						TEXT("Can't update camera rig '%s' because its evaluation context isn't valid."),
-						*GetNameSafe(Entry.CameraRig));
-				Entry.bLogWarnings = false;
-			}
+				if (Entry.bLogWarnings)
+				{
+					UE_LOG(LogCameraSystem, Warning,
+							TEXT("Can't update camera rig '%s' because its evaluation context isn't valid."),
+							*GetNameSafe(Entry.CameraRig));
+					Entry.bLogWarnings = false;
+				}
 #endif  // UE_GAMEPLAY_CAMERAS_TRACE
 
-			continue;
-		}
-
-		FCameraNodeEvaluationParams CurParams(Params);
-		CurParams.EvaluationContext = CurContext;
-		CurParams.bIsFirstFrame = Entry.bIsFirstFrame;
-
-		FCameraNodeEvaluationResult& CurResult(Entry.Result);
+				continue;
+			}
 
-		if (!Entry.bIsFrozen)
-		{
-			// If the context in which this camera rig runs doesn't have a valid result,
-			// skip it.
+			// Check that we have a valid result for this context.
 			const FCameraNodeEvaluationResult& ContextResult(CurContext->GetInitialResult());
 			if (UNLIKELY(!ContextResult.bIsValid))
 			{
-				CurResult.bIsValid = false;
+				Entry.Result.bIsValid = false;
 
 #if UE_GAMEPLAY_CAMERAS_TRACE
 				if (Entry.bLogWarnings)
@@ -195,48 +281,139 @@ void FBlendStackCameraNodeEvaluator::OnRun(const FCameraNodeEvaluationParams& Pa
 
 				continue;
 			}
+		}
+		// else: frozen entries may have null contexts or invalid initial results
+		//       because we're not going to update them anyway. We will however blend
+		//       them so we add them to the valid entries.
+
+		// This entry is valid, we'll process it.
+		ValidEntries.Add({ Entry, CurContext, Index });
 
-			// Start with the input given to us.
-			CurResult.CameraPose = OutResult.CameraPose;
-			CurResult.CameraPose.ClearAllChangedFlags();
-			CurResult.VariableTable.OverrideAll(OutResult.VariableTable);
-			CurResult.VariableTable.ClearAllWrittenThisFrameFlags();
+#if UE_GAMEPLAY_CAMERAS_TRACE
+		// This entry might have been temporarily invalid before. It's valid now, so let's
+		// re-enable warnings if it becomes invalid again in the future.
+		Entry.bLogWarnings = true;
+#endif  // UE_GAMEPLAY_CAMERAS_TRACE
+	}
 
-			// Override it with whatever the evaluation context has set on its result.
-			CurResult.CameraPose.OverrideChanged(ContextResult.CameraPose);
-			CurResult.VariableTable.OverrideAll(ContextResult.VariableTable);
-			CurResult.bIsCameraCut = OutResult.bIsCameraCut || ContextResult.bIsCameraCut;
-			CurResult.bIsValid = true;
+	// Gather parameters to pre-blend, and evaluate blend nodes.
+	for (FValidEntry& ValidEntry : ValidEntries)
+	{
+		FCameraRigEntry& Entry(ValidEntry.Entry);
+		TSharedPtr<const FCameraEvaluationContext> CurContext(ValidEntry.Context);
 
-			// Run the camera rig!
-			Entry.RootEvaluator->Run(CurParams, CurResult);
+		if (Entry.bIsFrozen)
+		{
+			continue;
 		}
-		else
+
+		FCameraNodeEvaluationParams CurParams(Params);
+		CurParams.EvaluationContext = CurContext;
+		CurParams.bIsFirstFrame = Entry.bIsFirstFrame;
+
+		FCameraNodeEvaluationResult& CurResult(Entry.Result);
+
+		// Gather input parameters.
+		if (!Entry.bInputRunThisFrame)
 		{
-			// Only evaluate the blend via the root node.
-			Entry.RootEvaluator->Run(CurParams, CurResult);
+			FCameraBlendedParameterUpdateParams InputParams(CurParams, CurResult.CameraPose);
+			FCameraBlendedParameterUpdateResult InputResult(CurResult.VariableTable);
+
+			for (FCameraNodeEvaluator* ParameterEvaluator : Entry.ParameterEvaluators)
+			{
+				ParameterEvaluator->UpdateParameters(InputParams, InputResult);
+			}
+
+			Entry.bInputRunThisFrame = true;
 		}
 
-#if UE_GAMEPLAY_CAMERAS_TRACE
-		Entry.bLogWarnings = true;
-#endif  // UE_GAMEPLAY_CAMERAS_TRACE
+		// Run blends.
+		if (!Entry.bBlendRunThisFrame)
+		{
+			FBlendCameraNodeEvaluator* BlendEvaluator = Entry.RootEvaluator->GetBlendEvaluator();
+			if (ensure(BlendEvaluator))
+			{
+				BlendEvaluator->Run(CurParams, CurResult);
+			}
+
+			Entry.bBlendRunThisFrame = true;
+		}
+	}
+
+	// Blend input variables.
+	for (FValidEntry& ValidEntry : ValidEntries)
+	{
+		FCameraRigEntry& Entry(ValidEntry.Entry);
+		TSharedPtr<const FCameraEvaluationContext> CurContext(ValidEntry.Context);
+
+		// Don't filter out frozen entries here, they still contribute to the blend
+		// using their last evaluated values.
+
+		if (!Entry.ParameterEvaluators.IsEmpty())
+		{
+			FCameraNodeEvaluationResult& CurResult(Entry.Result);
+
+			FCameraNodeEvaluationParams CurParams(Params);
+			CurParams.EvaluationContext = CurContext;
+			CurParams.bIsFirstFrame = Entry.bIsFirstFrame;
+			FCameraNodePreBlendParams PreBlendParams(CurParams, CurResult.CameraPose, CurResult.VariableTable);
+
+			FCameraNodePreBlendResult PreBlendResult(OutResult.VariableTable);
+
+			FBlendCameraNodeEvaluator* EntryBlendEvaluator = Entry.RootEvaluator->GetBlendEvaluator();
+			if (ensure(EntryBlendEvaluator))
+			{
+				EntryBlendEvaluator->BlendParameters(PreBlendParams, PreBlendResult);
+			}
+		}
+	}
+
+	// Run the root nodes. They will use the pre-blended inputs from the last step.
+	for (FValidEntry& ValidEntry : ValidEntries)
+	{
+		FCameraRigEntry& Entry(ValidEntry.Entry);
+		TSharedPtr<const FCameraEvaluationContext> CurContext(ValidEntry.Context);
+
+		FCameraNodeEvaluationParams CurParams(Params);
+		CurParams.EvaluationContext = CurContext;
+		CurParams.bIsFirstFrame = Entry.bIsFirstFrame;
+
+		FCameraNodeEvaluationResult& CurResult(Entry.Result);
+
+		// Start with the input given to us.
+		CurResult.CameraPose = OutResult.CameraPose;
+		CurResult.CameraPose.ClearAllChangedFlags();
+		CurResult.VariableTable.OverrideAll(OutResult.VariableTable);
+		CurResult.VariableTable.ClearAllWrittenThisFrameFlags();
+
+		// Override it with whatever the evaluation context has set on its result.
+		const FCameraNodeEvaluationResult& ContextResult(CurContext->GetInitialResult());
+		CurResult.CameraPose.OverrideChanged(ContextResult.CameraPose);
+		CurResult.VariableTable.OverrideAll(ContextResult.VariableTable);
+		CurResult.bIsCameraCut = OutResult.bIsCameraCut || ContextResult.bIsCameraCut;
+		CurResult.bIsValid = true;
+
+		// Run the camera rig's root node.
+		FCameraNodeEvaluator* RootEvaluator = Entry.RootEvaluator->GetRootEvaluator();
+		if (ensure(RootEvaluator))
+		{
+			RootEvaluator->Run(CurParams, CurResult);
+		}
 	}
 
 	// Now blend all the results, keeping track of blends that have reached 100% so
 	// that we can remove any camera rigs below (since the would have been completely
 	// blended out by that).
-	int32 EntryIndex = 0;
 	int32 PopEntriesBelow = INDEX_NONE;
-	for (FCameraRigEntry& Entry : Entries)
+	for (FValidEntry& ValidEntry : ValidEntries)
 	{
+		FCameraRigEntry& Entry(ValidEntry.Entry);
+		TSharedPtr<const FCameraEvaluationContext> CurContext(ValidEntry.Context);
+
 		FCameraNodeEvaluationResult& CurResult(Entry.Result);
-		if (UNLIKELY(!CurResult.bIsValid))
-		{
-			continue;
-		}
 
 		FCameraNodeEvaluationParams CurParams(Params);
-		CurParams.EvaluationContext = Entry.EvaluationContext.Pin();
+		CurParams.EvaluationContext = CurContext;
 		CurParams.bIsFirstFrame = Entry.bIsFirstFrame;
 		FCameraNodeBlendParams BlendParams(CurParams, CurResult);
 
@@ -249,17 +426,15 @@ void FBlendStackCameraNodeEvaluator::OnRun(const FCameraNodeEvaluationParams& Pa
 
 			if (BlendResult.bIsBlendFull && BlendResult.bIsBlendFinished)
 			{
-				PopEntriesBelow = EntryIndex;
+				PopEntriesBelow = ValidEntry.EntryIndex;
 			}
 		}
 		else
 		{
-			OutResult.CameraPose.OverrideChanged(CurResult.CameraPose);
+			OutResult.CameraPose.OverrideAll(CurResult.CameraPose);
 
-			PopEntriesBelow = EntryIndex;
+			PopEntriesBelow = ValidEntry.EntryIndex;
 		}
-		
-		++EntryIndex;
 	}
 
 	// Pop out camera rigs that have been blended out.
@@ -268,10 +443,12 @@ void FBlendStackCameraNodeEvaluator::OnRun(const FCameraNodeEvaluationParams& Pa
 		PopEntries(PopEntriesBelow);
 	}
 
-	// Reset first frame flags.
+	// Reset transient flags.
 	for (FCameraRigEntry& Entry : Entries)
 	{
 		Entry.bIsFirstFrame = false;
+		Entry.bInputRunThisFrame = false;
+		Entry.bBlendRunThisFrame = false;
 	}
 }
 
@@ -454,16 +631,17 @@ void FBlendStackCameraNodeEvaluator::OnPostBuildAsset(const FGameplayCameraAsset
 			Entry.RootNode->Blend = NewObject<UPopBlendCameraNode>(Entry.RootNode, NAME_None);
 
 			// Rebuild the evaluator tree.
-			FCameraNodeEvaluatorTreeBuildParams BuildParams;
-			BuildParams.RootCameraNode = Entry.RootNode;
-			BuildParams.Evaluator = OwningEvaluator;
-			BuildParams.EvaluationContext = Entry.EvaluationContext.Pin();
-			FCameraNodeEvaluator* RootEvaluator = Entry.EvaluatorStorage.BuildEvaluatorTree(BuildParams);
-
-			Entry.RootEvaluator = RootEvaluator->CastThisChecked<FBlendStackRootCameraNodeEvaluator>();
-
-			// This is the first frame for this new hierarchy of evaluators.
-			Entry.bIsFirstFrame = true;
+			const bool bInitialized = InitializeEntry(
+					Entry,
+					Entry.CameraRig,
+					OwningEvaluator,
+					Entry.EvaluationContext.Pin(),
+					Entry.RootNode);
+			if (!bInitialized)
+			{
+				Entry.bIsFrozen = true;
+				continue;
+			}
 		}
 	}
 }
diff --git a/Engine/Plugins/Cameras/GameplayCameras/Source/GameplayCameras/Private/Core/BlendStackRootCameraNode.cpp b/Engine/Plugins/Cameras/GameplayCameras/Source/GameplayCameras/Private/Core/BlendStackRootCameraNode.cpp
--- a/Engine/Plugins/Cameras/GameplayCameras/Source/GameplayCameras/Private/Core/BlendStackRootCameraNode.cpp
+++ b/Engine/Plugins/Cameras/GameplayCameras/Source/GameplayCameras/Private/Core/BlendStackRootCameraNode.cpp
@@ -64,14 +64,7 @@ void FBlendStackRootCameraNodeEvaluator::OnBuild(const FCameraNodeEvaluatorBuild
 
 void FBlendStackRootCameraNodeEvaluator::OnRun(const FCameraNodeEvaluationParams& Params, FCameraNodeEvaluationResult& OutResult)
 {
-	if (BlendEvaluator)
-	{
-		BlendEvaluator->Run(Params, OutResult);
-	}
-	if (RootEvaluator)
-	{
-		RootEvaluator->Run(Params, OutResult);
-	}
+	// Nothing to do, the blend stack node runs our blend and root nodes selectively.
 }
 
 #if UE_GAMEPLAY_CAMERAS_DEBUG
diff --git a/Engine/Plugins/Cameras/GameplayCameras/Source/GameplayCameras/Private/Core/BuiltInCameraVariables.cpp b/Engine/Plugins/Cameras/GameplayCameras/Source/GameplayCameras/Private/Core/BuiltInCameraVariables.cpp
new file mode 100644
--- /dev/null
+++ b/Engine/Plugins/Cameras/GameplayCameras/Source/GameplayCameras/Private/Core/BuiltInCameraVariables.cpp
@@ -0,0 +1,97 @@
+// Copyright Epic Games, Inc. All Rights Reserved.
+
+#include "Core/BuiltInCameraVariables.h"
+
+namespace UE::Cameras
+{
+
+namespace Private
+{
+
+struct FBuiltInCameraVariablesGuids
+{
+	FGuid YawGuid;       // {6E23348F-290E-460F-9432-ED80E7CA03F2}
+	FGuid PitchGuid;     // {B7142BDA-6775-45A4-B06C-D07B712D1F89}
+	FGuid RollGuid;      // {DABC149B-DE7D-4917-820C-7EEBE2A9846C}
+	FGuid ZoomGuid;      // {56BF6A94-87B3-4648-84A4-391BDAD27061}
+
+	FGuid YawPitchGuid;  // {67B7828D-C645-4907-92BE-B40DEB8C838C}
+
+	FBuiltInCameraVariablesGuids()
+	{
+		YawGuid = FGuid(0x6E23348F, 0x290E460F, 0x9432ED80, 0xE7CA03F2);
+		PitchGuid = FGuid(0xB7142BDA, 0x677545A4, 0xB06CD07B, 0x712D1F89);
+		RollGuid = FGuid(0xDABC149B, 0xDE7D4917, 0x820C7EEB, 0xE2A9846C);
+		ZoomGuid = FGuid(0x56BF6A94, 0x87B34648, 0x84A4391B, 0xDAD27061);
+
+		YawPitchGuid = FGuid(0x67B7828D, 0xC6454907, 0x92BEB40D, 0xEB8C838C);
+	}
+};
+
+void MakeCameraVariableDefinition(
+		FCameraVariableDefinition& OutDefinition,
+		const FGuid& VariableGuid, ECameraVariableType VariableType, bool bIsInput,
+		const FString& VariableName)
+{
+	OutDefinition.VariableID =  FCameraVariableID::FromHashValue(GetTypeHash(VariableGuid));
+	OutDefinition.VariableType = VariableType;
+	OutDefinition.bIsInput = bIsInput;
+#if WITH_EDITORONLY_DATA
+	OutDefinition.VariableName = VariableName;
+#endif
+}
+
+}  // namespace Private
+
+const FBuiltInCameraVariables& FBuiltInCameraVariables::Get()
+{
+	static FBuiltInCameraVariables Instance;
+	return Instance;
+}
+
+FBuiltInCameraVariables::FBuiltInCameraVariables()
+{
+	using namespace Private;
+
+	FBuiltInCameraVariablesGuids KnownGuids;
+
+	MakeCameraVariableDefinition(YawDefinition, KnownGuids.YawGuid, ECameraVariableType::Double, true, TEXT("Yaw"));
+	MakeCameraVariableDefinition(PitchDefinition, KnownGuids.PitchGuid, ECameraVariableType::Double, true, TEXT("Pitch"));
+	MakeCameraVariableDefinition(RollDefinition, KnownGuids.RollGuid, ECameraVariableType::Double, true, TEXT("Roll"));
+	MakeCameraVariableDefinition(ZoomDefinition, KnownGuids.ZoomGuid, ECameraVariableType::Double, true, TEXT("Zoom"));
+
+	MakeCameraVariableDefinition(YawPitchDefinition, KnownGuids.YawPitchGuid, ECameraVariableType::Vector2d, true, TEXT("YawPitch"));
+}
+
+const FCameraVariableDefinition& FBuiltInCameraVariables::GetDefinition(EBuiltInDoubleCameraVariable BuiltInVariable) const
+{
+	static const FCameraVariableDefinition Invalid;
+	switch (BuiltInVariable)
+	{
+		case EBuiltInDoubleCameraVariable::Yaw:
+			return YawDefinition;
+		case EBuiltInDoubleCameraVariable::Pitch:
+			return PitchDefinition;
+		case EBuiltInDoubleCameraVariable::Roll:
+			return RollDefinition;
+		case EBuiltInDoubleCameraVariable::Zoom:
+			return ZoomDefinition;
+		default:
+			return Invalid;
+	}
+}
+
+const FCameraVariableDefinition& FBuiltInCameraVariables::GetDefinition(EBuiltInVector2dCameraVariable BuiltInVariable) const
+{
+	static const FCameraVariableDefinition Invalid;
+	switch (BuiltInVariable)
+	{
+		case EBuiltInVector2dCameraVariable::YawPitch:
+			return YawPitchDefinition;
+		default:
+			return Invalid;
+	}
+}
+
+}  // namespace UE::Cameras
+
diff --git a/Engine/Plugins/Cameras/GameplayCameras/Source/GameplayCameras/Private/Core/CameraNodeEvaluator.cpp b/Engine/Plugins/Cameras/GameplayCameras/Source/GameplayCameras/Private/Core/CameraNodeEvaluator.cpp
--- a/Engine/Plugins/Cameras/GameplayCameras/Source/GameplayCameras/Private/Core/CameraNodeEvaluator.cpp
+++ b/Engine/Plugins/Cameras/GameplayCameras/Source/GameplayCameras/Private/Core/CameraNodeEvaluator.cpp
@@ -41,6 +41,11 @@ void FCameraNodeEvaluator::SetPrivateCameraNode(TObjectPtr<const UCameraNode> In
 	PrivateCameraNode = InCameraNode;
 }
 
+void FCameraNodeEvaluator::SetNodeEvaluatorFlags(ECameraNodeEvaluatorFlags InFlags)
+{
+	PrivateFlags = InFlags;
+}
+
 FCameraNodeEvaluatorChildrenView FCameraNodeEvaluator::GetChildren()
 {
 	return OnGetChildren();
@@ -82,11 +87,45 @@ void FCameraNodeEvaluator::Initialize(const FCameraNodeEvaluatorInitializeParams
 	}
 }
 
+void FCameraNodeEvaluator::UpdateParameters(const FCameraBlendedParameterUpdateParams& Params, FCameraBlendedParameterUpdateResult& OutResult)
+{
+	if (!PrivateCameraNode || PrivateCameraNode->bIsEnabled)
+	{
+		if (EnumHasAnyFlags(PrivateFlags, ECameraNodeEvaluatorFlags::NeedsParameterUpdate))
+		{
+			OnUpdateParameters(Params, OutResult);
+		}
+		else
+		{
+			for (FCameraNodeEvaluator* Child : GetChildren())
+			{
+				if (Child)
+				{
+					Child->UpdateParameters(Params, OutResult);
+				}
+			}
+		}
+	}
+}
+
 void FCameraNodeEvaluator::Run(const FCameraNodeEvaluationParams& Params, FCameraNodeEvaluationResult& OutResult)
 {
 	if (!PrivateCameraNode || PrivateCameraNode->bIsEnabled)
 	{
-		OnRun(Params, OutResult);
+		if (EnumHasAnyFlags(PrivateFlags, ECameraNodeEvaluatorFlags::NeedsEvaluationUpdate))
+		{
+			OnRun(Params, OutResult);
+		}
+		else
+		{
+			for (FCameraNodeEvaluator* Child : GetChildren())
+			{
+				if (Child)
+				{
+					Child->Run(Params, OutResult);
+				}
+			}
+		}
 	}
 }
 
diff --git a/Engine/Plugins/Cameras/GameplayCameras/Source/GameplayCameras/Private/Core/CameraNodeEvaluatorStorage.cpp b/Engine/Plugins/Cameras/GameplayCameras/Source/GameplayCameras/Private/Core/CameraNodeEvaluatorStorage.cpp
--- a/Engine/Plugins/Cameras/GameplayCameras/Source/GameplayCameras/Private/Core/CameraNodeEvaluatorStorage.cpp
+++ b/Engine/Plugins/Cameras/GameplayCameras/Source/GameplayCameras/Private/Core/CameraNodeEvaluatorStorage.cpp
@@ -16,7 +16,10 @@ FCameraNodeEvaluatorPtr FCameraNodeEvaluatorStorage::BuildEvaluatorTree(const FC
 	{
 		const uint16 NewCapacity = Params.AllocationInfo->TotalSizeof;
 		const uint16 NewAlignment = Params.AllocationInfo->MaxAlignof;
-		Super::AllocatePage(NewCapacity, NewAlignment);
+		if (NewCapacity > 0 && NewAlignment > 0)
+		{
+			Super::AllocatePage(NewCapacity, NewAlignment);
+		}
 	}
 
 	FCameraNodeEvaluatorBuilder Builder(*this);
@@ -24,15 +27,6 @@ FCameraNodeEvaluatorPtr FCameraNodeEvaluatorStorage::BuildEvaluatorTree(const FC
 
 	FCameraNodeEvaluatorPtr RootEvaluator = BuildParams.BuildEvaluator(Params.RootCameraNode);
 
-	if (Params.bInitialize)
-	{
-		FCameraNodeEvaluatorInitializeParams InitParams;
-		InitParams.Evaluator = Params.Evaluator;
-		InitParams.EvaluationContext = Params.EvaluationContext;
-
-		RootEvaluator->Initialize(InitParams);
-	}
-
 	return RootEvaluator;
 }
 
diff --git a/Engine/Plugins/Cameras/GameplayCameras/Source/GameplayCameras/Private/Core/CameraRigAssetBuilder.cpp b/Engine/Plugins/Cameras/GameplayCameras/Source/GameplayCameras/Private/Core/CameraRigAssetBuilder.cpp
--- a/Engine/Plugins/Cameras/GameplayCameras/Source/GameplayCameras/Private/Core/CameraRigAssetBuilder.cpp
+++ b/Engine/Plugins/Cameras/GameplayCameras/Source/GameplayCameras/Private/Core/CameraRigAssetBuilder.cpp
@@ -471,7 +471,6 @@ void FCameraRigAssetBuilder::BuildAllocationInfo()
 	// Build a mock tree of evaluators.
 	FCameraNodeEvaluatorTreeBuildParams BuildParams;
 	BuildParams.RootCameraNode = CameraRig->RootNode;
-	BuildParams.bInitialize = false;  // Don't initialize, we just want to know how much room they take.
 	FCameraNodeEvaluatorStorage Storage;
 	Storage.BuildEvaluatorTree(BuildParams);
 
diff --git a/Engine/Plugins/Cameras/GameplayCameras/Source/GameplayCameras/Private/Core/CameraSystemEvaluator.cpp b/Engine/Plugins/Cameras/GameplayCameras/Source/GameplayCameras/Private/Core/CameraSystemEvaluator.cpp
--- a/Engine/Plugins/Cameras/GameplayCameras/Source/GameplayCameras/Private/Core/CameraSystemEvaluator.cpp
+++ b/Engine/Plugins/Cameras/GameplayCameras/Source/GameplayCameras/Private/Core/CameraSystemEvaluator.cpp
@@ -58,9 +58,15 @@ void FCameraSystemEvaluator::Initialize(const FCameraSystemEvaluatorCreateParams
 	ContextStack.Initialize(*this);
 
 	FCameraNodeEvaluatorTreeBuildParams BuildParams;
-	BuildParams.Evaluator = this;
 	BuildParams.RootCameraNode = RootNode;
 	RootEvaluator = static_cast<FRootCameraNodeEvaluator*>(RootEvaluatorStorage.BuildEvaluatorTree(BuildParams));
+
+	if (ensure(RootEvaluator))
+	{
+		FCameraNodeEvaluatorInitializeParams InitParams;
+		InitParams.Evaluator = this;
+		RootEvaluator->Initialize(InitParams);
+	}
 }
 
 void FCameraSystemEvaluator::AddReferencedObjects(FReferenceCollector& Collector)
diff --git a/Engine/Plugins/Cameras/GameplayCameras/Source/GameplayCameras/Private/Core/CameraVariableTable.cpp b/Engine/Plugins/Cameras/GameplayCameras/Source/GameplayCameras/Private/Core/CameraVariableTable.cpp
--- a/Engine/Plugins/Cameras/GameplayCameras/Source/GameplayCameras/Private/Core/CameraVariableTable.cpp
+++ b/Engine/Plugins/Cameras/GameplayCameras/Source/GameplayCameras/Private/Core/CameraVariableTable.cpp
@@ -89,6 +89,15 @@ void FCameraVariableTable::Initialize(const FCameraVariableTableAllocationInfo&
 		NewEntry.ID = VariableDefinition.VariableID;
 		NewEntry.Type = VariableDefinition.VariableType;
 		NewEntry.Offset = NewEntryOffset;
+		NewEntry.Flags = EEntryFlags::None;
+		if (VariableDefinition.bIsPrivate)
+		{
+			NewEntry.Flags |= EEntryFlags::Private;
+		}
+		if (VariableDefinition.bIsInput)
+		{
+			NewEntry.Flags |= EEntryFlags::Input;
+		}
 #if WITH_EDITORONLY_DATA
 		NewEntry.DebugName = VariableDefinition.VariableName;
 #endif
@@ -237,30 +246,37 @@ void FCameraVariableTable::ClearAllWrittenThisFrameFlags()
 
 void FCameraVariableTable::OverrideAll(const FCameraVariableTable& OtherTable)
 {
-	InternalOverride(OtherTable, nullptr, false, nullptr, false);
+	const ECameraVariableTableFilter Filter = ECameraVariableTableFilter::All;
+	InternalOverride(OtherTable, Filter, nullptr, false, nullptr);
 }
 
-void FCameraVariableTable::OverrideChanged(const FCameraVariableTable& OtherTable)
+void FCameraVariableTable::Override(const FCameraVariableTable& OtherTable, ECameraVariableTableFilter Filter)
 {
-	InternalOverride(OtherTable, nullptr, false, nullptr, true);
+	InternalOverride(OtherTable, Filter, nullptr, false, nullptr);
 }
 
-void FCameraVariableTable::OverrideChanged(const FCameraVariableTable& OtherTable, const FCameraVariableTableFlags& InMask, bool bInvertMask, FCameraVariableTableFlags& OutMask)
+void FCameraVariableTable::Override(const FCameraVariableTable& OtherTable, ECameraVariableTableFilter Filter, const FCameraVariableTableFlags& InMask, bool bInvertMask, FCameraVariableTableFlags& OutMask)
 {
-	InternalOverride(OtherTable, &InMask, bInvertMask, &OutMask, true);
+	InternalOverride(OtherTable, Filter, &InMask, bInvertMask, &OutMask);
 }
 
-void FCameraVariableTable::InternalOverride(const FCameraVariableTable& OtherTable, const FCameraVariableTableFlags* InMask, bool bInvertMask, FCameraVariableTableFlags* OutMask, bool bChangedOnly)
+void FCameraVariableTable::InternalOverride(const FCameraVariableTable& OtherTable, ECameraVariableTableFilter Filter, const FCameraVariableTableFlags* InMask, bool bInvertMask, FCameraVariableTableFlags* OutMask)
 {
 	using namespace UE::Cameras::Private;
 
+	const bool bChangedOnly = EnumHasAnyFlags(Filter, ECameraVariableTableFilter::ChangedOnly);
+	const bool bInputs = EnumHasAnyFlags(Filter, ECameraVariableTableFilter::Input);
+	const bool bOutputs = EnumHasAnyFlags(Filter, ECameraVariableTableFilter::Output);
+
 	for (auto& OtherPair : OtherTable.Entries)
 	{
 		// Look for entries in the other table that have been written to, and aren't private.
 		const FEntry& OtherEntry  = OtherPair.Value;
 		const EEntryFlags OtherFlags = OtherEntry.Flags;
+		const bool bOtherEntryIsInput = EnumHasAnyFlags(OtherFlags, EEntryFlags::Input);
 		if (EnumHasAnyFlags(OtherFlags, EEntryFlags::Written)
 				&& (!bChangedOnly || EnumHasAnyFlags(OtherFlags, EEntryFlags::WrittenThisFrame))
+				&& ((bInputs && bOtherEntryIsInput) || (bOutputs && !bOtherEntryIsInput))
 				&& !EnumHasAnyFlags(OtherFlags, EEntryFlags::Private)
 				&& IsVariableInMask(OtherEntry.ID, InMask, bInvertMask))
 		{
@@ -322,30 +338,37 @@ void FCameraVariableTable::InternalOverride(const FCameraVariableTable& OtherTab
 
 void FCameraVariableTable::LerpAll(const FCameraVariableTable& ToTable, float Factor)
 {
-	InternalLerp(ToTable, Factor, nullptr, false, nullptr, false);
+	const ECameraVariableTableFilter Filter = ECameraVariableTableFilter::All;
+	InternalLerp(ToTable, Filter, Factor, nullptr, false, nullptr);
 }
 
-void FCameraVariableTable::LerpChanged(const FCameraVariableTable& ToTable, float Factor)
+void FCameraVariableTable::Lerp(const FCameraVariableTable& ToTable, ECameraVariableTableFilter Filter, float Factor)
 {
-	InternalLerp(ToTable, Factor, nullptr, false, nullptr, true);
+	InternalLerp(ToTable, Filter, Factor, nullptr, false, nullptr);
 }
 
-void FCameraVariableTable::LerpChanged(const FCameraVariableTable& ToTable, float Factor, const FCameraVariableTableFlags& InMask, bool bInvertMask, FCameraVariableTableFlags& OutMask)
+void FCameraVariableTable::Lerp(const FCameraVariableTable& ToTable, ECameraVariableTableFilter Filter, float Factor, const FCameraVariableTableFlags& InMask, bool bInvertMask, FCameraVariableTableFlags& OutMask)
 {
-	InternalLerp(ToTable, Factor, &InMask, false, &OutMask, true);
+	InternalLerp(ToTable, Filter, Factor, &InMask, false, &OutMask);
 }
 
-void FCameraVariableTable::InternalLerp(const FCameraVariableTable& ToTable, float Factor, const FCameraVariableTableFlags* InMask, bool bInvertMask, FCameraVariableTableFlags* OutMask, bool bChangedOnly)
+void FCameraVariableTable::InternalLerp(const FCameraVariableTable& ToTable, ECameraVariableTableFilter Filter, float Factor, const FCameraVariableTableFlags* InMask, bool bInvertMask, FCameraVariableTableFlags* OutMask)
 {
 	using namespace UE::Cameras::Private;
 
+	const bool bChangedOnly = EnumHasAnyFlags(Filter, ECameraVariableTableFilter::ChangedOnly);
+	const bool bInputs = EnumHasAnyFlags(Filter, ECameraVariableTableFilter::Input);
+	const bool bOutputs = EnumHasAnyFlags(Filter, ECameraVariableTableFilter::Output);
+
 	for (auto& ToPair : ToTable.Entries)
 	{
 		// Look for entries in the other table that have been written to, and aren't private.
 		const FEntry& ToEntry  = ToPair.Value;
 		const EEntryFlags ToFlags = ToEntry.Flags;
+		const bool bToEntryIsInput = EnumHasAnyFlags(ToFlags, EEntryFlags::Input);
 		if (EnumHasAnyFlags(ToFlags, EEntryFlags::Written)
 				&& (!bChangedOnly || EnumHasAnyFlags(ToFlags, EEntryFlags::WrittenThisFrame))
+				&& ((bInputs && bToEntryIsInput) || (bOutputs && !bToEntryIsInput))
 				&& !EnumHasAnyFlags(ToFlags, EEntryFlags::Private)
 				&& IsVariableInMask(ToEntry.ID, InMask, bInvertMask))
 		{
diff --git a/Engine/Plugins/Cameras/GameplayCameras/Source/GameplayCameras/Private/Nodes/Blends/PopBlendCameraNode.cpp b/Engine/Plugins/Cameras/GameplayCameras/Source/GameplayCameras/Private/Nodes/Blends/PopBlendCameraNode.cpp
--- a/Engine/Plugins/Cameras/GameplayCameras/Source/GameplayCameras/Private/Nodes/Blends/PopBlendCameraNode.cpp
+++ b/Engine/Plugins/Cameras/GameplayCameras/Source/GameplayCameras/Private/Nodes/Blends/PopBlendCameraNode.cpp
@@ -13,6 +13,7 @@ class FPopBlendCameraNodeEvaluator : public FBlendCameraNodeEvaluator
 
 protected:
 	virtual void OnRun(const FCameraNodeEvaluationParams& Params, FCameraNodeEvaluationResult& OutResult) override;
+	virtual void OnBlendParameters(const FCameraNodePreBlendParams& Params, FCameraNodePreBlendResult& OutResult) override;
 	virtual void OnBlendResults(const FCameraNodeBlendParams& Params, FCameraNodeBlendResult& OutResult) override;
 };
 
@@ -22,12 +23,19 @@ void FPopBlendCameraNodeEvaluator::OnRun(const FCameraNodeEvaluationParams& Para
 {
 }
 
+void FPopBlendCameraNodeEvaluator::OnBlendParameters(const FCameraNodePreBlendParams& Params, FCameraNodePreBlendResult& OutResult)
+{
+	const FCameraVariableTable& ChildVariableTable(Params.ChildVariableTable);
+	OutResult.VariableTable.Override(ChildVariableTable, ECameraVariableTableFilter::Input);
+}
+
 void FPopBlendCameraNodeEvaluator::OnBlendResults(const FCameraNodeBlendParams& Params, FCameraNodeBlendResult& OutResult)
 {
 	const FCameraNodeEvaluationResult& ChildResult(Params.ChildResult);
 	FCameraNodeEvaluationResult& BlendedResult(OutResult.BlendedResult);
 	
-	BlendedResult.CameraPose.OverrideChanged(ChildResult.CameraPose);
+	BlendedResult.CameraPose.OverrideAll(ChildResult.CameraPose);
+	BlendedResult.VariableTable.OverrideAll(ChildResult.VariableTable);
 	if (ChildResult.bIsCameraCut || Params.ChildParams.bIsFirstFrame)
 	{
 		BlendedResult.bIsCameraCut = true;
diff --git a/Engine/Plugins/Cameras/GameplayCameras/Source/GameplayCameras/Private/Nodes/Blends/SimpleBlendCameraNode.cpp b/Engine/Plugins/Cameras/GameplayCameras/Source/GameplayCameras/Private/Nodes/Blends/SimpleBlendCameraNode.cpp
--- a/Engine/Plugins/Cameras/GameplayCameras/Source/GameplayCameras/Private/Nodes/Blends/SimpleBlendCameraNode.cpp
+++ b/Engine/Plugins/Cameras/GameplayCameras/Source/GameplayCameras/Private/Nodes/Blends/SimpleBlendCameraNode.cpp
@@ -26,6 +26,12 @@ void FSimpleBlendCameraNodeEvaluator::OnRun(const FCameraNodeEvaluationParams& P
 	BlendFactor = FactorResult.BlendFactor;
 }
 
+void FSimpleBlendCameraNodeEvaluator::OnBlendParameters(const FCameraNodePreBlendParams& Params, FCameraNodePreBlendResult& OutResult)
+{
+	const FCameraVariableTable& ChildVariableTable(Params.ChildVariableTable);
+	OutResult.VariableTable.Lerp(ChildVariableTable, ECameraVariableTableFilter::Input, BlendFactor);
+}
+
 void FSimpleBlendCameraNodeEvaluator::OnBlendResults(const FCameraNodeBlendParams& Params, FCameraNodeBlendResult& OutResult)
 {
 	const FCameraNodeEvaluationResult& ChildResult(Params.ChildResult);
diff --git a/Engine/Plugins/Cameras/GameplayCameras/Source/GameplayCameras/Private/Nodes/Common/BoomArmCameraNode.cpp b/Engine/Plugins/Cameras/GameplayCameras/Source/GameplayCameras/Private/Nodes/Common/BoomArmCameraNode.cpp
--- a/Engine/Plugins/Cameras/GameplayCameras/Source/GameplayCameras/Private/Nodes/Common/BoomArmCameraNode.cpp
+++ b/Engine/Plugins/Cameras/GameplayCameras/Source/GameplayCameras/Private/Nodes/Common/BoomArmCameraNode.cpp
@@ -3,6 +3,8 @@
 #include "Nodes/Common/BoomArmCameraNode.h"
 
 #include "Core/CameraEvaluationContext.h"
+#include "Core/CameraParameterReader.h"
+#include "Nodes/Input/CameraRigInput2DSlot.h"
 #include "GameFramework/Pawn.h"
 #include "GameFramework/PlayerController.h"
 #include "GameplayCameras.h"
@@ -18,15 +20,46 @@ class FBoomArmCameraNodeEvaluator : public FCameraNodeEvaluator
 
 protected:
 
+	virtual void OnBuild(const FCameraNodeEvaluatorBuildParams& Params) override;
+	virtual void OnInitialize(const FCameraNodeEvaluatorInitializeParams& Params) override;
+	virtual FCameraNodeEvaluatorChildrenView OnGetChildren() override;
 	virtual void OnRun(const FCameraNodeEvaluationParams& Params, FCameraNodeEvaluationResult& OutResult) override;
+
+private:
+
+	TCameraParameterReader<FVector3d> BoomOffsetReader;
+	FCameraRigInput2DSlotEvaluator* InputSlotEvaluator = nullptr;
 };
 
 UE_DEFINE_CAMERA_NODE_EVALUATOR(FBoomArmCameraNodeEvaluator)
 
+void FBoomArmCameraNodeEvaluator::OnBuild(const FCameraNodeEvaluatorBuildParams& Params)
+{
+	const UBoomArmCameraNode* BoomArmNode = GetCameraNodeAs<UBoomArmCameraNode>();
+	InputSlotEvaluator = Params.BuildEvaluatorAs<FCameraRigInput2DSlotEvaluator>(BoomArmNode->InputSlot);
+}
+
+void FBoomArmCameraNodeEvaluator::OnInitialize(const FCameraNodeEvaluatorInitializeParams& Params)
+{
+	const UBoomArmCameraNode* BoomArmNode = GetCameraNodeAs<UBoomArmCameraNode>();
+	BoomOffsetReader.Initialize(BoomArmNode->BoomOffset);
+}
+
+FCameraNodeEvaluatorChildrenView FBoomArmCameraNodeEvaluator::OnGetChildren()
+{
+	return FCameraNodeEvaluatorChildrenView({ InputSlotEvaluator });
+}
+
 void FBoomArmCameraNodeEvaluator::OnRun(const FCameraNodeEvaluationParams& Params, FCameraNodeEvaluationResult& OutResult)
 {
-	FRotator3d BoomRotation;
-	if (Params.EvaluationContext)
+	FRotator3d BoomRotation = FRotator3d::ZeroRotator;
+	if (InputSlotEvaluator)
+	{
+		InputSlotEvaluator->Run(Params, OutResult);
+		const FVector2d YawPitch = InputSlotEvaluator->GetInputValue();
+		BoomRotation = FRotator3d(YawPitch.Y, YawPitch.X, 0);
+	}
+	else if (Params.EvaluationContext)
 	{
 		if (APlayerController* PlayerController = Params.EvaluationContext->GetPlayerController())
 		{
@@ -48,7 +81,7 @@ void FBoomArmCameraNodeEvaluator::OnRun(const FCameraNodeEvaluationParams& Param
 	// by using the fact that BoomRotation is, well, just a rotation, and CameraPose.Location is of
 	// course just a translation. So we can put them both in the same transform:
 	const FTransform3d BoomPivot(BoomRotation, OutResult.CameraPose.GetLocation());
-	const FTransform3d BoomOffset(BoomArmNode->BoomOffset);
+	const FTransform3d BoomOffset(BoomOffsetReader.Get(OutResult.VariableTable));
 
 	const FTransform3d FinalTransform(BoomOffset * BoomPivot);
 
@@ -57,6 +90,11 @@ void FBoomArmCameraNodeEvaluator::OnRun(const FCameraNodeEvaluationParams& Param
 
 }  // namespace UE::Cameras
 
+FCameraNodeChildrenView UBoomArmCameraNode::OnGetChildren()
+{
+	return FCameraNodeChildrenView({ InputSlot });
+}
+
 FCameraNodeEvaluatorPtr UBoomArmCameraNode::OnBuildEvaluator(FCameraNodeEvaluatorBuilder& Builder) const
 {
 	using namespace UE::Cameras;
diff --git a/Engine/Plugins/Cameras/GameplayCameras/Source/GameplayCameras/Private/Nodes/Input/CameraRigInput1DSlot.cpp b/Engine/Plugins/Cameras/GameplayCameras/Source/GameplayCameras/Private/Nodes/Input/CameraRigInput1DSlot.cpp
new file mode 100644
--- /dev/null
+++ b/Engine/Plugins/Cameras/GameplayCameras/Source/GameplayCameras/Private/Nodes/Input/CameraRigInput1DSlot.cpp
@@ -0,0 +1,140 @@
+// Copyright Epic Games, Inc. All Rights Reserved.
+
+#include "Nodes/Input/CameraRigInput1DSlot.h"
+
+#include "Core/CameraBuildLog.h"
+#include "Core/CameraRigAsset.h"
+#include "Core/CameraRigBuildContext.h"
+#include "Core/CameraVariableAssets.h"
+#include "Core/CameraVariableTable.h"
+#include "Nodes/Input/Input1DCameraNode.h"
+
+#define LOCTEXT_NAMESPACE "CameraRigInputSlot"
+
+namespace UE::Cameras
+{
+
+UE_DEFINE_CAMERA_NODE_EVALUATOR(FCameraRigInput1DSlotEvaluator)
+
+FCameraRigInput1DSlotEvaluator::FCameraRigInput1DSlotEvaluator()
+{
+	SetNodeEvaluatorFlags(ECameraNodeEvaluatorFlags::NeedsParameterUpdate | ECameraNodeEvaluatorFlags::NeedsEvaluationUpdate);
+}
+
+void FCameraRigInput1DSlotEvaluator::OnBuild(const FCameraNodeEvaluatorBuildParams& Params)
+{
+	const UCameraRigInput1DSlot* SlotNode = GetCameraNodeAs<UCameraRigInput1DSlot>();
+	ChildEvaluator = Params.BuildEvaluatorAs<FInput1DCameraNodeEvaluator>(SlotNode->Child);
+}
+
+void FCameraRigInput1DSlotEvaluator::OnInitialize(const FCameraNodeEvaluatorInitializeParams& Params)
+{
+	TransientInputValue = 0.f;
+	InputValue = 0.f;
+
+	const UCameraRigInput1DSlot* SlotNode = GetCameraNodeAs<UCameraRigInput1DSlot>();
+	if (SlotNode->GetVariableID().IsValid() && Params.LastActiveCameraRig.IsSet())
+	{
+		const FCameraRigEvaluationInfo& LastActiveRig = Params.LastActiveCameraRig.GetValue();
+		LastActiveRig.LastResult.VariableTable.TryGetValue<double>(SlotNode->GetVariableID(), InputValue);
+	}
+}
+
+FCameraNodeEvaluatorChildrenView FCameraRigInput1DSlotEvaluator::OnGetChildren()
+{
+	return FCameraNodeEvaluatorChildrenView{ ChildEvaluator };
+}
+
+void FCameraRigInput1DSlotEvaluator::OnUpdateParameters(const FCameraBlendedParameterUpdateParams& Params, FCameraBlendedParameterUpdateResult& OutResult)
+{
+	if (ChildEvaluator)
+	{
+		ChildEvaluator->UpdateParameters(Params, OutResult);
+
+		TransientInputValue = ChildEvaluator->GetInputValue();
+	}
+
+	const UCameraRigInput1DSlot* SlotNode = GetCameraNodeAs<UCameraRigInput1DSlot>();
+	if (SlotNode->InputSlotParameters.bIsPreBlended)
+	{
+		OutResult.VariableTable.SetValue<double>(SlotNode->GetTransientVariableID(), TransientInputValue);
+	}
+}
+
+void FCameraRigInput1DSlotEvaluator::OnRun(const FCameraNodeEvaluationParams& Params, FCameraNodeEvaluationResult& OutResult)
+{
+	const UCameraRigInput1DSlot* SlotNode = GetCameraNodeAs<UCameraRigInput1DSlot>();
+
+	if (SlotNode->InputSlotParameters.bIsPreBlended)
+	{
+		TransientInputValue = OutResult.VariableTable.GetValue<double>(SlotNode->GetTransientVariableID());
+	}
+
+	if (SlotNode->InputSlotParameters.bIsAccumulated)
+	{
+		InputValue += TransientInputValue;
+	}
+	else
+	{
+		InputValue = TransientInputValue;
+	}
+
+	InputValue = SlotNode->Normalize.NormalizeValue(InputValue);
+	InputValue = SlotNode->Clamp.ClampValue(InputValue);
+
+	OutResult.VariableTable.SetValue<double>(SlotNode->GetVariableID(), InputValue);
+}
+
+}  // namespace UE::Cameras
+
+FCameraNodeChildrenView UCameraRigInput1DSlot::OnGetChildren()
+{
+	return FCameraNodeChildrenView{ Child };
+}
+
+void UCameraRigInput1DSlot::OnBuild(FCameraRigBuildContext& BuildContext) 
+{
+	using namespace UE::Cameras;
+
+	FCameraVariableDefinition VariableDefinition;
+
+	if (BuiltInVariable != EBuiltInDoubleCameraVariable::None)
+	{
+		VariableDefinition = FBuiltInCameraVariables::Get().GetDefinition(BuiltInVariable);
+	}
+	else if (Variable)
+	{
+		VariableDefinition = Variable->GetVariableDefinition();
+	}
+	else if (InputSlotParameters.bIsPreBlended)
+	{
+		BuildContext.BuildLog.AddMessage(
+				EMessageSeverity::Error,
+				this,
+				LOCTEXT("PreBlendedInputSlotRequiresVariable",
+					"An input slot with pre-blend enabled must specify a variable (built-in or custom) "
+					"to blend with other input slots"));
+	}
+
+	if (VariableDefinition.IsValid())
+	{
+		VariableDefinition.bIsInput = true;
+
+		FCameraVariableTableAllocationInfo& VariableTableInfo = BuildContext.AllocationInfo.VariableTableInfo;
+		VariableTableInfo.VariableDefinitions.Add(VariableDefinition);
+
+		FCameraVariableDefinition TransientVariableDefinition = VariableDefinition.CreateVariant(TEXT("Transient"));
+		VariableTableInfo.VariableDefinitions.Add(TransientVariableDefinition);
+
+		VariableID = VariableDefinition.VariableID;
+		TransientVariableID = TransientVariableDefinition.VariableID;
+	}
+}
+
+FCameraNodeEvaluatorPtr UCameraRigInput1DSlot::OnBuildEvaluator(FCameraNodeEvaluatorBuilder& Builder) const
+{
+	return Builder.BuildEvaluator<UE::Cameras::FCameraRigInput1DSlotEvaluator>();
+}
+
+#undef LOCTEXT_NAMESPACE
+
diff --git a/Engine/Plugins/Cameras/GameplayCameras/Source/GameplayCameras/Private/Nodes/Input/CameraRigInput2DSlot.cpp b/Engine/Plugins/Cameras/GameplayCameras/Source/GameplayCameras/Private/Nodes/Input/CameraRigInput2DSlot.cpp
new file mode 100644
--- /dev/null
+++ b/Engine/Plugins/Cameras/GameplayCameras/Source/GameplayCameras/Private/Nodes/Input/CameraRigInput2DSlot.cpp
@@ -0,0 +1,143 @@
+// Copyright Epic Games, Inc. All Rights Reserved.
+
+#include "Nodes/Input/CameraRigInput2DSlot.h"
+
+#include "Core/CameraBuildLog.h"
+#include "Core/CameraRigAsset.h"
+#include "Core/CameraRigBuildContext.h"
+#include "Core/CameraVariableAssets.h"
+#include "Core/CameraVariableTable.h"
+#include "Nodes/Input/Input2DCameraNode.h"
+
+#define LOCTEXT_NAMESPACE "CameraRigInputSlot"
+
+namespace UE::Cameras
+{
+
+UE_DEFINE_CAMERA_NODE_EVALUATOR(FCameraRigInput2DSlotEvaluator)
+
+FCameraRigInput2DSlotEvaluator::FCameraRigInput2DSlotEvaluator()
+{
+	SetNodeEvaluatorFlags(ECameraNodeEvaluatorFlags::NeedsParameterUpdate | ECameraNodeEvaluatorFlags::NeedsEvaluationUpdate);
+}
+
+void FCameraRigInput2DSlotEvaluator::OnBuild(const FCameraNodeEvaluatorBuildParams& Params)
+{
+	const UCameraRigInput2DSlot* SlotNode = GetCameraNodeAs<UCameraRigInput2DSlot>();
+	ChildEvaluator = Params.BuildEvaluatorAs<FInput2DCameraNodeEvaluator>(SlotNode->Child);
+}
+
+void FCameraRigInput2DSlotEvaluator::OnInitialize(const FCameraNodeEvaluatorInitializeParams& Params)
+{
+	TransientInputValue = FVector2d::ZeroVector;
+	InputValue = FVector2d::ZeroVector;
+
+	const UCameraRigInput2DSlot* SlotNode = GetCameraNodeAs<UCameraRigInput2DSlot>();
+	if (SlotNode->GetVariableID().IsValid() && Params.LastActiveCameraRig.IsSet())
+	{
+		const FCameraRigEvaluationInfo& LastActiveRig = Params.LastActiveCameraRig.GetValue();
+		LastActiveRig.LastResult.VariableTable.TryGetValue<FVector2d>(SlotNode->GetVariableID(), InputValue);
+	}
+}
+
+FCameraNodeEvaluatorChildrenView FCameraRigInput2DSlotEvaluator::OnGetChildren()
+{
+	return FCameraNodeEvaluatorChildrenView{ ChildEvaluator };
+}
+
+void FCameraRigInput2DSlotEvaluator::OnUpdateParameters(const FCameraBlendedParameterUpdateParams& Params, FCameraBlendedParameterUpdateResult& OutResult)
+{
+	if (ChildEvaluator)
+	{
+		ChildEvaluator->UpdateParameters(Params, OutResult);
+
+		TransientInputValue = ChildEvaluator->GetInputValue();
+	}
+
+	const UCameraRigInput2DSlot* SlotNode = GetCameraNodeAs<UCameraRigInput2DSlot>();
+	if (SlotNode->InputSlotParameters.bIsPreBlended)
+	{
+		OutResult.VariableTable.SetValue<FVector2d>(SlotNode->GetTransientVariableID(), TransientInputValue);
+	}
+}
+
+void FCameraRigInput2DSlotEvaluator::OnRun(const FCameraNodeEvaluationParams& Params, FCameraNodeEvaluationResult& OutResult)
+{
+	const UCameraRigInput2DSlot* SlotNode = GetCameraNodeAs<UCameraRigInput2DSlot>();
+
+	if (SlotNode->InputSlotParameters.bIsPreBlended)
+	{
+		TransientInputValue = OutResult.VariableTable.GetValue<FVector2d>(SlotNode->GetTransientVariableID());
+	}
+
+	if (SlotNode->InputSlotParameters.bIsAccumulated)
+	{
+		InputValue += TransientInputValue;
+	}
+	else
+	{
+		InputValue = TransientInputValue;
+	}
+
+	InputValue.X = SlotNode->NormalizeX.NormalizeValue(InputValue.X);
+	InputValue.Y = SlotNode->NormalizeY.NormalizeValue(InputValue.Y);
+
+	InputValue.X = SlotNode->ClampX.ClampValue(InputValue.X);
+	InputValue.Y = SlotNode->ClampY.ClampValue(InputValue.Y);
+
+	OutResult.VariableTable.SetValue<FVector2d>(SlotNode->GetVariableID(), InputValue);
+}
+
+}  // namespace UE::Cameras
+
+FCameraNodeChildrenView UCameraRigInput2DSlot::OnGetChildren()
+{
+	return FCameraNodeChildrenView{ Child };
+}
+
+void UCameraRigInput2DSlot::OnBuild(FCameraRigBuildContext& BuildContext)
+{
+	using namespace UE::Cameras;
+
+	FCameraVariableDefinition VariableDefinition;
+
+	if (BuiltInVariable != EBuiltInVector2dCameraVariable::None)
+	{
+		VariableDefinition = FBuiltInCameraVariables::Get().GetDefinition(BuiltInVariable);
+	}
+	else if (Variable)
+	{
+		VariableDefinition = Variable->GetVariableDefinition();
+	}
+	else if (InputSlotParameters.bIsPreBlended)
+	{
+		BuildContext.BuildLog.AddMessage(
+				EMessageSeverity::Error,
+				this,
+				LOCTEXT("PreBlendedInputSlotRequiresVariable",
+					"An input slot with pre-blend enabled must specify a variable (built-in or custom) "
+					"to blend with other input slots"));
+	}
+
+	if (VariableDefinition.IsValid())
+	{
+		VariableDefinition.bIsInput = true;
+
+		FCameraVariableTableAllocationInfo& VariableTableInfo = BuildContext.AllocationInfo.VariableTableInfo;
+		VariableTableInfo.VariableDefinitions.Add(VariableDefinition);
+
+		FCameraVariableDefinition TransientVariableDefinition = VariableDefinition.CreateVariant(TEXT("Transient"));
+		VariableTableInfo.VariableDefinitions.Add(TransientVariableDefinition);
+
+		VariableID = VariableDefinition.VariableID;
+		TransientVariableID = TransientVariableDefinition.VariableID;
+	}
+}
+
+FCameraNodeEvaluatorPtr UCameraRigInput2DSlot::OnBuildEvaluator(FCameraNodeEvaluatorBuilder& Builder) const
+{
+	return Builder.BuildEvaluator<UE::Cameras::FCameraRigInput2DSlotEvaluator>();
+}
+
+#undef LOCTEXT_NAMESPACE
+
diff --git a/Engine/Plugins/Cameras/GameplayCameras/Source/GameplayCameras/Private/Nodes/Input/CameraRigInputSlotTypes.cpp b/Engine/Plugins/Cameras/GameplayCameras/Source/GameplayCameras/Private/Nodes/Input/CameraRigInputSlotTypes.cpp
new file mode 100644
--- /dev/null
+++ b/Engine/Plugins/Cameras/GameplayCameras/Source/GameplayCameras/Private/Nodes/Input/CameraRigInputSlotTypes.cpp
@@ -0,0 +1,29 @@
+// Copyright Epic Games, Inc. All Rights Reserved.
+
+#include "Nodes/Input/CameraRigInputSlotTypes.h"
+
+double FCameraParameterClamping::ClampValue(double Value) const
+{
+	if (bClampMin && Value < MinValue)
+	{
+		Value = MinValue;
+	}
+	if (bClampMax && Value > MaxValue)
+	{
+		Value = MaxValue;
+	}
+	return Value;
+}
+
+double FCameraParameterNormalization::NormalizeValue(double Value) const
+{
+	if (bNormalize && MaxValue > 0)
+	{
+		while (Value > MaxValue)
+		{
+			Value -= MaxValue;
+		}
+	}
+	return Value;
+}
+
diff --git a/Engine/Plugins/Cameras/GameplayCameras/Source/GameplayCameras/Private/Nodes/Input/Input1DCameraNode.cpp b/Engine/Plugins/Cameras/GameplayCameras/Source/GameplayCameras/Private/Nodes/Input/Input1DCameraNode.cpp
new file mode 100644
--- /dev/null
+++ b/Engine/Plugins/Cameras/GameplayCameras/Source/GameplayCameras/Private/Nodes/Input/Input1DCameraNode.cpp
@@ -0,0 +1,19 @@
+// Copyright Epic Games, Inc. All Rights Reserved.
+
+#include "Nodes/Input/Input1DCameraNode.h"
+
+#include UE_INLINE_GENERATED_CPP_BY_NAME(Input1DCameraNode)
+
+namespace UE::Cameras
+{
+
+UE_DEFINE_CAMERA_NODE_EVALUATOR(FInput1DCameraNodeEvaluator)
+
+FInput1DCameraNodeEvaluator::FInput1DCameraNodeEvaluator()
+{
+	// Only run during parameter update.
+	SetNodeEvaluatorFlags(ECameraNodeEvaluatorFlags::NeedsParameterUpdate);
+}
+
+}  // namespace UE::Cameras
+
diff --git a/Engine/Plugins/Cameras/GameplayCameras/Source/GameplayCameras/Private/Nodes/Input/Input2DCameraNode.cpp b/Engine/Plugins/Cameras/GameplayCameras/Source/GameplayCameras/Private/Nodes/Input/Input2DCameraNode.cpp
new file mode 100644
--- /dev/null
+++ b/Engine/Plugins/Cameras/GameplayCameras/Source/GameplayCameras/Private/Nodes/Input/Input2DCameraNode.cpp
@@ -0,0 +1,19 @@
+// Copyright Epic Games, Inc. All Rights Reserved.
+
+#include "Nodes/Input/Input2DCameraNode.h"
+
+#include UE_INLINE_GENERATED_CPP_BY_NAME(Input2DCameraNode)
+
+namespace UE::Cameras
+{
+
+UE_DEFINE_CAMERA_NODE_EVALUATOR(FInput2DCameraNodeEvaluator)
+
+FInput2DCameraNodeEvaluator::FInput2DCameraNodeEvaluator()
+{
+	// Only run during parameter update.
+	SetNodeEvaluatorFlags(ECameraNodeEvaluatorFlags::NeedsParameterUpdate);
+}
+
+}  // namespace UE::Cameras
+
diff --git a/Engine/Plugins/Cameras/GameplayCameras/Source/GameplayCameras/Private/Nodes/Input/InputAxisBinding2DCameraNode.cpp b/Engine/Plugins/Cameras/GameplayCameras/Source/GameplayCameras/Private/Nodes/Input/InputAxisBinding2DCameraNode.cpp
new file mode 100644
--- /dev/null
+++ b/Engine/Plugins/Cameras/GameplayCameras/Source/GameplayCameras/Private/Nodes/Input/InputAxisBinding2DCameraNode.cpp
@@ -0,0 +1,83 @@
+// Copyright Epic Games, Inc. All Rights Reserved.
+
+#include "Nodes/Input/InputAxisBinding2DCameraNode.h"
+
+#include "Components/InputComponent.h"
+#include "Core/CameraEvaluationContext.h"
+#include "EnhancedInputComponent.h"
+#include "GameFramework/Actor.h"
+#include "InputAction.h"
+
+namespace UE::Cameras
+{
+
+class FInputAxisBinding2DCameraNodeEvaluator : public FInput2DCameraNodeEvaluator
+{
+	UE_DECLARE_CAMERA_NODE_EVALUATOR_EX(GAMEPLAYCAMERAS_API, FInputAxisBinding2DCameraNodeEvaluator, FInput2DCameraNodeEvaluator)
+
+protected:
+
+	virtual void OnInitialize(const FCameraNodeEvaluatorInitializeParams& Params) override;
+	virtual void OnUpdateParameters(const FCameraBlendedParameterUpdateParams& Params, FCameraBlendedParameterUpdateResult& OutResult) override;
+
+private:
+
+	TObjectPtr<UEnhancedInputComponent> InputComponent;
+	FEnhancedInputActionValueBinding* AxisValueBinding = nullptr;
+};
+
+UE_DEFINE_CAMERA_NODE_EVALUATOR(FInputAxisBinding2DCameraNodeEvaluator)
+
+void FInputAxisBinding2DCameraNodeEvaluator::OnInitialize(const FCameraNodeEvaluatorInitializeParams& Params)
+{
+	if (UObject* ContextOwner = Params.EvaluationContext->GetOwner())
+	{
+		if (AActor* OuterActor = ContextOwner->GetTypedOuter<AActor>())
+		{
+			InputComponent = Cast<UEnhancedInputComponent>(OuterActor->InputComponent);
+		}
+	}
+
+	const UInputAxisBinding2DCameraNode* AxisBindingNode = GetCameraNodeAs<UInputAxisBinding2DCameraNode>();
+	if (InputComponent && AxisBindingNode->AxisAction)
+	{
+		AxisValueBinding = &InputComponent->BindActionValue(AxisBindingNode->AxisAction);
+	}
+}
+
+void FInputAxisBinding2DCameraNodeEvaluator::OnUpdateParameters(const FCameraBlendedParameterUpdateParams& Params, FCameraBlendedParameterUpdateResult& OutResult)
+{
+	const UInputAxisBinding2DCameraNode* AxisBindingNode = GetCameraNodeAs<UInputAxisBinding2DCameraNode>();
+	if (AxisValueBinding)
+	{
+		const FVector2d Value = AxisValueBinding->GetValue().Get<FVector2D>();
+
+		InputValue = FVector2d(
+				Value.X * AxisBindingNode->Multiplier.X, 
+				Value.Y * AxisBindingNode->Multiplier.Y);
+
+		if (AxisBindingNode->RevertAxisX)
+		{
+			InputValue.X = -InputValue.X;
+		}
+		if (AxisBindingNode->RevertAxisY)
+		{
+			InputValue.Y = -InputValue.Y;
+		}
+	}
+}
+
+}  // namespace UE::Cameras
+
+UInputAxisBinding2DCameraNode::UInputAxisBinding2DCameraNode(const FObjectInitializer& ObjInit)
+	: Super(ObjInit)
+{
+	Multiplier = FVector2D(1, 1);
+}
+
+FCameraNodeEvaluatorPtr UInputAxisBinding2DCameraNode::OnBuildEvaluator(FCameraNodeEvaluatorBuilder& Builder) const
+{
+	using namespace UE::Cameras;
+	return Builder.BuildEvaluator<FInputAxisBinding2DCameraNodeEvaluator>();
+}
+
diff --git a/Engine/Plugins/Cameras/GameplayCameras/Source/GameplayCameras/Public/Core/BlendCameraNode.h b/Engine/Plugins/Cameras/GameplayCameras/Source/GameplayCameras/Public/Core/BlendCameraNode.h
--- a/Engine/Plugins/Cameras/GameplayCameras/Source/GameplayCameras/Public/Core/BlendCameraNode.h
+++ b/Engine/Plugins/Cameras/GameplayCameras/Source/GameplayCameras/Public/Core/BlendCameraNode.h
@@ -19,6 +19,41 @@ class UBlendCameraNode : public UCameraNode
 namespace UE::Cameras
 {
 
+/**
+ * Parameter struct for blending camera node parameters.
+ */
+struct FCameraNodePreBlendParams
+{
+	FCameraNodePreBlendParams(
+			const FCameraNodeEvaluationParams& InEvaluationParams,
+			const FCameraPose& InLastCameraPose,
+			const FCameraVariableTable& InChildVariableTable)
+		: EvaluationParams(InEvaluationParams)
+		, LastCameraPose(InLastCameraPose)
+		, ChildVariableTable(InChildVariableTable)
+	{}
+
+	/** The parameters for the evaluation that will happen afterwards. */
+	const FCameraNodeEvaluationParams& EvaluationParams;
+	/** Last frame's camera pose. */
+	const FCameraPose& LastCameraPose;
+	/** The variable table of the node tree being blended. */
+	const FCameraVariableTable& ChildVariableTable;
+};
+
+/**
+ * Result struct for blending camera node parameters.
+ */
+struct FCameraNodePreBlendResult
+{
+	FCameraNodePreBlendResult(FCameraVariableTable& InVariableTable)
+		: VariableTable(InVariableTable)
+	{}
+
+	/** The variable table to received blended parameters. */
+	FCameraVariableTable& VariableTable;
+};
+
 /**
  * Parameter struct for blending camera node tree results.
  */
@@ -65,11 +100,17 @@ class FBlendCameraNodeEvaluator : public FCameraNodeEvaluator
 
 public:
 
+	/** Blend the parameters produced by a camera node tree over another set of values. */
+	void BlendParameters(const FCameraNodePreBlendParams& Params, FCameraNodePreBlendResult& OutResult);
+
 	/** Blend the result of a camera node tree over another result. */
 	void BlendResults(const FCameraNodeBlendParams& Params, FCameraNodeBlendResult& OutResult);
 
 protected:
 
+	/** Blend the parameters produced by a camera node tree over another set of values. */
+	virtual void OnBlendParameters(const FCameraNodePreBlendParams& Params, FCameraNodePreBlendResult& OutResult) {}
+
 	/** Blend the result of a camera node tree over another result. */
 	virtual void OnBlendResults(const FCameraNodeBlendParams& Params, FCameraNodeBlendResult& OutResult) {}
 };
diff --git a/Engine/Plugins/Cameras/GameplayCameras/Source/GameplayCameras/Public/Core/BlendStackCameraNode.h b/Engine/Plugins/Cameras/GameplayCameras/Source/GameplayCameras/Public/Core/BlendStackCameraNode.h
--- a/Engine/Plugins/Cameras/GameplayCameras/Source/GameplayCameras/Public/Core/BlendStackCameraNode.h
+++ b/Engine/Plugins/Cameras/GameplayCameras/Source/GameplayCameras/Public/Core/BlendStackCameraNode.h
@@ -6,6 +6,7 @@
 #include "Core/CameraNodeEvaluator.h"
 #include "Core/CameraNodeEvaluatorStorage.h"
 #include "Core/CameraRigAsset.h"
+#include "Core/CameraRigEvaluationInfo.h"
 #include "Debug/CameraDebugBlock.h"
 #include "IGameplayCamerasLiveEditListener.h"
 
@@ -94,6 +95,9 @@ public:
 	/** Push a new camera rig onto the blend stack. */
 	void Push(const FBlendStackCameraPushParams& Params);
 
+	/** Returns information about the top (active) camera rig, if any. */
+	TOptional<FCameraRigEvaluationInfo> GetActiveCameraRigEvaluationInfo() const;
+
 #if UE_GAMEPLAY_CAMERAS_DEBUG
 	FBlendStackCameraDebugBlock* BuildDetailedDebugBlock(const FCameraDebugBlockBuildParams& Params, FCameraDebugBlockBuilder& Builder);
 #endif  // UE_GAMEPLAY_CAMERAS_DEBUG
@@ -117,6 +121,8 @@ protected:
 
 protected:
 
+	struct FCameraRigEntry;
+
 	// Utility functions for finding an appropriate transition.
 	const UCameraRigTransition* FindTransition(const FBlendStackCameraPushParams& Params) const;
 	const UCameraRigTransition* FindTransition(
@@ -126,6 +132,15 @@ protected:
 
 	void PopEntries(int32 FirstIndexToKeep);
 
+	bool InitializeEntry(
+		FCameraRigEntry& NewEntry, 
+		const UCameraRigAsset* CameraRig,
+		FCameraSystemEvaluator* Evaluator,
+		TSharedPtr<const FCameraEvaluationContext> EvaluationContext,
+		UBlendStackRootCameraNode* EntryRootNode);
+
+	void GatherEntryParameterEvaluators(FCameraNodeEvaluator* RootEvaluator, TArray<FCameraNodeEvaluator*>& OutParameterEvaluators);
+
 protected:
 
 	struct FCameraRigEntry
@@ -140,10 +155,16 @@ protected:
 		FCameraNodeEvaluatorStorage EvaluatorStorage;
 		/** Root evaluator. */
 		FBlendStackRootCameraNodeEvaluator* RootEvaluator = nullptr;
+		/** Evaluators needing parameter update. */
+		TArray<FCameraNodeEvaluator*> ParameterEvaluators;
 		/** Result for this node tree. */
 		FCameraNodeEvaluationResult Result;
 		/** Whether this is the first frame this entry runs. */
 		bool bIsFirstFrame = false;
+		/** Whether input slots were run (possibly from a preview update). */
+		bool bInputRunThisFrame = false;
+		/** Whether the blend node was run (possibly from a preview update). */
+		bool bBlendRunThisFrame = false;
 		/** Whether this entry is frozen. */
 		bool bIsFrozen = false;
 
diff --git a/Engine/Plugins/Cameras/GameplayCameras/Source/GameplayCameras/Public/Core/BuiltInCameraVariables.h b/Engine/Plugins/Cameras/GameplayCameras/Source/GameplayCameras/Public/Core/BuiltInCameraVariables.h
new file mode 100644
--- /dev/null
+++ b/Engine/Plugins/Cameras/GameplayCameras/Source/GameplayCameras/Public/Core/BuiltInCameraVariables.h
@@ -0,0 +1,69 @@
+// Copyright Epic Games, Inc. All Rights Reserved.
+
+#pragma once
+
+#include "Core/CameraVariableTableFwd.h"
+#include "CoreTypes.h"
+
+#include "BuiltInCameraVariables.generated.h"
+
+class UDoubleCameraVariable;
+class UVector2dCameraVariable;
+
+/**
+ * Built-in floating point camera variables.
+ */
+UENUM()
+enum class EBuiltInDoubleCameraVariable
+{
+	None,
+	Yaw,
+	Pitch,
+	Roll,
+	Zoom
+};
+
+/**
+ * Built-in 2-dimensional camera variables.
+ */
+UENUM()
+enum class EBuiltInVector2dCameraVariable
+{
+	None,
+	YawPitch
+};
+
+namespace UE::Cameras
+{
+
+/**
+ * A utility class that provides definitions for "well-known" camera variables.
+ */
+class FBuiltInCameraVariables
+{
+public:
+
+	/** Singleton access. */
+	static const FBuiltInCameraVariables& Get();
+
+	/** Get the definition of a built-in camera variable. */
+	const FCameraVariableDefinition& GetDefinition(EBuiltInDoubleCameraVariable BuiltInVariable) const;
+	/** Get the definition of a built-in camera variable. */
+	const FCameraVariableDefinition& GetDefinition(EBuiltInVector2dCameraVariable BuiltInVariable) const;
+
+public:
+
+	FCameraVariableDefinition YawDefinition;
+	FCameraVariableDefinition PitchDefinition;
+	FCameraVariableDefinition RollDefinition;
+	FCameraVariableDefinition ZoomDefinition;
+
+	FCameraVariableDefinition YawPitchDefinition;
+
+protected:
+
+	FBuiltInCameraVariables();
+};
+
+}  // namespace UE::Cameras
+
diff --git a/Engine/Plugins/Cameras/GameplayCameras/Source/GameplayCameras/Public/Core/CameraNodeEvaluator.h b/Engine/Plugins/Cameras/GameplayCameras/Source/GameplayCameras/Public/Core/CameraNodeEvaluator.h
--- a/Engine/Plugins/Cameras/GameplayCameras/Source/GameplayCameras/Public/Core/CameraNodeEvaluator.h
+++ b/Engine/Plugins/Cameras/GameplayCameras/Source/GameplayCameras/Public/Core/CameraNodeEvaluator.h
@@ -2,6 +2,7 @@
 
 #pragma once
 
+#include "Core/CameraRigEvaluationInfo.h"
 #include "Core/CameraObjectRtti.h"
 #include "Core/CameraPose.h"
 #include "Core/CameraVariableTable.h"
@@ -9,10 +10,12 @@
 #include "CoreTypes.h"
 #include "Debug/RootCameraDebugBlock.h"
 #include "GameplayCameras.h"
+#include "Misc/Optional.h"
 #include "UObject/ObjectPtr.h"
 
 class FReferenceCollector;
 class UCameraNode;
+class UCameraRigAsset;
 
 namespace UE::Cameras
 {
@@ -20,23 +23,66 @@ namespace UE::Cameras
 class FCameraEvaluationContext;
 class FCameraNodeEvaluator;
 class FCameraSystemEvaluator;
+struct FCameraNodeEvaluationParams;
 struct FCameraNodeEvaluatorBuilder;
+struct FCameraRigEvaluationInfo;
 
 #if UE_GAMEPLAY_CAMERAS_DEBUG
 struct FCameraDebugBlockBuilder;
 #endif  // UE_GAMEPLAY_CAMERAS_DEBUG
 
+/**
+ * Flags describing the needs of a camera node evaluator.
+ */
+enum class ECameraNodeEvaluatorFlags
+{
+	None = 0,
+	NeedsParameterUpdate = 1 << 0,
+	NeedsEvaluationUpdate = 1 << 1
+};
+ENUM_CLASS_FLAGS(ECameraNodeEvaluatorFlags)
+
+/**
+ * Parameter structure for updating the pre-blended parameters of a camera node.
+ */
+struct FCameraBlendedParameterUpdateParams
+{
+	FCameraBlendedParameterUpdateParams(const FCameraNodeEvaluationParams& InEvaluationParams, const FCameraPose& InLastCameraPose)
+		: EvaluationParams(InEvaluationParams)
+		, LastCameraPose(InLastCameraPose)
+	{}
+
+	/** Information about the evaluation pass that will happen afterwards. */
+	const FCameraNodeEvaluationParams& EvaluationParams;
+	/** Last frame's camera pose. */
+	const FCameraPose& LastCameraPose;
+};
+
+/**
+ * Result of updating the pre-blended parameters of a camera node.
+ */
+struct FCameraBlendedParameterUpdateResult
+{
+	FCameraBlendedParameterUpdateResult(FCameraVariableTable& InVariableTable)
+		: VariableTable(InVariableTable)
+	{}
+
+	/** Variable table in which parameters should be stored or obtained. */
+	FCameraVariableTable& VariableTable;
+};
+
 /**
  * Parameter structure for running a camera node evaluator.
  */
 struct FCameraNodeEvaluationParams
 {
-	/** The evaluation running this evaluation. */
+	/** The evaluator running this evaluation. */
 	FCameraSystemEvaluator* Evaluator = nullptr;
 	/** The evaluation context (if any) responsible for this branch of the evaluation. */
 	TSharedPtr<const FCameraEvaluationContext> EvaluationContext;
 	/** The time interval for the evaluation. */
 	float DeltaTime = 0.f;
+	
 	/** Whether this is the first evaluation of this camera node hierarchy. */
 	bool bIsFirstFrame = false;
 };
@@ -92,6 +138,12 @@ struct FCameraNodeEvaluatorInitializeParams
 	FCameraSystemEvaluator* Evaluator = nullptr;
 	/** The evaluation context (if any) responsible for this branch of the evaluation. */
 	TSharedPtr<const FCameraEvaluationContext> EvaluationContext;
+
+	/**
+	 * Information about the last active camera rig if the node tree being initialized
+	 * is being pushed on top of a non-empty blend stack.
+	 */
+	TOptional<FCameraRigEvaluationInfo> LastActiveCameraRig;
 };
 
 /** View on a camera node evaluator's children. */
@@ -130,12 +182,21 @@ public:
 	/** Get the list of children under this evaluator. */
 	FCameraNodeEvaluatorChildrenView GetChildren();
 
+	/** Called to update and store the blended parameters for this node. */
+	void UpdateParameters(const FCameraBlendedParameterUpdateParams& Params, FCameraBlendedParameterUpdateResult& OutResult);
+
 	/** Run this evaluator. */
 	void Run(const FCameraNodeEvaluationParams& Params, FCameraNodeEvaluationResult& OutResult);
 
 	/** Collect referenced UObjects. */
 	void AddReferencedObjects(FReferenceCollector& Collector);
 
+	/** Gets the flags for this evaluator. */
+	ECameraNodeEvaluatorFlags  GetNodeEvaluatorFlags() const { return PrivateFlags; }
+
+	/** Get the camera node. */
+	const UCameraNode* GetCameraNode() const { return PrivateCameraNode; }
+
 	/** Get the camera node. */
 	template<typename CameraNodeType>
 	const CameraNodeType* GetCameraNodeAs() const
@@ -153,6 +214,11 @@ public:
 	// Internal API.
 	void SetPrivateCameraNode(TObjectPtr<const UCameraNode> InCameraNode);
 
+protected:
+
+	/** Sets the flags for this evaluator. */
+	void SetNodeEvaluatorFlags(ECameraNodeEvaluatorFlags InFlags);
+
 protected:
 
 	/** Called to build any children evaluators. */
@@ -162,7 +228,10 @@ protected:
 	GAMEPLAYCAMERAS_API virtual void OnInitialize(const FCameraNodeEvaluatorInitializeParams& Params) {}
 
 	/** Get the list of children under this evaluator. */
-	virtual FCameraNodeEvaluatorChildrenView OnGetChildren() { return FCameraNodeEvaluatorChildrenView(); }
+	GAMEPLAYCAMERAS_API virtual FCameraNodeEvaluatorChildrenView OnGetChildren() { return FCameraNodeEvaluatorChildrenView(); }
+
+	/** Called to update and store the blended parameters for this node. */
+	GAMEPLAYCAMERAS_API virtual void OnUpdateParameters(const FCameraBlendedParameterUpdateParams& Params, FCameraBlendedParameterUpdateResult& OutResult) {}
 
 	/** Run this evaluator. */
 	GAMEPLAYCAMERAS_API virtual void OnRun(const FCameraNodeEvaluationParams& Params, FCameraNodeEvaluationResult& OutResult) {}
@@ -179,6 +248,9 @@ private:
 
 	/** The camera node to run. */
 	TObjectPtr<const UCameraNode> PrivateCameraNode;
+
+	/** The flags for this evaluator. */
+	ECameraNodeEvaluatorFlags PrivateFlags = ECameraNodeEvaluatorFlags::NeedsEvaluationUpdate;
 };
 
 /** Utility base class for camera node evaluators of a specific camera node type. */
@@ -199,9 +271,11 @@ public:
 template<typename EvaluatorType>
 EvaluatorType* FCameraNodeEvaluatorBuildParams::BuildEvaluatorAs(const UCameraNode* InNode) const
 {
-	FCameraNodeEvaluator* NewEvaluator = BuildEvaluator(InNode);
-	check(NewEvaluator);
-	return NewEvaluator->CastThisChecked<EvaluatorType>();
+	if (FCameraNodeEvaluator* NewEvaluator = BuildEvaluator(InNode))
+	{
+		return NewEvaluator->CastThisChecked<EvaluatorType>();
+	}
+	return nullptr;
 }
 
 }  // namespace UE::Cameras
diff --git a/Engine/Plugins/Cameras/GameplayCameras/Source/GameplayCameras/Public/Core/CameraNodeEvaluatorStorage.h b/Engine/Plugins/Cameras/GameplayCameras/Source/GameplayCameras/Public/Core/CameraNodeEvaluatorStorage.h
--- a/Engine/Plugins/Cameras/GameplayCameras/Source/GameplayCameras/Public/Core/CameraNodeEvaluatorStorage.h
+++ b/Engine/Plugins/Cameras/GameplayCameras/Source/GameplayCameras/Public/Core/CameraNodeEvaluatorStorage.h
@@ -25,15 +25,8 @@ struct FCameraNodeEvaluatorTreeBuildParams
 {
 	/** The root node of the tree. */
 	TObjectPtr<const UCameraNode> RootCameraNode;
-	/** The evaluator running this evaluation. */
-	FCameraSystemEvaluator* Evaluator = nullptr;
-	/** The evaluation context (if any) responsible for this branch of the evaluation. */
-	TSharedPtr<const FCameraEvaluationContext> EvaluationContext;
 	/** An optional allocation information to optimize storage. */
-	FCameraNodeEvaluatorAllocationInfo* AllocationInfo = nullptr;
-
-	/** Whether to initialize the tree. */
-	bool bInitialize = true;
+	const FCameraNodeEvaluatorAllocationInfo* AllocationInfo = nullptr;
 };
 
 /**
diff --git a/Engine/Plugins/Cameras/GameplayCameras/Source/GameplayCameras/Public/Core/CameraObjectStorage.h b/Engine/Plugins/Cameras/GameplayCameras/Source/GameplayCameras/Public/Core/CameraObjectStorage.h
--- a/Engine/Plugins/Cameras/GameplayCameras/Source/GameplayCameras/Public/Core/CameraObjectStorage.h
+++ b/Engine/Plugins/Cameras/GameplayCameras/Source/GameplayCameras/Public/Core/CameraObjectStorage.h
@@ -66,7 +66,7 @@ protected:
 	void* BuildObjectUninitialized(uint32 Sizeof, uint32 Alignof);
 
 	/**
-	 * Called to registered an object after it has been constructed.
+	 * Called to register an object after it has been constructed.
 	 */
 	void RegisterInitializedObject(BaseObjectType* BaseObjectPtr);
 
@@ -257,6 +257,7 @@ void TCameraObjectStorage<BaseObjectType>::GetAllocationInfo(uint32& OutTotalUse
 template<typename BaseObjectType>
 void TCameraObjectStorage<BaseObjectType>::AllocatePage(uint32 InCapacity, uint32 InAlignment)
 {
+	check(InCapacity > 0 && InAlignment > 0);
 	FAllocation& NewAllocation = Allocations.Emplace_GetRef();
 	NewAllocation.Memory = reinterpret_cast<uint8*>(FMemory::Malloc(InCapacity, InAlignment));
 	NewAllocation.Alignment = InAlignment;
diff --git a/Engine/Plugins/Cameras/GameplayCameras/Source/GameplayCameras/Public/Core/CameraRigEvaluationInfo.h b/Engine/Plugins/Cameras/GameplayCameras/Source/GameplayCameras/Public/Core/CameraRigEvaluationInfo.h
new file mode 100644
--- /dev/null
+++ b/Engine/Plugins/Cameras/GameplayCameras/Source/GameplayCameras/Public/Core/CameraRigEvaluationInfo.h
@@ -0,0 +1,46 @@
+// Copyright Epic Games, Inc. All Rights Reserved.
+
+#pragma once
+
+#include "Templates/SharedPointer.h"
+
+class UCameraRigAsset;
+
+namespace UE::Cameras
+{
+
+class FCameraEvaluationContext;
+class FCameraNodeEvaluator;
+struct FCameraNodeEvaluationResult;
+
+/**
+ * A structure describing an active camera rig being evaluated, generally
+ * inside a blend stack.
+ */
+struct FCameraRigEvaluationInfo
+{
+	/** The camera rig being evaluated. */
+	const UCameraRigAsset* CameraRig;
+	/** The context inside which the evaluation occurs. */
+	TSharedPtr<const FCameraEvaluationContext> EvaluationContext;
+	/** The last evaluated result for this camera rig. */
+	const FCameraNodeEvaluationResult& LastResult;
+	/** The root node evaluator of the camera rig. */
+	const FCameraNodeEvaluator* RootEvaluator;
+	/** Whether the camera rig is frozen. */
+	bool bIsFrozen = false;
+
+	FCameraRigEvaluationInfo(
+			const UCameraRigAsset* InCameraRig,
+			TSharedPtr<const FCameraEvaluationContext> InEvaluationContext,
+			const FCameraNodeEvaluationResult& InLastResult,
+			const FCameraNodeEvaluator* InRootEvaluator)
+		: CameraRig(InCameraRig)
+		, EvaluationContext(InEvaluationContext)
+		, LastResult(InLastResult)
+		, RootEvaluator(InRootEvaluator)
+	{}
+};
+
+}  // namespace UE::Cameras
+
diff --git a/Engine/Plugins/Cameras/GameplayCameras/Source/GameplayCameras/Public/Core/CameraVariableTable.h b/Engine/Plugins/Cameras/GameplayCameras/Source/GameplayCameras/Public/Core/CameraVariableTable.h
--- a/Engine/Plugins/Cameras/GameplayCameras/Source/GameplayCameras/Public/Core/CameraVariableTable.h
+++ b/Engine/Plugins/Cameras/GameplayCameras/Source/GameplayCameras/Public/Core/CameraVariableTable.h
@@ -20,6 +20,26 @@ struct TCameraVariableTraits;
 template<typename ValueType>
 struct TCameraVariableInterpolation;
 
+/**
+ * Filter for variable table operations.
+ */
+enum class ECameraVariableTableFilter
+{
+	None = 0,
+	/** Only include input variables. */
+	Input = 1 << 0,
+	/** Only include output variables (i.e. anything not an input). */
+	Output = 1 << 1,
+	/** Only include changed variables. */
+	ChangedOnly = 1 << 2,
+
+	/** All variables. */
+	All = Input | Output,
+	/** All changed variables. */
+	AllChanged = Input | Output | ChangedOnly
+};
+ENUM_CLASS_FLAGS(ECameraVariableTableFilter)
+
 /**
  * A structure that keeps track of which variables have been processed in a
  * camera variable table.
@@ -107,12 +127,12 @@ public:
 	// Interpolation.
 	
 	void OverrideAll(const FCameraVariableTable& OtherTable);
-	void OverrideChanged(const FCameraVariableTable& OtherTable);
-	void OverrideChanged(const FCameraVariableTable& OtherTable, const FCameraVariableTableFlags& InMask, bool bInvertMask, FCameraVariableTableFlags& OutMask);
+	void Override(const FCameraVariableTable& OtherTable, ECameraVariableTableFilter Filter);
+	void Override(const FCameraVariableTable& OtherTable, ECameraVariableTableFilter Filter, const FCameraVariableTableFlags& InMask, bool bInvertMask, FCameraVariableTableFlags& OutMask);
 
 	void LerpAll(const FCameraVariableTable& ToTable, float Factor);
-	void LerpChanged(const FCameraVariableTable& ToTable, float Factor);
-	void LerpChanged(const FCameraVariableTable& ToTable, float Factor, const FCameraVariableTableFlags& InMask, bool bInvertMask, FCameraVariableTableFlags& OutMask);
+	void Lerp(const FCameraVariableTable& ToTable, ECameraVariableTableFilter Filter, float Factor);
+	void Lerp(const FCameraVariableTable& ToTable, ECameraVariableTableFilter Filter, float Factor, const FCameraVariableTableFlags& InMask, bool bInvertMask, FCameraVariableTableFlags& OutMask);
 
 public:
 
@@ -137,8 +157,8 @@ private:
 
 	void ReallocateBuffer(uint32 MinRequired = 0);
 
-	void InternalOverride(const FCameraVariableTable& OtherTable, const FCameraVariableTableFlags* InMask, bool bInvertMask, FCameraVariableTableFlags* OutMask, bool bChangedOnly);
-	void InternalLerp(const FCameraVariableTable& ToTable, float Factor, const FCameraVariableTableFlags* InMask, bool bInvertMask, FCameraVariableTableFlags* OutMask, bool bChangedOnly);
+	void InternalOverride(const FCameraVariableTable& OtherTable, ECameraVariableTableFilter Filter, const FCameraVariableTableFlags* InMask, bool bInvertMask, FCameraVariableTableFlags* OutMask);
+	void InternalLerp(const FCameraVariableTable& ToTable, ECameraVariableTableFilter Filter, float Factor, const FCameraVariableTableFlags* InMask, bool bInvertMask, FCameraVariableTableFlags* OutMask);
 
 private:
 
@@ -146,8 +166,9 @@ private:
 	{
 		None = 0,
 		Private = 1 << 0,
-		Written = 1 << 1,
-		WrittenThisFrame = 1 << 2
+		Input = 1 << 1,
+		Written = 1 << 2,
+		WrittenThisFrame = 1 << 3
 	};
 	FRIEND_ENUM_CLASS_FLAGS(EEntryFlags)
 
@@ -195,20 +216,26 @@ const ValueType* FCameraVariableTable::FindValue(FCameraVariableID VariableID) c
 template<typename ValueType>
 const ValueType& FCameraVariableTable::GetValue(FCameraVariableID VariableID) const
 {
-	const FEntry& Entry = Entries.FindChecked(VariableID);
-	CheckVariableType<ValueType>(Entry.Type);
+	const FEntry* Entry = Entries.Find(VariableID);
+	if (ensureMsgf(Entry, TEXT("Can't get camera variable (ID '%d') because it doesn't exist in the table."), VariableID.GetValue()))
+	{
+		CheckVariableType<ValueType>(Entry->Type);
 #if WITH_EDITORONLY_DATA
-	checkf(
-			EnumHasAnyFlags(Entry.Flags, EEntryFlags::Written),
-			TEXT("Variable '%s' has never been written to. GetValue() will return uninitialized memory!"),
-			*Entry.DebugName);
+		checkf(
+				EnumHasAnyFlags(Entry->Flags, EEntryFlags::Written),
+				TEXT("Variable '%s' has never been written to. GetValue() will return uninitialized memory!"),
+				*Entry->DebugName);
 #else
-	checkf(
-			EnumHasAnyFlags(Entry.Flags, EEntryFlags::Written),
-			TEXT("Variable '%s' has never been written to. GetValue() will return uninitialized memory!"),
-			*LexToString(VariableID.GetValue()));
+		checkf(
+				EnumHasAnyFlags(Entry->Flags, EEntryFlags::Written),
+				TEXT("Variable '%s' has never been written to. GetValue() will return uninitialized memory!"),
+				*LexToString(VariableID.GetValue()));
 #endif
-	return *reinterpret_cast<ValueType*>(Memory + Entry.Offset);
+		return *reinterpret_cast<ValueType*>(Memory + Entry->Offset);
+	}
+
+	static ValueType DefaultValue = ValueType();
+	return DefaultValue;
 }
 
 template<typename ValueType>
@@ -235,11 +262,14 @@ bool FCameraVariableTable::TryGetValue(FCameraVariableID VariableID, ValueType&
 template<typename ValueType>
 void FCameraVariableTable::SetValue(FCameraVariableID VariableID, typename TCallTraits<ValueType>::ParamType Value)
 {
-	FEntry& Entry = Entries.FindChecked(VariableID);
-	CheckVariableType<ValueType>(Entry.Type);
-	ValueType* ValuePtr = reinterpret_cast<ValueType*>(Memory + Entry.Offset);
-	*ValuePtr = Value;
-	Entry.Flags |= EEntryFlags::Written | EEntryFlags::WrittenThisFrame;
+	FEntry* Entry = Entries.Find(VariableID);
+	if (ensureMsgf(Entry, TEXT("Can't set camera variable (ID '%d') because it doesn't exist in the table."), VariableID.GetValue()))
+	{
+		CheckVariableType<ValueType>(Entry->Type);
+		ValueType* ValuePtr = reinterpret_cast<ValueType*>(Memory + Entry->Offset);
+		*ValuePtr = Value;
+		Entry->Flags |= EEntryFlags::Written | EEntryFlags::WrittenThisFrame;
+	}
 }
 
 template<typename ValueType>
diff --git a/Engine/Plugins/Cameras/GameplayCameras/Source/GameplayCameras/Public/Core/CameraVariableTableFwd.h b/Engine/Plugins/Cameras/GameplayCameras/Source/GameplayCameras/Public/Core/CameraVariableTableFwd.h
--- a/Engine/Plugins/Cameras/GameplayCameras/Source/GameplayCameras/Public/Core/CameraVariableTableFwd.h
+++ b/Engine/Plugins/Cameras/GameplayCameras/Source/GameplayCameras/Public/Core/CameraVariableTableFwd.h
@@ -4,6 +4,7 @@
 
 #include "Containers/UnrealString.h"
 #include "CoreTypes.h"
+#include "Templates/TypeHash.h"
 
 #include "CameraVariableTableFwd.generated.h"
 
@@ -108,10 +109,32 @@ struct FCameraVariableDefinition
 	UPROPERTY()
 	bool bIsPrivate = false;
 
+	UPROPERTY()
+	bool bIsInput = false;
+
 #if WITH_EDITORONLY_DATA
 	UPROPERTY()
 	FString VariableName;
 #endif
+
+	bool IsValid() const
+	{
+		return VariableID.IsValid();
+	}
+
+	FCameraVariableDefinition CreateVariant(const FString& VariantID) const
+	{
+		FCameraVariableDefinition VariantDefinition(*this);
+		VariantDefinition.VariableID = FCameraVariableID::FromHashValue(
+				HashCombineFast(VariableID.GetValue(), GetTypeHash(VariantID)));
+#if WITH_EDITORONLY_DATA
+		if (!VariableName.IsEmpty())
+		{
+			VariantDefinition.VariableName += FString::Format(TEXT("_{0}Variant"), { VariantID });
+		}
+#endif
+		return VariantDefinition;
+	}
 };
 
 USTRUCT()
diff --git a/Engine/Plugins/Cameras/GameplayCameras/Source/GameplayCameras/Public/Nodes/Blends/SimpleBlendCameraNode.h b/Engine/Plugins/Cameras/GameplayCameras/Source/GameplayCameras/Public/Nodes/Blends/SimpleBlendCameraNode.h
--- a/Engine/Plugins/Cameras/GameplayCameras/Source/GameplayCameras/Public/Nodes/Blends/SimpleBlendCameraNode.h
+++ b/Engine/Plugins/Cameras/GameplayCameras/Source/GameplayCameras/Public/Nodes/Blends/SimpleBlendCameraNode.h
@@ -53,6 +53,7 @@ public:
 protected:
 
 	virtual void OnRun(const FCameraNodeEvaluationParams& Params, FCameraNodeEvaluationResult& OutResult) override;
+	virtual void OnBlendParameters(const FCameraNodePreBlendParams& Params, FCameraNodePreBlendResult& OutResult) override;
 	virtual void OnBlendResults(const FCameraNodeBlendParams& Params, FCameraNodeBlendResult& OutResult) override;
 
 	virtual void OnComputeBlendFactor(const FCameraNodeEvaluationParams& Params, FSimpleBlendCameraNodeEvaluationResult& OutResult) {}
diff --git a/Engine/Plugins/Cameras/GameplayCameras/Source/GameplayCameras/Public/Nodes/Common/BoomArmCameraNode.h b/Engine/Plugins/Cameras/GameplayCameras/Source/GameplayCameras/Public/Nodes/Common/BoomArmCameraNode.h
--- a/Engine/Plugins/Cameras/GameplayCameras/Source/GameplayCameras/Public/Nodes/Common/BoomArmCameraNode.h
+++ b/Engine/Plugins/Cameras/GameplayCameras/Source/GameplayCameras/Public/Nodes/Common/BoomArmCameraNode.h
@@ -3,10 +3,12 @@
 #pragma once
 
 #include "Core/CameraNode.h"
-#include "Nodes/CameraNodeTypes.h"
+#include "Core/CameraParameters.h"
 
 #include "BoomArmCameraNode.generated.h"
 
+class UCameraRigInput2DSlot;
+
 /**
  * A camera node that can rotate the camera in yaw and pitch based on player input.
  */
@@ -18,11 +20,20 @@ class UBoomArmCameraNode : public UCameraNode
 protected:
 
 	// UCameraNode interface.
+	virtual FCameraNodeChildrenView OnGetChildren() override;
 	virtual FCameraNodeEvaluatorPtr OnBuildEvaluator(FCameraNodeEvaluatorBuilder& Builder) const override;
 
 public:
 
+	/** The offset of the boom. Rotation occurs at the base (i.e. before the offset). */
 	UPROPERTY(EditAnywhere, Category=Common)
-	FVector3d BoomOffset;
+	FVector3dCameraParameter BoomOffset;
+
+	/**
+	 * The input slot for controlling the boom arm.
+	 * If no input slot is specified, the boom arm will use the player controller view rotation.
+	 */
+	UPROPERTY()
+	TObjectPtr<UCameraRigInput2DSlot> InputSlot;
 };
 
diff --git a/Engine/Plugins/Cameras/GameplayCameras/Source/GameplayCameras/Public/Nodes/Input/CameraRigInput1DSlot.h b/Engine/Plugins/Cameras/GameplayCameras/Source/GameplayCameras/Public/Nodes/Input/CameraRigInput1DSlot.h
new file mode 100644
--- /dev/null
+++ b/Engine/Plugins/Cameras/GameplayCameras/Source/GameplayCameras/Public/Nodes/Input/CameraRigInput1DSlot.h
@@ -0,0 +1,103 @@
+// Copyright Epic Games, Inc. All Rights Reserved.
+
+#pragma once
+
+#include "Core/BuiltInCameraVariables.h"
+#include "Core/CameraNode.h"
+#include "Core/CameraNodeEvaluator.h"
+#include "Nodes/Input/CameraRigInputSlotTypes.h"
+#include "UObject/ObjectPtr.h"
+
+#include "CameraRigInput1DSlot.generated.h"
+
+class UDoubleCameraVariable;
+class UInput1DCameraNode;
+
+/**
+ * A node that can handle and accumulate a chain of player input nodes.
+ */
+UCLASS(MinimalAPI, meta=(CameraNodeCategories="Input"))
+class UCameraRigInput1DSlot : public UCameraNode
+{
+	GENERATED_BODY()
+
+public:
+
+	/** Input processing parameters. */
+	UPROPERTY(EditAnywhere, Category="Input", meta=(ShowOnlyInnerProperties))
+	FCameraRigInputSlotParameters InputSlotParameters;
+
+	/** Clamping of the final input value. */
+	UPROPERTY(EditAnywhere, Category="Input")
+	FCameraParameterClamping Clamp;
+
+	/** Normalization of the final input value. */
+	UPROPERTY(EditAnywhere, Category="Input")
+	FCameraParameterNormalization Normalize;
+
+	/** The variable to use to blend with other input slots. */
+	UPROPERTY(EditAnywhere, Category="Input")
+	EBuiltInDoubleCameraVariable BuiltInVariable = EBuiltInDoubleCameraVariable::Yaw;
+
+	/** The variable to use to blend with other input slots. */
+	UPROPERTY(EditAnywhere, Category="Input", meta=(EditCondition="BuiltInVariable == EBuiltInDoubleCameraVariable::None"))
+	TObjectPtr<UDoubleCameraVariable> Variable;
+
+	/** A node providing an incremental value. */
+	UPROPERTY()
+	TObjectPtr<UInput1DCameraNode> Child;
+
+public:
+
+	FCameraVariableID GetVariableID() const { return VariableID; }
+	FCameraVariableID GetTransientVariableID() const { return TransientVariableID; }
+
+protected:
+
+	// UCameraNode interface.
+	virtual FCameraNodeChildrenView OnGetChildren() override;
+	virtual void OnBuild(FCameraRigBuildContext& BuildContext) override;
+	virtual FCameraNodeEvaluatorPtr OnBuildEvaluator(FCameraNodeEvaluatorBuilder& Builder) const override;
+
+private:
+
+	UPROPERTY()
+	FCameraVariableID TransientVariableID;
+	UPROPERTY()
+	FCameraVariableID VariableID;
+};
+
+namespace UE::Cameras
+{
+
+class FInput1DCameraNodeEvaluator;
+
+class FCameraRigInput1DSlotEvaluator : public FCameraNodeEvaluator
+{
+	UE_DECLARE_CAMERA_NODE_EVALUATOR(GAMEPLAYCAMERAS_API, FCameraRigInput1DSlotEvaluator)
+
+public:
+
+	FCameraRigInput1DSlotEvaluator();
+
+	double GetInputValue() const { return InputValue; }
+
+protected:
+
+	// FCameraNodeEvaluator interface.
+	virtual void OnBuild(const FCameraNodeEvaluatorBuildParams& Params) override;
+	virtual void OnInitialize(const FCameraNodeEvaluatorInitializeParams& Params) override;
+	virtual FCameraNodeEvaluatorChildrenView OnGetChildren() override;
+	virtual void OnUpdateParameters(const FCameraBlendedParameterUpdateParams& Params, FCameraBlendedParameterUpdateResult& OutResult) override;
+	virtual void OnRun(const FCameraNodeEvaluationParams& Params, FCameraNodeEvaluationResult& OutResult) override;
+
+protected:
+
+	FInput1DCameraNodeEvaluator* ChildEvaluator;
+
+	double TransientInputValue = 0.f;
+	double InputValue = 0.f;
+};
+
+}  // namespace UE::Cameras
+
diff --git a/Engine/Plugins/Cameras/GameplayCameras/Source/GameplayCameras/Public/Nodes/Input/CameraRigInput2DSlot.h b/Engine/Plugins/Cameras/GameplayCameras/Source/GameplayCameras/Public/Nodes/Input/CameraRigInput2DSlot.h
new file mode 100644
--- /dev/null
+++ b/Engine/Plugins/Cameras/GameplayCameras/Source/GameplayCameras/Public/Nodes/Input/CameraRigInput2DSlot.h
@@ -0,0 +1,112 @@
+// Copyright Epic Games, Inc. All Rights Reserved.
+
+#pragma once
+
+#include "Core/BuiltInCameraVariables.h"
+#include "Core/CameraNode.h"
+#include "Core/CameraNodeEvaluator.h"
+#include "Nodes/Input/CameraRigInputSlotTypes.h"
+#include "UObject/ObjectPtr.h"
+
+#include "CameraRigInput2DSlot.generated.h"
+
+class UInput2DCameraNode;
+class UVector2dCameraVariable;
+
+/**
+ * A node that can handle and accumulate a chain of player input nodes.
+ */
+UCLASS(MinimalAPI, meta=(CameraNodeCategories="Input"))
+class UCameraRigInput2DSlot : public UCameraNode
+{
+	GENERATED_BODY()
+
+public:
+
+	/** Input processing parameters. */
+	UPROPERTY(EditAnywhere, Category="Input", meta=(ShowOnlyInnerProperties))
+	FCameraRigInputSlotParameters InputSlotParameters;
+
+	/** Clamping of the final input value. */
+	UPROPERTY(EditAnywhere, Category="Input")
+	FCameraParameterClamping ClampX;
+
+	/** Clamping of the final input value. */
+	UPROPERTY(EditAnywhere, Category="Input")
+	FCameraParameterClamping ClampY;
+
+	/** Normalization of the final input value. */
+	UPROPERTY(EditAnywhere, Category="Input")
+	FCameraParameterNormalization NormalizeX;
+
+	/** Normalization of the final input value. */
+	UPROPERTY(EditAnywhere, Category="Input")
+	FCameraParameterNormalization NormalizeY;
+
+	/** The variable to use to blend with other input slots. */
+	UPROPERTY(EditAnywhere, Category="Input")
+	EBuiltInVector2dCameraVariable BuiltInVariable = EBuiltInVector2dCameraVariable::YawPitch;
+
+	/** The variable to use to blend with other input slots. */
+	UPROPERTY(EditAnywhere, Category="Input", meta=(EditCondition="BuiltInVariable == EBuiltInVector2dCameraVariable::None"))
+	TObjectPtr<UVector2dCameraVariable> Variable;
+
+	/** A node providing an incremental value. */
+	UPROPERTY()
+	TObjectPtr<UInput2DCameraNode> Child;
+
+public:
+
+	FCameraVariableID GetVariableID() const { return VariableID; }
+	FCameraVariableID GetTransientVariableID() const { return TransientVariableID; }
+
+protected:
+
+	// UCameraNode interface.
+	virtual FCameraNodeChildrenView OnGetChildren() override;
+	virtual void OnBuild(FCameraRigBuildContext& BuildContext) override;
+	virtual FCameraNodeEvaluatorPtr OnBuildEvaluator(FCameraNodeEvaluatorBuilder& Builder) const override;
+
+private:
+
+	UPROPERTY()
+	FCameraVariableID TransientVariableID;
+	UPROPERTY()
+	FCameraVariableID VariableID;
+};
+
+namespace UE::Cameras
+{
+
+class FInput2DCameraNodeEvaluator;
+
+class FCameraRigInput2DSlotEvaluator : public FCameraNodeEvaluator
+{
+	UE_DECLARE_CAMERA_NODE_EVALUATOR(GAMEPLAYCAMERAS_API, FCameraRigInput2DSlotEvaluator)
+
+public:
+
+	FCameraRigInput2DSlotEvaluator();
+
+	FVector2d GetInputValue() const { return InputValue; }
+
+protected:
+
+	// FCameraNodeEvaluator interface.
+	virtual void OnBuild(const FCameraNodeEvaluatorBuildParams& Params) override;
+	virtual void OnInitialize(const FCameraNodeEvaluatorInitializeParams& Params) override;
+	virtual FCameraNodeEvaluatorChildrenView OnGetChildren() override;
+	virtual void OnUpdateParameters(const FCameraBlendedParameterUpdateParams& Params, FCameraBlendedParameterUpdateResult& OutResult) override;
+	virtual void OnRun(const FCameraNodeEvaluationParams& Params, FCameraNodeEvaluationResult& OutResult) override;
+
+protected:
+
+	FInput2DCameraNodeEvaluator* ChildEvaluator;
+
+	FVector2d TransientInputValue = FVector2d::ZeroVector;
+	FVector2d InputValue = FVector2d::ZeroVector;
+};
+
+}  // namespace UE::Cameras
+
+
diff --git a/Engine/Plugins/Cameras/GameplayCameras/Source/GameplayCameras/Public/Nodes/Input/CameraRigInputSlotTypes.h b/Engine/Plugins/Cameras/GameplayCameras/Source/GameplayCameras/Public/Nodes/Input/CameraRigInputSlotTypes.h
new file mode 100644
--- /dev/null
+++ b/Engine/Plugins/Cameras/GameplayCameras/Source/GameplayCameras/Public/Nodes/Input/CameraRigInputSlotTypes.h
@@ -0,0 +1,69 @@
+// Copyright Epic Games, Inc. All Rights Reserved.
+
+#pragma once
+
+#include "CoreTypes.h"
+
+#include "CameraRigInputSlotTypes.generated.h"
+
+/**
+ * General options for an input slot.
+ */
+USTRUCT()
+struct FCameraRigInputSlotParameters
+{
+	GENERATED_BODY()
+
+	UPROPERTY(EditAnywhere, Category="Input")
+	bool bIsAccumulated = true;
+
+	UPROPERTY(EditAnywhere, Category="Input")
+	bool bIsPreBlended = true;
+};
+
+/**
+ * Value clamping parameters.
+ */
+USTRUCT()
+struct FCameraParameterClamping
+{
+	GENERATED_BODY()
+
+	UPROPERTY(EditAnywhere, Category="Common", meta=(EditCondition="bClampMin"))
+	double MinValue = 0;
+
+	UPROPERTY(EditAnywhere, Category ="Common", meta=(EditCondition="bClampMax"))
+	double MaxValue = 0;
+
+	UPROPERTY(EditAnywhere, Category="Common")
+	bool bClampMin = false;
+
+	UPROPERTY(EditAnywhere, Category="Common")
+	bool bClampMax = false;
+
+public:
+
+	/** Clamps the given value. */
+	double ClampValue(double Value) const;
+};
+
+/**
+ * Value normalization parameters.
+ */
+USTRUCT()
+struct FCameraParameterNormalization
+{
+	GENERATED_BODY()
+
+	UPROPERTY(EditAnywhere, Category="Common", meta=(ClampMin=0, EditCondition="bNormalize"))
+	double MaxValue = 0;
+
+	UPROPERTY(EditAnywhere, Category="Common")
+	bool bNormalize = false;
+
+public:
+
+	/** Normalizes the given value. */
+	double NormalizeValue(double Value) const;
+};
+
diff --git a/Engine/Plugins/Cameras/GameplayCameras/Source/GameplayCameras/Public/Nodes/Input/Input1DCameraNode.h b/Engine/Plugins/Cameras/GameplayCameras/Source/GameplayCameras/Public/Nodes/Input/Input1DCameraNode.h
new file mode 100644
--- /dev/null
+++ b/Engine/Plugins/Cameras/GameplayCameras/Source/GameplayCameras/Public/Nodes/Input/Input1DCameraNode.h
@@ -0,0 +1,42 @@
+// Copyright Epic Games, Inc. All Rights Reserved.
+
+#pragma once
+
+#include "Core/CameraNode.h"
+#include "Core/CameraNodeEvaluator.h"
+
+#include "Input1DCameraNode.generated.h"
+
+/**
+ * A camera node that provides a floating-point value to an input slot.
+ */
+UCLASS(MinimalAPI, meta=(CameraNodeCategories="Input"))
+class UInput1DCameraNode : public UCameraNode
+{
+	GENERATED_BODY()
+};
+
+namespace UE::Cameras
+{
+
+/**
+ * Base class for a 1D input value node evaluator.
+ */
+class FInput1DCameraNodeEvaluator : public FCameraNodeEvaluator
+{
+	UE_DECLARE_CAMERA_NODE_EVALUATOR(GAMEPLAYCAMERAS_API, FInput1DCameraNodeEvaluator)
+
+public:
+
+	FInput1DCameraNodeEvaluator();
+
+	/** Get the current input value. */
+	double GetInputValue() const { return InputValue; }
+
+protected:
+
+	float InputValue;
+};
+
+}  // namespace UE::Cameras
+
diff --git a/Engine/Plugins/Cameras/GameplayCameras/Source/GameplayCameras/Public/Nodes/Input/Input2DCameraNode.h b/Engine/Plugins/Cameras/GameplayCameras/Source/GameplayCameras/Public/Nodes/Input/Input2DCameraNode.h
new file mode 100644
--- /dev/null
+++ b/Engine/Plugins/Cameras/GameplayCameras/Source/GameplayCameras/Public/Nodes/Input/Input2DCameraNode.h
@@ -0,0 +1,43 @@
+// Copyright Epic Games, Inc. All Rights Reserved.
+
+#pragma once
+
+#include "Core/CameraNode.h"
+#include "Core/CameraNodeEvaluator.h"
+#include "Math/MathFwd.h"
+
+#include "Input2DCameraNode.generated.h"
+
+/**
+ * A camera node that provides a two-dimensional value to an input slot.
+ */
+UCLASS(MinimalAPI, meta=(CameraNodeCategories="Input"))
+class UInput2DCameraNode : public UCameraNode
+{
+	GENERATED_BODY()
+};
+
+namespace UE::Cameras
+{
+
+/**
+ * Base class for a 2D input value node evaluator.
+ */
+class FInput2DCameraNodeEvaluator : public FCameraNodeEvaluator
+{
+	UE_DECLARE_CAMERA_NODE_EVALUATOR(GAMEPLAYCAMERAS_API, FInput2DCameraNodeEvaluator)
+
+public:
+
+	FInput2DCameraNodeEvaluator();
+
+	/** Get the current input value. */
+	FVector2d GetInputValue() const { return InputValue; }
+
+protected:
+
+	FVector2d InputValue;
+};
+
+}  // namespace UE::Cameras
+
diff --git a/Engine/Plugins/Cameras/GameplayCameras/Source/GameplayCameras/Public/Nodes/Input/InputAxisBinding2DCameraNode.h b/Engine/Plugins/Cameras/GameplayCameras/Source/GameplayCameras/Public/Nodes/Input/InputAxisBinding2DCameraNode.h
new file mode 100644
--- /dev/null
+++ b/Engine/Plugins/Cameras/GameplayCameras/Source/GameplayCameras/Public/Nodes/Input/InputAxisBinding2DCameraNode.h
@@ -0,0 +1,47 @@
+// Copyright Epic Games, Inc. All Rights Reserved.
+
+#pragma once
+
+#include "Math/MathFwd.h"
+#include "Nodes/Input/Input2DCameraNode.h"
+
+#include "InputAxisBinding2DCameraNode.generated.h"
+
+class UInputAction;
+
+/**
+ * An input node that reads player input from an input action.
+ */
+UCLASS(MinimalAPI, meta=(CameraNodeCategories="Input"))
+class UInputAxisBinding2DCameraNode : public UInput2DCameraNode
+{
+	GENERATED_BODY()
+
+public:
+
+	/** The axis input action to read from. */
+	UPROPERTY(EditAnywhere, Category="Input")
+	TObjectPtr<UInputAction> AxisAction;
+
+	/** Whether to revert the X axis. */
+	UPROPERTY(EditAnywhere, Category="Input Processing")
+	bool RevertAxisX = false;
+
+	/** Whether to revert the Y axis. */
+	UPROPERTY(EditAnywhere, Category="Input Processing")
+	bool RevertAxisY = false;
+
+	/** A multiplier to use on the input values. */
+	UPROPERTY(EditAnywhere, Category="Input Processing")
+	FVector2D Multiplier;
+
+public:
+
+	UInputAxisBinding2DCameraNode(const FObjectInitializer& ObjInit);
+	
+protected:
+
+	// UCameraNode interface.
+	virtual FCameraNodeEvaluatorPtr OnBuildEvaluator(FCameraNodeEvaluatorBuilder& Builder) const override;
+};
+
