Cameras: allow activating camera rigs directly from a Blueprint director again

The Blueprint camera director can now activate cameras rigs either directly, or via a proxy rig asset. However, in the former case, we check that only one camera asset references this director. If not, we show a warning message.

diff --git a/Engine/Plugins/Cameras/GameplayCameras/Source/GameplayCameras/Private/Directors/BlueprintCameraDirector.cpp b/Engine/Plugins/Cameras/GameplayCameras/Source/GameplayCameras/Private/Directors/BlueprintCameraDirector.cpp
--- a/Engine/Plugins/Cameras/GameplayCameras/Source/GameplayCameras/Private/Directors/BlueprintCameraDirector.cpp
+++ b/Engine/Plugins/Cameras/GameplayCameras/Source/GameplayCameras/Private/Directors/BlueprintCameraDirector.cpp
@@ -78,18 +78,36 @@ void FBlueprintCameraDirectorEvaluator::OnRun(const FCameraDirectorEvaluationPar
 
 		EvaluatorBlueprint->NativeRunCameraDirector(BlueprintParams, BlueprintResult);
 
-		// The BP interface doesn't specify the evaluation context for the chosen camera rigs: we always automatically
-		// make them run in our own owner context.
-		for (const UCameraRigProxyAsset* ActiveCameraRigProxy : BlueprintResult.ActiveCameraRigs)
+		TArray<const UCameraRigAsset*, TInlineAllocator<2>> CameraRigs;
+		const UCameraAsset* CameraAsset = Params.OwnerContext->GetCameraAsset();
+
+		// Gather camera rigs.
+		for (const UCameraRigAsset* ActiveCameraRig : BlueprintResult.ActiveCameraRigs)
+		{
+			if (ActiveCameraRig)
+			{
+				CameraRigs.Add(ActiveCameraRig);
+			}
+			else
+			{
+				UE_LOG(
+						LogCameraSystem, 
+						Error, 
+						TEXT("Null camera rig specified in camera director '%s'."),
+						*EvaluatorBlueprint->GetClass()->GetPathName());
+			}
+		}
+
+		// Resolve camera rig proxies.
+		for (const UCameraRigProxyAsset* ActiveCameraRigProxy : BlueprintResult.ActiveCameraRigProxies)
 		{
 			const UCameraRigAsset* ActiveCameraRig = FindCameraRigByProxy(ActiveCameraRigProxy);
 			if (ActiveCameraRig)
 			{
-				OutResult.Add(Params.OwnerContext, ActiveCameraRig);
+				CameraRigs.Add(ActiveCameraRig);
 			}
 			else
 			{
-				const UCameraAsset* CameraAsset = Params.OwnerContext->GetCameraAsset();
 				UE_LOG(
 						LogCameraSystem, 
 						Error, 
@@ -97,6 +115,13 @@ void FBlueprintCameraDirectorEvaluator::OnRun(const FCameraDirectorEvaluationPar
 						*ActiveCameraRigProxy->GetPathName(), *CameraAsset->GetPathName());
 			}
 		}
+
+		// The BP interface doesn't specify the evaluation context for the chosen camera rigs: we always automatically
+		// make them run in our own owner context.
+		for (const UCameraRigAsset* ActiveCameraRig : CameraRigs)
+		{
+			OutResult.Add(Params.OwnerContext, ActiveCameraRig);
+		}
 	}
 	else
 	{
@@ -130,9 +155,14 @@ void FBlueprintCameraDirectorEvaluator::OnAddReferencedObjects(FReferenceCollect
 
 }  // namespace UE::Cameras
 
-void UBlueprintCameraDirectorEvaluator::ActivateCameraRig(UCameraRigProxyAsset* CameraRigProxy)
+void UBlueprintCameraDirectorEvaluator::ActivateCameraRig(UCameraRigAsset* CameraRig)
 {
-	CurrentResult.ActiveCameraRigs.Add(CameraRigProxy);
+	CurrentResult.ActiveCameraRigs.Add(CameraRig);
+}
+
+void UBlueprintCameraDirectorEvaluator::ActivateCameraRigViaProxy(UCameraRigProxyAsset* CameraRigProxy)
+{
+	CurrentResult.ActiveCameraRigProxies.Add(CameraRigProxy);
 }
 
 void UBlueprintCameraDirectorEvaluator::NativeRunCameraDirector(const FBlueprintCameraDirectorEvaluationParams& Params, FBlueprintCameraDirectorEvaluationResult& OutResult)
@@ -154,13 +184,14 @@ FCameraDirectorEvaluatorPtr UBlueprintCameraDirector::OnBuildEvaluator(FCameraDi
 
 void UBlueprintCameraDirector::OnBuildCameraDirector(UE::Cameras::FCameraBuildLog& BuildLog)
 {
+	using namespace UE::Cameras;
+
+	// Check that a camera director evaluator Blueprint was specified.
 	if (!CameraDirectorEvaluatorClass)
 	{
 		BuildLog.AddMessage(EMessageSeverity::Error, LOCTEXT("MissingBlueprintClass", "No evaluator Blueprint class is set."));
 		return;
 	}
-
-	// TODO: check that the proxy table is complete.
 }
 
 #if WITH_EDITOR
diff --git a/Engine/Plugins/Cameras/GameplayCameras/Source/GameplayCameras/Public/Directors/BlueprintCameraDirector.h b/Engine/Plugins/Cameras/GameplayCameras/Source/GameplayCameras/Public/Directors/BlueprintCameraDirector.h
--- a/Engine/Plugins/Cameras/GameplayCameras/Source/GameplayCameras/Public/Directors/BlueprintCameraDirector.h
+++ b/Engine/Plugins/Cameras/GameplayCameras/Source/GameplayCameras/Public/Directors/BlueprintCameraDirector.h
@@ -38,13 +38,17 @@ struct FBlueprintCameraDirectorEvaluationResult
 
 	/** The list of camera rigs that should be active this frame. */
 	UPROPERTY(BlueprintReadWrite, Category="Evaluation")
-	TArray<UCameraRigProxyAsset*> ActiveCameraRigs;
+	TArray<UCameraRigProxyAsset*> ActiveCameraRigProxies;
+
+	/** The list of camera rigs that should be active this frame. */
+	UPROPERTY(BlueprintReadWrite, Category="Evaluation")
+	TArray<UCameraRigAsset*> ActiveCameraRigs;
 };
 
 /**
  * Base class for a Blueprint camera director evaluator.
  */
-UCLASS(MinimalAPI, Blueprintable)
+UCLASS(MinimalAPI, Blueprintable, Abstract)
 class UBlueprintCameraDirectorEvaluator : public UObject
 {
 	GENERATED_BODY()
@@ -58,11 +62,16 @@ public:
 	UFUNCTION(BlueprintCallable, BlueprintImplementableEvent, Category="Evaluation")
 	void RunCameraDirector(const FBlueprintCameraDirectorEvaluationParams& Params);
 
+	/** Specifies a camera rig to be active this frame. */
+	UFUNCTION(BlueprintCallable, Category="Evaluation")
+	void ActivateCameraRig(UPARAM(meta=(UseCameraRigPicker=true)) UCameraRigAsset* CameraRig);
+
 	/**
-	 * Specifies a camera rig to be active this frame.
+	 * Specifies a camera rig to be active this frame, via a proxy which is later resolved
+	 * via the proxy table of the Blueprint camera director.
 	 */
 	UFUNCTION(BlueprintCallable, Category="Evaluation")
-	void ActivateCameraRig(UCameraRigProxyAsset* CameraRigProxy);
+	void ActivateCameraRigViaProxy(UCameraRigProxyAsset* CameraRigProxy);
 
 	/** Native wrapper for RunCameraDirector. */
 	void NativeRunCameraDirector(
@@ -74,7 +83,6 @@ protected:
 	/** The current camera director evaluation result. */
 	UPROPERTY(BlueprintReadWrite, Category="Evaluation")
 	FBlueprintCameraDirectorEvaluationResult CurrentResult;
-
 };
 
 /**
diff --git a/Engine/Plugins/Cameras/GameplayCameras/Source/GameplayCamerasEditor/Private/Editors/GameplayCamerasGraphPanelPinFactory.cpp b/Engine/Plugins/Cameras/GameplayCameras/Source/GameplayCamerasEditor/Private/Editors/GameplayCamerasGraphPanelPinFactory.cpp
--- a/Engine/Plugins/Cameras/GameplayCameras/Source/GameplayCamerasEditor/Private/Editors/GameplayCamerasGraphPanelPinFactory.cpp
+++ b/Engine/Plugins/Cameras/GameplayCameras/Source/GameplayCamerasEditor/Private/Editors/GameplayCamerasGraphPanelPinFactory.cpp
@@ -34,12 +34,19 @@ TSharedPtr<class SGraphPin> FGameplayCamerasGraphPanelPinFactory::CreatePin(UEdG
 		return nullptr;
 	}
 
-	if (!ParameterProperty->HasMetaData(TEXT("UseCameraRigNamePicker")))
+	if (ParameterProperty->HasMetaData(TEXT("UseCameraRigNamePicker")))
 	{
-		return nullptr;
+		return SNew(SCameraRigNameGraphPin, Pin)
+			.PinMode(ECameraRigNameGraphPinMode::NamePin);
+	}
+
+	if (ParameterProperty->HasMetaData(TEXT("UseCameraRigPicker")))
+	{
+		return SNew(SCameraRigNameGraphPin, Pin)
+			.PinMode(ECameraRigNameGraphPinMode::ReferencePin);
 	}
 
-	return SNew(SCameraRigNameGraphPin, Pin);
+	return nullptr;
 }
 
 }  // namespace UE::Cameras
diff --git a/Engine/Plugins/Cameras/GameplayCameras/Source/GameplayCamerasEditor/Private/Editors/SCameraRigNameGraphPin.cpp b/Engine/Plugins/Cameras/GameplayCameras/Source/GameplayCamerasEditor/Private/Editors/SCameraRigNameGraphPin.cpp
--- a/Engine/Plugins/Cameras/GameplayCameras/Source/GameplayCamerasEditor/Private/Editors/SCameraRigNameGraphPin.cpp
+++ b/Engine/Plugins/Cameras/GameplayCameras/Source/GameplayCamerasEditor/Private/Editors/SCameraRigNameGraphPin.cpp
@@ -7,6 +7,7 @@
 #include "Core/CameraRigAsset.h"
 #include "Editors/CameraRigPickerConfig.h"
 #include "Framework/Views/ITypedTableView.h"
+#include "Helpers/CameraDirectorHelper.h"
 #include "IContentBrowserSingleton.h"
 #include "IGameplayCamerasEditorModule.h"
 #include "Kismet2/BlueprintEditorUtils.h"
@@ -30,6 +31,8 @@ namespace UE::Cameras
 void SCameraRigNameGraphPin::Construct(const FArguments& InArgs, UEdGraphPin* InGraphPinObj)
 {
 	SGraphPin::Construct(SGraphPin::FArguments(), InGraphPinObj);
+
+	PinMode = InArgs._PinMode;
 }
 
 TSharedRef<SWidget>	SCameraRigNameGraphPin::GetDefaultValueWidget()
@@ -75,7 +78,7 @@ TSharedRef<SWidget>	SCameraRigNameGraphPin::GetDefaultValueWidget()
 			.ButtonColorAndOpacity(this, &SCameraRigNameGraphPin::OnGetWidgetBackground)
 			.OnClicked(this, &SCameraRigNameGraphPin::OnResetButtonClicked)
 			.ContentPadding(1.f)
-			.ToolTipText(LOCTEXT("ResetButtonToolTip", "Reset the name to an empty string."))
+			.ToolTipText(LOCTEXT("ResetButtonToolTip", "Reset the camera rig reference."))
 			.IsEnabled(this, &SGraphPin::IsEditingEnabled)
 			[
 				SNew(SImage)
@@ -119,9 +122,19 @@ FText SCameraRigNameGraphPin::OnGetComboText() const
 	
 	if (GraphPinObj != nullptr)
 	{
-		if (!GraphPinObj->DefaultValue.IsEmpty())
+		switch (PinMode)
 		{
-			Value = FText::FromString(GraphPinObj->DefaultValue);
+			case ECameraRigNameGraphPinMode::NamePin:
+				{
+					Value = FText::FromString(GraphPinObj->DefaultValue);
+				}
+				break;
+			case ECameraRigNameGraphPinMode::ReferencePin:
+				if (const UCameraRigAsset* CameraRig = Cast<const UCameraRigAsset>(GraphPinObj->DefaultObject))
+				{
+					Value = FText::FromString(CameraRig->GetDisplayName());
+				}
+				break;
 		}
 	}
 	return Value;
@@ -135,34 +148,64 @@ FText SCameraRigNameGraphPin::OnGetComboToolTipText() const
 TSharedRef<SWidget> SCameraRigNameGraphPin::OnBuildCameraRigNamePicker()
 {
 	FCameraRigPickerConfig CameraRigPickerConfig;
-	CameraRigPickerConfig.bCanSelectCameraAsset = true;
-	CameraRigPickerConfig.CameraAssetSelectionMode = ESelectionMode::Multi;
-	CameraRigPickerConfig.CameraAssetViewType = EAssetViewType::List;
-	CameraRigPickerConfig.CameraAssetSaveSettingsName = TEXT("CameraRigNamePicker");
+	CameraRigPickerConfig.bCanSelectCameraAsset = false;
+	CameraRigPickerConfig.bFocusCameraRigSearchBoxWhenOpened = true;
 	CameraRigPickerConfig.OnCameraRigSelected = FOnCameraRigSelected::CreateSP(this, &SCameraRigNameGraphPin::OnPickerAssetSelected);
 
-	if (TSharedPtr<SGraphNode> OwnerNodeWidget = OwnerNodePtr.Pin())
-	{
-		FAssetRegistryModule& AssetRegistryModule = FModuleManager::LoadModuleChecked<FAssetRegistryModule>("AssetRegistry");
+	TSharedPtr<SGraphNode> OwnerNodeWidget = OwnerNodePtr.Pin();
+	check(OwnerNodeWidget);
+	
+	UEdGraphNode* OwnerNode = OwnerNodeWidget->GetNodeObj();
+	UBlueprint* OwnerBlueprint = FBlueprintEditorUtils::FindBlueprintForNode(OwnerNode);
+
+	TArray<UCameraAsset*> ReferencingCameraAssets;
+	FCameraDirectorHelper::GetReferencingCameraAssets(OwnerBlueprint, ReferencingCameraAssets);
 
-		UEdGraphNode* OwnerNode = OwnerNodeWidget->GetNodeObj();
-		UBlueprint* OwnerBlueprint = FBlueprintEditorUtils::FindBlueprintForNode(OwnerNode);
+	if (ReferencingCameraAssets.Num() == 0)
+	{
+		CameraRigPickerConfig.WarningMessage = LOCTEXT("NoReferencingCameraAssetWarning",
+				"No camera asset references this Blueprint, so no camera rig list can be displayed. "
+				"Make a camera asset use this Blueprint as its camera director evaluator, or use "
+				"ActivateCameraRigViaProxy.");
+	}
+	else
+	{
+		CameraRigPickerConfig.InitialCameraAssetSelection = ReferencingCameraAssets[0];
 
-		FARFilter Filter;
-		IAssetRegistry& AssetRegistry = AssetRegistryModule.Get();
-		AssetRegistry.GetReferencers(OwnerBlueprint->GetOutermost()->GetFName(), Filter.PackageNames);
-		TArray<FAssetData> ReferencerAssetData;
-		AssetRegistry.GetAssets(Filter, ReferencerAssetData);
-		if (!ReferencerAssetData.IsEmpty())
+		if (ReferencingCameraAssets.Num() > 1)
 		{
-			//TODO: multi-select all referencers.
-			CameraRigPickerConfig.InitialCameraAssetSelection = ReferencerAssetData[0];
+			CameraRigPickerConfig.WarningMessage = LOCTEXT("ManyReferencingCameraAssetsWarning",
+				"More than one camera asset references this Blueprint. Only camera rigs from the first "
+				"one will be displayed. Even then, shared camera director Blueprints should use "
+				"ActivateCameraRigViaProxy instead.");
 		}
 	}
 
-	if (GraphPinObj && !GraphPinObj->DefaultValue.IsEmpty())
+	if (GraphPinObj)
 	{
-		CameraRigPickerConfig.InitialCameraRigSelectionName = GraphPinObj->DefaultValue;
+		switch (PinMode)
+		{
+			case ECameraRigNameGraphPinMode::NamePin:
+				if (!ReferencingCameraAssets.IsEmpty())
+				{
+					TArrayView<const TObjectPtr<UCameraRigAsset>> CameraRigs = ReferencingCameraAssets[0]->GetCameraRigs();
+					const TObjectPtr<UCameraRigAsset>* FoundItem = CameraRigs.FindByPredicate(
+							[this](UCameraRigAsset* Item)
+							{
+							return Item->GetDisplayName() == GraphPinObj->DefaultValue;
+							});
+					if (FoundItem)
+					{
+						CameraRigPickerConfig.InitialCameraRigSelection = *FoundItem;
+					}
+				}
+				break;
+			case ECameraRigNameGraphPinMode::ReferencePin:
+				{
+					CameraRigPickerConfig.InitialCameraRigSelection = Cast<UCameraRigAsset>(GraphPinObj->DefaultObject);
+				}
+				break;
+		}
 	}
 
 	IGameplayCamerasEditorModule& CamerasEditorModule = FModuleManager::LoadModuleChecked<IGameplayCamerasEditorModule>("GameplayCamerasEditor");
@@ -174,22 +217,32 @@ void SCameraRigNameGraphPin::OnPickerAssetSelected(UCameraRigAsset* SelectedItem
 	if (SelectedItem)
 	{
 		CameraRigPickerButton->SetIsOpen(false);
-		SetCameraRigName(SelectedItem->GetDisplayName());
+		SetCameraRig(SelectedItem);
 	}
 }
 
 FReply SCameraRigNameGraphPin::OnResetButtonClicked()
 {
 	CameraRigPickerButton->SetIsOpen(false);
-	SetCameraRigName(FString());
+	SetCameraRig(nullptr);
 	return FReply::Handled();
 }
 
-void SCameraRigNameGraphPin::SetCameraRigName(const FString& InCameraRigName)
+void SCameraRigNameGraphPin::SetCameraRig(UCameraRigAsset* SelectedCameraRig)
 {
 	const FScopedTransaction Transaction(LOCTEXT("ChangeObjectPinValue", "Change Object Pin Value"));
+
 	GraphPinObj->Modify();
-	GraphPinObj->GetSchema()->TrySetDefaultValue(*GraphPinObj, InCameraRigName);
+
+	switch (PinMode)
+	{
+		case ECameraRigNameGraphPinMode::NamePin:
+			GraphPinObj->GetSchema()->TrySetDefaultValue(*GraphPinObj, SelectedCameraRig->GetDisplayName());
+			break;
+		case ECameraRigNameGraphPinMode::ReferencePin:
+			GraphPinObj->GetSchema()->TrySetDefaultObject(*GraphPinObj, SelectedCameraRig);
+			break;
+	}
 }
 
 }  // namespace UE::Cameras
diff --git a/Engine/Plugins/Cameras/GameplayCameras/Source/GameplayCamerasEditor/Private/Editors/SCameraRigPicker.cpp b/Engine/Plugins/Cameras/GameplayCameras/Source/GameplayCamerasEditor/Private/Editors/SCameraRigPicker.cpp
--- a/Engine/Plugins/Cameras/GameplayCameras/Source/GameplayCamerasEditor/Private/Editors/SCameraRigPicker.cpp
+++ b/Engine/Plugins/Cameras/GameplayCameras/Source/GameplayCamerasEditor/Private/Editors/SCameraRigPicker.cpp
@@ -6,15 +6,19 @@
 #include "ContentBrowserModule.h"
 #include "Core/CameraAsset.h"
 #include "Core/CameraRigAsset.h"
+#include "Editor.h"
+#include "Editor/EditorEngine.h"
 #include "Editors/CameraRigPickerConfig.h"
 #include "Framework/Application/SlateApplication.h"
-#include "IContentBrowserSingleton.h"
+#include "Framework/Views/ITypedTableView.h"
 #include "Layout/WidgetPath.h"
 #include "PropertyHandle.h"
 #include "Styles/GameplayCamerasEditorStyle.h"
 #include "Types/SlateEnums.h"
+#include "Widgets/Input/SHyperlink.h"
 #include "Widgets/Input/SSearchBox.h"
 #include "Widgets/Layout/SBox.h"
+#include "Widgets/SBoxPanel.h"
 
 #define LOCTEXT_NAMESPACE "SCameraRigPicker"
 
@@ -25,12 +29,15 @@ void SCameraRigPicker::Construct(const FArguments& InArgs)
 {
 	const FCameraRigPickerConfig& PickerConfig = InArgs._CameraRigPickerConfig;
 
+	const ISlateStyle& AppStyle = FAppStyle::Get();
+
 	SearchTextFilter = MakeShareable(new FTextFilter(
 		FTextFilter::FItemToStringArray::CreateSP(this, &SCameraRigPicker::GetEntryStrings)));
 
-	// Camera asset picker.
 	TSharedRef<SVerticalBox> LayoutBox = SNew(SVerticalBox);
+	const float CamerRigPickerFillHeight = PickerConfig.bCanSelectCameraAsset ? 0.45f : 1.f;
 
+	// Camera asset picker.
 	if (PickerConfig.bCanSelectCameraAsset)
 	{
 		FAssetRegistryModule& AssetRegistryModule = FModuleManager::LoadModuleChecked<FAssetRegistryModule>("AssetRegistry");
@@ -44,9 +51,10 @@ void SCameraRigPicker::Construct(const FArguments& InArgs)
 		AssetPickerConfig.OnAssetSelected = FOnAssetSelected::CreateSP(this, &SCameraRigPicker::OnCameraAssetSelected);
 		AssetPickerConfig.GetCurrentSelectionDelegates.Add(&GetCurrentCameraAssetPickerSelection);
 
-		AssetPickerConfig.SelectionMode = PickerConfig.CameraAssetSelectionMode;
+		AssetPickerConfig.SelectionMode = ESelectionMode::Single;
 		AssetPickerConfig.InitialAssetViewType = PickerConfig.CameraAssetViewType;
 		AssetPickerConfig.SaveSettingsName = PickerConfig.CameraAssetSaveSettingsName;
+
 		AssetPickerConfig.InitialAssetSelection = PickerConfig.InitialCameraAssetSelection;
 
 		LayoutBox->AddSlot()
@@ -55,38 +63,98 @@ void SCameraRigPicker::Construct(const FArguments& InArgs)
 			ContentBrowserModule.Get().CreateAssetPicker(AssetPickerConfig)
 		];
 	}
-
-	// Camera rig picker.
-	const float CamerRigPickerFillHeight = PickerConfig.bCanSelectCameraAsset ? 0.45f : 1.f;
-	LayoutBox->AddSlot()
-	.FillHeight(CamerRigPickerFillHeight)
-	[
-		SNew(SVerticalBox)
-		+SVerticalBox::Slot()
+	// Which camera asset is being shown.
+	else
+	{
+		LayoutBox->AddSlot()
 		.AutoHeight()
 		[
-			SNew(SBorder)
-			.BorderImage(FAppStyle::GetBrush("ToolPanel.GroupBorder"))
-			.Padding(8.f)
+			SNew(SHorizontalBox)
+			+SHorizontalBox::Slot()
+			.AutoWidth()
+			.VAlign(VAlign_Center)
+			.Padding(4.f, 4.f, 0.f, 4.f)
 			[
-				SAssignNew(SearchBox, SSearchBox)
-				.HintText(LOCTEXT("SearchHint", "Search"))
-				.OnTextChanged(this, &SCameraRigPicker::OnSearchTextChanged)
-				.OnTextCommitted(this, &SCameraRigPicker::OnSearchTextCommitted)
-				.OnKeyDownHandler(this, &SCameraRigPicker::OnSearchKeyDown)
+				SNew(STextBlock)
+				.Text(LOCTEXT("CameraAssetInfo", "Showing camera rigs from "))
 			]
-		]
-		+SVerticalBox::Slot()
-		.Padding(0.f, 3.f)
+			+SHorizontalBox::Slot()
+			.AutoWidth()
+			.VAlign(VAlign_Center)
+			.Padding(0.f, 4.f)
+			[
+				SNew(SHyperlink)
+				.Text(this, &SCameraRigPicker::GetSelectedCameraAssetName)
+				.OnNavigate(this, &SCameraRigPicker::NavigateToSelectedCameraAsset)
+			]
+			+ SHorizontalBox::Slot()
+			.AutoWidth()
+			.VAlign(VAlign_Center)
+			.Padding(4.f, 4.f, 4.f, 4.f)
+			[
+				SNew(SImage)
+				.ColorAndOpacity(FSlateColor::UseForeground())
+				.Image(FAppStyle::Get().GetBrush("Icons.BrowseContent"))
+			]
+		];
+	}
+
+	// Search box.
+	LayoutBox->AddSlot()
+	.AutoHeight()
+	[
+		SNew(SBorder)
+		.BorderImage(FAppStyle::GetBrush("ToolPanel.GroupBorder"))
+		.Padding(8.f)
 		[
-			SAssignNew(CameraRigListView, SListView<UCameraRigAsset*>)
-			.ListItemsSource(&CameraRigFilteredItemsSource)
-			.OnGenerateRow(this, &SCameraRigPicker::OnCameraRigListGenerateRow)
-			.OnSelectionChanged(this, &SCameraRigPicker::OnCameraRigListSelectionChanged)
+			SAssignNew(SearchBox, SSearchBox)
+			.HintText(LOCTEXT("SearchHint", "Search"))
+			.OnTextChanged(this, &SCameraRigPicker::OnSearchTextChanged)
+			.OnTextCommitted(this, &SCameraRigPicker::OnSearchTextCommitted)
+			.OnKeyDownHandler(this, &SCameraRigPicker::OnSearchKeyDown)
 		]
 	];
+
+	// List of camera rig names.
+	LayoutBox->AddSlot()
+	.FillHeight(CamerRigPickerFillHeight)
+	.Padding(0.f, 3.f)
+	[
+		SAssignNew(CameraRigListView, SListView<UCameraRigAsset*>)
+		.ListItemsSource(&CameraRigFilteredItemsSource)
+		.OnGenerateRow(this, &SCameraRigPicker::OnCameraRigListGenerateRow)
+		.OnSelectionChanged(this, &SCameraRigPicker::OnCameraRigListSelectionChanged)
+	];
+
+	// Optional warning message.
+	if (!PickerConfig.WarningMessage.IsEmpty())
+	{
+		LayoutBox->AddSlot()
+		.AutoHeight()
+		[
+			SNew(STextBlock)
+			.ColorAndOpacity(AppStyle.GetSlateColor("Colors.Warning"))
+			.Margin(4.f)
+			.Text(PickerConfig.WarningMessage)
+			.AutoWrapText(true)
+		];
+	}
+
+	// Optional error message.
+	if (!PickerConfig.ErrorMessage.IsEmpty())
+	{
+		LayoutBox->AddSlot()
+		.AutoHeight()
+		[
+			SNew(STextBlock)
+			.ColorAndOpacity(AppStyle.GetSlateColor("Colors.Error"))
+			.Margin(4.f)
+			.Text(PickerConfig.ErrorMessage)
+			.AutoWrapText(true)
+		];
+	}
 	
-	// Full layout.
+	// Assemble it all.
 	ChildSlot
 	[
 		SNew(SBox)
@@ -101,26 +169,29 @@ void SCameraRigPicker::Construct(const FArguments& InArgs)
 		]
 	];
 
-	// If we have an initially selected camera asset, fill the list of camera rigs immediately.
-	if (PickerConfig.InitialCameraAssetSelection.IsValid())
+	// If we have an initially selected camera asset, set it up immediately.
+	if (!PickerConfig.bCanSelectCameraAsset)
 	{
-		UpdateCameraRigItemsSource(TArray<FAssetData>{ PickerConfig.InitialCameraAssetSelection });
+		FixedCameraAssetSelection = PickerConfig.InitialCameraAssetSelection;
 	}
-
+	UpdateCameraRigItemsSource();
 	UpdateCameraRigFilteredItemsSource();
 
 	// If we have an initially selected camera rig, select it in the list immediately too.
 	UCameraRigAsset* InitialCameraRigSelection = PickerConfig.InitialCameraRigSelection;
-	if (!PickerConfig.InitialCameraRigSelectionName.IsEmpty() && !InitialCameraRigSelection)
+	if (!InitialCameraRigSelection && PickerConfig.InitialCameraRigSelectionGuid.IsValid())
 	{
-		UCameraRigAsset** FoundItem = CameraRigFilteredItemsSource.FindByPredicate(
-				[&PickerConfig](UCameraRigAsset* Item)
-				{
-					return Item->GetDisplayName() == PickerConfig.InitialCameraRigSelectionName;
-				});
-		if (FoundItem)
+		if (UCameraAsset* CameraAsset = GetSelectedCameraAsset())
 		{
-			InitialCameraRigSelection = *FoundItem;
+			const TObjectPtr<UCameraRigAsset>* FoundItem = CameraAsset->GetCameraRigs().FindByPredicate(
+					[&PickerConfig](UCameraRigAsset* Item)
+					{
+						return Item->GetGuid() == PickerConfig.InitialCameraRigSelectionGuid;
+					});
+			if (FoundItem)
+			{
+				InitialCameraRigSelection = *FoundItem;
+			}
 		}
 	}
 	if (InitialCameraRigSelection)
@@ -147,8 +218,11 @@ EActiveTimerReturnType SCameraRigPicker::FocusCameraRigSearchBox(double InCurren
 		FWidgetPath WidgetToFocusPath;
 		FSlateApplication::Get().GeneratePathToWidgetUnchecked(SearchBox.ToSharedRef(), WidgetToFocusPath);
 		FSlateApplication::Get().SetKeyboardFocus(WidgetToFocusPath, EFocusCause::SetDirectly);
-		WidgetToFocusPath.GetWindow()->SetWidgetToFocusOnActivate(SearchBox);
-		return EActiveTimerReturnType::Stop;
+		if (WidgetToFocusPath.IsValid())
+		{
+			WidgetToFocusPath.GetWindow()->SetWidgetToFocusOnActivate(SearchBox);
+			return EActiveTimerReturnType::Stop;
+		}
 	}
 
 	return EActiveTimerReturnType::Continue;
@@ -158,11 +232,7 @@ void SCameraRigPicker::Tick(const FGeometry& AllottedGeometry, const double InCu
 {
 	if (bUpdateItemsSource)
 	{
-		if (GetCurrentCameraAssetPickerSelection.IsBound())
-		{
-			TArray<FAssetData> SelectedAssets = GetCurrentCameraAssetPickerSelection.Execute();
-			UpdateCameraRigItemsSource(SelectedAssets);
-		}
+		UpdateCameraRigItemsSource();
 	}
 	if (bUpdateFilteredItemsSource || bUpdateItemsSource)
 	{
@@ -181,11 +251,6 @@ void SCameraRigPicker::Tick(const FGeometry& AllottedGeometry, const double InCu
 	SCompoundWidget::Tick(AllottedGeometry, InCurrentTime, InDeltaTime);
 }
 
-void SCameraRigPicker::OnCameraAssetSelected(const FAssetData& AssetData)
-{
-	bUpdateItemsSource = true;
-}
-
 TSharedRef<ITableRow> SCameraRigPicker::OnCameraRigListGenerateRow(UCameraRigAsset* Item, const TSharedRef<STableViewBase>& OwnerTable)
 {
 	TSharedRef<FGameplayCamerasEditorStyle> GameplayCamerasStyle = FGameplayCamerasEditorStyle::Get();
@@ -216,29 +281,90 @@ TSharedRef<ITableRow> SCameraRigPicker::OnCameraRigListGenerateRow(UCameraRigAss
 		];
 }
 
+UCameraAsset* SCameraRigPicker::GetSelectedCameraAsset() const
+{
+	if (GetCurrentCameraAssetPickerSelection.IsBound())
+	{
+		TArray<FAssetData> Selection = GetCurrentCameraAssetPickerSelection.Execute();
+		if (!Selection.IsEmpty())
+		{
+			return Cast<UCameraAsset>(Selection[0].GetAsset());
+		}
+		return nullptr;
+	}
+	else
+	{
+		return Cast<UCameraAsset>(FixedCameraAssetSelection.GetAsset());
+	}
+}
+
+FText SCameraRigPicker::GetSelectedCameraAssetName() const
+{
+	if (UCameraAsset* CameraAsset = GetSelectedCameraAsset())
+	{
+		return FText::FromName(CameraAsset->GetFName());
+	}
+	return LOCTEXT("NoCameraAssetName", "None");
+}
+
+void SCameraRigPicker::NavigateToSelectedCameraAsset() const
+{
+	if (UCameraAsset* CameraAsset = GetSelectedCameraAsset())
+	{
+		FAssetData AssetData(CameraAsset);
+		GEditor->SyncBrowserToObject(AssetData);
+	}
+}
+
+void SCameraRigPicker::OnCameraAssetSelected(const FAssetData& AssetData)
+{
+	bUpdateItemsSource = true;
+}
+
 void SCameraRigPicker::OnCameraRigListSelectionChanged(UCameraRigAsset* Item, ESelectInfo::Type SelectInfo)
 {
 	if (SelectInfo != ESelectInfo::Direct)
 	{
 		if (PropertyToSet)
 		{
-			PropertyToSet->SetValue(Item);
+			FProperty* Property = PropertyToSet->GetProperty();
+			if (FObjectProperty* ObjectProperty = CastField<FObjectProperty>(Property))
+			{
+				if (ensure(ObjectProperty->PropertyClass && ObjectProperty->PropertyClass->IsChildOf<UCameraRigAsset>()))
+				{
+					PropertyToSet->SetValue(Item);
+				}
+			}
+			else if (FStrProperty* StringProperty = CastField<FStrProperty>(Property))
+			{
+				PropertyToSet->SetValue(Item->GetDisplayName());
+			}
+			else if (FStructProperty* StructProperty = CastField<FStructProperty>(Property))
+			{
+				if (StructProperty->Struct && StructProperty->Struct == TBaseStructure<FGuid>::Get())
+				{
+					PropertyToSet->SetValue(Item->GetGuid().ToString());
+				}
+			}
+			else
+			{
+				ensureMsgf(false, TEXT("Don't know how to set camera rig on property: %s"), *Property->GetFullName());
+			}
 		}
 
 		OnCameraRigSelected.ExecuteIfBound(Item);
 	}
 }
 
-void SCameraRigPicker::UpdateCameraRigItemsSource(const TArray<FAssetData>& Assets)
+void SCameraRigPicker::UpdateCameraRigItemsSource()
 {
-	CameraRigItemsSource.Reset();
-
-	for (const FAssetData& SelectedAsset : Assets)
+	if (UCameraAsset* CameraAsset = GetSelectedCameraAsset())
 	{
-		if (const UCameraAsset* CameraAsset = Cast<UCameraAsset>(SelectedAsset.GetAsset()))
-		{
-			CameraRigItemsSource.Append(CameraAsset->GetCameraRigs());
-		}
+		CameraRigItemsSource = CameraAsset->GetCameraRigs();
+	}
+	else
+	{
+		CameraRigItemsSource.Reset();
 	}
 }
 
diff --git a/Engine/Plugins/Cameras/GameplayCameras/Source/GameplayCamerasEditor/Private/Editors/SCameraRigPicker.h b/Engine/Plugins/Cameras/GameplayCameras/Source/GameplayCamerasEditor/Private/Editors/SCameraRigPicker.h
--- a/Engine/Plugins/Cameras/GameplayCameras/Source/GameplayCamerasEditor/Private/Editors/SCameraRigPicker.h
+++ b/Engine/Plugins/Cameras/GameplayCameras/Source/GameplayCamerasEditor/Private/Editors/SCameraRigPicker.h
@@ -34,11 +34,15 @@ private:
 
 	EActiveTimerReturnType FocusCameraRigSearchBox(double InCurrentTime, float InDeltaTime);
 
+	UCameraAsset* GetSelectedCameraAsset() const;
+	FText GetSelectedCameraAssetName() const;
+	void NavigateToSelectedCameraAsset() const;
+
 	void OnCameraAssetSelected(const FAssetData& AssetData);
 
 	TSharedRef<ITableRow> OnCameraRigListGenerateRow(UCameraRigAsset* Item, const TSharedRef<STableViewBase>& OwnerTable);
 	void OnCameraRigListSelectionChanged(UCameraRigAsset* Item, ESelectInfo::Type SelectInfo);
-	void UpdateCameraRigItemsSource(const TArray<FAssetData>& Assets);
+	void UpdateCameraRigItemsSource();
 	void UpdateCameraRigFilteredItemsSource();
 
 	void GetEntryStrings(const UCameraRigAsset* InItem, TArray<FString>& OutStrings);
@@ -49,6 +53,7 @@ private:
 
 private:
 
+	FAssetData FixedCameraAssetSelection;
 	FGetCurrentSelectionDelegate GetCurrentCameraAssetPickerSelection;
 
 	TSharedPtr<SListView<UCameraRigAsset*>> CameraRigListView;
diff --git a/Engine/Plugins/Cameras/GameplayCameras/Source/GameplayCamerasEditor/Private/Helpers/CameraDirectorHelper.cpp b/Engine/Plugins/Cameras/GameplayCameras/Source/GameplayCamerasEditor/Private/Helpers/CameraDirectorHelper.cpp
new file mode 100644
--- /dev/null
+++ b/Engine/Plugins/Cameras/GameplayCameras/Source/GameplayCamerasEditor/Private/Helpers/CameraDirectorHelper.cpp
@@ -0,0 +1,66 @@
+// Copyright Epic Games, Inc. All Rights Reserved.
+
+#include "Helpers/CameraDirectorHelper.h"
+
+#include "AssetRegistry/AssetData.h"
+#include "AssetRegistry/IAssetRegistry.h"
+#include "Core/CameraAsset.h"
+#include "Misc/AssetRegistryInterface.h"
+#include "UObject/Object.h"
+#include "UObject/Package.h"
+#include "UObject/ReferencerFinder.h"
+
+namespace UE::Cameras
+{
+
+void FCameraDirectorHelper::GetReferencingCameraAssets(UObject* CameraDirectorObject, TArray<UCameraAsset*>& OutReferencers)
+{
+	TSet<UCameraAsset*> UniqueReferencers;
+	UPackage* CameraDirectorObjectPackage = CameraDirectorObject->GetOutermost();
+
+	// Assume the asset registry module is already loaded.
+	IAssetRegistry* AssetRegistry = IAssetRegistry::Get();
+
+	// Get on-disk referencers.
+	TArray<FAssetIdentifier> ReferencerIds;
+	FAssetIdentifier AssetIdentifier(CameraDirectorObjectPackage->GetFName());
+	AssetRegistry->GetReferencers(AssetIdentifier, ReferencerIds);
+
+	TArray<FAssetData> AllAssetData;
+	for (FAssetIdentifier ReferencerId : ReferencerIds)
+	{
+		AssetRegistry->GetAssetsByPackageName(ReferencerId.PackageName, AllAssetData);
+	}
+	for (const FAssetData& AssetData : AllAssetData)
+	{
+		UObject* Asset = AssetData.GetAsset();
+		if (UCameraAsset* CameraAsset = Cast<UCameraAsset>(Asset))
+		{
+			UniqueReferencers.Add(CameraAsset);
+		}
+	}
+
+	// Get in-memory referencers.
+	TArray<UObject*> ReferencedObjects;
+	GetObjectsWithPackage(CameraDirectorObjectPackage, ReferencedObjects);
+	ReferencedObjects.Add(CameraDirectorObjectPackage);
+
+	TArray<UObject*> AllReferencers = FReferencerFinder::GetAllReferencers(
+			ReferencedObjects, nullptr, EReferencerFinderFlags::SkipWeakReferences);
+	for (UObject* Referencer : AllReferencers)
+	{
+		if (UCameraAsset* CameraAsset = Cast<UCameraAsset>(Referencer))
+		{
+			UniqueReferencers.Add(CameraAsset);
+		}
+		else if (UCameraAsset* OuterCameraAsset = Referencer->GetTypedOuter<UCameraAsset>())
+		{
+			UniqueReferencers.Add(OuterCameraAsset);
+		}
+	}
+
+	OutReferencers.Append(UniqueReferencers.Array());
+}
+
+}  // namespace UE::Cameras
+
diff --git a/Engine/Plugins/Cameras/GameplayCameras/Source/GameplayCamerasEditor/Public/Editors/CameraRigPickerConfig.h b/Engine/Plugins/Cameras/GameplayCameras/Source/GameplayCamerasEditor/Public/Editors/CameraRigPickerConfig.h
--- a/Engine/Plugins/Cameras/GameplayCameras/Source/GameplayCamerasEditor/Public/Editors/CameraRigPickerConfig.h
+++ b/Engine/Plugins/Cameras/GameplayCameras/Source/GameplayCamerasEditor/Public/Editors/CameraRigPickerConfig.h
@@ -2,11 +2,12 @@
 
 #pragma once
 
-#include "Framework/Views/ITypedTableView.h"
+#include "CoreTypes.h"
 #include "IContentBrowserSingleton.h"
-#include "UObject/ObjectKey.h"
+#include "Internationalization/Text.h"
 
 class IPropertyHandle;
+class UCameraAsset;
 class UCameraRigAsset;
 
 namespace UE::Cameras
@@ -16,9 +17,8 @@ DECLARE_DELEGATE_OneParam(FOnCameraRigSelected, UCameraRigAsset*);
 
 /**
  * Configuration structure for a camera rig picker widget.
- * This is a widget that shows:
- *	- An asset picker for a camera asset, and below it...
- *	- A list view with the camera rigs inside that camera asset.
+ * This is a widget that shows the list of camera rigs inside a given
+ * camera asset.
  *
  * See IGameplayCamerasEditorModule for creating that widget.
  */
@@ -27,38 +27,41 @@ struct FCameraRigPickerConfig
 	/** The initially selected camera asset, if any. */
 	FAssetData InitialCameraAssetSelection;
 
-	/** 
-	 * The initially selected camera rig, specified as a pointer, if any.
-	 * This shouldn't be set if InitialCameraAssetSelectionName is set.
-	 */
-	UCameraRigAsset* InitialCameraRigSelection;
+	/** The initially selected camera rig, specified as a pointer, if any. */
+	UCameraRigAsset* InitialCameraRigSelection = nullptr;
 
-	/**
-	 * The initially selected camera rig, picked by name, if any.
-	 * This shouldn't be set if InitialCameraAssetSelection is set.
-	 */
-	FString InitialCameraRigSelectionName;
+	/** The initially selected camera rig, specified by its Guid, if any. */
+	FGuid InitialCameraRigSelectionGuid;
 
 	/** 
-	 * Whether a camera asset can be picked, or whether the InitialCameraAssetSelection
-	 * is the only asset that should be used. If the latter (when the value is false),
+	 * Whether a camera asset can be picked, or whether InitialCameraAssetSelection
+	 * determines the only asset to be used. If the latter (when the value is false),
 	 * then no camera asset picker is shown. Only the list of rigs is shown.
 	 */
 	bool bCanSelectCameraAsset = false;
 
 	/** Asset picker view type for the camera asset picker. */
 	EAssetViewType::Type CameraAssetViewType = EAssetViewType::List;
-	/** Asset picker selection mode for the camera asset picker. */
-	ESelectionMode::Type CameraAssetSelectionMode = ESelectionMode::Single;
+
 	/** Asset picker settings name for the camera asset picker. */
 	FString CameraAssetSaveSettingsName;
 
 	/** Whether the camera rig search box should be focused initially. */
 	bool bFocusCameraRigSearchBoxWhenOpened = true;
 
+	/** An optional warning message to display. */
+	FText WarningMessage;
+
+	/** An optional error message to display. */
+	FText ErrorMessage;
+
 	/** Callback for when a camera rig has been selected. */
 	FOnCameraRigSelected OnCameraRigSelected;
 
+	/** 
+	 * An optional property handle to set the selected camera rig on.
+	 * Properties of type FGuid and UCameraRigAsset* are handled.
+	 */
 	TSharedPtr<IPropertyHandle> PropertyToSet;
 };
 
diff --git a/Engine/Plugins/Cameras/GameplayCameras/Source/GameplayCamerasEditor/Public/Editors/SCameraRigNameGraphPin.h b/Engine/Plugins/Cameras/GameplayCameras/Source/GameplayCamerasEditor/Public/Editors/SCameraRigNameGraphPin.h
--- a/Engine/Plugins/Cameras/GameplayCameras/Source/GameplayCamerasEditor/Public/Editors/SCameraRigNameGraphPin.h
+++ b/Engine/Plugins/Cameras/GameplayCameras/Source/GameplayCamerasEditor/Public/Editors/SCameraRigNameGraphPin.h
@@ -12,6 +12,12 @@ class UCameraRigAsset;
 namespace UE::Cameras
 {
 
+enum class ECameraRigNameGraphPinMode
+{
+	NamePin,
+	ReferencePin
+};
+
 /**
  * A custom widget for a graph editor pin that shows a camera rig picker dialog.
  */
@@ -20,7 +26,9 @@ class SCameraRigNameGraphPin : public SGraphPin
 public:
 
 	SLATE_BEGIN_ARGS(SCameraRigNameGraphPin)
+		: _PinMode(ECameraRigNameGraphPinMode::NamePin)
 	{}
+		SLATE_ARGUMENT(ECameraRigNameGraphPinMode, PinMode)
 	SLATE_END_ARGS()
 
 	void Construct(const FArguments& InArgs, UEdGraphPin* InGraphPinObj);
@@ -53,11 +61,12 @@ private:
 
 	FReply OnResetButtonClicked();
 
-	void SetCameraRigName(const FString& InCameraRigName);
+	void SetCameraRig(UCameraRigAsset* SelectedCameraRig);
 
 private:
 
 	TSharedPtr<SComboButton> CameraRigPickerButton;
+	ECameraRigNameGraphPinMode PinMode;
 };
 
 }  // namespace UE::Cameras
diff --git a/Engine/Plugins/Cameras/GameplayCameras/Source/GameplayCamerasEditor/Public/Helpers/CameraDirectorHelper.h b/Engine/Plugins/Cameras/GameplayCameras/Source/GameplayCamerasEditor/Public/Helpers/CameraDirectorHelper.h
new file mode 100644
--- /dev/null
+++ b/Engine/Plugins/Cameras/GameplayCameras/Source/GameplayCamerasEditor/Public/Helpers/CameraDirectorHelper.h
@@ -0,0 +1,26 @@
+// Copyright Epic Games, Inc. All Rights Reserved.
+
+#pragma once
+
+#include "CoreTypes.h"
+#include "Containers/Array.h"
+
+class UCameraAsset;
+class UObject;
+
+namespace UE::Cameras
+{
+
+/**
+ * A helper class for working with camera directors.
+ */
+class FCameraDirectorHelper
+{
+public:
+
+	/** Gets the list of camera assets that reference the given object as their camera director. */
+	static void GetReferencingCameraAssets(UObject* CameraDirectorObject, TArray<UCameraAsset*>& OutReferencers);
+};
+
+}  // namespace UE::Cameras
+
