Cameras: add custom build steps to camera assets and camera rigs, and implement one for checked shared Blueprint directors

diff --git a/Engine/Plugins/Cameras/GameplayCameras/Source/GameplayCameras/Private/Core/CameraAssetBuilder.cpp b/Engine/Plugins/Cameras/GameplayCameras/Source/GameplayCameras/Private/Core/CameraAssetBuilder.cpp
--- a/Engine/Plugins/Cameras/GameplayCameras/Source/GameplayCameras/Private/Core/CameraAssetBuilder.cpp
+++ b/Engine/Plugins/Cameras/GameplayCameras/Source/GameplayCameras/Private/Core/CameraAssetBuilder.cpp
@@ -20,6 +20,11 @@ FCameraAssetBuilder::FCameraAssetBuilder(FCameraBuildLog& InBuildLog)
 }
 
 void FCameraAssetBuilder::BuildCamera(UCameraAsset* InCameraAsset)
+{
+	BuildCamera(InCameraAsset, FCustomBuildStep::CreateLambda([](UCameraAsset*, FCameraBuildLog&) {}));
+}
+	
+void FCameraAssetBuilder::BuildCamera(UCameraAsset* InCameraAsset, FCustomBuildStep InCustomBuildStep)
 {
 	if (!ensure(InCameraAsset))
 	{
@@ -30,6 +35,8 @@ void FCameraAssetBuilder::BuildCamera(UCameraAsset* InCameraAsset)
 	BuildLog.SetLoggingPrefix(InCameraAsset->GetPathName() + TEXT(": "));
 	{
 		BuildCameraImpl();
+
+		InCustomBuildStep.ExecuteIfBound(CameraAsset, BuildLog);
 	}
 	BuildLog.SetLoggingPrefix(FString());
 	UpdateBuildStatus();
diff --git a/Engine/Plugins/Cameras/GameplayCameras/Source/GameplayCameras/Private/Core/CameraRigAssetBuilder.cpp b/Engine/Plugins/Cameras/GameplayCameras/Source/GameplayCameras/Private/Core/CameraRigAssetBuilder.cpp
--- a/Engine/Plugins/Cameras/GameplayCameras/Source/GameplayCameras/Private/Core/CameraRigAssetBuilder.cpp
+++ b/Engine/Plugins/Cameras/GameplayCameras/Source/GameplayCameras/Private/Core/CameraRigAssetBuilder.cpp
@@ -205,6 +205,11 @@ FCameraRigAssetBuilder::FCameraRigAssetBuilder(FCameraBuildLog& InBuildLog)
 }
 
 void FCameraRigAssetBuilder::BuildCameraRig(UCameraRigAsset* InCameraRig)
+{
+	BuildCameraRig(InCameraRig, FCustomBuildStep::CreateLambda([](UCameraRigAsset*, FCameraBuildLog&) {}));
+}
+
+void FCameraRigAssetBuilder::BuildCameraRig(UCameraRigAsset* InCameraRig, FCustomBuildStep InCustomBuildStep)
 {
 	if (!ensure(InCameraRig))
 	{
@@ -215,6 +220,8 @@ void FCameraRigAssetBuilder::BuildCameraRig(UCameraRigAsset* InCameraRig)
 	BuildLog.SetLoggingPrefix(InCameraRig->GetPathName() + TEXT(": "));
 	{
 		BuildCameraRigImpl();
+
+		InCustomBuildStep.ExecuteIfBound(CameraRig, BuildLog);
 	}
 	BuildLog.SetLoggingPrefix(FString());
 	UpdateBuildStatus();
diff --git a/Engine/Plugins/Cameras/GameplayCameras/Source/GameplayCameras/Public/Core/CameraAssetBuilder.h b/Engine/Plugins/Cameras/GameplayCameras/Source/GameplayCameras/Public/Core/CameraAssetBuilder.h
--- a/Engine/Plugins/Cameras/GameplayCameras/Source/GameplayCameras/Public/Core/CameraAssetBuilder.h
+++ b/Engine/Plugins/Cameras/GameplayCameras/Source/GameplayCameras/Public/Core/CameraAssetBuilder.h
@@ -4,6 +4,7 @@
 
 #include "Core/CameraBuildLog.h"
 #include "CoreTypes.h"
+#include "GameplayCameras.h"
 #include "Templates/Tuple.h"
 
 class UCameraAsset;
@@ -14,16 +15,21 @@ namespace UE::Cameras
 /**
  * A class that can prepare a camera asset for runtime use.
  */
-class FCameraAssetBuilder
+class GAMEPLAYCAMERAS_API FCameraAssetBuilder
 {
 public:
 
+	DECLARE_DELEGATE_TwoParams(FCustomBuildStep, UCameraAsset*, FCameraBuildLog&);
+
 	/** Creates a new camera builder. */
 	FCameraAssetBuilder(FCameraBuildLog& InBuildLog);
 
 	/** Builds the given camera. */
 	void BuildCamera(UCameraAsset* InCameraAsset);
 
+	/** Builds the given camera. */
+	void BuildCamera(UCameraAsset* InCameraAsset, FCustomBuildStep InCustomBuildStep);
+
 private:
 
 	void BuildCameraImpl();
diff --git a/Engine/Plugins/Cameras/GameplayCameras/Source/GameplayCameras/Public/Core/CameraRigAssetBuilder.h b/Engine/Plugins/Cameras/GameplayCameras/Source/GameplayCameras/Public/Core/CameraRigAssetBuilder.h
--- a/Engine/Plugins/Cameras/GameplayCameras/Source/GameplayCameras/Public/Core/CameraRigAssetBuilder.h
+++ b/Engine/Plugins/Cameras/GameplayCameras/Source/GameplayCameras/Public/Core/CameraRigAssetBuilder.h
@@ -7,6 +7,7 @@
 #include "Core/CameraRigAsset.h"
 #include "Core/CameraVariableTableFwd.h"
 #include "CoreTypes.h"
+#include "GameplayCameras.h"
 #include "Templates/Tuple.h"
 
 class FStructProperty;
@@ -29,16 +30,21 @@ namespace Internal { struct FPrivateVariableBuilder; }
  *
  * Once the build process is done, the BuildStatus property is set on the camera rig.
  */
-class FCameraRigAssetBuilder
+class GAMEPLAYCAMERAS_API FCameraRigAssetBuilder
 {
 public:
 
+	DECLARE_DELEGATE_TwoParams(FCustomBuildStep, UCameraRigAsset*, FCameraBuildLog&);
+
 	/** Creates a new camera rig builder. */
 	FCameraRigAssetBuilder(FCameraBuildLog& InBuildLog);
 
 	/** Builds the given camera rig. */
 	void BuildCameraRig(UCameraRigAsset* InCameraRig);
 
+	/** Builds the given camera rig. */
+	void BuildCameraRig(UCameraRigAsset* InCameraRig, FCustomBuildStep InCustomBuildStep);
+
 private:
 
 	void BuildCameraRigImpl();
diff --git a/Engine/Plugins/Cameras/GameplayCameras/Source/GameplayCamerasEditor/Private/Builders/BlueprintCameraDirectorEditorBuilder.cpp b/Engine/Plugins/Cameras/GameplayCameras/Source/GameplayCamerasEditor/Private/Builders/BlueprintCameraDirectorEditorBuilder.cpp
new file mode 100644
--- /dev/null
+++ b/Engine/Plugins/Cameras/GameplayCameras/Source/GameplayCamerasEditor/Private/Builders/BlueprintCameraDirectorEditorBuilder.cpp
@@ -0,0 +1,88 @@
+// Copyright Epic Games, Inc. All Rights Reserved.
+
+#include "Builders/BlueprintCameraDirectorEditorBuilder.h"
+
+#include "Core/CameraAsset.h"
+#include "Core/CameraBuildLog.h"
+#include "Directors/BlueprintCameraDirector.h"
+#include "EdGraph/EdGraph.h"
+#include "Engine/Blueprint.h"
+#include "Helpers/CameraDirectorHelper.h"
+#include "K2Node_CallFunction.h"
+
+#define LOCTEXT_NAMESPACE "BlueprintCameraDirectorEditorBuilder"
+
+namespace UE::Cameras
+{
+
+void FBlueprintCameraDirectorEditorBuilder::OnBuildCameraAsset(UCameraAsset* CameraAsset, UE::Cameras::FCameraBuildLog& BuildLog)
+{
+	UBlueprintCameraDirector* BlueprintCameraDirector = Cast<UBlueprintCameraDirector>(CameraAsset->GetCameraDirector());
+	if (!BlueprintCameraDirector)
+	{
+		// Skip other camera directors.
+		return;
+	}
+
+	TSubclassOf<UBlueprintCameraDirectorEvaluator> CameraDirectorEvaluatorClass = BlueprintCameraDirector->CameraDirectorEvaluatorClass;
+	if (!CameraDirectorEvaluatorClass)
+	{
+		// This case is already checked by the Blueprint camera director itself.
+		return;
+	}
+
+	// If the evaluator Blueprint is shared with other camera assets, check that it only uses
+	// camera rig proxies for activation.
+	TArray<UCameraAsset*> ReferencingCameraAssets;
+	UBlueprint* Blueprint = CastChecked<UBlueprint>(CameraDirectorEvaluatorClass->ClassGeneratedBy);
+	FCameraDirectorHelper::GetReferencingCameraAssets(Blueprint, ReferencingCameraAssets);
+
+	UCameraAsset* ThisCameraAsset = BlueprintCameraDirector->GetTypedOuter<UCameraAsset>();
+	ensure(ThisCameraAsset);
+	ReferencingCameraAssets.Remove(ThisCameraAsset);
+
+	if (ReferencingCameraAssets.Num() > 0)
+	{
+		const bool bUsesDirectRigActivation = UsesDirectCameraRigActivation(Blueprint);
+		if (bUsesDirectRigActivation)
+		{
+			BuildLog.AddMessage(
+					EMessageSeverity::Error, 
+					BlueprintCameraDirector,
+					FText::Format(LOCTEXT("SharedDirectorMustUseProxies",
+						"Blueprint camera director evaluator '{0}' is shared with {1} other camera assets, but uses "
+						"ActivateCameraRig nodes. Shared Blueprints must use ActivateCameraRigViaProxy to work between "
+						"multiple camera assets."),
+						{
+							FText::FromString(CameraDirectorEvaluatorClass->GetOutermost()->GetName()),
+							ReferencingCameraAssets.Num() }));
+		}
+	}
+}
+
+bool FBlueprintCameraDirectorEditorBuilder::UsesDirectCameraRigActivation(UBlueprint* Blueprint)
+{
+	TArray<UEdGraph*> BlueprintGraphs;
+	Blueprint->GetAllGraphs(BlueprintGraphs);
+	for (const UEdGraph* BlueprintGraph : BlueprintGraphs)
+	{
+		TArray<UK2Node_CallFunction*> CallFunctionNodes;
+		BlueprintGraph->GetNodesOfClass(CallFunctionNodes);
+		for (UK2Node_CallFunction* CallFunctionNode : CallFunctionNodes)
+		{
+			// Look for calls of `this->ActivateCameraRig()`.
+			const FName FunctionName = CallFunctionNode->FunctionReference.GetMemberName();
+			if (CallFunctionNode->FunctionReference.IsSelfContext()
+					&& FunctionName == GET_FUNCTION_NAME_CHECKED(UBlueprintCameraDirectorEvaluator, ActivateCameraRig))
+			{
+				return true;
+			}
+		}
+	}
+	return false;
+}
+
+}  // namespace UE::Cameras
+
+#undef LOCTEXT_NAMESPACE
+
diff --git a/Engine/Plugins/Cameras/GameplayCameras/Source/GameplayCamerasEditor/Private/Builders/BlueprintCameraDirectorEditorBuilder.h b/Engine/Plugins/Cameras/GameplayCameras/Source/GameplayCamerasEditor/Private/Builders/BlueprintCameraDirectorEditorBuilder.h
new file mode 100644
--- /dev/null
+++ b/Engine/Plugins/Cameras/GameplayCameras/Source/GameplayCamerasEditor/Private/Builders/BlueprintCameraDirectorEditorBuilder.h
@@ -0,0 +1,29 @@
+// Copyright Epic Games, Inc. All Rights Reserved.
+
+#pragma once
+
+class UBlueprint;
+class UCameraAsset;
+
+namespace UE::Cameras
+{
+
+class FCameraBuildLog;
+
+/**
+ * Camera asset builder that does extra validation related to Blueprint camera directors.
+ */
+class FBlueprintCameraDirectorEditorBuilder
+{
+public:
+
+	/** Callback registered on IGameplayCamerasEditorModule's camera asset builders */
+	static void OnBuildCameraAsset(UCameraAsset* CameraAsset, UE::Cameras::FCameraBuildLog& BuildLog);
+
+private:
+
+	static bool UsesDirectCameraRigActivation(UBlueprint* Blueprint);
+};
+
+}  // namespace UE::Cameras
+
diff --git a/Engine/Plugins/Cameras/GameplayCameras/Source/GameplayCamerasEditor/Private/GameplayCamerasEditorModule.cpp b/Engine/Plugins/Cameras/GameplayCameras/Source/GameplayCamerasEditor/Private/GameplayCamerasEditorModule.cpp
--- a/Engine/Plugins/Cameras/GameplayCameras/Source/GameplayCamerasEditor/Private/GameplayCamerasEditorModule.cpp
+++ b/Engine/Plugins/Cameras/GameplayCameras/Source/GameplayCamerasEditor/Private/GameplayCamerasEditorModule.cpp
@@ -6,6 +6,7 @@
 #include "AssetTools/CameraRigAssetEditor.h"
 #include "AssetTools/CameraRigProxyAssetEditor.h"
 #include "AssetTools/CameraVariableCollectionEditor.h"
+#include "Builders/BlueprintCameraDirectorEditorBuilder.h"
 #include "Commands/CameraAssetEditorCommands.h"
 #include "Commands/CameraRigAssetEditorCommands.h"
 #include "Commands/CameraRigTransitionEditorCommands.h"
@@ -82,6 +83,7 @@ public:
 
 		RegisterSettings();
 		RegisterCameraDirectorEditors();
+		RegisterBuilders();
 		RegisterCoreDebugCategories();
 		RegisterRewindDebuggerFeatures();
 		RegisterDetailsCustomizations();
@@ -107,6 +109,7 @@ public:
 
 		UnregisterSettings();
 		UnregisterCameraDirectorEditors();
+		UnregisterBuilders();
 		UnregisterCoreDebugCategories();
 		UnregisterRewindDebuggerFeatures();
 		UnregisterDetailsCustomizations();
@@ -169,7 +172,7 @@ public:
 		return CameraDirectorEditorCreators;
 	}
 
-	virtual void UnregisterCameraDirectorEditor(FDelegateHandle InHandle)
+	virtual void UnregisterCameraDirectorEditor(FDelegateHandle InHandle) override
 	{
 		CameraDirectorEditorCreators.RemoveAll(
 				[=](const FOnCreateCameraDirectorAssetEditorMode& Delegate) 
@@ -178,6 +181,46 @@ public:
 				});
 	}
 
+	virtual FDelegateHandle RegisterCameraAssetBuilder(FOnBuildCameraAsset InOnBuildCameraAsset) override
+	{
+		CameraAssetBuilders.Add(InOnBuildCameraAsset);
+		return CameraAssetBuilders.Last().GetHandle();
+	}
+
+	virtual TArrayView<const FOnBuildCameraAsset> GetCameraAssetBuilders() const override
+	{
+		return CameraAssetBuilders;
+	}
+
+	virtual void UnregisterCameraAssetBuilder(FDelegateHandle InHandle) override
+	{
+		CameraAssetBuilders.RemoveAll(
+				[=](const FOnBuildCameraAsset& Delegate)
+				{
+					return Delegate.GetHandle() == InHandle;
+				});
+	}
+
+	virtual FDelegateHandle RegisterCameraRigAssetBuilder(FOnBuildCameraRigAsset InOnBuildCameraRigAsset) override
+	{
+		CameraRigAssetBuilders.Add(InOnBuildCameraRigAsset);
+		return CameraRigAssetBuilders.Last().GetHandle();
+	}
+	
+	virtual TArrayView<const FOnBuildCameraRigAsset> GetCameraRigAssetBuilders() const override
+	{
+		return CameraRigAssetBuilders;
+	}
+
+	virtual void UnregisterCameraRigAssetBuilder(FDelegateHandle InHandle) override
+	{
+		CameraRigAssetBuilders.RemoveAll(
+				[=](const FOnBuildCameraRigAsset& Delegate)
+				{
+					return Delegate.GetHandle() == InHandle;
+				});
+	}
+
 	virtual void RegisterDebugCategory(const UE::Cameras::FCameraDebugCategoryInfo& InCategoryInfo) override
 	{
 		if (!ensureMsgf(!InCategoryInfo.Name.IsEmpty(), TEXT("A debug category must at least specify a name!")))
@@ -280,6 +323,35 @@ private:
 
 	void UnregisterCameraDirectorEditors()
 	{
+		for (FDelegateHandle Handle : BuiltInDirectorCreatorHandles)
+		{
+			UnregisterCameraDirectorEditor(Handle);
+		}
+		BuiltInDirectorCreatorHandles.Reset();
+	}
+
+	void RegisterBuilders()
+	{
+		using namespace UE::Cameras;
+
+		BuiltInCameraAssetBuilders.Add(
+				RegisterCameraAssetBuilder(FOnBuildCameraAsset::CreateStatic(
+						&FBlueprintCameraDirectorEditorBuilder::OnBuildCameraAsset)));
+	}
+
+	void UnregisterBuilders()
+	{
+		for (FDelegateHandle Handle : BuiltInCameraAssetBuilders)
+		{
+			UnregisterCameraAssetBuilder(Handle);
+		}
+		BuiltInCameraAssetBuilders.Reset();
+
+		for (FDelegateHandle Handle : BuiltInCameraRigAssetBuilders)
+		{
+			UnregisterCameraRigAssetBuilder(Handle);
+		}
+		BuiltInCameraRigAssetBuilders.Reset();
 	}
 
 	void RegisterCoreDebugCategories()
@@ -506,6 +578,12 @@ private:
 	TArray<FOnCreateCameraDirectorAssetEditorMode> CameraDirectorEditorCreators;
 	TArray<FDelegateHandle> BuiltInDirectorCreatorHandles;
 
+	TArray<FOnBuildCameraAsset> CameraAssetBuilders;
+	TArray<FDelegateHandle> BuiltInCameraAssetBuilders;
+
+	TArray<FOnBuildCameraRigAsset> CameraRigAssetBuilders;
+	TArray<FDelegateHandle> BuiltInCameraRigAssetBuilders;
+
 	TSharedPtr<UE::Cameras::FGameplayCamerasGraphPanelPinFactory> GraphPanelPinFactory;
 
 	TMap<FString, UE::Cameras::FCameraDebugCategoryInfo> DebugCategoryInfos;
diff --git a/Engine/Plugins/Cameras/GameplayCameras/Source/GameplayCamerasEditor/Private/Toolkits/CameraAssetEditorToolkit.cpp b/Engine/Plugins/Cameras/GameplayCameras/Source/GameplayCamerasEditor/Private/Toolkits/CameraAssetEditorToolkit.cpp
--- a/Engine/Plugins/Cameras/GameplayCameras/Source/GameplayCamerasEditor/Private/Toolkits/CameraAssetEditorToolkit.cpp
+++ b/Engine/Plugins/Cameras/GameplayCameras/Source/GameplayCamerasEditor/Private/Toolkits/CameraAssetEditorToolkit.cpp
@@ -5,6 +5,7 @@
 #include "AssetTools/CameraAssetEditor.h"
 #include "Commands/CameraAssetEditorCommands.h"
 #include "Core/CameraAsset.h"
+#include "Core/CameraAssetBuilder.h"
 #include "Core/CameraBuildLog.h"
 #include "Core/CameraDirector.h"
 #include "Editors/ObjectTreeGraphConfig.h"
@@ -263,13 +264,34 @@ void FCameraAssetEditorToolkit::OnEditorToolkitModeActivated()
 
 void FCameraAssetEditorToolkit::OnBuild()
 {
+	using namespace UE::Cameras;
+
 	if (!CameraAsset)
 	{
 		return;
 	}
 
 	FCameraBuildLog BuildLog;
-	CameraAsset->BuildCamera(BuildLog);
+	FCameraAssetBuilder Builder(BuildLog);
+	Builder.BuildCamera(
+			CameraAsset,
+			FCameraAssetBuilder::FCustomBuildStep::CreateLambda(
+				[](UCameraAsset* InCameraAsset, FCameraBuildLog& BuildLog)
+				{
+					IGameplayCamerasEditorModule& GameplayCamerasEditorModule = IGameplayCamerasEditorModule::Get();
+					for (const FOnBuildCameraAsset& Builder : GameplayCamerasEditorModule.GetCameraAssetBuilders())
+					{
+						Builder.ExecuteIfBound(InCameraAsset, BuildLog);
+					}
+					for (UCameraRigAsset* CameraRig : InCameraAsset->GetCameraRigs())
+					{
+						for (const FOnBuildCameraRigAsset& Builder : GameplayCamerasEditorModule.GetCameraRigAssetBuilders())
+						{
+							Builder.ExecuteIfBound(CameraRig, BuildLog);
+						}
+					}
+				}));
+	
 	BuildLogToolkit->PopulateMessageListing(BuildLog);
 
 	if (CameraAsset->GetBuildStatus() != ECameraBuildStatus::Clean)
diff --git a/Engine/Plugins/Cameras/GameplayCameras/Source/GameplayCamerasEditor/Private/Toolkits/CameraRigAssetEditorToolkit.cpp b/Engine/Plugins/Cameras/GameplayCameras/Source/GameplayCamerasEditor/Private/Toolkits/CameraRigAssetEditorToolkit.cpp
--- a/Engine/Plugins/Cameras/GameplayCameras/Source/GameplayCamerasEditor/Private/Toolkits/CameraRigAssetEditorToolkit.cpp
+++ b/Engine/Plugins/Cameras/GameplayCameras/Source/GameplayCamerasEditor/Private/Toolkits/CameraRigAssetEditorToolkit.cpp
@@ -5,21 +5,23 @@
 #include "Commands/CameraRigAssetEditorCommands.h"
 #include "Core/CameraBuildLog.h"
 #include "Core/CameraRigAsset.h"
+#include "Core/CameraRigAssetBuilder.h"
 #include "Editors/CameraNodeGraphSchema.h"
 #include "Editors/CameraRigTransitionGraphSchema.h"
 #include "Editors/SCameraRigAssetEditor.h"
 #include "Editors/SFindInObjectTreeGraph.h"
 #include "Framework/Docking/LayoutExtender.h"
 #include "Framework/Docking/TabManager.h"
+#include "IGameplayCamerasEditorModule.h"
 #include "IGameplayCamerasLiveEditManager.h"
 #include "IGameplayCamerasModule.h"
 #include "Modules/ModuleManager.h"
 #include "Styles/GameplayCamerasEditorStyle.h"
+#include "ToolMenus.h"
 #include "Toolkits/BuildButtonToolkit.h"
 #include "Toolkits/CameraBuildLogToolkit.h"
 #include "Toolkits/CameraRigAssetEditorToolkitBase.h"
 #include "Toolkits/StandardToolkitLayout.h"
-#include "ToolMenus.h"
 #include "Widgets/Docking/SDockTab.h"
 
 #include UE_INLINE_GENERATED_CPP_BY_NAME(CameraRigAssetEditorToolkit)
@@ -222,7 +224,19 @@ void FCameraRigAssetEditorToolkit::OnBuild()
 	}
 
 	FCameraBuildLog BuildLog;
-	CameraRigAsset->BuildCameraRig(BuildLog);
+	FCameraRigAssetBuilder Builder(BuildLog);
+	Builder.BuildCameraRig(
+			CameraRigAsset,
+			FCameraRigAssetBuilder::FCustomBuildStep::CreateLambda(
+				[](UCameraRigAsset* InCameraRigAsset, FCameraBuildLog& BuildLog)
+				{
+					IGameplayCamerasEditorModule& GameplayCamerasEditorModule = IGameplayCamerasEditorModule::Get();
+					for (const FOnBuildCameraRigAsset& Builder : GameplayCamerasEditorModule.GetCameraRigAssetBuilders())
+					{
+						Builder.ExecuteIfBound(InCameraRigAsset, BuildLog);
+					}
+				}));
+
 	BuildLogToolkit->PopulateMessageListing(BuildLog);
 
 	if (CameraRigAsset->BuildStatus != ECameraBuildStatus::Clean)
diff --git a/Engine/Plugins/Cameras/GameplayCameras/Source/GameplayCamerasEditor/Public/IGameplayCamerasEditorModule.h b/Engine/Plugins/Cameras/GameplayCameras/Source/GameplayCamerasEditor/Public/IGameplayCamerasEditorModule.h
--- a/Engine/Plugins/Cameras/GameplayCameras/Source/GameplayCamerasEditor/Public/IGameplayCamerasEditorModule.h
+++ b/Engine/Plugins/Cameras/GameplayCameras/Source/GameplayCamerasEditor/Public/IGameplayCamerasEditorModule.h
@@ -2,6 +2,7 @@
 
 #pragma once
 
+#include "Core/CameraBuildLog.h"
 #include "CoreTypes.h"
 #include "Delegates/Delegate.h"
 #include "Modules/ModuleInterface.h"
@@ -20,6 +21,7 @@ class UCameraVariableCollectionEditor;
 namespace UE::Cameras
 {
 
+class FCameraBuildLog;
 class FCameraDirectorAssetEditorMode;
 struct FCameraRigPickerConfig;
 
@@ -34,6 +36,8 @@ struct FCameraDebugCategoryInfo
 }  // namespace UE::Cameras
 
 DECLARE_DELEGATE_RetVal_OneParam(TSharedPtr<UE::Cameras::FCameraDirectorAssetEditorMode>, FOnCreateCameraDirectorAssetEditorMode, UCameraAsset*);
+DECLARE_DELEGATE_TwoParams(FOnBuildCameraAsset, UCameraAsset*, UE::Cameras::FCameraBuildLog&);
+DECLARE_DELEGATE_TwoParams(FOnBuildCameraRigAsset, UCameraRigAsset*, UE::Cameras::FCameraBuildLog&);
 DECLARE_DELEGATE_RetVal_OneParam(TSharedRef<SWidget>, FOnCreateDebugCategoryPanel, const FString&);
 
 /**
@@ -78,6 +82,20 @@ public:
 	/** Unregisters a camera director editor creator. */
 	virtual void UnregisterCameraDirectorEditor(FDelegateHandle InHandle) = 0;
 
+	/** Registers a custom camera asset builder. */
+	virtual FDelegateHandle RegisterCameraAssetBuilder(FOnBuildCameraAsset InOnBuildCameraAsset) = 0;
+	/** Gets the registered custom camera asset builders. */
+	virtual TArrayView<const FOnBuildCameraAsset> GetCameraAssetBuilders() const = 0;
+	/** Unregisters a custom camera asset builder. */
+	virtual void UnregisterCameraAssetBuilder(FDelegateHandle InHandle) = 0;
+
+	/** Registers a custom camera rig builder. */
+	virtual FDelegateHandle RegisterCameraRigAssetBuilder(FOnBuildCameraRigAsset InOnBuildCameraRigAsset) = 0;
+	/** Gets the registered custom camera rig builders. */
+	virtual TArrayView<const FOnBuildCameraRigAsset> GetCameraRigAssetBuilders() const = 0;
+	/** Unregisters a custom camera rig builder. */
+	virtual void UnregisterCameraRigAssetBuilder(FDelegateHandle InHandle) = 0;
+
 public:
 
 	using FCameraDebugCategoryInfo = UE::Cameras::FCameraDebugCategoryInfo;
