Cameras: fix camera rigs and variables missing RF_Public to be referenceable from other assets

diff --git a/Engine/Plugins/Cameras/GameplayCameras/Source/GameplayCameras/Private/Core/CameraRigAsset.cpp b/Engine/Plugins/Cameras/GameplayCameras/Source/GameplayCameras/Private/Core/CameraRigAsset.cpp
--- a/Engine/Plugins/Cameras/GameplayCameras/Source/GameplayCameras/Private/Core/CameraRigAsset.cpp
+++ b/Engine/Plugins/Cameras/GameplayCameras/Source/GameplayCameras/Private/Core/CameraRigAsset.cpp
@@ -2,6 +2,7 @@
 
 #include "Core/CameraRigAsset.h"
 
+#include "Core/CameraAsset.h"
 #include "Core/CameraBuildLog.h"
 #include "Core/CameraNode.h"
 #include "Core/CameraParameters.h"
@@ -47,6 +48,18 @@ const FName UCameraRigAsset::TransitionsGraphName(TEXT("Transitions"));
 
 void UCameraRigAsset::PostLoad()
 {
+#if WITH_EDITOR
+
+	UCameraAsset* OuterCameraAsset = GetTypedOuter<UCameraAsset>();
+	if (OuterCameraAsset && !HasAllFlags(RF_Public | RF_Transactional))
+	{
+		Modify();
+
+		SetFlags(RF_Public | RF_Transactional);
+	}
+
+#endif
+
 #if WITH_EDITORONLY_DATA
 
 	if (GraphNodePosX_DEPRECATED != 0 || GraphNodePosY_DEPRECATED != 0)
diff --git a/Engine/Plugins/Cameras/GameplayCameras/Source/GameplayCamerasEditor/Private/Editors/SCameraRigList.cpp b/Engine/Plugins/Cameras/GameplayCameras/Source/GameplayCamerasEditor/Private/Editors/SCameraRigList.cpp
--- a/Engine/Plugins/Cameras/GameplayCameras/Source/GameplayCamerasEditor/Private/Editors/SCameraRigList.cpp
+++ b/Engine/Plugins/Cameras/GameplayCameras/Source/GameplayCamerasEditor/Private/Editors/SCameraRigList.cpp
@@ -282,7 +282,11 @@ void SCameraRigList::OnAddCameraRig()
 	CameraAsset->Modify();
 
 	const FName NewCameraRigName = MakeUniqueObjectName(CameraAsset, UCameraRigAsset::StaticClass(), TEXT("NewCameraRig"));
-	UCameraRigAsset* NewCameraRig = NewObject<UCameraRigAsset>(CameraAsset, NewCameraRigName, RF_Transactional);
+	UCameraRigAsset* NewCameraRig = NewObject<UCameraRigAsset>(
+			CameraAsset, 
+			NewCameraRigName, 
+			RF_Transactional | RF_Public  // Must be referenceable from camera directors.
+			);
 	CameraAsset->AddCameraRig(NewCameraRig);
 
 	bUpdateItemSource = true;
diff --git a/Engine/Plugins/Cameras/GameplayCameras/Source/GameplayCamerasEditor/Private/Factories/CameraAssetFactory.cpp b/Engine/Plugins/Cameras/GameplayCameras/Source/GameplayCamerasEditor/Private/Factories/CameraAssetFactory.cpp
--- a/Engine/Plugins/Cameras/GameplayCameras/Source/GameplayCamerasEditor/Private/Factories/CameraAssetFactory.cpp
+++ b/Engine/Plugins/Cameras/GameplayCameras/Source/GameplayCamerasEditor/Private/Factories/CameraAssetFactory.cpp
@@ -60,7 +60,7 @@ UObject* UCameraAssetFactory::FactoryCreateNew(UClass* Class, UObject* Parent, F
 {
 	UCameraAsset* NewCameraAsset = NewObject<UCameraAsset>(Parent, Class, Name, Flags | RF_Transactional);
 
-	UCameraRigAsset* FirstCameraRig = NewObject<UCameraRigAsset>(NewCameraAsset, TEXT("NewCameraRig"));
+	UCameraRigAsset* FirstCameraRig = NewObject<UCameraRigAsset>(NewCameraAsset, TEXT("NewCameraRig"), RF_Transactional | RF_Public);
 	NewCameraAsset->AddCameraRig(FirstCameraRig);
 
 	if (CameraDirectorClass)
diff --git a/Engine/Plugins/Cameras/GameplayCameras/Source/GameplayCamerasEditor/Private/Toolkits/CameraVariableCollectionEditorToolkit.cpp b/Engine/Plugins/Cameras/GameplayCameras/Source/GameplayCamerasEditor/Private/Toolkits/CameraVariableCollectionEditorToolkit.cpp
--- a/Engine/Plugins/Cameras/GameplayCameras/Source/GameplayCamerasEditor/Private/Toolkits/CameraVariableCollectionEditorToolkit.cpp
+++ b/Engine/Plugins/Cameras/GameplayCameras/Source/GameplayCamerasEditor/Private/Toolkits/CameraVariableCollectionEditorToolkit.cpp
@@ -462,7 +462,12 @@ FLinearColor FCameraVariableCollectionEditorToolkit::GetWorldCentricTabColorScal
 
 void FCameraVariableCollectionEditorToolkit::OnCreateVariable(TSubclassOf<UCameraVariableAsset> InVariableClass)
 {
-	UCameraVariableAsset* NewVariable = NewObject<UCameraVariableAsset>(VariableCollection, InVariableClass.Get());
+	UCameraVariableAsset* NewVariable = NewObject<UCameraVariableAsset>(
+			VariableCollection, 
+			InVariableClass.Get(),
+			NAME_None,
+			RF_Transactional | RF_Public  // Must be referenceable by camera parameters.
+			);
 	VariableCollection->Variables.Add(NewVariable);
 
 	VariableCollectionEditorWidget->RequestListRefresh();
