Cameras: add evaluation services and observers to the camera system

- Evaluation services are classes that can run arbitrary logic before and after the main update.
- They can also respond to "events" happening in the node tree such as camera rigs being pushed/popped from a blend stack.
- An observer API is used to accomplish this on blend stacks and root nodes.

diff --git a/Engine/Plugins/Cameras/GameplayCameras/Source/GameplayCameras/Private/Core/BlendStackCameraNode.cpp b/Engine/Plugins/Cameras/GameplayCameras/Source/GameplayCameras/Private/Core/BlendStackCameraNode.cpp
--- a/Engine/Plugins/Cameras/GameplayCameras/Source/GameplayCameras/Private/Core/BlendStackCameraNode.cpp
+++ b/Engine/Plugins/Cameras/GameplayCameras/Source/GameplayCameras/Private/Core/BlendStackCameraNode.cpp
@@ -3,6 +3,7 @@
 #include "Core/BlendStackCameraNode.h"
 
 #include "Core/BlendCameraNode.h"
+#include "Core/BlendStackCameraNodeObserver.h"
 #include "Core/BlendStackRootCameraNode.h"
 #include "Core/CameraAsset.h"
 #include "Core/CameraEvaluationContext.h"
@@ -71,6 +72,7 @@ void FBlendStackCameraNodeEvaluator::Push(const FBlendStackCameraPushParams& Par
 	// blend node for this transition.
 	// We need to const-cast here to be able to use our own blend stack node as the outer
 	// of the new node.
+	const UCameraRigTransition* UsedTransition = nullptr;
 	UObject* Outer = const_cast<UObject*>((UObject*)GetCameraNode());
 	UBlendStackRootCameraNode* EntryRootNode = NewObject<UBlendStackRootCameraNode>(Outer, NAME_None);
 	{
@@ -82,6 +84,7 @@ void FBlendStackCameraNodeEvaluator::Push(const FBlendStackCameraPushParams& Par
 		if (const UCameraRigTransition* Transition = FindTransition(Params))
 		{
 			ModeBlend = Transition->Blend;
+			UsedTransition = Transition;
 		}
 		if (!ModeBlend)
 		{
@@ -121,6 +124,12 @@ void FBlendStackCameraNodeEvaluator::Push(const FBlendStackCameraPushParams& Par
 	// Important: we need to move the new entry here because copying evaluator storage
 	// is disabled.
 	Entries.Add(MoveTemp(NewEntry));
+
+	// Notify observers.
+	if (!Observers.IsEmpty())
+	{
+		NotifyObservers(EBlendStackCameraRigEventType::Pushed, Entries.Last(), UsedTransition);
+	}
 }
 
 bool FBlendStackCameraNodeEvaluator::InitializeEntry(
@@ -190,7 +199,7 @@ void FBlendStackCameraNodeEvaluator::GatherEntryParameterEvaluators(FCameraNodeE
 	}
 }
 
-TOptional<FCameraRigEvaluationInfo> FBlendStackCameraNodeEvaluator::GetActiveCameraRigEvaluationInfo() const
+FCameraRigEvaluationInfo FBlendStackCameraNodeEvaluator::GetActiveCameraRigEvaluationInfo() const
 {
 	if (Entries.Num() > 0)
 	{
@@ -203,7 +212,7 @@ TOptional<FCameraRigEvaluationInfo> FBlendStackCameraNodeEvaluator::GetActiveCam
 		Info.bIsFrozen = ActiveEntry.bIsFrozen;
 		return Info;
 	}
-	return TOptional<FCameraRigEvaluationInfo>();
+	return FCameraRigEvaluationInfo();
 }
 
 FCameraNodeEvaluatorChildrenView FBlendStackCameraNodeEvaluator::OnGetChildren()
@@ -466,8 +475,9 @@ void FBlendStackCameraNodeEvaluator::PopEntries(int32 FirstIndexToKeep)
 
 	for (int32 Index = 0; Index < FirstIndexToKeep; ++Index)
 	{
-#if WITH_EDITOR
 		const FCameraRigEntry& FirstEntry = Entries[0];
+
+#if WITH_EDITOR
 		for (const UPackage* ListenPackage : FirstEntry.ListenedPackages)
 		{
 			int32* NumListens = AllListenedPackages.Find(ListenPackage);
@@ -483,6 +493,11 @@ void FBlendStackCameraNodeEvaluator::PopEntries(int32 FirstIndexToKeep)
 		}
 #endif  // WITH_EDITOR
 
+		if (!Observers.IsEmpty())
+		{
+			NotifyObservers(EBlendStackCameraRigEventType::Popped, FirstEntry);
+		}
+
 		Entries.RemoveAt(0);
 	}
 }
@@ -603,6 +618,35 @@ const UCameraRigTransition* FBlendStackCameraNodeEvaluator::FindTransition(
 	return nullptr;
 }
 
+void FBlendStackCameraNodeEvaluator::RegisterObserver(IBlendStackCameraNodeObserver* Observer)
+{
+	Observers.Add(Observer);
+}
+
+void FBlendStackCameraNodeEvaluator::UnregisterObserver(IBlendStackCameraNodeObserver* Observer)
+{
+	Observers.Remove(Observer);
+}
+
+void FBlendStackCameraNodeEvaluator::NotifyObservers(EBlendStackCameraRigEventType EventType, const FCameraRigEntry& Entry, const UCameraRigTransition* Transition) const
+{
+	FBlendStackCameraRigEvent Event;
+	Event.EventType = EventType;
+	Event.BlendStackEvaluator = this;
+	Event.CameraRigInfo = FCameraRigEvaluationInfo(
+			Entry.CameraRig,
+			Entry.EvaluationContext.Pin(),
+			Entry.Result,
+			Entry.RootEvaluator);
+	Event.CameraRigInfo.bIsFrozen = Entry.bIsFrozen;
+	Event.Transition = Transition;
+
+	for (IBlendStackCameraNodeObserver* Observer : Observers)
+	{
+		Observer->OnBlendStackEvent(Event);
+	}
+}
+
 void FBlendStackCameraNodeEvaluator::OnAddReferencedObjects(FReferenceCollector& Collector)
 {
 	for (FCameraRigEntry& Entry : Entries)
diff --git a/Engine/Plugins/Cameras/GameplayCameras/Source/GameplayCameras/Private/Core/CameraEvaluationService.cpp b/Engine/Plugins/Cameras/GameplayCameras/Source/GameplayCameras/Private/Core/CameraEvaluationService.cpp
new file mode 100644
--- /dev/null
+++ b/Engine/Plugins/Cameras/GameplayCameras/Source/GameplayCameras/Private/Core/CameraEvaluationService.cpp
@@ -0,0 +1,53 @@
+// Copyright Epic Games, Inc. All Rights Reserved.
+
+#include "Core/CameraEvaluationService.h"
+
+namespace UE::Cameras
+{
+
+UE_GAMEPLAY_CAMERAS_DEFINE_RTTI(FCameraEvaluationService)
+
+FCameraEvaluationService::FCameraEvaluationService()
+{
+}
+
+FCameraEvaluationService::~FCameraEvaluationService()
+{
+}
+
+void FCameraEvaluationService::Initialize(const FCameraEvaluationServiceInitializeParams& Params)
+{
+	OnInitialize(Params);
+}
+
+void FCameraEvaluationService::PreUpdate(const FCameraEvaluationServiceUpdateParams& Params, FCameraEvaluationServiceUpdateResult& OutResult)
+{
+	OnPreUpdate(Params, OutResult);
+}
+
+void FCameraEvaluationService::PostUpdate(const FCameraEvaluationServiceUpdateParams& Params, FCameraEvaluationServiceUpdateResult& OutResult)
+{
+	OnPostUpdate(Params, OutResult);
+}
+
+void FCameraEvaluationService::Teardown(const FCameraEvaluationServiceTeardownParams& Params)
+{
+	OnTeardown(Params);
+}
+
+void FCameraEvaluationService::NotifyRootCameraNodeEvent(const FRootCameraNodeCameraRigEvent& InEvent)
+{
+	OnRootCameraNodeEvent(InEvent);
+}
+
+bool FCameraEvaluationService::HasAnyEvaluationServiceFlags(ECameraEvaluationServiceFlags InFlags) const
+{
+	return EnumHasAnyFlags(PrivateFlags, InFlags);
+}
+
+void FCameraEvaluationService::SetEvaluationServiceFlags(ECameraEvaluationServiceFlags InFlags)
+{
+	PrivateFlags = InFlags;
+}
+
+}  // namespace UE::Cameras
diff --git a/Engine/Plugins/Cameras/GameplayCameras/Source/GameplayCameras/Private/Core/CameraSystemEvaluator.cpp b/Engine/Plugins/Cameras/GameplayCameras/Source/GameplayCameras/Private/Core/CameraSystemEvaluator.cpp
--- a/Engine/Plugins/Cameras/GameplayCameras/Source/GameplayCameras/Private/Core/CameraSystemEvaluator.cpp
+++ b/Engine/Plugins/Cameras/GameplayCameras/Source/GameplayCameras/Private/Core/CameraSystemEvaluator.cpp
@@ -6,6 +6,7 @@
 #include "Core/CameraAsset.h"
 #include "Core/CameraDirectorEvaluator.h"
 #include "Core/CameraEvaluationContext.h"
+#include "Core/CameraEvaluationService.h"
 #include "Core/CameraRigAsset.h"
 #include "Core/DefaultRootCameraNode.h"
 #include "Debug/CameraDebugBlock.h"
@@ -69,6 +70,19 @@ void FCameraSystemEvaluator::Initialize(const FCameraSystemEvaluatorCreateParams
 	}
 }
 
+FCameraSystemEvaluator::~FCameraSystemEvaluator()
+{
+	{
+		FCameraEvaluationServiceTeardownParams TeardownParams;
+		TeardownParams.Evaluator = this;
+		for (TSharedPtr<FCameraEvaluationService> EvaluationService : EvaluationServices)
+		{
+			EvaluationService->Teardown(TeardownParams);
+		}
+		EvaluationServices.Reset();
+	}
+}
+
 void FCameraSystemEvaluator::AddReferencedObjects(FReferenceCollector& Collector)
 {
 	Collector.AddReferencedObject(RootNode);
@@ -94,10 +108,58 @@ void FCameraSystemEvaluator::PopEvaluationContext()
 	ContextStack.PopContext();
 }
 
+void FCameraSystemEvaluator::RegisterEvaluationService(TSharedRef<FCameraEvaluationService> EvaluationService)
+{
+	EvaluationServices.Add(EvaluationService);
+	{
+		FCameraEvaluationServiceInitializeParams InitParams;
+		InitParams.Evaluator = this;
+		EvaluationService->Initialize(InitParams);
+	}
+}
+
+void FCameraSystemEvaluator::UnregisterEvaluationService(TSharedRef<FCameraEvaluationService> EvaluationService)
+{
+	{
+		FCameraEvaluationServiceTeardownParams TeardownParams;
+		TeardownParams.Evaluator = this;
+		EvaluationService->Teardown(TeardownParams);
+	}
+	EvaluationServices.Remove(EvaluationService);
+}
+
+void FCameraSystemEvaluator::NotifyRootCameraNodeEvent(const FRootCameraNodeCameraRigEvent& InEvent)
+{
+	for (TSharedPtr<FCameraEvaluationService> EvaluationService : EvaluationServices)
+	{
+		if (EvaluationService->HasAnyEvaluationServiceFlags(ECameraEvaluationServiceFlags::NeedsRootCameraNodeEvents))
+		{
+			EvaluationService->NotifyRootCameraNodeEvent(InEvent);
+		}
+	}
+}
+
 void FCameraSystemEvaluator::Update(const FCameraSystemEvaluationUpdateParams& Params)
 {
 	SCOPE_CYCLE_COUNTER(CameraSystemEval_Total);
 
+	// Pre-update all services.
+	{
+		FCameraEvaluationServiceUpdateParams ServiceUpdateParams;
+		ServiceUpdateParams.Evaluator = this;
+		ServiceUpdateParams.DeltaTime = Params.DeltaTime;
+
+		FCameraEvaluationServiceUpdateResult ServiceUpdateResult(RootNodeResult);
+
+		for (TSharedPtr<FCameraEvaluationService> EvaluationService : EvaluationServices)
+		{
+			if (EvaluationService->HasAnyEvaluationServiceFlags(ECameraEvaluationServiceFlags::NeedsPreUpdate))
+			{
+				EvaluationService->PreUpdate(ServiceUpdateParams, ServiceUpdateResult);
+			}
+		}
+	}
+
 	// Get the active evaluation context.
 	TSharedPtr<FCameraEvaluationContext> ActiveContext = ContextStack.GetActiveContext();
 	if (UNLIKELY(!ActiveContext.IsValid()))
@@ -145,6 +207,23 @@ void FCameraSystemEvaluator::Update(const FCameraSystemEvaluationUpdateParams& P
 	Result.CameraPose = RootNodeResult.CameraPose;
 	Result.bIsCameraCut = RootNodeResult.bIsCameraCut;
 	Result.bIsValid = true;
+
+	// Post-update all services.
+	{
+		FCameraEvaluationServiceUpdateParams ServiceUpdateParams;
+		ServiceUpdateParams.Evaluator = this;
+		ServiceUpdateParams.DeltaTime = Params.DeltaTime;
+
+		FCameraEvaluationServiceUpdateResult ServiceUpdateResult(RootNodeResult);
+
+		for (TSharedPtr<FCameraEvaluationService> EvaluationService : EvaluationServices)
+		{
+			if (EvaluationService->HasAnyEvaluationServiceFlags(ECameraEvaluationServiceFlags::NeedsPostUpdate))
+			{
+				EvaluationService->PostUpdate(ServiceUpdateParams, ServiceUpdateResult);
+			}
+		}
+	}
 }
 
 void FCameraSystemEvaluator::GetEvaluatedCameraView(FMinimalViewInfo& DesiredView)
diff --git a/Engine/Plugins/Cameras/GameplayCameras/Source/GameplayCameras/Private/Core/DefaultRootCameraNode.cpp b/Engine/Plugins/Cameras/GameplayCameras/Source/GameplayCameras/Private/Core/DefaultRootCameraNode.cpp
--- a/Engine/Plugins/Cameras/GameplayCameras/Source/GameplayCameras/Private/Core/DefaultRootCameraNode.cpp
+++ b/Engine/Plugins/Cameras/GameplayCameras/Source/GameplayCameras/Private/Core/DefaultRootCameraNode.cpp
@@ -4,6 +4,7 @@
 
 #include "Core/BlendStackCameraNode.h"
 #include "Core/CameraRigAsset.h"
+#include "Core/RootCameraNodeObserver.h"
 #include "Debug/BlendStacksCameraDebugBlock.h"
 #include "Debug/CameraDebugBlockBuilder.h"
 #include "Debug/RootCameraDebugBlock.h"
@@ -51,10 +52,17 @@ UE_DEFINE_CAMERA_NODE_EVALUATOR(FDefaultRootCameraNodeEvaluator)
 void FDefaultRootCameraNodeEvaluator::OnBuild(const FCameraNodeEvaluatorBuildParams& Params)
 {
 	const UDefaultRootCameraNode* Data = GetCameraNodeAs<UDefaultRootCameraNode>();
-	BaseLayer = Params.BuildEvaluatorAs<FBlendStackCameraNodeEvaluator>(Data->BaseLayer);
-	MainLayer = Params.BuildEvaluatorAs<FBlendStackCameraNodeEvaluator>(Data->MainLayer);
-	GlobalLayer = Params.BuildEvaluatorAs<FBlendStackCameraNodeEvaluator>(Data->GlobalLayer);
-	VisualLayer = Params.BuildEvaluatorAs<FBlendStackCameraNodeEvaluator>(Data->VisualLayer);
+	BaseLayer = BuildBlendStackEvaluator(Params, Data->BaseLayer);
+	MainLayer = BuildBlendStackEvaluator(Params, Data->MainLayer);
+	GlobalLayer = BuildBlendStackEvaluator(Params, Data->GlobalLayer);
+	VisualLayer = BuildBlendStackEvaluator(Params, Data->VisualLayer);
+}
+
+FBlendStackCameraNodeEvaluator* FDefaultRootCameraNodeEvaluator::BuildBlendStackEvaluator(const FCameraNodeEvaluatorBuildParams& Params, UBlendStackCameraNode* BlendStackNode)
+{
+	FBlendStackCameraNodeEvaluator* BlendStackEvaluator = Params.BuildEvaluatorAs<FBlendStackCameraNodeEvaluator>(BlendStackNode);
+	BlendStackEvaluator->RegisterObserver(this);
+	return BlendStackEvaluator;
 }
 
 FCameraNodeEvaluatorChildrenView FDefaultRootCameraNodeEvaluator::OnGetChildren()
@@ -99,6 +107,46 @@ void FDefaultRootCameraNodeEvaluator::OnActivateCameraRig(const FActivateCameraR
 	}
 }
 
+void FDefaultRootCameraNodeEvaluator::OnBlendStackEvent(const FBlendStackCameraRigEvent& InEvent)
+{
+	if (InEvent.EventType == EBlendStackCameraRigEventType::Pushed ||
+			InEvent.EventType == EBlendStackCameraRigEventType::Popped)
+	{
+		FRootCameraNodeCameraRigEvent RootEvent;
+		RootEvent.CameraRigInfo = InEvent.CameraRigInfo;
+		RootEvent.Transition = InEvent.Transition;
+
+		switch (InEvent.EventType)
+		{
+		case EBlendStackCameraRigEventType::Pushed:
+			RootEvent.EventType = ERootCameraNodeCameraRigEventType::Activated;
+			break;
+		case EBlendStackCameraRigEventType::Popped:
+			RootEvent.EventType = ERootCameraNodeCameraRigEventType::Deactivated;
+			break;
+		}
+
+		if (InEvent.BlendStackEvaluator == BaseLayer)
+		{
+			RootEvent.EventLayer = ECameraRigLayer::Base;
+		}
+		else if (InEvent.BlendStackEvaluator == MainLayer)
+		{
+			RootEvent.EventLayer = ECameraRigLayer::Main;
+		}
+		else if (InEvent.BlendStackEvaluator == GlobalLayer)
+		{
+			RootEvent.EventLayer = ECameraRigLayer::Global;
+		}
+		else if (InEvent.BlendStackEvaluator == VisualLayer)
+		{
+			RootEvent.EventLayer = ECameraRigLayer::Visual;
+		}
+
+		NotifyObservers(RootEvent);
+	}
+}
+
 #if UE_GAMEPLAY_CAMERAS_DEBUG
 
 UE_DECLARE_CAMERA_DEBUG_BLOCK_START(GAMEPLAYCAMERAS_API, FDefaultRootCameraNodeEvaluatorDebugBlock)
diff --git a/Engine/Plugins/Cameras/GameplayCameras/Source/GameplayCameras/Private/Core/RootCameraNode.cpp b/Engine/Plugins/Cameras/GameplayCameras/Source/GameplayCameras/Private/Core/RootCameraNode.cpp
--- a/Engine/Plugins/Cameras/GameplayCameras/Source/GameplayCameras/Private/Core/RootCameraNode.cpp
+++ b/Engine/Plugins/Cameras/GameplayCameras/Source/GameplayCameras/Private/Core/RootCameraNode.cpp
@@ -2,17 +2,51 @@
 
 #include "Core/RootCameraNode.h"
 
-#include "Core/CameraRigAsset.h"
+#include "Core/CameraSystemEvaluator.h"
+#include "Core/RootCameraNodeObserver.h"
 
 #include UE_INLINE_GENERATED_CPP_BY_NAME(RootCameraNode)
 
 namespace UE::Cameras
 {
 
+void FRootCameraNodeEvaluator::OnInitialize(const FCameraNodeEvaluatorInitializeParams& Params)
+{
+	OwningEvaluator = Params.Evaluator;
+}
+
 void FRootCameraNodeEvaluator::ActivateCameraRig(const FActivateCameraRigParams& Params)
 {
 	OnActivateCameraRig(Params);
 }
 
+void FRootCameraNodeEvaluator::RegisterObserver(IRootCameraNodeObserver* Observer)
+{
+	Observers.Add(Observer);
+}
+
+void FRootCameraNodeEvaluator::UnregisterObserver(IRootCameraNodeObserver* Observer)
+{
+	Observers.Remove(Observer);
+}
+
+bool FRootCameraNodeEvaluator::HasObservers() const
+{
+	return !Observers.IsEmpty();
+}
+
+void FRootCameraNodeEvaluator::NotifyObservers(const FRootCameraNodeCameraRigEvent& InEvent) const
+{
+	if (ensure(OwningEvaluator))
+	{
+		OwningEvaluator->NotifyRootCameraNodeEvent(InEvent);
+	}
+
+	for (IRootCameraNodeObserver* Observer : Observers)
+	{
+		Observer->OnRootCameraNodeEvent(InEvent);
+	}
+}
+
 }  // namespace UE::Cameras
 
diff --git a/Engine/Plugins/Cameras/GameplayCameras/Source/GameplayCameras/Private/Nodes/Input/CameraRigInput1DSlot.cpp b/Engine/Plugins/Cameras/GameplayCameras/Source/GameplayCameras/Private/Nodes/Input/CameraRigInput1DSlot.cpp
--- a/Engine/Plugins/Cameras/GameplayCameras/Source/GameplayCameras/Private/Nodes/Input/CameraRigInput1DSlot.cpp
+++ b/Engine/Plugins/Cameras/GameplayCameras/Source/GameplayCameras/Private/Nodes/Input/CameraRigInput1DSlot.cpp
@@ -33,10 +33,10 @@ void FCameraRigInput1DSlotEvaluator::OnInitialize(const FCameraNodeEvaluatorInit
 	InputValue = 0.f;
 
 	const UCameraRigInput1DSlot* SlotNode = GetCameraNodeAs<UCameraRigInput1DSlot>();
-	if (SlotNode->GetVariableID().IsValid() && Params.LastActiveCameraRig.IsSet())
+	if (SlotNode->GetVariableID().IsValid() && Params.LastActiveCameraRig.IsValid())
 	{
-		const FCameraRigEvaluationInfo& LastActiveRig = Params.LastActiveCameraRig.GetValue();
-		LastActiveRig.LastResult.VariableTable.TryGetValue<double>(SlotNode->GetVariableID(), InputValue);
+		const FCameraVariableTable& LastActiveRigVariableTable = Params.LastActiveCameraRig.LastResult->VariableTable;
+		LastActiveRigVariableTable.TryGetValue<double>(SlotNode->GetVariableID(), InputValue);
 	}
 }
 
diff --git a/Engine/Plugins/Cameras/GameplayCameras/Source/GameplayCameras/Private/Nodes/Input/CameraRigInput2DSlot.cpp b/Engine/Plugins/Cameras/GameplayCameras/Source/GameplayCameras/Private/Nodes/Input/CameraRigInput2DSlot.cpp
--- a/Engine/Plugins/Cameras/GameplayCameras/Source/GameplayCameras/Private/Nodes/Input/CameraRigInput2DSlot.cpp
+++ b/Engine/Plugins/Cameras/GameplayCameras/Source/GameplayCameras/Private/Nodes/Input/CameraRigInput2DSlot.cpp
@@ -33,10 +33,10 @@ void FCameraRigInput2DSlotEvaluator::OnInitialize(const FCameraNodeEvaluatorInit
 	InputValue = FVector2d::ZeroVector;
 
 	const UCameraRigInput2DSlot* SlotNode = GetCameraNodeAs<UCameraRigInput2DSlot>();
-	if (SlotNode->GetVariableID().IsValid() && Params.LastActiveCameraRig.IsSet())
+	if (SlotNode->GetVariableID().IsValid() && Params.LastActiveCameraRig.IsValid())
 	{
-		const FCameraRigEvaluationInfo& LastActiveRig = Params.LastActiveCameraRig.GetValue();
-		LastActiveRig.LastResult.VariableTable.TryGetValue<FVector2d>(SlotNode->GetVariableID(), InputValue);
+		const FCameraVariableTable& LastActiveRigVariableTable = Params.LastActiveCameraRig.LastResult->VariableTable;
+		LastActiveRigVariableTable.TryGetValue<FVector2d>(SlotNode->GetVariableID(), InputValue);
 	}
 }
 
diff --git a/Engine/Plugins/Cameras/GameplayCameras/Source/GameplayCameras/Public/Core/BlendStackCameraNode.h b/Engine/Plugins/Cameras/GameplayCameras/Source/GameplayCameras/Public/Core/BlendStackCameraNode.h
--- a/Engine/Plugins/Cameras/GameplayCameras/Source/GameplayCameras/Public/Core/BlendStackCameraNode.h
+++ b/Engine/Plugins/Cameras/GameplayCameras/Source/GameplayCameras/Public/Core/BlendStackCameraNode.h
@@ -23,6 +23,8 @@ namespace UE::Cameras
 class FBlendStackRootCameraNodeEvaluator;
 class FCameraEvaluationContext;
 class FCameraSystemEvaluator;
+class IBlendStackCameraNodeObserver;
+enum class EBlendStackCameraRigEventType;
 
 #if UE_GAMEPLAY_CAMERAS_DEBUG
 class FBlendStackCameraDebugBlock;
@@ -96,12 +98,19 @@ public:
 	void Push(const FBlendStackCameraPushParams& Params);
 
 	/** Returns information about the top (active) camera rig, if any. */
-	TOptional<FCameraRigEvaluationInfo> GetActiveCameraRigEvaluationInfo() const;
+	FCameraRigEvaluationInfo GetActiveCameraRigEvaluationInfo() const;
 
 #if UE_GAMEPLAY_CAMERAS_DEBUG
 	FBlendStackCameraDebugBlock* BuildDetailedDebugBlock(const FCameraDebugBlockBuildParams& Params, FCameraDebugBlockBuilder& Builder);
 #endif  // UE_GAMEPLAY_CAMERAS_DEBUG
 
+public:
+
+	/** Register a blend stack observer. */
+	void RegisterObserver(IBlendStackCameraNodeObserver* Observer);
+	/** Unregister a blend stack observer. */
+	void UnregisterObserver(IBlendStackCameraNodeObserver* Observer);
+
 protected:
 
 	// FCameraNodeEvaluator interface
@@ -141,6 +150,8 @@ protected:
 
 	void GatherEntryParameterEvaluators(FCameraNodeEvaluator* RootEvaluator, TArray<FCameraNodeEvaluator*>& OutParameterEvaluators);
 
+	void NotifyObservers(EBlendStackCameraRigEventType EventType, const FCameraRigEntry& Entry, const UCameraRigTransition* Transition = nullptr) const;
+
 protected:
 
 	struct FCameraRigEntry
@@ -183,6 +194,9 @@ protected:
 	/** Entries in the blend stack. */
 	TArray<FCameraRigEntry> Entries;
 
+	/** The list of observers for this blend stack. */
+	TArray<IBlendStackCameraNodeObserver*> Observers;
+
 #if WITH_EDITOR
 	TMap<const UPackage*, int32> AllListenedPackages;
 #endif  // WITH_EDITOR
diff --git a/Engine/Plugins/Cameras/GameplayCameras/Source/GameplayCameras/Public/Core/BlendStackCameraNodeObserver.h b/Engine/Plugins/Cameras/GameplayCameras/Source/GameplayCameras/Public/Core/BlendStackCameraNodeObserver.h
new file mode 100644
--- /dev/null
+++ b/Engine/Plugins/Cameras/GameplayCameras/Source/GameplayCameras/Public/Core/BlendStackCameraNodeObserver.h
@@ -0,0 +1,55 @@
+// Copyright Epic Games, Inc. All Rights Reserved.
+
+#pragma once
+
+#include "Core/CameraRigEvaluationInfo.h"
+#include "UObject/ObjectPtr.h"
+
+class UCameraRigTransition;
+
+namespace UE::Cameras
+{
+
+class FBlendStackCameraNodeEvaluator;
+
+/** A type of blend stack event. */
+enum class EBlendStackCameraRigEventType
+{
+	/** A camera rig was pushed on the blend stack. */
+	Pushed,
+	/** A camera rig was frozen. */
+	Frozen,
+	/** A camera rig was popped out of the blend stack. */
+	Popped
+};
+
+/**
+ * A structure describing an event happening on a blend stack.
+ */
+struct FBlendStackCameraRigEvent
+{
+	/** The event type. */
+	EBlendStackCameraRigEventType EventType;
+	/** The blend stack evaluator in which the event is happening. */
+	const FBlendStackCameraNodeEvaluator* BlendStackEvaluator = nullptr;
+	/** The evaluation information for the camera rig associated with the event. */
+	FCameraRigEvaluationInfo CameraRigInfo;
+	/** If the event is Pushed, the transition used for the new camera rig. */
+	TObjectPtr<const UCameraRigTransition> Transition;
+};
+
+/**
+ * An interface for observing blend stack events.
+ */
+class IBlendStackCameraNodeObserver
+{
+public:
+
+	virtual ~IBlendStackCameraNodeObserver() {}
+
+	/** Called when a relevant event happens in the observed blend stack. */
+	virtual void OnBlendStackEvent(const FBlendStackCameraRigEvent& InEvent) {}
+};
+
+}  // namespace UE::Cameras
+
diff --git a/Engine/Plugins/Cameras/GameplayCameras/Source/GameplayCameras/Public/Core/CameraEvaluationService.h b/Engine/Plugins/Cameras/GameplayCameras/Source/GameplayCameras/Public/Core/CameraEvaluationService.h
new file mode 100644
--- /dev/null
+++ b/Engine/Plugins/Cameras/GameplayCameras/Source/GameplayCameras/Public/Core/CameraEvaluationService.h
@@ -0,0 +1,117 @@
+// Copyright Epic Games, Inc. All Rights Reserved.
+
+#pragma once
+
+#include "Core/CameraObjectRtti.h"
+#include "GameplayCameras.h"
+#include "Templates/SharedPointer.h"
+
+namespace UE::Cameras
+{
+
+class FCameraSystemEvaluator;
+struct FCameraNodeEvaluationResult;
+struct FRootCameraNodeCameraRigEvent;
+
+/** Flags for what callbacks an evaluation service wants to opt-into. */
+enum class ECameraEvaluationServiceFlags
+{
+	None = 0,
+	NeedsPreUpdate = 1 << 0,
+	NeedsPostUpdate = 1 << 1,
+	NeedsRootCameraNodeEvents = 1 << 2
+};
+ENUM_CLASS_FLAGS(ECameraEvaluationServiceFlags);
+
+/** Parameter structure for initializing an evaluation service. */
+struct FCameraEvaluationServiceInitializeParams
+{
+	FCameraSystemEvaluator* Evaluator = nullptr;
+};
+
+/** Parameter structure for tearing down an evaluation service. */
+struct FCameraEvaluationServiceTeardownParams
+{
+	FCameraSystemEvaluator* Evaluator = nullptr;
+};
+
+/** Parameter structure for updating an evaluation service. */
+struct FCameraEvaluationServiceUpdateParams
+{
+	FCameraSystemEvaluator* Evaluator = nullptr;
+	float DeltaTime = 0.f;
+};
+
+/** Result structure for updating an evaluation service. */
+struct FCameraEvaluationServiceUpdateResult
+{
+	FCameraEvaluationServiceUpdateResult(FCameraNodeEvaluationResult& InEvaluationResult)
+		: EvaluationResult(InEvaluationResult)
+	{}
+
+	FCameraNodeEvaluationResult& EvaluationResult;
+};
+
+/**
+ * An evaluation service running on a camera system.
+ *
+ * Evaluation services can run arbitrary logic before or after the root camera node update,
+ * and respond to events in the node tree such as when camera rigs are activated or deactivated.
+ */
+class FCameraEvaluationService : public TSharedFromThis<FCameraEvaluationService>
+{
+	UE_GAMEPLAY_CAMERAS_DECLARE_RTTI_BASE(GAMEPLAYCAMERAS_API, FCameraEvaluationService)
+
+public:
+
+	GAMEPLAYCAMERAS_API FCameraEvaluationService();
+
+	GAMEPLAYCAMERAS_API virtual ~FCameraEvaluationService();
+
+public:
+
+	/** Initializes the evaluation service. */
+	void Initialize(const FCameraEvaluationServiceInitializeParams& Params);
+	/** Runs at the start of the camera system update. */
+	void PreUpdate(const FCameraEvaluationServiceUpdateParams& Params, FCameraEvaluationServiceUpdateResult& OutResult);
+	/** Runs at the end of the camera system update. */
+	void PostUpdate(const FCameraEvaluationServiceUpdateParams& Params, FCameraEvaluationServiceUpdateResult& OutResult);
+	/** Tears down the evaluation service. */
+	void Teardown(const FCameraEvaluationServiceTeardownParams& Params);
+
+public:
+
+	// Internal API.
+
+	void NotifyRootCameraNodeEvent(const FRootCameraNodeCameraRigEvent& InEvent);
+
+	ECameraEvaluationServiceFlags GetEvaluationServiceFlags() const { return PrivateFlags; }
+	bool HasAnyEvaluationServiceFlags(ECameraEvaluationServiceFlags InFlags) const;
+
+protected:
+
+	/** Initializes the evaluation service. */
+	virtual void OnInitialize(const FCameraEvaluationServiceInitializeParams& Params) {}
+	/** Runs at the start of the camera system update. */
+	virtual void OnPreUpdate(const FCameraEvaluationServiceUpdateParams& Params, FCameraEvaluationServiceUpdateResult& OutResult) {}
+	/** Runs at the end of the camera system update. */
+	virtual void OnPostUpdate(const FCameraEvaluationServiceUpdateParams& Params, FCameraEvaluationServiceUpdateResult& OutResult) {}
+	/** Tears down the evaluation service. */
+	virtual void OnTeardown(const FCameraEvaluationServiceTeardownParams& Params) {}
+
+	/** Called when the root camera node experiences an event. */
+	virtual void OnRootCameraNodeEvent(const FRootCameraNodeCameraRigEvent& InEvent) {}
+
+protected:
+
+	/** Sets the flags on this service. */
+	void SetEvaluationServiceFlags(ECameraEvaluationServiceFlags InFlags);
+
+private:
+
+	/** Evaluation service flags. */
+	ECameraEvaluationServiceFlags PrivateFlags = ECameraEvaluationServiceFlags::None;
+};
+
+}  // namespace UE::Cameras
+
diff --git a/Engine/Plugins/Cameras/GameplayCameras/Source/GameplayCameras/Public/Core/CameraNodeEvaluator.h b/Engine/Plugins/Cameras/GameplayCameras/Source/GameplayCameras/Public/Core/CameraNodeEvaluator.h
--- a/Engine/Plugins/Cameras/GameplayCameras/Source/GameplayCameras/Public/Core/CameraNodeEvaluator.h
+++ b/Engine/Plugins/Cameras/GameplayCameras/Source/GameplayCameras/Public/Core/CameraNodeEvaluator.h
@@ -10,7 +10,6 @@
 #include "CoreTypes.h"
 #include "Debug/RootCameraDebugBlock.h"
 #include "GameplayCameras.h"
-#include "Misc/Optional.h"
 #include "UObject/ObjectPtr.h"
 
 class FReferenceCollector;
@@ -143,7 +142,7 @@ struct FCameraNodeEvaluatorInitializeParams
 	 * Information about the last active camera rig if the node tree being initialized
 	 * is being pushed on top of a non-empty blend stack.
 	 */
-	TOptional<FCameraRigEvaluationInfo> LastActiveCameraRig;
+	FCameraRigEvaluationInfo LastActiveCameraRig;
 };
 
 /** View on a camera node evaluator's children. */
diff --git a/Engine/Plugins/Cameras/GameplayCameras/Source/GameplayCameras/Public/Core/CameraRigEvaluationInfo.h b/Engine/Plugins/Cameras/GameplayCameras/Source/GameplayCameras/Public/Core/CameraRigEvaluationInfo.h
--- a/Engine/Plugins/Cameras/GameplayCameras/Source/GameplayCameras/Public/Core/CameraRigEvaluationInfo.h
+++ b/Engine/Plugins/Cameras/GameplayCameras/Source/GameplayCameras/Public/Core/CameraRigEvaluationInfo.h
@@ -20,16 +20,19 @@ struct FCameraNodeEvaluationResult;
 struct FCameraRigEvaluationInfo
 {
 	/** The camera rig being evaluated. */
-	const UCameraRigAsset* CameraRig;
+	const UCameraRigAsset* CameraRig = nullptr;
 	/** The context inside which the evaluation occurs. */
 	TSharedPtr<const FCameraEvaluationContext> EvaluationContext;
 	/** The last evaluated result for this camera rig. */
-	const FCameraNodeEvaluationResult& LastResult;
+	const FCameraNodeEvaluationResult* LastResult = nullptr;
 	/** The root node evaluator of the camera rig. */
 	const FCameraNodeEvaluator* RootEvaluator;
 	/** Whether the camera rig is frozen. */
 	bool bIsFrozen = false;
 
+	FCameraRigEvaluationInfo()
+	{}
+
 	FCameraRigEvaluationInfo(
 			const UCameraRigAsset* InCameraRig,
 			TSharedPtr<const FCameraEvaluationContext> InEvaluationContext,
@@ -37,9 +40,20 @@ struct FCameraRigEvaluationInfo
 			const FCameraNodeEvaluator* InRootEvaluator)
 		: CameraRig(InCameraRig)
 		, EvaluationContext(InEvaluationContext)
-		, LastResult(InLastResult)
+		, LastResult(&InLastResult)
 		, RootEvaluator(InRootEvaluator)
 	{}
+
+	/** 
+	 * Whether this evaluation info is valid.
+	 *
+	 * A valid evaluation info may not have a valid CameraRig, EvaluationContext, or RootEvaluator
+	 * if the related camera rig was frozen. However, it would always have a valid LastResult.
+	 */
+	bool IsValid() const
+	{
+		return LastResult != nullptr;
+	}
 };
 
 }  // namespace UE::Cameras
diff --git a/Engine/Plugins/Cameras/GameplayCameras/Source/GameplayCameras/Public/Core/CameraSystemEvaluator.h b/Engine/Plugins/Cameras/GameplayCameras/Source/GameplayCameras/Public/Core/CameraSystemEvaluator.h
--- a/Engine/Plugins/Cameras/GameplayCameras/Source/GameplayCameras/Public/Core/CameraSystemEvaluator.h
+++ b/Engine/Plugins/Cameras/GameplayCameras/Source/GameplayCameras/Public/Core/CameraSystemEvaluator.h
@@ -24,7 +24,9 @@ namespace UE::Cameras
 {
 
 class FCameraEvaluationContext;
+class FCameraEvaluationService;
 class FRootCameraNodeEvaluator;
+struct FRootCameraNodeCameraRigEvent;
 
 #if UE_GAMEPLAY_CAMERAS_DEBUG
 class FRootCameraDebugBlock;
@@ -89,6 +91,8 @@ public:
 	/** Initializes the camera system. */
 	GAMEPLAYCAMERAS_API void Initialize(const FCameraSystemEvaluatorCreateParams& Params);
 
+	GAMEPLAYCAMERAS_API ~FCameraSystemEvaluator();
+
 public:
 
 	/** Gets the owner of this camera system, if any, and if still valid. */
@@ -108,6 +112,13 @@ public:
 	/** Gets the context stack. */
 	const FCameraEvaluationContextStack& GetEvaluationContextStack() const { return ContextStack; }
 
+public:
+
+	/** Registers an evaluation service on this camera system. */
+	GAMEPLAYCAMERAS_API void RegisterEvaluationService(TSharedRef<FCameraEvaluationService> EvaluationService);
+	/** Unregisters an evaluation service from this camera system. */
+	GAMEPLAYCAMERAS_API void UnregisterEvaluationService(TSharedRef<FCameraEvaluationService> EvaluationService);
+
 public:
 
 	/** Run an update of the camera system. */
@@ -130,6 +141,10 @@ public:
 
 	GAMEPLAYCAMERAS_API void AddReferencedObjects(FReferenceCollector& Collector);
 
+protected:
+
+	void NotifyRootCameraNodeEvent(const FRootCameraNodeCameraRigEvent& InEvent);
+
 private:
 
 	/** The owner (if any) of this camera system evaluator. */
@@ -141,6 +156,9 @@ private:
 	/** The stack of active evaluation context. */
 	FCameraEvaluationContextStack ContextStack;
 
+	/** The list of evaluation services. */
+	TArray<TSharedPtr<FCameraEvaluationService>> EvaluationServices;
+
 	/** Storage buffer for the root evaluator. */
 	FCameraNodeEvaluatorStorage RootEvaluatorStorage;
 
@@ -160,6 +178,8 @@ private:
 	/** The root debug drawing block. */
 	FRootCameraDebugBlock* RootDebugBlock = nullptr;
 #endif  // UE_GAMEPLAY_CAMERAS_DEBUG
+
+	friend class FRootCameraNodeEvaluator;
 };
 
 }  // namespace UE::Cameras
diff --git a/Engine/Plugins/Cameras/GameplayCameras/Source/GameplayCameras/Public/Core/DefaultRootCameraNode.h b/Engine/Plugins/Cameras/GameplayCameras/Source/GameplayCameras/Public/Core/DefaultRootCameraNode.h
--- a/Engine/Plugins/Cameras/GameplayCameras/Source/GameplayCameras/Public/Core/DefaultRootCameraNode.h
+++ b/Engine/Plugins/Cameras/GameplayCameras/Source/GameplayCameras/Public/Core/DefaultRootCameraNode.h
@@ -3,6 +3,7 @@
 #pragma once
 
 #include "Core/RootCameraNode.h"
+#include "Core/BlendStackCameraNodeObserver.h"
 
 #include "DefaultRootCameraNode.generated.h"
 
@@ -47,7 +48,9 @@ class FBlendStackCameraNodeEvaluator;
 /**
  * Evaluator for the default root camera node.
  */
-class FDefaultRootCameraNodeEvaluator : public FRootCameraNodeEvaluator
+class FDefaultRootCameraNodeEvaluator 
+	: public FRootCameraNodeEvaluator
+	, public IBlendStackCameraNodeObserver
 {
 	UE_DECLARE_CAMERA_NODE_EVALUATOR(GAMEPLAYCAMERAS_API, FDefaultRootCameraNodeEvaluator)
 
@@ -63,6 +66,13 @@ protected:
 	virtual void OnBuildDebugBlocks(const FCameraDebugBlockBuildParams& Params, FCameraDebugBlockBuilder& Builder) override;
 #endif  // UE_GAMEPLAY_CAMERAS_DEBUG
 
+	// IBlendStackCameraNodeObserver interface.
+	virtual void OnBlendStackEvent(const FBlendStackCameraRigEvent& InEvent) override;
+
+private:
+
+	FBlendStackCameraNodeEvaluator* BuildBlendStackEvaluator(const FCameraNodeEvaluatorBuildParams& Params, UBlendStackCameraNode* BlendStackNode);
+
 private:
 
 	FBlendStackCameraNodeEvaluator* BaseLayer;
diff --git a/Engine/Plugins/Cameras/GameplayCameras/Source/GameplayCameras/Public/Core/RootCameraNode.h b/Engine/Plugins/Cameras/GameplayCameras/Source/GameplayCameras/Public/Core/RootCameraNode.h
--- a/Engine/Plugins/Cameras/GameplayCameras/Source/GameplayCameras/Public/Core/RootCameraNode.h
+++ b/Engine/Plugins/Cameras/GameplayCameras/Source/GameplayCameras/Public/Core/RootCameraNode.h
@@ -40,6 +40,8 @@ namespace UE::Cameras
 
 class FCameraEvaluationContext;
 class FCameraSystemEvaluator;
+class IRootCameraNodeObserver;
+struct FRootCameraNodeCameraRigEvent;
 
 /**
  * Parameter structure for activating a new camera rig.
@@ -69,10 +71,31 @@ public:
 	/** Activates a camera rig. */
 	void ActivateCameraRig(const FActivateCameraRigParams& Params);
 
-private:
+	/** Registers an observer to this root node. */
+	void RegisterObserver(IRootCameraNodeObserver* Observer);
+	/** Unregisters an observer from this root node. */
+	void UnregisterObserver(IRootCameraNodeObserver* Observer);
+
+protected:
+
+	// FCameraNodeEvaluator interface.
+	virtual void OnInitialize(const FCameraNodeEvaluatorInitializeParams& Params) override;
+
+protected:
 
 	/** Activates a camera rig. */
 	virtual void OnActivateCameraRig(const FActivateCameraRigParams& Params) {}
+
+	bool HasObservers() const;
+	void NotifyObservers(const FRootCameraNodeCameraRigEvent& InEvent) const;
+
+private:
+
+	/** The camera system that owns this root node. */
+	FCameraSystemEvaluator* OwningEvaluator = nullptr;
+
+	/** The list of observers. */
+	TArray<IRootCameraNodeObserver*> Observers;
 };
 
 }  // namespace UE::Cameras
diff --git a/Engine/Plugins/Cameras/GameplayCameras/Source/GameplayCameras/Public/Core/RootCameraNodeObserver.h b/Engine/Plugins/Cameras/GameplayCameras/Source/GameplayCameras/Public/Core/RootCameraNodeObserver.h
new file mode 100644
--- /dev/null
+++ b/Engine/Plugins/Cameras/GameplayCameras/Source/GameplayCameras/Public/Core/RootCameraNodeObserver.h
@@ -0,0 +1,53 @@
+// Copyright Epic Games, Inc. All Rights Reserved.
+
+#pragma once
+
+#include "Core/CameraRigEvaluationInfo.h"
+#include "UObject/ObjectPtr.h"
+
+class UCameraRigTransition;
+enum class ECameraRigLayer : uint8;
+
+namespace UE::Cameras
+{
+
+/** The type of root node event. */
+enum class ERootCameraNodeCameraRigEventType
+{
+	/** A camera rig was activated. */
+	Activated,
+	/** A camera rig was deactivated. */
+	Deactivated
+};
+
+/**
+ * A structure describing an event on a root node.
+ */
+struct FRootCameraNodeCameraRigEvent
+{
+	/** The type of event. */
+	ERootCameraNodeCameraRigEventType EventType;
+	/** The layer on which the event happened. */
+	ECameraRigLayer EventLayer;
+	/** The evaluation information of the associated camera rig. */
+	FCameraRigEvaluationInfo CameraRigInfo;
+	/** If a camera rig was activated, the transition used to activate it, if any. */
+	TObjectPtr<const UCameraRigTransition> Transition;
+};
+
+/**
+ * An interface for observing events on a root camera node.
+ */
+class IRootCameraNodeObserver
+{
+public:
+
+	virtual ~IRootCameraNodeObserver() {}
+
+	/** Called when an event happens in a root camera node. */
+	virtual void OnRootCameraNodeEvent(const FRootCameraNodeCameraRigEvent& InEvent) {}
+};
+
+}  // namespace UE::Cameras
+
+
