Cameras: tear-down any still-ongoing evaluation contexts when a camera system is destroyed.

diff --git a/Engine/Plugins/Cameras/GameplayCameras/Source/GameplayCameras/Private/Core/CameraEvaluationContextStack.cpp b/Engine/Plugins/Cameras/GameplayCameras/Source/GameplayCameras/Private/Core/CameraEvaluationContextStack.cpp
--- a/Engine/Plugins/Cameras/GameplayCameras/Source/GameplayCameras/Private/Core/CameraEvaluationContextStack.cpp
+++ b/Engine/Plugins/Cameras/GameplayCameras/Source/GameplayCameras/Private/Core/CameraEvaluationContextStack.cpp
@@ -2,16 +2,18 @@
 
 #include "Core/CameraEvaluationContextStack.h"
 
-#include "Core/CameraAsset.h"
-#include "Core/CameraDirector.h"
 #include "Core/CameraDirectorEvaluator.h"
 #include "Core/CameraEvaluationContext.h"
 #include "Core/CameraSystemEvaluator.h"
-#include "UObject/Package.h"
 
 namespace UE::Cameras
 {
 
+FCameraEvaluationContextStack::~FCameraEvaluationContextStack()
+{
+	Reset();
+}
+
 TSharedPtr<FCameraEvaluationContext> FCameraEvaluationContextStack::GetActiveContext() const
 {
 	for (const FContextEntry& Entry : ReverseIterate(Entries))
@@ -100,6 +102,19 @@ void FCameraEvaluationContextStack::PopContext()
 	Entries.Pop();
 }
 
+void FCameraEvaluationContextStack::Reset()
+{
+	for (FContextEntry& Entry : Entries)
+	{
+		if (TSharedPtr<FCameraEvaluationContext> Context = Entry.WeakContext.Pin())
+		{
+			FCameraEvaluationContextDeactivateParams DeactivateParams;
+			Context->Deactivate(DeactivateParams);
+		}
+	}
+	Entries.Reset();
+}
+
 void FCameraEvaluationContextStack::Initialize(FCameraSystemEvaluator& InEvaluator)
 {
 	Evaluator = &InEvaluator;
diff --git a/Engine/Plugins/Cameras/GameplayCameras/Source/GameplayCameras/Private/Core/CameraSystemEvaluator.cpp b/Engine/Plugins/Cameras/GameplayCameras/Source/GameplayCameras/Private/Core/CameraSystemEvaluator.cpp
--- a/Engine/Plugins/Cameras/GameplayCameras/Source/GameplayCameras/Private/Core/CameraSystemEvaluator.cpp
+++ b/Engine/Plugins/Cameras/GameplayCameras/Source/GameplayCameras/Private/Core/CameraSystemEvaluator.cpp
@@ -72,6 +72,8 @@ void FCameraSystemEvaluator::Initialize(const FCameraSystemEvaluatorCreateParams
 
 FCameraSystemEvaluator::~FCameraSystemEvaluator()
 {
+	ContextStack.Reset();
+
 	{
 		FCameraEvaluationServiceTeardownParams TeardownParams;
 		TeardownParams.Evaluator = this;
diff --git a/Engine/Plugins/Cameras/GameplayCameras/Source/GameplayCameras/Public/Core/CameraEvaluationContextStack.h b/Engine/Plugins/Cameras/GameplayCameras/Source/GameplayCameras/Public/Core/CameraEvaluationContextStack.h
--- a/Engine/Plugins/Cameras/GameplayCameras/Source/GameplayCameras/Public/Core/CameraEvaluationContextStack.h
+++ b/Engine/Plugins/Cameras/GameplayCameras/Source/GameplayCameras/Public/Core/CameraEvaluationContextStack.h
@@ -21,6 +21,10 @@ class FCameraSystemEvaluator;
  */
 struct FCameraEvaluationContextStack
 {
+public:
+
+	~FCameraEvaluationContextStack();
+
 public:
 
 	/** Gets the active (top) context. */
@@ -43,6 +47,9 @@ public:
 	/** The number of contexts on the stack. */
 	int32 NumContexts() const { return Entries.Num(); }
 
+	/** Empties the stack of all contexts. */
+	void Reset();
+
 public:
 
 	// Internal API
