Cameras: add service for auto-resetting camera variables to their default value every frame

diff --git a/Engine/Plugins/Cameras/GameplayCameras/Source/GameplayCameras/Private/Core/CameraRigAssetBuilder.cpp b/Engine/Plugins/Cameras/GameplayCameras/Source/GameplayCameras/Private/Core/CameraRigAssetBuilder.cpp
--- a/Engine/Plugins/Cameras/GameplayCameras/Source/GameplayCameras/Private/Core/CameraRigAssetBuilder.cpp
+++ b/Engine/Plugins/Cameras/GameplayCameras/Source/GameplayCameras/Private/Core/CameraRigAssetBuilder.cpp
@@ -516,6 +516,10 @@ void FCameraRigAssetBuilder::BuildAllocationInfo(UCameraNode* CameraNode)
 			{\
 				FCameraVariableDefinition VariableDefinition = CameraParameterPtr->Variable->GetVariableDefinition();\
 				AllocationInfo.VariableTableInfo.VariableDefinitions.Add(VariableDefinition);\
+				if (CameraParameterPtr->Variable->bAutoReset)\
+				{\
+					AllocationInfo.VariableTableInfo.AutoResetVariables.Add(CameraParameterPtr->Variable);\
+				}\
 			}\
 		}\
 		else
diff --git a/Engine/Plugins/Cameras/GameplayCameras/Source/GameplayCameras/Private/Core/CameraSystemEvaluator.cpp b/Engine/Plugins/Cameras/GameplayCameras/Source/GameplayCameras/Private/Core/CameraSystemEvaluator.cpp
--- a/Engine/Plugins/Cameras/GameplayCameras/Source/GameplayCameras/Private/Core/CameraSystemEvaluator.cpp
+++ b/Engine/Plugins/Cameras/GameplayCameras/Source/GameplayCameras/Private/Core/CameraSystemEvaluator.cpp
@@ -17,6 +17,7 @@
 #include "Debug/RootCameraDebugBlock.h"
 #include "HAL/IConsoleManager.h"
 #include "IGameplayCamerasModule.h"
+#include "Services/AutoResetCameraVariableService.h"
 #include "UObject/Package.h"
 #include "UObject/UObjectGlobals.h"
 
@@ -68,6 +69,8 @@ void FCameraSystemEvaluator::Initialize(const FCameraSystemEvaluatorCreateParams
 		InitParams.Evaluator = this;
 		RootEvaluator->Initialize(InitParams);
 	}
+
+	RegisterEvaluationService(MakeShared<FAutoResetCameraVariableService>());
 }
 
 FCameraSystemEvaluator::~FCameraSystemEvaluator()
diff --git a/Engine/Plugins/Cameras/GameplayCameras/Source/GameplayCameras/Private/Core/CameraVariableTable.cpp b/Engine/Plugins/Cameras/GameplayCameras/Source/GameplayCameras/Private/Core/CameraVariableTable.cpp
--- a/Engine/Plugins/Cameras/GameplayCameras/Source/GameplayCameras/Private/Core/CameraVariableTable.cpp
+++ b/Engine/Plugins/Cameras/GameplayCameras/Source/GameplayCameras/Private/Core/CameraVariableTable.cpp
@@ -202,6 +202,35 @@ bool FCameraVariableTable::ContainsValue(FCameraVariableID VariableID) const
 	return Entries.Contains(VariableID);
 }
 
+void FCameraVariableTable::SetValue(FCameraVariableID VariableID, ECameraVariableType ExpectedVariableType, const uint8* InRawValuePtr)
+{
+	FEntry* Entry = Entries.Find(VariableID);
+	if (ensureMsgf(Entry, TEXT("Can't set camera variable (ID '%d') because it doesn't exist in the table."), VariableID.GetValue()))
+	{
+		check(ExpectedVariableType == Entry->Type);
+		uint32 SizeOf, AlignOf;
+		GetVariableTypeAllocationInfo(Entry->Type, SizeOf, AlignOf);
+		uint8* ValuePtr = Memory + Entry->Offset;
+		FMemory::Memcpy(ValuePtr, InRawValuePtr, SizeOf);
+		Entry->Flags |= EEntryFlags::Written | EEntryFlags::WrittenThisFrame;
+	}
+}
+
+bool FCameraVariableTable::TrySetValue(FCameraVariableID VariableID, ECameraVariableType ExpectedVariableType, const uint8* InRawValuePtr)
+{
+	if (FEntry* Entry = Entries.Find(VariableID))
+	{
+		check(ExpectedVariableType == Entry->Type);
+		uint32 SizeOf, AlignOf;
+		GetVariableTypeAllocationInfo(Entry->Type, SizeOf, AlignOf);
+		uint8* ValuePtr = Memory + Entry->Offset;
+		FMemory::Memcpy(ValuePtr, InRawValuePtr, SizeOf);
+		Entry->Flags |= EEntryFlags::Written | EEntryFlags::WrittenThisFrame;
+		return true;
+	}
+	return false;
+}
+
 bool FCameraVariableTable::IsValueWritten(FCameraVariableID VariableID) const
 {
 	if (const FEntry* Entry = Entries.Find(VariableID))
diff --git a/Engine/Plugins/Cameras/GameplayCameras/Source/GameplayCameras/Private/Services/AutoResetCameraVariableService.cpp b/Engine/Plugins/Cameras/GameplayCameras/Source/GameplayCameras/Private/Services/AutoResetCameraVariableService.cpp
new file mode 100644
--- /dev/null
+++ b/Engine/Plugins/Cameras/GameplayCameras/Source/GameplayCameras/Private/Services/AutoResetCameraVariableService.cpp
@@ -0,0 +1,90 @@
+// Copyright Epic Games, Inc. All Rights Reserved.
+
+#include "Services/AutoResetCameraVariableService.h"
+
+#include "Core/CameraNodeEvaluator.h"
+#include "Core/CameraRigAsset.h"
+#include "Core/CameraVariableAssets.h"
+#include "Core/CameraVariableTable.h"
+#include "Core/RootCameraNodeObserver.h"
+
+namespace UE::Cameras
+{
+
+void FAutoResetCameraVariableService::OnInitialize(const FCameraEvaluationServiceInitializeParams& Params)
+{
+	SetEvaluationServiceFlags(
+			ECameraEvaluationServiceFlags::NeedsPreUpdate |
+			ECameraEvaluationServiceFlags::NeedsRootCameraNodeEvents);
+}
+
+void FAutoResetCameraVariableService::OnPreUpdate(const FCameraEvaluationServiceUpdateParams& Params, FCameraEvaluationServiceUpdateResult& OutResult)
+{
+	FCameraVariableTable& VariableTable = OutResult.EvaluationResult.VariableTable;
+
+	for (const TPair<const UCameraVariableAsset*, uint32>& Pair : AutoResetVariables)
+	{
+		const FCameraVariableID VariableID = Pair.Key->GetVariableID();
+		const uint8* DefaultValuePtr = Pair.Key->GetDefaultValuePtr();
+		const ECameraVariableType VariableType = Pair.Key->GetVariableType();
+		// Using TrySetValue instead of SetValue because we only know of variables *possibly* used by camera rigs. 
+		// This doesn't mean these variables have been added to the table and written to. For instance, a camera
+		// parameter on a node might be configured to use a variable that isn't set, which makes it fallback to
+		// the variable's default value when it's not found in the variable table. So we silently ignore variables
+		// not present in the variable table here.
+		VariableTable.TrySetValue(VariableID, VariableType, DefaultValuePtr);
+	}
+}
+
+void FAutoResetCameraVariableService::OnRootCameraNodeEvent(const FRootCameraNodeCameraRigEvent& InEvent)
+{
+	const UCameraRigAsset* CameraRig = InEvent.CameraRigInfo.CameraRig;
+	if (!CameraRig)
+	{
+		return;
+	}
+
+	const FCameraVariableTableAllocationInfo& VariableTableInfo = CameraRig->AllocationInfo.VariableTableInfo;
+
+	switch (InEvent.EventType)
+	{
+		case ERootCameraNodeCameraRigEventType::Activated:
+			{
+				for (const UCameraVariableAsset* Variable : VariableTableInfo.AutoResetVariables)
+				{
+					AddAutoResetVariable(Variable);
+				}
+			}
+			break;
+		case ERootCameraNodeCameraRigEventType::Deactivated:
+			{
+				for (const UCameraVariableAsset* Variable : VariableTableInfo.AutoResetVariables)
+				{
+					RemoveAutoResetVariable(Variable);
+				}
+			}
+			break;
+	}
+}
+
+void FAutoResetCameraVariableService::AddAutoResetVariable(const UCameraVariableAsset* InVariable)
+{
+	uint32& RefCount = AutoResetVariables.FindOrAdd(InVariable, 0);
+	++RefCount;
+}
+
+void FAutoResetCameraVariableService::RemoveAutoResetVariable(const UCameraVariableAsset* InVariable)
+{
+	uint32* RefCount = AutoResetVariables.Find(InVariable);
+	if (ensure(RefCount))
+	{
+		--(*RefCount);
+		if (*RefCount == 0)
+		{
+			AutoResetVariables.Remove(InVariable);
+		}
+	}
+}
+
+}  // namespace UE::Cameras
+
diff --git a/Engine/Plugins/Cameras/GameplayCameras/Source/GameplayCameras/Private/Services/AutoResetCameraVariableService.h b/Engine/Plugins/Cameras/GameplayCameras/Source/GameplayCameras/Private/Services/AutoResetCameraVariableService.h
new file mode 100644
--- /dev/null
+++ b/Engine/Plugins/Cameras/GameplayCameras/Source/GameplayCameras/Private/Services/AutoResetCameraVariableService.h
@@ -0,0 +1,33 @@
+// Copyright Epic Games, Inc. All Rights Reserved.
+
+#pragma once
+
+#include "Core/CameraEvaluationService.h"
+#include "UObject/ObjectPtr.h"
+
+class UCameraVariableAsset;
+
+namespace UE::Cameras
+{
+
+class FAutoResetCameraVariableService : public FCameraEvaluationService
+{
+protected:
+
+	// FCameraEvaluationService interface.
+	virtual void OnInitialize(const FCameraEvaluationServiceInitializeParams& Params) override;
+	virtual void OnPreUpdate(const FCameraEvaluationServiceUpdateParams& Params, FCameraEvaluationServiceUpdateResult& OutResult) override;
+	virtual void OnRootCameraNodeEvent(const FRootCameraNodeCameraRigEvent& InEvent) override;
+
+private:
+
+	void AddAutoResetVariable(const UCameraVariableAsset* InVariable);
+	void RemoveAutoResetVariable(const UCameraVariableAsset* InVariable);
+
+private:
+
+	TMap<const UCameraVariableAsset*, uint32> AutoResetVariables;
+};
+
+}  // namespace UE::Cameras
+
diff --git a/Engine/Plugins/Cameras/GameplayCameras/Source/GameplayCameras/Public/Core/CameraVariableTable.h b/Engine/Plugins/Cameras/GameplayCameras/Source/GameplayCameras/Public/Core/CameraVariableTable.h
--- a/Engine/Plugins/Cameras/GameplayCameras/Source/GameplayCameras/Public/Core/CameraVariableTable.h
+++ b/Engine/Plugins/Cameras/GameplayCameras/Source/GameplayCameras/Public/Core/CameraVariableTable.h
@@ -138,6 +138,8 @@ public:
 
 	// Lower level API.
 
+	void SetValue(FCameraVariableID VariableID, ECameraVariableType ExpectedVariableType, const uint8* InRawValuePtr);
+	bool TrySetValue(FCameraVariableID VariableID, ECameraVariableType ExpectedVariableType, const uint8* InRawValuePtr);
 	bool IsValueWritten(FCameraVariableID VariableID) const;
 	void UnsetValue(FCameraVariableID VariableID);
 	void UnsetAllValues();
diff --git a/Engine/Plugins/Cameras/GameplayCameras/Source/GameplayCameras/Public/Core/CameraVariableTableFwd.h b/Engine/Plugins/Cameras/GameplayCameras/Source/GameplayCameras/Public/Core/CameraVariableTableFwd.h
--- a/Engine/Plugins/Cameras/GameplayCameras/Source/GameplayCameras/Public/Core/CameraVariableTableFwd.h
+++ b/Engine/Plugins/Cameras/GameplayCameras/Source/GameplayCameras/Public/Core/CameraVariableTableFwd.h
@@ -8,6 +8,8 @@
 
 #include "CameraVariableTableFwd.generated.h"
 
+class UCameraVariableAsset;
+
 #define UE_CAMERA_VARIABLE_FOR_ALL_TYPES()\
 	UE_CAMERA_VARIABLE_FOR_TYPE(bool, Boolean)\
 	UE_CAMERA_VARIABLE_FOR_TYPE(int32, Integer32)\
@@ -24,6 +26,12 @@
 	UE_CAMERA_VARIABLE_FOR_TYPE(FTransform3f, Transform3f)\
 	UE_CAMERA_VARIABLE_FOR_TYPE(FTransform3d, Transform3d)
 
+/**
+ * The type of a camera variable. 
+ *
+ * Only a fixed set of types are supported for camera variables because of
+ * simplicity, and because these types need to be blendable.
+ */
 UENUM()
 enum class ECameraVariableType
 {
@@ -43,6 +51,9 @@ enum class ECameraVariableType
 	Transform3d
 };
 
+/**
+ * The ID of a camera variable, used to refer to it in a camera variable table.
+ */
 USTRUCT()
 struct FCameraVariableID
 {
@@ -95,24 +106,41 @@ private:
 	uint32 Value;
 };
 
+/**
+ * A structure that describes a camera variable.
+ */
 USTRUCT()
 struct FCameraVariableDefinition
 {
 	GENERATED_BODY()
 
+	/** The ID of the variable. */
 	UPROPERTY()
 	FCameraVariableID VariableID;
 
+	/** The type of the variable. */
 	UPROPERTY()
 	ECameraVariableType VariableType = ECameraVariableType::Boolean;
 
+	/**
+	 * Whether the variable is private. 
+	 *
+	 * Private variables are not propagated from one table to another when
+	 * interpolating or overriding a table.
+	 */
 	UPROPERTY()
 	bool bIsPrivate = false;
 
+	/**
+	 * Whether the variable is an input variable.
+	 *
+	 * Input variables are blended during the pre-blend parameter update phase.
+	 */
 	UPROPERTY()
 	bool bIsInput = false;
 
 #if WITH_EDITORONLY_DATA
+	/** The name of the variable, for debugging purposes. */
 	UPROPERTY()
 	FString VariableName;
 #endif
@@ -137,12 +165,20 @@ struct FCameraVariableDefinition
 	}
 };
 
+/**
+ * A structure that describes the required camera variable table setup of a camera rig.
+ */
 USTRUCT()
 struct FCameraVariableTableAllocationInfo
 {
 	GENERATED_BODY()
 
+	/** The list of variables that should be allocated in a table. */
 	UPROPERTY()
 	TArray<FCameraVariableDefinition> VariableDefinitions;
+
+	/** The list of variables that should be auto-reset to their default value every frame. */
+	UPROPERTY()
+	TArray<TObjectPtr<UCameraVariableAsset>> AutoResetVariables;
 };
 
