Cameras: fix missing blend stack root debug block

diff --git a/Engine/Plugins/Cameras/GameplayCameras/Source/GameplayCameras/Private/Core/BlendStackRootCameraNode.cpp b/Engine/Plugins/Cameras/GameplayCameras/Source/GameplayCameras/Private/Core/BlendStackRootCameraNode.cpp
--- a/Engine/Plugins/Cameras/GameplayCameras/Source/GameplayCameras/Private/Core/BlendStackRootCameraNode.cpp
+++ b/Engine/Plugins/Cameras/GameplayCameras/Source/GameplayCameras/Private/Core/BlendStackRootCameraNode.cpp
@@ -62,6 +62,20 @@ void FBlendStackRootCameraNodeEvaluator::OnBuild(const FCameraNodeEvaluatorBuild
 	RootEvaluator = Params.BuildEvaluator(RootNode->RootNode);
 }
 
+void FBlendStackRootCameraNodeEvaluator::OnInitialize(const FCameraNodeEvaluatorInitializeParams& Params)
+{
+#if UE_GAMEPLAY_CAMERAS_DEBUG
+	const UBlendStackRootCameraNode* RootNode = GetCameraNodeAs<UBlendStackRootCameraNode>();
+	if (RootNode->RootNode)
+	{
+		if (const UCameraRigAsset* CameraRig = RootNode->RootNode->GetTypedOuter<UCameraRigAsset>())
+		{
+			CameraRigAssetName = GetNameSafe(CameraRig);
+		}
+	}
+#endif  // UE_GAMEPLAY_CAMERAS_DEBUG
+}
+
 void FBlendStackRootCameraNodeEvaluator::OnRun(const FCameraNodeEvaluationParams& Params, FCameraNodeEvaluationResult& OutResult)
 {
 	// Nothing to do, the blend stack node runs our blend and root nodes selectively.
@@ -71,6 +85,9 @@ void FBlendStackRootCameraNodeEvaluator::OnRun(const FCameraNodeEvaluationParams
 
 void FBlendStackRootCameraNodeEvaluator::OnBuildDebugBlocks(const FCameraDebugBlockBuildParams& Params, FCameraDebugBlockBuilder& Builder)
 {
+	FBlendStackRootCameraDebugBlock& DebugBlock = Builder.StartChildDebugBlock<FBlendStackRootCameraDebugBlock>();
+	DebugBlock.CameraRigAssetName = CameraRigAssetName;
+
 	if (BlendEvaluator)
 	{
 		BlendEvaluator->BuildDebugBlocks(Params, Builder);
@@ -93,6 +110,7 @@ void FBlendStackRootCameraNodeEvaluator::OnBuildDebugBlocks(const FCameraDebugBl
 		Builder.EndChildDebugBlock();
 	}
 
+	Builder.EndChildDebugBlock();
 	Builder.SkipChildren();
 }
 
@@ -100,12 +118,12 @@ void FBlendStackRootCameraDebugBlock::OnDebugDraw(const FCameraDebugBlockDrawPar
 {
 	TArrayView<FCameraDebugBlock*> ChildrenView(GetChildren());
 
-	Renderer.AddText("{yellowgreen}[Blend]\n");
+	Renderer.AddText("{cam_passive}<Blend>\n");
 	Renderer.AddIndent();
 	ChildrenView[0]->DebugDraw(Params, Renderer);
 	Renderer.RemoveIndent();
 
-	Renderer.AddText("{yellowgreen}[Camera Rig]\n");
+	Renderer.AddText(TEXT("{cam_passive}<CameraRig %s>\n"), *CameraRigAssetName);
 	Renderer.AddIndent();
 	ChildrenView[1]->DebugDraw(Params, Renderer);
 	Renderer.RemoveIndent();
diff --git a/Engine/Plugins/Cameras/GameplayCameras/Source/GameplayCameras/Private/Core/BlendStackRootCameraNode.h b/Engine/Plugins/Cameras/GameplayCameras/Source/GameplayCameras/Private/Core/BlendStackRootCameraNode.h
--- a/Engine/Plugins/Cameras/GameplayCameras/Source/GameplayCameras/Private/Core/BlendStackRootCameraNode.h
+++ b/Engine/Plugins/Cameras/GameplayCameras/Source/GameplayCameras/Private/Core/BlendStackRootCameraNode.h
@@ -59,6 +59,7 @@ protected:
 	// FCameraNodeEvaluator interface.
 	virtual FCameraNodeEvaluatorChildrenView OnGetChildren() override;
 	virtual void OnBuild(const FCameraNodeEvaluatorBuildParams& Params) override;
+	virtual void OnInitialize(const FCameraNodeEvaluatorInitializeParams& Params) override;
 	virtual void OnRun(const FCameraNodeEvaluationParams& Params, FCameraNodeEvaluationResult& OutResult) override;
 
 #if UE_GAMEPLAY_CAMERAS_DEBUG
@@ -69,6 +70,10 @@ private:
 
 	FBlendCameraNodeEvaluator* BlendEvaluator = nullptr;
 	FCameraNodeEvaluator* RootEvaluator = nullptr;
+
+#if UE_GAMEPLAY_CAMERAS_DEBUG
+	FString CameraRigAssetName;
+#endif  // UE_GAMEPLAY_CAMERAS_DEBUG
 };
 
 }  // namespace UE::Cameras
