Cameras: fix possible crash while building debug blocks

diff --git a/Engine/Plugins/Cameras/GameplayCameras/Source/GameplayCameras/Private/Core/CameraNodeEvaluator.cpp b/Engine/Plugins/Cameras/GameplayCameras/Source/GameplayCameras/Private/Core/CameraNodeEvaluator.cpp
--- a/Engine/Plugins/Cameras/GameplayCameras/Source/GameplayCameras/Private/Core/CameraNodeEvaluator.cpp
+++ b/Engine/Plugins/Cameras/GameplayCameras/Source/GameplayCameras/Private/Core/CameraNodeEvaluator.cpp
@@ -158,7 +158,10 @@ void FCameraNodeEvaluator::BuildDebugBlocks(const FCameraDebugBlockBuildParams&
 			FCameraNodeEvaluatorChildrenView ChildrenView(GetChildren());
 			for (FCameraNodeEvaluator* Child : ChildrenView)
 			{
-				Child->BuildDebugBlocks(Params, Builder);
+				if (Child)
+				{
+					Child->BuildDebugBlocks(Params, Builder);
+				}
 			}
 		}
 	}
