Cameras: add serialization to camera pose and variable table

To make it easier, the variable table now separates the storing of entries and the lookup for them.

diff --git a/Engine/Plugins/Cameras/GameplayCameras/Source/GameplayCameras/Private/Core/CameraPose.cpp b/Engine/Plugins/Cameras/GameplayCameras/Source/GameplayCameras/Private/Core/CameraPose.cpp
--- a/Engine/Plugins/Cameras/GameplayCameras/Source/GameplayCameras/Private/Core/CameraPose.cpp
+++ b/Engine/Plugins/Cameras/GameplayCameras/Source/GameplayCameras/Private/Core/CameraPose.cpp
@@ -354,3 +354,8 @@ UE_CAMERA_POSE_FOR_ALL_PROPERTIES()
 #undef UE_CAMERA_POSE_FOR_PROPERTY
 }
 
+void FCameraPose::SerializeWithFlags(FArchive& Ar)
+{
+	FCameraPose::SerializeWithFlags(Ar, *this);
+}
+
diff --git a/Engine/Plugins/Cameras/GameplayCameras/Source/GameplayCameras/Private/Core/CameraVariableTable.cpp b/Engine/Plugins/Cameras/GameplayCameras/Source/GameplayCameras/Private/Core/CameraVariableTable.cpp
--- a/Engine/Plugins/Cameras/GameplayCameras/Source/GameplayCameras/Private/Core/CameraVariableTable.cpp
+++ b/Engine/Plugins/Cameras/GameplayCameras/Source/GameplayCameras/Private/Core/CameraVariableTable.cpp
@@ -31,6 +31,7 @@ FCameraVariableTable::FCameraVariableTable()
 
 FCameraVariableTable::FCameraVariableTable(FCameraVariableTable&& Other)
 	: Entries(MoveTemp(Other.Entries))
+	, EntryLookup(MoveTemp(Other.EntryLookup))
 	, Memory(Other.Memory)
 	, Capacity(Other.Capacity)
 	, Used(Other.Used)
@@ -43,6 +44,7 @@ FCameraVariableTable::FCameraVariableTable(FCameraVariableTable&& Other)
 FCameraVariableTable& FCameraVariableTable::operator=(FCameraVariableTable&& Other)
 {
 	Entries = MoveTemp(Other.Entries);
+	EntryLookup = MoveTemp(Other.EntryLookup);
 	Memory = Other.Memory;
 	Capacity = Other.Capacity;
 	Used = Other.Used;
@@ -68,6 +70,7 @@ void FCameraVariableTable::Initialize(const FCameraVariableTableAllocationInfo&
 {
 	// Reset any previous state.
 	Entries.Reset();
+	EntryLookup.Reset();
 	if (Memory)
 	{
 		FMemory::Free(Memory);
@@ -101,7 +104,8 @@ void FCameraVariableTable::Initialize(const FCameraVariableTableAllocationInfo&
 #if WITH_EDITORONLY_DATA
 		NewEntry.DebugName = VariableDefinition.VariableName;
 #endif
-		Entries.Add(NewEntry.ID, NewEntry);
+		Entries.Add(NewEntry);
+		EntryLookup.Add(NewEntry.ID, Entries.Num() - 1);
 	}
 
 	// Allocate the memory buffer.
@@ -110,9 +114,8 @@ void FCameraVariableTable::Initialize(const FCameraVariableTableAllocationInfo&
 	Used = 0;
 
 	// Go back to our entries and initialize each entry to the default value for that variable type.
-	for (auto& Pair : Entries)
+	for (const FEntry& Entry : Entries)
 	{
-		const FEntry& Entry = Pair.Value;
 		uint8* ValuePtr = Memory + Entry.Offset;
 		switch (Entry.Type)
 		{
@@ -155,7 +158,8 @@ void FCameraVariableTable::AddVariable(const FCameraVariableDefinition& Variable
 	NewEntry.DebugName = VariableDefinition.VariableName;
 #endif
 
-	Entries.Add(VariableDefinition.VariableID, NewEntry);
+	Entries.Add(NewEntry);
+	EntryLookup.Add(VariableDefinition.VariableID, Entries.Num() - 1);
 }
 
 void FCameraVariableTable::ReallocateBuffer(uint32 MinRequired)
@@ -182,6 +186,18 @@ void FCameraVariableTable::ReallocateBuffer(uint32 MinRequired)
 	Capacity = NewCapacity;
 }
 
+FCameraVariableTable::FEntry* FCameraVariableTable::FindEntry(FCameraVariableID VariableID)
+{
+	const int32* IndexPtr = EntryLookup.Find(VariableID);
+	return IndexPtr ? &Entries[*IndexPtr] : nullptr;
+}
+
+const FCameraVariableTable::FEntry* FCameraVariableTable::FindEntry(FCameraVariableID VariableID) const
+{
+	const int32* IndexPtr = EntryLookup.Find(VariableID);
+	return IndexPtr ? &Entries[*IndexPtr] : nullptr;
+}
+
 bool FCameraVariableTable::GetVariableTypeAllocationInfo(ECameraVariableType VariableType, uint32& OutSizeOf, uint32& OutAlignOf)
 {
 	switch (VariableType)
@@ -199,12 +215,12 @@ UE_CAMERA_VARIABLE_FOR_ALL_TYPES()
 
 bool FCameraVariableTable::ContainsValue(FCameraVariableID VariableID) const
 {
-	return Entries.Contains(VariableID);
+	return EntryLookup.Contains(VariableID);
 }
 
 void FCameraVariableTable::SetValue(FCameraVariableID VariableID, ECameraVariableType ExpectedVariableType, const uint8* InRawValuePtr)
 {
-	FEntry* Entry = Entries.Find(VariableID);
+	FEntry* Entry = FindEntry(VariableID);
 	if (ensureMsgf(Entry, TEXT("Can't set camera variable (ID '%d') because it doesn't exist in the table."), VariableID.GetValue()))
 	{
 		check(ExpectedVariableType == Entry->Type);
@@ -218,7 +234,7 @@ void FCameraVariableTable::SetValue(FCameraVariableID VariableID, ECameraVariabl
 
 bool FCameraVariableTable::TrySetValue(FCameraVariableID VariableID, ECameraVariableType ExpectedVariableType, const uint8* InRawValuePtr)
 {
-	if (FEntry* Entry = Entries.Find(VariableID))
+	if (FEntry* Entry = FindEntry(VariableID))
 	{
 		check(ExpectedVariableType == Entry->Type);
 		uint32 SizeOf, AlignOf;
@@ -233,7 +249,7 @@ bool FCameraVariableTable::TrySetValue(FCameraVariableID VariableID, ECameraVari
 
 bool FCameraVariableTable::IsValueWritten(FCameraVariableID VariableID) const
 {
-	if (const FEntry* Entry = Entries.Find(VariableID))
+	if (const FEntry* Entry = FindEntry(VariableID))
 	{
 		return EnumHasAnyFlags(Entry->Flags, EEntryFlags::Written);
 	}
@@ -242,7 +258,7 @@ bool FCameraVariableTable::IsValueWritten(FCameraVariableID VariableID) const
 
 void FCameraVariableTable::UnsetValue(FCameraVariableID VariableID)
 {
-	if (FEntry* Entry = Entries.Find(VariableID))
+	if (FEntry* Entry = FindEntry(VariableID))
 	{
 		EnumRemoveFlags(Entry->Flags, EEntryFlags::Written | EEntryFlags::WrittenThisFrame);
 	}
@@ -250,15 +266,15 @@ void FCameraVariableTable::UnsetValue(FCameraVariableID VariableID)
 
 void FCameraVariableTable::UnsetAllValues()
 {
-	for (auto& Pair : Entries)
+	for (FEntry& Entry : Entries)
 	{
-		EnumRemoveFlags(Pair.Value.Flags, EEntryFlags::Written | EEntryFlags::WrittenThisFrame);
+		EnumRemoveFlags(Entry.Flags, EEntryFlags::Written | EEntryFlags::WrittenThisFrame);
 	}
 }
 
 bool FCameraVariableTable::IsValueWrittenThisFrame(FCameraVariableID VariableID) const
 {
-	if (const FEntry* Entry = Entries.Find(VariableID))
+	if (const FEntry* Entry = FindEntry(VariableID))
 	{
 		return EnumHasAnyFlags(Entry->Flags, EEntryFlags::WrittenThisFrame);
 	}
@@ -267,9 +283,41 @@ bool FCameraVariableTable::IsValueWrittenThisFrame(FCameraVariableID VariableID)
 
 void FCameraVariableTable::ClearAllWrittenThisFrameFlags()
 {
-	for (auto& Pair : Entries)
+	for (FEntry& Entry : Entries)
+	{
+		EnumRemoveFlags(Entry.Flags, EEntryFlags::WrittenThisFrame);
+	}
+}
+
+void FCameraVariableTable::Serialize(FArchive& Ar)
+{
+	if (Ar.IsSaving())
 	{
-		EnumRemoveFlags(Pair.Value.Flags, EEntryFlags::WrittenThisFrame);
+		Ar << Capacity;
+		Ar << Used;
+		Ar.Serialize(Memory, Capacity);
+
+		int32 NumEntries = Entries.Num();
+		Ar << NumEntries;
+		Ar.Serialize(Entries.GetData(), Entries.Num() * sizeof(FEntry));
+	}
+
+	if (Ar.IsLoading())
+	{
+		uint32 LoadedCapacity = 0;
+		Ar << LoadedCapacity;
+
+		uint32 LoadedUsed = 0;
+		Ar << LoadedUsed;
+
+		ensure(LoadedCapacity <= Capacity);
+		Ar.Serialize(Memory, LoadedCapacity);
+
+		int32 LoadedNumEntries = 0;
+		Ar << LoadedNumEntries;
+		
+		ensure(LoadedNumEntries == Entries.Num());
+		Ar.Serialize(Entries.GetData(), LoadedNumEntries * sizeof(FEntry));
 	}
 }
 
@@ -297,10 +345,9 @@ void FCameraVariableTable::InternalOverride(const FCameraVariableTable& OtherTab
 	const bool bInputs = EnumHasAnyFlags(Filter, ECameraVariableTableFilter::Input);
 	const bool bOutputs = EnumHasAnyFlags(Filter, ECameraVariableTableFilter::Output);
 
-	for (auto& OtherPair : OtherTable.Entries)
+	for (const FEntry& OtherEntry : OtherTable.Entries)
 	{
 		// Look for entries in the other table that have been written to, and aren't private.
-		const FEntry& OtherEntry  = OtherPair.Value;
 		const EEntryFlags OtherFlags = OtherEntry.Flags;
 		const bool bOtherEntryIsInput = EnumHasAnyFlags(OtherFlags, EEntryFlags::Input);
 		if (EnumHasAnyFlags(OtherFlags, EEntryFlags::Written)
@@ -310,7 +357,7 @@ void FCameraVariableTable::InternalOverride(const FCameraVariableTable& OtherTab
 				&& IsVariableInMask(OtherEntry.ID, InMask, bInvertMask))
 		{
 			// See if we know this variable.
-			FEntry* ThisEntry = Entries.Find(OtherEntry.ID);
+			FEntry* ThisEntry = FindEntry(OtherEntry.ID);
 			if (ThisEntry)
 			{
 				// We already have the other table's variable in our table. Let's check
@@ -342,7 +389,7 @@ void FCameraVariableTable::InternalOverride(const FCameraVariableTable& OtherTab
 #endif
 				AddVariable(NewVariableDefinition);
 
-				ThisEntry = Entries.Find(OtherEntry.ID);
+				ThisEntry = FindEntry(OtherEntry.ID);
 			}
 
 			if (ensure(ThisEntry))
@@ -389,10 +436,9 @@ void FCameraVariableTable::InternalLerp(const FCameraVariableTable& ToTable, ECa
 	const bool bInputs = EnumHasAnyFlags(Filter, ECameraVariableTableFilter::Input);
 	const bool bOutputs = EnumHasAnyFlags(Filter, ECameraVariableTableFilter::Output);
 
-	for (auto& ToPair : ToTable.Entries)
+	for (const FEntry& ToEntry : ToTable.Entries)
 	{
 		// Look for entries in the other table that have been written to, and aren't private.
-		const FEntry& ToEntry  = ToPair.Value;
 		const EEntryFlags ToFlags = ToEntry.Flags;
 		const bool bToEntryIsInput = EnumHasAnyFlags(ToFlags, EEntryFlags::Input);
 		if (EnumHasAnyFlags(ToFlags, EEntryFlags::Written)
@@ -402,7 +448,7 @@ void FCameraVariableTable::InternalLerp(const FCameraVariableTable& ToTable, ECa
 				&& IsVariableInMask(ToEntry.ID, InMask, bInvertMask))
 		{
 			// See if we know this variable.
-			FEntry* FromEntry = Entries.Find(ToEntry.ID);
+			FEntry* FromEntry = FindEntry(ToEntry.ID);
 			if (FromEntry)
 			{
 				// We already have the other table's variable in our table. Let's check
@@ -452,7 +498,7 @@ UE_CAMERA_VARIABLE_FOR_ALL_TYPES()
 #endif
 				AddVariable(NewVariableDefinition);
 
-				FromEntry = Entries.Find(ToEntry.ID);
+				FromEntry = FindEntry(ToEntry.ID);
 				check(FromEntry);
 
 				uint32 ValueSize = 0, ValueAlignment = 0;
diff --git a/Engine/Plugins/Cameras/GameplayCameras/Source/GameplayCameras/Private/Debug/VariableTableDebugBlock.cpp b/Engine/Plugins/Cameras/GameplayCameras/Source/GameplayCameras/Private/Debug/VariableTableDebugBlock.cpp
--- a/Engine/Plugins/Cameras/GameplayCameras/Source/GameplayCameras/Private/Debug/VariableTableDebugBlock.cpp
+++ b/Engine/Plugins/Cameras/GameplayCameras/Source/GameplayCameras/Private/Debug/VariableTableDebugBlock.cpp
@@ -26,11 +26,8 @@ FVariableTableDebugBlock::FVariableTableDebugBlock(const FCameraVariableTable& I
 
 void FVariableTableDebugBlock::Initialize(const FCameraVariableTable& InVariableTable)
 {
-	for (const TPair<FCameraVariableID, FCameraVariableTable::FEntry>& Pair : InVariableTable.Entries)
+	for (const FCameraVariableTable::FEntry& Entry : InVariableTable.Entries)
 	{
-		const uint32 EntryID = Pair.Key.GetValue();
-		const FCameraVariableTable::FEntry& Entry = Pair.Value;
-
 		FString EntryName;
 #if WITH_EDITORONLY_DATA
 		EntryName = Entry.DebugName;
@@ -41,7 +38,7 @@ void FVariableTableDebugBlock::Initialize(const FCameraVariableTable& InVariable
 			case ECameraVariableType::ValueName:\
 				if (EnumHasAnyFlags(Entry.Flags, FCameraVariableTable::EEntryFlags::Written))\
 				{\
-					const ValueType EntryValue = InVariableTable.GetValue<ValueType>(FCameraVariableID::FromHashValue(EntryID));\
+					const ValueType EntryValue = InVariableTable.GetValue<ValueType>(FCameraVariableID::FromHashValue(Entry.ID.GetValue()));\
 					EntryValueStr = ToDebugString(EntryValue);\
 				}\
 				break;
@@ -51,7 +48,7 @@ void FVariableTableDebugBlock::Initialize(const FCameraVariableTable& InVariable
 		}
 #undef UE_CAMERA_VARIABLE_FOR_TYPE
 
-		FEntryDebugInfo EntryDebugInfo{ EntryID, EntryName, EntryValueStr };
+		FEntryDebugInfo EntryDebugInfo{ Entry.ID.GetValue(), EntryName, EntryValueStr};
 		EntryDebugInfo.bWritten = EnumHasAnyFlags(Entry.Flags, FCameraVariableTable::EEntryFlags::Written);
 		EntryDebugInfo.bWrittenThisFrame = EnumHasAnyFlags(Entry.Flags, FCameraVariableTable::EEntryFlags::WrittenThisFrame);
 		Entries.Add(EntryDebugInfo);
diff --git a/Engine/Plugins/Cameras/GameplayCameras/Source/GameplayCameras/Public/Core/CameraPose.h b/Engine/Plugins/Cameras/GameplayCameras/Source/GameplayCameras/Public/Core/CameraPose.h
--- a/Engine/Plugins/Cameras/GameplayCameras/Source/GameplayCameras/Public/Core/CameraPose.h
+++ b/Engine/Plugins/Cameras/GameplayCameras/Source/GameplayCameras/Public/Core/CameraPose.h
@@ -175,8 +175,12 @@ public:
 
 public:
 
+	/** Serializes the given camera pose including the written-property flags. */
 	static void SerializeWithFlags(FArchive& Ar, FCameraPose& CameraPose);
 
+	/** Serializes this camera pose including the written-property flags. */
+	void SerializeWithFlags(FArchive& Ar);
+
 private:
 
 	void InternalOverrideChanged(const FCameraPose& OtherPose, bool bChangedOnly);
diff --git a/Engine/Plugins/Cameras/GameplayCameras/Source/GameplayCameras/Public/Core/CameraVariableTable.h b/Engine/Plugins/Cameras/GameplayCameras/Source/GameplayCameras/Public/Core/CameraVariableTable.h
--- a/Engine/Plugins/Cameras/GameplayCameras/Source/GameplayCameras/Public/Core/CameraVariableTable.h
+++ b/Engine/Plugins/Cameras/GameplayCameras/Source/GameplayCameras/Public/Core/CameraVariableTable.h
@@ -147,8 +147,12 @@ public:
 	bool IsValueWrittenThisFrame(FCameraVariableID VariableID) const;
 	void ClearAllWrittenThisFrameFlags();
 
+	void Serialize(FArchive& Ar);
+
 private:
 
+	struct FEntry;
+
 	static bool GetVariableTypeAllocationInfo(ECameraVariableType VariableType, uint32& OutSizeOf, uint32& OutAlignOf);
 
 	template<typename ValueType>
@@ -159,6 +163,9 @@ private:
 
 	void ReallocateBuffer(uint32 MinRequired = 0);
 
+	GAMEPLAYCAMERAS_API FEntry* FindEntry(FCameraVariableID VariableID);
+	GAMEPLAYCAMERAS_API const FEntry* FindEntry(FCameraVariableID VariableID) const;
+
 	void InternalOverride(const FCameraVariableTable& OtherTable, ECameraVariableTableFilter Filter, const FCameraVariableTableFlags* InMask, bool bInvertMask, FCameraVariableTableFlags* OutMask);
 	void InternalLerp(const FCameraVariableTable& ToTable, ECameraVariableTableFilter Filter, float Factor, const FCameraVariableTableFlags* InMask, bool bInvertMask, FCameraVariableTableFlags* OutMask);
 
@@ -185,7 +192,8 @@ private:
 #endif
 	};
 
-	TMap<FCameraVariableID, FEntry> Entries;
+	TArray<FEntry> Entries;
+	TMap<FCameraVariableID, int32> EntryLookup;
 
 	uint8* Memory = nullptr;
 	uint32 Capacity = 0;
@@ -204,7 +212,7 @@ ENUM_CLASS_FLAGS(FCameraVariableTable::EEntryFlags)
 template<typename ValueType>
 const ValueType* FCameraVariableTable::FindValue(FCameraVariableID VariableID) const
 {
-	if (const FEntry* Entry = Entries.Find(VariableID))
+	if (const FEntry* Entry = FindEntry(VariableID))
 	{
 		CheckVariableType<ValueType>(Entry->Type);
 		if (EnumHasAnyFlags(Entry->Flags, EEntryFlags::Written))
@@ -218,7 +226,7 @@ const ValueType* FCameraVariableTable::FindValue(FCameraVariableID VariableID) c
 template<typename ValueType>
 const ValueType& FCameraVariableTable::GetValue(FCameraVariableID VariableID) const
 {
-	const FEntry* Entry = Entries.Find(VariableID);
+	const FEntry* Entry = FindEntry(VariableID);
 	if (ensureMsgf(Entry, TEXT("Can't get camera variable (ID '%d') because it doesn't exist in the table."), VariableID.GetValue()))
 	{
 		CheckVariableType<ValueType>(Entry->Type);
@@ -264,7 +272,7 @@ bool FCameraVariableTable::TryGetValue(FCameraVariableID VariableID, ValueType&
 template<typename ValueType>
 void FCameraVariableTable::SetValue(FCameraVariableID VariableID, typename TCallTraits<ValueType>::ParamType Value)
 {
-	FEntry* Entry = Entries.Find(VariableID);
+	FEntry* Entry = FindEntry(VariableID);
 	if (ensureMsgf(Entry, TEXT("Can't set camera variable (ID '%d') because it doesn't exist in the table."), VariableID.GetValue()))
 	{
 		CheckVariableType<ValueType>(Entry->Type);
@@ -277,7 +285,7 @@ void FCameraVariableTable::SetValue(FCameraVariableID VariableID, typename TCall
 template<typename ValueType>
 bool FCameraVariableTable::TrySetValue(FCameraVariableID VariableID, typename TCallTraits<ValueType>::ParamType Value)
 {
-	if (FEntry* Entry = Entries.Find(VariableID))
+	if (FEntry* Entry = FindEntry(VariableID))
 	{
 		CheckVariableType<ValueType>(Entry->Type);
 		ValueType* ValuePtr = reinterpret_cast<ValueType*>(Memory + Entry->Offset);
