Cameras: add support for iterative IK operations, use them for orientation initialization

Camera operations are mostly meant for "manipulating" a camera rig, such as turning it to aim at a desired target. Camera node evaluators can opt-in to execute such operations. Operations are simply-typed so that a node can test for a specific operation. This is a pretty barebones typing system that can be expanded later if needed.

To support iterative operations, camera nodes can be serialized to/from a memory buffer.

The main operation right now is yaw/pitch correction. It is used by a new "orientation initialization" service that looks at new settings on camera rigs and camera rig transitions. These settings allow for preserving the previous camera rig's orientation or target, among others. The actual aiming is done with a new utility class, FCameraIKAim, that can be used for all kinds of aiming needs.

To let the aiming system do its job, the camera nodes report "joints" on their camera rig. This gives us information about where a camera's pivot is, for instance.

diff --git a/Engine/Plugins/Cameras/GameplayCameras/Source/GameplayCameras/GameplayCameras.Build.cs b/Engine/Plugins/Cameras/GameplayCameras/Source/GameplayCameras/GameplayCameras.Build.cs
--- a/Engine/Plugins/Cameras/GameplayCameras/Source/GameplayCameras/GameplayCameras.Build.cs
+++ b/Engine/Plugins/Cameras/GameplayCameras/Source/GameplayCameras/GameplayCameras.Build.cs
@@ -19,6 +19,7 @@ public class GameplayCameras : ModuleRules
 				"HeadMountedDisplay",
 				"MovieScene",
 				"MovieSceneTracks",
+				"Settings",
 				"TemplateSequence",
 				"TraceLog"
 			}
diff --git a/Engine/Plugins/Cameras/GameplayCameras/Source/GameplayCameras/Private/Core/BlendStackCameraNode.cpp b/Engine/Plugins/Cameras/GameplayCameras/Source/GameplayCameras/Private/Core/BlendStackCameraNode.cpp
--- a/Engine/Plugins/Cameras/GameplayCameras/Source/GameplayCameras/Private/Core/BlendStackCameraNode.cpp
+++ b/Engine/Plugins/Cameras/GameplayCameras/Source/GameplayCameras/Private/Core/BlendStackCameraNode.cpp
@@ -153,7 +153,7 @@ bool FBlendStackCameraNodeEvaluator::InitializeEntry(
 	FCameraNodeEvaluatorInitializeParams InitParams;
 	InitParams.Evaluator = Evaluator;
 	InitParams.EvaluationContext = EvaluationContext;
-	InitParams.LastActiveCameraRig = GetActiveCameraRigEvaluationInfo();
+	InitParams.LastActiveCameraRigInfo = GetActiveCameraRigEvaluationInfo();
 	RootEvaluator->Initialize(InitParams);
 
 	// Gather blended parameter evaluators.
@@ -205,11 +205,10 @@ FCameraRigEvaluationInfo FBlendStackCameraNodeEvaluator::GetActiveCameraRigEvalu
 	{
 		const FCameraRigEntry& ActiveEntry = Entries[0];
 		FCameraRigEvaluationInfo Info(
-				ActiveEntry.CameraRig, 
 				ActiveEntry.EvaluationContext.Pin(),
-				ActiveEntry.Result,
-				ActiveEntry.RootEvaluator);
-		Info.bIsFrozen = ActiveEntry.bIsFrozen;
+				ActiveEntry.CameraRig, 
+				&ActiveEntry.Result,
+				ActiveEntry.RootEvaluator->GetRootEvaluator());
 		return Info;
 	}
 	return FCameraRigEvaluationInfo();
@@ -305,7 +304,8 @@ void FBlendStackCameraNodeEvaluator::OnRun(const FCameraNodeEvaluationParams& Pa
 #endif  // UE_GAMEPLAY_CAMERAS_TRACE
 	}
 
-	// Gather parameters to pre-blend, and evaluate blend nodes.
+	// Setup the entries' variable tables for this frame, gather parameters to pre-blend,
+	// and evaluate blend nodes.
 	for (FValidEntry& ValidEntry : ValidEntries)
 	{
 		FCameraRigEntry& Entry(ValidEntry.Entry);
@@ -322,6 +322,14 @@ void FBlendStackCameraNodeEvaluator::OnRun(const FCameraNodeEvaluationParams& Pa
 
 		FCameraNodeEvaluationResult& CurResult(Entry.Result);
 
+		// Setup the variable table: take everything from the previous evaluation layer,
+		// and override it with the context's initial result.
+		CurResult.VariableTable.OverrideAll(OutResult.VariableTable);
+		CurResult.VariableTable.ClearAllWrittenThisFrameFlags();
+
+		const FCameraNodeEvaluationResult& ContextResult(CurContext->GetInitialResult());
+		CurResult.VariableTable.OverrideAll(ContextResult.VariableTable);
+
 		// Gather input parameters.
 		if (!Entry.bInputRunThisFrame)
 		{
@@ -337,6 +345,8 @@ void FBlendStackCameraNodeEvaluator::OnRun(const FCameraNodeEvaluationParams& Pa
 		}
 
 		// Run blends.
+		// Note that we pass last frame's camera pose to the Run() method. This may change.
+		// Blends aren't expected to use the camera pose to do any logic until BlendResults().
 		if (!Entry.bBlendRunThisFrame)
 		{
 			FBlendCameraNodeEvaluator* BlendEvaluator = Entry.RootEvaluator->GetBlendEvaluator();
@@ -392,13 +402,10 @@ void FBlendStackCameraNodeEvaluator::OnRun(const FCameraNodeEvaluationParams& Pa
 		// Start with the input given to us.
 		CurResult.CameraPose = OutResult.CameraPose;
 		CurResult.CameraPose.ClearAllChangedFlags();
-		CurResult.VariableTable.OverrideAll(OutResult.VariableTable);
-		CurResult.VariableTable.ClearAllWrittenThisFrameFlags();
 
 		// Override it with whatever the evaluation context has set on its result.
 		const FCameraNodeEvaluationResult& ContextResult(CurContext->GetInitialResult());
 		CurResult.CameraPose.OverrideChanged(ContextResult.CameraPose);
-		CurResult.VariableTable.OverrideAll(ContextResult.VariableTable);
 		CurResult.bIsCameraCut = OutResult.bIsCameraCut || ContextResult.bIsCameraCut;
 		CurResult.bIsValid = true;
 
@@ -411,7 +418,7 @@ void FBlendStackCameraNodeEvaluator::OnRun(const FCameraNodeEvaluationParams& Pa
 	}
 
 	// Now blend all the results, keeping track of blends that have reached 100% so
-	// that we can remove any camera rigs below (since the would have been completely
+	// that we can remove any camera rigs below (since they would have been completely
 	// blended out by that).
 	int32 PopEntriesBelow = INDEX_NONE;
 	for (FValidEntry& ValidEntry : ValidEntries)
@@ -634,11 +641,10 @@ void FBlendStackCameraNodeEvaluator::NotifyObservers(EBlendStackCameraRigEventTy
 	Event.EventType = EventType;
 	Event.BlendStackEvaluator = this;
 	Event.CameraRigInfo = FCameraRigEvaluationInfo(
-			Entry.CameraRig,
 			Entry.EvaluationContext.Pin(),
-			Entry.Result,
+			Entry.CameraRig,
+			&Entry.Result,
 			Entry.RootEvaluator);
-	Event.CameraRigInfo.bIsFrozen = Entry.bIsFrozen;
 	Event.Transition = Transition;
 
 	for (IBlendStackCameraNodeObserver* Observer : Observers)
@@ -656,6 +662,34 @@ void FBlendStackCameraNodeEvaluator::OnAddReferencedObjects(FReferenceCollector&
 	}
 }
 
+void FBlendStackCameraNodeEvaluator::OnSerialize(const FCameraNodeEvaluatorSerializeParams& Params, FArchive& Ar)
+{
+	if (Ar.IsSaving())
+	{
+		int32 NumEntries = Entries.Num();
+		Ar << NumEntries;
+	}
+	else if (Ar.IsLoading())
+	{
+		int32 LoadedNumEntries = 0;
+		Ar << LoadedNumEntries;
+
+		ensure(LoadedNumEntries == Entries.Num());
+	}
+
+	for (FCameraRigEntry& Entry : Entries)
+	{
+		Entry.Result.Serialize(Ar);
+		Ar << Entry.bIsFirstFrame;
+		Ar << Entry.bInputRunThisFrame;
+		Ar << Entry.bBlendRunThisFrame;
+		Ar << Entry.bIsFrozen;
+#if UE_GAMEPLAY_CAMERAS_TRACE
+		Ar << Entry.bLogWarnings;
+#endif  // UE_GAMEPLAY_CAMERAS_TRACE
+	}
+}
+
 #if WITH_EDITOR
 
 void FBlendStackCameraNodeEvaluator::OnPostBuildAsset(const FGameplayCameraAssetBuildEvent& BuildEvent)
diff --git a/Engine/Plugins/Cameras/GameplayCameras/Source/GameplayCameras/Private/Core/BlendStackRootCameraNode.cpp b/Engine/Plugins/Cameras/GameplayCameras/Source/GameplayCameras/Private/Core/BlendStackRootCameraNode.cpp
--- a/Engine/Plugins/Cameras/GameplayCameras/Source/GameplayCameras/Private/Core/BlendStackRootCameraNode.cpp
+++ b/Engine/Plugins/Cameras/GameplayCameras/Source/GameplayCameras/Private/Core/BlendStackRootCameraNode.cpp
@@ -78,7 +78,14 @@ void FBlendStackRootCameraNodeEvaluator::OnInitialize(const FCameraNodeEvaluator
 
 void FBlendStackRootCameraNodeEvaluator::OnRun(const FCameraNodeEvaluationParams& Params, FCameraNodeEvaluationResult& OutResult)
 {
-	// Nothing to do, the blend stack node runs our blend and root nodes selectively.
+	if (BlendEvaluator)
+	{
+		BlendEvaluator->Run(Params, OutResult);
+	}
+	if (RootEvaluator)
+	{
+		RootEvaluator->Run(Params, OutResult);
+	}
 }
 
 #if UE_GAMEPLAY_CAMERAS_DEBUG
diff --git a/Engine/Plugins/Cameras/GameplayCameras/Source/GameplayCameras/Private/Core/CameraEvaluationService.cpp b/Engine/Plugins/Cameras/GameplayCameras/Source/GameplayCameras/Private/Core/CameraEvaluationService.cpp
--- a/Engine/Plugins/Cameras/GameplayCameras/Source/GameplayCameras/Private/Core/CameraEvaluationService.cpp
+++ b/Engine/Plugins/Cameras/GameplayCameras/Source/GameplayCameras/Private/Core/CameraEvaluationService.cpp
@@ -40,9 +40,9 @@ void FCameraEvaluationService::NotifyRootCameraNodeEvent(const FRootCameraNodeCa
 	OnRootCameraNodeEvent(InEvent);
 }
 
-bool FCameraEvaluationService::HasAnyEvaluationServiceFlags(ECameraEvaluationServiceFlags InFlags) const
+bool FCameraEvaluationService::HasAllEvaluationServiceFlags(ECameraEvaluationServiceFlags InFlags) const
 {
-	return EnumHasAnyFlags(PrivateFlags, InFlags);
+	return EnumHasAllFlags(PrivateFlags, InFlags);
 }
 
 void FCameraEvaluationService::SetEvaluationServiceFlags(ECameraEvaluationServiceFlags InFlags)
diff --git a/Engine/Plugins/Cameras/GameplayCameras/Source/GameplayCameras/Private/Core/CameraIKAim.cpp b/Engine/Plugins/Cameras/GameplayCameras/Source/GameplayCameras/Private/Core/CameraIKAim.cpp
new file mode 100644
--- /dev/null
+++ b/Engine/Plugins/Cameras/GameplayCameras/Source/GameplayCameras/Private/Core/CameraIKAim.cpp
@@ -0,0 +1,424 @@
+// Copyright Epic Games, Inc. All Rights Reserved.
+
+#include "Core/CameraIKAim.h"
+
+#include "Core/BuiltInCameraVariables.h"
+#include "Core/CameraEvaluationContext.h"
+#include "Core/CameraEvaluationService.h"
+#include "Core/CameraOperation.h"
+#include "Core/CameraRigAsset.h"
+#include "Core/CameraRigEvaluationInfo.h"
+#include "Core/CameraRigJoints.h"
+#include "Core/CameraSystemEvaluator.h"
+#include "Core/CameraVariableTable.h"
+#include "Core/RootCameraNode.h"
+#include "GameplayCameras.h"
+#include "GameplayCamerasSettings.h"
+#include "Serialization/MemoryReader.h"
+#include "Serialization/MemoryWriter.h"
+
+namespace UE::Cameras
+{
+
+bool FCameraIKAim::Run(const FCameraIKAimParams& Params, const FCameraRigEvaluationInfo& CameraRigInfo)
+{
+	FCameraIKAimParams ValidatedParams(Params);
+
+	// Replace invalid parameter values with defaults from the settings.
+	const UGameplayCamerasSettings* Settings = UGameplayCamerasSettings::StaticClass()->GetDefaultObject<UGameplayCamerasSettings>();
+	if (ValidatedParams.AngleTolerance <= 0)
+	{
+		ValidatedParams.AngleTolerance = FMath::Max(0.1, Settings->DefaultIKAimingAngleTolerance);
+	}
+	if (ValidatedParams.DistanceTolerance <= 0)
+	{
+		ValidatedParams.DistanceTolerance = FMath::Max(0.1, Settings->DefaultIKAimingDistanceTolerance);
+	}
+	if (ValidatedParams.MinDistance <= 0)
+	{
+		ValidatedParams.MinDistance = FMath::Max(0.1, Settings->DefaultIKAimingMinDistance);
+	}
+	if (ValidatedParams.MaxIterations == 0)
+	{
+		ValidatedParams.MaxIterations = FMath::Max((uint8)1, Settings->DefaultIKAimingMaxIterations);
+	}
+
+	return DoRun(ValidatedParams, CameraRigInfo);
+}
+
+bool FCameraIKAim::DoRun(const FCameraIKAimParams& Params, const FCameraRigEvaluationInfo& CameraRigInfo)
+{
+	if (!ensureMsgf(CameraRigInfo.CameraRig, TEXT("Can't aim invalid camera rig!")))
+	{
+		return false;
+	}
+	if (!ensureMsgf(CameraRigInfo.EvaluationContext, TEXT("Can't aim camera rig '%s', it has no evalution context!"), *CameraRigInfo.CameraRig->GetPathName()))
+	{
+		return false;
+	}
+
+	FCameraSystemEvaluator* CameraSystemEvaluator = Params.Evaluator;
+
+	// Initialize our scratch result.
+	ScratchResult.VariableTable.Initialize(CameraRigInfo.CameraRig->AllocationInfo.VariableTableInfo);
+	ScratchResult.VariableTable.OverrideAll(CameraSystemEvaluator->GetEvaluatedResult().VariableTable);
+
+	// Iterate on the solution.
+	FAimIterationInfo IterationInfo;
+	double LastErrorAngle = TNumericLimits<double>::Max();
+	double LastErrorDistance = TNumericLimits<double>::Max();
+
+	while (IterationInfo.IterationIndex < Params.MaxIterations)
+	{
+		DoRunIteration(Params, CameraRigInfo, IterationInfo);
+
+		if (IterationInfo.Result == EAimResult::Failed || IterationInfo.Result == EAimResult::Aborted || IterationInfo.Result == EAimResult::Completed)
+		{
+			break;
+		}
+		
+		ensure(IterationInfo.Result == EAimResult::Corrected);
+
+		// Check that we are getting closer to a solution.
+		if (IterationInfo.ErrorAngle >= LastErrorAngle || IterationInfo.ErrorDistance >= LastErrorDistance)
+		{
+			UE_LOG(LogCameraSystem, Error, TEXT("Can't converge towards a solution while aiming camera rig '%s'. Aborting."), *CameraRigInfo.CameraRig->GetPathName());
+			IterationInfo.Result = EAimResult::Aborted;
+			break;
+		}
+
+		LastErrorAngle = IterationInfo.ErrorAngle;
+		LastErrorDistance = IterationInfo.ErrorDistance;
+		++IterationInfo.IterationIndex;
+	}
+
+	return IterationInfo.Result == EAimResult::Completed || IterationInfo.Result == EAimResult::Corrected;
+}
+
+void FCameraIKAim::DoRunIteration(const FCameraIKAimParams& Params, const FCameraRigEvaluationInfo& CameraRigInfo, FAimIterationInfo& IterationInfo)
+{
+	// Run the system. We restore its state after each run.
+	RunRootCameraNode(Params, CameraRigInfo);
+
+	const bool bContinueIteration = CheckTolerance(Params, CameraRigInfo, IterationInfo);
+	if (!bContinueIteration)
+	{
+		return;
+	}
+
+	// Look at what we have... in 99% of cases, we're dealing with a camera rig that can be boiled down
+	// to 2 bones and one yaw/pitch articulation.
+	const FCameraVariableDefinition& YawPitchDefinition = FBuiltInCameraVariables::Get().YawPitchDefinition;
+	const FCameraRigJoint* FoundItem = ScratchResult.CameraRigJoints.GetJoints().FindByPredicate(
+			[&YawPitchDefinition](const FCameraRigJoint& Item)
+			{
+				return Item.VariableID == YawPitchDefinition.VariableID;
+			});
+	if (!FoundItem)
+	{
+		UE_LOG(LogCameraSystem, Error, TEXT("Can't figure out how to aim camera rig '%s'."), *CameraRigInfo.CameraRig->GetPathName());
+		IterationInfo.Result = EAimResult::Failed;
+		return;
+	}
+
+	const FCameraRigJoint& PivotJoint(*FoundItem);
+	AimTwoBonesCameraRig(Params, CameraRigInfo, PivotJoint.Transform, IterationInfo);
+}
+
+void FCameraIKAim::RunRootCameraNode(const FCameraIKAimParams& Params, const FCameraRigEvaluationInfo& CameraRigInfo)
+{
+	const FCameraSystemEvaluationResult& LastResult = Params.Evaluator->GetEvaluatedResult();
+	const FCameraNodeEvaluationResult& ContextInitialResult = CameraRigInfo.EvaluationContext->GetInitialResult();
+
+	FRootCameraNodeEvaluator* RootEvaluator = Params.Evaluator->GetRootNodeEvaluator();
+
+	// Save the initial state of the camera rig.
+	// TODO: we could optimize this a bit by caching the list of nodes that want serialization.
+	{
+		EvaluatorSnapshot.Reset();
+
+		FCameraNodeEvaluatorSerializeParams SerializeParams;
+		FMemoryWriter Writer(EvaluatorSnapshot);
+		RootEvaluator->Serialize(SerializeParams, Writer);
+	}
+
+	// Run the system.
+	{
+		FSingleCameraRigEvaluationParams SingleParams;
+		SingleParams.EvaluationParams.DeltaTime = Params.DeltaTime;
+		SingleParams.EvaluationParams.bIsFirstFrame = Params.bIsFirstFrame;
+		SingleParams.EvaluationParams.EvaluationType = ECameraNodeEvaluationType::IK;
+		SingleParams.EvaluationParams.EvaluationContext = CameraRigInfo.EvaluationContext;
+		SingleParams.EvaluationParams.Evaluator = Params.Evaluator;
+		SingleParams.CameraRigInfo = CameraRigInfo;
+		RootEvaluator->RunSingleCameraRig(SingleParams, ScratchResult);
+	}
+
+	// Restore the state of the camera rig.
+	{
+		FCameraNodeEvaluatorSerializeParams SerializeParams;
+		FMemoryReader Reader(EvaluatorSnapshot);
+		RootEvaluator->Serialize(SerializeParams, Reader);
+	}
+}
+
+bool FCameraIKAim::CheckTolerance(const FCameraIKAimParams& Params, const FCameraRigEvaluationInfo& CameraRigInfo, FAimIterationInfo& IterationInfo)
+{
+	const FCameraPose& ResultPose = ScratchResult.CameraPose;
+
+	// Figure out whether the distance to the current and/or desired target is too short.
+	// This is sometimes a good indication that the player is up against an obstacle and it's
+	// undesirable to turn their camera.
+	const double TargetDistance = ResultPose.GetTargetDistance();
+	const FVector3d CurrentAim(ResultPose.GetAimDir() * TargetDistance);
+	const FVector3d CurrentLocationToDesiredTarget(Params.TargetLocation - ResultPose.GetLocation());
+	const double DistanceToDesiredTarget(CurrentLocationToDesiredTarget.Length());
+
+	if (TargetDistance < Params.MinDistance || DistanceToDesiredTarget < Params.MinDistance)
+	{
+		UE_LOG(LogCameraSystem, Warning, TEXT("Aborting aiming of camera rig '%s': current target is %f away, minimium distance is %f"),
+				*CameraRigInfo.CameraRig->GetPathName(), TargetDistance, Params.MinDistance);
+		IterationInfo.Result = EAimResult::Aborted;
+		return false;
+	}
+
+	// See if we are within the angle or distance tolerance.
+	//
+	//    ||R ^ D|| = ||R||*||D||*sin(A)
+	//
+	// ..with R being the line of sight to the current target, D being the vector pointing to the 
+	// desired target, and A being the angle between the two. So:
+	//
+	//    A = asin(||R ^ D|| / (||R||*||D||)) 
+	//
+	// We also want to find H, the distance between the current line of sight and the desired target:
+	//
+	//    sin(A) = H / D
+	//    H = sin(A) * D
+	//    H = ||R ^ D|| / ||R||
+	//
+	// In the following code:
+	//   OrthLength = ||R ^ D||
+	//   SinAngle = sin(A)
+	//   ErrorAngle = A
+	//   ErrorDistance = H
+	//
+	const double OrthLength = FVector3d::CrossProduct(CurrentAim, CurrentLocationToDesiredTarget).Length();
+	const double SinAngle = OrthLength / (TargetDistance * DistanceToDesiredTarget);
+
+	const double ErrorAngle = FMath::RadiansToDegrees(FMath::Asin(SinAngle));
+	IterationInfo.ErrorAngle = ErrorAngle;
+	if (ErrorAngle <= Params.AngleTolerance)
+	{
+		IterationInfo.Result = EAimResult::Completed;
+		return false;
+	}
+
+	const double ErrorDistance = OrthLength / TargetDistance;
+	IterationInfo.ErrorDistance = ErrorDistance;
+	if (ErrorDistance <= Params.DistanceTolerance)
+	{
+		IterationInfo.Result = EAimResult::Completed;
+		return false;
+	}
+
+	return true;
+}
+
+void FCameraIKAim::AimTwoBonesCameraRig(const FCameraIKAimParams& Params, const FCameraRigEvaluationInfo& CameraRigInfo, const FTransform3d& PivotTransform, FAimIterationInfo& IterationInfo)
+{
+	FRotator3d Correction;
+	const bool bGotCorrection = ComputeTwoBonesCorrection(
+			ScratchResult.CameraPose, PivotTransform.GetLocation(), Params.TargetLocation, Correction);
+	if (!bGotCorrection)
+	{
+		IterationInfo.Result = EAimResult::Failed;
+		return;
+	}
+
+	FCameraOperationParams OperationParams;
+	OperationParams.Evaluator = Params.Evaluator;
+	OperationParams.EvaluationContext = CameraRigInfo.EvaluationContext;
+
+	FYawPitchCameraOperation Operation;
+	Operation.Yaw = FConsumableDouble::Delta(Correction.Yaw);
+	Operation.Pitch = FConsumableDouble::Delta(Correction.Pitch);
+
+	CameraRigInfo.RootEvaluator->ExecuteOperation(OperationParams, Operation);
+
+	if (Operation.Yaw.HasValue() || Operation.Pitch.HasValue())
+	{
+		UE_LOG(LogCameraSystem, Warning, TEXT("Aborting aiming of camera rig '%s': not all corrections were consumed by the camera nodes."),
+				*CameraRigInfo.CameraRig->GetPathName());
+		IterationInfo.Result = EAimResult::Aborted;
+		return;
+	}
+
+	IterationInfo.Result = EAimResult::Corrected;
+}
+
+bool FCameraIKAim::ComputeTwoBonesCorrection(const FCameraPose& CurrentPose, const FVector3d& PivotLocation, const FVector3d& DesiredTarget, FRotator3d& OutCorrection)
+{
+	// This is roughly the situation we are looking at, as seen from above:
+	//
+	//                         T
+	//                         |
+	//                         |    	
+	//             , - ~ D ~ - X
+	//         , '       |    /| ' ,
+	//       ,           |   / |     ,
+	//      ,            |  /  |      ,
+	//     ,             | /   |       ,
+	//     ,             |/    |       ,
+	//     ,             P     |       ,
+	//     ,              .    |       ,
+	//      ,               .  |      ,
+	//       ,                .C     ,
+	//         ,                  , '
+	//           ' - , _ _ _ ,  '
+	//    
+	// ...where:
+	//
+	//    P : the pivot
+	//    C : the current camera position
+	//    T : the current camera target
+	//    D : the desired camera target
+	//
+	// The sphere is centered on P, with its radius determined by D.
+	// The intersection of the camera's current line of sight with this circle is X.
+	// What we want is to turn the camera by A, the angle between PD and PX.
+	//
+	// We first compute the sphere's properties:
+	const FVector3d PivotToDesiredTarget = (DesiredTarget - PivotLocation);
+	const double SphereRadius = PivotToDesiredTarget.Length();
+
+	// Next we compute the intersection between the camera's line of sight and that sphere.
+	const FVector3d CameraLocation = CurrentPose.GetLocation();
+	const FVector3d CameraAim = CurrentPose.GetAimDir();
+
+	double DistanceToX = 0.0;
+	const bool bGotX = RaySphereIntersectExit(CameraLocation, CameraAim, PivotLocation, SphereRadius, DistanceToX);
+	if (!bGotX)
+	{
+		return false;
+	}
+
+	// Finally we compute the angle between PX and PD.
+	const FVector3d X = CameraLocation + CameraAim * DistanceToX;
+	const FVector3d PX = (X - PivotLocation);
+	const FVector3d PD = (DesiredTarget - PivotLocation);
+	// IMPORTANT NOTE: This assumes a vertical pivot axis!
+	const FRotator3d RotPX = PX.ToOrientationRotator();
+	const FRotator3d RotPD = PD.ToOrientationRotator();
+	OutCorrection = RotPD - RotPX;
+	return true;
+}
+
+bool FCameraIKAim::RaySphereIntersectExit(const FRay3d& Ray, const FVector3d& SphereOrigin, double SphereRadius, double& OutRayIntersectDistance)
+{
+	return RaySphereIntersectExit(Ray.Origin, Ray.Direction, SphereOrigin, SphereRadius, OutRayIntersectDistance);
+}
+
+bool FCameraIKAim::RaySphereIntersectExit(const FVector3d& RayStart, const FVector3d& RayDir, const FVector3d& SphereOrigin, double SphereRadius, double& OutRayIntersectDistance)
+{
+	// The equation for a point on the sphere is:
+	//
+	//    (x - Ox)^2 + (y - Oy)^2 + (z - Oz)^2 = R^2
+	//
+	// ...where:
+	//    x,y,z are the coordinates of the point
+	//    Ox,Oy,Oz are the coordinates of the sphere's origin
+	//    R is the radius of the sphere
+	//
+	// The equation for a point on the ray is:
+	//
+	//    x = rx + L*dx
+	//    y = ry + L*dy
+	//    z = rz + L*dz
+	//
+	// ...where:
+	//    rx,ry,rz are the coordinates of the ray's start
+	//    dx,dy,dz are the coordinates of the ray's direction vector
+	//    L is the linear coordinate of the point along the ray
+	//
+	// So we have:
+	//
+	//    (rx + L*dx - Ox)^2 + (ry + L*dy - Oy)^2 + (rz + L*dz - Oz)^2 = R^2
+	//
+	// Decomposing each term, we have:
+	//
+	//    dx*dx*L^2 + (rx - Ox)^2 + 2*L*dx*(rx - Ox)
+	//    dy*dy*L^2 + (ry - Oy)^2 + 2*L*dy*(ry - Oy)
+	//    dz*dz*L^2 + (rz - Oz)^2 + 2*L*dz*(rz - Oz)
+	//
+	// We can group then around L, which is the variable we're solving for:
+	//
+	//    (dx*dx + dy*dy + dz*dz)*L^2 + 2*(dx*Fx + dy*Fy + dz*Fz)*L + (Fx*Fx + Fy*Fy + Fz*Fz) = R^2
+	//
+	// ...where Fx,Fy,Fz are the coordinates of the vector from the sphere origin to 
+	// the ray's start.
+	//
+	// This is an equation of the form:  
+	//
+	//    a*L^2 + b*L + c = 0
+	//
+	// ...where:
+	//
+	//    a = (dx*dx + dy*dy + dz*dz) = d.d
+	//    b = 2*(dx*Fx + dy*Fy + dz*Fz) = 2*d.F
+	//    c = (Fx*Fx + Fy*Fy + Fz*Fz) - R^2 = F.F - R^2
+	//
+	// The discriminant for an equation of this form is:
+	//
+	//    D = b^2 - 4*a*c
+	//
+	// If the discriminant is positive, the solutions for the equations are:
+	//
+	//    L = (-b +- sqrt(D)) / (2*a)
+	//
+	// If the discriminant is zero, there's only one solution (such as when the ray barely "touches"
+	// the sphere). If the discriminant is negative, there's no real solution.
+	//
+	// In practice, we only want the solution that's "in front" of the ray, and is the "farthest",
+	// because we only want the exit point.
+	//
+	const FVector3d F = (RayStart - SphereOrigin);
+	const double dxd = RayDir.SizeSquared();
+	const double FxF = F.SizeSquared();
+	const double dxF = FVector3d::DotProduct(RayDir, F);
+
+	const double a = dxd;
+	const double b = 2.0 * dxF;
+	const double c = FxF - SphereRadius * SphereRadius;
+
+	const double D = b * b - 4.0 * a * c;
+
+	if (D > 0.0)
+	{
+		const double SqrtD = FMath::Sqrt(D);
+		const double two_a = 2.0 * a;
+
+		const double L1 = (-b + SqrtD) / two_a; 
+		const double L2 = (-b - SqrtD) / two_a; 
+
+		OutRayIntersectDistance = FMath::Max(L1, L2);
+		return true;
+	}
+
+	if (D == 0.0)
+	{
+		const double SqrtD = FMath::Sqrt(D);
+		const double two_a = 2.0 * a;
+
+		const double L = (-b / two_a);
+
+		OutRayIntersectDistance = L;
+		return L >= 0.0;
+	}
+
+	// D < 0.0
+	return false;
+}
+
+}  // namespace UE::Cameras
+
diff --git a/Engine/Plugins/Cameras/GameplayCameras/Source/GameplayCameras/Private/Core/CameraNodeEvaluator.cpp b/Engine/Plugins/Cameras/GameplayCameras/Source/GameplayCameras/Private/Core/CameraNodeEvaluator.cpp
--- a/Engine/Plugins/Cameras/GameplayCameras/Source/GameplayCameras/Private/Core/CameraNodeEvaluator.cpp
+++ b/Engine/Plugins/Cameras/GameplayCameras/Source/GameplayCameras/Private/Core/CameraNodeEvaluator.cpp
@@ -3,9 +3,7 @@
 #include "Core/CameraNodeEvaluator.h"
 
 #include "Core/CameraNode.h"
-#include "Debug/CameraDebugBlock.h"
 #include "Debug/CameraDebugBlockBuilder.h"
-#include "Debug/CameraDebugRenderer.h"
 #include "Debug/CameraNodeEvaluatorDebugBlock.h"
 #include "UObject/UObjectGlobals.h"
 
@@ -14,13 +12,28 @@ namespace UE::Cameras
 
 UE_GAMEPLAY_CAMERAS_DEFINE_RTTI(FCameraNodeEvaluator)
 
-void FCameraNodeEvaluationResult::Reset()
+void FCameraNodeEvaluationResult::Reset(bool bResetVariableTable)
 {
 	CameraPose.Reset();
+	CameraRigJoints.Reset();
+
+	if (bResetVariableTable)
+	{
+		VariableTable.UnsetAllValues();
+	}
+
 	bIsCameraCut = false;
 	bIsValid = false;
 }
 
+void FCameraNodeEvaluationResult::Serialize(FArchive& Ar)
+{
+	CameraPose.SerializeWithFlags(Ar);
+	VariableTable.Serialize(Ar);
+	Ar << bIsCameraCut;
+	Ar << bIsValid;
+}
+
 FCameraNodeEvaluator* FCameraNodeEvaluatorBuildParams::BuildEvaluator(const UCameraNode* InNode) const
 {
 	if (InNode)
@@ -51,6 +64,27 @@ FCameraNodeEvaluatorChildrenView FCameraNodeEvaluator::GetChildren()
 	return OnGetChildren();
 }
 
+void FCameraNodeEvaluator::ExecuteOperation(const FCameraOperationParams& Params, FCameraOperation& Operation)
+{
+	if (!PrivateCameraNode || PrivateCameraNode->bIsEnabled)
+	{
+		if (EnumHasAnyFlags(PrivateFlags, ECameraNodeEvaluatorFlags::SupportsOperations))
+		{
+			OnExecuteOperation(Params, Operation);
+		}
+		else
+		{
+			for (FCameraNodeEvaluator* Child : GetChildren())
+			{
+				if (Child)
+				{
+					Child->ExecuteOperation(Params, Operation);
+				}
+			}
+		}
+	}
+}
+
 void FCameraNodeEvaluator::AddReferencedObjects(FReferenceCollector& Collector)
 {
 	if (PrivateCameraNode)
@@ -69,6 +103,19 @@ void FCameraNodeEvaluator::AddReferencedObjects(FReferenceCollector& Collector)
 	}
 }
 
+void FCameraNodeEvaluator::Serialize(const FCameraNodeEvaluatorSerializeParams& Params, FArchive& Ar)
+{
+	OnSerialize(Params, Ar);
+
+	for (FCameraNodeEvaluator* Child : GetChildren())
+	{
+		if (Child)
+		{
+			Child->Serialize(Params, Ar);
+		}
+	}
+}
+
 void FCameraNodeEvaluator::Build(const FCameraNodeEvaluatorBuildParams& Params)
 {
 	OnBuild(Params);
diff --git a/Engine/Plugins/Cameras/GameplayCameras/Source/GameplayCameras/Private/Core/CameraPose.cpp b/Engine/Plugins/Cameras/GameplayCameras/Source/GameplayCameras/Private/Core/CameraPose.cpp
--- a/Engine/Plugins/Cameras/GameplayCameras/Source/GameplayCameras/Private/Core/CameraPose.cpp
+++ b/Engine/Plugins/Cameras/GameplayCameras/Source/GameplayCameras/Private/Core/CameraPose.cpp
@@ -167,6 +167,11 @@ FVector3d FCameraPose::GetAimDir() const
 	return Rotation.RotateVector(FVector3d{ 1, 0, 0 });
 }
 
+FVector3d FCameraPose::GetTarget() const
+{
+	return Location + TargetDistance * GetAimDir();
+}
+
 void FCameraPose::OverrideAll(const FCameraPose& OtherPose)
 {
 	InternalOverrideChanged(OtherPose, false);
diff --git a/Engine/Plugins/Cameras/GameplayCameras/Source/GameplayCameras/Private/Core/CameraRigJoints.cpp b/Engine/Plugins/Cameras/GameplayCameras/Source/GameplayCameras/Private/Core/CameraRigJoints.cpp
new file mode 100644
--- /dev/null
+++ b/Engine/Plugins/Cameras/GameplayCameras/Source/GameplayCameras/Private/Core/CameraRigJoints.cpp
@@ -0,0 +1,31 @@
+// Copyright Epic Games, Inc. All Rights Reserved.
+
+#include "Core/CameraRigJoints.h"
+
+#include "Core/BuiltInCameraVariables.h"
+
+namespace UE::Cameras
+{
+
+void FCameraRigJoints::AddJoint(const FCameraRigJoint& InJoint)
+{
+	Joints.Add(InJoint);
+}
+
+void FCameraRigJoints::AddJoint(const FCameraVariableDefinition& InVariableDefinition, const FTransform3d& InTransform)
+{
+	AddJoint(FCameraRigJoint{ InVariableDefinition.VariableID, InTransform });
+}
+
+void FCameraRigJoints::AddYawPitchJoint(const FTransform3d& InTransform)
+{
+	AddJoint(FBuiltInCameraVariables::Get().YawPitchDefinition, InTransform);
+}
+
+void FCameraRigJoints::Reset()
+{
+	Joints.Reset();
+}
+
+}  // namespace UE::Cameras
+
diff --git a/Engine/Plugins/Cameras/GameplayCameras/Source/GameplayCameras/Private/Core/CameraSystemEvaluator.cpp b/Engine/Plugins/Cameras/GameplayCameras/Source/GameplayCameras/Private/Core/CameraSystemEvaluator.cpp
--- a/Engine/Plugins/Cameras/GameplayCameras/Source/GameplayCameras/Private/Core/CameraSystemEvaluator.cpp
+++ b/Engine/Plugins/Cameras/GameplayCameras/Source/GameplayCameras/Private/Core/CameraSystemEvaluator.cpp
@@ -3,21 +3,17 @@
 #include "Core/CameraSystemEvaluator.h"
 
 #include "Camera/CameraTypes.h"
-#include "Core/CameraAsset.h"
 #include "Core/CameraDirectorEvaluator.h"
 #include "Core/CameraEvaluationContext.h"
 #include "Core/CameraEvaluationService.h"
-#include "Core/CameraRigAsset.h"
 #include "Core/DefaultRootCameraNode.h"
 #include "Debug/CameraDebugBlock.h"
 #include "Debug/CameraDebugBlockBuilder.h"
 #include "Debug/CameraDebugRenderer.h"
 #include "Debug/CameraSystemTrace.h"
-#include "Debug/CategoryTitleDebugBlock.h"
 #include "Debug/RootCameraDebugBlock.h"
-#include "HAL/IConsoleManager.h"
-#include "IGameplayCamerasModule.h"
 #include "Services/AutoResetCameraVariableService.h"
+#include "Services/OrientationInitializationService.h"
 #include "UObject/Package.h"
 #include "UObject/UObjectGlobals.h"
 
@@ -70,7 +66,9 @@ void FCameraSystemEvaluator::Initialize(const FCameraSystemEvaluatorCreateParams
 		RootEvaluator->Initialize(InitParams);
 	}
 
-	RegisterEvaluationService(MakeShared<FAutoResetCameraVariableService>());
+	VariableAutoResetService = MakeShared<FAutoResetCameraVariableService>();
+	RegisterEvaluationService(VariableAutoResetService.ToSharedRef());
+	RegisterEvaluationService(MakeShared<FOrientationInitializationService>());
 }
 
 FCameraSystemEvaluator::~FCameraSystemEvaluator()
@@ -137,33 +135,19 @@ void FCameraSystemEvaluator::NotifyRootCameraNodeEvent(const FRootCameraNodeCame
 {
 	for (TSharedPtr<FCameraEvaluationService> EvaluationService : EvaluationServices)
 	{
-		if (EvaluationService->HasAnyEvaluationServiceFlags(ECameraEvaluationServiceFlags::NeedsRootCameraNodeEvents))
+		if (EvaluationService->HasAllEvaluationServiceFlags(ECameraEvaluationServiceFlags::NeedsRootCameraNodeEvents))
 		{
 			EvaluationService->NotifyRootCameraNodeEvent(InEvent);
 		}
 	}
 }
 
-void FCameraSystemEvaluator::Update(const FCameraSystemEvaluationUpdateParams& Params)
+void FCameraSystemEvaluator::Update(const FCameraSystemEvaluationParams& Params)
 {
 	SCOPE_CYCLE_COUNTER(CameraSystemEval_Total);
 
 	// Pre-update all services.
-	{
-		FCameraEvaluationServiceUpdateParams ServiceUpdateParams;
-		ServiceUpdateParams.Evaluator = this;
-		ServiceUpdateParams.DeltaTime = Params.DeltaTime;
-
-		FCameraEvaluationServiceUpdateResult ServiceUpdateResult(RootNodeResult);
-
-		for (TSharedPtr<FCameraEvaluationService> EvaluationService : EvaluationServices)
-		{
-			if (EvaluationService->HasAnyEvaluationServiceFlags(ECameraEvaluationServiceFlags::NeedsPreUpdate))
-			{
-				EvaluationService->PreUpdate(ServiceUpdateParams, ServiceUpdateResult);
-			}
-		}
-	}
+	PreUpdateServices(Params.DeltaTime, ECameraEvaluationServiceFlags::None);
 
 	// Get the active evaluation context.
 	TSharedPtr<FCameraEvaluationContext> ActiveContext = ContextStack.GetActiveContext();
@@ -198,35 +182,60 @@ void FCameraSystemEvaluator::Update(const FCameraSystemEvaluationUpdateParams& P
 		// TODO: handle the case of composite camera rigs.
 	}
 
-	// Setup the params/result for running the root camera node.
-	FCameraNodeEvaluationParams NodeParams;
-	NodeParams.Evaluator = this;
-	NodeParams.DeltaTime = Params.DeltaTime;
+	{
+		// Setup the params/result for running the root camera node.
+		FCameraNodeEvaluationParams NodeParams;
+		NodeParams.Evaluator = this;
+		NodeParams.DeltaTime = Params.DeltaTime;
+
+		RootNodeResult.Reset(false);
 
-	RootNodeResult.Reset();
+		// Run the root camera node.
+		RootEvaluator->Run(NodeParams, RootNodeResult);
 
-	// Run the root camera node.
-	RootEvaluator->Run(NodeParams, RootNodeResult);
+		RootNodeResult.bIsValid = true;
+	}
 
 	// Harvest the result.
-	Result.CameraPose = RootNodeResult.CameraPose;
+	Result.CameraPose.OverrideAll(RootNodeResult.CameraPose);
+	Result.VariableTable.OverrideAll(RootNodeResult.VariableTable);
 	Result.bIsCameraCut = RootNodeResult.bIsCameraCut;
 	Result.bIsValid = true;
 
 	// Post-update all services.
+	PostUpdateServices(Params.DeltaTime, ECameraEvaluationServiceFlags::None);
+}
+
+void FCameraSystemEvaluator::PreUpdateServices(float DeltaTime, ECameraEvaluationServiceFlags ExtraFlags)
+{
+	FCameraEvaluationServiceUpdateParams ServiceUpdateParams;
+	ServiceUpdateParams.Evaluator = this;
+	ServiceUpdateParams.DeltaTime = DeltaTime;
+
+	FCameraEvaluationServiceUpdateResult ServiceUpdateResult(RootNodeResult);
+
+	for (TSharedPtr<FCameraEvaluationService> EvaluationService : EvaluationServices)
 	{
-		FCameraEvaluationServiceUpdateParams ServiceUpdateParams;
-		ServiceUpdateParams.Evaluator = this;
-		ServiceUpdateParams.DeltaTime = Params.DeltaTime;
+		if (EvaluationService->HasAllEvaluationServiceFlags(ECameraEvaluationServiceFlags::NeedsPreUpdate | ExtraFlags))
+		{
+			EvaluationService->PreUpdate(ServiceUpdateParams, ServiceUpdateResult);
+		}
+	}
+}
 
-		FCameraEvaluationServiceUpdateResult ServiceUpdateResult(RootNodeResult);
+void FCameraSystemEvaluator::PostUpdateServices(float DeltaTime, ECameraEvaluationServiceFlags ExtraFlags)
+{
+	FCameraEvaluationServiceUpdateParams ServiceUpdateParams;
+	ServiceUpdateParams.Evaluator = this;
+	ServiceUpdateParams.DeltaTime = DeltaTime;
 
-		for (TSharedPtr<FCameraEvaluationService> EvaluationService : EvaluationServices)
+	FCameraEvaluationServiceUpdateResult ServiceUpdateResult(RootNodeResult);
+
+	for (TSharedPtr<FCameraEvaluationService> EvaluationService : EvaluationServices)
+	{
+		if (EvaluationService->HasAllEvaluationServiceFlags(ECameraEvaluationServiceFlags::NeedsPostUpdate | ExtraFlags))
 		{
-			if (EvaluationService->HasAnyEvaluationServiceFlags(ECameraEvaluationServiceFlags::NeedsPostUpdate))
-			{
-				EvaluationService->PostUpdate(ServiceUpdateParams, ServiceUpdateResult);
-			}
+			EvaluationService->PostUpdate(ServiceUpdateParams, ServiceUpdateResult);
 		}
 	}
 }
diff --git a/Engine/Plugins/Cameras/GameplayCameras/Source/GameplayCameras/Private/Core/DefaultRootCameraNode.cpp b/Engine/Plugins/Cameras/GameplayCameras/Source/GameplayCameras/Private/Core/DefaultRootCameraNode.cpp
--- a/Engine/Plugins/Cameras/GameplayCameras/Source/GameplayCameras/Private/Core/DefaultRootCameraNode.cpp
+++ b/Engine/Plugins/Cameras/GameplayCameras/Source/GameplayCameras/Private/Core/DefaultRootCameraNode.cpp
@@ -3,6 +3,7 @@
 #include "Core/DefaultRootCameraNode.h"
 
 #include "Core/BlendStackCameraNode.h"
+#include "Core/CameraEvaluationContext.h"
 #include "Core/CameraRigAsset.h"
 #include "Core/RootCameraNodeObserver.h"
 #include "Debug/BlendStacksCameraDebugBlock.h"
@@ -80,23 +81,7 @@ void FDefaultRootCameraNodeEvaluator::OnRun(const FCameraNodeEvaluationParams& P
 
 void FDefaultRootCameraNodeEvaluator::OnActivateCameraRig(const FActivateCameraRigParams& Params)
 {
-	FBlendStackCameraNodeEvaluator* TargetStack = nullptr;
-	switch (Params.Layer)
-	{
-		case ECameraRigLayer::Base:
-			TargetStack = BaseLayer;
-			break;
-		case ECameraRigLayer::Main:
-			TargetStack = MainLayer;
-			break;
-		case ECameraRigLayer::Global:
-			TargetStack = GlobalLayer;
-			break;
-		case ECameraRigLayer::Visual:
-			TargetStack = VisualLayer;
-			break;
-	}
-
+	FBlendStackCameraNodeEvaluator* TargetStack = GetBlendStackEvaluator(Params.Layer);
 	if (ensure(TargetStack))
 	{
 		FBlendStackCameraPushParams PushParams;
@@ -107,6 +92,53 @@ void FDefaultRootCameraNodeEvaluator::OnActivateCameraRig(const FActivateCameraR
 	}
 }
 
+void FDefaultRootCameraNodeEvaluator::OnRunSingleCameraRig(const FSingleCameraRigEvaluationParams& Params, FCameraNodeEvaluationResult& OutResult)
+{
+	BaseLayer->Run(Params.EvaluationParams, OutResult);
+
+	FCameraNodeEvaluator* RootEvaluator = Params.CameraRigInfo.RootEvaluator;
+
+	{
+		const FCameraNodeEvaluationResult* CameraRigResult = Params.CameraRigInfo.LastResult;
+		FCameraBlendedParameterUpdateParams InputParams(Params.EvaluationParams, CameraRigResult->CameraPose);
+		FCameraBlendedParameterUpdateResult InputResult(OutResult.VariableTable);
+		RootEvaluator->UpdateParameters(InputParams, InputResult);
+	}
+
+	{
+		OutResult.CameraPose.ClearAllChangedFlags();
+		OutResult.VariableTable.ClearAllWrittenThisFrameFlags();
+
+		const FCameraNodeEvaluationResult& InitialResult = Params.CameraRigInfo.EvaluationContext->GetInitialResult();
+		OutResult.CameraPose.OverrideChanged(InitialResult.CameraPose);
+		OutResult.VariableTable.OverrideAll(InitialResult.VariableTable);
+
+		OutResult.bIsValid = true;
+
+		RootEvaluator->Run(Params.EvaluationParams, OutResult);
+	}
+
+	GlobalLayer->Run(Params.EvaluationParams, OutResult);
+	// Don't run the visual layer.
+}
+
+FBlendStackCameraNodeEvaluator* FDefaultRootCameraNodeEvaluator::GetBlendStackEvaluator(ECameraRigLayer Layer) const
+{
+	switch (Layer)
+	{
+		case ECameraRigLayer::Base:
+			return BaseLayer;
+		case ECameraRigLayer::Main:
+			return MainLayer;
+		case ECameraRigLayer::Global:
+			return GlobalLayer;
+		case ECameraRigLayer::Visual:
+			return VisualLayer;
+		default:
+			return nullptr;
+	}
+}
+
 void FDefaultRootCameraNodeEvaluator::OnBlendStackEvent(const FBlendStackCameraRigEvent& InEvent)
 {
 	if (InEvent.EventType == EBlendStackCameraRigEventType::Pushed ||
diff --git a/Engine/Plugins/Cameras/GameplayCameras/Source/GameplayCameras/Private/Core/RootCameraNode.cpp b/Engine/Plugins/Cameras/GameplayCameras/Source/GameplayCameras/Private/Core/RootCameraNode.cpp
--- a/Engine/Plugins/Cameras/GameplayCameras/Source/GameplayCameras/Private/Core/RootCameraNode.cpp
+++ b/Engine/Plugins/Cameras/GameplayCameras/Source/GameplayCameras/Private/Core/RootCameraNode.cpp
@@ -2,8 +2,10 @@
 
 #include "Core/RootCameraNode.h"
 
+#include "Core/CameraEvaluationService.h"
 #include "Core/CameraSystemEvaluator.h"
 #include "Core/RootCameraNodeObserver.h"
+#include "Services/AutoResetCameraVariableService.h"
 
 #include UE_INLINE_GENERATED_CPP_BY_NAME(RootCameraNode)
 
@@ -20,6 +22,25 @@ void FRootCameraNodeEvaluator::ActivateCameraRig(const FActivateCameraRigParams&
 	OnActivateCameraRig(Params);
 }
 
+void FRootCameraNodeEvaluator::RunSingleCameraRig(const FSingleCameraRigEvaluationParams& Params, FCameraNodeEvaluationResult& OutResult)
+{
+	// Before we do the actual evaluation we need to ask the system to auto-reset
+	// any camera variable that needs auto-resetting. Otherwise, we might end up with
+	// an update result that isn't representative of what would happen normally.
+	if (FCameraSystemEvaluator* Evaluator = Params.EvaluationParams.Evaluator)
+	{
+		FCameraEvaluationServiceUpdateParams ServiceUpdateParams;
+		ServiceUpdateParams.Evaluator = Evaluator;
+		ServiceUpdateParams.DeltaTime = Params.EvaluationParams.DeltaTime;
+
+		FCameraEvaluationServiceUpdateResult ServiceUpdateResult(OutResult);
+
+		Evaluator->VariableAutoResetService->PreUpdate(ServiceUpdateParams, ServiceUpdateResult);
+	}
+
+	OnRunSingleCameraRig(Params, OutResult);
+}
+
 void FRootCameraNodeEvaluator::RegisterObserver(IRootCameraNodeObserver* Observer)
 {
 	Observers.Add(Observer);
diff --git a/Engine/Plugins/Cameras/GameplayCameras/Source/GameplayCameras/Private/Debug/CameraSystemTrace.cpp b/Engine/Plugins/Cameras/GameplayCameras/Source/GameplayCameras/Private/Debug/CameraSystemTrace.cpp
--- a/Engine/Plugins/Cameras/GameplayCameras/Source/GameplayCameras/Private/Debug/CameraSystemTrace.cpp
+++ b/Engine/Plugins/Cameras/GameplayCameras/Source/GameplayCameras/Private/Debug/CameraSystemTrace.cpp
@@ -293,7 +293,7 @@ bool FCameraSystemTrace::IsTraceEnabled()
 	return GGameplayCamerasDebugTrace || UE_TRACE_CHANNELEXPR_IS_ENABLED(CameraSystemChannel);
 }
 
-void FCameraSystemTrace::TraceEvaluation(UWorld* InWorld, const FCameraSystemEvaluationUpdateResult& InResult, FCameraDebugBlock& InRootDebugBlock)
+void FCameraSystemTrace::TraceEvaluation(UWorld* InWorld, const FCameraSystemEvaluationResult& InResult, FCameraDebugBlock& InRootDebugBlock)
 {
 	if (!IsTraceEnabled())
 	{
diff --git a/Engine/Plugins/Cameras/GameplayCameras/Source/GameplayCameras/Private/Debug/RootCameraDebugBlock.cpp b/Engine/Plugins/Cameras/GameplayCameras/Source/GameplayCameras/Private/Debug/RootCameraDebugBlock.cpp
--- a/Engine/Plugins/Cameras/GameplayCameras/Source/GameplayCameras/Private/Debug/RootCameraDebugBlock.cpp
+++ b/Engine/Plugins/Cameras/GameplayCameras/Source/GameplayCameras/Private/Debug/RootCameraDebugBlock.cpp
@@ -96,7 +96,7 @@ void FRootCameraDebugBlock::BuildDebugBlocks(const FCameraSystemEvaluator& Camer
 		PoseStatsCategory.Title = TEXT("Evaluated Camera");
 		PoseStatsCategory.Category = FCameraDebugCategories::PoseStats;
 
-		const FCameraSystemEvaluationUpdateResult& Result = CameraSystem.GetEvaluatedResult();
+		const FCameraSystemEvaluationResult& Result = CameraSystem.GetEvaluatedResult();
 		Builder.AttachDebugBlock<FCameraPoseDebugBlock>(Result.CameraPose)
 			.WithShowUnchangedCVar(TEXT("GameplayCameras.Debug.PoseStats.ShowUnchanged"));
 	}
diff --git a/Engine/Plugins/Cameras/GameplayCameras/Source/GameplayCameras/Private/GameFramework/GameplayCameraSystemComponent.cpp b/Engine/Plugins/Cameras/GameplayCameras/Source/GameplayCameras/Private/GameFramework/GameplayCameraSystemComponent.cpp
--- a/Engine/Plugins/Cameras/GameplayCameras/Source/GameplayCameras/Private/GameFramework/GameplayCameraSystemComponent.cpp
+++ b/Engine/Plugins/Cameras/GameplayCameras/Source/GameplayCameras/Private/GameFramework/GameplayCameraSystemComponent.cpp
@@ -39,7 +39,7 @@ void UGameplayCameraSystemComponent::GetCameraView(float DeltaTime, FMinimalView
 
 	if (Evaluator.IsValid())
 	{
-		FCameraSystemEvaluationUpdateParams UpdateParams;
+		FCameraSystemEvaluationParams UpdateParams;
 		UpdateParams.DeltaTime = DeltaTime;
 		Evaluator->Update(UpdateParams);
 
diff --git a/Engine/Plugins/Cameras/GameplayCameras/Source/GameplayCameras/Private/GameplayCamerasModule.cpp b/Engine/Plugins/Cameras/GameplayCameras/Source/GameplayCameras/Private/GameplayCamerasModule.cpp
--- a/Engine/Plugins/Cameras/GameplayCameras/Source/GameplayCameras/Private/GameplayCamerasModule.cpp
+++ b/Engine/Plugins/Cameras/GameplayCameras/Source/GameplayCameras/Private/GameplayCamerasModule.cpp
@@ -8,9 +8,13 @@
 #include "Debug/CameraDebugColors.h"
 #include "Features/IModularFeatures.h"
 #include "GameplayCameras.h"
+#include "GameplayCamerasSettings.h"
+#include "ISettingsModule.h"
 #include "Logging/MessageLog.h"
 #include "Modules/ModuleManager.h"
 
+#define LOCTEXT_NAMESPACE "GameplayCamerasModule"
+
 DEFINE_LOG_CATEGORY(LogCameraSystem);
 
 IGameplayCamerasModule& IGameplayCamerasModule::Get()
@@ -25,6 +29,8 @@ public:
 	// IModuleInterface interface
 	virtual void StartupModule() override
 	{
+		RegisterSettings();
+
 		CameraModularFeature = MakeShared<FCameraModularFeature>();
 		if (CameraModularFeature.IsValid())
 		{
@@ -38,6 +44,8 @@ public:
 
 	virtual void ShutdownModule() override
 	{
+		UnregisterSettings();
+
 		if (CameraModularFeature.IsValid())
 		{
 			IModularFeatures::Get().UnregisterModularFeature(ICameraModularFeature::GetModularFeatureName(), CameraModularFeature.Get());
@@ -45,6 +53,8 @@ public:
 		}
 	}
 
+public:
+
 	// IGameplayCamerasModule interface
 #if WITH_EDITOR
 	virtual TSharedPtr<IGameplayCamerasLiveEditManager> GetLiveEditManager() const override
@@ -59,6 +69,33 @@ public:
 #endif
 
 private:
+
+	void RegisterSettings()
+	{
+		ISettingsModule* SettingsModule = FModuleManager::GetModulePtr<ISettingsModule>("Settings");
+
+		if (SettingsModule != nullptr)
+		{
+			SettingsModule->RegisterSettings("Project", "Plugins", "Gameplay Cameras",
+				LOCTEXT("GameplayCamerasProjectSettingsName", "Gameplay Cameras"),
+				LOCTEXT("GameplayCamerasProjectSettingsDescription", "Configure gameplay cameras."),
+				GetMutableDefault<UGameplayCamerasSettings>()
+			);
+		}
+	}
+
+	void UnregisterSettings()
+	{
+		ISettingsModule* SettingsModule = FModuleManager::GetModulePtr<ISettingsModule>("Settings");
+
+		if (SettingsModule != nullptr)
+		{
+			SettingsModule->UnregisterSettings("Project", "Plugins", "Gameplay Cameras");
+		}
+	}
+
+private:
+
 	class FCameraModularFeature : public ICameraModularFeature
 	{
 		// ICameraModularFeature interface
@@ -76,3 +113,6 @@ private:
 };
 
 IMPLEMENT_MODULE(FGameplayCamerasModule, GameplayCameras);
+
+#undef LOCTEXT_NAMESPACE
+
diff --git a/Engine/Plugins/Cameras/GameplayCameras/Source/GameplayCameras/Private/Nodes/Blends/SimpleBlendCameraNode.cpp b/Engine/Plugins/Cameras/GameplayCameras/Source/GameplayCameras/Private/Nodes/Blends/SimpleBlendCameraNode.cpp
--- a/Engine/Plugins/Cameras/GameplayCameras/Source/GameplayCameras/Private/Nodes/Blends/SimpleBlendCameraNode.cpp
+++ b/Engine/Plugins/Cameras/GameplayCameras/Source/GameplayCameras/Private/Nodes/Blends/SimpleBlendCameraNode.cpp
@@ -52,6 +52,12 @@ void FSimpleBlendCameraNodeEvaluator::OnBlendResults(const FCameraNodeBlendParam
 	OutResult.bIsBlendFinished = bIsBlendFinished;
 }
 
+void FSimpleBlendCameraNodeEvaluator::OnSerialize(const FCameraNodeEvaluatorSerializeParams& Params, FArchive& Ar)
+{
+	Ar << BlendFactor;
+	Ar << bIsBlendFinished;
+}
+
 #if UE_GAMEPLAY_CAMERAS_DEBUG
 
 void FSimpleBlendCameraNodeEvaluator::OnBuildDebugBlocks(const FCameraDebugBlockBuildParams& Params, FCameraDebugBlockBuilder& Builder)
diff --git a/Engine/Plugins/Cameras/GameplayCameras/Source/GameplayCameras/Private/Nodes/Collision/OcclusionMaterialCameraNode.cpp b/Engine/Plugins/Cameras/GameplayCameras/Source/GameplayCameras/Private/Nodes/Collision/OcclusionMaterialCameraNode.cpp
--- a/Engine/Plugins/Cameras/GameplayCameras/Source/GameplayCameras/Private/Nodes/Collision/OcclusionMaterialCameraNode.cpp
+++ b/Engine/Plugins/Cameras/GameplayCameras/Source/GameplayCameras/Private/Nodes/Collision/OcclusionMaterialCameraNode.cpp
@@ -70,6 +70,12 @@ void FOcclusionMaterialCameraNodeEvaluator::OnRun(const FCameraNodeEvaluationPar
 		return;
 	}
 
+	if (Params.EvaluationType != ECameraNodeEvaluationType::Standard)
+	{
+		// Don't run occlusion traces during IK/stateless updates.
+		return;
+	}
+
 	UWorld* World = Params.EvaluationContext->GetWorld();
 	APlayerController* PlayerController = Params.EvaluationContext->GetPlayerController();
 	if (!World || !PlayerController)
diff --git a/Engine/Plugins/Cameras/GameplayCameras/Source/GameplayCameras/Private/Nodes/Common/BoomArmCameraNode.cpp b/Engine/Plugins/Cameras/GameplayCameras/Source/GameplayCameras/Private/Nodes/Common/BoomArmCameraNode.cpp
--- a/Engine/Plugins/Cameras/GameplayCameras/Source/GameplayCameras/Private/Nodes/Common/BoomArmCameraNode.cpp
+++ b/Engine/Plugins/Cameras/GameplayCameras/Source/GameplayCameras/Private/Nodes/Common/BoomArmCameraNode.cpp
@@ -3,11 +3,12 @@
 #include "Nodes/Common/BoomArmCameraNode.h"
 
 #include "Core/CameraEvaluationContext.h"
+#include "Core/CameraOperation.h"
 #include "Core/CameraParameterReader.h"
-#include "Nodes/Input/CameraRigInput2DSlot.h"
+#include "Core/CameraRigJoints.h"
 #include "GameFramework/Pawn.h"
 #include "GameFramework/PlayerController.h"
-#include "GameplayCameras.h"
+#include "Nodes/Input/CameraRigInput2DSlot.h"
 
 #include UE_INLINE_GENERATED_CPP_BY_NAME(BoomArmCameraNode)
 
@@ -24,6 +25,11 @@ protected:
 	virtual void OnInitialize(const FCameraNodeEvaluatorInitializeParams& Params) override;
 	virtual FCameraNodeEvaluatorChildrenView OnGetChildren() override;
 	virtual void OnRun(const FCameraNodeEvaluationParams& Params, FCameraNodeEvaluationResult& OutResult) override;
+	virtual void OnExecuteOperation(const FCameraOperationParams& Params, FCameraOperation& Operation) override;
+
+private:
+
+	APlayerController* GetPlayerController(TSharedPtr<const FCameraEvaluationContext> EvaluationContext) const;
 
 private:
 
@@ -41,6 +47,10 @@ void FBoomArmCameraNodeEvaluator::OnBuild(const FCameraNodeEvaluatorBuildParams&
 
 void FBoomArmCameraNodeEvaluator::OnInitialize(const FCameraNodeEvaluatorInitializeParams& Params)
 {
+	SetNodeEvaluatorFlags(
+			ECameraNodeEvaluatorFlags::NeedsEvaluationUpdate |
+			ECameraNodeEvaluatorFlags::SupportsOperations);
+
 	const UBoomArmCameraNode* BoomArmNode = GetCameraNodeAs<UBoomArmCameraNode>();
 	BoomOffsetReader.Initialize(BoomArmNode->BoomOffset);
 }
@@ -59,16 +69,10 @@ void FBoomArmCameraNodeEvaluator::OnRun(const FCameraNodeEvaluationParams& Param
 		const FVector2d YawPitch = InputSlotEvaluator->GetInputValue();
 		BoomRotation = FRotator3d(YawPitch.Y, YawPitch.X, 0);
 	}
-	else if (Params.EvaluationContext)
+	else if (APlayerController* PlayerController = GetPlayerController(Params.EvaluationContext))
 	{
-		if (APlayerController* PlayerController = Params.EvaluationContext->GetPlayerController())
-		{
-			if (APawn* Pawn = PlayerController->GetPawn())
-			{
-				const FRotator3d PawnViewRotation = Pawn->GetViewRotation();
-				BoomRotation = PawnViewRotation;
-			}
-		}
+		const FRotator3d ControlRotation = PlayerController->GetControlRotation();
+		BoomRotation = ControlRotation;
 	}
 
 	const UBoomArmCameraNode* BoomArmNode = GetCameraNodeAs<UBoomArmCameraNode>();
@@ -86,6 +90,40 @@ void FBoomArmCameraNodeEvaluator::OnRun(const FCameraNodeEvaluationParams& Param
 	const FTransform3d FinalTransform(BoomOffset * BoomPivot);
 
 	OutResult.CameraPose.SetTransform(FinalTransform);
+	
+	OutResult.CameraRigJoints.AddYawPitchJoint(BoomPivot);
+}
+
+void FBoomArmCameraNodeEvaluator::OnExecuteOperation(const FCameraOperationParams& Params, FCameraOperation& Operation)
+{
+	if (InputSlotEvaluator)
+	{
+		InputSlotEvaluator->ExecuteOperation(Params, Operation);
+	}
+	else
+	{
+		// If we don't have an input slot, we use the pawn rotation directly in OnRun. So let's handle
+		// some operations by affecting that pawn rotation ourselves.
+		if (FYawPitchCameraOperation* Op = Operation.CastOperation<FYawPitchCameraOperation>())
+		{
+			if (APlayerController* PlayerController = GetPlayerController(Params.EvaluationContext))
+			{
+				FRotator3d ControlRotation = PlayerController->GetControlRotation();
+				ControlRotation.Yaw = Op->Yaw.Apply(ControlRotation.Yaw);
+				ControlRotation.Pitch = Op->Pitch.Apply(ControlRotation.Pitch);
+				PlayerController->SetControlRotation(ControlRotation);
+			}
+		}
+	}
+}
+
+APlayerController* FBoomArmCameraNodeEvaluator::GetPlayerController(TSharedPtr<const FCameraEvaluationContext> EvaluationContext) const
+{
+	if (EvaluationContext)
+	{
+		return EvaluationContext->GetPlayerController();
+	}
+	return nullptr;
 }
 
 }  // namespace UE::Cameras
diff --git a/Engine/Plugins/Cameras/GameplayCameras/Source/GameplayCameras/Private/Nodes/Common/DampenPositionCameraNode.cpp b/Engine/Plugins/Cameras/GameplayCameras/Source/GameplayCameras/Private/Nodes/Common/DampenPositionCameraNode.cpp
--- a/Engine/Plugins/Cameras/GameplayCameras/Source/GameplayCameras/Private/Nodes/Common/DampenPositionCameraNode.cpp
+++ b/Engine/Plugins/Cameras/GameplayCameras/Source/GameplayCameras/Private/Nodes/Common/DampenPositionCameraNode.cpp
@@ -24,6 +24,7 @@ protected:
 
 	virtual void OnInitialize(const FCameraNodeEvaluatorInitializeParams& Params) override;
 	virtual void OnRun(const FCameraNodeEvaluationParams& Params, FCameraNodeEvaluationResult& OutResult) override;
+	virtual void OnSerialize(const FCameraNodeEvaluatorSerializeParams& Params, FArchive& Ar) override;
 
 #if UE_GAMEPLAY_CAMERAS_DEBUG
 	virtual void OnBuildDebugBlocks(const FCameraDebugBlockBuildParams& Params, FCameraDebugBlockBuilder& Builder) override;
@@ -137,6 +138,20 @@ void FDampenPositionCameraNodeEvaluator::OnRun(const FCameraNodeEvaluationParams
 	OutResult.CameraPose.SetLocation(NextLocation);
 }
 
+void FDampenPositionCameraNodeEvaluator::OnSerialize(const FCameraNodeEvaluatorSerializeParams& Params, FArchive& Ar)
+{
+	Ar << ForwardDamper;
+	Ar << LateralDamper;
+	Ar << VerticalDamper;
+
+	Ar << PreviousLocation;
+
+#if UE_GAMEPLAY_CAMERAS_DEBUG
+	Ar << LastUndampedPosition;
+	Ar << LastDampedPosition;
+#endif  // UE_GAMEPLAY_CAMERAS_DEBUG
+}
+
 #if UE_GAMEPLAY_CAMERAS_DEBUG
 
 void FDampenPositionCameraNodeEvaluator::OnBuildDebugBlocks(const FCameraDebugBlockBuildParams& Params, FCameraDebugBlockBuilder& Builder)
diff --git a/Engine/Plugins/Cameras/GameplayCameras/Source/GameplayCameras/Private/Nodes/Input/CameraRigInput1DSlot.cpp b/Engine/Plugins/Cameras/GameplayCameras/Source/GameplayCameras/Private/Nodes/Input/CameraRigInput1DSlot.cpp
--- a/Engine/Plugins/Cameras/GameplayCameras/Source/GameplayCameras/Private/Nodes/Input/CameraRigInput1DSlot.cpp
+++ b/Engine/Plugins/Cameras/GameplayCameras/Source/GameplayCameras/Private/Nodes/Input/CameraRigInput1DSlot.cpp
@@ -3,6 +3,7 @@
 #include "Nodes/Input/CameraRigInput1DSlot.h"
 
 #include "Core/CameraBuildLog.h"
+#include "Core/CameraOperation.h"
 #include "Core/CameraRigAsset.h"
 #include "Core/CameraRigBuildContext.h"
 #include "Core/CameraVariableAssets.h"
@@ -18,7 +19,10 @@ UE_DEFINE_CAMERA_NODE_EVALUATOR(FCameraRigInput1DSlotEvaluator)
 
 FCameraRigInput1DSlotEvaluator::FCameraRigInput1DSlotEvaluator()
 {
-	SetNodeEvaluatorFlags(ECameraNodeEvaluatorFlags::NeedsParameterUpdate | ECameraNodeEvaluatorFlags::NeedsEvaluationUpdate);
+	SetNodeEvaluatorFlags(
+			ECameraNodeEvaluatorFlags::NeedsParameterUpdate | 
+			ECameraNodeEvaluatorFlags::NeedsEvaluationUpdate |
+			ECameraNodeEvaluatorFlags::SupportsOperations);
 }
 
 void FCameraRigInput1DSlotEvaluator::OnBuild(const FCameraNodeEvaluatorBuildParams& Params)
@@ -33,9 +37,9 @@ void FCameraRigInput1DSlotEvaluator::OnInitialize(const FCameraNodeEvaluatorInit
 	InputValue = 0.f;
 
 	const UCameraRigInput1DSlot* SlotNode = GetCameraNodeAs<UCameraRigInput1DSlot>();
-	if (SlotNode->GetVariableID().IsValid() && Params.LastActiveCameraRig.IsValid())
+	if (SlotNode->GetVariableID().IsValid() && Params.LastActiveCameraRigInfo.LastResult)
 	{
-		const FCameraVariableTable& LastActiveRigVariableTable = Params.LastActiveCameraRig.LastResult->VariableTable;
+		const FCameraVariableTable& LastActiveRigVariableTable = Params.LastActiveCameraRigInfo.LastResult->VariableTable;
 		LastActiveRigVariableTable.TryGetValue<double>(SlotNode->GetVariableID(), InputValue);
 	}
 }
@@ -85,6 +89,25 @@ void FCameraRigInput1DSlotEvaluator::OnRun(const FCameraNodeEvaluationParams& Pa
 	OutResult.VariableTable.SetValue<double>(SlotNode->GetVariableID(), InputValue);
 }
 
+void FCameraRigInput1DSlotEvaluator::OnExecuteOperation(const FCameraOperationParams& Params, FCameraOperation& Operation)
+{
+	if (FSingleValueCameraOperation* Op = Operation.CastOperation<FSingleValueCameraOperation>())
+	{
+		const UCameraRigInput1DSlot* SlotNode = GetCameraNodeAs<UCameraRigInput1DSlot>();
+
+		double MinValue, MaxValue;
+		SlotNode->Clamp.GetEffectiveClamping(MinValue, MaxValue);
+
+		InputValue = Op->Value.Apply(InputValue, MinValue, MaxValue);
+	}
+}
+
+void FCameraRigInput1DSlotEvaluator::OnSerialize(const FCameraNodeEvaluatorSerializeParams& Params, FArchive& Ar)
+{
+	Ar << TransientInputValue;
+	Ar << InputValue;
+}
+
 }  // namespace UE::Cameras
 
 FCameraNodeChildrenView UCameraRigInput1DSlot::OnGetChildren()
diff --git a/Engine/Plugins/Cameras/GameplayCameras/Source/GameplayCameras/Private/Nodes/Input/CameraRigInput2DSlot.cpp b/Engine/Plugins/Cameras/GameplayCameras/Source/GameplayCameras/Private/Nodes/Input/CameraRigInput2DSlot.cpp
--- a/Engine/Plugins/Cameras/GameplayCameras/Source/GameplayCameras/Private/Nodes/Input/CameraRigInput2DSlot.cpp
+++ b/Engine/Plugins/Cameras/GameplayCameras/Source/GameplayCameras/Private/Nodes/Input/CameraRigInput2DSlot.cpp
@@ -3,6 +3,7 @@
 #include "Nodes/Input/CameraRigInput2DSlot.h"
 
 #include "Core/CameraBuildLog.h"
+#include "Core/CameraOperation.h"
 #include "Core/CameraRigAsset.h"
 #include "Core/CameraRigBuildContext.h"
 #include "Core/CameraVariableAssets.h"
@@ -18,7 +19,10 @@ UE_DEFINE_CAMERA_NODE_EVALUATOR(FCameraRigInput2DSlotEvaluator)
 
 FCameraRigInput2DSlotEvaluator::FCameraRigInput2DSlotEvaluator()
 {
-	SetNodeEvaluatorFlags(ECameraNodeEvaluatorFlags::NeedsParameterUpdate | ECameraNodeEvaluatorFlags::NeedsEvaluationUpdate);
+	SetNodeEvaluatorFlags(
+			ECameraNodeEvaluatorFlags::NeedsParameterUpdate | 
+			ECameraNodeEvaluatorFlags::NeedsEvaluationUpdate |
+			ECameraNodeEvaluatorFlags::SupportsOperations);
 }
 
 void FCameraRigInput2DSlotEvaluator::OnBuild(const FCameraNodeEvaluatorBuildParams& Params)
@@ -33,9 +37,9 @@ void FCameraRigInput2DSlotEvaluator::OnInitialize(const FCameraNodeEvaluatorInit
 	InputValue = FVector2d::ZeroVector;
 
 	const UCameraRigInput2DSlot* SlotNode = GetCameraNodeAs<UCameraRigInput2DSlot>();
-	if (SlotNode->GetVariableID().IsValid() && Params.LastActiveCameraRig.IsValid())
+	if (SlotNode->GetVariableID().IsValid() && Params.LastActiveCameraRigInfo.LastResult)
 	{
-		const FCameraVariableTable& LastActiveRigVariableTable = Params.LastActiveCameraRig.LastResult->VariableTable;
+		const FCameraVariableTable& LastActiveRigVariableTable = Params.LastActiveCameraRigInfo.LastResult->VariableTable;
 		LastActiveRigVariableTable.TryGetValue<FVector2d>(SlotNode->GetVariableID(), InputValue);
 	}
 }
@@ -88,6 +92,29 @@ void FCameraRigInput2DSlotEvaluator::OnRun(const FCameraNodeEvaluationParams& Pa
 	OutResult.VariableTable.SetValue<FVector2d>(SlotNode->GetVariableID(), InputValue);
 }
 
+void FCameraRigInput2DSlotEvaluator::OnExecuteOperation(const FCameraOperationParams& Params, FCameraOperation& Operation)
+{
+	if (FYawPitchCameraOperation* Op = Operation.CastOperation<FYawPitchCameraOperation>())
+	{
+		const UCameraRigInput2DSlot* SlotNode = GetCameraNodeAs<UCameraRigInput2DSlot>();
+
+		double MinValueX, MaxValueX;
+		SlotNode->ClampX.GetEffectiveClamping(MinValueX, MaxValueX);
+
+		double MinValueY, MaxValueY;
+		SlotNode->ClampY.GetEffectiveClamping(MinValueY, MaxValueY);
+
+		InputValue.X = Op->Yaw.Apply(InputValue.X, MinValueX, MaxValueX);
+		InputValue.Y = Op->Pitch.Apply(InputValue.Y, MinValueY, MaxValueY);
+	}
+}
+
+void FCameraRigInput2DSlotEvaluator::OnSerialize(const FCameraNodeEvaluatorSerializeParams& Params, FArchive& Ar)
+{
+	Ar << TransientInputValue;
+	Ar << InputValue;
+}
+
 }  // namespace UE::Cameras
 
 FCameraNodeChildrenView UCameraRigInput2DSlot::OnGetChildren()
diff --git a/Engine/Plugins/Cameras/GameplayCameras/Source/GameplayCameras/Private/Nodes/Input/CameraRigInputSlotTypes.cpp b/Engine/Plugins/Cameras/GameplayCameras/Source/GameplayCameras/Private/Nodes/Input/CameraRigInputSlotTypes.cpp
--- a/Engine/Plugins/Cameras/GameplayCameras/Source/GameplayCameras/Private/Nodes/Input/CameraRigInputSlotTypes.cpp
+++ b/Engine/Plugins/Cameras/GameplayCameras/Source/GameplayCameras/Private/Nodes/Input/CameraRigInputSlotTypes.cpp
@@ -2,6 +2,8 @@
 
 #include "Nodes/Input/CameraRigInputSlotTypes.h"
 
+#include "Math/NumericLimits.h"
+
 double FCameraParameterClamping::ClampValue(double Value) const
 {
 	if (bClampMin && Value < MinValue)
@@ -15,6 +17,12 @@ double FCameraParameterClamping::ClampValue(double Value) const
 	return Value;
 }
 
+void FCameraParameterClamping::GetEffectiveClamping(double& OutMinValue, double& OutMaxValue) const
+{
+	OutMinValue = bClampMin ? MinValue : TNumericLimits<double>::Lowest();
+	OutMaxValue = bClampMax ? MaxValue : TNumericLimits<double>::Max();
+}
+
 double FCameraParameterNormalization::NormalizeValue(double Value) const
 {
 	if (bNormalize && MaxValue > 0)
diff --git a/Engine/Plugins/Cameras/GameplayCameras/Source/GameplayCameras/Private/Services/OrientationInitializationService.cpp b/Engine/Plugins/Cameras/GameplayCameras/Source/GameplayCameras/Private/Services/OrientationInitializationService.cpp
new file mode 100644
--- /dev/null
+++ b/Engine/Plugins/Cameras/GameplayCameras/Source/GameplayCameras/Private/Services/OrientationInitializationService.cpp
@@ -0,0 +1,197 @@
+// Copyright Epic Games, Inc. All Rights Reserved.
+
+#include "Services/OrientationInitializationService.h"
+
+#include "Core/CameraEvaluationContext.h"
+#include "Core/CameraEvaluationContextStack.h"
+#include "Core/CameraIKAim.h"
+#include "Core/CameraOperation.h"
+#include "Core/CameraRigAsset.h"
+#include "Core/CameraSystemEvaluator.h"
+#include "Core/RootCameraNode.h"
+#include "Core/RootCameraNodeObserver.h"
+
+namespace UE::Cameras
+{
+
+void FOrientationInitializationService::OnInitialize(const FCameraEvaluationServiceInitializeParams& Params)
+{
+	SetEvaluationServiceFlags(
+			ECameraEvaluationServiceFlags::NeedsRootCameraNodeEvents |
+			ECameraEvaluationServiceFlags::NeedsPostUpdate);
+
+	Evaluator = Params.Evaluator;
+}
+
+void FOrientationInitializationService::OnPostUpdate(const FCameraEvaluationServiceUpdateParams& Params, FCameraEvaluationServiceUpdateResult& OutResult)
+{
+	TSharedPtr<FCameraEvaluationContext> ActiveContext = Params.Evaluator->GetEvaluationContextStack().GetActiveContext();
+
+	bHasPreviousContextTransform = false;
+	if (ActiveContext)
+	{
+		const FCameraNodeEvaluationResult& InitialResult = ActiveContext->GetInitialResult();
+		PreviousContextLocation = InitialResult.CameraPose.GetLocation();
+		PreviousContextRotation = InitialResult.CameraPose.GetRotation();
+		bHasPreviousContextTransform = true;
+	}
+
+	PreviousEvaluationContext = ActiveContext;
+}
+
+void FOrientationInitializationService::OnRootCameraNodeEvent(const FRootCameraNodeCameraRigEvent& InEvent)
+{
+	if (InEvent.EventType == ERootCameraNodeCameraRigEventType::Activated && InEvent.EventLayer == ECameraRigLayer::Main)
+	{
+		ECameraRigInitialOrientation InitialOrientation = ECameraRigInitialOrientation::None;
+
+		if (InEvent.CameraRigInfo.CameraRig)
+		{
+			InitialOrientation = InEvent.CameraRigInfo.CameraRig->InitialOrientation;
+		}
+
+		if (InEvent.Transition && InEvent.Transition->bOverrideInitialOrientation)
+		{
+			InitialOrientation = InEvent.Transition->InitialOrientation;
+		}
+
+		switch (InitialOrientation)
+		{
+			case ECameraRigInitialOrientation::ContextYawPitch:
+				TryInitializeContextYawPitch(InEvent.CameraRigInfo);
+				break;
+			case ECameraRigInitialOrientation::PreviousYawPitch:
+				TryPreserveYawPitch(InEvent.CameraRigInfo);
+				break;
+			case ECameraRigInitialOrientation::PreviousAbsoluteTarget:
+				TryPreserveTarget(InEvent.CameraRigInfo, false);
+				break;
+			case ECameraRigInitialOrientation::PreviousRelativeTarget:
+				TryPreserveTarget(InEvent.CameraRigInfo, true);
+				break;
+			default:
+				break;
+		}
+	}
+}
+
+void FOrientationInitializationService::TryInitializeContextYawPitch(const FCameraRigEvaluationInfo& CameraRigInfo)
+{
+	if (!CameraRigInfo.CameraRig)
+	{
+		UE_LOG(LogCameraSystem, Error, TEXT("Can't initialize orientation on invalid camera rig."));
+		return;
+	}
+
+	if (!CameraRigInfo.EvaluationContext)
+	{
+		UE_LOG(LogCameraSystem, Error, TEXT("Can't initialize orientation on camera rig '%s' with invalid evaluation context."),
+				*CameraRigInfo.CameraRig->GetPathName());
+		return;
+	}
+
+	const FCameraNodeEvaluationResult& InitialResult = CameraRigInfo.EvaluationContext->GetInitialResult();
+	if (!InitialResult.bIsValid)
+	{
+		UE_LOG(LogCameraSystem, Error, TEXT("Can't initialize orientation on camera rig '%s' with invalid initial context result."),
+				*CameraRigInfo.CameraRig->GetPathName());
+		return;
+	}
+
+	const FRotator3d& InitialRotation = InitialResult.CameraPose.GetRotation();
+	TryInitializeYawPitch(CameraRigInfo, InitialRotation.Yaw, InitialRotation.Pitch);
+}
+
+void FOrientationInitializationService::TryPreserveYawPitch(const FCameraRigEvaluationInfo& CameraRigInfo)
+{
+	const FCameraSystemEvaluationResult& LastResult = Evaluator->GetEvaluatedResult();
+	if (!LastResult.bIsValid)
+	{
+		UE_LOG(LogCameraSystem, Error, TEXT("Can't initialize camera rig orientation when previous camera result is invalid."));
+		return;
+	}
+
+	if (!CameraRigInfo.CameraRig)
+	{
+		UE_LOG(LogCameraSystem, Error, TEXT("Can't initialize camera rig orientation with invalid camera rig."));
+		return;
+	}
+	
+	const FRotator3d LastOrientation = LastResult.CameraPose.GetRotation();
+	TryInitializeYawPitch(CameraRigInfo, LastOrientation.Yaw, LastOrientation.Pitch);
+}
+
+void FOrientationInitializationService::TryInitializeYawPitch(const FCameraRigEvaluationInfo& CameraRigInfo, TOptional<double> Yaw, TOptional<double> Pitch)
+{
+	if (!CameraRigInfo.RootEvaluator)
+	{
+		UE_LOG(LogCameraSystem, Error, TEXT("Can't initialize orientation on camera rig '%s' because it has no evaluator."),
+				*CameraRigInfo.CameraRig->GetPathName());
+		return;
+	}
+
+	FCameraOperationParams OperationParams;
+	OperationParams.Evaluator = Evaluator;
+	OperationParams.EvaluationContext = CameraRigInfo.EvaluationContext;
+
+	FYawPitchCameraOperation Operation;
+	if (Yaw.IsSet())
+	{
+		Operation.Yaw = FConsumableDouble::Absolute(Yaw.GetValue());
+	}
+	if (Pitch.IsSet())
+	{
+		Operation.Pitch = FConsumableDouble::Absolute(Pitch.GetValue());
+	}
+	CameraRigInfo.RootEvaluator->ExecuteOperation(OperationParams, Operation);
+}
+
+void FOrientationInitializationService::TryPreserveTarget(const FCameraRigEvaluationInfo& CameraRigInfo, bool bUseRelativeTarget)
+{
+	const FCameraSystemEvaluationResult& LastResult = Evaluator->GetEvaluatedResult();
+	if (!LastResult.bIsValid)
+	{
+		// The previous result might be invalid on the very first frame of the game, when the first
+		// camera rig activates.
+		return;
+	}
+
+	if (!CameraRigInfo.CameraRig)
+	{
+		UE_LOG(LogCameraSystem, Error, TEXT("Can't initialize camera rig orientation with invalid camera rig."));
+		return;
+	}
+	
+	if (!CameraRigInfo.RootEvaluator)
+	{
+		UE_LOG(LogCameraSystem, Error, TEXT("Can't initialize orientation on camera rig '%s' because it has no evaluator."),
+				*CameraRigInfo.CameraRig->GetPathName());
+		return;
+	}
+
+	FVector3d TargetToPreserve = LastResult.CameraPose.GetTarget();
+
+	if (bUseRelativeTarget && bHasPreviousContextTransform &&
+			(PreviousEvaluationContext == nullptr || CameraRigInfo.EvaluationContext == PreviousEvaluationContext))
+	{
+		const FRotator3d PreviousInverseContextRotation = PreviousContextRotation.GetInverse();
+		const FVector3d LastRelativeTarget = PreviousInverseContextRotation.RotateVector(TargetToPreserve - PreviousContextLocation);
+
+		const FCameraNodeEvaluationResult& InitialResult = CameraRigInfo.EvaluationContext->GetInitialResult();
+		const FVector3d& InitialLocation = InitialResult.CameraPose.GetLocation();
+		const FRotator3d& InitialRotation = InitialResult.CameraPose.GetRotation();
+		TargetToPreserve = InitialRotation.RotateVector(LastRelativeTarget) + InitialLocation;
+	}
+
+	FCameraIKAimParams AimParams;
+	AimParams.bIsFirstFrame = true;
+	AimParams.DeltaTime = 0.f;
+	AimParams.Evaluator = Evaluator;
+	AimParams.TargetLocation = TargetToPreserve;
+
+	FCameraIKAim CameraAim;
+	CameraAim.Run(AimParams, CameraRigInfo);
+}
+
+}  // namespace UE::Cameras
+
diff --git a/Engine/Plugins/Cameras/GameplayCameras/Source/GameplayCameras/Private/Services/OrientationInitializationService.h b/Engine/Plugins/Cameras/GameplayCameras/Source/GameplayCameras/Private/Services/OrientationInitializationService.h
new file mode 100644
--- /dev/null
+++ b/Engine/Plugins/Cameras/GameplayCameras/Source/GameplayCameras/Private/Services/OrientationInitializationService.h
@@ -0,0 +1,41 @@
+// Copyright Epic Games, Inc. All Rights Reserved.
+
+#pragma once
+
+#include "Core/CameraEvaluationService.h"
+#include "Core/CameraRigTransition.h"
+
+namespace UE::Cameras
+{
+
+class FCameraEvaluationContext;
+struct FCameraRigEvaluationInfo;
+
+class FOrientationInitializationService : public FCameraEvaluationService
+{
+protected:
+
+	// FCameraEvaluationService interface.
+	virtual void OnInitialize(const FCameraEvaluationServiceInitializeParams& Params) override;
+	virtual void OnPostUpdate(const FCameraEvaluationServiceUpdateParams& Params, FCameraEvaluationServiceUpdateResult& OutResult) override;
+	virtual void OnRootCameraNodeEvent(const FRootCameraNodeCameraRigEvent& InEvent) override;
+
+private:
+
+	void TryInitializeContextYawPitch(const FCameraRigEvaluationInfo& CameraRigInfo);
+	void TryPreserveYawPitch(const FCameraRigEvaluationInfo& CameraRigInfo);
+	void TryInitializeYawPitch(const FCameraRigEvaluationInfo& CameraRigInfo, TOptional<double> Yaw, TOptional<double> Pitch);
+	void TryPreserveTarget(const FCameraRigEvaluationInfo& CameraRigInfo, bool bUseRelativeTarget);
+
+private:
+
+	FCameraSystemEvaluator* Evaluator = nullptr;
+
+	TWeakPtr<FCameraEvaluationContext> PreviousEvaluationContext;
+	FVector3d PreviousContextLocation;
+	FRotator3d PreviousContextRotation;
+	bool bHasPreviousContextTransform = false;
+};
+
+}  // namespace UE::Cameras
+
diff --git a/Engine/Plugins/Cameras/GameplayCameras/Source/GameplayCameras/Public/Core/BlendStackCameraNode.h b/Engine/Plugins/Cameras/GameplayCameras/Source/GameplayCameras/Public/Core/BlendStackCameraNode.h
--- a/Engine/Plugins/Cameras/GameplayCameras/Source/GameplayCameras/Public/Core/BlendStackCameraNode.h
+++ b/Engine/Plugins/Cameras/GameplayCameras/Source/GameplayCameras/Public/Core/BlendStackCameraNode.h
@@ -118,6 +118,7 @@ protected:
 	virtual void OnInitialize(const FCameraNodeEvaluatorInitializeParams& Params) override;
 	virtual void OnRun(const FCameraNodeEvaluationParams& Params, FCameraNodeEvaluationResult& OutResult) override;
 	virtual void OnAddReferencedObjects(FReferenceCollector& Collector) override;
+	virtual void OnSerialize(const FCameraNodeEvaluatorSerializeParams& Params, FArchive& Ar) override;
 
 #if UE_GAMEPLAY_CAMERAS_DEBUG
 	virtual void OnBuildDebugBlocks(const FCameraDebugBlockBuildParams& Params, FCameraDebugBlockBuilder& Builder) override;
diff --git a/Engine/Plugins/Cameras/GameplayCameras/Source/GameplayCameras/Public/Core/CameraEvaluationService.h b/Engine/Plugins/Cameras/GameplayCameras/Source/GameplayCameras/Public/Core/CameraEvaluationService.h
--- a/Engine/Plugins/Cameras/GameplayCameras/Source/GameplayCameras/Public/Core/CameraEvaluationService.h
+++ b/Engine/Plugins/Cameras/GameplayCameras/Source/GameplayCameras/Public/Core/CameraEvaluationService.h
@@ -86,7 +86,7 @@ public:
 	void NotifyRootCameraNodeEvent(const FRootCameraNodeCameraRigEvent& InEvent);
 
 	ECameraEvaluationServiceFlags GetEvaluationServiceFlags() const { return PrivateFlags; }
-	bool HasAnyEvaluationServiceFlags(ECameraEvaluationServiceFlags InFlags) const;
+	bool HasAllEvaluationServiceFlags(ECameraEvaluationServiceFlags InFlags) const;
 
 protected:
 
diff --git a/Engine/Plugins/Cameras/GameplayCameras/Source/GameplayCameras/Public/Core/CameraIKAim.h b/Engine/Plugins/Cameras/GameplayCameras/Source/GameplayCameras/Public/Core/CameraIKAim.h
new file mode 100644
--- /dev/null
+++ b/Engine/Plugins/Cameras/GameplayCameras/Source/GameplayCameras/Public/Core/CameraIKAim.h
@@ -0,0 +1,89 @@
+// Copyright Epic Games, Inc. All Rights Reserved.
+
+#pragma once
+
+#include "Core/CameraNodeEvaluator.h"
+#include "Core/CameraNodeEvaluatorStorage.h"
+#include "Math/MathFwd.h"
+
+namespace UE::Cameras
+{
+
+struct FCameraRigEvaluationInfo;
+
+/**
+ * Parameter structure for aiming a camera rig at a target.
+ */
+struct FCameraIKAimParams
+{
+	/** The time interval to use when updating the camera rig. */
+	float DeltaTime = 0.f;
+	/** Whether this is the first update of the camera rig. */
+	bool bIsFirstFrame = false;
+	
+	/** The desired target that the camera rig should be aiming at. */
+	FVector3d TargetLocation = FVector3d::ZeroVector;
+
+	/** The camera system inside which the evaluation takes place. */
+	FCameraSystemEvaluator* Evaluator = nullptr;
+
+	/** The distance below which aiming should not take place. */
+	double MinDistance = -1.0;
+	/** The angle between desired and actual target below which we consider aiming is complete. */
+	double AngleTolerance = -1.0;
+	/** The distance between desired target and line of sight below which we consider aiming is complete. */
+	double DistanceTolerance = -1.0;
+	/** The maximum number of iterations to run. */
+	uint8 MaxIterations = 0;
+};
+
+/**
+ * A class that can manipulate a camera rig in order to aim it at a desired target.
+ */
+class FCameraIKAim
+{
+public:
+
+	/** Executes the aiming. */
+	bool Run(const FCameraIKAimParams& Params, const FCameraRigEvaluationInfo& CameraRigInfo);
+
+private:
+
+	enum class EAimResult
+	{
+		Failed,
+		Aborted,
+		Corrected,
+		Completed
+	};
+
+	struct FAimIterationInfo
+	{
+		uint8 IterationIndex = 0;
+		double ErrorAngle = 0.0;
+		double ErrorDistance = 0.0;
+		EAimResult Result = EAimResult::Failed;
+	};
+
+	bool DoRun(const FCameraIKAimParams& Params, const FCameraRigEvaluationInfo& CameraRigInfo);
+	void DoRunIteration(const FCameraIKAimParams& Params, const FCameraRigEvaluationInfo& CameraRigInfo, FAimIterationInfo& IterationInfo);
+
+	void AimTwoBonesCameraRig(const FCameraIKAimParams& Params, const FCameraRigEvaluationInfo& CameraRigInfo, const FTransform3d& PivotTransform, FAimIterationInfo& IterationInfo);
+
+	void RunRootCameraNode(const FCameraIKAimParams& Params, const FCameraRigEvaluationInfo& CameraRigInfo);
+	bool CheckTolerance(const FCameraIKAimParams& Params, const FCameraRigEvaluationInfo& CameraRigInfo, FAimIterationInfo& IterationInfo);
+
+	static bool ComputeTwoBonesCorrection(const FCameraPose& CurrentPose, const FVector3d& PivotLocation, const FVector3d& DesiredTarget, FRotator3d& OutCorrection);
+
+	static bool RaySphereIntersectExit(const FRay3d& Ray, const FVector3d& SphereOrigin, double SphereRadius, double& OutRayIntersectDistance);
+	static bool RaySphereIntersectExit(const FVector3d& RayStart, const FVector3d& RayDir, const FVector3d& SphereOrigin, double SphereRadius, double& OutRayIntersectDistance);
+
+private:
+
+	FCameraNodeEvaluationResult ScratchResult;
+	
+	TArray<uint8> EvaluatorSnapshot;
+};
+
+}  // namespace UE::Cameras
+
diff --git a/Engine/Plugins/Cameras/GameplayCameras/Source/GameplayCameras/Public/Core/CameraNodeEvaluator.h b/Engine/Plugins/Cameras/GameplayCameras/Source/GameplayCameras/Public/Core/CameraNodeEvaluator.h
--- a/Engine/Plugins/Cameras/GameplayCameras/Source/GameplayCameras/Public/Core/CameraNodeEvaluator.h
+++ b/Engine/Plugins/Cameras/GameplayCameras/Source/GameplayCameras/Public/Core/CameraNodeEvaluator.h
@@ -5,6 +5,7 @@
 #include "Core/CameraRigEvaluationInfo.h"
 #include "Core/CameraObjectRtti.h"
 #include "Core/CameraPose.h"
+#include "Core/CameraRigJoints.h"
 #include "Core/CameraVariableTable.h"
 #include "Core/ObjectChildrenView.h"
 #include "CoreTypes.h"
@@ -22,6 +23,7 @@ namespace UE::Cameras
 class FCameraEvaluationContext;
 class FCameraNodeEvaluator;
 class FCameraSystemEvaluator;
+struct FCameraOperation;
 struct FCameraNodeEvaluationParams;
 struct FCameraNodeEvaluatorBuilder;
 struct FCameraRigEvaluationInfo;
@@ -37,10 +39,53 @@ enum class ECameraNodeEvaluatorFlags
 {
 	None = 0,
 	NeedsParameterUpdate = 1 << 0,
-	NeedsEvaluationUpdate = 1 << 1
+	NeedsEvaluationUpdate = 1 << 1,
+	SupportsOperations = 1 << 2
 };
 ENUM_CLASS_FLAGS(ECameraNodeEvaluatorFlags)
 
+/** View on a camera node evaluator's children. */
+using FCameraNodeEvaluatorChildrenView = TObjectChildrenView<FCameraNodeEvaluator*>;
+
+/**
+ * Structure for building the tree of camera node evaluators.
+ */
+struct FCameraNodeEvaluatorBuildParams
+{
+	FCameraNodeEvaluatorBuildParams(FCameraNodeEvaluatorBuilder& InBuilder)
+		: Builder(InBuilder)
+	{}
+
+	/** Builds an evaluator for the given camera node. */
+	FCameraNodeEvaluator* BuildEvaluator(const UCameraNode* InNode) const;
+
+	/** Builds an evaluator for the given camera node, and down-cast it to the given type. */
+	template<typename EvaluatorType>
+	EvaluatorType* BuildEvaluatorAs(const UCameraNode* InNode) const;
+
+private:
+
+	/** Builder object for building children evaluators. */
+	FCameraNodeEvaluatorBuilder& Builder;
+};
+
+/**
+ * Structure for initializing a camera node evaluator.
+ */
+struct FCameraNodeEvaluatorInitializeParams
+{
+	/** The evaluation running this evaluation. */
+	FCameraSystemEvaluator* Evaluator = nullptr;
+	/** The evaluation context (if any) responsible for this branch of the evaluation. */
+	TSharedPtr<const FCameraEvaluationContext> EvaluationContext;
+
+	/**
+	 * Information about the last active camera rig if the node tree being initialized
+	 * is being pushed on top of a non-empty blend stack.
+	 */
+	FCameraRigEvaluationInfo LastActiveCameraRigInfo;
+};
+
 /**
  * Parameter structure for updating the pre-blended parameters of a camera node.
  */
@@ -70,6 +115,15 @@ struct FCameraBlendedParameterUpdateResult
 	FCameraVariableTable& VariableTable;
 };
 
+/** The type of evaluation being run. */
+enum class ECameraNodeEvaluationType
+{
+	/** Normal evaluation. */
+	Standard,
+	/** Evaluation for IK aiming. */
+	IK
+};
+
 /**
  * Parameter structure for running a camera node evaluator.
  */
@@ -81,7 +135,8 @@ struct FCameraNodeEvaluationParams
 	TSharedPtr<const FCameraEvaluationContext> EvaluationContext;
 	/** The time interval for the evaluation. */
 	float DeltaTime = 0.f;
-	
+	/** The type of evaluation being run. */
+	ECameraNodeEvaluationType EvaluationType = ECameraNodeEvaluationType::Standard;
 	/** Whether this is the first evaluation of this camera node hierarchy. */
 	bool bIsFirstFrame = false;
 };
@@ -97,56 +152,40 @@ struct GAMEPLAYCAMERAS_API FCameraNodeEvaluationResult
 	/** The variable table. */
 	FCameraVariableTable VariableTable;
 
+	/** The list of joints in the current camera rig. */
+	FCameraRigJoints CameraRigJoints;
+
 	/** Whether the current frame is a camera cut. */
 	bool bIsCameraCut = false;
 	/** Whether this result is valid. */
 	bool bIsValid = false;
 
-	/** Reset this result to its default (non-valid) state. */
-	void Reset();
-};
-
-/**
- * Structure for building the tree of camera node evaluators.
- */
-struct FCameraNodeEvaluatorBuildParams
-{
-	FCameraNodeEvaluatorBuildParams(FCameraNodeEvaluatorBuilder& InBuilder)
-		: Builder(InBuilder)
-	{}
-
-	/** Builds an evaluator for the given camera node. */
-	FCameraNodeEvaluator* BuildEvaluator(const UCameraNode* InNode) const;
-
-	/** Builds an evaluator for the given camera node, and down-cast it to the given type. */
-	template<typename EvaluatorType>
-	EvaluatorType* BuildEvaluatorAs(const UCameraNode* InNode) const;
+public:
 
-private:
+	/** Reset this result to its default (non-valid) state.  */
+	void Reset(bool bResetVariableTable);
 
-	/** Builder object for building children evaluators. */
-	FCameraNodeEvaluatorBuilder& Builder;
+	/** Serializes this result to the given archive. */
+	void Serialize(FArchive& Ar);
 };
 
 /**
- * Structure for initializing a camera node evaluator.
+ * Parameter structure for executing camera operations.
  */
-struct FCameraNodeEvaluatorInitializeParams
+struct FCameraOperationParams
 {
-	/** The evaluation running this evaluation. */
+	/** The evaluator running this operation. */
 	FCameraSystemEvaluator* Evaluator = nullptr;
 	/** The evaluation context (if any) responsible for this branch of the evaluation. */
 	TSharedPtr<const FCameraEvaluationContext> EvaluationContext;
-
-	/**
-	 * Information about the last active camera rig if the node tree being initialized
-	 * is being pushed on top of a non-empty blend stack.
-	 */
-	FCameraRigEvaluationInfo LastActiveCameraRig;
 };
 
-/** View on a camera node evaluator's children. */
-using FCameraNodeEvaluatorChildrenView = TObjectChildrenView<FCameraNodeEvaluator*>;
+/**
+ * Parameter structure for serializing the state of a camera node evaluator.
+ */
+struct GAMEPLAYCAMERAS_API FCameraNodeEvaluatorSerializeParams
+{
+};
 
 #if UE_GAMEPLAY_CAMERAS_DEBUG
 
@@ -187,9 +226,15 @@ public:
 	/** Run this evaluator. */
 	void Run(const FCameraNodeEvaluationParams& Params, FCameraNodeEvaluationResult& OutResult);
 
+	/** Execute an IK operation. */
+	void ExecuteOperation(const FCameraOperationParams& Params, FCameraOperation& Operation);
+
 	/** Collect referenced UObjects. */
 	void AddReferencedObjects(FReferenceCollector& Collector);
 
+	/** Serializes the state of this evaluator. */
+	void Serialize(const FCameraNodeEvaluatorSerializeParams& Params, FArchive& Ar);
+
 	/** Gets the flags for this evaluator. */
 	ECameraNodeEvaluatorFlags  GetNodeEvaluatorFlags() const { return PrivateFlags; }
 
@@ -235,9 +280,15 @@ protected:
 	/** Run this evaluator. */
 	GAMEPLAYCAMERAS_API virtual void OnRun(const FCameraNodeEvaluationParams& Params, FCameraNodeEvaluationResult& OutResult) {}
 
+	/** Execute an IK operation. */
+	GAMEPLAYCAMERAS_API virtual void OnExecuteOperation(const FCameraOperationParams& Params, FCameraOperation& Operation) {}
+
 	/** Collect referenced UObjects. */
 	GAMEPLAYCAMERAS_API virtual void OnAddReferencedObjects(FReferenceCollector& Collector) {}
 
+	/** Serializes the state of this evaluator. */
+	GAMEPLAYCAMERAS_API virtual void OnSerialize(const FCameraNodeEvaluatorSerializeParams& Params, FArchive& Ar) {}
+
 #if UE_GAMEPLAY_CAMERAS_DEBUG
 	/** Called to create debug blocks for this node evaluator. */
 	GAMEPLAYCAMERAS_API virtual void OnBuildDebugBlocks(const FCameraDebugBlockBuildParams& Params, FCameraDebugBlockBuilder& Builder);
diff --git a/Engine/Plugins/Cameras/GameplayCameras/Source/GameplayCameras/Public/Core/CameraOperation.h b/Engine/Plugins/Cameras/GameplayCameras/Source/GameplayCameras/Public/Core/CameraOperation.h
new file mode 100644
--- /dev/null
+++ b/Engine/Plugins/Cameras/GameplayCameras/Source/GameplayCameras/Public/Core/CameraOperation.h
@@ -0,0 +1,109 @@
+// Copyright Epic Games, Inc. All Rights Reserved.
+
+#pragma once
+
+#include "CoreTypes.h"
+#include "Math/ConsumableValue.h"
+
+namespace UE::Cameras
+{
+
+/**
+ * Hard-coded type IDs for built-in camera operations.
+ */
+enum EBuiltInCameraOperationTypes : uint32
+{
+	YawPitch = 0,
+	SingleValue = 1,
+
+	MAX = 3
+};
+
+/**
+ * Simple type ID for camera operations.
+ */
+struct FCameraOperationTypeID
+{
+	FCameraOperationTypeID() : Value(INVALID) {}
+
+	FCameraOperationTypeID(uint32 InValue) : Value(InValue) {}
+
+public:
+
+	/** Returns whether this ID is valid. */
+	bool IsValid() const { return Value != INVALID; }
+
+	friend bool operator==(FCameraOperationTypeID A, FCameraOperationTypeID B)
+	{
+		return A.Value == B.Value;
+	}
+
+private:
+
+	static const uint32 INVALID = (uint32)-1;
+
+	uint32 Value;
+};
+
+/**
+ * Base class for an operation to be executed on a camera rig.
+ */
+struct FCameraOperation
+{
+public:
+
+	/** Attempts to cast this operation into a sub-class. */
+	template<typename OperationType>
+	OperationType* CastOperation()
+	{
+		if (OperationTypeID == OperationType::GetOperationTypeID())
+		{
+			return static_cast<OperationType*>(this);
+		}
+		return nullptr;
+	}
+
+protected:
+
+	FCameraOperation(FCameraOperationTypeID InOperationTypeID)
+		: OperationTypeID(InOperationTypeID)
+	{}
+
+private:
+
+	FCameraOperationTypeID OperationTypeID;
+};
+
+#define UE_DEFINE_CAMERA_IK_OPERATION(OperationName, OperationType)\
+	static FCameraOperationTypeID GetOperationTypeID()\
+	{\
+		return FCameraOperationTypeID(OperationType);\
+	}\
+	OperationName()\
+		: FCameraOperation(OperationName::GetOperationTypeID())\
+	{}
+
+/**
+ * A camera operation that tries to correct the yaw/pitch of a camera rig.
+ */
+struct FYawPitchCameraOperation : public FCameraOperation
+{
+	UE_DEFINE_CAMERA_IK_OPERATION(FYawPitchCameraOperation, EBuiltInCameraOperationTypes::YawPitch)
+
+	FConsumableDouble Yaw;
+	FConsumableDouble Pitch;
+};
+
+/**
+ * A camera operation that tries to correct a single undetermined input value on a
+ * camera rig.
+ */
+struct FSingleValueCameraOperation : public FCameraOperation
+{
+	UE_DEFINE_CAMERA_IK_OPERATION(FSingleValueCameraOperation, EBuiltInCameraOperationTypes::SingleValue)
+
+	FConsumableDouble Value;
+};
+
+}  // namespace UE::Cameras
+
diff --git a/Engine/Plugins/Cameras/GameplayCameras/Source/GameplayCameras/Public/Core/CameraPose.h b/Engine/Plugins/Cameras/GameplayCameras/Source/GameplayCameras/Public/Core/CameraPose.h
--- a/Engine/Plugins/Cameras/GameplayCameras/Source/GameplayCameras/Public/Core/CameraPose.h
+++ b/Engine/Plugins/Cameras/GameplayCameras/Source/GameplayCameras/Public/Core/CameraPose.h
@@ -158,6 +158,9 @@ public:
 	/** Gets the aiming direction of the camera. */
 	FVector3d GetAimDir() const;
 
+	/** Gets the location of the camera's target. */
+	FVector3d GetTarget() const;
+
 public:
 
 	// Interpolation
@@ -198,7 +201,7 @@ private:
 
 	/** Distance to the target */
 	UPROPERTY()
-	double TargetDistance = 100.0;
+	double TargetDistance = 1000.0;
 
 	/**
 	 * The horizontal field of view of the camera, in degrees 
diff --git a/Engine/Plugins/Cameras/GameplayCameras/Source/GameplayCameras/Public/Core/CameraRigAsset.h b/Engine/Plugins/Cameras/GameplayCameras/Source/GameplayCameras/Public/Core/CameraRigAsset.h
--- a/Engine/Plugins/Cameras/GameplayCameras/Source/GameplayCameras/Public/Core/CameraRigAsset.h
+++ b/Engine/Plugins/Cameras/GameplayCameras/Source/GameplayCameras/Public/Core/CameraRigAsset.h
@@ -167,6 +167,10 @@ public:
 	/** List of exist transitions for this camera rig. */
 	UPROPERTY(Instanced, meta=(ObjectTreeGraphPinDirection=Output))
 	TArray<TObjectPtr<UCameraRigTransition>> ExitTransitions;
+	
+	/** Default orientation initialization when this camera rig is activated. */
+	UPROPERTY(EditAnywhere, Category="Transition")
+	ECameraRigInitialOrientation InitialOrientation = ECameraRigInitialOrientation::None;
 
 	/** Allocation information for all the nodes and variables in this camera rig. */
 	UPROPERTY()
diff --git a/Engine/Plugins/Cameras/GameplayCameras/Source/GameplayCameras/Public/Core/CameraRigEvaluationInfo.h b/Engine/Plugins/Cameras/GameplayCameras/Source/GameplayCameras/Public/Core/CameraRigEvaluationInfo.h
--- a/Engine/Plugins/Cameras/GameplayCameras/Source/GameplayCameras/Public/Core/CameraRigEvaluationInfo.h
+++ b/Engine/Plugins/Cameras/GameplayCameras/Source/GameplayCameras/Public/Core/CameraRigEvaluationInfo.h
@@ -3,14 +3,17 @@
 #pragma once
 
 #include "Templates/SharedPointer.h"
+#include "UObject/ObjectPtr.h"
 
 class UCameraRigAsset;
 
 namespace UE::Cameras
 {
 
+class FBlendCameraNodeEvaluator;
 class FCameraEvaluationContext;
 class FCameraNodeEvaluator;
+class FCameraNodeEvaluatorStorage;
 struct FCameraNodeEvaluationResult;
 
 /**
@@ -19,41 +22,28 @@ struct FCameraNodeEvaluationResult;
  */
 struct FCameraRigEvaluationInfo
 {
-	/** The camera rig being evaluated. */
-	const UCameraRigAsset* CameraRig = nullptr;
 	/** The context inside which the evaluation occurs. */
 	TSharedPtr<const FCameraEvaluationContext> EvaluationContext;
+	/** The camera rig being evaluated. */
+	TObjectPtr<const UCameraRigAsset> CameraRig;
 	/** The last evaluated result for this camera rig. */
 	const FCameraNodeEvaluationResult* LastResult = nullptr;
-	/** The root node evaluator of the camera rig. */
-	const FCameraNodeEvaluator* RootEvaluator;
-	/** Whether the camera rig is frozen. */
-	bool bIsFrozen = false;
+	/** The blend node evaluator of the camera rig. */
+	FCameraNodeEvaluator* RootEvaluator = nullptr;
 
 	FCameraRigEvaluationInfo()
 	{}
 
 	FCameraRigEvaluationInfo(
-			const UCameraRigAsset* InCameraRig,
 			TSharedPtr<const FCameraEvaluationContext> InEvaluationContext,
-			const FCameraNodeEvaluationResult& InLastResult,
-			const FCameraNodeEvaluator* InRootEvaluator)
-		: CameraRig(InCameraRig)
-		, EvaluationContext(InEvaluationContext)
-		, LastResult(&InLastResult)
+			TObjectPtr<const UCameraRigAsset> InCameraRig,
+			const FCameraNodeEvaluationResult* InLastResult,
+			FCameraNodeEvaluator* InRootEvaluator)
+		: EvaluationContext(InEvaluationContext)
+		, CameraRig(InCameraRig)
+		, LastResult(InLastResult)
 		, RootEvaluator(InRootEvaluator)
 	{}
-
-	/** 
-	 * Whether this evaluation info is valid.
-	 *
-	 * A valid evaluation info may not have a valid CameraRig, EvaluationContext, or RootEvaluator
-	 * if the related camera rig was frozen. However, it would always have a valid LastResult.
-	 */
-	bool IsValid() const
-	{
-		return LastResult != nullptr;
-	}
 };
 
 }  // namespace UE::Cameras
diff --git a/Engine/Plugins/Cameras/GameplayCameras/Source/GameplayCameras/Public/Core/CameraRigJoints.h b/Engine/Plugins/Cameras/GameplayCameras/Source/GameplayCameras/Public/Core/CameraRigJoints.h
new file mode 100644
--- /dev/null
+++ b/Engine/Plugins/Cameras/GameplayCameras/Source/GameplayCameras/Public/Core/CameraRigJoints.h
@@ -0,0 +1,51 @@
+// Copyright Epic Games, Inc. All Rights Reserved.
+
+#pragma once
+
+#include "Core/CameraVariableTableFwd.h"
+#include "CoreTypes.h"
+
+namespace UE::Cameras
+{
+
+/**
+ * A structure describing a joint in a camera rig.
+ */
+struct FCameraRigJoint
+{
+	/** The variable driving the rotation of this joint. */
+	FCameraVariableID VariableID;
+	/** The position of the this joint. */
+	FTransform3d Transform;
+};
+
+/**
+ * A structure describing the joints of a camera rig.
+ * These joints allow for "manipulating" the rig, e.g. to make it point
+ * towards a desired target or direction.
+ */
+class FCameraRigJoints
+{
+public:
+
+	/** Add a joint. */
+	void AddJoint(const FCameraRigJoint& InJoint);
+	/** Add a joint. */
+	void AddJoint(const FCameraVariableDefinition& InVariableDefinition, const FTransform3d& InTransform);
+	/** Add a joint related to the yaw/pitch built-in variable. */
+	void AddYawPitchJoint(const FTransform3d& InTransform);
+
+	/** Gets the joints. */
+	TArrayView<const FCameraRigJoint> GetJoints() const { return Joints; }
+
+	/** Removes all previously added joints. */
+	void Reset();
+
+private:
+
+	using FJointArray = TArray<FCameraRigJoint, TInlineAllocator<2>>;
+	FJointArray Joints;
+};
+
+}  // namespace UE::Cameras
+
diff --git a/Engine/Plugins/Cameras/GameplayCameras/Source/GameplayCameras/Public/Core/CameraRigTransition.h b/Engine/Plugins/Cameras/GameplayCameras/Source/GameplayCameras/Public/Core/CameraRigTransition.h
--- a/Engine/Plugins/Cameras/GameplayCameras/Source/GameplayCameras/Public/Core/CameraRigTransition.h
+++ b/Engine/Plugins/Cameras/GameplayCameras/Source/GameplayCameras/Public/Core/CameraRigTransition.h
@@ -86,6 +86,31 @@ private:
 #endif  // WITH_EDITORONLY_DATA
 };
 
+/**
+ * Determines how a camera rig's initial orientation should be initialized.
+ */
+UENUM()
+enum class ECameraRigInitialOrientation
+{
+	/** Leave the camera rig to its default orientation. */
+	None,
+	/** Orient the camera rig in the same direction as its context's initial transform. */
+	ContextYawPitch,
+	/** Orient the camera rig in the same direction as the previously active camera rig. */
+	PreviousYawPitch,
+	/** 
+	 * Make the camera rig point at the same target as the previously active camera rig's 
+	 * last frame target.
+	 */
+	PreviousAbsoluteTarget,
+	/** 
+	 * Make the camera rig point at the same target as the previously active camera rig. 
+	 * Last frame's target will be moved and turned by an offset equal to how much the 
+	 * active evaluation context has moved and turned since last frame.
+	 */
+	PreviousRelativeTarget
+};
+
 /**
  * A camera transition.
  */
@@ -106,6 +131,14 @@ public:
 	UPROPERTY(Instanced)
 	TObjectPtr<UBlendCameraNode> Blend;
 
+	/** The orientation to set on the camera rig. */
+	UPROPERTY(EditAnywhere, Category="Transition", meta=(EditCondition="bOverrideInitialOrientation"))
+	ECameraRigInitialOrientation InitialOrientation = ECameraRigInitialOrientation::None;
+
+	/** Whether to override the default orientation to set on the camera rig. */
+	UPROPERTY(EditAnywhere, Category="Transition")
+	bool bOverrideInitialOrientation = false;
+
 protected:
 
 	// UObject interface.
diff --git a/Engine/Plugins/Cameras/GameplayCameras/Source/GameplayCameras/Public/Core/CameraSystemEvaluator.h b/Engine/Plugins/Cameras/GameplayCameras/Source/GameplayCameras/Public/Core/CameraSystemEvaluator.h
--- a/Engine/Plugins/Cameras/GameplayCameras/Source/GameplayCameras/Public/Core/CameraSystemEvaluator.h
+++ b/Engine/Plugins/Cameras/GameplayCameras/Source/GameplayCameras/Public/Core/CameraSystemEvaluator.h
@@ -18,14 +18,17 @@ class FCanvas;
 class UCameraDirector;
 class UCameraRigAsset;
 class URootCameraNode;
+enum class ECameraRigLayer : uint8;
 struct FMinimalViewInfo;
 
 namespace UE::Cameras
 {
 
+class FAutoResetCameraVariableService;
 class FCameraEvaluationContext;
 class FCameraEvaluationService;
 class FRootCameraNodeEvaluator;
+enum class ECameraEvaluationServiceFlags;
 struct FRootCameraNodeCameraRigEvent;
 
 #if UE_GAMEPLAY_CAMERAS_DEBUG
@@ -48,20 +51,23 @@ struct GAMEPLAYCAMERAS_API FCameraSystemEvaluatorCreateParams
 /**
  * Parameter structure for updating the camera system.
  */
-struct GAMEPLAYCAMERAS_API FCameraSystemEvaluationUpdateParams
+struct GAMEPLAYCAMERAS_API FCameraSystemEvaluationParams
 {
-	/** Time interface for the update. */
+	/** Time interval for the update. */
 	float DeltaTime = 0.f;
 };
 
 /**
  * Result structure for updating the camera system.
  */
-struct FCameraSystemEvaluationUpdateResult
+struct FCameraSystemEvaluationResult
 {
 	/** The result camera pose. */
 	FCameraPose CameraPose;
 
+	/** The result camera variable table. */
+	FCameraVariableTable VariableTable;
+
 	/** Whether this evaluation was a camera cut. */
 	bool bIsCameraCut = false;
 
@@ -122,26 +128,28 @@ public:
 public:
 
 	/** Run an update of the camera system. */
-	GAMEPLAYCAMERAS_API void Update(const FCameraSystemEvaluationUpdateParams& Params);
+	GAMEPLAYCAMERAS_API void Update(const FCameraSystemEvaluationParams& Params);
 
 	/** Returns the root node evaluator. */
 	FRootCameraNodeEvaluator* GetRootNodeEvaluator() const { return RootEvaluator; }
 
 	/** Gets the evaluated result. */
-	const FCameraSystemEvaluationUpdateResult& GetEvaluatedResult() const { return Result; }
+	const FCameraSystemEvaluationResult& GetEvaluatedResult() const { return Result; }
 
 	/** Get the last evaluated camera. */
 	GAMEPLAYCAMERAS_API void GetEvaluatedCameraView(FMinimalViewInfo& DesiredView);
 
+	/** Collect reference objects for the garbage collector. */
+	GAMEPLAYCAMERAS_API void AddReferencedObjects(FReferenceCollector& Collector);
+
 #if UE_GAMEPLAY_CAMERAS_DEBUG
 	GAMEPLAYCAMERAS_API void DebugUpdate(const FCameraSystemDebugUpdateParams& Params);
 #endif  // UE_GAMEPLAY_CAMERAS_DEBUG
 
-public:
-
-	GAMEPLAYCAMERAS_API void AddReferencedObjects(FReferenceCollector& Collector);
+private:
 
-protected:
+	void PreUpdateServices(float DeltaTime, ECameraEvaluationServiceFlags ExtraFlags);
+	void PostUpdateServices(float DeltaTime, ECameraEvaluationServiceFlags ExtraFlags);
 
 	void NotifyRootCameraNodeEvent(const FRootCameraNodeCameraRigEvent& InEvent);
 
@@ -159,6 +167,9 @@ private:
 	/** The list of evaluation services. */
 	TArray<TSharedPtr<FCameraEvaluationService>> EvaluationServices;
 
+	/** Quick access to the variable auto-reset service. */
+	TSharedPtr<FAutoResetCameraVariableService> VariableAutoResetService;
+
 	/** Storage buffer for the root evaluator. */
 	FCameraNodeEvaluatorStorage RootEvaluatorStorage;
 
@@ -169,7 +180,7 @@ private:
 	FCameraNodeEvaluationResult RootNodeResult;
 
 	/** The current overall result of the camera system. */
-	FCameraSystemEvaluationUpdateResult Result;
+	FCameraSystemEvaluationResult Result;
 
 #if UE_GAMEPLAY_CAMERAS_DEBUG
 	/** Storage for debug drawing blocks. */
diff --git a/Engine/Plugins/Cameras/GameplayCameras/Source/GameplayCameras/Public/Core/DefaultRootCameraNode.h b/Engine/Plugins/Cameras/GameplayCameras/Source/GameplayCameras/Public/Core/DefaultRootCameraNode.h
--- a/Engine/Plugins/Cameras/GameplayCameras/Source/GameplayCameras/Public/Core/DefaultRootCameraNode.h
+++ b/Engine/Plugins/Cameras/GameplayCameras/Source/GameplayCameras/Public/Core/DefaultRootCameraNode.h
@@ -56,16 +56,19 @@ class FDefaultRootCameraNodeEvaluator
 
 protected:
 
-	// FRootCameraNodeEvaluator interface.
+	// FCameraNodeEvaluator interface.
 	virtual FCameraNodeEvaluatorChildrenView OnGetChildren() override;
 	virtual void OnBuild(const FCameraNodeEvaluatorBuildParams& Params) override;
 	virtual void OnRun(const FCameraNodeEvaluationParams& Params, FCameraNodeEvaluationResult& OutResult) override;
-	virtual void OnActivateCameraRig(const FActivateCameraRigParams& Params) override;
 
 #if UE_GAMEPLAY_CAMERAS_DEBUG
 	virtual void OnBuildDebugBlocks(const FCameraDebugBlockBuildParams& Params, FCameraDebugBlockBuilder& Builder) override;
 #endif  // UE_GAMEPLAY_CAMERAS_DEBUG
 
+	// FRootCameraNodeEvaluator interface.
+	virtual void OnActivateCameraRig(const FActivateCameraRigParams& Params) override;
+	virtual void OnRunSingleCameraRig(const FSingleCameraRigEvaluationParams& Params, FCameraNodeEvaluationResult& OutResult) override;
+
 	// IBlendStackCameraNodeObserver interface.
 	virtual void OnBlendStackEvent(const FBlendStackCameraRigEvent& InEvent) override;
 
@@ -73,6 +76,8 @@ private:
 
 	FBlendStackCameraNodeEvaluator* BuildBlendStackEvaluator(const FCameraNodeEvaluatorBuildParams& Params, UBlendStackCameraNode* BlendStackNode);
 
+	FBlendStackCameraNodeEvaluator* GetBlendStackEvaluator(ECameraRigLayer Layer) const;
+
 private:
 
 	FBlendStackCameraNodeEvaluator* BaseLayer;
diff --git a/Engine/Plugins/Cameras/GameplayCameras/Source/GameplayCameras/Public/Core/RootCameraNode.h b/Engine/Plugins/Cameras/GameplayCameras/Source/GameplayCameras/Public/Core/RootCameraNode.h
--- a/Engine/Plugins/Cameras/GameplayCameras/Source/GameplayCameras/Public/Core/RootCameraNode.h
+++ b/Engine/Plugins/Cameras/GameplayCameras/Source/GameplayCameras/Public/Core/RootCameraNode.h
@@ -4,6 +4,7 @@
 
 #include "Core/CameraNode.h"
 #include "Core/CameraNodeEvaluator.h"
+#include "Core/CameraRigEvaluationInfo.h"
 
 #include "RootCameraNode.generated.h"
 
@@ -19,6 +20,7 @@ enum class ECameraRigLayer : uint8
 	Main UMETA(DisplayName="Main Layer"),
 	Global UMETA(DisplayName="Global Layer"),
 	Visual UMETA(DisplayName="Visual Layer"),
+	ScratchMain UMETA(DisplayName="Scratch Main Layer"),
 	User0,
 	User1,
 	User2
@@ -61,6 +63,18 @@ struct FActivateCameraRigParams
 	ECameraRigLayer Layer = ECameraRigLayer::Main;
 };
 
+/**
+ * Parameter structure for evaluating a single camera rig.
+ */
+struct FSingleCameraRigEvaluationParams
+{
+	/** The evaluation parameters. */
+	FCameraNodeEvaluationParams EvaluationParams;
+
+	/** The camera rig to evaluate. */
+	FCameraRigEvaluationInfo CameraRigInfo;
+};
+
 /**
  * Base class for the evaluator of a root camera node.
  */
@@ -71,6 +85,13 @@ public:
 	/** Activates a camera rig. */
 	void ActivateCameraRig(const FActivateCameraRigParams& Params);
 
+	/**
+	 * Evaluates a single camera rig.
+	 * This is expected to run all layers as usual, except for the main layer which should
+	 * only run the given camera rig instead.
+	 */
+	void RunSingleCameraRig(const FSingleCameraRigEvaluationParams& Params, FCameraNodeEvaluationResult& OutResult);
+
 	/** Registers an observer to this root node. */
 	void RegisterObserver(IRootCameraNodeObserver* Observer);
 	/** Unregisters an observer from this root node. */
@@ -86,7 +107,15 @@ protected:
 	/** Activates a camera rig. */
 	virtual void OnActivateCameraRig(const FActivateCameraRigParams& Params) {}
 
+	/** Evaluates a single camera rig. See comments on RunSingleCameraRig. */
+	virtual void OnRunSingleCameraRig(const FSingleCameraRigEvaluationParams& Params, FCameraNodeEvaluationResult& OutResult) {}
+
+protected:
+
+	/** Returns whether this root camera node has any observers. */
 	bool HasObservers() const;
+	
+	/** Notifies observers of a root camera node event. */
 	void NotifyObservers(const FRootCameraNodeCameraRigEvent& InEvent) const;
 
 private:
diff --git a/Engine/Plugins/Cameras/GameplayCameras/Source/GameplayCameras/Public/Debug/CameraSystemTrace.h b/Engine/Plugins/Cameras/GameplayCameras/Source/GameplayCameras/Public/Debug/CameraSystemTrace.h
--- a/Engine/Plugins/Cameras/GameplayCameras/Source/GameplayCameras/Public/Debug/CameraSystemTrace.h
+++ b/Engine/Plugins/Cameras/GameplayCameras/Source/GameplayCameras/Public/Debug/CameraSystemTrace.h
@@ -16,7 +16,7 @@ namespace UE::Cameras
 
 class FCameraDebugBlock;
 class FCameraDebugBlockStorage;
-struct FCameraSystemEvaluationUpdateResult;
+struct FCameraSystemEvaluationResult;
 
 /**
  * Trace utility class for the camera system.
@@ -39,7 +39,7 @@ public:
 	/** Returns whether tracing of camera system evaluation is enabled. */
 	GAMEPLAYCAMERAS_API static bool IsTraceEnabled();
 	/** Records one frame of camera system evaluation. */
-	GAMEPLAYCAMERAS_API static void TraceEvaluation(UWorld* InWorld, const FCameraSystemEvaluationUpdateResult& InResult, FCameraDebugBlock& InRootDebugBlock);
+	GAMEPLAYCAMERAS_API static void TraceEvaluation(UWorld* InWorld, const FCameraSystemEvaluationResult& InResult, FCameraDebugBlock& InRootDebugBlock);
 	/** Reads back one frame of camera system evaluation. */
 	GAMEPLAYCAMERAS_API static FCameraDebugBlock* ReadEvaluationTrace(TArray<uint8> InSerializedBlocks, FCameraDebugBlockStorage& InStorage);
 
diff --git a/Engine/Plugins/Cameras/GameplayCameras/Source/GameplayCameras/Public/GameplayCamerasSettings.h b/Engine/Plugins/Cameras/GameplayCameras/Source/GameplayCameras/Public/GameplayCamerasSettings.h
new file mode 100644
--- /dev/null
+++ b/Engine/Plugins/Cameras/GameplayCameras/Source/GameplayCameras/Public/GameplayCamerasSettings.h
@@ -0,0 +1,37 @@
+// Copyright Epic Games, Inc. All Rights Reserved.
+
+#pragma once
+
+#include "CoreTypes.h"
+#include "UObject/ObjectMacros.h"
+#include "UObject/Object.h"
+
+#include "GameplayCamerasSettings.generated.h"
+
+/**
+ * The settings for the Gameplay Cameras runtime.
+ */
+UCLASS(config=Project, MinimalAPI)
+class UGameplayCamerasSettings : public UObject
+{
+	GENERATED_BODY()
+
+public:
+
+	/** The default angle tolerance to accept an aiming operation. */
+	UPROPERTY(EditAnywhere, Config, Category="IK Aiming")
+	double DefaultIKAimingAngleTolerance = 0.1;  // 0.1 degrees
+
+	/** The default distance tolerance to accept an aiming operation. */
+	UPROPERTY(EditAnywhere, Config, Category="IK Aiming")
+	double DefaultIKAimingDistanceTolerance = 1.0;  // 1cm
+
+	/** The default number of iterations for an aiming operation. */
+	UPROPERTY(EditAnywhere, Config, Category="IK Aiming")
+	uint8 DefaultIKAimingMaxIterations = 3;
+
+	/** The distance below which any IK aiming operation is disabled. */
+	UPROPERTY(EditAnywhere, Config, Category="IK Aiming")
+	double DefaultIKAimingMinDistance = 100.0;  // 1m
+};
+
diff --git a/Engine/Plugins/Cameras/GameplayCameras/Source/GameplayCameras/Public/Math/ConsumableValue.h b/Engine/Plugins/Cameras/GameplayCameras/Source/GameplayCameras/Public/Math/ConsumableValue.h
new file mode 100644
--- /dev/null
+++ b/Engine/Plugins/Cameras/GameplayCameras/Source/GameplayCameras/Public/Math/ConsumableValue.h
@@ -0,0 +1,184 @@
+// Copyright Epic Games, Inc. All Rights Reserved.
+
+#pragma once
+
+#include "Misc/AssertionMacros.h"
+
+namespace UE::Cameras
+{
+
+/**
+ * A value that can be "consumed" and applied to another value.
+ * The consumable value can be a "delta", which is added to another value,
+ * or an "absolute", which is set on another value.
+ */
+template<typename ValueType>
+struct TConsumableValue
+{
+public:
+
+	/** Returns an absolute consumable with the given value. */
+	static TConsumableValue<ValueType> Absolute(ValueType InValue)
+	{
+		return TConsumableValue<ValueType>(InValue, false);
+	}
+
+	/** Returns a delta consumable with the given value. */
+	static TConsumableValue<ValueType> Delta(ValueType InValue)
+	{
+		return TConsumableValue<ValueType>(InValue, true);
+	}
+
+	/** Creates a new consumable with no value. */
+	TConsumableValue()
+		: bHasValue(false)
+	{}
+
+	/** Creates a new consumable. */
+	TConsumableValue(ValueType InValue, bool bInIsDelta)
+		: Value(InValue)
+		, bHasValue(true)
+		, bIsDelta(bInIsDelta)
+	{}
+
+public:
+
+	/** Gets the current value. Asserts if there is no value. */
+	ValueType Get() const
+	{
+		check(bHasValue);
+		return Value;
+	}
+
+	/** Gets the current value, or the default one if there's no value. */
+	ValueType GetOrDefault(ValueType DefaultValue) const
+	{
+		return bHasValue ? Value : DefaultValue;
+	}
+
+	/** Returns whether there's any value left in the consumable. */
+	bool HasValue() const
+	{
+		return bHasValue;
+	}
+
+	/** Returns whether the consumable is a delta. */
+	bool IsDelta() const
+	{
+		return bIsDelta;
+	}
+
+public:
+
+	/**
+	 * Applies the consumable to the given value and returns the new value.
+	 *
+	 * This is meant to be used as follows:
+	 *
+	 *   MyFooBar = Consumable.Apply(MyFooBar);
+	 *
+	 * After this, the consumable doesn't have any value left in it anymore,
+	 * so HasValue() will return false and further calls to Apply() will just
+	 * return the same value as the given parameter.
+	 */
+	ValueType Apply(const ValueType InTarget)
+	{
+		if (bHasValue)
+		{
+			if (bIsDelta)
+			{
+				ValueType Result = InTarget + Value;
+				Value = ValueType();
+				bHasValue = false;
+				return Result;
+			}
+			else
+			{
+				ValueType Result = Value;
+				Value = ValueType();
+				bHasValue = false;
+				return Result;
+			}
+		}
+		return InTarget;
+	}
+
+	/**
+	 * Applies the consumable to the given value and returns the new value. However,
+	 * check that the return value falls within the given min/max bounds. If not, only
+	 * consume enough to reach those bounds and leave the rest to be consumed later.
+	 *
+	 * This is meant to be used as follows:
+	 *
+	 *   MyFooBar = Consumable.Apply(MyFooBar, MyMin, MyMax);
+	 *
+	 * After this, the consumable may or may not have any value left in it, depending
+	 * on whether the given min/max bounds were reached.
+	 */
+	ValueType Apply(const ValueType InTarget, const ValueType InMinTarget, const ValueType InMaxTarget)
+	{
+		check(InMinTarget <= InMaxTarget);
+		if (bHasValue)
+		{
+			if (bIsDelta)
+			{
+				ValueType Result = InTarget + Value;
+				if (Result < InMinTarget)
+				{
+					ValueType ActuallyConsumable = InMinTarget - InTarget;
+					Value -= ActuallyConsumable;
+					return InMinTarget;
+				}
+				else if (Result > InMaxTarget)
+				{
+					ValueType ActuallyConsumable = InMaxTarget - InTarget;
+					Value -= ActuallyConsumable;
+					return InMaxTarget;
+				}
+				else
+				{
+					Value = ValueType();
+					bHasValue = false;
+					return Result;
+				}
+			}
+			else
+			{
+				if (Value < InMinTarget)
+				{
+					ValueType Remaining = Value - InMinTarget;
+					Value = Remaining;
+					bIsDelta = true;
+					return InMinTarget;
+				}
+				else if (Value > InMaxTarget)
+				{
+					ValueType Remaining = Value - InMaxTarget;
+					Value = Remaining;
+					bIsDelta = true;
+					return InMaxTarget;
+				}
+				else
+				{
+					ValueType Result = Value;
+					Value = ValueType();
+					bHasValue = false;
+					return Result;
+				}
+			}
+		}
+		return InTarget;
+	}
+
+private:
+
+	ValueType Value = ValueType();
+	bool bHasValue = false;
+	bool bIsDelta = false;
+};
+
+using FConsumableFloat = TConsumableValue<float>;
+using FConsumableDouble = TConsumableValue<double>;
+
+}  // namespace UE::Cameras
+
diff --git a/Engine/Plugins/Cameras/GameplayCameras/Source/GameplayCameras/Public/Nodes/Blends/SimpleBlendCameraNode.h b/Engine/Plugins/Cameras/GameplayCameras/Source/GameplayCameras/Public/Nodes/Blends/SimpleBlendCameraNode.h
--- a/Engine/Plugins/Cameras/GameplayCameras/Source/GameplayCameras/Public/Nodes/Blends/SimpleBlendCameraNode.h
+++ b/Engine/Plugins/Cameras/GameplayCameras/Source/GameplayCameras/Public/Nodes/Blends/SimpleBlendCameraNode.h
@@ -55,6 +55,7 @@ protected:
 	virtual void OnRun(const FCameraNodeEvaluationParams& Params, FCameraNodeEvaluationResult& OutResult) override;
 	virtual void OnBlendParameters(const FCameraNodePreBlendParams& Params, FCameraNodePreBlendResult& OutResult) override;
 	virtual void OnBlendResults(const FCameraNodeBlendParams& Params, FCameraNodeBlendResult& OutResult) override;
+	virtual void OnSerialize(const FCameraNodeEvaluatorSerializeParams& Params, FArchive& Ar) override;
 
 	virtual void OnComputeBlendFactor(const FCameraNodeEvaluationParams& Params, FSimpleBlendCameraNodeEvaluationResult& OutResult) {}
 
diff --git a/Engine/Plugins/Cameras/GameplayCameras/Source/GameplayCameras/Public/Nodes/Input/CameraRigInput1DSlot.h b/Engine/Plugins/Cameras/GameplayCameras/Source/GameplayCameras/Public/Nodes/Input/CameraRigInput1DSlot.h
--- a/Engine/Plugins/Cameras/GameplayCameras/Source/GameplayCameras/Public/Nodes/Input/CameraRigInput1DSlot.h
+++ b/Engine/Plugins/Cameras/GameplayCameras/Source/GameplayCameras/Public/Nodes/Input/CameraRigInput1DSlot.h
@@ -90,6 +90,8 @@ protected:
 	virtual FCameraNodeEvaluatorChildrenView OnGetChildren() override;
 	virtual void OnUpdateParameters(const FCameraBlendedParameterUpdateParams& Params, FCameraBlendedParameterUpdateResult& OutResult) override;
 	virtual void OnRun(const FCameraNodeEvaluationParams& Params, FCameraNodeEvaluationResult& OutResult) override;
+	virtual void OnExecuteOperation(const FCameraOperationParams& Params, FCameraOperation& Operation) override;
+	virtual void OnSerialize(const FCameraNodeEvaluatorSerializeParams& Params, FArchive& Ar) override;
 
 protected:
 
diff --git a/Engine/Plugins/Cameras/GameplayCameras/Source/GameplayCameras/Public/Nodes/Input/CameraRigInput2DSlot.h b/Engine/Plugins/Cameras/GameplayCameras/Source/GameplayCameras/Public/Nodes/Input/CameraRigInput2DSlot.h
--- a/Engine/Plugins/Cameras/GameplayCameras/Source/GameplayCameras/Public/Nodes/Input/CameraRigInput2DSlot.h
+++ b/Engine/Plugins/Cameras/GameplayCameras/Source/GameplayCameras/Public/Nodes/Input/CameraRigInput2DSlot.h
@@ -98,6 +98,8 @@ protected:
 	virtual FCameraNodeEvaluatorChildrenView OnGetChildren() override;
 	virtual void OnUpdateParameters(const FCameraBlendedParameterUpdateParams& Params, FCameraBlendedParameterUpdateResult& OutResult) override;
 	virtual void OnRun(const FCameraNodeEvaluationParams& Params, FCameraNodeEvaluationResult& OutResult) override;
+	virtual void OnExecuteOperation(const FCameraOperationParams& Params, FCameraOperation& Operation) override;
+	virtual void OnSerialize(const FCameraNodeEvaluatorSerializeParams& Params, FArchive& Ar) override;
 
 protected:
 
diff --git a/Engine/Plugins/Cameras/GameplayCameras/Source/GameplayCameras/Public/Nodes/Input/CameraRigInputSlotTypes.h b/Engine/Plugins/Cameras/GameplayCameras/Source/GameplayCameras/Public/Nodes/Input/CameraRigInputSlotTypes.h
--- a/Engine/Plugins/Cameras/GameplayCameras/Source/GameplayCameras/Public/Nodes/Input/CameraRigInputSlotTypes.h
+++ b/Engine/Plugins/Cameras/GameplayCameras/Source/GameplayCameras/Public/Nodes/Input/CameraRigInputSlotTypes.h
@@ -45,6 +45,13 @@ public:
 
 	/** Clamps the given value. */
 	double ClampValue(double Value) const;
+
+	/** 
+	 * Gets the effective min/max values for this struct.
+	 * If a bound is disabled, the effective value will be lowest or max
+	 * double precision values.
+	 */
+	void GetEffectiveClamping(double& OutMinValue, double& OutMaxValue) const;
 };
 
 /**
