Cameras: replace observers with multicast delegates

diff --git a/Engine/Plugins/Cameras/GameplayCameras/Source/GameplayCameras/Private/Core/BlendStackCameraNode.cpp b/Engine/Plugins/Cameras/GameplayCameras/Source/GameplayCameras/Private/Core/BlendStackCameraNode.cpp
--- a/Engine/Plugins/Cameras/GameplayCameras/Source/GameplayCameras/Private/Core/BlendStackCameraNode.cpp
+++ b/Engine/Plugins/Cameras/GameplayCameras/Source/GameplayCameras/Private/Core/BlendStackCameraNode.cpp
@@ -3,7 +3,7 @@
 #include "Core/BlendStackCameraNode.h"
 
 #include "Core/BlendCameraNode.h"
-#include "Core/BlendStackCameraNodeObserver.h"
+#include "Core/BlendStackCameraRigEvent.h"
 #include "Core/BlendStackRootCameraNode.h"
 #include "Core/CameraAsset.h"
 #include "Core/CameraEvaluationContext.h"
@@ -125,10 +125,9 @@ void FBlendStackCameraNodeEvaluator::Push(const FBlendStackCameraPushParams& Par
 	// is disabled.
 	Entries.Add(MoveTemp(NewEntry));
 
-	// Notify observers.
-	if (!Observers.IsEmpty())
+	if (OnCameraRigEventDelegate.IsBound())
 	{
-		NotifyObservers(EBlendStackCameraRigEventType::Pushed, Entries.Last(), UsedTransition);
+		BroadcastCameraRigEvent(EBlendStackCameraRigEventType::Pushed, Entries.Last(), UsedTransition);
 	}
 }
 
@@ -500,9 +499,9 @@ void FBlendStackCameraNodeEvaluator::PopEntries(int32 FirstIndexToKeep)
 		}
 #endif  // WITH_EDITOR
 
-		if (!Observers.IsEmpty())
+		if (OnCameraRigEventDelegate.IsBound())
 		{
-			NotifyObservers(EBlendStackCameraRigEventType::Popped, FirstEntry);
+			BroadcastCameraRigEvent(EBlendStackCameraRigEventType::Popped, FirstEntry);
 		}
 
 		Entries.RemoveAt(0);
@@ -625,17 +624,7 @@ const UCameraRigTransition* FBlendStackCameraNodeEvaluator::FindTransition(
 	return nullptr;
 }
 
-void FBlendStackCameraNodeEvaluator::RegisterObserver(IBlendStackCameraNodeObserver* Observer)
-{
-	Observers.Add(Observer);
-}
-
-void FBlendStackCameraNodeEvaluator::UnregisterObserver(IBlendStackCameraNodeObserver* Observer)
-{
-	Observers.Remove(Observer);
-}
-
-void FBlendStackCameraNodeEvaluator::NotifyObservers(EBlendStackCameraRigEventType EventType, const FCameraRigEntry& Entry, const UCameraRigTransition* Transition) const
+void FBlendStackCameraNodeEvaluator::BroadcastCameraRigEvent(EBlendStackCameraRigEventType EventType, const FCameraRigEntry& Entry, const UCameraRigTransition* Transition) const
 {
 	FBlendStackCameraRigEvent Event;
 	Event.EventType = EventType;
@@ -647,10 +636,7 @@ void FBlendStackCameraNodeEvaluator::NotifyObservers(EBlendStackCameraRigEventTy
 			Entry.RootEvaluator);
 	Event.Transition = Transition;
 
-	for (IBlendStackCameraNodeObserver* Observer : Observers)
-	{
-		Observer->OnBlendStackEvent(Event);
-	}
+	OnCameraRigEventDelegate.Broadcast(Event);
 }
 
 void FBlendStackCameraNodeEvaluator::OnAddReferencedObjects(FReferenceCollector& Collector)
diff --git a/Engine/Plugins/Cameras/GameplayCameras/Source/GameplayCameras/Private/Core/DefaultRootCameraNode.cpp b/Engine/Plugins/Cameras/GameplayCameras/Source/GameplayCameras/Private/Core/DefaultRootCameraNode.cpp
--- a/Engine/Plugins/Cameras/GameplayCameras/Source/GameplayCameras/Private/Core/DefaultRootCameraNode.cpp
+++ b/Engine/Plugins/Cameras/GameplayCameras/Source/GameplayCameras/Private/Core/DefaultRootCameraNode.cpp
@@ -4,8 +4,7 @@
 
 #include "Core/BlendStackCameraNode.h"
 #include "Core/CameraEvaluationContext.h"
-#include "Core/CameraRigAsset.h"
-#include "Core/RootCameraNodeObserver.h"
+#include "Core/RootCameraNodeCameraRigEvent.h"
 #include "Debug/BlendStacksCameraDebugBlock.h"
 #include "Debug/CameraDebugBlockBuilder.h"
 #include "Debug/RootCameraDebugBlock.h"
@@ -62,7 +61,7 @@ void FDefaultRootCameraNodeEvaluator::OnBuild(const FCameraNodeEvaluatorBuildPar
 FBlendStackCameraNodeEvaluator* FDefaultRootCameraNodeEvaluator::BuildBlendStackEvaluator(const FCameraNodeEvaluatorBuildParams& Params, UBlendStackCameraNode* BlendStackNode)
 {
 	FBlendStackCameraNodeEvaluator* BlendStackEvaluator = Params.BuildEvaluatorAs<FBlendStackCameraNodeEvaluator>(BlendStackNode);
-	BlendStackEvaluator->RegisterObserver(this);
+	BlendStackEvaluator->OnCameraRigEvent().AddRaw(this, &FDefaultRootCameraNodeEvaluator::OnBlendStackEvent);
 	return BlendStackEvaluator;
 }
 
@@ -175,7 +174,7 @@ void FDefaultRootCameraNodeEvaluator::OnBlendStackEvent(const FBlendStackCameraR
 			RootEvent.EventLayer = ECameraRigLayer::Visual;
 		}
 
-		NotifyObservers(RootEvent);
+		BroadcastCameraRigEvent(RootEvent);
 	}
 }
 
diff --git a/Engine/Plugins/Cameras/GameplayCameras/Source/GameplayCameras/Private/Core/RootCameraNode.cpp b/Engine/Plugins/Cameras/GameplayCameras/Source/GameplayCameras/Private/Core/RootCameraNode.cpp
--- a/Engine/Plugins/Cameras/GameplayCameras/Source/GameplayCameras/Private/Core/RootCameraNode.cpp
+++ b/Engine/Plugins/Cameras/GameplayCameras/Source/GameplayCameras/Private/Core/RootCameraNode.cpp
@@ -4,7 +4,7 @@
 
 #include "Core/CameraEvaluationService.h"
 #include "Core/CameraSystemEvaluator.h"
-#include "Core/RootCameraNodeObserver.h"
+#include "Core/RootCameraNodeCameraRigEvent.h"
 #include "Services/AutoResetCameraVariableService.h"
 
 #include UE_INLINE_GENERATED_CPP_BY_NAME(RootCameraNode)
@@ -41,32 +41,14 @@ void FRootCameraNodeEvaluator::RunSingleCameraRig(const FSingleCameraRigEvaluati
 	OnRunSingleCameraRig(Params, OutResult);
 }
 
-void FRootCameraNodeEvaluator::RegisterObserver(IRootCameraNodeObserver* Observer)
-{
-	Observers.Add(Observer);
-}
-
-void FRootCameraNodeEvaluator::UnregisterObserver(IRootCameraNodeObserver* Observer)
-{
-	Observers.Remove(Observer);
-}
-
-bool FRootCameraNodeEvaluator::HasObservers() const
-{
-	return !Observers.IsEmpty();
-}
-
-void FRootCameraNodeEvaluator::NotifyObservers(const FRootCameraNodeCameraRigEvent& InEvent) const
+void FRootCameraNodeEvaluator::BroadcastCameraRigEvent(const FRootCameraNodeCameraRigEvent& InEvent) const
 {
 	if (ensure(OwningEvaluator))
 	{
 		OwningEvaluator->NotifyRootCameraNodeEvent(InEvent);
 	}
 
-	for (IRootCameraNodeObserver* Observer : Observers)
-	{
-		Observer->OnRootCameraNodeEvent(InEvent);
-	}
+	OnCameraRigEventDelegate.Broadcast(InEvent);
 }
 
 }  // namespace UE::Cameras
diff --git a/Engine/Plugins/Cameras/GameplayCameras/Source/GameplayCameras/Private/Services/AutoResetCameraVariableService.cpp b/Engine/Plugins/Cameras/GameplayCameras/Source/GameplayCameras/Private/Services/AutoResetCameraVariableService.cpp
--- a/Engine/Plugins/Cameras/GameplayCameras/Source/GameplayCameras/Private/Services/AutoResetCameraVariableService.cpp
+++ b/Engine/Plugins/Cameras/GameplayCameras/Source/GameplayCameras/Private/Services/AutoResetCameraVariableService.cpp
@@ -6,7 +6,7 @@
 #include "Core/CameraRigAsset.h"
 #include "Core/CameraVariableAssets.h"
 #include "Core/CameraVariableTable.h"
-#include "Core/RootCameraNodeObserver.h"
+#include "Core/RootCameraNodeCameraRigEvent.h"
 
 namespace UE::Cameras
 {
diff --git a/Engine/Plugins/Cameras/GameplayCameras/Source/GameplayCameras/Private/Services/OrientationInitializationService.cpp b/Engine/Plugins/Cameras/GameplayCameras/Source/GameplayCameras/Private/Services/OrientationInitializationService.cpp
--- a/Engine/Plugins/Cameras/GameplayCameras/Source/GameplayCameras/Private/Services/OrientationInitializationService.cpp
+++ b/Engine/Plugins/Cameras/GameplayCameras/Source/GameplayCameras/Private/Services/OrientationInitializationService.cpp
@@ -9,7 +9,7 @@
 #include "Core/CameraRigAsset.h"
 #include "Core/CameraSystemEvaluator.h"
 #include "Core/RootCameraNode.h"
-#include "Core/RootCameraNodeObserver.h"
+#include "Core/RootCameraNodeCameraRigEvent.h"
 
 namespace UE::Cameras
 {
diff --git a/Engine/Plugins/Cameras/GameplayCameras/Source/GameplayCameras/Public/Core/BlendStackCameraNode.h b/Engine/Plugins/Cameras/GameplayCameras/Source/GameplayCameras/Public/Core/BlendStackCameraNode.h
--- a/Engine/Plugins/Cameras/GameplayCameras/Source/GameplayCameras/Public/Core/BlendStackCameraNode.h
+++ b/Engine/Plugins/Cameras/GameplayCameras/Source/GameplayCameras/Public/Core/BlendStackCameraNode.h
@@ -23,8 +23,8 @@ namespace UE::Cameras
 class FBlendStackRootCameraNodeEvaluator;
 class FCameraEvaluationContext;
 class FCameraSystemEvaluator;
-class IBlendStackCameraNodeObserver;
 enum class EBlendStackCameraRigEventType;
+struct FBlendStackCameraRigEvent;
 
 #if UE_GAMEPLAY_CAMERAS_DEBUG
 class FBlendStackCameraDebugBlock;
@@ -79,6 +79,8 @@ struct FBlendStackCameraPushParams
 	TObjectPtr<const UCameraRigAsset> CameraRig;
 };
 
+DECLARE_MULTICAST_DELEGATE_OneParam(FOnBlendStackCameraRigEvent, const FBlendStackCameraRigEvent&);
+
 /**
  * Evaluator for a blend stack camera node.
  */
@@ -106,10 +108,8 @@ public:
 
 public:
 
-	/** Register a blend stack observer. */
-	void RegisterObserver(IBlendStackCameraNodeObserver* Observer);
-	/** Unregister a blend stack observer. */
-	void UnregisterObserver(IBlendStackCameraNodeObserver* Observer);
+	/** Gets the delegate for blend stack events. */
+	FOnBlendStackCameraRigEvent& OnCameraRigEvent() { return OnCameraRigEventDelegate; }
 
 protected:
 
@@ -151,7 +151,7 @@ protected:
 
 	void GatherEntryParameterEvaluators(FCameraNodeEvaluator* RootEvaluator, TArray<FCameraNodeEvaluator*>& OutParameterEvaluators);
 
-	void NotifyObservers(EBlendStackCameraRigEventType EventType, const FCameraRigEntry& Entry, const UCameraRigTransition* Transition = nullptr) const;
+	void BroadcastCameraRigEvent(EBlendStackCameraRigEventType EventType, const FCameraRigEntry& Entry, const UCameraRigTransition* Transition = nullptr) const;
 
 protected:
 
@@ -195,8 +195,8 @@ protected:
 	/** Entries in the blend stack. */
 	TArray<FCameraRigEntry> Entries;
 
-	/** The list of observers for this blend stack. */
-	TArray<IBlendStackCameraNodeObserver*> Observers;
+	/** The delegate to invoke when an event occurs in this blend stack. */
+	FOnBlendStackCameraRigEvent OnCameraRigEventDelegate;
 
 #if WITH_EDITOR
 	TMap<const UPackage*, int32> AllListenedPackages;
diff --git a/Engine/Plugins/Cameras/GameplayCameras/Source/GameplayCameras/Public/Core/BlendStackCameraNodeObserver.h b/Engine/Plugins/Cameras/GameplayCameras/Source/GameplayCameras/Public/Core/BlendStackCameraNodeObserver.h
deleted file mode 100644
diff --git a/Engine/Plugins/Cameras/GameplayCameras/Source/GameplayCameras/Public/Core/BlendStackCameraRigEvent.h b/Engine/Plugins/Cameras/GameplayCameras/Source/GameplayCameras/Public/Core/BlendStackCameraRigEvent.h
new file mode 100644
--- /dev/null
+++ b/Engine/Plugins/Cameras/GameplayCameras/Source/GameplayCameras/Public/Core/BlendStackCameraRigEvent.h
@@ -0,0 +1,42 @@
+// Copyright Epic Games, Inc. All Rights Reserved.
+
+#pragma once
+
+#include "Core/CameraRigEvaluationInfo.h"
+#include "UObject/ObjectPtr.h"
+
+class UCameraRigTransition;
+
+namespace UE::Cameras
+{
+
+class FBlendStackCameraNodeEvaluator;
+
+/** A type of blend stack event. */
+enum class EBlendStackCameraRigEventType
+{
+	/** A camera rig was pushed on the blend stack. */
+	Pushed,
+	/** A camera rig was frozen. */
+	Frozen,
+	/** A camera rig was popped out of the blend stack. */
+	Popped
+};
+
+/**
+ * A structure describing an event happening on a blend stack.
+ */
+struct FBlendStackCameraRigEvent
+{
+	/** The event type. */
+	EBlendStackCameraRigEventType EventType;
+	/** The blend stack evaluator in which the event is happening. */
+	const FBlendStackCameraNodeEvaluator* BlendStackEvaluator = nullptr;
+	/** The evaluation information for the camera rig associated with the event. */
+	FCameraRigEvaluationInfo CameraRigInfo;
+	/** If the event is Pushed, the transition used for the new camera rig. */
+	TObjectPtr<const UCameraRigTransition> Transition;
+};
+
+}  // namespace UE::Cameras
+
diff --git a/Engine/Plugins/Cameras/GameplayCameras/Source/GameplayCameras/Public/Core/DefaultRootCameraNode.h b/Engine/Plugins/Cameras/GameplayCameras/Source/GameplayCameras/Public/Core/DefaultRootCameraNode.h
--- a/Engine/Plugins/Cameras/GameplayCameras/Source/GameplayCameras/Public/Core/DefaultRootCameraNode.h
+++ b/Engine/Plugins/Cameras/GameplayCameras/Source/GameplayCameras/Public/Core/DefaultRootCameraNode.h
@@ -3,7 +3,7 @@
 #pragma once
 
 #include "Core/RootCameraNode.h"
-#include "Core/BlendStackCameraNodeObserver.h"
+#include "Core/BlendStackCameraRigEvent.h"
 
 #include "DefaultRootCameraNode.generated.h"
 
@@ -48,9 +48,7 @@ class FBlendStackCameraNodeEvaluator;
 /**
  * Evaluator for the default root camera node.
  */
-class FDefaultRootCameraNodeEvaluator 
-	: public FRootCameraNodeEvaluator
-	, public IBlendStackCameraNodeObserver
+class FDefaultRootCameraNodeEvaluator : public FRootCameraNodeEvaluator
 {
 	UE_DECLARE_CAMERA_NODE_EVALUATOR(GAMEPLAYCAMERAS_API, FDefaultRootCameraNodeEvaluator)
 
@@ -69,15 +67,14 @@ protected:
 	virtual void OnActivateCameraRig(const FActivateCameraRigParams& Params) override;
 	virtual void OnRunSingleCameraRig(const FSingleCameraRigEvaluationParams& Params, FCameraNodeEvaluationResult& OutResult) override;
 
-	// IBlendStackCameraNodeObserver interface.
-	virtual void OnBlendStackEvent(const FBlendStackCameraRigEvent& InEvent) override;
-
 private:
 
 	FBlendStackCameraNodeEvaluator* BuildBlendStackEvaluator(const FCameraNodeEvaluatorBuildParams& Params, UBlendStackCameraNode* BlendStackNode);
 
 	FBlendStackCameraNodeEvaluator* GetBlendStackEvaluator(ECameraRigLayer Layer) const;
 
+	void OnBlendStackEvent(const FBlendStackCameraRigEvent& InEvent);
+
 private:
 
 	FBlendStackCameraNodeEvaluator* BaseLayer;
diff --git a/Engine/Plugins/Cameras/GameplayCameras/Source/GameplayCameras/Public/Core/RootCameraNode.h b/Engine/Plugins/Cameras/GameplayCameras/Source/GameplayCameras/Public/Core/RootCameraNode.h
--- a/Engine/Plugins/Cameras/GameplayCameras/Source/GameplayCameras/Public/Core/RootCameraNode.h
+++ b/Engine/Plugins/Cameras/GameplayCameras/Source/GameplayCameras/Public/Core/RootCameraNode.h
@@ -42,7 +42,6 @@ namespace UE::Cameras
 
 class FCameraEvaluationContext;
 class FCameraSystemEvaluator;
-class IRootCameraNodeObserver;
 struct FRootCameraNodeCameraRigEvent;
 
 /**
@@ -75,6 +74,8 @@ struct FSingleCameraRigEvaluationParams
 	FCameraRigEvaluationInfo CameraRigInfo;
 };
 
+DECLARE_MULTICAST_DELEGATE_OneParam(FOnRootCameraNodeCameraRigEvent, const FRootCameraNodeCameraRigEvent&);
+
 /**
  * Base class for the evaluator of a root camera node.
  */
@@ -92,10 +93,8 @@ public:
 	 */
 	void RunSingleCameraRig(const FSingleCameraRigEvaluationParams& Params, FCameraNodeEvaluationResult& OutResult);
 
-	/** Registers an observer to this root node. */
-	void RegisterObserver(IRootCameraNodeObserver* Observer);
-	/** Unregisters an observer from this root node. */
-	void UnregisterObserver(IRootCameraNodeObserver* Observer);
+	/** Gets the delegate for camera rig events. */
+	FOnRootCameraNodeCameraRigEvent& OnCameraRigEvent() { return OnCameraRigEventDelegate; }
 
 protected:
 
@@ -112,19 +111,15 @@ protected:
 
 protected:
 
-	/** Returns whether this root camera node has any observers. */
-	bool HasObservers() const;
-	
-	/** Notifies observers of a root camera node event. */
-	void NotifyObservers(const FRootCameraNodeCameraRigEvent& InEvent) const;
+	void BroadcastCameraRigEvent(const FRootCameraNodeCameraRigEvent& InEvent) const;
 
 private:
 
 	/** The camera system that owns this root node. */
 	FCameraSystemEvaluator* OwningEvaluator = nullptr;
 
-	/** The list of observers. */
-	TArray<IRootCameraNodeObserver*> Observers;
+	/** The delegate to notify when an event happens. */
+	FOnRootCameraNodeCameraRigEvent OnCameraRigEventDelegate;
 };
 
 }  // namespace UE::Cameras
diff --git a/Engine/Plugins/Cameras/GameplayCameras/Source/GameplayCameras/Public/Core/RootCameraNodeCameraRigEvent.h b/Engine/Plugins/Cameras/GameplayCameras/Source/GameplayCameras/Public/Core/RootCameraNodeCameraRigEvent.h
new file mode 100644
--- /dev/null
+++ b/Engine/Plugins/Cameras/GameplayCameras/Source/GameplayCameras/Public/Core/RootCameraNodeCameraRigEvent.h
@@ -0,0 +1,40 @@
+// Copyright Epic Games, Inc. All Rights Reserved.
+
+#pragma once
+
+#include "Core/CameraRigEvaluationInfo.h"
+#include "UObject/ObjectPtr.h"
+
+class UCameraRigTransition;
+enum class ECameraRigLayer : uint8;
+
+namespace UE::Cameras
+{
+
+/** The type of root node event. */
+enum class ERootCameraNodeCameraRigEventType
+{
+	/** A camera rig was activated. */
+	Activated,
+	/** A camera rig was deactivated. */
+	Deactivated
+};
+
+/**
+ * A structure describing an event on a root node.
+ */
+struct FRootCameraNodeCameraRigEvent
+{
+	/** The type of event. */
+	ERootCameraNodeCameraRigEventType EventType;
+	/** The layer on which the event happened. */
+	ECameraRigLayer EventLayer;
+	/** The evaluation information of the associated camera rig. */
+	FCameraRigEvaluationInfo CameraRigInfo;
+	/** If a camera rig was activated, the transition used to activate it, if any. */
+	TObjectPtr<const UCameraRigTransition> Transition;
+};
+
+}  // namespace UE::Cameras
+
+
diff --git a/Engine/Plugins/Cameras/GameplayCameras/Source/GameplayCameras/Public/Core/RootCameraNodeObserver.h b/Engine/Plugins/Cameras/GameplayCameras/Source/GameplayCameras/Public/Core/RootCameraNodeObserver.h
deleted file mode 100644
