Cameras: aspect ratio axis constraint should be flipped at 50%, not integer-interpolated

diff --git a/Engine/Plugins/Cameras/GameplayCameras/Source/GameplayCameras/Private/Core/CameraPose.cpp b/Engine/Plugins/Cameras/GameplayCameras/Source/GameplayCameras/Private/Core/CameraPose.cpp
--- a/Engine/Plugins/Cameras/GameplayCameras/Source/GameplayCameras/Private/Core/CameraPose.cpp
+++ b/Engine/Plugins/Cameras/GameplayCameras/Source/GameplayCameras/Private/Core/CameraPose.cpp
@@ -241,7 +241,7 @@ void FCameraPose::InternalLerpChanged(const FCameraPose& ToPose, float Factor, c
 
 		UE_CAMERA_POSE_FOR_TRANSFORM_PROPERTIES()
 		UE_CAMERA_POSE_FOR_INTERPOLABLE_PROPERTIES()
-		UE_CAMERA_POSE_FOR_BOOL_PROPERTIES()
+		UE_CAMERA_POSE_FOR_FLIPPING_PROPERTIES()
 
 #undef UE_CAMERA_POSE_FOR_PROPERTY
 
@@ -324,7 +324,7 @@ void FCameraPose::InternalLerpChanged(const FCameraPose& ToPose, float Factor, c
 			}
 		}
 
-		// Last, do booleans, which just flip their value once we reach 50% interpolation.
+		// Last, do booleans and other properties that just flip their value once we reach 50% interpolation.
 
 #define UE_CAMERA_POSE_FOR_PROPERTY(PropType, PropName)\
 		if ((!bInvertMask && InMask.PropName) || (bInvertMask && !InMask.PropName))\
@@ -337,7 +337,7 @@ void FCameraPose::InternalLerpChanged(const FCameraPose& ToPose, float Factor, c
 			OutMask.PropName = true;\
 		}
 
-		UE_CAMERA_POSE_FOR_BOOL_PROPERTIES()
+		UE_CAMERA_POSE_FOR_FLIPPING_PROPERTIES()
 
 #undef UE_CAMERA_POSE_FOR_PROPERTY
 
diff --git a/Engine/Plugins/Cameras/GameplayCameras/Source/GameplayCameras/Public/Core/CameraPose.h b/Engine/Plugins/Cameras/GameplayCameras/Source/GameplayCameras/Public/Core/CameraPose.h
--- a/Engine/Plugins/Cameras/GameplayCameras/Source/GameplayCameras/Public/Core/CameraPose.h
+++ b/Engine/Plugins/Cameras/GameplayCameras/Source/GameplayCameras/Public/Core/CameraPose.h
@@ -30,18 +30,15 @@ enum EAspectRatioAxisConstraint : int;
 	UE_CAMERA_POSE_FOR_PROPERTY(float, FieldOfView)\
 	UE_CAMERA_POSE_FOR_PROPERTY(float, FocalLength)
 
-#define UE_CAMERA_POSE_FOR_BOOL_PROPERTIES()\
-	UE_CAMERA_POSE_FOR_PROPERTY(bool, ConstrainAspectRatio)
-
-#define UE_CAMERA_POSE_FOR_INTEGER_PROPERTIES()\
+#define UE_CAMERA_POSE_FOR_FLIPPING_PROPERTIES()\
+	UE_CAMERA_POSE_FOR_PROPERTY(bool, ConstrainAspectRatio)\
 	UE_CAMERA_POSE_FOR_PROPERTY(EAspectRatioAxisConstraint, AspectRatioAxisConstraint)
 
 #define UE_CAMERA_POSE_FOR_ALL_PROPERTIES()\
 	UE_CAMERA_POSE_FOR_TRANSFORM_PROPERTIES()\
 	UE_CAMERA_POSE_FOR_INTERPOLABLE_PROPERTIES()\
 	UE_CAMERA_POSE_FOR_FOV_PROPERTIES()\
-	UE_CAMERA_POSE_FOR_BOOL_PROPERTIES()\
-	UE_CAMERA_POSE_FOR_INTEGER_PROPERTIES()
+	UE_CAMERA_POSE_FOR_FLIPPING_PROPERTIES()
 
 /**
  * Boolean flags for each of the properties inside FCameraPose.
