Cameras: remove registered settings from runtime module

diff --git a/Engine/Plugins/Cameras/GameplayCameras/Source/GameplayCameras/GameplayCameras.Build.cs b/Engine/Plugins/Cameras/GameplayCameras/Source/GameplayCameras/GameplayCameras.Build.cs
--- a/Engine/Plugins/Cameras/GameplayCameras/Source/GameplayCameras/GameplayCameras.Build.cs
+++ b/Engine/Plugins/Cameras/GameplayCameras/Source/GameplayCameras/GameplayCameras.Build.cs
@@ -19,7 +19,6 @@ public class GameplayCameras : ModuleRules
 				"HeadMountedDisplay",
 				"MovieScene",
 				"MovieSceneTracks",
-				"Settings",
 				"TemplateSequence",
 				"TraceLog"
 			}
diff --git a/Engine/Plugins/Cameras/GameplayCameras/Source/GameplayCameras/Private/GameplayCamerasModule.cpp b/Engine/Plugins/Cameras/GameplayCameras/Source/GameplayCameras/Private/GameplayCamerasModule.cpp
--- a/Engine/Plugins/Cameras/GameplayCameras/Source/GameplayCameras/Private/GameplayCamerasModule.cpp
+++ b/Engine/Plugins/Cameras/GameplayCameras/Source/GameplayCameras/Private/GameplayCamerasModule.cpp
@@ -8,8 +8,6 @@
 #include "Debug/CameraDebugColors.h"
 #include "Features/IModularFeatures.h"
 #include "GameplayCameras.h"
-#include "GameplayCamerasSettings.h"
-#include "ISettingsModule.h"
 #include "Logging/MessageLog.h"
 #include "Modules/ModuleManager.h"
 
@@ -29,8 +27,6 @@ public:
 	// IModuleInterface interface
 	virtual void StartupModule() override
 	{
-		RegisterSettings();
-
 		CameraModularFeature = MakeShared<FCameraModularFeature>();
 		if (CameraModularFeature.IsValid())
 		{
@@ -44,8 +40,6 @@ public:
 
 	virtual void ShutdownModule() override
 	{
-		UnregisterSettings();
-
 		if (CameraModularFeature.IsValid())
 		{
 			IModularFeatures::Get().UnregisterModularFeature(ICameraModularFeature::GetModularFeatureName(), CameraModularFeature.Get());
@@ -68,32 +62,6 @@ public:
 	}
 #endif
 
-private:
-
-	void RegisterSettings()
-	{
-		ISettingsModule* SettingsModule = FModuleManager::GetModulePtr<ISettingsModule>("Settings");
-
-		if (SettingsModule != nullptr)
-		{
-			SettingsModule->RegisterSettings("Project", "Plugins", "Gameplay Cameras",
-				LOCTEXT("GameplayCamerasProjectSettingsName", "Gameplay Cameras"),
-				LOCTEXT("GameplayCamerasProjectSettingsDescription", "Configure gameplay cameras."),
-				GetMutableDefault<UGameplayCamerasSettings>()
-			);
-		}
-	}
-
-	void UnregisterSettings()
-	{
-		ISettingsModule* SettingsModule = FModuleManager::GetModulePtr<ISettingsModule>("Settings");
-
-		if (SettingsModule != nullptr)
-		{
-			SettingsModule->UnregisterSettings("Project", "Plugins", "Gameplay Cameras");
-		}
-	}
-
 private:
 
 	class FCameraModularFeature : public ICameraModularFeature
