Cameras: add utility method to camera pose to get the target at a given distance

diff --git a/Engine/Plugins/Cameras/GameplayCameras/Source/GameplayCameras/Private/Core/CameraPose.cpp b/Engine/Plugins/Cameras/GameplayCameras/Source/GameplayCameras/Private/Core/CameraPose.cpp
--- a/Engine/Plugins/Cameras/GameplayCameras/Source/GameplayCameras/Private/Core/CameraPose.cpp
+++ b/Engine/Plugins/Cameras/GameplayCameras/Source/GameplayCameras/Private/Core/CameraPose.cpp
@@ -172,6 +172,11 @@ FVector3d FCameraPose::GetTarget() const
 	return Location + TargetDistance * GetAimDir();
 }
 
+FVector3d FCameraPose::GetTarget(double InTargetDistance) const
+{
+	return Location + InTargetDistance * GetAimDir();
+}
+
 void FCameraPose::OverrideAll(const FCameraPose& OtherPose)
 {
 	InternalOverrideChanged(OtherPose, false);
diff --git a/Engine/Plugins/Cameras/GameplayCameras/Source/GameplayCameras/Public/Core/CameraPose.h b/Engine/Plugins/Cameras/GameplayCameras/Source/GameplayCameras/Public/Core/CameraPose.h
--- a/Engine/Plugins/Cameras/GameplayCameras/Source/GameplayCameras/Public/Core/CameraPose.h
+++ b/Engine/Plugins/Cameras/GameplayCameras/Source/GameplayCameras/Public/Core/CameraPose.h
@@ -158,6 +158,9 @@ public:
 	/** Gets the location of the camera's target. */
 	FVector3d GetTarget() const;
 
+	/** Gets the location of the camera's target given a specific distance. */
+	FVector3d GetTarget(double InTargetDistance) const;
+
 public:
 
 	// Interpolation
