Cameras: add a ray-cast node to place the camera target in the world

diff --git a/Engine/Plugins/Cameras/GameplayCameras/Source/GameplayCameras/Private/Debug/CameraDebugRenderer.cpp b/Engine/Plugins/Cameras/GameplayCameras/Source/GameplayCameras/Private/Debug/CameraDebugRenderer.cpp
--- a/Engine/Plugins/Cameras/GameplayCameras/Source/GameplayCameras/Private/Debug/CameraDebugRenderer.cpp
+++ b/Engine/Plugins/Cameras/GameplayCameras/Source/GameplayCameras/Private/Debug/CameraDebugRenderer.cpp
@@ -268,6 +268,14 @@ void FCameraDebugRenderer::DrawLine(const FVector3d& Start, const FVector3d& End
 	}
 }
 
+void FCameraDebugRenderer::DrawSphere(const FVector3d& Center, float Radius, int32 Segments, const FLinearColor& LineColor, float LineThickness)
+{
+	if (ULineBatchComponent* LineBatcher = GetDebugLineBatcher())
+	{
+		LineBatcher->DrawSphere(Center, Radius, Segments, LineColor, 0.f, SDPG_Foreground, LineThickness);
+	}
+}
+
 ULineBatchComponent* FCameraDebugRenderer::GetDebugLineBatcher() const
 {
 	return World ? World->ForegroundLineBatcher : nullptr;
diff --git a/Engine/Plugins/Cameras/GameplayCameras/Source/GameplayCameras/Private/Nodes/Common/TargetRayCastCameraNode.cpp b/Engine/Plugins/Cameras/GameplayCameras/Source/GameplayCameras/Private/Nodes/Common/TargetRayCastCameraNode.cpp
new file mode 100644
--- /dev/null
+++ b/Engine/Plugins/Cameras/GameplayCameras/Source/GameplayCameras/Private/Nodes/Common/TargetRayCastCameraNode.cpp
@@ -0,0 +1,169 @@
+// Copyright Epic Games, Inc. All Rights Reserved.
+
+#include "Nodes/Common/TargetRayCastCameraNode.h"
+
+#include "CollisionQueryParams.h"
+#include "Core/CameraEvaluationContext.h"
+#include "Core/CameraNodeEvaluator.h"
+#include "Debug/CameraDebugBlock.h"
+#include "Debug/CameraDebugBlockBuilder.h"
+#include "Debug/CameraDebugRenderer.h"
+#include "Engine/HitResult.h"
+#include "Engine/World.h"
+#include "GameFramework/Pawn.h"
+#include "GameFramework/PlayerController.h"
+#include "WorldCollision.h"
+
+#include UE_INLINE_GENERATED_CPP_BY_NAME(TargetRayCastCameraNode)
+
+namespace UE::Cameras
+{
+
+float GTargetRayCastLength = 100000.0;
+static FAutoConsoleVariableRef CVarTargetRayCastLength(
+	TEXT("GameplayCameras.TargetRayCastLength"),
+	GTargetRayCastLength,
+	TEXT("(Default: 100000cm. Sets the length of the line trace test that determines the camera's target distance."));
+
+class FTargetRayCastCameraNodeEvaluator : public FCameraNodeEvaluator
+{
+	UE_DECLARE_CAMERA_NODE_EVALUATOR(GAMEPLAYCAMERAS_API, FTargetRayCastCameraNodeEvaluator)
+
+protected:
+
+	// FCameraNodeEvaluator interface.
+	virtual void OnRun(const FCameraNodeEvaluationParams& Params, FCameraNodeEvaluationResult& OutResult) override;
+#if UE_GAMEPLAY_CAMERAS_DEBUG
+	virtual void OnBuildDebugBlocks(const FCameraDebugBlockBuildParams& Params, FCameraDebugBlockBuilder& Builder) override;
+#endif  // UE_GAMEPLAY_CAMERAS_DEBUG
+
+private:
+
+	void RunLineTrace(UWorld* World, APlayerController* PlayerController, FCameraNodeEvaluationResult& OutResult);
+
+private:
+
+	double LastHitResultDistance = 1000.0; // Same as the default CameraPose target distance.
+
+#if UE_GAMEPLAY_CAMERAS_DEBUG
+	FVector3d LastHitResultLocation;
+	FString LastHitResultName;
+	bool bGotLastHitResult;
+#endif  // UE_GAMEPLAY_CAMERAS_DEBUG
+};
+
+UE_DEFINE_CAMERA_NODE_EVALUATOR(FTargetRayCastCameraNodeEvaluator)
+
+UE_DECLARE_CAMERA_DEBUG_BLOCK_START(GAMEPLAYCAMERAS_API, FTargetRayCastCameraDebugBlock)
+	UE_DECLARE_CAMERA_DEBUG_BLOCK_FIELD(FVector3d, HitResultLocation);
+	UE_DECLARE_CAMERA_DEBUG_BLOCK_FIELD(double, HitResultDistance);
+	UE_DECLARE_CAMERA_DEBUG_BLOCK_FIELD(FString, HitResultName);
+	UE_DECLARE_CAMERA_DEBUG_BLOCK_FIELD(bool, bGotHitResult);
+UE_DECLARE_CAMERA_DEBUG_BLOCK_END()
+
+UE_DEFINE_CAMERA_DEBUG_BLOCK_WITH_FIELDS(FTargetRayCastCameraDebugBlock)
+
+void FTargetRayCastCameraNodeEvaluator::OnRun(const FCameraNodeEvaluationParams& Params, FCameraNodeEvaluationResult& OutResult)
+{
+	if (!ensure(Params.EvaluationContext))
+	{
+		return;
+	}
+
+	if (Params.EvaluationType != ECameraNodeEvaluationType::Standard)
+	{
+		// Don't run actual ray-casts during IK/stateless updates.
+		OutResult.CameraPose.SetTargetDistance(LastHitResultDistance);
+		return;
+	}
+
+	UWorld* World = Params.EvaluationContext->GetWorld();
+	APlayerController* PlayerController = Params.EvaluationContext->GetPlayerController();
+	if (!World || !PlayerController)
+	{
+		return;
+	}
+
+	RunLineTrace(World, PlayerController, OutResult);
+}
+
+void FTargetRayCastCameraNodeEvaluator::RunLineTrace(UWorld* World, APlayerController* PlayerController, FCameraNodeEvaluationResult& OutResult)
+{
+	static FName LineTraceTag(TEXT("TargetRayCast"));
+	static FName LineTraceOwnerTag(TEXT("TargetRayCastCameraNode"));
+
+	APawn* Pawn = PlayerController->GetPawn();
+
+	FCameraPose& CameraPose = OutResult.CameraPose;
+	const FVector3d RayStart = CameraPose.GetLocation();
+	const FVector3d RayEnd = RayStart + CameraPose.GetAimDir() * GTargetRayCastLength;
+
+	FHitResult HitResult;
+
+	// Ignore the player pawn by default.
+	FCollisionQueryParams QueryParams(SCENE_QUERY_STAT(TargetLineTrace), false, Pawn);
+	QueryParams.TraceTag = LineTraceTag;
+	QueryParams.OwnerTag = LineTraceOwnerTag;
+
+	const bool bGotHit = World->LineTraceSingleByChannel(
+			HitResult, 
+			RayStart, RayEnd,
+			ECollisionChannel::ECC_Camera,
+			QueryParams);
+
+#if UE_GAMEPLAY_CAMERAS_DEBUG
+	bGotLastHitResult = bGotHit;
+#endif  // UE_GAMEPLAY_CAMERAS_DEBUG
+
+	if (bGotHit)
+	{
+		const double TargetDistance = (HitResult.ImpactPoint - RayStart).Length();
+		CameraPose.SetTargetDistance(TargetDistance);
+		LastHitResultDistance = TargetDistance;
+
+#if UE_GAMEPLAY_CAMERAS_DEBUG
+		LastHitResultLocation = HitResult.ImpactPoint;
+		LastHitResultName = GetNameSafe(HitResult.GetActor());
+#endif  // UE_GAMEPLAY_CAMERAS_DEBUG
+	}
+	else
+	{
+		CameraPose.SetTargetDistance(GTargetRayCastLength);
+	}
+}
+
+#if UE_GAMEPLAY_CAMERAS_DEBUG
+
+void FTargetRayCastCameraNodeEvaluator::OnBuildDebugBlocks(const FCameraDebugBlockBuildParams& Params, FCameraDebugBlockBuilder& Builder)
+{
+	FTargetRayCastCameraDebugBlock& DebugBlock = Builder.AttachDebugBlock<FTargetRayCastCameraDebugBlock>();
+	DebugBlock.HitResultLocation = LastHitResultLocation;
+	DebugBlock.HitResultDistance = LastHitResultDistance;
+	DebugBlock.HitResultName = LastHitResultName;
+	DebugBlock.bGotHitResult = bGotLastHitResult;
+}
+
+void FTargetRayCastCameraDebugBlock::OnDebugDraw(const FCameraDebugBlockDrawParams& Params, FCameraDebugRenderer& Renderer)
+{
+	if (bGotHitResult)
+	{
+		Renderer.AddText(TEXT("hit '%s', distance %.3f"), *HitResultName, HitResultDistance);
+	}
+	else
+	{
+		Renderer.AddText(TEXT("no hit"));
+	}
+
+	Renderer.DrawSphere(HitResultLocation, 1.f, 8, FLinearColor::Blue, 1.f);
+}
+
+#endif  // UE_GAMEPLAY_CAMERAS_DEBUG
+
+}  // namespace UE::Cameras
+
+FCameraNodeEvaluatorPtr UTargetRayCastCameraNode::OnBuildEvaluator(FCameraNodeEvaluatorBuilder& Builder) const
+{
+	using namespace UE::Cameras;
+	return Builder.BuildEvaluator<FTargetRayCastCameraNodeEvaluator>();
+}
+
diff --git a/Engine/Plugins/Cameras/GameplayCameras/Source/GameplayCameras/Public/Debug/CameraDebugRenderer.h b/Engine/Plugins/Cameras/GameplayCameras/Source/GameplayCameras/Public/Debug/CameraDebugRenderer.h
--- a/Engine/Plugins/Cameras/GameplayCameras/Source/GameplayCameras/Public/Debug/CameraDebugRenderer.h
+++ b/Engine/Plugins/Cameras/GameplayCameras/Source/GameplayCameras/Public/Debug/CameraDebugRenderer.h
@@ -76,7 +76,10 @@ public:
 	/** Draws a 2D circle. */
 	GAMEPLAYCAMERAS_API void Draw2DCircle(const FVector2D& Center, float Radius, const FLinearColor& LineColor, float LineThickness = 1.f, int32 NumSides = 0);
 
+	/** Draws a 3D line. */
 	GAMEPLAYCAMERAS_API void DrawLine(const FVector3d& Start, const FVector3d& End, const FLinearColor& LineColor, float LineThickness = 1.f);
+	/** Draws a 3D sphere. */
+	GAMEPLAYCAMERAS_API void DrawSphere(const FVector3d& Center, float Radius, int32 Segments, const FLinearColor& LineColor, float LineThickness);
 
 public:
 
diff --git a/Engine/Plugins/Cameras/GameplayCameras/Source/GameplayCameras/Public/Nodes/Common/TargetRayCastCameraNode.h b/Engine/Plugins/Cameras/GameplayCameras/Source/GameplayCameras/Public/Nodes/Common/TargetRayCastCameraNode.h
new file mode 100644
--- /dev/null
+++ b/Engine/Plugins/Cameras/GameplayCameras/Source/GameplayCameras/Public/Nodes/Common/TargetRayCastCameraNode.h
@@ -0,0 +1,23 @@
+// Copyright Epic Games, Inc. All Rights Reserved.
+
+#pragma once
+
+#include "Core/CameraNode.h"
+
+#include "TargetRayCastCameraNode.generated.h"
+
+/**
+ * A camera node that determines and sets the camera's target by running a ray-cast
+ * from the current camera position.
+ */
+UCLASS(MinimalAPI, meta=(CameraNodeCategories="Common,Target"))
+class UTargetRayCastCameraNode : public UCameraNode
+{
+	GENERATED_BODY()
+
+protected:
+
+	// UCameraNode interface.
+	virtual FCameraNodeEvaluatorPtr OnBuildEvaluator(FCameraNodeEvaluatorBuilder& Builder) const override;
+};
+
