Cameras: add orbit blend

diff --git a/Engine/Plugins/Cameras/GameplayCameras/Source/GameplayCameras/Private/Nodes/Blends/OrbitBlendCameraNode.cpp b/Engine/Plugins/Cameras/GameplayCameras/Source/GameplayCameras/Private/Nodes/Blends/OrbitBlendCameraNode.cpp
new file mode 100644
--- /dev/null
+++ b/Engine/Plugins/Cameras/GameplayCameras/Source/GameplayCameras/Private/Nodes/Blends/OrbitBlendCameraNode.cpp
@@ -0,0 +1,201 @@
+// Copyright Epic Games, Inc. All Rights Reserved.
+
+#include "Nodes/Blends/OrbitBlendCameraNode.h"
+
+#include UE_INLINE_GENERATED_CPP_BY_NAME(OrbitBlendCameraNode)
+
+namespace UE::Cameras
+{
+
+class FOrbitBlendCameraNodeEvaluator : public FBlendCameraNodeEvaluator
+{
+	UE_DECLARE_BLEND_CAMERA_NODE_EVALUATOR(GAMEPLAYCAMERAS_API, FOrbitBlendCameraNodeEvaluator)
+
+protected:
+
+	virtual void OnBuild(const FCameraNodeEvaluatorBuildParams& Params) override;
+	virtual FCameraNodeEvaluatorChildrenView OnGetChildren() override;
+	virtual void OnRun(const FCameraNodeEvaluationParams& Params, FCameraNodeEvaluationResult& OutResult) override;
+	virtual void OnBlendResults(const FCameraNodeBlendParams& Params, FCameraNodeBlendResult& OutResult) override;
+
+	static bool ClosestPoints(const FRay3d& A, const FRay3d& B, double& OutParameterA, double& OutParameterB);
+
+private:
+
+	FSimpleBlendCameraNodeEvaluator* DrivingBlendEvaluator = nullptr;
+};
+
+UE_DEFINE_BLEND_CAMERA_NODE_EVALUATOR(FOrbitBlendCameraNodeEvaluator)
+
+void FOrbitBlendCameraNodeEvaluator::OnBuild(const FCameraNodeEvaluatorBuildParams& Params)
+{
+	const UOrbitBlendCameraNode* OrbitBlend = GetCameraNodeAs<UOrbitBlendCameraNode>();
+	if (OrbitBlend->DrivingBlend)
+	{
+		DrivingBlendEvaluator = Params.BuildEvaluatorAs<FSimpleBlendCameraNodeEvaluator>(OrbitBlend->DrivingBlend);
+	}
+}
+
+FCameraNodeEvaluatorChildrenView FOrbitBlendCameraNodeEvaluator::OnGetChildren()
+{
+	return FCameraNodeEvaluatorChildrenView({ DrivingBlendEvaluator });
+}
+
+void FOrbitBlendCameraNodeEvaluator::OnRun(const FCameraNodeEvaluationParams& Params, FCameraNodeEvaluationResult& OutResult)
+{
+	if (DrivingBlendEvaluator)
+	{
+		DrivingBlendEvaluator->Run(Params, OutResult);
+	}
+}
+
+void FOrbitBlendCameraNodeEvaluator::OnBlendResults(const FCameraNodeBlendParams& Params, FCameraNodeBlendResult& OutResult)
+{
+	const FCameraNodeEvaluationResult& ChildResult(Params.ChildResult);
+	FCameraNodeEvaluationResult& BlendedResult(OutResult.BlendedResult);
+
+	// If we don't have a driving blend, just cut to the new camera rig.
+	if (!DrivingBlendEvaluator)
+	{
+		BlendedResult.CameraPose.OverrideAll(ChildResult.CameraPose);
+		BlendedResult.VariableTable.OverrideAll(ChildResult.VariableTable);
+		BlendedResult.bIsCameraCut = true;
+
+		OutResult.bIsBlendFinished = true;
+		OutResult.bIsBlendFull = true;
+
+		return;
+	}
+	
+	// Let our underlying blend do most of the blending, but overwrite the camera transform with our
+	// own blending algorithm.
+	//
+	// But first, remember a few things about the original camera poses.
+	const FRay3d FromAim = BlendedResult.CameraPose.GetAimRay();
+	const FRay3d ToAim = ChildResult.CameraPose.GetAimRay();
+
+	const FVector3d FromLocation = BlendedResult.CameraPose.GetLocation();
+	const FVector3d ToLocation = ChildResult.CameraPose.GetLocation();
+
+	// Run the underlying blend.
+	DrivingBlendEvaluator->BlendResults(Params, OutResult);
+
+	// If the blend reached 100%, we're done.
+	if (OutResult.bIsBlendFull)
+	{
+		return;
+	}
+
+	// Find the points on each line of sight that are the closest to each other. If successful,
+	// start blending around an interpolating mid-point between the two.
+	double FromClosestParam, ToClosestParam;
+	const bool bFoundClosestPoints = ClosestPoints(FromAim, ToAim, FromClosestParam, ToClosestParam);
+	if (bFoundClosestPoints)
+	{
+		const float Factor = DrivingBlendEvaluator->GetBlendFactor();
+
+		const FVector3d BlendedLocation = FMath::Lerp(FromLocation, ToLocation, Factor);
+
+		// Rotate around a point that is interpolating from the first line of sight to the other
+		// line of sight.
+		const FVector3d FromOrbitCenter = FromAim.PointAt(FromClosestParam);
+		const FVector3d ToOrbitCenter = ToAim.PointAt(ToClosestParam);
+		const FVector3d BlendedOrbitCenter = FMath::Lerp(FromOrbitCenter, ToOrbitCenter, Factor);
+
+		const FVector3d BlendedAimDir = (BlendedOrbitCenter - BlendedLocation).GetUnsafeNormal();
+		const FRotator3d BlendedRotation = BlendedAimDir.ToOrientationRotator();
+
+		const double FromOrbitCenterDistance = FVector3d::Distance(FromLocation, FromOrbitCenter);
+		const double ToOrbitCenterDistance = FVector3d::Distance(ToLocation, ToOrbitCenter);
+		const double BlendedTargetDistance = FMath::Lerp(FromOrbitCenterDistance, ToOrbitCenterDistance, Factor);
+
+		// So, instead of interpolating between the two positions and letting the target move
+		// forwards or backwards as TargetDistance interpolates, we do the opposite: we "anchor"
+		// the target at the orbit center, and push or pull the position based on the interpolated
+		// TargetDistance.
+		const FRay3d BlendedReverseAim(BlendedOrbitCenter, -BlendedAimDir, true);
+		const FVector3d OrbitingLocation = BlendedReverseAim.PointAt(BlendedTargetDistance);
+
+		BlendedResult.CameraPose.SetLocation(OrbitingLocation);
+		BlendedResult.CameraPose.SetRotation(BlendedRotation);
+	}
+}
+
+bool FOrbitBlendCameraNodeEvaluator::ClosestPoints(const FRay3d& A, const FRay3d& B, double& OutParameterA, double& OutParameterB)
+{
+	// The points closest to each other on rays A and B are named T1 and T1. They are such that the
+	// vector T1T2 is orthogonal to both A and B's direction vectors. So the dot products should be zero:
+	//
+	//    (T2 - T1).A = 0
+	//    (T2 - T1).B = 0
+	//
+	// We can define T1 and T2 using the parametric equations of the rays:
+	//
+	//    T1 = O1 + x1*A 
+	//    T2 = O2 + x2*B
+	//
+	// Where O1 and O2 are the origin points of the rays, and x1 and x2 are the linear parameters.
+	//
+	// So we can rewrite our conditions:
+	//
+	//    (O2 + x2*B - O1 - x1*A).A = 0
+	//    (O2 + x2*B - O1 - x1*A).B = 0
+	//
+	//    (O2 - O1).A + x2*(B.A) - x1*(A.A) = 0
+	//    (O2 - O1).B + x2*(B.B) - x1*(A.B) = 0
+	//
+	// A and B are unit vectors so A.A and B.B equal 1.
+	// Also, let's note D = (O2 - O1) and C = (A.B)
+	//
+	//    D.A + x2*C - x1 = 0
+	//    D.B + x2 - x1*C = 0
+	//
+	// Let's solve for x2:
+	//
+	//    x1 = D.A + x2*C
+	//    D.B + x2 - (D.A + x2*C)*C = 0
+	//    D.B + x2 - (D.A)*C - x2*C*C = 0
+	//	  x2 = ((D.A)*C - D.B) / (1 - C*C)
+	//
+	// And x1:
+	//
+	//    x2 = x1*C - D.B
+	//    D.A + (x1*C - D.B)*C - x1 = 0
+	//	  D.A + x1*C*C - (D.B)*C - x1 = 0
+	//	  (D.A - (D.B)*C) / (1 - C*C) = x1
+	//
+	// We can see that there is no solution if C*C == 1, which corresponds to parallel rays.
+	//
+	
+	const FVector3d D = (B.Origin - A.Origin);
+
+	const FVector3d DirA = A.Direction.GetUnsafeNormal();
+	const FVector3d DirB = B.Direction.GetUnsafeNormal();
+	const double C = FVector3d::DotProduct(DirA, DirB);
+
+	const double OneMinusCC = (1.0 - C * C);
+
+	if (OneMinusCC != 0.0)
+	{
+		OutParameterA = (D.Dot(DirA) - (D.Dot(DirB)*C)) / OneMinusCC;
+		OutParameterB = ((D.Dot(DirA)*C) - D.Dot(DirB)) / OneMinusCC;
+
+		return true;
+	}
+
+	return false;
+}
+
+}  // namespace UE::Cameras
+
+FCameraNodeChildrenView UOrbitBlendCameraNode::OnGetChildren()
+{
+	return FCameraNodeChildrenView({ DrivingBlend });
+}
+
+FCameraNodeEvaluatorPtr UOrbitBlendCameraNode::OnBuildEvaluator(FCameraNodeEvaluatorBuilder& Builder) const
+{
+	using namespace UE::Cameras;
+	return Builder.BuildEvaluator<FOrbitBlendCameraNodeEvaluator>();
+}
+
diff --git a/Engine/Plugins/Cameras/GameplayCameras/Source/GameplayCameras/Public/Nodes/Blends/OrbitBlendCameraNode.h b/Engine/Plugins/Cameras/GameplayCameras/Source/GameplayCameras/Public/Nodes/Blends/OrbitBlendCameraNode.h
new file mode 100644
--- /dev/null
+++ b/Engine/Plugins/Cameras/GameplayCameras/Source/GameplayCameras/Public/Nodes/Blends/OrbitBlendCameraNode.h
@@ -0,0 +1,29 @@
+// Copyright Epic Games, Inc. All Rights Reserved.
+
+#pragma once
+
+#include "Core/BlendCameraNode.h"
+#include "Nodes/Blends/SimpleBlendCameraNode.h"
+
+#include "OrbitBlendCameraNode.generated.h"
+
+/**
+ * Orbit blend node.
+ */
+UCLASS(MinimalAPI)
+class UOrbitBlendCameraNode : public UBlendCameraNode
+{
+	GENERATED_BODY()
+
+public:
+
+	UPROPERTY()
+	TObjectPtr<USimpleBlendCameraNode> DrivingBlend;
+
+protected:
+
+	// UCameraNode interface.
+	virtual FCameraNodeChildrenView OnGetChildren() override;
+	virtual FCameraNodeEvaluatorPtr OnBuildEvaluator(FCameraNodeEvaluatorBuilder& Builder) const override;
+};
+
