Cameras: use a dynamic player controller component to handle globally activated camera rigs

diff --git a/Engine/Plugins/Cameras/GameplayCameras/Source/GameplayCameras/Private/GameFramework/ActivateCameraRigFunctions.cpp b/Engine/Plugins/Cameras/GameplayCameras/Source/GameplayCameras/Private/GameFramework/ActivateCameraRigFunctions.cpp
--- a/Engine/Plugins/Cameras/GameplayCameras/Source/GameplayCameras/Private/GameFramework/ActivateCameraRigFunctions.cpp
+++ b/Engine/Plugins/Cameras/GameplayCameras/Source/GameplayCameras/Private/GameFramework/ActivateCameraRigFunctions.cpp
@@ -16,8 +16,6 @@
 
 #include UE_INLINE_GENERATED_CPP_BY_NAME(ActivateCameraRigFunctions)
 
-TSharedPtr<UE::Cameras::FCameraEvaluationContext> UActivateCameraRigFunctions::GlobalContext;
-
 void UActivateCameraRigFunctions::ActivateBaseCameraRig(UObject* WorldContextObject, APlayerController* PlayerController, UCameraRigAsset* CameraRig)
 {
 	ActivateCameraRig(WorldContextObject, PlayerController, CameraRig, ECameraRigLayer::Base);
@@ -43,42 +41,115 @@ void UActivateCameraRigFunctions::ActivateCameraRig(UObject* WorldContextObject,
 		return;
 	}
 
-	if (FCameraSystemEvaluator* SystemEvaluator = FindCameraSystemEvaluator(PlayerController))
+	// Register our evaluation component on the given player controller, if it's not there already.
+	UControllerGameplayCameraEvaluationComponent* CameraEvaluationComponent = PlayerController->FindComponentByClass<UControllerGameplayCameraEvaluationComponent>();
+	if (!CameraEvaluationComponent)
 	{
-		FActivateCameraRigParams Params;
-		Params.CameraRig = CameraRig;
-		Params.EvaluationContext = EnsureGlobalContext(WorldContextObject, PlayerController);
-		Params.Evaluator = SystemEvaluator;
-		Params.Layer = EvaluationLayer;
-		SystemEvaluator->GetRootNodeEvaluator()->ActivateCameraRig(Params);
+		CameraEvaluationComponent = NewObject<UControllerGameplayCameraEvaluationComponent>(
+				PlayerController, TEXT("ControllerGameplayCameraEvaluationComponent"), RF_Transient);
+		CameraEvaluationComponent->RegisterComponent();
 	}
+
+	// Activate the camera rig.
+	CameraEvaluationComponent->ActivateCameraRig(CameraRig, EvaluationLayer);
 }
 
-UE::Cameras::FCameraSystemEvaluator* UActivateCameraRigFunctions::FindCameraSystemEvaluator(APlayerController* PlayerController)
+UControllerGameplayCameraEvaluationComponent::UControllerGameplayCameraEvaluationComponent(const FObjectInitializer& ObjectInitializer)
+	: Super(ObjectInitializer)
 {
-	if (PlayerController && PlayerController->PlayerCameraManager)
+	bAutoActivate = true;
+}
+
+void UControllerGameplayCameraEvaluationComponent::ActivateCameraRig(UCameraRigAsset* CameraRig, ECameraRigLayer EvaluationLayer)
+{
+	FCameraRigInfo NewCameraRigInfo;
+	NewCameraRigInfo.CameraRig = CameraRig;
+	NewCameraRigInfo.EvaluationLayer = EvaluationLayer;
+	NewCameraRigInfo.bActivated = false;
+	CameraRigInfos.Add(NewCameraRigInfo);
+
+	if (IsActive())
 	{
-		AActor* ViewTarget = PlayerController->PlayerCameraManager->GetViewTarget();
-		if (AGameplayCameraSystemActor* SystemActor = Cast<AGameplayCameraSystemActor>(ViewTarget))
+		ActivateCameraRigs();
+	}
+}
+
+void UControllerGameplayCameraEvaluationComponent::BeginPlay()
+{
+	Super::BeginPlay();
+
+	ActivateCameraRigs();
+}
+
+void UControllerGameplayCameraEvaluationComponent::EndPlay(const EEndPlayReason::Type EndPlayReason)
+{
+	CameraRigInfos.Reset();
+	EvaluationContext.Reset();
+
+	Super::EndPlay(EndPlayReason);
+}
+
+void UControllerGameplayCameraEvaluationComponent::ActivateCameraRigs()
+{
+	using namespace UE::Cameras;
+
+	APlayerController* PlayerController = GetOwner<APlayerController>();
+	FCameraSystemEvaluator* SystemEvaluator = FindCameraSystemEvaluator(PlayerController);
+	if (!SystemEvaluator)
+	{
+		return;
+	}
+	
+	EnsureEvaluationContext();
+	if (!EvaluationContext)
+	{
+		return;
+	}
+
+	FRootCameraNodeEvaluator* RootNodeEvaluator = SystemEvaluator->GetRootNodeEvaluator();
+
+	for (FCameraRigInfo& CameraRigInfo : CameraRigInfos)
+	{
+		if (!CameraRigInfo.bActivated)
 		{
-			return SystemActor->GetCameraSystemComponent()->GetCameraSystemEvaluator().Get();
+			FActivateCameraRigParams Params;
+			Params.CameraRig = CameraRigInfo.CameraRig;
+			Params.EvaluationContext = EvaluationContext;
+			Params.Evaluator = SystemEvaluator;
+			Params.Layer = CameraRigInfo.EvaluationLayer;
+			RootNodeEvaluator->ActivateCameraRig(Params);
+
+			CameraRigInfo.bActivated = true;
 		}
 	}
-	return nullptr;
 }
 
-TSharedPtr<UE::Cameras::FCameraEvaluationContext> UActivateCameraRigFunctions::EnsureGlobalContext(UObject* WorldContextObject, APlayerController* PlayerController)
+void UControllerGameplayCameraEvaluationComponent::EnsureEvaluationContext()
 {
 	using namespace UE::Cameras;
 
-	if (!GlobalContext.IsValid())
+	if (!EvaluationContext.IsValid())
 	{
+		APlayerController* PlayerController = GetOwner<APlayerController>();
+
 		FCameraEvaluationContextInitializeParams InitParams;
-		InitParams.Owner = WorldContextObject;
+		InitParams.Owner = this;
 		InitParams.PlayerController = PlayerController;
-		GlobalContext = MakeShared<FCameraEvaluationContext>(InitParams);
-		GlobalContext->GetInitialResult().bIsValid = true;	
+		EvaluationContext = MakeShared<FCameraEvaluationContext>(InitParams);
+		EvaluationContext->GetInitialResult().bIsValid = true;	
 	}
-	return GlobalContext;
+}
+
+UE::Cameras::FCameraSystemEvaluator* UControllerGameplayCameraEvaluationComponent::FindCameraSystemEvaluator(APlayerController* PlayerController)
+{
+	if (PlayerController && PlayerController->PlayerCameraManager)
+	{
+		AActor* ViewTarget = PlayerController->PlayerCameraManager->GetViewTarget();
+		if (AGameplayCameraSystemActor* SystemActor = Cast<AGameplayCameraSystemActor>(ViewTarget))
+		{
+			return SystemActor->GetCameraSystemComponent()->GetCameraSystemEvaluator().Get();
+		}
+	}
+	return nullptr;
 }
 
diff --git a/Engine/Plugins/Cameras/GameplayCameras/Source/GameplayCameras/Public/GameFramework/ActivateCameraRigFunctions.h b/Engine/Plugins/Cameras/GameplayCameras/Source/GameplayCameras/Public/GameFramework/ActivateCameraRigFunctions.h
--- a/Engine/Plugins/Cameras/GameplayCameras/Source/GameplayCameras/Public/GameFramework/ActivateCameraRigFunctions.h
+++ b/Engine/Plugins/Cameras/GameplayCameras/Source/GameplayCameras/Public/GameFramework/ActivateCameraRigFunctions.h
@@ -2,6 +2,7 @@
 
 #pragma once
 
+#include "Components/ActorComponent.h"
 #include "Core/CameraEvaluationContext.h"
 #include "Core/RootCameraNode.h"
 #include "CoreTypes.h"
@@ -50,14 +51,47 @@ public:
 	/** Activates the given camera rig in the given layer. */
 	UFUNCTION(BlueprintCallable, Category="Camera", meta = (WorldContext = "WorldContextObject"))
 	static void ActivateCameraRig(UObject* WorldContextObject, APlayerController* PlayerController, UCameraRigAsset* CameraRig, ECameraRigLayer EvaluationLayer);
+};
+
+/**
+ * A component, attached to a player controller, that can run camera rigs activated from
+ * a global place like the Blueprint functions inside UActivateCameraRigFunctions.
+ */
+UCLASS(Hidden, MinimalAPI)
+class UControllerGameplayCameraEvaluationComponent : public UActorComponent
+{
+	GENERATED_BODY()
+
+public:
+
+	UControllerGameplayCameraEvaluationComponent(const FObjectInitializer& ObjectInitializer);
+
+	/** Activates a new camera rig. */
+	void ActivateCameraRig(UCameraRigAsset* CameraRig, ECameraRigLayer EvaluationLayer);
+
+public:
+
+	// UActorComponent interface.
+	virtual void BeginPlay() override;
+	virtual void EndPlay(const EEndPlayReason::Type EndPlayReason) override;
 
 private:
 
+	void ActivateCameraRigs();
+	void EnsureEvaluationContext();
+
 	static UE::Cameras::FCameraSystemEvaluator* FindCameraSystemEvaluator(APlayerController* PlayerController);
-	static TSharedPtr<UE::Cameras::FCameraEvaluationContext> EnsureGlobalContext(UObject* WorldContextObject, APlayerController* PlayerController);
 
 private:
 
-	static TSharedPtr<UE::Cameras::FCameraEvaluationContext> GlobalContext;
+	struct FCameraRigInfo
+	{
+		TObjectPtr<UCameraRigAsset> CameraRig;
+		ECameraRigLayer EvaluationLayer;
+		bool bActivated = false;
+	};
+
+	TArray<FCameraRigInfo> CameraRigInfos;
+	TSharedPtr<UE::Cameras::FCameraEvaluationContext> EvaluationContext;
 };
 
