Cameras: add option to set the player controller rotation to the camera system component's output camera rotation

diff --git a/Engine/Plugins/Cameras/GameplayCameras/Source/GameplayCameras/Private/GameFramework/GameplayCameraSystemComponent.cpp b/Engine/Plugins/Cameras/GameplayCameras/Source/GameplayCameras/Private/GameFramework/GameplayCameraSystemComponent.cpp
--- a/Engine/Plugins/Cameras/GameplayCameras/Source/GameplayCameras/Private/GameFramework/GameplayCameraSystemComponent.cpp
+++ b/Engine/Plugins/Cameras/GameplayCameras/Source/GameplayCameras/Private/GameFramework/GameplayCameraSystemComponent.cpp
@@ -43,6 +43,15 @@ void UGameplayCameraSystemComponent::GetCameraView(float DeltaTime, FMinimalView
 		UpdateParams.DeltaTime = DeltaTime;
 		Evaluator->Update(UpdateParams);
 
+		if (bSetPlayerControllerRotation && ActivatedForPlayerIndex != INDEX_NONE)
+		{
+			APlayerController* PC = UGameplayStatics::GetPlayerController(this, ActivatedForPlayerIndex);
+			if (PC)
+			{
+				PC->SetControlRotation(Evaluator->GetEvaluatedResult().CameraPose.GetRotation());
+			}
+		}
+
 		Evaluator->GetEvaluatedCameraView(DesiredView);
 	}
 }
diff --git a/Engine/Plugins/Cameras/GameplayCameras/Source/GameplayCameras/Public/GameFramework/GameplayCameraSystemComponent.h b/Engine/Plugins/Cameras/GameplayCameras/Source/GameplayCameras/Public/GameFramework/GameplayCameraSystemComponent.h
--- a/Engine/Plugins/Cameras/GameplayCameras/Source/GameplayCameras/Public/GameFramework/GameplayCameraSystemComponent.h
+++ b/Engine/Plugins/Cameras/GameplayCameras/Source/GameplayCameras/Public/GameFramework/GameplayCameraSystemComponent.h
@@ -89,6 +89,13 @@ public:
 	UPROPERTY(EditAnywhere, Category=Camera)
 	TEnumAsByte<EAutoReceiveInput::Type> AutoActivateForPlayer;
 
+	/**
+	 * If enabled, sets the evaluated camera orientation as the player controller rotation every frame.
+	 * This is set on the player controller that this component was activated for.
+	 */
+	UPROPERTY(EditAnywhere, Category=Camera)
+	bool bSetPlayerControllerRotation = true;
+
 private:
 	
 	TSharedPtr<FCameraSystemEvaluator> Evaluator;
