Cameras: fix ordering of operation in blend stack so that pre-blended parameters are correctly re-assigned to their respective tables.

diff --git a/Engine/Plugins/Cameras/GameplayCameras/Source/GameplayCameras/Private/Core/BlendStackCameraNode.cpp b/Engine/Plugins/Cameras/GameplayCameras/Source/GameplayCameras/Private/Core/BlendStackCameraNode.cpp
--- a/Engine/Plugins/Cameras/GameplayCameras/Source/GameplayCameras/Private/Core/BlendStackCameraNode.cpp
+++ b/Engine/Plugins/Cameras/GameplayCameras/Source/GameplayCameras/Private/Core/BlendStackCameraNode.cpp
@@ -303,8 +303,7 @@ void FBlendStackCameraNodeEvaluator::OnRun(const FCameraNodeEvaluationParams& Pa
 #endif  // UE_GAMEPLAY_CAMERAS_TRACE
 	}
 
-	// Setup the entries' variable tables for this frame, gather parameters to pre-blend,
-	// and evaluate blend nodes.
+	// Gather parameters to pre-blend, and evaluate blend nodes.
 	for (FValidEntry& ValidEntry : ValidEntries)
 	{
 		FCameraRigEntry& Entry(ValidEntry.Entry);
@@ -321,14 +320,6 @@ void FBlendStackCameraNodeEvaluator::OnRun(const FCameraNodeEvaluationParams& Pa
 
 		FCameraNodeEvaluationResult& CurResult(Entry.Result);
 
-		// Setup the variable table: take everything from the previous evaluation layer,
-		// and override it with the context's initial result.
-		CurResult.VariableTable.OverrideAll(OutResult.VariableTable);
-		CurResult.VariableTable.ClearAllWrittenThisFrameFlags();
-
-		const FCameraNodeEvaluationResult& ContextResult(CurContext->GetInitialResult());
-		CurResult.VariableTable.OverrideAll(ContextResult.VariableTable);
-
 		// Gather input parameters.
 		if (!Entry.bInputRunThisFrame)
 		{
@@ -401,10 +392,13 @@ void FBlendStackCameraNodeEvaluator::OnRun(const FCameraNodeEvaluationParams& Pa
 		// Start with the input given to us.
 		CurResult.CameraPose = OutResult.CameraPose;
 		CurResult.CameraPose.ClearAllChangedFlags();
+		CurResult.VariableTable.OverrideAll(OutResult.VariableTable);
+		CurResult.VariableTable.ClearAllWrittenThisFrameFlags();
 
 		// Override it with whatever the evaluation context has set on its result.
 		const FCameraNodeEvaluationResult& ContextResult(CurContext->GetInitialResult());
 		CurResult.CameraPose.OverrideChanged(ContextResult.CameraPose);
+		CurResult.VariableTable.OverrideAll(ContextResult.VariableTable);
 		CurResult.bIsCameraCut = OutResult.bIsCameraCut || ContextResult.bIsCameraCut;
 		CurResult.bIsValid = true;
 
