Cameras: properly initialize flags on dynamically added variables

diff --git a/Engine/Plugins/Cameras/GameplayCameras/Source/GameplayCameras/Private/Core/CameraVariableTable.cpp b/Engine/Plugins/Cameras/GameplayCameras/Source/GameplayCameras/Private/Core/CameraVariableTable.cpp
--- a/Engine/Plugins/Cameras/GameplayCameras/Source/GameplayCameras/Private/Core/CameraVariableTable.cpp
+++ b/Engine/Plugins/Cameras/GameplayCameras/Source/GameplayCameras/Private/Core/CameraVariableTable.cpp
@@ -153,7 +153,15 @@ void FCameraVariableTable::AddVariable(const FCameraVariableDefinition& Variable
 	NewEntry.ID = VariableDefinition.VariableID;
 	NewEntry.Type = VariableDefinition.VariableType;
 	NewEntry.Offset = VariablePtr - Memory;
-	NewEntry.Flags = VariableDefinition.bIsPrivate ? EEntryFlags::Private : EEntryFlags::None;
+	NewEntry.Flags = EEntryFlags::None;
+	if (VariableDefinition.bIsPrivate)
+	{
+		NewEntry.Flags |= EEntryFlags::Private;
+	}
+	if (VariableDefinition.bIsInput)
+	{
+		NewEntry.Flags |= EEntryFlags::Input;
+	}
 #if WITH_EDITORONLY_DATA
 	NewEntry.DebugName = VariableDefinition.VariableName;
 #endif
@@ -384,6 +392,7 @@ void FCameraVariableTable::InternalOverride(const FCameraVariableTable& OtherTab
 				FCameraVariableDefinition NewVariableDefinition;
 				NewVariableDefinition.VariableID = OtherEntry.ID;
 				NewVariableDefinition.VariableType = OtherEntry.Type;
+				NewVariableDefinition.bIsInput = EnumHasAllFlags(OtherEntry.Flags, EEntryFlags::Input);
 #if WITH_EDITORONLY_DATA
 				NewVariableDefinition.VariableName = OtherEntry.DebugName;
 #endif
@@ -493,6 +502,7 @@ UE_CAMERA_VARIABLE_FOR_ALL_TYPES()
 				FCameraVariableDefinition NewVariableDefinition;
 				NewVariableDefinition.VariableID = ToEntry.ID;
 				NewVariableDefinition.VariableType = ToEntry.Type;
+				NewVariableDefinition.bIsInput = EnumHasAllFlags(ToEntry.Flags, EEntryFlags::Input);
 #if WITH_EDITORONLY_DATA
 				NewVariableDefinition.VariableName = ToEntry.DebugName;
 #endif
