Cameras: tidy up camera rig picker's way of showing warnings/errors

diff --git a/Engine/Plugins/Cameras/GameplayCameras/Source/GameplayCamerasEditor/Private/Editors/SCameraRigPicker.cpp b/Engine/Plugins/Cameras/GameplayCameras/Source/GameplayCamerasEditor/Private/Editors/SCameraRigPicker.cpp
--- a/Engine/Plugins/Cameras/GameplayCameras/Source/GameplayCamerasEditor/Private/Editors/SCameraRigPicker.cpp
+++ b/Engine/Plugins/Cameras/GameplayCameras/Source/GameplayCamerasEditor/Private/Editors/SCameraRigPicker.cpp
@@ -126,31 +126,45 @@ void SCameraRigPicker::Construct(const FArguments& InArgs)
 		.OnSelectionChanged(this, &SCameraRigPicker::OnCameraRigListSelectionChanged)
 	];
 
+	// Number of items in the camera rig list.
+	TSharedPtr<SHorizontalBox> MessageBar;
+	LayoutBox->AddSlot()
+	.AutoHeight()
+	[
+		SAssignNew(MessageBar, SHorizontalBox)
+		+SHorizontalBox::Slot()
+		.FillWidth(1.f)
+		.VAlign(VAlign_Center)
+		.Padding(8, 5)
+		[
+			SNew(STextBlock)
+			.Text(this, &SCameraRigPicker::GetCameraRigCountText)
+		]
+	];
+
 	// Optional warning message.
 	if (!PickerConfig.WarningMessage.IsEmpty())
 	{
-		LayoutBox->AddSlot()
-		.AutoHeight()
+		MessageBar->InsertSlot(0)
+		.AutoWidth()
+		.VAlign(VAlign_Center)
 		[
-			SNew(STextBlock)
-			.ColorAndOpacity(AppStyle.GetSlateColor("Colors.Warning"))
-			.Margin(4.f)
-			.Text(PickerConfig.WarningMessage)
-			.AutoWrapText(true)
+			SNew(SImage)
+			.Image(FAppStyle::GetBrush("Icons.WarningWithColor"))
+			.ToolTipText(PickerConfig.WarningMessage)
 		];
 	}
 
 	// Optional error message.
 	if (!PickerConfig.ErrorMessage.IsEmpty())
 	{
-		LayoutBox->AddSlot()
-		.AutoHeight()
+		MessageBar->InsertSlot(0)
+		.AutoWidth()
+		.VAlign(VAlign_Center)
 		[
-			SNew(STextBlock)
-			.ColorAndOpacity(AppStyle.GetSlateColor("Colors.Error"))
-			.Margin(4.f)
-			.Text(PickerConfig.ErrorMessage)
-			.AutoWrapText(true)
+			SNew(SImage)
+			.Image(FAppStyle::GetBrush("Icons.ErrorWithColor"))
+			.ToolTipText(PickerConfig.ErrorMessage)
 		];
 	}
 	
@@ -386,6 +400,22 @@ void SCameraRigPicker::UpdateCameraRigFilteredItemsSource()
 	}
 }
 
+FText SCameraRigPicker::GetCameraRigCountText() const
+{
+	const int32 NumCameraRigs = CameraRigFilteredItemsSource.Num();
+
+	FText CountText = FText::GetEmpty();
+	if (NumCameraRigs == 1)
+	{
+		CountText = LOCTEXT("CameraRigCountTextSingular", "1 item");
+	}
+	else
+	{
+		CountText = FText::Format(LOCTEXT("CameraRigCountTextPlural", "{0} items"), FText::AsNumber(NumCameraRigs));
+	}
+	return CountText;
+}
+
 void SCameraRigPicker::GetEntryStrings(const UCameraRigAsset* InItem, TArray<FString>& OutStrings)
 {
 	OutStrings.Add(InItem->GetDisplayName());
diff --git a/Engine/Plugins/Cameras/GameplayCameras/Source/GameplayCamerasEditor/Private/Editors/SCameraRigPicker.h b/Engine/Plugins/Cameras/GameplayCameras/Source/GameplayCamerasEditor/Private/Editors/SCameraRigPicker.h
--- a/Engine/Plugins/Cameras/GameplayCameras/Source/GameplayCamerasEditor/Private/Editors/SCameraRigPicker.h
+++ b/Engine/Plugins/Cameras/GameplayCameras/Source/GameplayCamerasEditor/Private/Editors/SCameraRigPicker.h
@@ -44,6 +44,7 @@ private:
 	void OnCameraRigListSelectionChanged(UCameraRigAsset* Item, ESelectInfo::Type SelectInfo);
 	void UpdateCameraRigItemsSource();
 	void UpdateCameraRigFilteredItemsSource();
+	FText GetCameraRigCountText() const;
 
 	void GetEntryStrings(const UCameraRigAsset* InItem, TArray<FString>& OutStrings);
 	void OnSearchTextChanged(const FText& InFilterText);
