Cameras: add button to more easily rename a variable in the camera variable collection editor

diff --git a/Engine/Plugins/Cameras/GameplayCameras/Source/GameplayCamerasEditor/Private/Commands/CameraVariableCollectionEditorCommands.cpp b/Engine/Plugins/Cameras/GameplayCameras/Source/GameplayCamerasEditor/Private/Commands/CameraVariableCollectionEditorCommands.cpp
--- a/Engine/Plugins/Cameras/GameplayCameras/Source/GameplayCamerasEditor/Private/Commands/CameraVariableCollectionEditorCommands.cpp
+++ b/Engine/Plugins/Cameras/GameplayCameras/Source/GameplayCamerasEditor/Private/Commands/CameraVariableCollectionEditorCommands.cpp
@@ -28,6 +28,8 @@ void FCameraVariableCollectionEditorCommands::RegisterCommands()
 {
 	UI_COMMAND(CreateVariable, "Create Variable", "Adds a new camera variable to the collection",
 			EUserInterfaceActionType::Button, FInputChord());
+	UI_COMMAND(RenameVariable, "Rename Variable", "Renames a camera variable",
+			EUserInterfaceActionType::Button, FInputChord());
 	UI_COMMAND(DeleteVariable, "Delete Variable", "Removes a camera variable from the collection",
 			EUserInterfaceActionType::Button, FInputChord());
 }
diff --git a/Engine/Plugins/Cameras/GameplayCameras/Source/GameplayCamerasEditor/Private/Styles/GameplayCamerasEditorStyle.cpp b/Engine/Plugins/Cameras/GameplayCameras/Source/GameplayCamerasEditor/Private/Styles/GameplayCamerasEditorStyle.cpp
--- a/Engine/Plugins/Cameras/GameplayCameras/Source/GameplayCamerasEditor/Private/Styles/GameplayCamerasEditorStyle.cpp
+++ b/Engine/Plugins/Cameras/GameplayCameras/Source/GameplayCamerasEditor/Private/Styles/GameplayCamerasEditorStyle.cpp
@@ -89,6 +89,7 @@ FGameplayCamerasEditorStyle::FGameplayCamerasEditorStyle()
 
 	// Camera variable collection icons.
 	Set("CameraVariableCollectionEditor.CreateVariable", new CORE_IMAGE_BRUSH_SVG("Starship/Common/plus", Icon16x16));
+	Set("CameraVariableCollectionEditor.RenameVariable", new CORE_IMAGE_BRUSH_SVG("Starship/Common/Rename", Icon16x16));
 	Set("CameraVariableCollectionEditor.DeleteVariable", new CORE_IMAGE_BRUSH_SVG("Starship/Common/minus", Icon16x16));
 
 	Set("CameraVariableCollectionEditor.Entry.Name", FTextBlockStyle(NormalText) .SetFont(DEFAULT_FONT("Bold", 12)));
diff --git a/Engine/Plugins/Cameras/GameplayCameras/Source/GameplayCamerasEditor/Private/Toolkits/CameraVariableCollectionEditorToolkit.cpp b/Engine/Plugins/Cameras/GameplayCameras/Source/GameplayCamerasEditor/Private/Toolkits/CameraVariableCollectionEditorToolkit.cpp
--- a/Engine/Plugins/Cameras/GameplayCameras/Source/GameplayCamerasEditor/Private/Toolkits/CameraVariableCollectionEditorToolkit.cpp
+++ b/Engine/Plugins/Cameras/GameplayCameras/Source/GameplayCamerasEditor/Private/Toolkits/CameraVariableCollectionEditorToolkit.cpp
@@ -16,6 +16,7 @@
 #include "Modules/ModuleManager.h"
 #include "ObjectTools.h"
 #include "PropertyEditorModule.h"
+#include "PropertyPath.h"
 #include "ScopedTransaction.h"
 #include "Styles/GameplayCamerasEditorStyle.h"
 #include "ToolMenus.h"
@@ -377,6 +378,9 @@ void FCameraVariableCollectionEditorToolkit::RegisterToolbar()
 				);
 		VariablesSection.AddEntry(CreateVariableEntry);
 
+		FToolMenuEntry RenameVariableEntry = FToolMenuEntry::InitToolBarButton(Commands.RenameVariable);
+		VariablesSection.AddEntry(RenameVariableEntry);
+
 		FToolMenuEntry DeleteVariableEntry = FToolMenuEntry::InitToolBarButton(Commands.DeleteVariable);
 		VariablesSection.AddEntry(DeleteVariableEntry);
 	}
@@ -433,6 +437,11 @@ void FCameraVariableCollectionEditorToolkit::PostInitAssetEditor()
 {
 	const FCameraVariableCollectionEditorCommands& Commands = FCameraVariableCollectionEditorCommands::Get();
 
+	ToolkitCommands->MapAction(
+		Commands.RenameVariable,
+		FExecuteAction::CreateSP(this, &FCameraVariableCollectionEditorToolkit::OnRenameVariable),
+		FCanExecuteAction::CreateSP(this, &FCameraVariableCollectionEditorToolkit::CanRenameVariable));
+
 	ToolkitCommands->MapAction(
 		Commands.DeleteVariable,
 		FExecuteAction::CreateSP(this, &FCameraVariableCollectionEditorToolkit::OnDeleteVariable),
@@ -473,6 +482,23 @@ void FCameraVariableCollectionEditorToolkit::OnCreateVariable(TSubclassOf<UCamer
 	VariableCollectionEditorWidget->RequestListRefresh();
 }
 
+void FCameraVariableCollectionEditorToolkit::OnRenameVariable()
+{
+	UClass* VariableAssetClass = UCameraVariableAsset::StaticClass();
+	FProperty* DisplayNameProperty = VariableAssetClass->FindPropertyByName(GET_MEMBER_NAME_CHECKED(UCameraVariableAsset, DisplayName));
+
+	FPropertyPath PropertyPath;
+	PropertyPath.AddProperty(FPropertyInfo(DisplayNameProperty));
+	DetailsView->HighlightProperty(PropertyPath);
+}
+
+bool FCameraVariableCollectionEditorToolkit::CanRenameVariable()
+{
+	TArray<UCameraVariableAsset*> Selection;
+	VariableCollectionEditorWidget->GetSelectedVariables(Selection);
+	return !Selection.IsEmpty();
+}
+
 void FCameraVariableCollectionEditorToolkit::OnDeleteVariable()
 {
 	TArray<UCameraVariableAsset*> Selection;
diff --git a/Engine/Plugins/Cameras/GameplayCameras/Source/GameplayCamerasEditor/Public/Commands/CameraVariableCollectionEditorCommands.h b/Engine/Plugins/Cameras/GameplayCameras/Source/GameplayCamerasEditor/Public/Commands/CameraVariableCollectionEditorCommands.h
--- a/Engine/Plugins/Cameras/GameplayCameras/Source/GameplayCamerasEditor/Public/Commands/CameraVariableCollectionEditorCommands.h
+++ b/Engine/Plugins/Cameras/GameplayCameras/Source/GameplayCamerasEditor/Public/Commands/CameraVariableCollectionEditorCommands.h
@@ -23,6 +23,7 @@ public:
 public:
 
 	TSharedPtr<FUICommandInfo> CreateVariable;
+	TSharedPtr<FUICommandInfo> RenameVariable;
 	TSharedPtr<FUICommandInfo> DeleteVariable;
 };
 
diff --git a/Engine/Plugins/Cameras/GameplayCameras/Source/GameplayCamerasEditor/Public/Toolkits/CameraVariableCollectionEditorToolkit.h b/Engine/Plugins/Cameras/GameplayCameras/Source/GameplayCamerasEditor/Public/Toolkits/CameraVariableCollectionEditorToolkit.h
--- a/Engine/Plugins/Cameras/GameplayCameras/Source/GameplayCamerasEditor/Public/Toolkits/CameraVariableCollectionEditorToolkit.h
+++ b/Engine/Plugins/Cameras/GameplayCameras/Source/GameplayCamerasEditor/Public/Toolkits/CameraVariableCollectionEditorToolkit.h
@@ -63,6 +63,9 @@ private:
 
 	void OnCreateVariable(TSubclassOf<UCameraVariableAsset> InVariableClass);
 
+	void OnRenameVariable();
+	bool CanRenameVariable();
+
 	void OnDeleteVariable();
 	bool CanDeleteVariable();
 
