Cameras: fix missing parameter blending in the orbit blend

diff --git a/Engine/Plugins/Cameras/GameplayCameras/Source/GameplayCameras/Private/Nodes/Blends/OrbitBlendCameraNode.cpp b/Engine/Plugins/Cameras/GameplayCameras/Source/GameplayCameras/Private/Nodes/Blends/OrbitBlendCameraNode.cpp
--- a/Engine/Plugins/Cameras/GameplayCameras/Source/GameplayCameras/Private/Nodes/Blends/OrbitBlendCameraNode.cpp
+++ b/Engine/Plugins/Cameras/GameplayCameras/Source/GameplayCameras/Private/Nodes/Blends/OrbitBlendCameraNode.cpp
@@ -16,6 +16,7 @@ protected:
 	virtual void OnBuild(const FCameraNodeEvaluatorBuildParams& Params) override;
 	virtual FCameraNodeEvaluatorChildrenView OnGetChildren() override;
 	virtual void OnRun(const FCameraNodeEvaluationParams& Params, FCameraNodeEvaluationResult& OutResult) override;
+	virtual void OnBlendParameters(const FCameraNodePreBlendParams& Params, FCameraNodePreBlendResult& OutResult) override;
 	virtual void OnBlendResults(const FCameraNodeBlendParams& Params, FCameraNodeBlendResult& OutResult) override;
 
 	static bool ClosestPoints(const FRay3d& A, const FRay3d& B, double& OutParameterA, double& OutParameterB);
@@ -41,6 +42,19 @@ FCameraNodeEvaluatorChildrenView FOrbitBlendCameraNodeEvaluator::OnGetChildren()
 	return FCameraNodeEvaluatorChildrenView({ DrivingBlendEvaluator });
 }
 
+void FOrbitBlendCameraNodeEvaluator::OnBlendParameters(const FCameraNodePreBlendParams& Params, FCameraNodePreBlendResult& OutResult)
+{
+	if (DrivingBlendEvaluator)
+	{
+		DrivingBlendEvaluator->BlendParameters(Params, OutResult);
+	}
+	else
+	{
+		const FCameraVariableTable& ChildVariableTable(Params.ChildVariableTable);
+		OutResult.VariableTable.Override(ChildVariableTable, ECameraVariableTableFilter::Input);
+	}
+}
+
 void FOrbitBlendCameraNodeEvaluator::OnRun(const FCameraNodeEvaluationParams& Params, FCameraNodeEvaluationResult& OutResult)
 {
 	if (DrivingBlendEvaluator)
