Cameras: make the input axis binding node handle multiple bindings

This is because it seems typical to have mouse and gamepad bindings separate, with different input modifiers.

diff --git a/Engine/Plugins/Cameras/GameplayCameras/Source/GameplayCameras/Private/Nodes/Input/InputAxisBinding2DCameraNode.cpp b/Engine/Plugins/Cameras/GameplayCameras/Source/GameplayCameras/Private/Nodes/Input/InputAxisBinding2DCameraNode.cpp
--- a/Engine/Plugins/Cameras/GameplayCameras/Source/GameplayCameras/Private/Nodes/Input/InputAxisBinding2DCameraNode.cpp
+++ b/Engine/Plugins/Cameras/GameplayCameras/Source/GameplayCameras/Private/Nodes/Input/InputAxisBinding2DCameraNode.cpp
@@ -23,7 +23,7 @@ protected:
 private:
 
 	TObjectPtr<UEnhancedInputComponent> InputComponent;
-	FEnhancedInputActionValueBinding* AxisValueBinding = nullptr;
+	TArray<FEnhancedInputActionValueBinding*> AxisValueBindings;
 };
 
 UE_DEFINE_CAMERA_NODE_EVALUATOR(FInputAxisBinding2DCameraNodeEvaluator)
@@ -39,32 +39,50 @@ void FInputAxisBinding2DCameraNodeEvaluator::OnInitialize(const FCameraNodeEvalu
 	}
 
 	const UInputAxisBinding2DCameraNode* AxisBindingNode = GetCameraNodeAs<UInputAxisBinding2DCameraNode>();
-	if (InputComponent && AxisBindingNode->AxisAction)
+	if (InputComponent)
 	{
-		AxisValueBinding = &InputComponent->BindActionValue(AxisBindingNode->AxisAction);
+		for (TObjectPtr<UInputAction> AxisAction : AxisBindingNode->AxisActions)
+		{
+			FEnhancedInputActionValueBinding* AxisValueBinding = &InputComponent->BindActionValue(AxisAction);
+			AxisValueBindings.Add(AxisValueBinding);
+		}
 	}
 }
 
 void FInputAxisBinding2DCameraNodeEvaluator::OnUpdateParameters(const FCameraBlendedParameterUpdateParams& Params, FCameraBlendedParameterUpdateResult& OutResult)
 {
 	const UInputAxisBinding2DCameraNode* AxisBindingNode = GetCameraNodeAs<UInputAxisBinding2DCameraNode>();
-	if (AxisValueBinding)
-	{
-		const FVector2d Value = AxisValueBinding->GetValue().Get<FVector2D>();
 
-		InputValue = FVector2d(
-				Value.X * AxisBindingNode->Multiplier.X, 
-				Value.Y * AxisBindingNode->Multiplier.Y);
+	FVector2d HighestValue(FVector2d::ZeroVector);
+	double HighestSquaredLenth = 0.f;
 
-		if (AxisBindingNode->RevertAxisX)
+	for (FEnhancedInputActionValueBinding* AxisValueBinding : AxisValueBindings)
+	{
+		if (!AxisValueBinding)
 		{
-			InputValue.X = -InputValue.X;
+			continue;
 		}
-		if (AxisBindingNode->RevertAxisY)
+
+		const FVector2d Value = AxisValueBinding->GetValue().Get<FVector2D>();
+		const double ValueSquaredLength = Value.SquaredLength();
+		if (ValueSquaredLength > HighestSquaredLenth)
 		{
-			InputValue.Y = -InputValue.Y;
+			HighestValue = Value;
 		}
 	}
+
+	InputValue = FVector2d(
+			HighestValue.X * AxisBindingNode->Multiplier.X, 
+			HighestValue.Y * AxisBindingNode->Multiplier.Y);
+
+	if (AxisBindingNode->RevertAxisX)
+	{
+		InputValue.X = -InputValue.X;
+	}
+	if (AxisBindingNode->RevertAxisY)
+	{
+		InputValue.Y = -InputValue.Y;
+	}
 }
 
 }  // namespace UE::Cameras
diff --git a/Engine/Plugins/Cameras/GameplayCameras/Source/GameplayCameras/Public/Nodes/Input/InputAxisBinding2DCameraNode.h b/Engine/Plugins/Cameras/GameplayCameras/Source/GameplayCameras/Public/Nodes/Input/InputAxisBinding2DCameraNode.h
--- a/Engine/Plugins/Cameras/GameplayCameras/Source/GameplayCameras/Public/Nodes/Input/InputAxisBinding2DCameraNode.h
+++ b/Engine/Plugins/Cameras/GameplayCameras/Source/GameplayCameras/Public/Nodes/Input/InputAxisBinding2DCameraNode.h
@@ -19,9 +19,9 @@ class UInputAxisBinding2DCameraNode : public UInput2DCameraNode
 
 public:
 
-	/** The axis input action to read from. */
+	/** The axis input action(s) to read from. */
 	UPROPERTY(EditAnywhere, Category="Input")
-	TObjectPtr<UInputAction> AxisAction;
+	TArray<TObjectPtr<UInputAction>> AxisActions;
 
 	/** Whether to revert the X axis. */
 	UPROPERTY(EditAnywhere, Category="Input Processing")
