Cameras: fix visual bug with the orbit blend when the two lines of sight don't converge, and add debug info

diff --git a/Engine/Plugins/Cameras/GameplayCameras/Source/GameplayCameras/Private/Nodes/Blends/OrbitBlendCameraNode.cpp b/Engine/Plugins/Cameras/GameplayCameras/Source/GameplayCameras/Private/Nodes/Blends/OrbitBlendCameraNode.cpp
--- a/Engine/Plugins/Cameras/GameplayCameras/Source/GameplayCameras/Private/Nodes/Blends/OrbitBlendCameraNode.cpp
+++ b/Engine/Plugins/Cameras/GameplayCameras/Source/GameplayCameras/Private/Nodes/Blends/OrbitBlendCameraNode.cpp
@@ -2,32 +2,71 @@
 
 #include "Nodes/Blends/OrbitBlendCameraNode.h"
 
+#include "Debug/CameraDebugBlock.h"
+#include "Debug/CameraDebugBlockBuilder.h"
+#include "Debug/CameraDebugRenderer.h"
+#include "HAL/IConsoleManager.h"
+
 #include UE_INLINE_GENERATED_CPP_BY_NAME(OrbitBlendCameraNode)
 
 namespace UE::Cameras
 {
 
+float GOrbitBlendDeactivationSmoothingTime = 0.3f;
+static FAutoConsoleVariableRef CVarOrbitBlendDeactivationSmoothingTime(
+	TEXT("GameplayCameras.OrbitBlend.DeactivationSmoothingTime"),
+	GOrbitBlendDeactivationSmoothingTime,
+	TEXT("(Default: 0.3 seconds. The time to smooth out any differences between a deactivated orbit blend and its underlying blend."));
+
+float GOrbitBlendDeactivationSmoothingMinTime = 0.01f;
+
 class FOrbitBlendCameraNodeEvaluator : public FBlendCameraNodeEvaluator
 {
 	UE_DECLARE_BLEND_CAMERA_NODE_EVALUATOR(GAMEPLAYCAMERAS_API, FOrbitBlendCameraNodeEvaluator)
 
 protected:
 
+	// FCameraNodeEvaluator interface.
 	virtual void OnBuild(const FCameraNodeEvaluatorBuildParams& Params) override;
 	virtual FCameraNodeEvaluatorChildrenView OnGetChildren() override;
 	virtual void OnRun(const FCameraNodeEvaluationParams& Params, FCameraNodeEvaluationResult& OutResult) override;
 	virtual void OnBlendParameters(const FCameraNodePreBlendParams& Params, FCameraNodePreBlendResult& OutResult) override;
 	virtual void OnBlendResults(const FCameraNodeBlendParams& Params, FCameraNodeBlendResult& OutResult) override;
+	virtual void OnSerialize(const FCameraNodeEvaluatorSerializeParams& Params, FArchive& Ar) override;
+#if UE_GAMEPLAY_CAMERAS_DEBUG
+	virtual void OnBuildDebugBlocks(const FCameraDebugBlockBuildParams& Params, FCameraDebugBlockBuilder& Builder) override;
+#endif  // UE_GAMEPLAY_CAMERAS_DEBUG
 
 	static bool ClosestPoints(const FRay3d& A, const FRay3d& B, double& OutParameterA, double& OutParameterB);
 
 private:
 
+	enum class FState
+	{
+		Active,
+		SmoothingOut,
+		Inactive
+	};
+
 	FSimpleBlendCameraNodeEvaluator* DrivingBlendEvaluator = nullptr;
+
+	FVector3d DeltaLocation = FVector3d::ZeroVector;
+	FRotator3d DeltaRotation = FRotator3d::ZeroRotator;
+	float SmoothingTimeLeft = -1.f;
+	FState State = FState::Active;
 };
 
 UE_DEFINE_BLEND_CAMERA_NODE_EVALUATOR(FOrbitBlendCameraNodeEvaluator)
 
+UE_DECLARE_CAMERA_DEBUG_BLOCK_START(GAMEPLAYCAMERAS_API, FOrbitBlendCameraDebugBlock)
+	UE_DECLARE_CAMERA_DEBUG_BLOCK_FIELD(FVector3d, DeltaLocation);
+	UE_DECLARE_CAMERA_DEBUG_BLOCK_FIELD(FRotator3d, DeltaRotation);
+	UE_DECLARE_CAMERA_DEBUG_BLOCK_FIELD(float, SmoothingTimeLeft);
+	UE_DECLARE_CAMERA_DEBUG_BLOCK_FIELD(bool, bIsActive);
+UE_DECLARE_CAMERA_DEBUG_BLOCK_END()
+
+UE_DEFINE_CAMERA_DEBUG_BLOCK_WITH_FIELDS(FOrbitBlendCameraDebugBlock)
+
 void FOrbitBlendCameraNodeEvaluator::OnBuild(const FCameraNodeEvaluatorBuildParams& Params)
 {
 	const UOrbitBlendCameraNode* OrbitBlend = GetCameraNodeAs<UOrbitBlendCameraNode>();
@@ -100,41 +139,132 @@ void FOrbitBlendCameraNodeEvaluator::OnBlendResults(const FCameraNodeBlendParams
 		return;
 	}
 
-	// Find the points on each line of sight that are the closest to each other. If successful,
-	// start blending around an interpolating mid-point between the two.
-	double FromClosestParam, ToClosestParam;
-	const bool bFoundClosestPoints = ClosestPoints(FromAim, ToAim, FromClosestParam, ToClosestParam);
-	if (bFoundClosestPoints)
+	if (State == FState::Active)
 	{
-		const float Factor = DrivingBlendEvaluator->GetBlendFactor();
+		// Find the points on each line of sight that are the closest to each other. If successful,
+		// start blending around an interpolating mid-point between the two.
+		double FromClosestParam, ToClosestParam;
+		const bool bFoundClosestPoints = ClosestPoints(FromAim, ToAim, FromClosestParam, ToClosestParam);
+
+		// We only do the orbit blend if we (1) found the appropriate points and (2) they are both in
+		// front of us.
+		// (1) may be false if both lines of sight are parallel.
+		// (2) may be false if the lines of sight aren't converging (i.e. the two cameras are looking
+		//     away from each other, even if just slighly).
+		const bool bDoOrbitBlend = bFoundClosestPoints && FromClosestParam > 0 && ToClosestParam > 0;
+
+		if (bDoOrbitBlend)
+		{
+			const float Factor = DrivingBlendEvaluator->GetBlendFactor();
+
+			const FVector3d BlendedLocation = FMath::Lerp(FromLocation, ToLocation, Factor);
+
+			// Rotate around a point that is interpolating from the first line of sight to the other
+			// line of sight.
+			const FVector3d FromOrbitCenter = FromAim.PointAt(FromClosestParam);
+			const FVector3d ToOrbitCenter = ToAim.PointAt(ToClosestParam);
+			const FVector3d BlendedOrbitCenter = FMath::Lerp(FromOrbitCenter, ToOrbitCenter, Factor);
+
+			const FVector3d BlendedAimDir = (BlendedOrbitCenter - BlendedLocation).GetUnsafeNormal();
+			const FRotator3d BlendedRotation = BlendedAimDir.ToOrientationRotator();
+
+			const double FromOrbitCenterDistance = FVector3d::Distance(FromLocation, FromOrbitCenter);
+			const double ToOrbitCenterDistance = FVector3d::Distance(ToLocation, ToOrbitCenter);
+			const double BlendedTargetDistance = FMath::Lerp(FromOrbitCenterDistance, ToOrbitCenterDistance, Factor);
+
+			// So, instead of interpolating between the two positions and letting the target move
+			// forwards or backwards as TargetDistance interpolates, we do the opposite: we "anchor"
+			// the target at the orbit center, and push or pull the position based on the interpolated
+			// TargetDistance.
+			const FRay3d BlendedReverseAim(BlendedOrbitCenter, -BlendedAimDir, true);
+			const FVector3d OrbitingLocation = BlendedReverseAim.PointAt(BlendedTargetDistance);
+
+			// Remember our offset from the underlying blend before we apply our orbit blend.
+			DeltaLocation = OrbitingLocation - BlendedResult.CameraPose.GetLocation();
+			DeltaRotation = BlendedRotation - BlendedResult.CameraPose.GetRotation();
+
+			// Apply the orbit blend!
+			BlendedResult.CameraPose.SetLocation(OrbitingLocation);
+			BlendedResult.CameraPose.SetRotation(BlendedRotation);
+		}
+		else
+		{
+			// The orbit blend was deactivated for one of the reasons mentioned above. We are going to
+			// smooth out the difference betwen us and our underlying driving blend over a short time
+			// to prevent creating artefacts... that is, unless we somehow have zero difference.
+			const bool bHasDelta = (!DeltaLocation.IsZero() || !DeltaRotation.IsZero());
+			State = bHasDelta ? FState::SmoothingOut : FState::Inactive;
+
+			const float SmoothingTime = FMath::Max(GOrbitBlendDeactivationSmoothingTime, GOrbitBlendDeactivationSmoothingMinTime);
+			SmoothingTimeLeft = SmoothingTime;
+
+			// NOTE: once we deactivated, we never try to reactivate, to keep things simple.
+		}
+	}
 
-		const FVector3d BlendedLocation = FMath::Lerp(FromLocation, ToLocation, Factor);
+	if (State == FState::SmoothingOut)
+	{
 
-		// Rotate around a point that is interpolating from the first line of sight to the other
-		// line of sight.
-		const FVector3d FromOrbitCenter = FromAim.PointAt(FromClosestParam);
-		const FVector3d ToOrbitCenter = ToAim.PointAt(ToClosestParam);
-		const FVector3d BlendedOrbitCenter = FMath::Lerp(FromOrbitCenter, ToOrbitCenter, Factor);
+		SmoothingTimeLeft -= Params.ChildParams.DeltaTime;
+		if (SmoothingTimeLeft > 0.f)
+		{
+			// Continue smoothing out the difference linearly.
+			const float SmoothingTime = FMath::Max(GOrbitBlendDeactivationSmoothingTime, GOrbitBlendDeactivationSmoothingMinTime);
+			const float Alpha = (SmoothingTime - SmoothingTimeLeft) / SmoothingTime;
+			const FVector3d CurDeltaLocation = FMath::Lerp(DeltaLocation, FVector3d::ZeroVector, Alpha);
+			const FRotator3d CurDeltaRotation = FMath::Lerp(DeltaRotation, FRotator3d::ZeroRotator, Alpha);
+
+			BlendedResult.CameraPose.SetLocation(BlendedResult.CameraPose.GetLocation() + CurDeltaLocation);
+			BlendedResult.CameraPose.SetRotation(BlendedResult.CameraPose.GetRotation() + CurDeltaRotation);
+
+			// Even if our underlying blend is finished, we still have some smoothing out to do.
+			OutResult.bIsBlendFull = false;
+			OutResult.bIsBlendFinished = false;
+		}
+		else
+		{
+			// We are done.
+			DeltaLocation = FVector3d::ZeroVector;
+			DeltaRotation = FRotator3d::ZeroRotator;
+			SmoothingTimeLeft = -1.f;
+			State = FState::Inactive;
+		}
+	}
+}
 
-		const FVector3d BlendedAimDir = (BlendedOrbitCenter - BlendedLocation).GetUnsafeNormal();
-		const FRotator3d BlendedRotation = BlendedAimDir.ToOrientationRotator();
+void FOrbitBlendCameraNodeEvaluator::OnSerialize(const FCameraNodeEvaluatorSerializeParams& Params, FArchive& Ar)
+{
+	Ar << DeltaLocation;
+	Ar << DeltaRotation;
+	Ar << SmoothingTimeLeft;
+	Ar << State;
+}
 
-		const double FromOrbitCenterDistance = FVector3d::Distance(FromLocation, FromOrbitCenter);
-		const double ToOrbitCenterDistance = FVector3d::Distance(ToLocation, ToOrbitCenter);
-		const double BlendedTargetDistance = FMath::Lerp(FromOrbitCenterDistance, ToOrbitCenterDistance, Factor);
+#if UE_GAMEPLAY_CAMERAS_DEBUG
 
-		// So, instead of interpolating between the two positions and letting the target move
-		// forwards or backwards as TargetDistance interpolates, we do the opposite: we "anchor"
-		// the target at the orbit center, and push or pull the position based on the interpolated
-		// TargetDistance.
-		const FRay3d BlendedReverseAim(BlendedOrbitCenter, -BlendedAimDir, true);
-		const FVector3d OrbitingLocation = BlendedReverseAim.PointAt(BlendedTargetDistance);
+void FOrbitBlendCameraNodeEvaluator::OnBuildDebugBlocks(const FCameraDebugBlockBuildParams& Params, FCameraDebugBlockBuilder& Builder)
+{
+	FOrbitBlendCameraDebugBlock& DebugBlock = Builder.AttachDebugBlock<FOrbitBlendCameraDebugBlock>();
+	DebugBlock.DeltaLocation = DeltaLocation;
+	DebugBlock.DeltaRotation = DeltaRotation;
+	DebugBlock.SmoothingTimeLeft = SmoothingTimeLeft;
+	DebugBlock.bIsActive = (State == FState::Active);
+}
 
-		BlendedResult.CameraPose.SetLocation(OrbitingLocation);
-		BlendedResult.CameraPose.SetRotation(BlendedRotation);
+void FOrbitBlendCameraDebugBlock::OnDebugDraw(const FCameraDebugBlockDrawParams& Params, FCameraDebugRenderer& Renderer)
+{
+	if (bIsActive)
+	{
+		Renderer.AddText(TEXT("orbiting (delta = %s)"), *DeltaLocation.ToString());
+	}
+	else
+	{
+		Renderer.AddText(TEXT("INACTIVE (delta = %s  smoothing time %f)"), *DeltaLocation.ToString(), SmoothingTimeLeft);
 	}
 }
 
+#endif  // UE_GAMEPLAY_CAMERAS_DEBUG
+
 bool FOrbitBlendCameraNodeEvaluator::ClosestPoints(const FRay3d& A, const FRay3d& B, double& OutParameterA, double& OutParameterB)
 {
 	// The points closest to each other on rays A and B are named T1 and T1. They are such that the
