Cameras: add metadata for customizing object tree graph pin directions, make all input camera nodes go left-to-right

diff --git a/Engine/Plugins/Cameras/GameplayCameras/Source/GameplayCameras/Public/Core/CameraRigAsset.h b/Engine/Plugins/Cameras/GameplayCameras/Source/GameplayCameras/Public/Core/CameraRigAsset.h
--- a/Engine/Plugins/Cameras/GameplayCameras/Source/GameplayCameras/Public/Core/CameraRigAsset.h
+++ b/Engine/Plugins/Cameras/GameplayCameras/Source/GameplayCameras/Public/Core/CameraRigAsset.h
@@ -45,7 +45,8 @@ struct FCameraRigAllocationInfo
  * An exposed camera rig parameter that drives a specific parameter on one of
  * its camera nodes.
  */
-UCLASS(MinimalAPI)
+UCLASS(MinimalAPI, meta=(
+			ObjectTreeGraphSelfPinDirection="Output"))
 class UCameraRigInterfaceParameter
 	: public UObject
 	, public IObjectTreeGraphObject
diff --git a/Engine/Plugins/Cameras/GameplayCameras/Source/GameplayCameras/Public/Core/CameraRigTransition.h b/Engine/Plugins/Cameras/GameplayCameras/Source/GameplayCameras/Public/Core/CameraRigTransition.h
--- a/Engine/Plugins/Cameras/GameplayCameras/Source/GameplayCameras/Public/Core/CameraRigTransition.h
+++ b/Engine/Plugins/Cameras/GameplayCameras/Source/GameplayCameras/Public/Core/CameraRigTransition.h
@@ -32,7 +32,10 @@ struct FCameraRigTransitionConditionMatchParams
 /**
  * Base class for a camera transition condition.
  */
-UCLASS(Abstract, DefaultToInstanced, MinimalAPI, meta=(ObjectTreeGraphCategory="Transition Conditions"))
+UCLASS(Abstract, DefaultToInstanced, MinimalAPI, meta=(
+			ObjectTreeGraphCategory="Transition Conditions",
+			ObjectTreeGraphSelfPinDirection="Output",
+			ObjectTreeGraphDefaultPropertyPinDirection="Input"))
 class UCameraRigTransitionCondition 
 	: public UObject
 	, public IObjectTreeGraphObject
diff --git a/Engine/Plugins/Cameras/GameplayCameras/Source/GameplayCameras/Public/Nodes/Common/BoomArmCameraNode.h b/Engine/Plugins/Cameras/GameplayCameras/Source/GameplayCameras/Public/Nodes/Common/BoomArmCameraNode.h
--- a/Engine/Plugins/Cameras/GameplayCameras/Source/GameplayCameras/Public/Nodes/Common/BoomArmCameraNode.h
+++ b/Engine/Plugins/Cameras/GameplayCameras/Source/GameplayCameras/Public/Nodes/Common/BoomArmCameraNode.h
@@ -33,7 +33,7 @@ public:
 	 * The input slot for controlling the boom arm.
 	 * If no input slot is specified, the boom arm will use the player controller view rotation.
 	 */
-	UPROPERTY()
+	UPROPERTY(meta=(ObjectTreeGraphPinDirection=Input))
 	TObjectPtr<UCameraRigInput2DSlot> InputSlot;
 };
 
diff --git a/Engine/Plugins/Cameras/GameplayCameras/Source/GameplayCameras/Public/Nodes/Input/CameraRigInput1DSlot.h b/Engine/Plugins/Cameras/GameplayCameras/Source/GameplayCameras/Public/Nodes/Input/CameraRigInput1DSlot.h
--- a/Engine/Plugins/Cameras/GameplayCameras/Source/GameplayCameras/Public/Nodes/Input/CameraRigInput1DSlot.h
+++ b/Engine/Plugins/Cameras/GameplayCameras/Source/GameplayCameras/Public/Nodes/Input/CameraRigInput1DSlot.h
@@ -16,7 +16,10 @@ class UInput1DCameraNode;
 /**
  * A node that can handle and accumulate a chain of player input nodes.
  */
-UCLASS(MinimalAPI, meta=(CameraNodeCategories="Input"))
+UCLASS(MinimalAPI, meta=(
+			CameraNodeCategories="Input", 
+			ObjectTreeGraphSelfPinDirection="Output", 
+			ObjectTreeGraphDefaultPropertyPinDirection="Input"))
 class UCameraRigInput1DSlot : public UCameraNode
 {
 	GENERATED_BODY()
diff --git a/Engine/Plugins/Cameras/GameplayCameras/Source/GameplayCameras/Public/Nodes/Input/CameraRigInput2DSlot.h b/Engine/Plugins/Cameras/GameplayCameras/Source/GameplayCameras/Public/Nodes/Input/CameraRigInput2DSlot.h
--- a/Engine/Plugins/Cameras/GameplayCameras/Source/GameplayCameras/Public/Nodes/Input/CameraRigInput2DSlot.h
+++ b/Engine/Plugins/Cameras/GameplayCameras/Source/GameplayCameras/Public/Nodes/Input/CameraRigInput2DSlot.h
@@ -16,7 +16,10 @@ class UVector2dCameraVariable;
 /**
  * A node that can handle and accumulate a chain of player input nodes.
  */
-UCLASS(MinimalAPI, meta=(CameraNodeCategories="Input"))
+UCLASS(MinimalAPI, meta=(
+			CameraNodeCategories="Input", 
+			ObjectTreeGraphSelfPinDirection="Output",
+			ObjectTreeGraphDefaultPropertyPinDirection="Input"))
 class UCameraRigInput2DSlot : public UCameraNode
 {
 	GENERATED_BODY()
diff --git a/Engine/Plugins/Cameras/GameplayCameras/Source/GameplayCameras/Public/Nodes/Input/Input1DCameraNode.h b/Engine/Plugins/Cameras/GameplayCameras/Source/GameplayCameras/Public/Nodes/Input/Input1DCameraNode.h
--- a/Engine/Plugins/Cameras/GameplayCameras/Source/GameplayCameras/Public/Nodes/Input/Input1DCameraNode.h
+++ b/Engine/Plugins/Cameras/GameplayCameras/Source/GameplayCameras/Public/Nodes/Input/Input1DCameraNode.h
@@ -10,7 +10,10 @@
 /**
  * A camera node that provides a floating-point value to an input slot.
  */
-UCLASS(MinimalAPI, meta=(CameraNodeCategories="Input"))
+UCLASS(MinimalAPI, meta=(
+			CameraNodeCategories="Input",
+			ObjectTreeGraphSelfPinDirection="Output",
+			ObjectTreeGraphDefaultPropertyPinDirection="Input"))
 class UInput1DCameraNode : public UCameraNode
 {
 	GENERATED_BODY()
diff --git a/Engine/Plugins/Cameras/GameplayCameras/Source/GameplayCameras/Public/Nodes/Input/Input2DCameraNode.h b/Engine/Plugins/Cameras/GameplayCameras/Source/GameplayCameras/Public/Nodes/Input/Input2DCameraNode.h
--- a/Engine/Plugins/Cameras/GameplayCameras/Source/GameplayCameras/Public/Nodes/Input/Input2DCameraNode.h
+++ b/Engine/Plugins/Cameras/GameplayCameras/Source/GameplayCameras/Public/Nodes/Input/Input2DCameraNode.h
@@ -11,7 +11,10 @@
 /**
  * A camera node that provides a two-dimensional value to an input slot.
  */
-UCLASS(MinimalAPI, meta=(CameraNodeCategories="Input"))
+UCLASS(MinimalAPI, meta=(
+			CameraNodeCategories="Input",
+			ObjectTreeGraphSelfPinDirection="Output",
+			ObjectTreeGraphDefaultPropertyPinDirection="Input"))
 class UInput2DCameraNode : public UCameraNode
 {
 	GENERATED_BODY()
diff --git a/Engine/Plugins/Cameras/GameplayCameras/Source/GameplayCamerasEditor/Private/Editors/CameraNodeGraphSchema.cpp b/Engine/Plugins/Cameras/GameplayCameras/Source/GameplayCamerasEditor/Private/Editors/CameraNodeGraphSchema.cpp
--- a/Engine/Plugins/Cameras/GameplayCameras/Source/GameplayCamerasEditor/Private/Editors/CameraNodeGraphSchema.cpp
+++ b/Engine/Plugins/Cameras/GameplayCameras/Source/GameplayCamerasEditor/Private/Editors/CameraNodeGraphSchema.cpp
@@ -43,7 +43,6 @@ FObjectTreeGraphConfig UCameraNodeGraphSchema::BuildGraphConfig() const
 		.CreateCategoryMetaData(TEXT("CameraNodeCategories"))
 		.GraphNodeClass(UCameraNodeGraphNode::StaticClass());
 	GraphConfig.ObjectClassConfigs.Emplace(UCameraRigInterfaceParameter::StaticClass())
-		.SelfPinDirection(EGPD_Output)
 		.SelfPinName(NAME_None)  // No self pin name, we just want the title
 		.CanCreateNew(false)
 		.GraphNodeClass(UCameraRigInterfaceParameterGraphNode::StaticClass());
diff --git a/Engine/Plugins/Cameras/GameplayCameras/Source/GameplayCamerasEditor/Private/Editors/CameraRigTransitionGraphSchemaBase.cpp b/Engine/Plugins/Cameras/GameplayCameras/Source/GameplayCamerasEditor/Private/Editors/CameraRigTransitionGraphSchemaBase.cpp
--- a/Engine/Plugins/Cameras/GameplayCameras/Source/GameplayCamerasEditor/Private/Editors/CameraRigTransitionGraphSchemaBase.cpp
+++ b/Engine/Plugins/Cameras/GameplayCameras/Source/GameplayCamerasEditor/Private/Editors/CameraRigTransitionGraphSchemaBase.cpp
@@ -26,8 +26,6 @@ FObjectTreeGraphConfig UCameraRigTransitionGraphSchemaBase::BuildGraphConfig() c
 	GraphConfig.ObjectClassConfigs.Emplace(UCameraRigTransition::StaticClass())
 		.NodeTitleColor(Settings->CameraRigTransitionTitleColor);
 	GraphConfig.ObjectClassConfigs.Emplace(UCameraRigTransitionCondition::StaticClass())
-		.SelfPinDirection(EGPD_Output)
-		.DefaultPropertyPinDirection(EGPD_Input)
 		.StripDisplayNameSuffix(TEXT("Transition Condition"))
 		.NodeTitleColor(Settings->CameraRigTransitionConditionTitleColor);
 	GraphConfig.ObjectClassConfigs.Emplace(UBlendCameraNode::StaticClass())
diff --git a/Engine/Plugins/Cameras/GameplayCameras/Source/GameplayCamerasEditor/Private/Editors/ObjectTreeGraphConfig.cpp b/Engine/Plugins/Cameras/GameplayCameras/Source/GameplayCamerasEditor/Private/Editors/ObjectTreeGraphConfig.cpp
--- a/Engine/Plugins/Cameras/GameplayCameras/Source/GameplayCamerasEditor/Private/Editors/ObjectTreeGraphConfig.cpp
+++ b/Engine/Plugins/Cameras/GameplayCameras/Source/GameplayCamerasEditor/Private/Editors/ObjectTreeGraphConfig.cpp
@@ -13,9 +13,7 @@
 FObjectTreeGraphClassConfig::FObjectTreeGraphClassConfig()
 	: _SelfPinName(NAME_Self)
 	, _SelfPinFriendlyName(FText::GetEmpty())
-	, _SelfPinDirection(EGPD_Input)
 	, _HasSelfPin(true)
-	, _DefaultPropertyPinDirection(EGPD_Output)
 	, _NodeTitleUsesObjectName(false)
 	, _CanCreateNew(true)
 	, _CanDelete(true)
@@ -228,6 +226,33 @@ void FObjectTreeGraphConfig::FormatDisplayNameText(const UObject* InObject, cons
 	OnFormatObjectDisplayName.ExecuteIfBound(InObject, InOutDisplayNameText);
 }
 
+EEdGraphPinDirection FObjectTreeGraphConfig::GetSelfPinDirection(const UClass* InObjectClass) const
+{
+	const FObjectTreeGraphClassConfig& ClassConfig = GetObjectClassConfig(InObjectClass);
+	TOptional<EEdGraphPinDirection> PinDirectionOverride = ClassConfig.SelfPinDirectionOverride();
+	if (PinDirectionOverride.IsSet())
+	{
+		return PinDirectionOverride.GetValue();
+	}
+
+	while (InObjectClass)
+	{
+		const FString& CustomDefaultDirection = InObjectClass->GetMetaData(TEXT("ObjectTreeGraphSelfPinDirection"));
+		if (CustomDefaultDirection == TEXT("Input"))
+		{
+			return EGPD_Input;
+		}
+		else if (CustomDefaultDirection == TEXT("Output"))
+		{
+			return EGPD_Output;
+		}
+
+		InObjectClass = InObjectClass->GetSuperClass();
+	}
+
+	return EGPD_Input;
+}
+
 EEdGraphPinDirection FObjectTreeGraphConfig::GetPropertyPinDirection(const UClass* InObjectClass, const FName& InPropertyName) const
 {
 	const FObjectTreeGraphClassConfig& ClassConfig = GetObjectClassConfig(InObjectClass);
@@ -248,7 +273,28 @@ EEdGraphPinDirection FObjectTreeGraphConfig::GetPropertyPinDirection(const UClas
 		return EGPD_Output;
 	}
 
-	return ClassConfig.DefaultPropertyPinDirection();
+	TOptional<EEdGraphPinDirection> DefaultPinDirectionOverride = ClassConfig.DefaultPropertyPinDirectionOverride();
+	if (DefaultPinDirectionOverride.IsSet())
+	{
+		return DefaultPinDirectionOverride.GetValue();
+	}
+
+	while (InObjectClass)
+	{
+		const FString& CustomDefaultDirection = InObjectClass->GetMetaData(TEXT("ObjectTreeGraphDefaultPropertyPinDirection"));
+		if (CustomDefaultDirection == TEXT("Input"))
+		{
+			return EGPD_Input;
+		}
+		else if (CustomDefaultDirection == TEXT("Output"))
+		{
+			return EGPD_Output;
+		}
+
+		InObjectClass = InObjectClass->GetSuperClass();
+	}
+
+	return EGPD_Output;
 }
 
 #undef LOCTEXT_NAMESPACE
diff --git a/Engine/Plugins/Cameras/GameplayCameras/Source/GameplayCamerasEditor/Private/Editors/ObjectTreeGraphNode.cpp b/Engine/Plugins/Cameras/GameplayCameras/Source/GameplayCamerasEditor/Private/Editors/ObjectTreeGraphNode.cpp
--- a/Engine/Plugins/Cameras/GameplayCameras/Source/GameplayCamerasEditor/Private/Editors/ObjectTreeGraphNode.cpp
+++ b/Engine/Plugins/Cameras/GameplayCameras/Source/GameplayCamerasEditor/Private/Editors/ObjectTreeGraphNode.cpp
@@ -77,7 +77,7 @@ void UObjectTreeGraphNode::AllocateDefaultPins()
 		FEdGraphPinType SelfPinType;
 		SelfPinType.PinCategory = UObjectTreeGraphSchema::PC_Self;
 		const FName& SelfPinName = ObjectClassConfig.SelfPinName();
-		UEdGraphPin* SelfPin = CreatePin(ObjectClassConfig.SelfPinDirection(), SelfPinType, SelfPinName);
+		UEdGraphPin* SelfPin = CreatePin(OuterGraphConfig.GetSelfPinDirection(NodeContext.ObjectClass), SelfPinType, SelfPinName);
 		SelfPin->PinFriendlyName = ObjectClassConfig.SelfPinFriendlyName();
 	}
 
diff --git a/Engine/Plugins/Cameras/GameplayCameras/Source/GameplayCamerasEditor/Public/Editors/ObjectTreeGraphConfig.h b/Engine/Plugins/Cameras/GameplayCameras/Source/GameplayCamerasEditor/Public/Editors/ObjectTreeGraphConfig.h
--- a/Engine/Plugins/Cameras/GameplayCameras/Source/GameplayCamerasEditor/Public/Editors/ObjectTreeGraphConfig.h
+++ b/Engine/Plugins/Cameras/GameplayCameras/Source/GameplayCamerasEditor/Public/Editors/ObjectTreeGraphConfig.h
@@ -50,12 +50,12 @@ public:
 	/** The display name of the self pin. */
 	OTGCC_FIELD(FText, SelfPinFriendlyName)
 	/** The direction of the self pin. */
-	OTGCC_FIELD(EEdGraphPinDirection, SelfPinDirection)
+	OTGCC_FIELD(TOptional<EEdGraphPinDirection>, SelfPinDirectionOverride)
 	/** Whether graph nodes for this class have a self pin. */
 	OTGCC_FIELD(bool, HasSelfPin)
 
 	/** Default direction of property pins. */
-	OTGCC_FIELD(EEdGraphPinDirection, DefaultPropertyPinDirection)
+	OTGCC_FIELD(TOptional<EEdGraphPinDirection>, DefaultPropertyPinDirectionOverride)
 
 	/** Color of the graph node's title. */
 	OTGCC_FIELD(TOptional<FLinearColor>, NodeTitleColor)
@@ -231,6 +231,9 @@ public:
 	/** Computes the display name of the given object class. */
 	FText GetDisplayNameText(const UClass* InClass) const;
 
+	/** Gets the "self" pin direction for a given class. */
+	EEdGraphPinDirection GetSelfPinDirection(const UClass* InObjectClass) const;
+
 	/** Gets the custom property pin direction for a given named property. */
 	EEdGraphPinDirection GetPropertyPinDirection(const UClass* InObjectClass, const FName& InPropertyName) const;
 
