Cameras: simplify input nodes

The InputAxisBinding input node is now a subclass of the input slot, so that users only need one input node to get things going.

diff --git a/Engine/Plugins/Cameras/GameplayCameras/Source/GameplayCameras/Private/Nodes/Common/BoomArmCameraNode.cpp b/Engine/Plugins/Cameras/GameplayCameras/Source/GameplayCameras/Private/Nodes/Common/BoomArmCameraNode.cpp
--- a/Engine/Plugins/Cameras/GameplayCameras/Source/GameplayCameras/Private/Nodes/Common/BoomArmCameraNode.cpp
+++ b/Engine/Plugins/Cameras/GameplayCameras/Source/GameplayCameras/Private/Nodes/Common/BoomArmCameraNode.cpp
@@ -8,7 +8,7 @@
 #include "Core/CameraRigJoints.h"
 #include "GameFramework/Pawn.h"
 #include "GameFramework/PlayerController.h"
-#include "Nodes/Input/CameraRigInput2DSlot.h"
+#include "Nodes/Input/Input2DCameraNode.h"
 
 #include UE_INLINE_GENERATED_CPP_BY_NAME(BoomArmCameraNode)
 
@@ -34,7 +34,7 @@ private:
 private:
 
 	TCameraParameterReader<FVector3d> BoomOffsetReader;
-	FCameraRigInput2DSlotEvaluator* InputSlotEvaluator = nullptr;
+	FInput2DCameraNodeEvaluator* InputSlotEvaluator = nullptr;
 };
 
 UE_DEFINE_CAMERA_NODE_EVALUATOR(FBoomArmCameraNodeEvaluator)
@@ -42,7 +42,7 @@ UE_DEFINE_CAMERA_NODE_EVALUATOR(FBoomArmCameraNodeEvaluator)
 void FBoomArmCameraNodeEvaluator::OnBuild(const FCameraNodeEvaluatorBuildParams& Params)
 {
 	const UBoomArmCameraNode* BoomArmNode = GetCameraNodeAs<UBoomArmCameraNode>();
-	InputSlotEvaluator = Params.BuildEvaluatorAs<FCameraRigInput2DSlotEvaluator>(BoomArmNode->InputSlot);
+	InputSlotEvaluator = Params.BuildEvaluatorAs<FInput2DCameraNodeEvaluator>(BoomArmNode->InputSlot);
 }
 
 void FBoomArmCameraNodeEvaluator::OnInitialize(const FCameraNodeEvaluatorInitializeParams& Params)
diff --git a/Engine/Plugins/Cameras/GameplayCameras/Source/GameplayCameras/Private/Nodes/Input/CameraRigInput1DSlot.cpp b/Engine/Plugins/Cameras/GameplayCameras/Source/GameplayCameras/Private/Nodes/Input/CameraRigInput1DSlot.cpp
--- a/Engine/Plugins/Cameras/GameplayCameras/Source/GameplayCameras/Private/Nodes/Input/CameraRigInput1DSlot.cpp
+++ b/Engine/Plugins/Cameras/GameplayCameras/Source/GameplayCameras/Private/Nodes/Input/CameraRigInput1DSlot.cpp
@@ -25,12 +25,6 @@ FCameraRigInput1DSlotEvaluator::FCameraRigInput1DSlotEvaluator()
 			ECameraNodeEvaluatorFlags::SupportsOperations);
 }
 
-void FCameraRigInput1DSlotEvaluator::OnBuild(const FCameraNodeEvaluatorBuildParams& Params)
-{
-	const UCameraRigInput1DSlot* SlotNode = GetCameraNodeAs<UCameraRigInput1DSlot>();
-	ChildEvaluator = Params.BuildEvaluatorAs<FInput1DCameraNodeEvaluator>(SlotNode->Child);
-}
-
 void FCameraRigInput1DSlotEvaluator::OnInitialize(const FCameraNodeEvaluatorInitializeParams& Params)
 {
 	TransientInputValue = 0.f;
@@ -44,20 +38,8 @@ void FCameraRigInput1DSlotEvaluator::OnInitialize(const FCameraNodeEvaluatorInit
 	}
 }
 
-FCameraNodeEvaluatorChildrenView FCameraRigInput1DSlotEvaluator::OnGetChildren()
-{
-	return FCameraNodeEvaluatorChildrenView{ ChildEvaluator };
-}
-
 void FCameraRigInput1DSlotEvaluator::OnUpdateParameters(const FCameraBlendedParameterUpdateParams& Params, FCameraBlendedParameterUpdateResult& OutResult)
 {
-	if (ChildEvaluator)
-	{
-		ChildEvaluator->UpdateParameters(Params, OutResult);
-
-		TransientInputValue = ChildEvaluator->GetInputValue();
-	}
-
 	const UCameraRigInput1DSlot* SlotNode = GetCameraNodeAs<UCameraRigInput1DSlot>();
 	if (SlotNode->InputSlotParameters.bIsPreBlended)
 	{
@@ -104,17 +86,13 @@ void FCameraRigInput1DSlotEvaluator::OnExecuteOperation(const FCameraOperationPa
 
 void FCameraRigInput1DSlotEvaluator::OnSerialize(const FCameraNodeEvaluatorSerializeParams& Params, FArchive& Ar)
 {
+	Super::OnSerialize(Params, Ar);
+
 	Ar << TransientInputValue;
-	Ar << InputValue;
 }
 
 }  // namespace UE::Cameras
 
-FCameraNodeChildrenView UCameraRigInput1DSlot::OnGetChildren()
-{
-	return FCameraNodeChildrenView{ Child };
-}
-
 void UCameraRigInput1DSlot::OnBuild(FCameraRigBuildContext& BuildContext) 
 {
 	using namespace UE::Cameras;
diff --git a/Engine/Plugins/Cameras/GameplayCameras/Source/GameplayCameras/Private/Nodes/Input/CameraRigInput2DSlot.cpp b/Engine/Plugins/Cameras/GameplayCameras/Source/GameplayCameras/Private/Nodes/Input/CameraRigInput2DSlot.cpp
--- a/Engine/Plugins/Cameras/GameplayCameras/Source/GameplayCameras/Private/Nodes/Input/CameraRigInput2DSlot.cpp
+++ b/Engine/Plugins/Cameras/GameplayCameras/Source/GameplayCameras/Private/Nodes/Input/CameraRigInput2DSlot.cpp
@@ -25,12 +25,6 @@ FCameraRigInput2DSlotEvaluator::FCameraRigInput2DSlotEvaluator()
 			ECameraNodeEvaluatorFlags::SupportsOperations);
 }
 
-void FCameraRigInput2DSlotEvaluator::OnBuild(const FCameraNodeEvaluatorBuildParams& Params)
-{
-	const UCameraRigInput2DSlot* SlotNode = GetCameraNodeAs<UCameraRigInput2DSlot>();
-	ChildEvaluator = Params.BuildEvaluatorAs<FInput2DCameraNodeEvaluator>(SlotNode->Child);
-}
-
 void FCameraRigInput2DSlotEvaluator::OnInitialize(const FCameraNodeEvaluatorInitializeParams& Params)
 {
 	TransientInputValue = FVector2d::ZeroVector;
@@ -44,20 +38,8 @@ void FCameraRigInput2DSlotEvaluator::OnInitialize(const FCameraNodeEvaluatorInit
 	}
 }
 
-FCameraNodeEvaluatorChildrenView FCameraRigInput2DSlotEvaluator::OnGetChildren()
-{
-	return FCameraNodeEvaluatorChildrenView{ ChildEvaluator };
-}
-
 void FCameraRigInput2DSlotEvaluator::OnUpdateParameters(const FCameraBlendedParameterUpdateParams& Params, FCameraBlendedParameterUpdateResult& OutResult)
 {
-	if (ChildEvaluator)
-	{
-		ChildEvaluator->UpdateParameters(Params, OutResult);
-
-		TransientInputValue = ChildEvaluator->GetInputValue();
-	}
-
 	const UCameraRigInput2DSlot* SlotNode = GetCameraNodeAs<UCameraRigInput2DSlot>();
 	if (SlotNode->InputSlotParameters.bIsPreBlended)
 	{
@@ -111,17 +93,13 @@ void FCameraRigInput2DSlotEvaluator::OnExecuteOperation(const FCameraOperationPa
 
 void FCameraRigInput2DSlotEvaluator::OnSerialize(const FCameraNodeEvaluatorSerializeParams& Params, FArchive& Ar)
 {
+	Super::OnSerialize(Params, Ar);
+
 	Ar << TransientInputValue;
-	Ar << InputValue;
 }
 
 }  // namespace UE::Cameras
 
-FCameraNodeChildrenView UCameraRigInput2DSlot::OnGetChildren()
-{
-	return FCameraNodeChildrenView{ Child };
-}
-
 void UCameraRigInput2DSlot::OnBuild(FCameraRigBuildContext& BuildContext)
 {
 	using namespace UE::Cameras;
diff --git a/Engine/Plugins/Cameras/GameplayCameras/Source/GameplayCameras/Private/Nodes/Input/Input1DCameraNode.cpp b/Engine/Plugins/Cameras/GameplayCameras/Source/GameplayCameras/Private/Nodes/Input/Input1DCameraNode.cpp
--- a/Engine/Plugins/Cameras/GameplayCameras/Source/GameplayCameras/Private/Nodes/Input/Input1DCameraNode.cpp
+++ b/Engine/Plugins/Cameras/GameplayCameras/Source/GameplayCameras/Private/Nodes/Input/Input1DCameraNode.cpp
@@ -15,5 +15,10 @@ FInput1DCameraNodeEvaluator::FInput1DCameraNodeEvaluator()
 	SetNodeEvaluatorFlags(ECameraNodeEvaluatorFlags::NeedsParameterUpdate);
 }
 
+void FInput1DCameraNodeEvaluator::OnSerialize(const FCameraNodeEvaluatorSerializeParams& Params, FArchive& Ar)
+{
+	Ar << InputValue;
+}
+
 }  // namespace UE::Cameras
 
diff --git a/Engine/Plugins/Cameras/GameplayCameras/Source/GameplayCameras/Private/Nodes/Input/Input2DCameraNode.cpp b/Engine/Plugins/Cameras/GameplayCameras/Source/GameplayCameras/Private/Nodes/Input/Input2DCameraNode.cpp
--- a/Engine/Plugins/Cameras/GameplayCameras/Source/GameplayCameras/Private/Nodes/Input/Input2DCameraNode.cpp
+++ b/Engine/Plugins/Cameras/GameplayCameras/Source/GameplayCameras/Private/Nodes/Input/Input2DCameraNode.cpp
@@ -15,5 +15,10 @@ FInput2DCameraNodeEvaluator::FInput2DCameraNodeEvaluator()
 	SetNodeEvaluatorFlags(ECameraNodeEvaluatorFlags::NeedsParameterUpdate);
 }
 
+void FInput2DCameraNodeEvaluator::OnSerialize(const FCameraNodeEvaluatorSerializeParams& Params, FArchive& Ar)
+{
+	Ar << InputValue;
+}
+
 }  // namespace UE::Cameras
 
diff --git a/Engine/Plugins/Cameras/GameplayCameras/Source/GameplayCameras/Private/Nodes/Input/InputAxisBinding2DCameraNode.cpp b/Engine/Plugins/Cameras/GameplayCameras/Source/GameplayCameras/Private/Nodes/Input/InputAxisBinding2DCameraNode.cpp
--- a/Engine/Plugins/Cameras/GameplayCameras/Source/GameplayCameras/Private/Nodes/Input/InputAxisBinding2DCameraNode.cpp
+++ b/Engine/Plugins/Cameras/GameplayCameras/Source/GameplayCameras/Private/Nodes/Input/InputAxisBinding2DCameraNode.cpp
@@ -6,14 +6,16 @@
 #include "Core/CameraEvaluationContext.h"
 #include "EnhancedInputComponent.h"
 #include "GameFramework/Actor.h"
+#include "GameplayCameras.h"
 #include "InputAction.h"
+#include "UObject/Package.h"
 
 namespace UE::Cameras
 {
 
-class FInputAxisBinding2DCameraNodeEvaluator : public FInput2DCameraNodeEvaluator
+class FInputAxisBinding2DCameraNodeEvaluator : public FCameraRigInput2DSlotEvaluator
 {
-	UE_DECLARE_CAMERA_NODE_EVALUATOR_EX(GAMEPLAYCAMERAS_API, FInputAxisBinding2DCameraNodeEvaluator, FInput2DCameraNodeEvaluator)
+	UE_DECLARE_CAMERA_NODE_EVALUATOR_EX(GAMEPLAYCAMERAS_API, FInputAxisBinding2DCameraNodeEvaluator, FCameraRigInput2DSlotEvaluator)
 
 protected:
 
@@ -30,7 +32,8 @@ UE_DEFINE_CAMERA_NODE_EVALUATOR(FInputAxisBinding2DCameraNodeEvaluator)
 
 void FInputAxisBinding2DCameraNodeEvaluator::OnInitialize(const FCameraNodeEvaluatorInitializeParams& Params)
 {
-	if (UObject* ContextOwner = Params.EvaluationContext->GetOwner())
+	UObject* ContextOwner = Params.EvaluationContext->GetOwner();
+	if (ContextOwner)
 	{
 		if (AActor* OuterActor = ContextOwner->GetTypedOuter<AActor>())
 		{
@@ -47,6 +50,15 @@ void FInputAxisBinding2DCameraNodeEvaluator::OnInitialize(const FCameraNodeEvalu
 			AxisValueBindings.Add(AxisValueBinding);
 		}
 	}
+	else
+	{
+		UE_LOG(LogCameraSystem, Error, TEXT("No input component found on context owner '%s' for node '%s' in '%s'."),
+				*GetNameSafe(ContextOwner), 
+				*GetNameSafe(AxisBindingNode),
+				*GetNameSafe(AxisBindingNode ? AxisBindingNode->GetOutermost() : nullptr));
+	}
+
+	Super::OnInitialize(Params);
 }
 
 void FInputAxisBinding2DCameraNodeEvaluator::OnUpdateParameters(const FCameraBlendedParameterUpdateParams& Params, FCameraBlendedParameterUpdateResult& OutResult)
@@ -71,18 +83,20 @@ void FInputAxisBinding2DCameraNodeEvaluator::OnUpdateParameters(const FCameraBle
 		}
 	}
 
-	InputValue = FVector2d(
+	TransientInputValue = FVector2d(
 			HighestValue.X * AxisBindingNode->Multiplier.X, 
 			HighestValue.Y * AxisBindingNode->Multiplier.Y);
 
 	if (AxisBindingNode->RevertAxisX)
 	{
-		InputValue.X = -InputValue.X;
+		TransientInputValue.X = -TransientInputValue.X;
 	}
 	if (AxisBindingNode->RevertAxisY)
 	{
-		InputValue.Y = -InputValue.Y;
+		TransientInputValue.Y = -TransientInputValue.Y;
 	}
+
+	Super::OnUpdateParameters(Params, OutResult);
 }
 
 }  // namespace UE::Cameras
diff --git a/Engine/Plugins/Cameras/GameplayCameras/Source/GameplayCameras/Public/Nodes/Common/BoomArmCameraNode.h b/Engine/Plugins/Cameras/GameplayCameras/Source/GameplayCameras/Public/Nodes/Common/BoomArmCameraNode.h
--- a/Engine/Plugins/Cameras/GameplayCameras/Source/GameplayCameras/Public/Nodes/Common/BoomArmCameraNode.h
+++ b/Engine/Plugins/Cameras/GameplayCameras/Source/GameplayCameras/Public/Nodes/Common/BoomArmCameraNode.h
@@ -7,7 +7,7 @@
 
 #include "BoomArmCameraNode.generated.h"
 
-class UCameraRigInput2DSlot;
+class UInput2DCameraNode;
 
 /**
  * A camera node that can rotate the camera in yaw and pitch based on player input.
@@ -34,6 +34,6 @@ public:
 	 * If no input slot is specified, the boom arm will use the player controller view rotation.
 	 */
 	UPROPERTY(meta=(ObjectTreeGraphPinDirection=Input))
-	TObjectPtr<UCameraRigInput2DSlot> InputSlot;
+	TObjectPtr<UInput2DCameraNode> InputSlot;
 };
 
diff --git a/Engine/Plugins/Cameras/GameplayCameras/Source/GameplayCameras/Public/Nodes/Input/CameraRigInput1DSlot.h b/Engine/Plugins/Cameras/GameplayCameras/Source/GameplayCameras/Public/Nodes/Input/CameraRigInput1DSlot.h
--- a/Engine/Plugins/Cameras/GameplayCameras/Source/GameplayCameras/Public/Nodes/Input/CameraRigInput1DSlot.h
+++ b/Engine/Plugins/Cameras/GameplayCameras/Source/GameplayCameras/Public/Nodes/Input/CameraRigInput1DSlot.h
@@ -3,24 +3,20 @@
 #pragma once
 
 #include "Core/BuiltInCameraVariables.h"
-#include "Core/CameraNode.h"
 #include "Core/CameraNodeEvaluator.h"
 #include "Nodes/Input/CameraRigInputSlotTypes.h"
+#include "Nodes/Input/Input1DCameraNode.h"
 #include "UObject/ObjectPtr.h"
 
 #include "CameraRigInput1DSlot.generated.h"
 
 class UDoubleCameraVariable;
-class UInput1DCameraNode;
 
 /**
- * A node that can handle and accumulate a chain of player input nodes.
+ * The base class for a node that can handle and accumulate raw player input values.
  */
-UCLASS(MinimalAPI, meta=(
-			CameraNodeCategories="Input", 
-			ObjectTreeGraphSelfPinDirection="Output", 
-			ObjectTreeGraphDefaultPropertyPinDirection="Input"))
-class UCameraRigInput1DSlot : public UCameraNode
+UCLASS(Abstract)
+class UCameraRigInput1DSlot : public UInput1DCameraNode
 {
 	GENERATED_BODY()
 
@@ -46,10 +42,6 @@ public:
 	UPROPERTY(EditAnywhere, Category="Input", meta=(EditCondition="BuiltInVariable == EBuiltInDoubleCameraVariable::None"))
 	TObjectPtr<UDoubleCameraVariable> Variable;
 
-	/** A node providing an incremental value. */
-	UPROPERTY()
-	TObjectPtr<UInput1DCameraNode> Child;
-
 public:
 
 	FCameraVariableID GetVariableID() const { return VariableID; }
@@ -58,7 +50,6 @@ public:
 protected:
 
 	// UCameraNode interface.
-	virtual FCameraNodeChildrenView OnGetChildren() override;
 	virtual void OnBuild(FCameraRigBuildContext& BuildContext) override;
 	virtual FCameraNodeEvaluatorPtr OnBuildEvaluator(FCameraNodeEvaluatorBuilder& Builder) const override;
 
@@ -75,22 +66,18 @@ namespace UE::Cameras
 
 class FInput1DCameraNodeEvaluator;
 
-class FCameraRigInput1DSlotEvaluator : public FCameraNodeEvaluator
+class FCameraRigInput1DSlotEvaluator : public FInput1DCameraNodeEvaluator
 {
-	UE_DECLARE_CAMERA_NODE_EVALUATOR(GAMEPLAYCAMERAS_API, FCameraRigInput1DSlotEvaluator)
+	UE_DECLARE_CAMERA_NODE_EVALUATOR_EX(GAMEPLAYCAMERAS_API, FCameraRigInput1DSlotEvaluator, FInput1DCameraNodeEvaluator)
 
 public:
 
 	FCameraRigInput1DSlotEvaluator();
 
-	double GetInputValue() const { return InputValue; }
-
 protected:
 
 	// FCameraNodeEvaluator interface.
-	virtual void OnBuild(const FCameraNodeEvaluatorBuildParams& Params) override;
 	virtual void OnInitialize(const FCameraNodeEvaluatorInitializeParams& Params) override;
-	virtual FCameraNodeEvaluatorChildrenView OnGetChildren() override;
 	virtual void OnUpdateParameters(const FCameraBlendedParameterUpdateParams& Params, FCameraBlendedParameterUpdateResult& OutResult) override;
 	virtual void OnRun(const FCameraNodeEvaluationParams& Params, FCameraNodeEvaluationResult& OutResult) override;
 	virtual void OnExecuteOperation(const FCameraOperationParams& Params, FCameraOperation& Operation) override;
@@ -98,10 +85,7 @@ protected:
 
 protected:
 
-	FInput1DCameraNodeEvaluator* ChildEvaluator;
-
 	double TransientInputValue = 0.f;
-	double InputValue = 0.f;
 };
 
 }  // namespace UE::Cameras
diff --git a/Engine/Plugins/Cameras/GameplayCameras/Source/GameplayCameras/Public/Nodes/Input/CameraRigInput2DSlot.h b/Engine/Plugins/Cameras/GameplayCameras/Source/GameplayCameras/Public/Nodes/Input/CameraRigInput2DSlot.h
--- a/Engine/Plugins/Cameras/GameplayCameras/Source/GameplayCameras/Public/Nodes/Input/CameraRigInput2DSlot.h
+++ b/Engine/Plugins/Cameras/GameplayCameras/Source/GameplayCameras/Public/Nodes/Input/CameraRigInput2DSlot.h
@@ -3,24 +3,20 @@
 #pragma once
 
 #include "Core/BuiltInCameraVariables.h"
-#include "Core/CameraNode.h"
 #include "Core/CameraNodeEvaluator.h"
 #include "Nodes/Input/CameraRigInputSlotTypes.h"
+#include "Nodes/Input/Input2DCameraNode.h"
 #include "UObject/ObjectPtr.h"
 
 #include "CameraRigInput2DSlot.generated.h"
 
-class UInput2DCameraNode;
 class UVector2dCameraVariable;
 
 /**
- * A node that can handle and accumulate a chain of player input nodes.
+ * The base class for a node that can handle and accumulate raw player input values.
  */
-UCLASS(MinimalAPI, meta=(
-			CameraNodeCategories="Input", 
-			ObjectTreeGraphSelfPinDirection="Output",
-			ObjectTreeGraphDefaultPropertyPinDirection="Input"))
-class UCameraRigInput2DSlot : public UCameraNode
+UCLASS(Abstract)
+class UCameraRigInput2DSlot : public UInput2DCameraNode
 {
 	GENERATED_BODY()
 
@@ -54,10 +50,6 @@ public:
 	UPROPERTY(EditAnywhere, Category="Input", meta=(EditCondition="BuiltInVariable == EBuiltInVector2dCameraVariable::None"))
 	TObjectPtr<UVector2dCameraVariable> Variable;
 
-	/** A node providing an incremental value. */
-	UPROPERTY()
-	TObjectPtr<UInput2DCameraNode> Child;
-
 public:
 
 	FCameraVariableID GetVariableID() const { return VariableID; }
@@ -66,7 +58,6 @@ public:
 protected:
 
 	// UCameraNode interface.
-	virtual FCameraNodeChildrenView OnGetChildren() override;
 	virtual void OnBuild(FCameraRigBuildContext& BuildContext) override;
 	virtual FCameraNodeEvaluatorPtr OnBuildEvaluator(FCameraNodeEvaluatorBuilder& Builder) const override;
 
@@ -83,22 +74,18 @@ namespace UE::Cameras
 
 class FInput2DCameraNodeEvaluator;
 
-class FCameraRigInput2DSlotEvaluator : public FCameraNodeEvaluator
+class FCameraRigInput2DSlotEvaluator : public FInput2DCameraNodeEvaluator
 {
-	UE_DECLARE_CAMERA_NODE_EVALUATOR(GAMEPLAYCAMERAS_API, FCameraRigInput2DSlotEvaluator)
+	UE_DECLARE_CAMERA_NODE_EVALUATOR_EX(GAMEPLAYCAMERAS_API, FCameraRigInput2DSlotEvaluator, FInput2DCameraNodeEvaluator)
 
 public:
 
 	FCameraRigInput2DSlotEvaluator();
 
-	FVector2d GetInputValue() const { return InputValue; }
-
 protected:
 
 	// FCameraNodeEvaluator interface.
-	virtual void OnBuild(const FCameraNodeEvaluatorBuildParams& Params) override;
 	virtual void OnInitialize(const FCameraNodeEvaluatorInitializeParams& Params) override;
-	virtual FCameraNodeEvaluatorChildrenView OnGetChildren() override;
 	virtual void OnUpdateParameters(const FCameraBlendedParameterUpdateParams& Params, FCameraBlendedParameterUpdateResult& OutResult) override;
 	virtual void OnRun(const FCameraNodeEvaluationParams& Params, FCameraNodeEvaluationResult& OutResult) override;
 	virtual void OnExecuteOperation(const FCameraOperationParams& Params, FCameraOperation& Operation) override;
@@ -106,12 +93,8 @@ protected:
 
 protected:
 
-	FInput2DCameraNodeEvaluator* ChildEvaluator;
-
 	FVector2d TransientInputValue = FVector2d::ZeroVector;
-	FVector2d InputValue = FVector2d::ZeroVector;
 };
 
 }  // namespace UE::Cameras
 
-
diff --git a/Engine/Plugins/Cameras/GameplayCameras/Source/GameplayCameras/Public/Nodes/Input/Input1DCameraNode.h b/Engine/Plugins/Cameras/GameplayCameras/Source/GameplayCameras/Public/Nodes/Input/Input1DCameraNode.h
--- a/Engine/Plugins/Cameras/GameplayCameras/Source/GameplayCameras/Public/Nodes/Input/Input1DCameraNode.h
+++ b/Engine/Plugins/Cameras/GameplayCameras/Source/GameplayCameras/Public/Nodes/Input/Input1DCameraNode.h
@@ -10,7 +10,7 @@
 /**
  * A camera node that provides a floating-point value to an input slot.
  */
-UCLASS(MinimalAPI, meta=(
+UCLASS(Abstract, MinimalAPI, meta=(
 			CameraNodeCategories="Input",
 			ObjectTreeGraphSelfPinDirection="Output",
 			ObjectTreeGraphDefaultPropertyPinDirection="Input"))
@@ -38,7 +38,12 @@ public:
 
 protected:
 
-	float InputValue;
+	// FCameraNodeEvaluator interface.
+	virtual void OnSerialize(const FCameraNodeEvaluatorSerializeParams& Params, FArchive& Ar) override;
+
+protected:
+
+	double InputValue;
 };
 
 }  // namespace UE::Cameras
diff --git a/Engine/Plugins/Cameras/GameplayCameras/Source/GameplayCameras/Public/Nodes/Input/Input2DCameraNode.h b/Engine/Plugins/Cameras/GameplayCameras/Source/GameplayCameras/Public/Nodes/Input/Input2DCameraNode.h
--- a/Engine/Plugins/Cameras/GameplayCameras/Source/GameplayCameras/Public/Nodes/Input/Input2DCameraNode.h
+++ b/Engine/Plugins/Cameras/GameplayCameras/Source/GameplayCameras/Public/Nodes/Input/Input2DCameraNode.h
@@ -11,7 +11,7 @@
 /**
  * A camera node that provides a two-dimensional value to an input slot.
  */
-UCLASS(MinimalAPI, meta=(
+UCLASS(Abstract, MinimalAPI, meta=(
 			CameraNodeCategories="Input",
 			ObjectTreeGraphSelfPinDirection="Output",
 			ObjectTreeGraphDefaultPropertyPinDirection="Input"))
@@ -37,6 +37,11 @@ public:
 	/** Get the current input value. */
 	FVector2d GetInputValue() const { return InputValue; }
 
+protected:
+
+	// FCameraNodeEvaluator interface.
+	virtual void OnSerialize(const FCameraNodeEvaluatorSerializeParams& Params, FArchive& Ar) override;
+
 protected:
 
 	FVector2d InputValue;
diff --git a/Engine/Plugins/Cameras/GameplayCameras/Source/GameplayCameras/Public/Nodes/Input/InputAxisBinding2DCameraNode.h b/Engine/Plugins/Cameras/GameplayCameras/Source/GameplayCameras/Public/Nodes/Input/InputAxisBinding2DCameraNode.h
--- a/Engine/Plugins/Cameras/GameplayCameras/Source/GameplayCameras/Public/Nodes/Input/InputAxisBinding2DCameraNode.h
+++ b/Engine/Plugins/Cameras/GameplayCameras/Source/GameplayCameras/Public/Nodes/Input/InputAxisBinding2DCameraNode.h
@@ -3,7 +3,7 @@
 #pragma once
 
 #include "Math/MathFwd.h"
-#include "Nodes/Input/Input2DCameraNode.h"
+#include "Nodes/Input/CameraRigInput2DSlot.h"
 
 #include "InputAxisBinding2DCameraNode.generated.h"
 
@@ -13,7 +13,7 @@ class UInputAction;
  * An input node that reads player input from an input action.
  */
 UCLASS(MinimalAPI, meta=(CameraNodeCategories="Input"))
-class UInputAxisBinding2DCameraNode : public UInput2DCameraNode
+class UInputAxisBinding2DCameraNode : public UCameraRigInput2DSlot
 {
 	GENERATED_BODY()
 
