Cameras: add credit in critical damper math explanation

diff --git a/Engine/Plugins/Cameras/GameplayCameras/Source/GameplayCameras/Private/Math/CriticalDamper.cpp b/Engine/Plugins/Cameras/GameplayCameras/Source/GameplayCameras/Private/Math/CriticalDamper.cpp
--- a/Engine/Plugins/Cameras/GameplayCameras/Source/GameplayCameras/Private/Math/CriticalDamper.cpp
+++ b/Engine/Plugins/Cameras/GameplayCameras/Source/GameplayCameras/Private/Math/CriticalDamper.cpp
@@ -155,7 +155,10 @@ void FCriticalDamper::InternalUpdate(float ForcedMovement, float DeltaTime)
 		//
 		//		x = P0 + P1*t
 		//
-		// To test this theory, let's derive this form and plug the result in the NHE:
+		// To test this theory, let's derive this form and plug the result in the NHE.
+		// By the way, I'd like to thank Matt Peters for this extra element that is 
+		// easily missed from spring/mass system equations (because most solutions found 
+		// online don't need to handle a moving target, but we do).
 		//
 		//		x' = P1
 		//		x'' = 0
