Cameras: add the framework for value interpolators, along with 5 default/built-in ones

diff --git a/Engine/Plugins/Cameras/GameplayCameras/Source/GameplayCameras/Private/Core/CameraValueInterpolator.cpp b/Engine/Plugins/Cameras/GameplayCameras/Source/GameplayCameras/Private/Core/CameraValueInterpolator.cpp
new file mode 100644
--- /dev/null
+++ b/Engine/Plugins/Cameras/GameplayCameras/Source/GameplayCameras/Private/Core/CameraValueInterpolator.cpp
@@ -0,0 +1,23 @@
+// Copyright Epic Games, Inc. All Rights Reserved.
+
+#include "Core/CameraValueInterpolator.h"
+
+#include UE_INLINE_GENERATED_CPP_BY_NAME(CameraValueInterpolator)
+
+TUniquePtr<UE::Cameras::TCameraValueInterpolator<double>> UCameraValueInterpolator::BuildDoubleInterpolator() const
+{
+	return OnBuildDoubleInterpolator();
+}
+
+TUniquePtr<UE::Cameras::TCameraValueInterpolator<FVector2d>> UCameraValueInterpolator::BuildVector2dInterpolator() const
+{
+	return OnBuildVector2dInterpolator();
+}
+
+TUniquePtr<UE::Cameras::TCameraValueInterpolator<FVector3d>> UCameraValueInterpolator::BuildVector3dInterpolator() const
+{
+	return OnBuildVector3dInterpolator();
+}
+
+UE_DEFINE_CAMERA_VALUE_INTERPOLATOR_GENERIC(UPopValueInterpolator, UE::Cameras::TPopValueInterpolator)
+
diff --git a/Engine/Plugins/Cameras/GameplayCameras/Source/GameplayCameras/Private/ValueInterpolators/AccelerationDecelerationValueInterpolator.cpp b/Engine/Plugins/Cameras/GameplayCameras/Source/GameplayCameras/Private/ValueInterpolators/AccelerationDecelerationValueInterpolator.cpp
new file mode 100644
--- /dev/null
+++ b/Engine/Plugins/Cameras/GameplayCameras/Source/GameplayCameras/Private/ValueInterpolators/AccelerationDecelerationValueInterpolator.cpp
@@ -0,0 +1,104 @@
+// Copyright Epic Games, Inc. All Rights Reserved.
+
+#include "ValueInterpolators/AccelerationDecelerationValueInterpolator.h"
+
+#include UE_INLINE_GENERATED_CPP_BY_NAME(AccelerationDecelerationValueInterpolator)
+
+namespace UE::Cameras
+{
+
+template<typename ValueType>
+class TAccelerationDecelerationValueInterpolator : public TCameraValueInterpolator<ValueType>
+{
+public:
+
+	TAccelerationDecelerationValueInterpolator(const UAccelerationDecelerationValueInterpolator* InInterpolator)
+		: TCameraValueInterpolator<ValueType>(InInterpolator)
+	{}
+
+protected:
+
+	virtual void OnRun(const FCameraValueInterpolationParams& Params, FCameraValueInterpolationResult& OutResult) override;
+	virtual void OnSerialize(const FCameraValueInterpolatorSerializeParams& Params, FArchive& Ar) override;
+
+private:
+
+	float CurrentSpeed = 0.f;
+};
+
+template<typename ValueType>
+void TAccelerationDecelerationValueInterpolator<ValueType>::OnRun(const FCameraValueInterpolationParams& Params, FCameraValueInterpolationResult& OutResult)
+{
+	// We need to start decelerating when it would take us to the target value.
+	//
+	// The equation is:
+	//
+	//    v1 = a*t + v0
+	//
+	// Where:
+	//    v1 is the next speed
+	//    v0 is the current speed
+	//    a is the acceleration (negative in the case of deceleration)
+	// 
+	// If t0 is the time at which we stop (v1 = 0), then:
+	//
+	//    0 = a*t0 + v0
+	//    t0 = -v0/a
+	//
+	// Distance travelled over that time is:
+	//
+	//    d = a/2*t^2 + v*t
+	//
+	// Let's call d0 the distance travelled before we stop:
+	//
+	//    d0 = a/2*t0^2 + v0*t0
+	//    d0 = a/2*(-v0/a)^2 + v0*(-v0)/a
+	//    d0 = v0^2/(2*a) - v0^2/a
+	//    d0 = (v0^2 - (2*v0^2))/(2*a)
+	//    d0 = -v0^2/(2*a)
+	//
+	// So the speed at which we should go when it's time to decelerate is:
+	//
+	//    (2*a)*d0 = -v0^2
+	//    sqrt(-2*a*d0) = v0
+	//
+	const double DistanceToTarget = TCameraValueInterpolationTraits<ValueType>::Distance(this->CurrentValue, this->TargetValue);
+	if (DistanceToTarget > 0)
+	{
+		const UAccelerationDecelerationValueInterpolator* SpeedInterpolator = this->template GetParametersAs<UAccelerationDecelerationValueInterpolator>();
+		const double MaxSpeed = FMath::Min(
+				SpeedInterpolator->MaxSpeed, 
+				FMath::Sqrt(2.0 * SpeedInterpolator->Deceleration * DistanceToTarget));
+		CurrentSpeed = FMath::Min(CurrentSpeed + SpeedInterpolator->Acceleration * Params.DeltaTime, MaxSpeed);
+
+		const double DistanceThisFrame = CurrentSpeed * Params.DeltaTime;
+		if (DistanceThisFrame >= DistanceToTarget)
+		{
+			CurrentSpeed = 0;
+			this->CurrentValue = this->TargetValue;
+			this->bIsFinished = true;
+		}
+		else
+		{
+			const ValueType Direction = TCameraValueInterpolationTraits<ValueType>::Direction(this->CurrentValue, this->TargetValue);
+			this->CurrentValue = this->CurrentValue + Direction * DistanceThisFrame;
+		}
+	}
+	else
+	{
+		this->CurrentSpeed = 0;
+		this->CurrentValue = this->TargetValue;
+		this->bIsFinished = true;
+	}
+}
+
+template<typename ValueType>
+void TAccelerationDecelerationValueInterpolator<ValueType>::OnSerialize(const FCameraValueInterpolatorSerializeParams& Params, FArchive& Ar)
+{
+	Ar << CurrentSpeed;
+}
+
+}  // namespace UE::Cameras
+
+UE_DEFINE_CAMERA_VALUE_INTERPOLATOR_GENERIC(UAccelerationDecelerationValueInterpolator, UE::Cameras::TAccelerationDecelerationValueInterpolator)
+
diff --git a/Engine/Plugins/Cameras/GameplayCameras/Source/GameplayCameras/Private/ValueInterpolators/CriticalDamperValueInterpolator.cpp b/Engine/Plugins/Cameras/GameplayCameras/Source/GameplayCameras/Private/ValueInterpolators/CriticalDamperValueInterpolator.cpp
new file mode 100644
--- /dev/null
+++ b/Engine/Plugins/Cameras/GameplayCameras/Source/GameplayCameras/Private/ValueInterpolators/CriticalDamperValueInterpolator.cpp
@@ -0,0 +1,69 @@
+// Copyright Epic Games, Inc. All Rights Reserved.
+
+#include "ValueInterpolators/CriticalDamperValueInterpolator.h"
+
+#include UE_INLINE_GENERATED_CPP_BY_NAME(CriticalDamperValueInterpolator)
+
+namespace UE::Cameras
+{
+
+template<typename ValueType>
+class TCriticalDamperValueInterpolator : public TCameraValueInterpolator<ValueType>
+{
+public:
+
+	using ValueTypeParam = TCameraValueInterpolator<ValueType>::ValueTypeParam;
+
+	TCriticalDamperValueInterpolator(const UCriticalDamperValueInterpolator* InInterpolator)
+		: TCameraValueInterpolator<ValueType>(InInterpolator)
+	{}
+
+protected:
+
+	virtual void OnReset(ValueTypeParam OldCurrentValue, ValueTypeParam OldTargetValue) override;
+	virtual void OnRun(const FCameraValueInterpolationParams& Params, FCameraValueInterpolationResult& OutResult) override;
+	virtual void OnSerialize(const FCameraValueInterpolatorSerializeParams& Params, FArchive& Ar) override;
+
+private:
+
+	UE::Cameras::FCriticalDamper Damper;
+	bool bIsFirstFrame = true;
+};
+
+template<typename ValueType>
+void TCriticalDamperValueInterpolator<ValueType>::OnReset(ValueTypeParam OldCurrentValue, ValueTypeParam OldTargetValue)
+{
+	const double DistanceToTarget = TCameraValueInterpolationTraits<ValueType>::Distance(this->CurrentValue, this->TargetValue);
+	Damper.Reset(DistanceToTarget, Damper.GetX0Derivative());
+}
+
+template<typename ValueType>
+void TCriticalDamperValueInterpolator<ValueType>::OnRun(const FCameraValueInterpolationParams& Params, FCameraValueInterpolationResult& OutResult)
+{
+	const double DistanceToTarget = TCameraValueInterpolationTraits<ValueType>::Distance(this->CurrentValue, this->TargetValue);
+	if (bIsFirstFrame)
+	{
+		const UCriticalDamperValueInterpolator* CriticalDamperInterpolator = this->template GetParametersAs<UCriticalDamperValueInterpolator>();
+		Damper.SetW0(CriticalDamperInterpolator->DampingFactor);
+		Damper.Reset(DistanceToTarget, 0);
+		bIsFirstFrame = false;
+	}
+
+	const double NextDistanceToTarget = Damper.Update(DistanceToTarget, Params.DeltaTime);
+	const double ClosingDistance = DistanceToTarget - NextDistanceToTarget;
+	const ValueType Direction = TCameraValueInterpolationTraits<ValueType>::Direction(this->CurrentValue, this->TargetValue);
+	this->CurrentValue = this->CurrentValue + Direction * ClosingDistance;
+	this->bIsFinished = ClosingDistance <= UE_DOUBLE_SMALL_NUMBER;
+}
+
+template<typename ValueType>
+void TCriticalDamperValueInterpolator<ValueType>::OnSerialize(const FCameraValueInterpolatorSerializeParams& Params, FArchive& Ar)
+{
+	Ar << Damper;
+	Ar << bIsFirstFrame;
+}
+
+}  // namespace UE::Cameras
+
+UE_DEFINE_CAMERA_VALUE_INTERPOLATOR_GENERIC(UCriticalDamperValueInterpolator, UE::Cameras::TCriticalDamperValueInterpolator)
+
diff --git a/Engine/Plugins/Cameras/GameplayCameras/Source/GameplayCameras/Private/ValueInterpolators/DoubleIIRValueInterpolator.cpp b/Engine/Plugins/Cameras/GameplayCameras/Source/GameplayCameras/Private/ValueInterpolators/DoubleIIRValueInterpolator.cpp
new file mode 100644
--- /dev/null
+++ b/Engine/Plugins/Cameras/GameplayCameras/Source/GameplayCameras/Private/ValueInterpolators/DoubleIIRValueInterpolator.cpp
@@ -0,0 +1,106 @@
+// Copyright Epic Games, Inc. All Rights Reserved.
+
+#include "ValueInterpolators/DoubleIIRValueInterpolator.h"
+
+#include "ValueInterpolators/IIRValueInterpolator.h"
+
+#include UE_INLINE_GENERATED_CPP_BY_NAME(DoubleIIRValueInterpolator)
+
+namespace UE::Cameras
+{
+
+template<typename ValueType>
+class TDoubleIIRValueInterpolator : public TCameraValueInterpolator<ValueType>
+{
+public:
+
+	TDoubleIIRValueInterpolator(const UDoubleIIRValueInterpolator* InParameters)
+		: TCameraValueInterpolator<ValueType>(InParameters)
+		, IntermediateInterpolator(InParameters->IntermediateSpeed, InParameters->bUseFixedStep)
+		, PrimaryInterpolator(InParameters->PrimarySpeed, InParameters->bUseFixedStep)
+		, bUseFixedStep(InParameters->bUseFixedStep)
+	{}
+
+protected:
+
+	using ValueTypeParam = TCameraValueInterpolator<ValueType>::ValueTypeParam;
+
+	virtual void OnReset(ValueTypeParam OldCurrentValue, ValueTypeParam OldTargetValue) override
+	{
+		IntermediateInterpolator.Reset(this->CurrentValue, this->TargetValue);
+		PrimaryInterpolator.Reset(this->CurrentValue, this->TargetValue);
+	}
+
+	virtual void OnRun(const FCameraValueInterpolationParams& Params, FCameraValueInterpolationResult& OutResult) override
+	{
+		if (bUseFixedStep)
+		{
+			float RemainingTime = Params.DeltaTime;
+
+			// Move the substep target value linearly toward the real target value while we evaluate the substeps.
+			ValueType LastToTargetValue(this->TargetValue - LastTargetValue);
+			if constexpr (std::is_same_v<FRotator3d, ValueType>)
+			{
+				LastToTargetValue = LastToTargetValue.GetNormalized();
+			}
+			const ValueType EquilibriumStepRate = LastToTargetValue * (1.f / RemainingTime);
+
+			ValueType LerpedTargetValue = LastTargetValue;
+
+			while (RemainingTime > KINDA_SMALL_NUMBER)
+			{
+				const float StepTime = FMath::Min(MaxSubstepTime, RemainingTime);
+
+				LerpedTargetValue += EquilibriumStepRate * StepTime;
+				RemainingTime -= StepTime;
+
+				RunSubstep(LerpedTargetValue, StepTime, OutResult);
+			}
+
+			this->CurrentValue = PrimaryInterpolator.GetCurrentValue();
+			LastTargetValue = this->TargetValue;
+		}
+		else
+		{
+			this->CurrentValue = RunSubstep(this->TargetValue, Params.DeltaTime, OutResult);
+		}
+	}
+
+	virtual void OnSerialize(const FCameraValueInterpolatorSerializeParams& Params, FArchive& Ar) override
+	{
+		IntermediateInterpolator.Serialize(Params, Ar);
+		PrimaryInterpolator.Serialize(Params, Ar);
+
+		Ar << LastTargetValue;
+		Ar << bUseFixedStep;
+	}
+
+private:
+
+	ValueType RunSubstep(ValueType SubstepTargetValue, float SubstepDeltaTime, FCameraValueInterpolationResult& OutResult)
+	{
+		FCameraValueInterpolationParams SubParams;
+		SubParams.DeltaTime = SubstepDeltaTime;
+
+		IntermediateInterpolator.Reset(IntermediateInterpolator.GetCurrentValue(), SubstepTargetValue);
+		ValueType IntermediateValue = IntermediateInterpolator.Run(SubParams, OutResult);
+
+		PrimaryInterpolator.Reset(PrimaryInterpolator.GetCurrentValue(), IntermediateValue);
+		return PrimaryInterpolator.Run(SubParams, OutResult);
+	}
+	
+private:
+
+	static constexpr float MaxSubstepTime = 1.f / 120.f;
+
+	TIIRValueInterpolator<ValueType> IntermediateInterpolator;
+	TIIRValueInterpolator<ValueType> PrimaryInterpolator;
+
+	ValueType LastTargetValue;
+	bool bUseFixedStep = false;
+};
+
+}  // namespace UE::Cameras
+
+UE_DEFINE_CAMERA_VALUE_INTERPOLATOR_GENERIC(UDoubleIIRValueInterpolator, UE::Cameras::TDoubleIIRValueInterpolator)
+
diff --git a/Engine/Plugins/Cameras/GameplayCameras/Source/GameplayCameras/Private/ValueInterpolators/IIRValueInterpolator.cpp b/Engine/Plugins/Cameras/GameplayCameras/Source/GameplayCameras/Private/ValueInterpolators/IIRValueInterpolator.cpp
new file mode 100644
--- /dev/null
+++ b/Engine/Plugins/Cameras/GameplayCameras/Source/GameplayCameras/Private/ValueInterpolators/IIRValueInterpolator.cpp
@@ -0,0 +1,8 @@
+// Copyright Epic Games, Inc. All Rights Reserved.
+
+#include "ValueInterpolators/IIRValueInterpolator.h"
+
+#include UE_INLINE_GENERATED_CPP_BY_NAME(IIRValueInterpolator)
+
+UE_DEFINE_CAMERA_VALUE_INTERPOLATOR_GENERIC(UIIRValueInterpolator, UE::Cameras::TIIRValueInterpolator)
+
diff --git a/Engine/Plugins/Cameras/GameplayCameras/Source/GameplayCameras/Public/Core/CameraValueInterpolator.h b/Engine/Plugins/Cameras/GameplayCameras/Source/GameplayCameras/Public/Core/CameraValueInterpolator.h
new file mode 100644
--- /dev/null
+++ b/Engine/Plugins/Cameras/GameplayCameras/Source/GameplayCameras/Public/Core/CameraValueInterpolator.h
@@ -0,0 +1,275 @@
+// Copyright Epic Games, Inc. All Rights Reserved.
+
+#pragma once
+
+#include "UObject/ObjectMacros.h"
+#include "Serialization/Archive.h"
+#include "Math/Vector.h"
+#include "Math/Vector2D.h"
+
+#include "CameraValueInterpolator.generated.h"
+
+class FArchive;
+class UCameraValueInterpolator;
+
+namespace UE::Cameras
+{
+
+class FCameraVariableTable;
+
+/** 
+ * Parameter structure for updating a value interpolator.
+ */
+struct FCameraValueInterpolationParams
+{
+	float DeltaTime = 0.f;
+	bool bIsCameraCut = false;
+};
+
+/**
+ * Result structure for updating a value interpolator.
+ */
+struct FCameraValueInterpolationResult
+{
+	FCameraValueInterpolationResult(FCameraVariableTable& InVariableTable)
+		: VariableTable(InVariableTable)
+	{}
+
+	FCameraVariableTable& VariableTable;
+};
+
+/**
+ * Parameter structure for serializing a value interpolator.
+ */
+struct FCameraValueInterpolatorSerializeParams
+{
+};
+
+/**
+ * A value interpolator is a stand-alone object that can interpolate a given value
+ * towards a given target value. Various sub-classes of this base class should 
+ * implement various interpolation algorithms.
+ */
+template<typename ValueType>
+class TCameraValueInterpolator
+{
+public:
+
+	using ValueTypeParam = TCallTraits<ValueType>::ParamType;
+
+	/** 
+	 * Creates a new value interpolator given a piece of data containing user-defined
+	 * settings.
+	 */
+	TCameraValueInterpolator(const UCameraValueInterpolator* InParameters)
+		: Parameters(InParameters)
+		, CurrentValue(ValueType())
+		, TargetValue(ValueType())
+		, bIsFinished(false)
+	{}
+
+	/** Destructor */
+	virtual ~TCameraValueInterpolator() {}
+
+	/** Gets the parameters class and casts it to the given sub-class.  */
+	template<typename TInterpolatorClass>
+	const TInterpolatorClass* GetParametersAs() const { return Cast<TInterpolatorClass>(Parameters); }
+
+	/** Gets the current value. */
+	ValueType GetCurrentValue() const { return CurrentValue; }
+	
+	/** Gets the target value. */
+	ValueType GetTargetValue() const { return TargetValue; }
+
+	/** 
+	 * Returns whether the interpolation has ended. 
+	 * The current value may not be equal to the target value if for some reason the 
+	 * implemented interpolation algorithm cannot reach the target value.
+	 */
+	bool IsFinished() const { return bIsFinished; }
+
+	/** Resets the current and target values. */
+	void Reset(ValueTypeParam NewCurrentValue, ValueTypeParam NewTargetValue)
+	{
+		ValueType OldCurrentValue = CurrentValue;
+		ValueType OldTargetValue = TargetValue;
+
+		CurrentValue = NewCurrentValue;
+		TargetValue = NewTargetValue;
+
+		OnReset(OldCurrentValue, OldTargetValue);
+	}
+
+	/** Evaluates the interpolator, advancing the current value towards the target. */
+	GAMEPLAYCAMERAS_API ValueType Run(const FCameraValueInterpolationParams& Params, FCameraValueInterpolationResult& OutResult)
+	{
+		OnRun(Params, OutResult);
+		return CurrentValue;
+	}
+
+	/** Serializes the value interpolator to/from the given archive. */
+	GAMEPLAYCAMERAS_API void Serialize(const FCameraValueInterpolatorSerializeParams& Params, FArchive& Ar)
+	{
+		Ar << CurrentValue;
+		Ar << TargetValue;
+		Ar << bIsFinished;
+
+		OnSerialize(Params, Ar);
+	}
+
+protected:
+
+	/** Called when the current and/or target values are forcibly changed. */
+	virtual void OnReset(ValueTypeParam OldCurrentValue, ValueTypeParam OldTargetValue) {}
+	/** Evaluates the interpolator, advancing the current value towards the target. */
+	virtual void OnRun(const FCameraValueInterpolationParams& Params, FCameraValueInterpolationResult& OutResult) {}
+	/** Serializes the value interpolator to/from the given archive. */
+	virtual void OnSerialize(const FCameraValueInterpolatorSerializeParams& Params, FArchive& Ar) {}
+
+protected:
+
+	const UCameraValueInterpolator* Parameters;
+	ValueType CurrentValue;
+	ValueType TargetValue;
+	bool bIsFinished;
+};
+
+/**
+ * Generic traits for interpolated value types.
+ * Empty by default, overriden for specifically supported value types.
+ */
+template<typename ValueType>
+struct TCameraValueInterpolationTraits
+{};
+
+template<>
+struct TCameraValueInterpolationTraits<double>
+{
+	static double Distance(double A, double B)
+	{
+		return FMath::Abs(B - A);
+	}
+
+	static double Direction(double A, double B)
+	{
+		return A < B ? 1 : -1;
+	}
+};
+
+template<>
+struct TCameraValueInterpolationTraits<FVector2d>
+{
+	static double Distance(const FVector2d& A, const FVector2d& B)
+	{
+		return FVector2d::Distance(A, B);
+	}
+
+	static FVector2d Direction(const FVector2d& A, const FVector2d& B)
+	{
+		return (B - A).GetSafeNormal();
+	}
+};
+
+template<>
+struct TCameraValueInterpolationTraits<FVector3d>
+{
+	static double Distance(const FVector3d& A, const FVector3d& B)
+	{
+		return FVector3d::Distance(A, B);
+	}
+
+	static FVector3d Direction(const FVector3d& A, const FVector3d& B)
+	{
+		return (B - A).GetSafeNormal();
+	}
+};
+
+using FCameraDoubleValueInterpolator = TCameraValueInterpolator<double>;
+using FCameraVector2dValueInterpolator = TCameraValueInterpolator<FVector2d>;
+using FCameraVector3dDoubleValueInterpolator = TCameraValueInterpolator<FVector3d>;
+
+/** A simple value interpolator that immediately "pops" to the target value. */
+template<typename ValueType>
+class TPopValueInterpolator : public TCameraValueInterpolator<ValueType>
+{
+public:
+
+	TPopValueInterpolator()
+		: TCameraValueInterpolator<ValueType>(nullptr)
+	{}
+
+	TPopValueInterpolator(const UPopValueInterpolator* InInterpolator)
+		: TCameraValueInterpolator<ValueType>(InInterpolator)
+	{}
+
+protected:
+
+	virtual void OnRun(const FCameraValueInterpolationParams& Params, FCameraValueInterpolationResult& OutResult) override
+	{
+		this->CurrentValue = this->TargetValue;
+		this->bIsFinished = true;
+	}
+};
+
+}  // namespace UE::Cameras
+
+/**
+ * Base class for value interpolator parameters.
+ */
+UCLASS(Abstract, DefaultToInstanced, EditInlineNew, CollapseCategories, MinimalAPI)
+class UCameraValueInterpolator : public UObject
+{
+	GENERATED_BODY()
+
+public:
+
+	template<typename ValueType> 
+	using TCameraValueInterpolator = UE::Cameras::TCameraValueInterpolator<ValueType>;
+
+	/** Creates a value interpolator for double precision floating point numbers. */
+	virtual TUniquePtr<TCameraValueInterpolator<double>> BuildDoubleInterpolator() const;
+	/** Creates a value interpolator for 2D vectors. */
+	virtual TUniquePtr<TCameraValueInterpolator<FVector2d>> BuildVector2dInterpolator() const;
+	/** Creates a value interpolator for 3D vectors. */
+	virtual TUniquePtr<TCameraValueInterpolator<FVector3d>> BuildVector3dInterpolator() const;
+
+protected:
+
+	/** Creates a value interpolator for double precision floating point numbers. */
+	virtual TUniquePtr<TCameraValueInterpolator<double>> OnBuildDoubleInterpolator() const { return nullptr; }
+	/** Creates a value interpolator for 2D vectors. */
+	virtual TUniquePtr<TCameraValueInterpolator<FVector2d>> OnBuildVector2dInterpolator() const { return nullptr; }
+	/** Creates a value interpolator for 3D vectors. */
+	virtual TUniquePtr<TCameraValueInterpolator<FVector3d>> OnBuildVector3dInterpolator() const { return nullptr; }
+};
+
+#define UE_DECLARE_CAMERA_VALUE_INTERPOLATOR()\
+	protected:\
+		virtual TUniquePtr<UE::Cameras::TCameraValueInterpolator<double>> OnBuildDoubleInterpolator() const override;\
+		virtual TUniquePtr<UE::Cameras::TCameraValueInterpolator<FVector2d>> OnBuildVector2dInterpolator() const override;\
+		virtual TUniquePtr<UE::Cameras::TCameraValueInterpolator<FVector3d>> OnBuildVector3dInterpolator() const override;\
+	private:
+
+#define UE_DEFINE_CAMERA_VALUE_INTERPOLATOR_GENERIC(ThisClass, GenericInterpolatorClass)\
+	TUniquePtr<UE::Cameras::TCameraValueInterpolator<double>> ThisClass::OnBuildDoubleInterpolator() const\
+	{\
+		return MakeUnique<GenericInterpolatorClass<double>>(this);\
+	}\
+	TUniquePtr<UE::Cameras::TCameraValueInterpolator<FVector2d>> ThisClass::OnBuildVector2dInterpolator() const\
+	{\
+		return MakeUnique<GenericInterpolatorClass<FVector2d>>(this);\
+	}\
+	TUniquePtr<UE::Cameras::TCameraValueInterpolator<FVector3d>> ThisClass::OnBuildVector3dInterpolator() const\
+	{\
+		return MakeUnique<GenericInterpolatorClass<FVector3d>>(this);\
+	}
+
+/** Simple interpolator that immediately "pops" to the target value. */
+UCLASS(meta=(DisplayName="Immediate Pop"))
+class UPopValueInterpolator : public UCameraValueInterpolator
+{
+	GENERATED_BODY()
+
+	UE_DECLARE_CAMERA_VALUE_INTERPOLATOR()
+};
+
diff --git a/Engine/Plugins/Cameras/GameplayCameras/Source/GameplayCameras/Public/ValueInterpolators/AccelerationDecelerationValueInterpolator.h b/Engine/Plugins/Cameras/GameplayCameras/Source/GameplayCameras/Public/ValueInterpolators/AccelerationDecelerationValueInterpolator.h
new file mode 100644
--- /dev/null
+++ b/Engine/Plugins/Cameras/GameplayCameras/Source/GameplayCameras/Public/ValueInterpolators/AccelerationDecelerationValueInterpolator.h
@@ -0,0 +1,34 @@
+// Copyright Epic Games, Inc. All Rights Reserved.
+
+#pragma once
+
+#include "Core/CameraValueInterpolator.h"
+
+#include "AccelerationDecelerationValueInterpolator.generated.h"
+
+/**
+ * A value interpolator that first accelerates up to a maximum speed, and then decelerates
+ * before reaching the target value.
+ */
+UCLASS(meta=(DisplayName="Acceleration-Deceleration"))
+class UAccelerationDecelerationValueInterpolator : public UCameraValueInterpolator
+{
+	GENERATED_BODY()
+
+	UE_DECLARE_CAMERA_VALUE_INTERPOLATOR()
+
+public:
+
+	/** The acceleration rate at the start of interpolation. */
+	UPROPERTY(EditAnywhere, Category="Speed")
+	float Acceleration = 1.f;
+
+	/** The maximum speed reachable during interpolation. */
+	UPROPERTY(EditAnywhere, Category="Speed")
+	float MaxSpeed = 10.f;
+
+	/** The deceleration rate at the end of interpolation. */
+	UPROPERTY(EditAnywhere, Category="Speed")
+	float Deceleration = 1.f;
+};
+
diff --git a/Engine/Plugins/Cameras/GameplayCameras/Source/GameplayCameras/Public/ValueInterpolators/CriticalDamperValueInterpolator.h b/Engine/Plugins/Cameras/GameplayCameras/Source/GameplayCameras/Public/ValueInterpolators/CriticalDamperValueInterpolator.h
new file mode 100644
--- /dev/null
+++ b/Engine/Plugins/Cameras/GameplayCameras/Source/GameplayCameras/Public/ValueInterpolators/CriticalDamperValueInterpolator.h
@@ -0,0 +1,26 @@
+// Copyright Epic Games, Inc. All Rights Reserved.
+
+#pragma once
+
+#include "Core/CameraValueInterpolator.h"
+#include "Math/CriticalDamper.h"
+
+#include "CriticalDamperValueInterpolator.generated.h"
+
+/**
+ * A value interpolator that uses a spring-mass system to converge towards the target value.
+ */
+UCLASS(meta=(DisplayName="Critical Damper"))
+class UCriticalDamperValueInterpolator : public UCameraValueInterpolator
+{
+	GENERATED_BODY()
+
+	UE_DECLARE_CAMERA_VALUE_INTERPOLATOR()
+
+public:
+
+	/** The damping factor of the spring-mass system. */
+	UPROPERTY(EditAnywhere, Category="Damping")
+	float DampingFactor = 1.f;
+};
+
diff --git a/Engine/Plugins/Cameras/GameplayCameras/Source/GameplayCameras/Public/ValueInterpolators/DoubleIIRValueInterpolator.h b/Engine/Plugins/Cameras/GameplayCameras/Source/GameplayCameras/Public/ValueInterpolators/DoubleIIRValueInterpolator.h
new file mode 100644
--- /dev/null
+++ b/Engine/Plugins/Cameras/GameplayCameras/Source/GameplayCameras/Public/ValueInterpolators/DoubleIIRValueInterpolator.h
@@ -0,0 +1,33 @@
+// Copyright Epic Games, Inc. All Rights Reserved.
+
+#pragma once
+
+#include "Core/CameraValueInterpolator.h"
+
+#include "DoubleIIRValueInterpolator.generated.h"
+
+/**
+ * Double infinite impulse response filter interpolator.
+ */
+UCLASS(meta=(DisplayName="Double-IIR"))
+class UDoubleIIRValueInterpolator : public UCameraValueInterpolator
+{
+	GENERATED_BODY()
+
+	UE_DECLARE_CAMERA_VALUE_INTERPOLATOR()
+
+public:
+
+	/** The primary speed of interpolation. */
+	UPROPERTY(EditAnywhere, Category="Speed")
+	float PrimarySpeed = 1.f;
+
+	/** The intermediate speed of interpolation. */
+	UPROPERTY(EditAnywhere, Category="Speed")
+	float IntermediateSpeed = 1.f;
+
+	/** Whether to use fixed-step evaluation. */
+	UPROPERTY(EditAnywhere, Category="Interpolation")
+	bool bUseFixedStep = true;
+};
+
diff --git a/Engine/Plugins/Cameras/GameplayCameras/Source/GameplayCameras/Public/ValueInterpolators/IIRValueInterpolator.h b/Engine/Plugins/Cameras/GameplayCameras/Source/GameplayCameras/Public/ValueInterpolators/IIRValueInterpolator.h
new file mode 100644
--- /dev/null
+++ b/Engine/Plugins/Cameras/GameplayCameras/Source/GameplayCameras/Public/ValueInterpolators/IIRValueInterpolator.h
@@ -0,0 +1,197 @@
+// Copyright Epic Games, Inc. All Rights Reserved.
+
+#pragma once
+
+#include "Core/CameraValueInterpolator.h"
+
+#include <type_traits>
+
+#include "IIRValueInterpolator.generated.h"
+
+/**
+ * Infinite impulse response filter interpolator.
+ */
+UCLASS(meta=(DisplayName="IIR"))
+class UIIRValueInterpolator : public UCameraValueInterpolator
+{
+	GENERATED_BODY()
+
+	UE_DECLARE_CAMERA_VALUE_INTERPOLATOR()
+
+public:
+
+	/** The speed of interpolation. */
+	UPROPERTY(EditAnywhere, Category="Interpolation")
+	float Speed = 1.f;
+
+	/** Whether to use fixed-step evaluation. */
+	UPROPERTY(EditAnywhere, Category="Interpolation")
+	bool bUseFixedStep = true;
+};
+
+namespace UE::Cameras
+{
+
+template<typename ValueType>
+struct TIIRValueInterpolatorTraits
+{};
+
+/**
+ * The actual evaluation code for the IIR interpolator. Exposed here in this header
+ * because it's also re-used in the double-IIR interpolator.
+ * Most of the code here courtesy of Jeff Farris.
+ */
+template<typename ValueType>
+class TIIRValueInterpolator : public TCameraValueInterpolator<ValueType>
+{
+public:
+
+	TIIRValueInterpolator(const UIIRValueInterpolator* InParameters)
+		: TCameraValueInterpolator<ValueType>(nullptr)
+		, Speed(InParameters->Speed)
+		, bUseFixedStep(InParameters->bUseFixedStep)
+	{}
+
+	TIIRValueInterpolator(float InSpeed, bool bInUseFixedStep)
+		: TCameraValueInterpolator<ValueType>(nullptr)
+		, Speed(InSpeed)
+		, bUseFixedStep(bInUseFixedStep)
+	{}
+
+protected:
+
+	using ValueTypeParam = TCameraValueInterpolator<ValueType>::ValueTypeParam;
+
+	virtual void OnReset(ValueTypeParam OldCurrentValue, ValueTypeParam OldTargetValue) override
+	{
+		LastTargetValue = this->TargetValue;
+		// Clear out any leftovers for rewind.
+		LastUpdateLeftoverTime = 0.f; 
+	}
+
+	virtual void OnRun(const FCameraValueInterpolationParams& Params, FCameraValueInterpolationResult& OutResult) override
+	{
+		if (bUseFixedStep)
+		{
+			float RemainingTime = Params.DeltaTime;
+
+			// Handle any leftover rewind.
+			if (bDoLeftoverRewind && (LastUpdateLeftoverTime > 0.f))
+			{
+				// Rewind back to the state at end of the last full-step update.
+				RemainingTime += LastUpdateLeftoverTime;
+				this->CurrentValue = ValueAfterLastFullStep;
+				LastUpdateLeftoverTime = 0.f;
+			}
+
+			// Move the substep target value linearly toward real target value while we evaluate the substeps.
+			ValueType LastToTargetValue(this->TargetValue - LastTargetValue);
+			if constexpr (std::is_same_v<FRotator3d, ValueType>)
+			{
+				LastToTargetValue = LastToTargetValue.GetNormalized();
+			}
+			const ValueType EquilibriumStepRate = LastToTargetValue * (1.f / RemainingTime);
+
+			ValueType LerpedTargetValue = LastTargetValue;
+
+			while (RemainingTime > KINDA_SMALL_NUMBER)
+			{
+				const float StepTime = FMath::Min(MaxSubstepTime, RemainingTime);
+
+				if (bDoLeftoverRewind && (StepTime < MaxSubstepTime))
+				{
+					// If this is the last substep, let's cache where we were after the last full-step 
+					// so we can resume from there on the next evaluation.
+					LastUpdateLeftoverTime = StepTime;
+					ValueAfterLastFullStep = this->CurrentValue;
+				}
+
+				LerpedTargetValue += EquilibriumStepRate * StepTime;
+				RemainingTime -= StepTime;
+
+				this->CurrentValue = RunSubstep(LerpedTargetValue, StepTime);
+			}
+
+			LastTargetValue = this->TargetValue;
+		}
+		else
+		{
+			this->CurrentValue = RunSubstep(this->TargetValue, Params.DeltaTime);
+			LastUpdateLeftoverTime = 0.f;
+		}
+	}
+
+	virtual void OnSerialize(const FCameraValueInterpolatorSerializeParams& Params, FArchive& Ar) override
+	{
+		Ar << ValueAfterLastFullStep;
+		Ar << LastTargetValue;
+		Ar << LastUpdateLeftoverTime;
+		Ar << bDoLeftoverRewind;
+	}
+
+private:
+
+	ValueType RunSubstep(ValueType SubstepTargetValue, float SubstepDeltaTime)
+	{
+		return TIIRValueInterpolatorTraits<ValueType>::InterpTo(this->CurrentValue, SubstepTargetValue, SubstepDeltaTime, Speed);
+	}
+
+private:
+
+	static constexpr float MaxSubstepTime = 1.f / 120.f;
+
+	ValueType ValueAfterLastFullStep;
+	ValueType LastTargetValue;
+	float LastUpdateLeftoverTime = 0.f;
+	float Speed = 1.0;
+	bool bDoLeftoverRewind = true;
+	bool bUseFixedStep = false;
+};
+
+template<>
+struct TIIRValueInterpolatorTraits<double>
+{
+	static double InterpTo(double CurrentValue, double TargetValue, float DeltaTime, double Speed)
+	{
+		return FMath::FInterpTo(CurrentValue, TargetValue, DeltaTime, Speed);
+	}
+};
+
+template<>
+struct TIIRValueInterpolatorTraits<FVector2d>
+{
+	static FVector2d InterpTo(FVector2d CurrentValue, FVector2d TargetValue, float DeltaTime, double Speed)
+	{
+		return FMath::Vector2DInterpTo(CurrentValue, TargetValue, DeltaTime, Speed);
+	}
+};
+
+template<>
+struct TIIRValueInterpolatorTraits<FVector3d>
+{
+	static FVector3d InterpTo(FVector3d CurrentValue, FVector3d TargetValue, float DeltaTime, double Speed)
+	{
+		return FMath::VInterpTo(CurrentValue, TargetValue, DeltaTime, Speed);
+	}
+};
+
+template<>
+struct TIIRValueInterpolatorTraits<FRotator>
+{
+	static FRotator InterpTo(FRotator CurrentValue, FRotator TargetValue, float DeltaTime, double Speed)
+	{
+		return FMath::RInterpTo(CurrentValue, TargetValue, DeltaTime, Speed);
+	}
+};
+
+template<>
+struct TIIRValueInterpolatorTraits<FQuat>
+{
+	static FQuat InterpTo(FQuat CurrentValue, FQuat TargetValue, float DeltaTime, double Speed)
+	{
+		return FMath::QInterpTo(CurrentValue, TargetValue, DeltaTime, Speed);
+	}
+};
+
+}  // namespace UE::Cameras
+
