Cameras: add auto-rotate input node

diff --git a/Engine/Plugins/Cameras/GameplayCameras/Source/GameplayCameras/Private/Nodes/Input/AutoRotateInput2DCameraNode.cpp b/Engine/Plugins/Cameras/GameplayCameras/Source/GameplayCameras/Private/Nodes/Input/AutoRotateInput2DCameraNode.cpp
new file mode 100644
--- /dev/null
+++ b/Engine/Plugins/Cameras/GameplayCameras/Source/GameplayCameras/Private/Nodes/Input/AutoRotateInput2DCameraNode.cpp
@@ -0,0 +1,256 @@
+// Copyright Epic Games, Inc. All Rights Reserved.
+
+#include "Nodes/Input/AutoRotateInput2DCameraNode.h"
+
+#include "Core/CameraEvaluationContext.h"
+#include "Core/CameraOperation.h"
+#include "Core/CameraValueInterpolator.h"
+#include "Debug/CameraDebugBlock.h"
+#include "Debug/CameraDebugBlockBuilder.h"
+#include "Debug/CameraDebugRenderer.h"
+
+#include UE_INLINE_GENERATED_CPP_BY_NAME(AutoRotateInput2DCameraNode)
+
+namespace UE::Cameras
+{
+
+class FAutoRotateInput2DCameraNodeEvaluator : public FInput2DCameraNodeEvaluator
+{
+	UE_DECLARE_CAMERA_NODE_EVALUATOR_EX(GAMEPLAYCAMERAS_API, FAutoRotateInput2DCameraNodeEvaluator, FInput2DCameraNodeEvaluator)
+
+public:
+
+	FAutoRotateInput2DCameraNodeEvaluator();
+
+protected:
+
+	virtual void OnBuild(const FCameraNodeEvaluatorBuildParams& Params) override;
+	virtual FCameraNodeEvaluatorChildrenView OnGetChildren() override;
+	virtual void OnInitialize(const FCameraNodeEvaluatorInitializeParams& Params) override;
+	virtual void OnRun(const FCameraNodeEvaluationParams& Params, FCameraNodeEvaluationResult& OutResult) override;
+
+#if UE_GAMEPLAY_CAMERAS_DEBUG
+	virtual void OnBuildDebugBlocks(const FCameraDebugBlockBuildParams& Params, FCameraDebugBlockBuilder& Builder) override;
+#endif  // UE_GAMEPLAY_CAMERAS_DEBUG
+
+private:
+
+	TUniquePtr<FCameraDoubleValueInterpolator> Interpolator;
+
+	FInput2DCameraNodeEvaluator* InputNodeEvaluator = nullptr;
+
+	FVector3d LastContextLocation;
+	FVector2d LastInputValue;
+	double RemainingWaitTime = 0.f;
+	bool bIsAutoRotating = false;
+};
+
+UE_DEFINE_CAMERA_NODE_EVALUATOR(FAutoRotateInput2DCameraNodeEvaluator)
+
+UE_DECLARE_CAMERA_DEBUG_BLOCK_START(GAMEPLAYCAMERAS_API, FAutoRotateInput2DCameraDebugBlock)
+	UE_DECLARE_CAMERA_DEBUG_BLOCK_FIELD(double, RemainingWaitTime);
+	UE_DECLARE_CAMERA_DEBUG_BLOCK_FIELD(double, InterpolationFactor);
+	UE_DECLARE_CAMERA_DEBUG_BLOCK_FIELD(double, InterpolationTarget);
+	UE_DECLARE_CAMERA_DEBUG_BLOCK_FIELD(bool, bIsAutoRotating);
+UE_DECLARE_CAMERA_DEBUG_BLOCK_END()
+
+UE_DEFINE_CAMERA_DEBUG_BLOCK_WITH_FIELDS(FAutoRotateInput2DCameraDebugBlock)
+
+FAutoRotateInput2DCameraNodeEvaluator::FAutoRotateInput2DCameraNodeEvaluator()
+{
+}
+
+void FAutoRotateInput2DCameraNodeEvaluator::OnBuild(const FCameraNodeEvaluatorBuildParams& Params)
+{
+	const UAutoRotateInput2DCameraNode* AutoRotateNode = GetCameraNodeAs<UAutoRotateInput2DCameraNode>();
+	if (AutoRotateNode->InputNode)
+	{
+		InputNodeEvaluator = Params.BuildEvaluatorAs<FInput2DCameraNodeEvaluator>(AutoRotateNode->InputNode);
+	}
+}
+
+FCameraNodeEvaluatorChildrenView FAutoRotateInput2DCameraNodeEvaluator::OnGetChildren()
+{
+	return FCameraNodeEvaluatorChildrenView({ InputNodeEvaluator });
+}
+
+void FAutoRotateInput2DCameraNodeEvaluator::OnInitialize(const FCameraNodeEvaluatorInitializeParams& Params)
+{
+	const FCameraNodeEvaluationResult& InitialResult = Params.EvaluationContext->GetInitialResult();
+	LastContextLocation = InitialResult.CameraPose.GetLocation();
+
+	LastInputValue = FVector2d::ZeroVector;
+	if (InputNodeEvaluator)
+	{
+		LastInputValue = InputNodeEvaluator->GetInputValue();
+	}
+
+	const UAutoRotateInput2DCameraNode* AutoRotateNode = GetCameraNodeAs<UAutoRotateInput2DCameraNode>();
+	RemainingWaitTime = AutoRotateNode->WaitTime;
+	bIsAutoRotating = false;
+}
+
+void FAutoRotateInput2DCameraNodeEvaluator::OnRun(const FCameraNodeEvaluationParams& Params, FCameraNodeEvaluationResult& OutResult)
+{
+	const UAutoRotateInput2DCameraNode* AutoRotateNode = GetCameraNodeAs<UAutoRotateInput2DCameraNode>();
+
+	// Evaluate our inner input node and grab the new input value.
+	if (InputNodeEvaluator)
+	{
+		InputNodeEvaluator->Run(Params, OutResult);
+
+		InputValue = InputNodeEvaluator->GetInputValue();
+	}
+
+	// Keep track of the context's movement this frame.
+	const FCameraNodeEvaluationResult& InitialResult = Params.EvaluationContext->GetInitialResult();
+	FVector3d ContextMovement = FVector3d::ZeroVector;
+	FVector3d CurrentContextLocation = InitialResult.CameraPose.GetLocation();
+	if (!Params.bIsFirstFrame && Params.DeltaTime > 0)
+	{
+		ContextMovement = (CurrentContextLocation - LastContextLocation);
+	}
+	LastContextLocation = CurrentContextLocation;
+
+	// Check if the input value changed from under us, and if it changed enough for
+	// us to deactivate auto-rotate.
+	const double YawChange = FMath::Abs(InputValue.X - LastInputValue.X);
+	const double PitchChange = FMath::Abs(InputValue.Y - LastInputValue.Y);
+	LastInputValue = InputValue;
+	if (YawChange >= AutoRotateNode->DeactivationThreshold || PitchChange >= AutoRotateNode->DeactivationThreshold)
+	{
+		bIsAutoRotating = false;
+		Interpolator.Reset();
+		RemainingWaitTime = AutoRotateNode->WaitTime;
+		return;
+	}
+
+	// We are almost good to auto-rotate... but maybe we need to wait a bit longer.
+	if (RemainingWaitTime > 0.f)
+	{
+		RemainingWaitTime -= Params.DeltaTime;
+		if (RemainingWaitTime > 0.f)
+		{
+			return;
+		}
+	}
+
+	// Figure out which direction we should auto-rotate towards.
+	FVector3d AutoRotateDir = InitialResult.CameraPose.GetAimDir();
+	if (AutoRotateNode->Direction == ECameraAutoRotateDirection::Movement)
+	{
+		AutoRotateDir = ContextMovement.GetSafeNormal(UE_SMALL_NUMBER, AutoRotateDir);
+	}
+	const FRotator3d AutoRotateRot = AutoRotateDir.ToOrientationRotator();
+
+	// Figure out how much work we have to do.
+	const double DeltaYaw = AutoRotateRot.Yaw - InputValue.X;
+	const double DeltaPitch = AutoRotateRot.Pitch - InputValue.Y;
+	if (FMath::Abs(DeltaYaw) < UE_DOUBLE_SMALL_NUMBER && FMath::Abs(DeltaPitch) < UE_DOUBLE_SMALL_NUMBER)
+	{
+		return;
+	}
+
+	// We will interpolate the length of the vector that represents the delta yaw/pitch
+	// that we need to compensate with.
+	const FVector2d DeltaVector = FVector2d(DeltaYaw, DeltaPitch);
+	const double DeltaMagnitude = DeltaVector.Length();
+
+	// Create our interpolator and update it.
+	if (!bIsAutoRotating)
+	{
+		if (AutoRotateNode->Interpolator)
+		{
+			Interpolator = AutoRotateNode->Interpolator->BuildDoubleInterpolator();
+		}
+		else
+		{
+			Interpolator = MakeUnique<TPopValueInterpolator<double>>();
+		}
+		bIsAutoRotating = true;
+	}
+	check(Interpolator);
+
+	Interpolator->Reset(DeltaMagnitude, 0);
+
+	FCameraValueInterpolationParams InterpParams;
+	InterpParams.DeltaTime = Params.DeltaTime;
+	FCameraValueInterpolationResult InterpResult(OutResult.VariableTable);
+	const double NewDeltaMagnitude = Interpolator->Run(InterpParams, InterpResult);
+
+	// Get the new delta yaw/pitch and try to adjust our inner input node.
+	const FVector2d NewDeltaVector = DeltaVector * (1.0 - NewDeltaMagnitude / DeltaMagnitude);
+	bool bDeactivateAutoRotate = false;
+
+	if (InputNodeEvaluator)
+	{
+		FCameraOperationParams OperationParams;
+		OperationParams.EvaluationContext = Params.EvaluationContext;
+		OperationParams.Evaluator = Params.Evaluator;
+
+		FYawPitchCameraOperation YawPitchOperation;
+		YawPitchOperation.Yaw = FConsumableDouble::Delta(NewDeltaVector.X);
+		YawPitchOperation.Pitch = FConsumableDouble::Delta(NewDeltaVector.Y);
+
+		InputNodeEvaluator->ExecuteOperation(OperationParams, YawPitchOperation);
+		InputValue = InputNodeEvaluator->GetInputValue();
+		LastInputValue = InputValue;
+
+		if (YawPitchOperation.Yaw.HasValue() || YawPitchOperation.Pitch.HasValue())
+		{
+			bDeactivateAutoRotate = true;
+		}
+	}
+
+	if (Interpolator->IsFinished() || bDeactivateAutoRotate)
+	{
+		bIsAutoRotating = false;
+		Interpolator.Reset();
+		RemainingWaitTime = AutoRotateNode->WaitTime;
+	}
+}
+
+#if UE_GAMEPLAY_CAMERAS_DEBUG
+
+void FAutoRotateInput2DCameraNodeEvaluator::OnBuildDebugBlocks(const FCameraDebugBlockBuildParams& Params, FCameraDebugBlockBuilder& Builder)
+{
+	FAutoRotateInput2DCameraDebugBlock& DebugBlock = Builder.AttachDebugBlock<FAutoRotateInput2DCameraDebugBlock>();
+
+	DebugBlock.RemainingWaitTime = RemainingWaitTime;
+	DebugBlock.bIsAutoRotating = bIsAutoRotating;
+	DebugBlock.InterpolationFactor = Interpolator ? Interpolator->GetCurrentValue() : 0.0f;
+	DebugBlock.InterpolationTarget = Interpolator ? Interpolator->GetTargetValue() : 0.0f;
+}
+
+void FAutoRotateInput2DCameraDebugBlock::OnDebugDraw(const FCameraDebugBlockDrawParams& Params, FCameraDebugRenderer& Renderer)
+{
+	if (bIsAutoRotating)
+	{
+		Renderer.AddText(TEXT("rotating: %.3f -> %.3f"), InterpolationFactor, InterpolationTarget);
+	}
+	else
+	{
+		Renderer.AddText(TEXT("waiting: %.3fsec"), RemainingWaitTime);
+	}
+}
+
+#endif  // UE_GAMEPLAY_CAMERAS_DEBUG
+
+}  // namespace UE::Cameras
+
+UAutoRotateInput2DCameraNode::UAutoRotateInput2DCameraNode(const FObjectInitializer& ObjInit)
+	: Super(ObjInit)
+{
+}
+
+FCameraNodeChildrenView UAutoRotateInput2DCameraNode::OnGetChildren()
+{
+	return FCameraNodeChildrenView({ InputNode });
+}
+
+FCameraNodeEvaluatorPtr UAutoRotateInput2DCameraNode::OnBuildEvaluator(FCameraNodeEvaluatorBuilder& Builder) const
+{
+	using namespace UE::Cameras;
+	return Builder.BuildEvaluator<FAutoRotateInput2DCameraNodeEvaluator>();
+}
+
diff --git a/Engine/Plugins/Cameras/GameplayCameras/Source/GameplayCameras/Public/Nodes/Input/AutoRotateInput2DCameraNode.h b/Engine/Plugins/Cameras/GameplayCameras/Source/GameplayCameras/Public/Nodes/Input/AutoRotateInput2DCameraNode.h
new file mode 100644
--- /dev/null
+++ b/Engine/Plugins/Cameras/GameplayCameras/Source/GameplayCameras/Public/Nodes/Input/AutoRotateInput2DCameraNode.h
@@ -0,0 +1,62 @@
+// Copyright Epic Games, Inc. All Rights Reserved.
+
+#pragma once
+
+#include "Math/MathFwd.h"
+#include "Nodes/Input/Input2DCameraNode.h"
+
+#include "AutoRotateInput2DCameraNode.generated.h"
+
+class UCameraValueInterpolator;
+
+/** Describes the type of auto-rotate. */
+UENUM()
+enum class ECameraAutoRotateDirection
+{
+	/** Re-align towards the evaluation context's facing. */
+	Facing,
+	/** Re-align towards the evaluation context's movement direction. */
+	Movement
+};
+
+/**
+ * An input node that modifies a yaw/pitch input in order to re-align its
+ * values to a given default direction.
+ */
+UCLASS(MinimalAPI, meta=(CameraNodeCategories="Input"))
+class UAutoRotateInput2DCameraNode : public UInput2DCameraNode
+{
+	GENERATED_BODY()
+
+public:
+
+	/** The direction to re-align towards. */
+	UPROPERTY(EditAnywhere, Category="Auto-Rotate")
+	ECameraAutoRotateDirection Direction;
+
+	/** The time, in seconds, to wait before re-aligning. */
+	UPROPERTY(EditAnywhere, Category="Auto-Rotate")
+	float WaitTime = 1.f;
+
+	/** The minimum manual rotation, in degrees, to deactivate auto-rtation. */
+	UPROPERTY(EditAnywhere, Category="Auto-Rotate")
+	float DeactivationThreshold = 0.01f;
+
+	/** The interpolation for re-alignment. */
+	UPROPERTY(EditAnywhere, Category="Auto-Rotate")
+	TObjectPtr<UCameraValueInterpolator> Interpolator;
+
+	UPROPERTY()
+	TObjectPtr<UInput2DCameraNode> InputNode;
+
+public:
+
+	UAutoRotateInput2DCameraNode(const FObjectInitializer& ObjInit);
+	
+protected:
+
+	// UCameraNode interface.
+	virtual FCameraNodeChildrenView OnGetChildren() override;
+	virtual FCameraNodeEvaluatorPtr OnBuildEvaluator(FCameraNodeEvaluatorBuilder& Builder) const override;
+};
+
