Cameras: fix mac compiler

diff --git a/Engine/Plugins/Cameras/GameplayCameras/Source/GameplayCameras/Private/ValueInterpolators/CriticalDamperValueInterpolator.cpp b/Engine/Plugins/Cameras/GameplayCameras/Source/GameplayCameras/Private/ValueInterpolators/CriticalDamperValueInterpolator.cpp
--- a/Engine/Plugins/Cameras/GameplayCameras/Source/GameplayCameras/Private/ValueInterpolators/CriticalDamperValueInterpolator.cpp
+++ b/Engine/Plugins/Cameras/GameplayCameras/Source/GameplayCameras/Private/ValueInterpolators/CriticalDamperValueInterpolator.cpp
@@ -12,7 +12,7 @@ class TCriticalDamperValueInterpolator : public TCameraValueInterpolator<ValueTy
 {
 public:
 
-	using ValueTypeParam = TCameraValueInterpolator<ValueType>::ValueTypeParam;
+	using ValueTypeParam = typename TCameraValueInterpolator<ValueType>::ValueTypeParam;
 
 	TCriticalDamperValueInterpolator(const UCriticalDamperValueInterpolator* InInterpolator)
 		: TCameraValueInterpolator<ValueType>(InInterpolator)
diff --git a/Engine/Plugins/Cameras/GameplayCameras/Source/GameplayCameras/Private/ValueInterpolators/DoubleIIRValueInterpolator.cpp b/Engine/Plugins/Cameras/GameplayCameras/Source/GameplayCameras/Private/ValueInterpolators/DoubleIIRValueInterpolator.cpp
--- a/Engine/Plugins/Cameras/GameplayCameras/Source/GameplayCameras/Private/ValueInterpolators/DoubleIIRValueInterpolator.cpp
+++ b/Engine/Plugins/Cameras/GameplayCameras/Source/GameplayCameras/Private/ValueInterpolators/DoubleIIRValueInterpolator.cpp
@@ -23,7 +23,7 @@ public:
 
 protected:
 
-	using ValueTypeParam = TCameraValueInterpolator<ValueType>::ValueTypeParam;
+	using ValueTypeParam = typename TCameraValueInterpolator<ValueType>::ValueTypeParam;
 
 	virtual void OnReset(ValueTypeParam OldCurrentValue, ValueTypeParam OldTargetValue) override
 	{
diff --git a/Engine/Plugins/Cameras/GameplayCameras/Source/GameplayCameras/Public/Core/CameraValueInterpolator.h b/Engine/Plugins/Cameras/GameplayCameras/Source/GameplayCameras/Public/Core/CameraValueInterpolator.h
--- a/Engine/Plugins/Cameras/GameplayCameras/Source/GameplayCameras/Public/Core/CameraValueInterpolator.h
+++ b/Engine/Plugins/Cameras/GameplayCameras/Source/GameplayCameras/Public/Core/CameraValueInterpolator.h
@@ -55,7 +55,7 @@ class TCameraValueInterpolator
 {
 public:
 
-	using ValueTypeParam = TCallTraits<ValueType>::ParamType;
+	using ValueTypeParam = typename TCallTraits<ValueType>::ParamType;
 
 	/** 
 	 * Creates a new value interpolator given a piece of data containing user-defined
diff --git a/Engine/Plugins/Cameras/GameplayCameras/Source/GameplayCameras/Public/ValueInterpolators/IIRValueInterpolator.h b/Engine/Plugins/Cameras/GameplayCameras/Source/GameplayCameras/Public/ValueInterpolators/IIRValueInterpolator.h
--- a/Engine/Plugins/Cameras/GameplayCameras/Source/GameplayCameras/Public/ValueInterpolators/IIRValueInterpolator.h
+++ b/Engine/Plugins/Cameras/GameplayCameras/Source/GameplayCameras/Public/ValueInterpolators/IIRValueInterpolator.h
@@ -60,7 +60,7 @@ public:
 
 protected:
 
-	using ValueTypeParam = TCameraValueInterpolator<ValueType>::ValueTypeParam;
+	using ValueTypeParam = typename TCameraValueInterpolator<ValueType>::ValueTypeParam;
 
 	virtual void OnReset(ValueTypeParam OldCurrentValue, ValueTypeParam OldTargetValue) override
 	{
