Cameras: add camera parameter constructors that take their value as argument

This makes initializing parameters easier.

diff --git a/Engine/Plugins/Cameras/GameplayCameras/Source/GameplayCameras/Public/Core/CameraParameters.h b/Engine/Plugins/Cameras/GameplayCameras/Source/GameplayCameras/Public/Core/CameraParameters.h
--- a/Engine/Plugins/Cameras/GameplayCameras/Source/GameplayCameras/Public/Core/CameraParameters.h
+++ b/Engine/Plugins/Cameras/GameplayCameras/Source/GameplayCameras/Public/Core/CameraParameters.h
@@ -8,6 +8,15 @@
 
 #include "CameraParameters.generated.h"
 
+#define UE_DEFINE_CAMERA_PARAMETER_VALUE_CONSTRUCTORS(ParameterClass)\
+	ParameterClass(typename TCallTraits<ParameterClass::ValueType>::ParamType InValue)\
+		: Value(InValue)\
+	{}
+
+#define UE_DEFINE_CAMERA_PARAMETER_ALL_CONSTRUCTORS(ParameterClass)\
+	ParameterClass() {}\
+	UE_DEFINE_CAMERA_PARAMETER_VALUE_CONSTRUCTORS(ParameterClass)
+
 /** Boolean camera parameter. */
 USTRUCT()
 struct GAMEPLAYCAMERAS_API FBooleanCameraParameter
@@ -24,6 +33,8 @@ struct GAMEPLAYCAMERAS_API FBooleanCameraParameter
 	TObjectPtr<UBooleanCameraVariable> Variable;
 
 	bool SerializeFromMismatchedTag(const FPropertyTag& Tag, FStructuredArchive::FSlot Slot);
+
+	UE_DEFINE_CAMERA_PARAMETER_ALL_CONSTRUCTORS(FBooleanCameraParameter)
 };
 
 /** Integer camera parameter. */
@@ -42,6 +53,8 @@ struct GAMEPLAYCAMERAS_API FInteger32CameraParameter
 	TObjectPtr<UInteger32CameraVariable> Variable;
 
 	bool SerializeFromMismatchedTag(const FPropertyTag& Tag, FStructuredArchive::FSlot Slot);
+
+	UE_DEFINE_CAMERA_PARAMETER_ALL_CONSTRUCTORS(FInteger32CameraParameter)
 };
 
 /** Float camera parameter. */
@@ -60,6 +73,8 @@ struct FFloatCameraParameter
 	TObjectPtr<UFloatCameraVariable> Variable;
 
 	bool SerializeFromMismatchedTag(const FPropertyTag& Tag, FStructuredArchive::FSlot Slot);
+
+	UE_DEFINE_CAMERA_PARAMETER_ALL_CONSTRUCTORS(FFloatCameraParameter)
 };
 
 /** Double camera parameter. */
@@ -78,11 +93,13 @@ struct GAMEPLAYCAMERAS_API FDoubleCameraParameter
 	TObjectPtr<UDoubleCameraVariable> Variable;
 
 	bool SerializeFromMismatchedTag(const FPropertyTag& Tag, FStructuredArchive::FSlot Slot);
+
+	UE_DEFINE_CAMERA_PARAMETER_ALL_CONSTRUCTORS(FDoubleCameraParameter)
 };
 
 /** Vector2f camera parameter. */
 USTRUCT()
-struct FVector2fCameraParameter
+struct GAMEPLAYCAMERAS_API FVector2fCameraParameter
 {
 	GENERATED_BODY()
 
@@ -97,6 +114,8 @@ struct FVector2fCameraParameter
 
 	FVector2fCameraParameter();
 	bool SerializeFromMismatchedTag(const FPropertyTag& Tag, FStructuredArchive::FSlot Slot);
+
+	UE_DEFINE_CAMERA_PARAMETER_VALUE_CONSTRUCTORS(FVector2fCameraParameter)
 };
 
 /** Vector2d camera parameter. */
@@ -116,11 +135,13 @@ struct GAMEPLAYCAMERAS_API FVector2dCameraParameter
 
 	FVector2dCameraParameter();
 	bool SerializeFromMismatchedTag(const FPropertyTag& Tag, FStructuredArchive::FSlot Slot);
+
+	UE_DEFINE_CAMERA_PARAMETER_VALUE_CONSTRUCTORS(FVector2dCameraParameter)
 };
 
 /** Vector3f camera parameter. */
 USTRUCT()
-struct FVector3fCameraParameter
+struct GAMEPLAYCAMERAS_API FVector3fCameraParameter
 {
 	GENERATED_BODY()
 
@@ -135,6 +156,8 @@ struct FVector3fCameraParameter
 
 	FVector3fCameraParameter();
 	bool SerializeFromMismatchedTag(const FPropertyTag& Tag, FStructuredArchive::FSlot Slot);
+
+	UE_DEFINE_CAMERA_PARAMETER_VALUE_CONSTRUCTORS(FVector3fCameraParameter)
 };
 
 /** Vector3d camera parameter. */
@@ -154,11 +177,13 @@ struct GAMEPLAYCAMERAS_API FVector3dCameraParameter
 
 	FVector3dCameraParameter();
 	bool SerializeFromMismatchedTag(const FPropertyTag& Tag, FStructuredArchive::FSlot Slot);
+
+	UE_DEFINE_CAMERA_PARAMETER_VALUE_CONSTRUCTORS(FVector3dCameraParameter)
 };
 
 /** Vector4f camera parameter. */
 USTRUCT()
-struct FVector4fCameraParameter
+struct GAMEPLAYCAMERAS_API FVector4fCameraParameter
 {
 	GENERATED_BODY()
 
@@ -173,11 +198,13 @@ struct FVector4fCameraParameter
 
 	FVector4fCameraParameter();
 	bool SerializeFromMismatchedTag(const FPropertyTag& Tag, FStructuredArchive::FSlot Slot);
+
+	UE_DEFINE_CAMERA_PARAMETER_VALUE_CONSTRUCTORS(FVector4fCameraParameter)
 };
 
 /** Vector4d camera parameter. */
 USTRUCT()
-struct FVector4dCameraParameter
+struct GAMEPLAYCAMERAS_API FVector4dCameraParameter
 {
 	GENERATED_BODY()
 
@@ -192,6 +219,8 @@ struct FVector4dCameraParameter
 
 	FVector4dCameraParameter();
 	bool SerializeFromMismatchedTag(const FPropertyTag& Tag, FStructuredArchive::FSlot Slot);
+
+	UE_DEFINE_CAMERA_PARAMETER_VALUE_CONSTRUCTORS(FVector4dCameraParameter)
 };
 
 /** Rotator3f camera parameter. */
@@ -211,6 +240,8 @@ struct GAMEPLAYCAMERAS_API FRotator3fCameraParameter
 
 	FRotator3fCameraParameter();
 	bool SerializeFromMismatchedTag(const FPropertyTag& Tag, FStructuredArchive::FSlot Slot);
+
+	UE_DEFINE_CAMERA_PARAMETER_VALUE_CONSTRUCTORS(FRotator3fCameraParameter)
 };
 
 /** Rotator3d camera parameter. */
@@ -230,11 +261,13 @@ struct GAMEPLAYCAMERAS_API FRotator3dCameraParameter
 
 	FRotator3dCameraParameter();
 	bool SerializeFromMismatchedTag(const FPropertyTag& Tag, FStructuredArchive::FSlot Slot);
+
+	UE_DEFINE_CAMERA_PARAMETER_VALUE_CONSTRUCTORS(FRotator3dCameraParameter)
 };
 
 /** Transform3f camera parameter. */
 USTRUCT()
-struct FTransform3fCameraParameter
+struct GAMEPLAYCAMERAS_API FTransform3fCameraParameter
 {
 	GENERATED_BODY()
 
@@ -248,6 +281,8 @@ struct FTransform3fCameraParameter
 	TObjectPtr<UTransform3fCameraVariable> Variable;
 
 	bool SerializeFromMismatchedTag(const FPropertyTag& Tag, FStructuredArchive::FSlot Slot);
+
+	UE_DEFINE_CAMERA_PARAMETER_ALL_CONSTRUCTORS(FTransform3fCameraParameter)
 };
 
 /** Transform3d camera parameter. */
@@ -266,8 +301,13 @@ struct GAMEPLAYCAMERAS_API FTransform3dCameraParameter
 	TObjectPtr<UTransform3dCameraVariable> Variable;
 
 	bool SerializeFromMismatchedTag(const FPropertyTag& Tag, FStructuredArchive::FSlot Slot);
+
+	UE_DEFINE_CAMERA_PARAMETER_ALL_CONSTRUCTORS(FTransform3dCameraParameter)
 };
 
+#undef UE_DEFINE_CAMERA_PARAMETER_VALUE_CONSTRUCTORS
+#undef UE_DEFINE_CAMERA_PARAMETER_ALL_CONSTRUCTORS
+
 // Any camera parameter might replace a previously non-parameterized property (i.e. a "fixed" property
 // of the underlying type, like bool, int32, float, etc.)
 // When someone upgrades the fixed property to a parameterized property, any previously saved data will
