Cameras: add boom arm length interpolation based on pivot movement

diff --git a/Engine/Plugins/Cameras/GameplayCameras/Source/GameplayCameras/Private/Nodes/Common/BoomArmCameraNode.cpp b/Engine/Plugins/Cameras/GameplayCameras/Source/GameplayCameras/Private/Nodes/Common/BoomArmCameraNode.cpp
--- a/Engine/Plugins/Cameras/GameplayCameras/Source/GameplayCameras/Private/Nodes/Common/BoomArmCameraNode.cpp
+++ b/Engine/Plugins/Cameras/GameplayCameras/Source/GameplayCameras/Private/Nodes/Common/BoomArmCameraNode.cpp
@@ -6,6 +6,10 @@
 #include "Core/CameraOperation.h"
 #include "Core/CameraParameterReader.h"
 #include "Core/CameraRigJoints.h"
+#include "Core/CameraValueInterpolator.h"
+#include "Debug/CameraDebugBlock.h"
+#include "Debug/CameraDebugBlockBuilder.h"
+#include "Debug/CameraDebugRenderer.h"
 #include "GameFramework/Pawn.h"
 #include "GameFramework/PlayerController.h"
 #include "Nodes/Input/Input2DCameraNode.h"
@@ -27,6 +31,10 @@ protected:
 	virtual void OnRun(const FCameraNodeEvaluationParams& Params, FCameraNodeEvaluationResult& OutResult) override;
 	virtual void OnExecuteOperation(const FCameraOperationParams& Params, FCameraOperation& Operation) override;
 
+#if UE_GAMEPLAY_CAMERAS_DEBUG
+	virtual void OnBuildDebugBlocks(const FCameraDebugBlockBuildParams& Params, FCameraDebugBlockBuilder& Builder) override;
+#endif  // UE_GAMEPLAY_CAMERAS_DEBUG
+
 private:
 
 	APlayerController* GetPlayerController(TSharedPtr<const FCameraEvaluationContext> EvaluationContext) const;
@@ -35,14 +43,38 @@ private:
 
 	TCameraParameterReader<FVector3d> BoomOffsetReader;
 	FInput2DCameraNodeEvaluator* InputSlotEvaluator = nullptr;
+	
+	TUniquePtr<TCameraValueInterpolator<double>> BoomLengthInterpolator;
+	TCameraParameterReader<double> MaxForwardInterpolationFactorReader;
+	TCameraParameterReader<double> MaxBackwardInterpolationFactorReader;
+	FVector3d LastPivotLocation;
+	double CumulativePull;
+
+#if UE_GAMEPLAY_CAMERAS_DEBUG
+	FVector2d DebugYawPitch;
+	bool bDebugDidClampPull;
+#endif  // UE_GAMEPLAY_CAMERAS_DEBUG
 };
 
 UE_DEFINE_CAMERA_NODE_EVALUATOR(FBoomArmCameraNodeEvaluator)
 
+UE_DECLARE_CAMERA_DEBUG_BLOCK_START(GAMEPLAYCAMERAS_API, FBoomArmCameraDebugBlock)
+	UE_DECLARE_CAMERA_DEBUG_BLOCK_FIELD(FVector2d, BoomYawPitch)
+	UE_DECLARE_CAMERA_DEBUG_BLOCK_FIELD(double, CumulativePull)
+	UE_DECLARE_CAMERA_DEBUG_BLOCK_FIELD(bool, bDidClampPull)
+	UE_DECLARE_CAMERA_DEBUG_BLOCK_FIELD(bool, bHasBoomLengthInterpolator)
+UE_DECLARE_CAMERA_DEBUG_BLOCK_END()
+
+UE_DEFINE_CAMERA_DEBUG_BLOCK_WITH_FIELDS(FBoomArmCameraDebugBlock)
+
 void FBoomArmCameraNodeEvaluator::OnBuild(const FCameraNodeEvaluatorBuildParams& Params)
 {
 	const UBoomArmCameraNode* BoomArmNode = GetCameraNodeAs<UBoomArmCameraNode>();
 	InputSlotEvaluator = Params.BuildEvaluatorAs<FInput2DCameraNodeEvaluator>(BoomArmNode->InputSlot);
+	if (BoomArmNode->BoomLengthInterpolator)
+	{
+		BoomLengthInterpolator = BoomArmNode->BoomLengthInterpolator->BuildDoubleInterpolator();
+	}
 }
 
 void FBoomArmCameraNodeEvaluator::OnInitialize(const FCameraNodeEvaluatorInitializeParams& Params)
@@ -53,6 +85,11 @@ void FBoomArmCameraNodeEvaluator::OnInitialize(const FCameraNodeEvaluatorInitial
 
 	const UBoomArmCameraNode* BoomArmNode = GetCameraNodeAs<UBoomArmCameraNode>();
 	BoomOffsetReader.Initialize(BoomArmNode->BoomOffset);
+	MaxForwardInterpolationFactorReader.Initialize(BoomArmNode->MaxForwardInterpolationFactor);
+	MaxBackwardInterpolationFactorReader.Initialize(BoomArmNode->MaxBackwardInterpolationFactor);
+
+	LastPivotLocation = FVector3d::ZeroVector;
+	CumulativePull = 0.0;
 }
 
 FCameraNodeEvaluatorChildrenView FBoomArmCameraNodeEvaluator::OnGetChildren()
@@ -85,9 +122,81 @@ void FBoomArmCameraNodeEvaluator::OnRun(const FCameraNodeEvaluationParams& Param
 	// by using the fact that BoomRotation is, well, just a rotation, and CameraPose.Location is of
 	// course just a translation. So we can put them both in the same transform:
 	const FTransform3d BoomPivot(BoomRotation, OutResult.CameraPose.GetLocation());
-	const FTransform3d BoomOffset(BoomOffsetReader.Get(OutResult.VariableTable));
+	const FVector3d BoomOffset(BoomOffsetReader.Get(OutResult.VariableTable));
+
+	FTransform3d FinalTransform(FTransform3d(BoomOffset) * BoomPivot);
+
+#if UE_GAMEPLAY_CAMERAS_DEBUG
+	DebugYawPitch = FVector2d(BoomRotation.Yaw, BoomRotation.Pitch);
+	bDebugDidClampPull = false;
+#endif
+
+	// If we have an interpolator for the boom length, let's run it now. The way we use it is that we
+	// keep track of the "pull" on the boom arm, i.e. how much the boom gets pulled in various directions
+	// as the pivot moves around (such as when running around with the player character or driving a
+	// vehicle). We compute the amount of pull along the pivot<->camera line, and then ask the interpolator
+	// to converge towards zero. This creates a sort of rubber-band/spring effect on the boom arm.
+	const double DefaultBoomLength = BoomOffset.Length();
+	if (BoomLengthInterpolator && DefaultBoomLength > 0)
+	{
+		if (!Params.bIsFirstFrame)
+		{
+			// The pull this frame is how much the base (pivot) of the boom arm has moved. We add that
+			// to our cumulative tally of the pull.
+			// Note that pull is positive when the pivot is moving forwards (away from the camera) and 
+			// negative when moving backwards (towards the camera).
+			const FVector3d PivotMovement = BoomPivot.GetLocation() - LastPivotLocation;
+			const FVector3d ForwardBoomOrientation = BoomRotation.RotateVector(FVector3d::ForwardVector);
+			const double PullThisFrame = PivotMovement.Dot(ForwardBoomOrientation);
+			CumulativePull += PullThisFrame;
+
+			// Update the interpolator to try and get back to zero pull.
+			BoomLengthInterpolator->Reset(CumulativePull, 0);
+
+			FCameraValueInterpolationParams InterpParams;
+			InterpParams.DeltaTime = Params.DeltaTime;
+			FCameraValueInterpolationResult InterpResult(OutResult.VariableTable);
+			double NewCumulativePull = BoomLengthInterpolator->Run(InterpParams, InterpResult);
+
+			// Clamp the cumulative pull to any maximums defined by the user.
+			double ClampedPull = NewCumulativePull;
+			if (ClampedPull < 0)
+			{
+				const double MaxForwardInterpolationFactor = MaxForwardInterpolationFactorReader.Get(OutResult.VariableTable);
+				if (MaxForwardInterpolationFactor > 0)
+				{
+					const double MaxForwardPush = DefaultBoomLength * MaxForwardInterpolationFactor;
+					ClampedPull = FMath::Max(-MaxForwardPush, ClampedPull);
+				}
+			}
+			else if (ClampedPull > 0)
+			{
+				const double MaxBackwardInterpolationFactor = MaxBackwardInterpolationFactorReader.Get(OutResult.VariableTable);
+				if (MaxBackwardInterpolationFactor > 0)
+				{
+					const double MaxBackwardPull = DefaultBoomLength * MaxBackwardInterpolationFactor;
+					ClampedPull = FMath::Min(MaxBackwardPull, ClampedPull);
+				}
+			}
+
+			// Add the pull to final transform. This effectively distorts the boom offset, since we move the 
+			// camera forwards/backwards based on the boom orientation, not the offset's orientation.
+			FinalTransform.SetLocation(FinalTransform.GetLocation() - ForwardBoomOrientation * ClampedPull);
 
-	const FTransform3d FinalTransform(BoomOffset * BoomPivot);
+#if UE_GAMEPLAY_CAMERAS_DEBUG
+			bDebugDidClampPull = (ClampedPull != NewCumulativePull);
+#endif
+
+			CumulativePull = ClampedPull;
+		}
+		else
+		{
+			LastPivotLocation = BoomPivot.GetLocation();
+			CumulativePull = 0.0;
+		}
+
+		LastPivotLocation = BoomPivot.GetLocation();
+	}
 
 	OutResult.CameraPose.SetTransform(FinalTransform);
 	
@@ -126,6 +235,34 @@ APlayerController* FBoomArmCameraNodeEvaluator::GetPlayerController(TSharedPtr<c
 	return nullptr;
 }
 
+#if UE_GAMEPLAY_CAMERAS_DEBUG
+
+void FBoomArmCameraNodeEvaluator::OnBuildDebugBlocks(const FCameraDebugBlockBuildParams& Params, FCameraDebugBlockBuilder& Builder)
+{
+	FBoomArmCameraDebugBlock& DebugBlock = Builder.AttachDebugBlock<FBoomArmCameraDebugBlock>();
+
+	DebugBlock.BoomYawPitch = DebugYawPitch;
+	DebugBlock.CumulativePull = CumulativePull;
+	DebugBlock.bDidClampPull = bDebugDidClampPull;
+	DebugBlock.bHasBoomLengthInterpolator = BoomLengthInterpolator.IsValid();
+}
+
+void FBoomArmCameraDebugBlock::OnDebugDraw(const FCameraDebugBlockDrawParams& Params, FCameraDebugRenderer& Renderer)
+{
+	Renderer.AddText(TEXT("yaw: %.3f pitch %.3f"), BoomYawPitch.X, BoomYawPitch.Y);
+
+	if (bHasBoomLengthInterpolator)
+	{
+		Renderer.AddText(TEXT(" (pull: %.3f)"), CumulativePull);
+		if (bDidClampPull)
+		{
+			Renderer.AddText(TEXT(" [CLAMPING]"));
+		}
+	}
+}
+
+#endif  // UE_GAMEPLAY_CAMERAS_DEBUG
+
 }  // namespace UE::Cameras
 
 FCameraNodeChildrenView UBoomArmCameraNode::OnGetChildren()
diff --git a/Engine/Plugins/Cameras/GameplayCameras/Source/GameplayCameras/Public/Nodes/Common/BoomArmCameraNode.h b/Engine/Plugins/Cameras/GameplayCameras/Source/GameplayCameras/Public/Nodes/Common/BoomArmCameraNode.h
--- a/Engine/Plugins/Cameras/GameplayCameras/Source/GameplayCameras/Public/Nodes/Common/BoomArmCameraNode.h
+++ b/Engine/Plugins/Cameras/GameplayCameras/Source/GameplayCameras/Public/Nodes/Common/BoomArmCameraNode.h
@@ -7,6 +7,7 @@
 
 #include "BoomArmCameraNode.generated.h"
 
+class UCameraValueInterpolator;
 class UInput2DCameraNode;
 
 /**
@@ -29,6 +30,24 @@ public:
 	UPROPERTY(EditAnywhere, Category=Common)
 	FVector3dCameraParameter BoomOffset;
 
+	/** The interpolator to use for changing the boom length based on its pivot's movements. */
+	UPROPERTY(EditAnywhere, Category=Common)
+	TObjectPtr<UCameraValueInterpolator> BoomLengthInterpolator;
+
+	/** 
+	 * The maximum amount of forward movement the interpolator can introduce, expressed
+	 * as a factor of the default boom length.
+	 */
+	UPROPERTY(EditAnywhere, Category=Common)
+	FDoubleCameraParameter MaxForwardInterpolationFactor = -1.0;
+
+	/** 
+	 * The maximum amount of backward movement the interpolator can introduce, expressed
+	 * as a factor of the default boom length.
+	 */
+	UPROPERTY(EditAnywhere, Category=Common)
+	FDoubleCameraParameter MaxBackwardInterpolationFactor = -1.0;
+
 	/**
 	 * The input slot for controlling the boom arm.
 	 * If no input slot is specified, the boom arm will use the player controller view rotation.
