Cameras: add Guids to camera rig interface parameters

diff --git a/Engine/Plugins/Cameras/GameplayCameras/Source/GameplayCameras/Private/Core/CameraRigAsset.cpp b/Engine/Plugins/Cameras/GameplayCameras/Source/GameplayCameras/Private/Core/CameraRigAsset.cpp
--- a/Engine/Plugins/Cameras/GameplayCameras/Source/GameplayCameras/Private/Core/CameraRigAsset.cpp
+++ b/Engine/Plugins/Cameras/GameplayCameras/Source/GameplayCameras/Private/Core/CameraRigAsset.cpp
@@ -29,6 +29,38 @@ void UCameraRigInterfaceParameter::OnGraphNodeMoved(FName InGraphName, int32 Nod
 
 #endif
 
+
+void UCameraRigInterfaceParameter::PostLoad()
+{
+	if (!Guid.IsValid())
+	{
+		Guid = FGuid::NewGuid();
+	}
+
+	Super::PostLoad();
+}
+
+void UCameraRigInterfaceParameter::PostInitProperties()
+{
+	Super::PostInitProperties();
+
+	if (!HasAnyFlags(RF_ClassDefaultObject | RF_ArchetypeObject | RF_NeedLoad | RF_WasLoaded) && 
+			!Guid.IsValid())
+	{
+		Guid = FGuid::NewGuid();
+	}
+}
+
+void UCameraRigInterfaceParameter::PostDuplicate(EDuplicateMode::Type DuplicateMode)
+{
+	Super::PostDuplicate(DuplicateMode);
+
+	if (DuplicateMode == EDuplicateMode::Normal)
+	{
+		Guid = FGuid::NewGuid();
+	}
+}
+
 UCameraRigInterfaceParameter* FCameraRigInterface::FindInterfaceParameterByName(const FString& ParameterName) const
 {
 	const TObjectPtr<UCameraRigInterfaceParameter>* FoundItem = InterfaceParameters.FindByPredicate([&ParameterName](UCameraRigInterfaceParameter* Item)
diff --git a/Engine/Plugins/Cameras/GameplayCameras/Source/GameplayCameras/Public/Core/CameraRigAsset.h b/Engine/Plugins/Cameras/GameplayCameras/Source/GameplayCameras/Public/Core/CameraRigAsset.h
--- a/Engine/Plugins/Cameras/GameplayCameras/Source/GameplayCameras/Public/Core/CameraRigAsset.h
+++ b/Engine/Plugins/Cameras/GameplayCameras/Source/GameplayCameras/Public/Core/CameraRigAsset.h
@@ -67,6 +67,9 @@ public:
 	UPROPERTY(EditAnywhere, Category=Camera)
 	FString InterfaceParameterName;
 
+	UPROPERTY()
+	FGuid Guid;
+
 	// Built on save/cook.
 
 	/**
@@ -85,6 +88,11 @@ protected:
 	virtual void OnGraphNodeMoved(FName InGraphName, int32 NodePosX, int32 NodePosY, bool bMarkDirty) override;
 #endif
 
+	// UObject interface.
+	virtual void PostLoad() override;
+	virtual void PostInitProperties() override;
+	virtual void PostDuplicate(EDuplicateMode::Type DuplicateMode) override;
+
 private:
 
 #if WITH_EDITORONLY_DATA
