Cameras: fix issues and crashes with freezing camera rigs in the blend stack

diff --git a/Engine/Plugins/Cameras/GameplayCameras/Source/GameplayCameras/Private/Core/BlendStackCameraNode.cpp b/Engine/Plugins/Cameras/GameplayCameras/Source/GameplayCameras/Private/Core/BlendStackCameraNode.cpp
--- a/Engine/Plugins/Cameras/GameplayCameras/Source/GameplayCameras/Private/Core/BlendStackCameraNode.cpp
+++ b/Engine/Plugins/Cameras/GameplayCameras/Source/GameplayCameras/Private/Core/BlendStackCameraNode.cpp
@@ -148,6 +148,9 @@ bool FBlendStackCameraNodeEvaluator::InitializeEntry(
 		return false;
 	}
 
+	// Allocate variables in the variable table.
+	NewEntry.Result.VariableTable.Initialize(CameraRig->AllocationInfo.VariableTableInfo);
+
 	// Initialize the node evaluators.
 	FCameraNodeEvaluatorInitializeParams InitParams;
 	InitParams.Evaluator = Evaluator;
@@ -159,9 +162,6 @@ bool FBlendStackCameraNodeEvaluator::InitializeEntry(
 	NewEntry.ParameterEvaluators.Reset();
 	GatherEntryParameterEvaluators(RootEvaluator, NewEntry.ParameterEvaluators);
 
-	// Allocate variables in the variable table.
-	NewEntry.Result.VariableTable.Initialize(CameraRig->AllocationInfo.VariableTableInfo);
-
 	// Wrap up!
 	NewEntry.EvaluationContext = EvaluationContext;
 	NewEntry.CameraRig = CameraRig;
@@ -172,6 +172,22 @@ bool FBlendStackCameraNodeEvaluator::InitializeEntry(
 	return true;
 }
 
+void FBlendStackCameraNodeEvaluator::FreezeEntry(FCameraRigEntry& Entry)
+{
+	// Deallocate our node evaluators and clear any pointers we kept to them.
+	Entry.EvaluatorStorage.DestroyEvaluatorTree(true);
+	Entry.RootEvaluator = nullptr;
+	Entry.ParameterEvaluators.Reset();
+
+	Entry.EvaluationContext.Reset();
+	
+#if WITH_EDITOR
+	RemoveListenedPackages(Entry);
+#endif
+	
+	Entry.bIsFrozen = true;
+}
+
 void FBlendStackCameraNodeEvaluator::GatherEntryParameterEvaluators(FCameraNodeEvaluator* RootEvaluator, TArray<FCameraNodeEvaluator*>& OutParameterEvaluators)
 {
 	TArray<FCameraNodeEvaluator*> EvaluatorStack;
@@ -235,33 +251,38 @@ void FBlendStackCameraNodeEvaluator::OnRun(const FCameraNodeEvaluationParams& Pa
 {
 	const UBlendStackCameraNode* BlendStackNode = GetCameraNodeAs<UBlendStackCameraNode>();
 
-	// Filter out invalid entries and go through other basic validation.
-	struct FValidEntry
+	// Build up these structures so we don't re-resolve evaluation context weak-pointers
+	// multiple times in this function..
+	struct FResolvedEntry
 	{
 		FCameraRigEntry& Entry;
 		TSharedPtr<const FCameraEvaluationContext> Context;
 		int32 EntryIndex;
 	};
 
-	TArray<FValidEntry> ValidEntries;
+	TArray<FResolvedEntry> ResolvedEntries;
 
+	// While we make these resolved entries, emit warnings and errors as needed.
 	for (int32 Index = 0; Index < Entries.Num(); ++Index)
 	{
 		FCameraRigEntry& Entry(Entries[Index]);
 		TSharedPtr<const FCameraEvaluationContext> CurContext = Entry.EvaluationContext.Pin();
 
+		ResolvedEntries.Add({ Entry, CurContext, Index });
+
 		if (!Entry.bIsFrozen)
 		{
-			// Check that we still have a valid context.
+			// Check that we still have a valid context. If not, let's freeze the entry, since
+			// we won't be able to evaluate it anymore.
 			if (UNLIKELY(!CurContext.IsValid()))
 			{
-				Entry.Result.bIsValid = false;
+				FreezeEntry(Entry);
 
 #if UE_GAMEPLAY_CAMERAS_TRACE
 				if (Entry.bLogWarnings)
 				{
 					UE_LOG(LogCameraSystem, Warning,
-							TEXT("Can't update camera rig '%s' because its evaluation context isn't valid."),
+							TEXT("Freezing camera rig '%s' because its evaluation context isn't valid anymore."),
 							*GetNameSafe(Entry.CameraRig));
 					Entry.bLogWarnings = false;
 				}
@@ -274,13 +295,11 @@ void FBlendStackCameraNodeEvaluator::OnRun(const FCameraNodeEvaluationParams& Pa
 			const FCameraNodeEvaluationResult& ContextResult(CurContext->GetInitialResult());
 			if (UNLIKELY(!ContextResult.bIsValid))
 			{
-				Entry.Result.bIsValid = false;
-
 #if UE_GAMEPLAY_CAMERAS_TRACE
 				if (Entry.bLogWarnings)
 				{
 					UE_LOG(LogCameraSystem, Warning,
-							TEXT("Can't update camera rig '%s' because its initial result isn't valid."),
+							TEXT("Camera rig '%s' may experience a hitch because its initial result isn't valid."),
 							*GetNameSafe(Entry.CameraRig));
 					Entry.bLogWarnings = false;
 				}
@@ -291,31 +310,27 @@ void FBlendStackCameraNodeEvaluator::OnRun(const FCameraNodeEvaluationParams& Pa
 		}
 		// else: frozen entries may have null contexts or invalid initial results
 		//       because we're not going to update them anyway. We will however blend
-		//       them so we add them to the valid entries.
-
-		// This entry is valid, we'll process it.
-		ValidEntries.Add({ Entry, CurContext, Index });
+		//       them so we add them to the list of entries too.
 
 #if UE_GAMEPLAY_CAMERAS_TRACE
-		// This entry might have been temporarily invalid before. It's valid now, so let's
+		// This entry might have has warnings before. It's valid now, so let's
 		// re-enable warnings if it becomes invalid again in the future.
 		Entry.bLogWarnings = true;
 #endif  // UE_GAMEPLAY_CAMERAS_TRACE
 	}
 
 	// Gather parameters to pre-blend, and evaluate blend nodes.
-	for (FValidEntry& ValidEntry : ValidEntries)
+	for (FResolvedEntry& ResolvedEntry : ResolvedEntries)
 	{
-		FCameraRigEntry& Entry(ValidEntry.Entry);
-		TSharedPtr<const FCameraEvaluationContext> CurContext(ValidEntry.Context);
+		FCameraRigEntry& Entry(ResolvedEntry.Entry);
 
-		if (Entry.bIsFrozen)
+		if (UNLIKELY(Entry.bIsFrozen))
 		{
 			continue;
 		}
 
 		FCameraNodeEvaluationParams CurParams(Params);
-		CurParams.EvaluationContext = CurContext;
+		CurParams.EvaluationContext = ResolvedEntry.Context;
 		CurParams.bIsFirstFrame = Entry.bIsFirstFrame;
 
 		FCameraNodeEvaluationResult& CurResult(Entry.Result);
@@ -350,20 +365,15 @@ void FBlendStackCameraNodeEvaluator::OnRun(const FCameraNodeEvaluationParams& Pa
 	}
 
 	// Blend input variables.
-	for (FValidEntry& ValidEntry : ValidEntries)
+	for (FResolvedEntry& ResolvedEntry : ResolvedEntries)
 	{
-		FCameraRigEntry& Entry(ValidEntry.Entry);
-		TSharedPtr<const FCameraEvaluationContext> CurContext(ValidEntry.Context);
-
-		// Don't filter out frozen entries here, they still contribute to the blend
-		// using their last evaluated values.
+		FCameraRigEntry& Entry(ResolvedEntry.Entry);
+		FCameraNodeEvaluationResult& CurResult(Entry.Result);
 
-		if (!Entry.ParameterEvaluators.IsEmpty())
+		if (!Entry.bIsFrozen)
 		{
-			FCameraNodeEvaluationResult& CurResult(Entry.Result);
-
 			FCameraNodeEvaluationParams CurParams(Params);
-			CurParams.EvaluationContext = CurContext;
+			CurParams.EvaluationContext = ResolvedEntry.Context;
 			CurParams.bIsFirstFrame = Entry.bIsFirstFrame;
 			FCameraNodePreBlendParams PreBlendParams(CurParams, CurResult.CameraPose, CurResult.VariableTable);
 
@@ -375,16 +385,26 @@ void FBlendStackCameraNodeEvaluator::OnRun(const FCameraNodeEvaluationParams& Pa
 				EntryBlendEvaluator->BlendParameters(PreBlendParams, PreBlendResult);
 			}
 		}
+		else
+		{
+			// Frozen entries still contribute to the blend using their last evaluated values.
+			OutResult.VariableTable.Override(CurResult.VariableTable, ECameraVariableTableFilter::Input);
+		}
 	}
 
 	// Run the root nodes. They will use the pre-blended inputs from the last step.
-	for (FValidEntry& ValidEntry : ValidEntries)
+	// Frozen entries are skipped, since they only ever use the last result they produced.
+	for (FResolvedEntry& ResolvedEntry : ResolvedEntries)
 	{
-		FCameraRigEntry& Entry(ValidEntry.Entry);
-		TSharedPtr<const FCameraEvaluationContext> CurContext(ValidEntry.Context);
+		FCameraRigEntry& Entry(ResolvedEntry.Entry);
+
+		if (UNLIKELY(Entry.bIsFrozen))
+		{
+			continue;
+		}
 
 		FCameraNodeEvaluationParams CurParams(Params);
-		CurParams.EvaluationContext = CurContext;
+		CurParams.EvaluationContext = ResolvedEntry.Context;
 		CurParams.bIsFirstFrame = Entry.bIsFirstFrame;
 
 		FCameraNodeEvaluationResult& CurResult(Entry.Result);
@@ -396,7 +416,7 @@ void FBlendStackCameraNodeEvaluator::OnRun(const FCameraNodeEvaluationParams& Pa
 		CurResult.VariableTable.ClearAllWrittenThisFrameFlags();
 
 		// Override it with whatever the evaluation context has set on its result.
-		const FCameraNodeEvaluationResult& ContextResult(CurContext->GetInitialResult());
+		const FCameraNodeEvaluationResult& ContextResult(ResolvedEntry.Context->GetInitialResult());
 		CurResult.CameraPose.OverrideChanged(ContextResult.CameraPose);
 		CurResult.VariableTable.OverrideAll(ContextResult.VariableTable);
 		CurResult.bIsCameraCut = OutResult.bIsCameraCut || ContextResult.bIsCameraCut;
@@ -414,35 +434,44 @@ void FBlendStackCameraNodeEvaluator::OnRun(const FCameraNodeEvaluationParams& Pa
 	// that we can remove any camera rigs below (since they would have been completely
 	// blended out by that).
 	int32 PopEntriesBelow = INDEX_NONE;
-	for (FValidEntry& ValidEntry : ValidEntries)
+	for (FResolvedEntry& ResolvedEntry : ResolvedEntries)
 	{
-		FCameraRigEntry& Entry(ValidEntry.Entry);
-		TSharedPtr<const FCameraEvaluationContext> CurContext(ValidEntry.Context);
-
+		FCameraRigEntry& Entry(ResolvedEntry.Entry);
 		FCameraNodeEvaluationResult& CurResult(Entry.Result);
 
-		FCameraNodeEvaluationParams CurParams(Params);
-		CurParams.EvaluationContext = CurContext;
-		CurParams.bIsFirstFrame = Entry.bIsFirstFrame;
-		FCameraNodeBlendParams BlendParams(CurParams, CurResult);
+		if (!Entry.bIsFrozen)
+		{
+			FCameraNodeEvaluationParams CurParams(Params);
+			CurParams.EvaluationContext = ResolvedEntry.Context;
+			CurParams.bIsFirstFrame = Entry.bIsFirstFrame;
+			FCameraNodeBlendParams BlendParams(CurParams, CurResult);
 
-		FCameraNodeBlendResult BlendResult(OutResult);
+			FCameraNodeBlendResult BlendResult(OutResult);
 
-		FBlendCameraNodeEvaluator* EntryBlendEvaluator = Entry.RootEvaluator->GetBlendEvaluator();
-		if (EntryBlendEvaluator)
-		{
-			EntryBlendEvaluator->BlendResults(BlendParams, BlendResult);
+			FBlendCameraNodeEvaluator* EntryBlendEvaluator = Entry.RootEvaluator->GetBlendEvaluator();
+			if (EntryBlendEvaluator)
+			{
+				EntryBlendEvaluator->BlendResults(BlendParams, BlendResult);
 
-			if (BlendResult.bIsBlendFull && BlendResult.bIsBlendFinished)
+				if (BlendResult.bIsBlendFull && BlendResult.bIsBlendFinished)
+				{
+					PopEntriesBelow = ResolvedEntry.EntryIndex;
+				}
+			}
+			else
 			{
-				PopEntriesBelow = ValidEntry.EntryIndex;
+				OutResult.CameraPose.OverrideAll(CurResult.CameraPose);
+				OutResult.VariableTable.OverrideAll(CurResult.VariableTable);
+
+				PopEntriesBelow = ResolvedEntry.EntryIndex;
 			}
 		}
 		else
 		{
 			OutResult.CameraPose.OverrideAll(CurResult.CameraPose);
+			OutResult.VariableTable.OverrideAll(CurResult.VariableTable);
 
-			PopEntriesBelow = ValidEntry.EntryIndex;
+			PopEntriesBelow = ResolvedEntry.EntryIndex;
 		}
 	}
 
@@ -475,22 +504,10 @@ void FBlendStackCameraNodeEvaluator::PopEntries(int32 FirstIndexToKeep)
 
 	for (int32 Index = 0; Index < FirstIndexToKeep; ++Index)
 	{
-		const FCameraRigEntry& FirstEntry = Entries[0];
+		FCameraRigEntry& FirstEntry = Entries[0];
 
 #if WITH_EDITOR
-		for (const UPackage* ListenPackage : FirstEntry.ListenedPackages)
-		{
-			int32* NumListens = AllListenedPackages.Find(ListenPackage);
-			if (ensure(NumListens))
-			{
-				--(*NumListens);
-				if (*NumListens == 0)
-				{
-					LiveEditManager->RemoveListener(ListenPackage, this);
-					AllListenedPackages.Remove(ListenPackage);
-				}
-			}
-		}
+		RemoveListenedPackages(LiveEditManager, FirstEntry);
 #endif  // WITH_EDITOR
 
 		if (OnCameraRigEventDelegate.IsBound())
@@ -502,6 +519,36 @@ void FBlendStackCameraNodeEvaluator::PopEntries(int32 FirstIndexToKeep)
 	}
 }
 
+#if WITH_EDITOR
+
+void FBlendStackCameraNodeEvaluator::RemoveListenedPackages(FCameraRigEntry& Entry)
+{
+	IGameplayCamerasModule& GameplayCamerasModule = FModuleManager::GetModuleChecked<IGameplayCamerasModule>("GameplayCameras");
+	TSharedPtr<IGameplayCamerasLiveEditManager> LiveEditManager = GameplayCamerasModule.GetLiveEditManager();
+	RemoveListenedPackages(LiveEditManager, Entry);
+}
+
+void FBlendStackCameraNodeEvaluator::RemoveListenedPackages(TSharedPtr<IGameplayCamerasLiveEditManager> LiveEditManager, FCameraRigEntry& Entry)
+{
+	for (const UPackage* ListenPackage : Entry.ListenedPackages)
+	{
+		int32* NumListens = AllListenedPackages.Find(ListenPackage);
+		if (ensure(NumListens))
+		{
+			--(*NumListens);
+			if (*NumListens == 0)
+			{
+				LiveEditManager->RemoveListener(ListenPackage, this);
+				AllListenedPackages.Remove(ListenPackage);
+			}
+		}
+	}
+
+	Entry.ListenedPackages.Reset();
+}
+
+#endif  // WITH_EDITOR
+
 const UCameraRigTransition* FBlendStackCameraNodeEvaluator::FindTransition(const FBlendStackCameraPushParams& Params) const
 {
 	const UBlendStackCameraNode* BlendStackNode = GetCameraNodeAs<UBlendStackCameraNode>();
diff --git a/Engine/Plugins/Cameras/GameplayCameras/Source/GameplayCameras/Public/Core/BlendStackCameraNode.h b/Engine/Plugins/Cameras/GameplayCameras/Source/GameplayCameras/Public/Core/BlendStackCameraNode.h
--- a/Engine/Plugins/Cameras/GameplayCameras/Source/GameplayCameras/Public/Core/BlendStackCameraNode.h
+++ b/Engine/Plugins/Cameras/GameplayCameras/Source/GameplayCameras/Public/Core/BlendStackCameraNode.h
@@ -30,6 +30,10 @@ struct FBlendStackCameraRigEvent;
 class FBlendStackCameraDebugBlock;
 #endif  // UE_GAMEPLAY_CAMERAS_DEBUG
 
+#if WITH_EDITOR
+class IGameplayCamerasLiveEditManager;
+#endif
+
 }  // namespace UE::Cameras
 
 /**
@@ -149,10 +153,17 @@ protected:
 		TSharedPtr<const FCameraEvaluationContext> EvaluationContext,
 		UBlendStackRootCameraNode* EntryRootNode);
 
+	void FreezeEntry(FCameraRigEntry& Entry);
+
 	void GatherEntryParameterEvaluators(FCameraNodeEvaluator* RootEvaluator, TArray<FCameraNodeEvaluator*>& OutParameterEvaluators);
 
 	void BroadcastCameraRigEvent(EBlendStackCameraRigEventType EventType, const FCameraRigEntry& Entry, const UCameraRigTransition* Transition = nullptr) const;
 
+#if WITH_EDITOR
+	void RemoveListenedPackages(FCameraRigEntry& Entry);
+	void RemoveListenedPackages(TSharedPtr<IGameplayCamerasLiveEditManager> LiveEditManager, FCameraRigEntry& Entry);
+#endif
+
 protected:
 
 	struct FCameraRigEntry
