Cameras: don't show internal variables in the UI for exposed camera parameters

diff --git a/Engine/Plugins/Cameras/GameplayCameras/Source/GameplayCamerasEditor/Private/Customizations/CameraParameterDetailsCustomizations.cpp b/Engine/Plugins/Cameras/GameplayCameras/Source/GameplayCamerasEditor/Private/Customizations/CameraParameterDetailsCustomizations.cpp
--- a/Engine/Plugins/Cameras/GameplayCameras/Source/GameplayCamerasEditor/Private/Customizations/CameraParameterDetailsCustomizations.cpp
+++ b/Engine/Plugins/Cameras/GameplayCameras/Source/GameplayCamerasEditor/Private/Customizations/CameraParameterDetailsCustomizations.cpp
@@ -238,13 +238,18 @@ void FCameraParameterDetailsCustomization::CustomizeHeader(TSharedRef<IPropertyH
 
 	UpdateVariableInfo();
 
-	const bool bHasText = !VariableInfoText.IsEmpty() || !VariableErrorText.IsEmpty();
+	// The value widget is enabled (i.e. the user can change the value) if the parameter isn't driven by
+	// a variable, or if that variable is a private variable meant to expose the parameter on the rig interface.
+	const bool bShowVariableText = !VariableInfoText.IsEmpty() && !bIsExposedParameterVariable;
+	const bool bShowVariableError = !VariableErrorText.IsEmpty() && !bIsExposedParameterVariable;
+	const bool bEnableVariableWidget = bShowVariableText || bShowVariableError;
 
 	TSharedRef<SWidget> ValueWidget = ValueProperty->CreatePropertyValueWidgetWithCustomization(nullptr);
-	ValueWidget->SetEnabled(!bHasText);
+	ValueWidget->SetEnabled(!bEnableVariableWidget);
+
 	// TODO: change SStandaloneCustomizedValueWidget so that it can tell us about the layout requirements
 	//		 of the value widget (min/max desired width/height, H/V alignment, etc.)
-	const float MaxValueWidgetDesiredWidth = bHasText ? 300.f : 0.f;
+	const float MaxValueWidgetDesiredWidth = bEnableVariableWidget ? 300.f : 0.f;
 
 	TSharedRef<FGameplayCamerasEditorStyle> GameplayCamerasStyle = FGameplayCamerasEditorStyle::Get();
 
@@ -285,7 +290,7 @@ void FCameraParameterDetailsCustomization::CustomizeHeader(TSharedRef<IPropertyH
 			.ContentPadding(1.f)
 			.ButtonStyle(FAppStyle::Get(), "SimpleButton")
 			.ToolTipText(LOCTEXT("SetVariable_ToolTip", "Selects a camera variable to drive this parameter"))
-			.IsEnabled(VariableClass != nullptr)
+			.IsEnabled(VariableClass != nullptr && !bIsExposedParameterVariable)
 			.ButtonContent()
 			[
 				SNew(SHorizontalBox)
@@ -305,7 +310,7 @@ void FCameraParameterDetailsCustomization::CustomizeHeader(TSharedRef<IPropertyH
 				[
 					SNew(SBox)
 					.VAlign(VAlign_Center)
-					.Visibility(!VariableInfoText.IsEmpty() ? EVisibility::Visible : EVisibility::Collapsed)
+					.Visibility(bShowVariableText ? EVisibility::Visible : EVisibility::Collapsed)
 					[
 						SNew(STextBlock)
 						.Text(VariableInfoText)
@@ -318,7 +323,7 @@ void FCameraParameterDetailsCustomization::CustomizeHeader(TSharedRef<IPropertyH
 				[
 					SNew(SBox)
 					.VAlign(VAlign_Center)
-					.Visibility(!VariableErrorText.IsEmpty() ? EVisibility::Visible : EVisibility::Collapsed)
+					.Visibility(bShowVariableError ? EVisibility::Visible : EVisibility::Collapsed)
 					[
 						SNew(STextBlock)
 						.Text(VariableErrorText)
@@ -379,6 +384,7 @@ void FCameraParameterDetailsCustomization::UpdateVariableInfo()
 	CommonVariable = nullptr;
 	VariableInfoText = FText::GetEmpty();
 	VariableErrorText = FText::GetEmpty();
+	bIsExposedParameterVariable = false;
 
 	UObject* VariableObject = nullptr;
 	FPropertyAccess::Result PropertyAccessResult = VariableProperty->GetValue(VariableObject);
@@ -392,6 +398,7 @@ void FCameraParameterDetailsCustomization::UpdateVariableInfo()
 				VariableInfoText = Variable->DisplayName.IsEmpty() ?
 					FText::FromName(Variable->GetFName()) :
 					FText::FromString(Variable->DisplayName);
+				bIsExposedParameterVariable = CommonVariable->bIsPrivate;
 			}
 			else
 			{
diff --git a/Engine/Plugins/Cameras/GameplayCameras/Source/GameplayCamerasEditor/Public/Customizations/CameraParameterDetailsCustomizations.h b/Engine/Plugins/Cameras/GameplayCameras/Source/GameplayCamerasEditor/Public/Customizations/CameraParameterDetailsCustomizations.h
--- a/Engine/Plugins/Cameras/GameplayCameras/Source/GameplayCamerasEditor/Public/Customizations/CameraParameterDetailsCustomizations.h
+++ b/Engine/Plugins/Cameras/GameplayCameras/Source/GameplayCamerasEditor/Public/Customizations/CameraParameterDetailsCustomizations.h
@@ -58,6 +58,7 @@ protected:
 	UCameraVariableAsset* CommonVariable = nullptr;
 	FText VariableInfoText;
 	FText VariableErrorText;
+	bool bIsExposedParameterVariable = false;
 
 	TSharedPtr<IPropertyUtilities> PropertyUtilities;
 
