Cameras: fix private camera variables not ending up private in the table

diff --git a/Engine/Plugins/Cameras/GameplayCameras/Source/GameplayCameras/Private/Core/CameraVariableAssets.cpp b/Engine/Plugins/Cameras/GameplayCameras/Source/GameplayCameras/Private/Core/CameraVariableAssets.cpp
--- a/Engine/Plugins/Cameras/GameplayCameras/Source/GameplayCameras/Private/Core/CameraVariableAssets.cpp
+++ b/Engine/Plugins/Cameras/GameplayCameras/Source/GameplayCameras/Private/Core/CameraVariableAssets.cpp
@@ -20,6 +20,7 @@ FCameraVariableDefinition UCameraVariableAsset::GetVariableDefinition() const
 	FCameraVariableDefinition VariableDefinition;
 	VariableDefinition.VariableID = GetVariableID();
 	VariableDefinition.VariableType = GetVariableType();
+	VariableDefinition.bIsPrivate = bIsPrivate;
 #if WITH_EDITORONLY_DATA
 	VariableDefinition.VariableName = GetName();
 #endif
