Cameras: reset variable table flags before we run pre-blended parameter updates

This lets us correctly track if variables were written-to during a given update or not.

diff --git a/Engine/Plugins/Cameras/GameplayCameras/Source/GameplayCameras/Private/Core/BlendStackCameraNode.cpp b/Engine/Plugins/Cameras/GameplayCameras/Source/GameplayCameras/Private/Core/BlendStackCameraNode.cpp
--- a/Engine/Plugins/Cameras/GameplayCameras/Source/GameplayCameras/Private/Core/BlendStackCameraNode.cpp
+++ b/Engine/Plugins/Cameras/GameplayCameras/Source/GameplayCameras/Private/Core/BlendStackCameraNode.cpp
@@ -335,6 +335,8 @@ void FBlendStackCameraNodeEvaluator::OnRun(const FCameraNodeEvaluationParams& Pa
 
 		FCameraNodeEvaluationResult& CurResult(Entry.Result);
 
+		CurResult.VariableTable.ClearAllWrittenThisFrameFlags();
+
 		// Gather input parameters.
 		if (!Entry.bInputRunThisFrame)
 		{
@@ -411,9 +413,7 @@ void FBlendStackCameraNodeEvaluator::OnRun(const FCameraNodeEvaluationParams& Pa
 
 		// Start with the input given to us.
 		CurResult.CameraPose = OutResult.CameraPose;
-		CurResult.CameraPose.ClearAllChangedFlags();
 		CurResult.VariableTable.OverrideAll(OutResult.VariableTable);
-		CurResult.VariableTable.ClearAllWrittenThisFrameFlags();
 
 		// Override it with whatever the evaluation context has set on its result.
 		const FCameraNodeEvaluationResult& ContextResult(ResolvedEntry.Context->GetInitialResult());
@@ -422,6 +422,8 @@ void FBlendStackCameraNodeEvaluator::OnRun(const FCameraNodeEvaluationParams& Pa
 		CurResult.bIsCameraCut = OutResult.bIsCameraCut || ContextResult.bIsCameraCut;
 		CurResult.bIsValid = true;
 
+		CurResult.CameraPose.ClearAllChangedFlags();
+
 		// Run the camera rig's root node.
 		FCameraNodeEvaluator* RootEvaluator = Entry.RootEvaluator->GetRootEvaluator();
 		if (ensure(RootEvaluator))
