Cameras: align variable table buffer on at least 32 bits

Aligning on less than that creates OOM crashes.

diff --git a/Engine/Plugins/Cameras/GameplayCameras/Source/GameplayCameras/Private/Core/CameraVariableTable.cpp b/Engine/Plugins/Cameras/GameplayCameras/Source/GameplayCameras/Private/Core/CameraVariableTable.cpp
--- a/Engine/Plugins/Cameras/GameplayCameras/Source/GameplayCameras/Private/Core/CameraVariableTable.cpp
+++ b/Engine/Plugins/Cameras/GameplayCameras/Source/GameplayCameras/Private/Core/CameraVariableTable.cpp
@@ -109,6 +109,7 @@ void FCameraVariableTable::Initialize(const FCameraVariableTableAllocationInfo&
 	}
 
 	// Allocate the memory buffer.
+	MaxAlignOf = FMath::Max(32u, MaxAlignOf);
 	Memory = reinterpret_cast<uint8*>(FMemory::Malloc(TotalSizeOf, MaxAlignOf));
 	Capacity = TotalSizeOf;
 	Used = 0;
