Cameras: add camera rig (prefab) node

- The camera rig node runs another camera rig, with some potential overrides on any exposed parameter.
- A details customization class provides an editing UI for those overrides.

diff --git a/Engine/Plugins/Cameras/GameplayCameras/Source/GameplayCameras/Private/Core/CameraRigAsset.cpp b/Engine/Plugins/Cameras/GameplayCameras/Source/GameplayCameras/Private/Core/CameraRigAsset.cpp
--- a/Engine/Plugins/Cameras/GameplayCameras/Source/GameplayCameras/Private/Core/CameraRigAsset.cpp
+++ b/Engine/Plugins/Cameras/GameplayCameras/Source/GameplayCameras/Private/Core/CameraRigAsset.cpp
@@ -13,6 +13,17 @@
 
 #include UE_INLINE_GENERATED_CPP_BY_NAME(CameraRigAsset)
 
+void FCameraRigAllocationInfo::Append(const FCameraRigAllocationInfo& OtherAllocationInfo)
+{
+	const FCameraNodeEvaluatorAllocationInfo& OtherEvaluatorInfo(OtherAllocationInfo.EvaluatorInfo);
+	EvaluatorInfo.MaxAlignof = FMath::Max(EvaluatorInfo.MaxAlignof, OtherEvaluatorInfo.MaxAlignof);
+	EvaluatorInfo.TotalSizeof = Align(EvaluatorInfo.TotalSizeof, OtherEvaluatorInfo.MaxAlignof) + OtherEvaluatorInfo.TotalSizeof;
+
+	const FCameraVariableTableAllocationInfo& OtherVariableTableInfo(OtherAllocationInfo.VariableTableInfo);
+	VariableTableInfo.AutoResetVariables.Append(OtherVariableTableInfo.AutoResetVariables);
+	VariableTableInfo.VariableDefinitions.Append(OtherVariableTableInfo.VariableDefinitions);
+}
+
 #if WITH_EDITOR
 
 void UCameraRigInterfaceParameter::GetGraphNodePosition(FName InGraphName, int32& NodePosX, int32& NodePosY) const
diff --git a/Engine/Plugins/Cameras/GameplayCameras/Source/GameplayCameras/Private/Nodes/Common/CameraRigCameraNode.cpp b/Engine/Plugins/Cameras/GameplayCameras/Source/GameplayCameras/Private/Nodes/Common/CameraRigCameraNode.cpp
new file mode 100644
--- /dev/null
+++ b/Engine/Plugins/Cameras/GameplayCameras/Source/GameplayCameras/Private/Nodes/Common/CameraRigCameraNode.cpp
@@ -0,0 +1,235 @@
+// Copyright Epic Games, Inc. All Rights Reserved.
+
+#include "Nodes/Common/CameraRigCameraNode.h"
+
+#include "Core/CameraBuildLog.h"
+#include "Core/CameraNodeEvaluator.h"
+#include "Core/CameraRigAsset.h"
+#include "Core/CameraRigBuildContext.h"
+
+#include UE_INLINE_GENERATED_CPP_BY_NAME(CameraRigCameraNode)
+
+#define LOCTEXT_NAMESPACE "CameraRigCameraNode"
+
+namespace UE::Cameras
+{
+
+namespace Internal
+{
+
+template<typename ParameterOverrideType>
+void ApplyParameterOverrides(
+		const UCameraRigAsset* CameraRig, 
+		TArrayView<const ParameterOverrideType> ParameterOverrides, 
+		FCameraVariableTable& OutVariableTable)
+{
+	for (const ParameterOverrideType& ParameterOverride : ParameterOverrides)
+	{
+		using ParameterType = decltype(ParameterOverrideType::Value);
+		using ValueType = typename ParameterType::ValueType;
+
+		FCameraVariableID InterfaceParameterID(FCameraVariableID::FromHashValue(GetTypeHash(ParameterOverride.PrivateVariableGuid)));
+
+		if (ParameterOverride.Value.Variable != nullptr)
+		{
+			// The override is driven by a variable... read its value and set it as the value for the
+			// prefab's variable. Basically, we forward the value from one variable to the next.
+			FCameraVariableDefinition OverrideDefinition(ParameterOverride.Value.Variable->GetVariableDefinition());
+
+			OutVariableTable.SetValue<ValueType>(
+					InterfaceParameterID,
+					OutVariableTable.GetValue<typename ParameterType::ValueType>(OverrideDefinition.VariableID));
+		}
+		else
+		{
+			// The override is a fixed value. Just set that on the prefab's variable.
+			OutVariableTable.SetValue<ValueType>(
+					InterfaceParameterID,
+					ParameterOverride.Value.Value);
+		}
+	}
+}
+
+}  // namespace Internal
+
+class FCameraRigCameraNodeEvaluator : public FCameraNodeEvaluator
+{
+	UE_DECLARE_CAMERA_NODE_EVALUATOR(GAMEPLAYCAMERAS_API, FCameraRigCameraNodeEvaluator)
+
+protected:
+
+	virtual FCameraNodeEvaluatorChildrenView OnGetChildren() override;
+	virtual void OnInitialize(const FCameraNodeEvaluatorInitializeParams& Params, FCameraNodeEvaluationResult& OutResult) override;
+	virtual void OnBuild(const FCameraNodeEvaluatorBuildParams& Params) override;
+	virtual void OnRun(const FCameraNodeEvaluationParams& Params, FCameraNodeEvaluationResult& OutResult) override;
+
+private:
+
+	void ApplyParameterOverrides(FCameraVariableTable& OutVariableTable);
+
+private:
+
+	FCameraNodeEvaluator* CameraRigRootEvaluator = nullptr;
+};
+
+UE_DEFINE_CAMERA_NODE_EVALUATOR(FCameraRigCameraNodeEvaluator)
+
+FCameraNodeEvaluatorChildrenView FCameraRigCameraNodeEvaluator::OnGetChildren()
+{
+	return FCameraNodeEvaluatorChildrenView({ CameraRigRootEvaluator });
+}
+
+void FCameraRigCameraNodeEvaluator::OnBuild(const FCameraNodeEvaluatorBuildParams& Params)
+{
+	const UCameraRigCameraNode* CameraRigNode = GetCameraNodeAs<UCameraRigCameraNode>();
+	if (CameraRigNode->CameraRig)
+	{
+		if (CameraRigNode->CameraRig->RootNode)
+		{
+			CameraRigRootEvaluator = Params.BuildEvaluator(CameraRigNode->CameraRig->RootNode);
+		}
+	}
+}
+
+void FCameraRigCameraNodeEvaluator::OnInitialize(const FCameraNodeEvaluatorInitializeParams& Params, FCameraNodeEvaluationResult& OutResult)
+{
+	// Apply overrides right away.
+	ApplyParameterOverrides(OutResult.VariableTable);
+}
+
+void FCameraRigCameraNodeEvaluator::OnRun(const FCameraNodeEvaluationParams& Params, FCameraNodeEvaluationResult& OutResult)
+{
+	// Keep applying overrides in case they are driven by a variable.
+	// TODO: we could skip this step for constant overrides.
+	ApplyParameterOverrides(OutResult.VariableTable);
+
+	if (CameraRigRootEvaluator)
+	{
+		CameraRigRootEvaluator->Run(Params, OutResult);
+	}
+}
+
+void FCameraRigCameraNodeEvaluator::ApplyParameterOverrides(FCameraVariableTable& OutVariableTable)
+{
+	const UCameraRigCameraNode* PrefabNode = GetCameraNodeAs<UCameraRigCameraNode>();
+
+#define UE_CAMERA_VARIABLE_FOR_TYPE(ValueType, ValueName)\
+	Internal::ApplyParameterOverrides(\
+			PrefabNode->CameraRig,\
+			TArrayView<const F##ValueName##CameraRigParameterOverride>(PrefabNode->ValueName##Overrides),\
+			OutVariableTable);
+UE_CAMERA_VARIABLE_FOR_ALL_TYPES()
+#undef UE_CAMERA_VARIABLE_FOR_TYPE
+}
+
+namespace Internal
+{
+
+struct FCameraRigCameraNodeBuilder
+{
+	UCameraRigCameraNode* CameraNode;
+	FCameraRigBuildContext& BuildContext;
+
+	FCameraRigCameraNodeBuilder(UCameraRigCameraNode* InCameraNode, FCameraRigBuildContext& InBuildContext)
+		: CameraNode(InCameraNode)
+		, BuildContext(InBuildContext)
+	{}
+
+	void Setup(const UCameraRigAsset* InCameraRig)
+	{
+		ParametersByGuid.Reset();
+		for (TObjectPtr<UCameraRigInterfaceParameter> InterfaceParameter : InCameraRig->Interface.InterfaceParameters)
+		{
+			ParametersByGuid.Add(InterfaceParameter->Guid, InterfaceParameter);
+		}
+	}
+
+	const UCameraRigInterfaceParameter* FindInterfaceParameter(const FGuid& InterfaceParameterGuid)
+	{
+		if (UCameraRigInterfaceParameter** FoundItem = ParametersByGuid.Find(InterfaceParameterGuid))
+		{
+			return *FoundItem;
+		}
+		return nullptr;
+	}
+
+	template<typename ParameterOverrideType>
+	void BuildCameraRigParameterOverride(ParameterOverrideType& ParameterOverride);
+
+private:
+
+	TMap<FGuid, UCameraRigInterfaceParameter*> ParametersByGuid;
+};
+
+template<typename ParameterOverrideType>
+void FCameraRigCameraNodeBuilder::BuildCameraRigParameterOverride(ParameterOverrideType& ParameterOverride)
+{
+	const UCameraRigInterfaceParameter* InterfaceParameter = FindInterfaceParameter(ParameterOverride.InterfaceParameterGuid);
+	if (!InterfaceParameter)
+	{
+		BuildContext.BuildLog.AddMessage(EMessageSeverity::Error, CameraNode,
+				FText::Format(
+					LOCTEXT("MissingInterfaceParameter", "No camera rig interface parameter named '{0}' exists on '{1}'."),
+					FText::FromString(ParameterOverride.InterfaceParameterName),
+					FText::FromString(GetNameSafe(CameraNode->CameraRig))));
+		return;
+	}
+
+	if (!InterfaceParameter->PrivateVariable)
+	{
+		BuildContext.BuildLog.AddMessage(EMessageSeverity::Error, CameraNode,
+				FText::Format(
+					LOCTEXT("UnbuiltInterfaceParameter", "Camera rig interface parameter '{0}' was not built correctly on '{1}'."),
+					FText::FromString(ParameterOverride.InterfaceParameterName),
+					FText::FromString(GetNameSafe(CameraNode->CameraRig))));
+		return;
+	}
+
+	ParameterOverride.PrivateVariableGuid = InterfaceParameter->PrivateVariable->GetGuid();
+}
+
+}  // namespace Internal
+
+}  // namespace UE::Cameras
+
+void UCameraRigCameraNode::OnBuild(FCameraRigBuildContext& BuildContext)
+{
+	using namespace UE::Cameras;
+	using namespace UE::Cameras::Internal;
+
+	if (!CameraRig)
+	{
+		BuildContext.BuildLog.AddMessage(EMessageSeverity::Error, this, LOCTEXT("MissingCameraRig", "No camera rig specified on camera rig node."));
+		return;
+	}
+
+	// Build the inner camera rig. Whatever allocations it needs for its evaluators and
+	// their camera variables, we add that to our camera rig's allocation info.
+	CameraRig->BuildCameraRig(BuildContext.BuildLog);
+	BuildContext.AllocationInfo.Append(CameraRig->AllocationInfo);
+
+	// Next, we set things up for the runtime. Mostly, we want to get the camera variable 
+	// Guids that we need to write the override values to.
+	FCameraRigCameraNodeBuilder InternalBuilder(this, BuildContext);
+	InternalBuilder.Setup(CameraRig);
+
+#define UE_CAMERA_VARIABLE_FOR_TYPE(ValueType, ValueName)\
+	{\
+		for (F##ValueName##CameraRigParameterOverride& ParameterOverride : ValueName##Overrides)\
+		{\
+			InternalBuilder.BuildCameraRigParameterOverride(ParameterOverride);\
+		}\
+	}
+UE_CAMERA_VARIABLE_FOR_ALL_TYPES()
+#undef UE_CAMERA_VARIABLE_FOR_TYPE
+
+}
+
+FCameraNodeEvaluatorPtr UCameraRigCameraNode::OnBuildEvaluator(FCameraNodeEvaluatorBuilder& Builder) const
+{
+	using namespace UE::Cameras;
+	return Builder.BuildEvaluator<FCameraRigCameraNodeEvaluator>();
+}
+
+#undef LOCTEXT_NAMESPACE
+
diff --git a/Engine/Plugins/Cameras/GameplayCameras/Source/GameplayCameras/Public/Core/CameraParameters.h b/Engine/Plugins/Cameras/GameplayCameras/Source/GameplayCameras/Public/Core/CameraParameters.h
--- a/Engine/Plugins/Cameras/GameplayCameras/Source/GameplayCameras/Public/Core/CameraParameters.h
+++ b/Engine/Plugins/Cameras/GameplayCameras/Source/GameplayCameras/Public/Core/CameraParameters.h
@@ -308,6 +308,24 @@ struct GAMEPLAYCAMERAS_API FTransform3dCameraParameter
 #undef UE_DEFINE_CAMERA_PARAMETER_VALUE_CONSTRUCTORS
 #undef UE_DEFINE_CAMERA_PARAMETER_ALL_CONSTRUCTORS
 
+template<typename ValueType>
+bool CameraParameterValueEquals(typename TCallTraits<ValueType>::ParamType A, typename TCallTraits<ValueType>::ParamType B)
+{
+	return A == B;
+}
+
+template<>
+inline bool CameraParameterValueEquals<FTransform3f>(const FTransform3f& A, const FTransform3f& B)
+{
+	return A.Equals(B);
+}
+
+template<>
+inline bool CameraParameterValueEquals<FTransform3d>(const FTransform3d& A, const FTransform3d& B)
+{
+	return A.Equals(B);
+}
+
 // Any camera parameter might replace a previously non-parameterized property (i.e. a "fixed" property
 // of the underlying type, like bool, int32, float, etc.)
 // When someone upgrades the fixed property to a parameterized property, any previously saved data will
diff --git a/Engine/Plugins/Cameras/GameplayCameras/Source/GameplayCameras/Public/Core/CameraRigAsset.h b/Engine/Plugins/Cameras/GameplayCameras/Source/GameplayCameras/Public/Core/CameraRigAsset.h
--- a/Engine/Plugins/Cameras/GameplayCameras/Source/GameplayCameras/Public/Core/CameraRigAsset.h
+++ b/Engine/Plugins/Cameras/GameplayCameras/Source/GameplayCameras/Public/Core/CameraRigAsset.h
@@ -39,6 +39,10 @@ struct FCameraRigAllocationInfo
 	/** Allocation info for the camera variable. */
 	UPROPERTY()
 	FCameraVariableTableAllocationInfo VariableTableInfo;
+
+public:
+
+	void Append(const FCameraRigAllocationInfo& OtherAllocationInfo);
 };
 
 /**
diff --git a/Engine/Plugins/Cameras/GameplayCameras/Source/GameplayCameras/Public/Nodes/Common/CameraRigCameraNode.h b/Engine/Plugins/Cameras/GameplayCameras/Source/GameplayCameras/Public/Nodes/Common/CameraRigCameraNode.h
new file mode 100644
--- /dev/null
+++ b/Engine/Plugins/Cameras/GameplayCameras/Source/GameplayCameras/Public/Nodes/Common/CameraRigCameraNode.h
@@ -0,0 +1,428 @@
+// Copyright Epic Games, Inc. All Rights Reserved.
+
+#pragma once
+
+#include "Core/CameraNode.h"
+#include "Core/CameraParameters.h"
+#include "Core/CameraRigAsset.h"
+
+#include "CameraRigCameraNode.generated.h"
+
+namespace UE::Cameras
+{
+	class FCameraRigCameraNodeEvaluator;
+}
+
+USTRUCT()
+struct FBooleanCameraRigParameterOverride
+{
+	GENERATED_BODY()
+
+	using CameraParameterType = FBooleanCameraParameter;
+
+	UPROPERTY()
+	FBooleanCameraParameter Value;
+
+	UPROPERTY()
+	FGuid PrivateVariableGuid;
+
+	UPROPERTY()
+	FGuid InterfaceParameterGuid;
+
+	UPROPERTY()
+	FString InterfaceParameterName;
+};
+
+USTRUCT()
+struct FInteger32CameraRigParameterOverride
+{
+	GENERATED_BODY()
+
+	using CameraParameterType = FInteger32CameraParameter;
+
+	UPROPERTY()
+	FInteger32CameraParameter Value;
+
+	UPROPERTY()
+	FGuid PrivateVariableGuid;
+
+	UPROPERTY()
+	FGuid InterfaceParameterGuid;
+
+	UPROPERTY()
+	FString InterfaceParameterName;
+};
+
+USTRUCT()
+struct FFloatCameraRigParameterOverride
+{
+	GENERATED_BODY()
+
+	using CameraParameterType = FFloatCameraParameter;
+
+	UPROPERTY()
+	FFloatCameraParameter Value;
+
+	UPROPERTY()
+	FGuid PrivateVariableGuid;
+
+	UPROPERTY()
+	FGuid InterfaceParameterGuid;
+
+	UPROPERTY()
+	FString InterfaceParameterName;
+};
+
+USTRUCT()
+struct FDoubleCameraRigParameterOverride
+{
+	GENERATED_BODY()
+
+	using CameraParameterType = FDoubleCameraParameter;
+
+	UPROPERTY()
+	FDoubleCameraParameter Value;
+
+	UPROPERTY()
+	FGuid PrivateVariableGuid;
+
+	UPROPERTY()
+	FGuid InterfaceParameterGuid;
+
+	UPROPERTY()
+	FString InterfaceParameterName;
+};
+
+USTRUCT()
+struct FVector2fCameraRigParameterOverride
+{
+	GENERATED_BODY()
+
+	using CameraParameterType = FVector2fCameraParameter;
+
+	UPROPERTY()
+	FVector2fCameraParameter Value;
+
+	UPROPERTY()
+	FGuid PrivateVariableGuid;
+
+	UPROPERTY()
+	FGuid InterfaceParameterGuid;
+
+	UPROPERTY()
+	FString InterfaceParameterName;
+};
+
+USTRUCT()
+struct FVector2dCameraRigParameterOverride
+{
+	GENERATED_BODY()
+
+	using CameraParameterType = FVector2dCameraParameter;
+
+	UPROPERTY()
+	FVector2dCameraParameter Value;
+
+	UPROPERTY()
+	FGuid PrivateVariableGuid;
+
+	UPROPERTY()
+	FGuid InterfaceParameterGuid;
+
+	UPROPERTY()
+	FString InterfaceParameterName;
+};
+
+USTRUCT()
+struct FVector3fCameraRigParameterOverride
+{
+	GENERATED_BODY()
+
+	using CameraParameterType = FVector3fCameraParameter;
+
+	UPROPERTY()
+	FVector3fCameraParameter Value;
+
+	UPROPERTY()
+	FGuid PrivateVariableGuid;
+
+	UPROPERTY()
+	FGuid InterfaceParameterGuid;
+
+	UPROPERTY()
+	FString InterfaceParameterName;
+};
+
+USTRUCT()
+struct FVector3dCameraRigParameterOverride
+{
+	GENERATED_BODY()
+
+	using CameraParameterType = FVector3dCameraParameter;
+
+	UPROPERTY()
+	FVector3dCameraParameter Value;
+
+	UPROPERTY()
+	FGuid PrivateVariableGuid;
+
+	UPROPERTY()
+	FGuid InterfaceParameterGuid;
+
+	UPROPERTY()
+	FString InterfaceParameterName;
+};
+
+USTRUCT()
+struct FVector4fCameraRigParameterOverride
+{
+	GENERATED_BODY()
+
+	using CameraParameterType = FVector4fCameraParameter;
+
+	UPROPERTY()
+	FVector4fCameraParameter Value;
+
+	UPROPERTY()
+	FGuid PrivateVariableGuid;
+
+	UPROPERTY()
+	FGuid InterfaceParameterGuid;
+
+	UPROPERTY()
+	FString InterfaceParameterName;
+};
+
+USTRUCT()
+struct FVector4dCameraRigParameterOverride
+{
+	GENERATED_BODY()
+
+	using CameraParameterType = FVector4dCameraParameter;
+
+	UPROPERTY()
+	FVector4dCameraParameter Value;
+
+	UPROPERTY()
+	FGuid PrivateVariableGuid;
+
+	UPROPERTY()
+	FGuid InterfaceParameterGuid;
+
+	UPROPERTY()
+	FString InterfaceParameterName;
+};
+
+USTRUCT()
+struct FRotator3fCameraRigParameterOverride
+{
+	GENERATED_BODY()
+
+	using CameraParameterType = FRotator3fCameraParameter;
+
+	UPROPERTY()
+	FRotator3fCameraParameter Value;
+
+	UPROPERTY()
+	FGuid PrivateVariableGuid;
+
+	UPROPERTY()
+	FGuid InterfaceParameterGuid;
+
+	UPROPERTY()
+	FString InterfaceParameterName;
+};
+
+USTRUCT()
+struct FRotator3dCameraRigParameterOverride
+{
+	GENERATED_BODY()
+
+	using CameraParameterType = FRotator3dCameraParameter;
+
+	UPROPERTY()
+	FRotator3dCameraParameter Value;
+
+	UPROPERTY()
+	FGuid PrivateVariableGuid;
+
+	UPROPERTY()
+	FGuid InterfaceParameterGuid;
+
+	UPROPERTY()
+	FString InterfaceParameterName;
+};
+
+USTRUCT()
+struct FTransform3fCameraRigParameterOverride
+{
+	GENERATED_BODY()
+
+	using CameraParameterType = FTransform3fCameraParameter;
+
+	UPROPERTY()
+	FTransform3fCameraParameter Value;
+
+	UPROPERTY()
+	FGuid PrivateVariableGuid;
+
+	UPROPERTY()
+	FGuid InterfaceParameterGuid;
+
+	UPROPERTY()
+	FString InterfaceParameterName;
+};
+
+USTRUCT()
+struct FTransform3dCameraRigParameterOverride
+{
+	GENERATED_BODY()
+
+	using CameraParameterType = FTransform3dCameraParameter;
+
+	UPROPERTY()
+	FTransform3dCameraParameter Value;
+
+	UPROPERTY()
+	FGuid PrivateVariableGuid;
+
+	UPROPERTY()
+	FGuid InterfaceParameterGuid;
+
+	UPROPERTY()
+	FString InterfaceParameterName;
+};
+
+/**
+ * A camera node that runs a camera rig's own node tree.
+ */
+UCLASS(MinimalAPI, meta=(CameraNodeCategories="Common,Utility"))
+class UCameraRigCameraNode : public UCameraNode
+{
+	GENERATED_BODY()
+
+public:
+
+	template<typename ParameterOverrideType>
+	ParameterOverrideType* FindParameterOverride(const FGuid& CameraRigParameterGuid);
+	
+	template<typename ParameterOverrideType>
+	ParameterOverrideType& GetOrAddParameterOverride(const UCameraRigInterfaceParameter* CameraRigParameter);
+
+	template<typename ParameterOverrideType>
+	void RemoveParameterOverride(const FGuid& CameraRigParameterGuid);
+
+private:
+
+	template<typename ParameterOverrideType>
+	ParameterOverrideType* FindParameterOverride(TArray<ParameterOverrideType>& OverridesArray, const FGuid& CameraRigParameterGuid);
+	
+	template<typename ParameterOverrideType>
+	ParameterOverrideType& GetOrAddParameterOverride(TArray<ParameterOverrideType>& OverridesArray, const UCameraRigInterfaceParameter* CameraRigParameter);
+
+	template<typename ParameterOverrideType>
+	void RemoveParameterOverride(TArray<ParameterOverrideType>& OverridesArray, const FGuid& CameraRigParameterGuid);
+
+protected:
+
+	// UCameraNode interface.
+	virtual void OnBuild(FCameraRigBuildContext& BuildContext) override;
+	virtual FCameraNodeEvaluatorPtr OnBuildEvaluator(FCameraNodeEvaluatorBuilder& Builder) const override;
+
+public:
+
+	/** The camera rig to run. */
+	UPROPERTY(EditAnywhere, Category=Common, meta=(ObjectTreeGraphHidden=true))
+	TObjectPtr<UCameraRigAsset> CameraRig;
+
+private:
+
+	UPROPERTY()
+	TArray<FBooleanCameraRigParameterOverride> BooleanOverrides;
+	UPROPERTY()
+	TArray<FInteger32CameraRigParameterOverride> Integer32Overrides;
+	UPROPERTY()
+	TArray<FFloatCameraRigParameterOverride> FloatOverrides;
+	UPROPERTY()
+	TArray<FDoubleCameraRigParameterOverride> DoubleOverrides;
+	UPROPERTY()
+	TArray<FVector2fCameraRigParameterOverride> Vector2fOverrides;
+	UPROPERTY()
+	TArray<FVector2dCameraRigParameterOverride> Vector2dOverrides;
+	UPROPERTY()
+	TArray<FVector3fCameraRigParameterOverride> Vector3fOverrides;
+	UPROPERTY()
+	TArray<FVector3dCameraRigParameterOverride> Vector3dOverrides;
+	UPROPERTY()
+	TArray<FVector4fCameraRigParameterOverride> Vector4fOverrides;
+	UPROPERTY()
+	TArray<FVector4dCameraRigParameterOverride> Vector4dOverrides;
+	UPROPERTY()
+	TArray<FRotator3fCameraRigParameterOverride> Rotator3fOverrides;
+	UPROPERTY()
+	TArray<FRotator3dCameraRigParameterOverride> Rotator3dOverrides;
+	UPROPERTY()
+	TArray<FTransform3fCameraRigParameterOverride> Transform3fOverrides;
+	UPROPERTY()
+	TArray<FTransform3dCameraRigParameterOverride> Transform3dOverrides;
+
+	friend class UE::Cameras::FCameraRigCameraNodeEvaluator;
+};
+
+template<typename ParameterOverrideType>
+ParameterOverrideType* UCameraRigCameraNode::FindParameterOverride(TArray<ParameterOverrideType>& OverridesArray, const FGuid& CameraRigParameterGuid)
+{
+	ParameterOverrideType* FoundItem = OverridesArray.FindByPredicate(
+			[CameraRigParameterGuid](ParameterOverrideType& Item)
+			{
+				return (Item.InterfaceParameterGuid == CameraRigParameterGuid);
+			});
+	return FoundItem;
+}
+
+template<typename ParameterOverrideType>
+ParameterOverrideType& UCameraRigCameraNode::GetOrAddParameterOverride(TArray<ParameterOverrideType>& OverridesArray, const UCameraRigInterfaceParameter* CameraRigParameter)
+{
+	ParameterOverrideType* Existing = FindParameterOverride<ParameterOverrideType>(CameraRigParameter->Guid);
+	if (Existing)
+	{
+		return *Existing;
+	}
+	else
+	{
+		ParameterOverrideType& NewOverride = OverridesArray.Emplace_GetRef();
+		NewOverride.InterfaceParameterGuid = CameraRigParameter->Guid;
+		NewOverride.InterfaceParameterName = CameraRigParameter->InterfaceParameterName;
+		return NewOverride;
+	}
+}
+
+template<typename ParameterOverrideType>
+void UCameraRigCameraNode::RemoveParameterOverride(TArray<ParameterOverrideType>& OverridesArray, const FGuid& CameraRigParameterGuid)
+{
+	OverridesArray.RemoveAll(
+			[CameraRigParameterGuid](ParameterOverrideType& Item)
+			{
+				return (Item.InterfaceParameterGuid == CameraRigParameterGuid);
+			});
+}
+
+#define UE_CAMERA_VARIABLE_FOR_TYPE(ValueType, ValueName)\
+template<>\
+inline F##ValueName##CameraRigParameterOverride* UCameraRigCameraNode::FindParameterOverride<F##ValueName##CameraRigParameterOverride>(const FGuid& CameraRigParameterGuid)\
+{\
+	return FindParameterOverride<F##ValueName##CameraRigParameterOverride>(ValueName##Overrides, CameraRigParameterGuid);\
+}\
+template<>\
+inline F##ValueName##CameraRigParameterOverride& UCameraRigCameraNode::GetOrAddParameterOverride<F##ValueName##CameraRigParameterOverride>(const UCameraRigInterfaceParameter* CameraRigParameter)\
+{\
+	return GetOrAddParameterOverride<F##ValueName##CameraRigParameterOverride>(ValueName##Overrides, CameraRigParameter);\
+}\
+template<>\
+inline void UCameraRigCameraNode::RemoveParameterOverride<F##ValueName##CameraRigParameterOverride>(const FGuid& CameraRigParameterGuid)\
+{\
+	RemoveParameterOverride<F##ValueName##CameraRigParameterOverride>(ValueName##Overrides, CameraRigParameterGuid);\
+}
+UE_CAMERA_VARIABLE_FOR_ALL_TYPES()
+#undef UE_CAMERA_VARIABLE_FOR_TYPE
+
diff --git a/Engine/Plugins/Cameras/GameplayCameras/Source/GameplayCamerasEditor/Private/Customizations/CameraRigCameraNodeDetailsCustomization.cpp b/Engine/Plugins/Cameras/GameplayCameras/Source/GameplayCamerasEditor/Private/Customizations/CameraRigCameraNodeDetailsCustomization.cpp
new file mode 100644
--- /dev/null
+++ b/Engine/Plugins/Cameras/GameplayCameras/Source/GameplayCamerasEditor/Private/Customizations/CameraRigCameraNodeDetailsCustomization.cpp
@@ -0,0 +1,228 @@
+// Copyright Epic Games, Inc. All Rights Reserved.
+
+#include "Customizations/CameraRigCameraNodeDetailsCustomization.h"
+
+#include "Core/CameraRigAsset.h"
+#include "DetailCategoryBuilder.h"
+#include "DetailLayoutBuilder.h"
+#include "IDetailChildrenBuilder.h"
+#include "IGameplayCamerasEditorModule.h"
+#include "IPropertyUtilities.h"
+#include "Modules/ModuleManager.h"
+#include "Nodes/Common/CameraRigCameraNode.h"
+#include "PropertyCustomizationHelpers.h"
+
+#define LOCTEXT_NAMESPACE "CameraRigCameraNodeDetailsCustomization"
+
+namespace UE::Cameras
+{
+
+class FParameterOverrideDetailRows : public TSharedFromThis<FParameterOverrideDetailRows>
+{
+public:
+
+	FParameterOverrideDetailRows(
+			UCameraRigCameraNode* InCameraRigNode,
+			IDetailLayoutBuilder& InLayoutBuilder,
+			IDetailCategoryBuilder& InOverridesCategory, 
+			TSharedRef<IPropertyUtilities> InPropertyUtilities)
+		: CameraRigNode(InCameraRigNode)
+		, LayoutBuilder(InLayoutBuilder)
+		, OverridesCategory(InOverridesCategory)
+		, PropertyUtilities(InPropertyUtilities)
+	{}
+
+	void AddParameterOverrideDetailRow(const UCameraRigInterfaceParameter* InterfaceParameter)
+	{
+		if (!InterfaceParameter)
+		{
+			return;
+		}
+
+		UCameraNode* TargetNode = InterfaceParameter->Target;
+		FName TargetPropertyName = InterfaceParameter->TargetPropertyName;
+		if (!TargetNode || TargetPropertyName.IsNone())
+		{
+			return;
+		}
+
+		UClass* TargetNodeClass = TargetNode->GetClass();
+		FStructProperty* TargetProperty = CastField<FStructProperty>(TargetNodeClass->FindPropertyByName(TargetPropertyName));
+		if (!TargetProperty)
+		{
+			return;
+		}
+
+#define UE_CAMERA_VARIABLE_FOR_TYPE(ValueType, ValueName)\
+		if (TargetProperty->Struct == F##ValueName##CameraParameter::StaticStruct())\
+		{\
+			using FCameraRigParameterOverrideType = F##ValueName##CameraRigParameterOverride;\
+			AddParameterOverrideDetailRowImpl<FCameraRigParameterOverrideType>(InterfaceParameter, TargetProperty);\
+		}\
+		else
+		UE_CAMERA_VARIABLE_FOR_ALL_TYPES()
+#undef UE_CAMERA_VARIABLE_FOR_TYPE
+		{
+			// Unknown sort of property...
+		}
+	}
+
+private:
+
+	template<typename CameraRigParameterOverrideType>
+	void AddParameterOverrideDetailRowImpl(const UCameraRigInterfaceParameter* InterfaceParameter, FStructProperty* TargetProperty)
+	{
+		using CameraParameterType = typename CameraRigParameterOverrideType::CameraParameterType;
+
+		const FGuid ParameterGuid = InterfaceParameter->Guid;
+
+		// Get the default value of this parameter from the original target node.
+		UObject* TargetNode = InterfaceParameter->Target;
+		const CameraParameterType* DefaultParameterPtr = TargetProperty->ContainerPtrToValuePtr<CameraParameterType>(TargetNode);
+
+		// Create a new version of the target node, and get a pointer to the parameter we're interested
+		// in showing in the details view.
+		UClass* TargetNodeClass = InterfaceParameter->Target->GetClass();
+		UObject* WrapperObject = NewObject<UObject>(GetTransientPackage(), TargetNodeClass);
+		CameraParameterType* ScratchParameterPtr = TargetProperty->ContainerPtrToValuePtr<CameraParameterType>(WrapperObject);
+		{
+			// If we already have an override, show its value. Otherwise, show the original value.
+			CameraRigParameterOverrideType* Override = CameraRigNode->FindParameterOverride<CameraRigParameterOverrideType>(ParameterGuid);
+			if (Override)
+			{
+				*ScratchParameterPtr = Override->Value;
+			}
+			else
+			{
+				ScratchParameterPtr->Value = DefaultParameterPtr->Value;
+			}
+		}
+		// Make sure this temporary object isn't GC'ed while the details view is still shown.
+		WrapperObjects.Add(TStrongObjectPtr<UObject>(WrapperObject));
+
+		// Add the scratch parameter to the details view.
+		TSharedPtr<IPropertyHandle> ParameterPropertyHandle = LayoutBuilder.AddObjectPropertyData(
+				{ WrapperObject }, InterfaceParameter->TargetPropertyName);
+		
+		FText ParameterNameText(FText::FromString(InterfaceParameter->InterfaceParameterName));
+		ParameterPropertyHandle->SetPropertyDisplayName(ParameterNameText);
+
+		IDetailPropertyRow& PropertyRow = OverridesCategory.AddProperty(ParameterPropertyHandle);
+		
+		// Register a change callback so that we can propagate edits made to the scratch parameter onto 
+		// the camera rig node's overrides.
+		PropertyRow.GetPropertyHandle()->SetOnPropertyValueChanged(FSimpleDelegate::CreateSP(
+					this,
+					&FParameterOverrideDetailRows::OnPropertyValueChanged<CameraRigParameterOverrideType, CameraParameterType>,
+					InterfaceParameter, DefaultParameterPtr, ScratchParameterPtr));
+		// Register the same callback for when the user edits a child property, such as the X/Y/Z coordinates
+		// of a 3d vector camera parameter.
+		PropertyRow.GetPropertyHandle()->SetOnChildPropertyValueChanged(FSimpleDelegate::CreateSP(
+					this,
+					&FParameterOverrideDetailRows::OnPropertyValueChanged<CameraRigParameterOverrideType, CameraParameterType>,
+					InterfaceParameter, DefaultParameterPtr, ScratchParameterPtr));
+
+		// Setup custom reset-to-default logic.
+		FResetToDefaultOverride ResetToDefault = FResetToDefaultOverride::Create(
+				FIsResetToDefaultVisible::CreateLambda(
+					[DefaultParameterPtr, ScratchParameterPtr](TSharedPtr<IPropertyHandle> PropertyHandle) -> bool
+					{
+						return ScratchParameterPtr->Variable != nullptr ||
+							!CameraParameterValueEquals<typename CameraParameterType::ValueType>(
+								ScratchParameterPtr->Value, DefaultParameterPtr->Value);
+					}),
+				FResetToDefaultHandler::CreateLambda(
+					[DefaultParameterPtr, ScratchParameterPtr](TSharedPtr<IPropertyHandle> PropertyHandle)
+					{
+						ScratchParameterPtr->Value = DefaultParameterPtr->Value;
+						ScratchParameterPtr->Variable = nullptr;
+					}));
+		PropertyRow.OverrideResetToDefault(ResetToDefault);
+	}
+
+	template<typename CameraRigParameterOverrideType, typename CameraParameterType>
+	void OnPropertyValueChanged(
+			const UCameraRigInterfaceParameter* InterfaceParameter, 
+			const CameraParameterType* DefaultParameterPtr,
+			CameraParameterType* ScratchParameterPtr)
+	{
+		CameraRigNode->Modify();
+
+		const bool bEqualValues = CameraParameterValueEquals<typename CameraParameterType::ValueType>(
+				ScratchParameterPtr->Value, DefaultParameterPtr->Value);
+		if (bEqualValues && ScratchParameterPtr->Variable == nullptr)
+		{
+			CameraRigNode->RemoveParameterOverride<CameraRigParameterOverrideType>(InterfaceParameter->Guid);
+		}
+		else
+		{
+			CameraRigParameterOverrideType& Override = CameraRigNode->GetOrAddParameterOverride<CameraRigParameterOverrideType>(InterfaceParameter);
+			Override.Value = *ScratchParameterPtr;
+		}
+
+		PropertyUtilities->RequestForceRefresh();
+	}
+
+private:
+
+	UCameraRigCameraNode* CameraRigNode;
+	IDetailLayoutBuilder& LayoutBuilder;
+	IDetailCategoryBuilder& OverridesCategory;
+	TSharedRef<IPropertyUtilities> PropertyUtilities;
+
+	TArray<TStrongObjectPtr<UObject>> WrapperObjects;
+};
+
+TSharedRef<IDetailCustomization> FCameraRigCameraNodeDetailsCustomization::MakeInstance()
+{
+	return MakeShared<FCameraRigCameraNodeDetailsCustomization>();
+}
+
+void FCameraRigCameraNodeDetailsCustomization::CustomizeDetails(IDetailLayoutBuilder& DetailBuilder)
+{
+	// Only support a single node selected for now.
+	TArray<TWeakObjectPtr<UCameraRigCameraNode>> WeakNodes = DetailBuilder.GetSelectedObjectsOfType<UCameraRigCameraNode>();
+	if (WeakNodes.IsEmpty())
+	{
+		return;
+	}
+
+	// Skip invalid nodes.
+	UCameraRigCameraNode* CameraRigNode = WeakNodes[0].Get();
+	if (!CameraRigNode || !CameraRigNode->CameraRig)
+	{
+		return;
+	}
+
+	// Skip camera rigs that don't have anything exposed.
+	FCameraRigInterface& Interface = CameraRigNode->CameraRig->Interface;
+	if (Interface.InterfaceParameters.IsEmpty())
+	{
+		return;
+	}
+
+	IDetailCategoryBuilder& OverridesCategory = DetailBuilder.EditCategory(
+			TEXT("ParameterOverrides"), 
+			LOCTEXT("ParameterOverridesCategory", "Parameter Overrides"));
+
+	ParameterOverrideRows = MakeShared<FParameterOverrideDetailRows>(
+			CameraRigNode,
+			DetailBuilder,
+			OverridesCategory,
+			DetailBuilder.GetPropertyUtilities());
+
+	for (const UCameraRigInterfaceParameter* InterfaceParameter : Interface.InterfaceParameters)
+	{
+		ParameterOverrideRows->AddParameterOverrideDetailRow(InterfaceParameter);
+	}
+}
+
+void FCameraRigCameraNodeDetailsCustomization::PendingDelete()
+{
+	ParameterOverrideRows.Reset();
+}
+
+}  // namespace UE::Cameras
+
+#undef LOCTEXT_NAMESPACE
+
diff --git a/Engine/Plugins/Cameras/GameplayCameras/Source/GameplayCamerasEditor/Private/GameplayCamerasEditorModule.cpp b/Engine/Plugins/Cameras/GameplayCameras/Source/GameplayCamerasEditor/Private/GameplayCamerasEditorModule.cpp
--- a/Engine/Plugins/Cameras/GameplayCameras/Source/GameplayCamerasEditor/Private/GameplayCamerasEditorModule.cpp
+++ b/Engine/Plugins/Cameras/GameplayCameras/Source/GameplayCamerasEditor/Private/GameplayCamerasEditorModule.cpp
@@ -13,6 +13,7 @@
 #include "Commands/CameraVariableCollectionEditorCommands.h"
 #include "Commands/GameplayCamerasDebuggerCommands.h"
 #include "Customizations/CameraParameterDetailsCustomizations.h"
+#include "Customizations/CameraRigCameraNodeDetailsCustomization.h"
 #include "Customizations/CameraProxyTableDetailsCustomization.h"
 #include "Customizations/SingleCameraDirectorDetailsCustomization.h"
 #include "Debug/CameraDebugCategories.h"
@@ -453,6 +454,8 @@ private:
 
 		FPropertyEditorModule& PropertyEditorModule = FModuleManager::GetModuleChecked<FPropertyEditorModule>("PropertyEditor");
 		FCameraParameterDetailsCustomization::Register(PropertyEditorModule);
+		PropertyEditorModule.RegisterCustomClassLayout("CameraRigCameraNode", FOnGetDetailCustomizationInstance::CreateStatic(
+					&FCameraRigCameraNodeDetailsCustomization::MakeInstance));
 		PropertyEditorModule.RegisterCustomPropertyTypeLayout("CameraRigProxyTableEntry", FOnGetPropertyTypeCustomizationInstance::CreateStatic(
 					&FCameraProxyTableEntryDetailsCustomization::MakeInstance));
 		PropertyEditorModule.RegisterCustomClassLayout("SingleCameraDirector", FOnGetDetailCustomizationInstance::CreateStatic(
diff --git a/Engine/Plugins/Cameras/GameplayCameras/Source/GameplayCamerasEditor/Public/Customizations/CameraRigCameraNodeDetailsCustomization.h b/Engine/Plugins/Cameras/GameplayCameras/Source/GameplayCamerasEditor/Public/Customizations/CameraRigCameraNodeDetailsCustomization.h
new file mode 100644
--- /dev/null
+++ b/Engine/Plugins/Cameras/GameplayCameras/Source/GameplayCamerasEditor/Public/Customizations/CameraRigCameraNodeDetailsCustomization.h
@@ -0,0 +1,33 @@
+// Copyright Epic Games, Inc. All Rights Reserved.
+
+#pragma once
+
+#include "IDetailCustomization.h"
+#include "Core/CameraParameters.h"
+
+class UCameraRigInterfaceParameterWrapper;
+
+namespace UE::Cameras
+{
+
+class FParameterOverrideDetailRows;
+
+class FCameraRigCameraNodeDetailsCustomization : public IDetailCustomization
+{
+public:
+
+	static TSharedRef<IDetailCustomization> MakeInstance();
+
+public:
+
+	/** IDetailCustomization interface */
+	virtual void CustomizeDetails(IDetailLayoutBuilder& DetailBuilder) override;
+	virtual void PendingDelete() override;
+
+private:
+
+	TSharedPtr<FParameterOverrideDetailRows> ParameterOverrideRows;
+};
+
+}  // namespace UE::Cameras
+
