Cameras: fix bool parameter's reading from mismatched tag

Apparently boolean values aren't fully written as booleans to save space so the real value needs to be read from `BoolVal`.

diff --git a/Engine/Plugins/Cameras/GameplayCameras/Source/GameplayCameras/Private/Core/CameraParameters.cpp b/Engine/Plugins/Cameras/GameplayCameras/Source/GameplayCameras/Private/Core/CameraParameters.cpp
--- a/Engine/Plugins/Cameras/GameplayCameras/Source/GameplayCameras/Private/Core/CameraParameters.cpp
+++ b/Engine/Plugins/Cameras/GameplayCameras/Source/GameplayCameras/Private/Core/CameraParameters.cpp
@@ -13,7 +13,7 @@ bool FBooleanCameraParameter::SerializeFromMismatchedTag(const FPropertyTag& Tag
 {
 	if (Tag.Type == NAME_BoolProperty)
 	{
-		Slot << Value;
+		Value = (Tag.BoolVal != 0);
 		return true;
 	}
 
