Cameras: find the input component on the context owner, not just its outer objects

diff --git a/Engine/Plugins/Cameras/GameplayCameras/Source/GameplayCameras/Private/Nodes/Input/InputAxisBinding2DCameraNode.cpp b/Engine/Plugins/Cameras/GameplayCameras/Source/GameplayCameras/Private/Nodes/Input/InputAxisBinding2DCameraNode.cpp
--- a/Engine/Plugins/Cameras/GameplayCameras/Source/GameplayCameras/Private/Nodes/Input/InputAxisBinding2DCameraNode.cpp
+++ b/Engine/Plugins/Cameras/GameplayCameras/Source/GameplayCameras/Private/Nodes/Input/InputAxisBinding2DCameraNode.cpp
@@ -35,7 +35,11 @@ void FInputAxisBinding2DCameraNodeEvaluator::OnInitialize(const FCameraNodeEvalu
 	UObject* ContextOwner = Params.EvaluationContext->GetOwner();
 	if (ContextOwner)
 	{
-		if (AActor* OuterActor = ContextOwner->GetTypedOuter<AActor>())
+		if (AActor* ContextOwnerActor = Cast<AActor>(ContextOwner))
+		{
+			InputComponent = Cast<UEnhancedInputComponent>(ContextOwnerActor->InputComponent);
+		}
+		else if (AActor* OuterActor = ContextOwner->GetTypedOuter<AActor>())
 		{
 			InputComponent = Cast<UEnhancedInputComponent>(OuterActor->InputComponent);
 		}
