Cameras: make the input axis binding node use camera parameters

diff --git a/Engine/Plugins/Cameras/GameplayCameras/Source/GameplayCameras/Private/Nodes/Input/InputAxisBinding2DCameraNode.cpp b/Engine/Plugins/Cameras/GameplayCameras/Source/GameplayCameras/Private/Nodes/Input/InputAxisBinding2DCameraNode.cpp
--- a/Engine/Plugins/Cameras/GameplayCameras/Source/GameplayCameras/Private/Nodes/Input/InputAxisBinding2DCameraNode.cpp
+++ b/Engine/Plugins/Cameras/GameplayCameras/Source/GameplayCameras/Private/Nodes/Input/InputAxisBinding2DCameraNode.cpp
@@ -4,6 +4,7 @@
 
 #include "Components/InputComponent.h"
 #include "Core/CameraEvaluationContext.h"
+#include "Core/CameraParameterReader.h"
 #include "EnhancedInputComponent.h"
 #include "GameFramework/Actor.h"
 #include "GameplayCameras.h"
@@ -25,6 +26,11 @@ protected:
 private:
 
 	TObjectPtr<UEnhancedInputComponent> InputComponent;
+
+	TCameraParameterReader<bool> RevertAxisXReader;
+	TCameraParameterReader<bool> RevertAxisYReader;
+	TCameraParameterReader<FVector2d> MultiplierReader;
+
 	TArray<FEnhancedInputActionValueBinding*> AxisValueBindings;
 };
 
@@ -46,6 +52,11 @@ void FInputAxisBinding2DCameraNodeEvaluator::OnInitialize(const FCameraNodeEvalu
 	}
 
 	const UInputAxisBinding2DCameraNode* AxisBindingNode = GetCameraNodeAs<UInputAxisBinding2DCameraNode>();
+
+	RevertAxisXReader.Initialize(AxisBindingNode->RevertAxisX);
+	RevertAxisYReader.Initialize(AxisBindingNode->RevertAxisY);
+	MultiplierReader.Initialize(AxisBindingNode->Multiplier);
+
 	if (InputComponent)
 	{
 		for (TObjectPtr<UInputAction> AxisAction : AxisBindingNode->AxisActions)
@@ -87,15 +98,14 @@ void FInputAxisBinding2DCameraNodeEvaluator::OnUpdateParameters(const FCameraBle
 		}
 	}
 
-	TransientInputValue = FVector2d(
-			HighestValue.X * AxisBindingNode->Multiplier.X, 
-			HighestValue.Y * AxisBindingNode->Multiplier.Y);
+	const FVector2d Multiplier = MultiplierReader.Get(OutResult.VariableTable);
+	TransientInputValue = FVector2d(HighestValue.X * Multiplier.X, HighestValue.Y * Multiplier.Y);
 
-	if (AxisBindingNode->RevertAxisX)
+	if (RevertAxisXReader.Get(OutResult.VariableTable))
 	{
 		TransientInputValue.X = -TransientInputValue.X;
 	}
-	if (AxisBindingNode->RevertAxisY)
+	if (RevertAxisYReader.Get(OutResult.VariableTable))
 	{
 		TransientInputValue.Y = -TransientInputValue.Y;
 	}
diff --git a/Engine/Plugins/Cameras/GameplayCameras/Source/GameplayCameras/Public/Nodes/Input/InputAxisBinding2DCameraNode.h b/Engine/Plugins/Cameras/GameplayCameras/Source/GameplayCameras/Public/Nodes/Input/InputAxisBinding2DCameraNode.h
--- a/Engine/Plugins/Cameras/GameplayCameras/Source/GameplayCameras/Public/Nodes/Input/InputAxisBinding2DCameraNode.h
+++ b/Engine/Plugins/Cameras/GameplayCameras/Source/GameplayCameras/Public/Nodes/Input/InputAxisBinding2DCameraNode.h
@@ -2,6 +2,7 @@
 
 #pragma once
 
+#include "Core/CameraParameters.h"
 #include "Math/MathFwd.h"
 #include "Nodes/Input/CameraRigInput2DSlot.h"
 
@@ -25,15 +26,15 @@ public:
 
 	/** Whether to revert the X axis. */
 	UPROPERTY(EditAnywhere, Category="Input Processing")
-	bool RevertAxisX = false;
+	FBooleanCameraParameter RevertAxisX = false;
 
 	/** Whether to revert the Y axis. */
 	UPROPERTY(EditAnywhere, Category="Input Processing")
-	bool RevertAxisY = false;
+	FBooleanCameraParameter RevertAxisY = false;
 
 	/** A multiplier to use on the input values. */
 	UPROPERTY(EditAnywhere, Category="Input Processing")
-	FVector2D Multiplier;
+	FVector2dCameraParameter Multiplier;
 
 public:
 
