Cameras: make the auto-rotate input node use camera parameters, and add a toggle parameter

diff --git a/Engine/Plugins/Cameras/GameplayCameras/Source/GameplayCameras/Private/Nodes/Input/AutoRotateInput2DCameraNode.cpp b/Engine/Plugins/Cameras/GameplayCameras/Source/GameplayCameras/Private/Nodes/Input/AutoRotateInput2DCameraNode.cpp
--- a/Engine/Plugins/Cameras/GameplayCameras/Source/GameplayCameras/Private/Nodes/Input/AutoRotateInput2DCameraNode.cpp
+++ b/Engine/Plugins/Cameras/GameplayCameras/Source/GameplayCameras/Private/Nodes/Input/AutoRotateInput2DCameraNode.cpp
@@ -4,6 +4,7 @@
 
 #include "Core/CameraEvaluationContext.h"
 #include "Core/CameraOperation.h"
+#include "Core/CameraParameterReader.h"
 #include "Core/CameraValueInterpolator.h"
 #include "Debug/CameraDebugBlock.h"
 #include "Debug/CameraDebugBlockBuilder.h"
@@ -35,6 +36,10 @@ protected:
 
 private:
 
+	TCameraParameterReader<float> WaitTimeReader;
+	TCameraParameterReader<float> DeactivationThresholdReader;
+	TCameraParameterReader<bool> EnableAutoRotateReader;
+
 	TUniquePtr<FCameraDoubleValueInterpolator> Interpolator;
 
 	FInput2DCameraNodeEvaluator* InputNodeEvaluator = nullptr;
@@ -86,7 +91,12 @@ void FAutoRotateInput2DCameraNodeEvaluator::OnInitialize(const FCameraNodeEvalua
 	}
 
 	const UAutoRotateInput2DCameraNode* AutoRotateNode = GetCameraNodeAs<UAutoRotateInput2DCameraNode>();
-	RemainingWaitTime = AutoRotateNode->WaitTime;
+
+	WaitTimeReader.Initialize(AutoRotateNode->WaitTime);
+	DeactivationThresholdReader.Initialize(AutoRotateNode->DeactivationThreshold);
+	EnableAutoRotateReader.Initialize(AutoRotateNode->EnableAutoRotate);
+
+	RemainingWaitTime = WaitTimeReader.Get(OutResult.VariableTable);
 	bIsAutoRotating = false;
 }
 
@@ -102,6 +112,20 @@ void FAutoRotateInput2DCameraNodeEvaluator::OnRun(const FCameraNodeEvaluationPar
 		InputValue = InputNodeEvaluator->GetInputValue();
 	}
 
+	// Bail out if auto-rotate is disabled.
+	const bool bAutoRotateEnabled = EnableAutoRotateReader.Get(OutResult.VariableTable);
+	if (!bAutoRotateEnabled)
+	{
+		// If we were disabled while auto-rotating, tear down everything we had.
+		if (bIsAutoRotating)
+		{
+			bIsAutoRotating = false;
+			Interpolator.Reset();
+			RemainingWaitTime = WaitTimeReader.Get(OutResult.VariableTable);
+		}
+		return;
+	}
+
 	// Keep track of the context's movement this frame.
 	const FCameraNodeEvaluationResult& InitialResult = Params.EvaluationContext->GetInitialResult();
 	FVector3d ContextMovement = FVector3d::ZeroVector;
@@ -116,12 +140,13 @@ void FAutoRotateInput2DCameraNodeEvaluator::OnRun(const FCameraNodeEvaluationPar
 	// us to deactivate auto-rotate.
 	const double YawChange = FMath::Abs(InputValue.X - LastInputValue.X);
 	const double PitchChange = FMath::Abs(InputValue.Y - LastInputValue.Y);
+	const float DeactivationThreshold = DeactivationThresholdReader.Get(OutResult.VariableTable);
 	LastInputValue = InputValue;
-	if (YawChange >= AutoRotateNode->DeactivationThreshold || PitchChange >= AutoRotateNode->DeactivationThreshold)
+	if (YawChange >= DeactivationThreshold || PitchChange >= DeactivationThreshold)
 	{
 		bIsAutoRotating = false;
 		Interpolator.Reset();
-		RemainingWaitTime = AutoRotateNode->WaitTime;
+		RemainingWaitTime = WaitTimeReader.Get(OutResult.VariableTable);
 		return;
 	}
 
@@ -206,7 +231,7 @@ void FAutoRotateInput2DCameraNodeEvaluator::OnRun(const FCameraNodeEvaluationPar
 	{
 		bIsAutoRotating = false;
 		Interpolator.Reset();
-		RemainingWaitTime = AutoRotateNode->WaitTime;
+		RemainingWaitTime = WaitTimeReader.Get(OutResult.VariableTable);
 	}
 }
 
diff --git a/Engine/Plugins/Cameras/GameplayCameras/Source/GameplayCameras/Public/Nodes/Input/AutoRotateInput2DCameraNode.h b/Engine/Plugins/Cameras/GameplayCameras/Source/GameplayCameras/Public/Nodes/Input/AutoRotateInput2DCameraNode.h
--- a/Engine/Plugins/Cameras/GameplayCameras/Source/GameplayCameras/Public/Nodes/Input/AutoRotateInput2DCameraNode.h
+++ b/Engine/Plugins/Cameras/GameplayCameras/Source/GameplayCameras/Public/Nodes/Input/AutoRotateInput2DCameraNode.h
@@ -2,6 +2,7 @@
 
 #pragma once
 
+#include "Core/CameraParameters.h"
 #include "Math/MathFwd.h"
 #include "Nodes/Input/Input2DCameraNode.h"
 
@@ -36,16 +37,21 @@ public:
 
 	/** The time, in seconds, to wait before re-aligning. */
 	UPROPERTY(EditAnywhere, Category="Auto-Rotate")
-	float WaitTime = 1.f;
+	FFloatCameraParameter WaitTime = 1.f;
 
 	/** The minimum manual rotation, in degrees, to deactivate auto-rtation. */
 	UPROPERTY(EditAnywhere, Category="Auto-Rotate")
-	float DeactivationThreshold = 0.01f;
+	FFloatCameraParameter DeactivationThreshold = 0.01f;
 
 	/** The interpolation for re-alignment. */
 	UPROPERTY(EditAnywhere, Category="Auto-Rotate")
 	TObjectPtr<UCameraValueInterpolator> Interpolator;
 
+	/** Whether to enable auto-rotation. */
+	UPROPERTY(EditAnywhere, Category="Auto-Rotate")
+	FBooleanCameraParameter EnableAutoRotate = true;
+
+	/** The underlying input node. */
 	UPROPERTY()
 	TObjectPtr<UInput2DCameraNode> InputNode;
 
