Cameras: pass the UCanvas object to the debug renderer instead of FCanvas, add more drawing methods

The UCanvas object seems to be required to get what we need to project/unproject points. This is used to draw text in the world, for instance.

diff --git a/Engine/Plugins/Cameras/GameplayCameras/Source/GameplayCameras/Private/Core/CameraSystemEvaluator.cpp b/Engine/Plugins/Cameras/GameplayCameras/Source/GameplayCameras/Private/Core/CameraSystemEvaluator.cpp
--- a/Engine/Plugins/Cameras/GameplayCameras/Source/GameplayCameras/Private/Core/CameraSystemEvaluator.cpp
+++ b/Engine/Plugins/Cameras/GameplayCameras/Source/GameplayCameras/Private/Core/CameraSystemEvaluator.cpp
@@ -289,7 +289,7 @@ void FCameraSystemEvaluator::DebugUpdate(const FCameraSystemDebugUpdateParams& P
 	}
 #endif
 	
-	FCameraDebugRenderer Renderer(OwnerWorld, Params.Canvas);
+	FCameraDebugRenderer Renderer(OwnerWorld, Params.CanvasObject);
 	RootDebugBlock->RootDebugDraw(Renderer);
 }
 
diff --git a/Engine/Plugins/Cameras/GameplayCameras/Source/GameplayCameras/Private/Debug/CameraDebugRenderer.cpp b/Engine/Plugins/Cameras/GameplayCameras/Source/GameplayCameras/Private/Debug/CameraDebugRenderer.cpp
--- a/Engine/Plugins/Cameras/GameplayCameras/Source/GameplayCameras/Private/Debug/CameraDebugRenderer.cpp
+++ b/Engine/Plugins/Cameras/GameplayCameras/Source/GameplayCameras/Private/Debug/CameraDebugRenderer.cpp
@@ -8,6 +8,7 @@
 #include "Components/LineBatchComponent.h"
 #include "Debug/CameraDebugColors.h"
 #include "Debug/DebugTextRenderer.h"
+#include "Engine/Canvas.h"
 #include "Engine/Engine.h"
 #include "Engine/Font.h"
 #include "Engine/World.h"
@@ -50,9 +51,9 @@ static FAutoConsoleVariableRef CVarGameplayCamerasDebugBackgroundDepthSortKey(
 	GGameplayCamerasDebugBackgroundDepthSortKey,
 	TEXT(""));
 
-FCameraDebugRenderer::FCameraDebugRenderer(UWorld* InWorld, FCanvas* InCanvas)
+FCameraDebugRenderer::FCameraDebugRenderer(UWorld* InWorld, UCanvas* InCanvasObject)
 	: World(InWorld)
-	, Canvas(InCanvas)
+	, CanvasObject(InCanvasObject)
 	, DrawColor(FColor::White)
 {
 	RenderFont = GEngine->GetSmallFont();
@@ -66,10 +67,19 @@ FCameraDebugRenderer::~FCameraDebugRenderer()
 	FlushText();
 }
 
+FCanvas* FCameraDebugRenderer::GetCanvas() const
+{
+	return CanvasObject ? CanvasObject->Canvas : nullptr;
+}
+
 FVector2D FCameraDebugRenderer::GetCanvasSize() const
 {
-	FIntPoint ParentSize = Canvas->GetParentCanvasSize();
-	return FVector2D(ParentSize.X, ParentSize.Y);
+	if (CanvasObject)
+	{
+		FIntPoint ParentSize = CanvasObject->Canvas->GetParentCanvasSize();
+		return FVector2D(ParentSize.X, ParentSize.Y);
+	}
+	return FVector2D::ZeroVector;
 }
 
 void FCameraDebugRenderer::AddText(const FString& InString)
@@ -138,7 +148,7 @@ void FCameraDebugRenderer::FlushText()
 		int32 ViewHeight = GetCanvasSize().Y;
 		if (NextDrawPosition.Y < ViewHeight)
 		{
-			FDebugTextRenderer TextRenderer(Canvas, DrawColor, RenderFont);
+			FDebugTextRenderer TextRenderer(GetCanvas(), DrawColor, RenderFont);
 			TextRenderer.LeftMargin = GetIndentMargin();
 			TextRenderer.RenderText(NextDrawPosition, LineBuilder.ToView());
 
@@ -194,7 +204,7 @@ void FCameraDebugRenderer::DrawTextBackgroundTile(float Opacity)
 	const FColor BackgroundColor = FCameraDebugColors::Get().Background.WithAlpha((uint8)(Opacity * 255));
 
 	// Draw the background behind the text.
-	if (Canvas)
+	if (FCanvas* Canvas = GetCanvas())
 	{
 		Canvas->PushDepthSortKey(GGameplayCamerasDebugBackgroundDepthSortKey);
 		{
@@ -208,7 +218,7 @@ void FCameraDebugRenderer::DrawTextBackgroundTile(float Opacity)
 
 void FCameraDebugRenderer::Draw2DLine(const FVector2D& Start, const FVector2D& End, const FLinearColor& LineColor, float LineThickness)
 {
-	if (Canvas)
+	if (FCanvas* Canvas = GetCanvas())
 	{
 		FCanvasLineItem LineItem(Start, End);
 		LineItem.SetColor(LineColor);
@@ -219,7 +229,7 @@ void FCameraDebugRenderer::Draw2DLine(const FVector2D& Start, const FVector2D& E
 
 void FCameraDebugRenderer::Draw2DBox(const FBox2D& Box, const FLinearColor& LineColor, float LineThickness)
 {
-	if (Canvas)
+	if (FCanvas* Canvas = GetCanvas())
 	{
 		FCanvasBoxItem BoxItem(Box.Min, Box.GetSize());
 		BoxItem.SetColor(LineColor);
@@ -230,7 +240,7 @@ void FCameraDebugRenderer::Draw2DBox(const FBox2D& Box, const FLinearColor& Line
 
 void FCameraDebugRenderer::Draw2DBox(const FVector2D& BoxPosition, const FVector2D& BoxSize, const FLinearColor& LineColor, float LineThickness)
 {
-	if (Canvas)
+	if (FCanvas* Canvas = GetCanvas())
 	{
 		FCanvasBoxItem BoxItem(BoxPosition, BoxSize);
 		BoxItem.SetColor(LineColor);
@@ -276,6 +286,34 @@ void FCameraDebugRenderer::DrawSphere(const FVector3d& Center, float Radius, int
 	}
 }
 
+void FCameraDebugRenderer::DrawDirectionalArrow(const FVector3d& Start, const FVector3d& End, float ArrowSize, const FLinearColor& LineColor, float LineThickness)
+{
+	if (ULineBatchComponent* LineBatcher = GetDebugLineBatcher())
+	{
+		LineBatcher->DrawDirectionalArrow(Start, End, ArrowSize, LineColor, 0.f, SDPG_Foreground, LineThickness);
+	}
+}
+
+void FCameraDebugRenderer::DrawText(const FVector3d& WorldPosition, const FString& Text, const FLinearColor& TextColor, UFont* TextFont)
+{
+	DrawText(WorldPosition, FVector2d::ZeroVector, Text, TextColor, TextFont);
+}
+
+void FCameraDebugRenderer::DrawText(const FVector3d& WorldPosition, const FVector2d& ScreenOffset, const FString& Text, const FLinearColor& TextColor, UFont* TextFont)
+{
+	if (CanvasObject)
+	{
+		const FColor PreviousColor = CanvasObject->DrawColor;
+		const FVector3d ScreenPosition = CanvasObject->Project(WorldPosition);
+		UFont* ActualTextFont = TextFont ? TextFont : GEngine->GetSmallFont();
+		CanvasObject->DrawColor = TextColor.ToFColor(true);
+		CanvasObject->DrawText(
+				ActualTextFont, Text, 
+				ScreenPosition.X + ScreenOffset.X, ScreenPosition.Y + ScreenOffset.Y);
+		CanvasObject->DrawColor = PreviousColor;
+	}
+}
+
 ULineBatchComponent* FCameraDebugRenderer::GetDebugLineBatcher() const
 {
 	return World ? World->ForegroundLineBatcher : nullptr;
diff --git a/Engine/Plugins/Cameras/GameplayCameras/Source/GameplayCameras/Private/GameFramework/GameplayCameraSystemComponent.cpp b/Engine/Plugins/Cameras/GameplayCameras/Source/GameplayCameras/Private/GameFramework/GameplayCameraSystemComponent.cpp
--- a/Engine/Plugins/Cameras/GameplayCameras/Source/GameplayCameras/Private/GameFramework/GameplayCameraSystemComponent.cpp
+++ b/Engine/Plugins/Cameras/GameplayCameras/Source/GameplayCameras/Private/GameFramework/GameplayCameraSystemComponent.cpp
@@ -239,7 +239,7 @@ void UGameplayCameraSystemComponent::DebugDraw(UCanvas* Canvas, APlayerControlle
 	if (Evaluator.IsValid())
 	{
 		FCameraSystemDebugUpdateParams DebugUpdateParams;
-		DebugUpdateParams.Canvas = Canvas->Canvas;
+		DebugUpdateParams.CanvasObject = Canvas;
 		Evaluator->DebugUpdate(DebugUpdateParams);
 	}
 }
diff --git a/Engine/Plugins/Cameras/GameplayCameras/Source/GameplayCameras/Public/Core/CameraSystemEvaluator.h b/Engine/Plugins/Cameras/GameplayCameras/Source/GameplayCameras/Public/Core/CameraSystemEvaluator.h
--- a/Engine/Plugins/Cameras/GameplayCameras/Source/GameplayCameras/Public/Core/CameraSystemEvaluator.h
+++ b/Engine/Plugins/Cameras/GameplayCameras/Source/GameplayCameras/Public/Core/CameraSystemEvaluator.h
@@ -14,9 +14,9 @@
 #include "Templates/SharedPointer.h"
 #include "UObject/GCObject.h"
 
-class FCanvas;
 class UCameraDirector;
 class UCameraRigAsset;
+class UCanvas;
 class URootCameraNode;
 enum class ECameraRigLayer : uint8;
 struct FMinimalViewInfo;
@@ -78,7 +78,7 @@ struct FCameraSystemEvaluationResult
 #if UE_GAMEPLAY_CAMERAS_DEBUG
 struct FCameraSystemDebugUpdateParams
 {
-	FCanvas* Canvas = nullptr;
+	UCanvas* CanvasObject = nullptr;
 };
 #endif  // UE_GAMEPLAY_CAMERAS_DEBUG
 
diff --git a/Engine/Plugins/Cameras/GameplayCameras/Source/GameplayCameras/Public/Debug/CameraDebugRenderer.h b/Engine/Plugins/Cameras/GameplayCameras/Source/GameplayCameras/Public/Debug/CameraDebugRenderer.h
--- a/Engine/Plugins/Cameras/GameplayCameras/Source/GameplayCameras/Public/Debug/CameraDebugRenderer.h
+++ b/Engine/Plugins/Cameras/GameplayCameras/Source/GameplayCameras/Public/Debug/CameraDebugRenderer.h
@@ -10,6 +10,7 @@
 #include "Misc/StringBuilder.h"
 
 class FCanvas;
+class UCanvas;
 class UFont;
 class ULineBatchComponent;
 class UWorld;
@@ -37,7 +38,7 @@ class FCameraDebugRenderer
 public:
 
 	/** Creates a new debug renderer. */
-	GAMEPLAYCAMERAS_API FCameraDebugRenderer(UWorld* InWorld, FCanvas* InCanvas);
+	GAMEPLAYCAMERAS_API FCameraDebugRenderer(UWorld* InWorld, UCanvas* InCanvasObject);
 	/** Destroys the debug renderer. */
 	GAMEPLAYCAMERAS_API ~FCameraDebugRenderer();
 
@@ -80,6 +81,12 @@ public:
 	GAMEPLAYCAMERAS_API void DrawLine(const FVector3d& Start, const FVector3d& End, const FLinearColor& LineColor, float LineThickness = 1.f);
 	/** Draws a 3D sphere. */
 	GAMEPLAYCAMERAS_API void DrawSphere(const FVector3d& Center, float Radius, int32 Segments, const FLinearColor& LineColor, float LineThickness);
+	/** Draws a 3D arrow. */
+	GAMEPLAYCAMERAS_API void DrawDirectionalArrow(const FVector3d& Start, const FVector3d& End, float ArrowSize, const FLinearColor& LineColor, float LineThickness);
+	/** Draws text at a projected 3D position. */
+	GAMEPLAYCAMERAS_API void DrawText(const FVector3d& WorldPosition, const FString& Text, const FLinearColor& TextColor, UFont* TextFont = nullptr);
+	/** Draws text at a projected 3D position, with an added screen-space offset. */
+	GAMEPLAYCAMERAS_API void DrawText(const FVector3d& WorldPosition, const FVector2d& ScreenOffset, const FString& Text, const FLinearColor& TextColor, UFont* TextFont = nullptr);
 
 public:
 
@@ -97,12 +104,12 @@ public:
 public:
 
 	/** Gets the drawing canvas. */
-	FCanvas* GetCanvas() const { return Canvas; }
+	FCanvas* GetCanvas() const;
 
 	/** Gets the size of the canvas. */
 	FVector2D GetCanvasSize() const;
 
-	bool HasCanvas() const { return Canvas != nullptr; }
+	bool HasCanvas() const { return CanvasObject != nullptr; }
 
 private:
 
@@ -119,7 +126,7 @@ private:
 	/** The world in which we might draw debug primitives. */
 	UWorld* World;
 	/** The canvas used to draw the text wall. */
-	FCanvas* Canvas;
+	UCanvas* CanvasObject;
 	/** The draw color of the canvas. */
 	FColor DrawColor;
 
diff --git a/Engine/Plugins/Cameras/GameplayCameras/Source/GameplayCamerasEditor/Private/Trace/CameraSystemRewindDebuggerExtension.cpp b/Engine/Plugins/Cameras/GameplayCameras/Source/GameplayCamerasEditor/Private/Trace/CameraSystemRewindDebuggerExtension.cpp
--- a/Engine/Plugins/Cameras/GameplayCameras/Source/GameplayCamerasEditor/Private/Trace/CameraSystemRewindDebuggerExtension.cpp
+++ b/Engine/Plugins/Cameras/GameplayCameras/Source/GameplayCamerasEditor/Private/Trace/CameraSystemRewindDebuggerExtension.cpp
@@ -102,7 +102,7 @@ void FCameraSystemRewindDebuggerExtension::DebugDraw(UCanvas* Canvas, APlayerCon
 {
 	if (RootDebugBlock && VisualizedWorld)
 	{
-		FCameraDebugRenderer CameraDebugRenderer(VisualizedWorld, Canvas->Canvas);
+		FCameraDebugRenderer CameraDebugRenderer(VisualizedWorld, Canvas);
 		RootDebugBlock->RootDebugDraw(CameraDebugRenderer);
 	}
 }
