Cameras: add debugging information for evaluation services, improve debug info of final result

- Add a new "services" debug category.
- Add a debug block for rendering services' debug info.
- Add APIs to services for building debug blocks.
- Services didn't have proper RTTI information setup (needed to print their class name on screen) so add the usual macros for that.
- Make the evaluation result debug block able to also display the camera system's result, not just camera node results, and use that for the "posetats" debug category.

diff --git a/Engine/Plugins/Cameras/GameplayCameras/Source/GameplayCameras/Private/Core/CameraEvaluationService.cpp b/Engine/Plugins/Cameras/GameplayCameras/Source/GameplayCameras/Private/Core/CameraEvaluationService.cpp
--- a/Engine/Plugins/Cameras/GameplayCameras/Source/GameplayCameras/Private/Core/CameraEvaluationService.cpp
+++ b/Engine/Plugins/Cameras/GameplayCameras/Source/GameplayCameras/Private/Core/CameraEvaluationService.cpp
@@ -50,4 +50,13 @@ void FCameraEvaluationService::SetEvaluationServiceFlags(ECameraEvaluationServic
 	PrivateFlags = InFlags;
 }
 
+#if UE_GAMEPLAY_CAMERAS_DEBUG
+
+void FCameraEvaluationService::BuildDebugBlocks(const FCameraDebugBlockBuildParams& Params, FCameraDebugBlockBuilder& Builder)
+{
+	OnBuildDebugBlocks(Params, Builder);
+}
+
+#endif  // UE_GAMEPLAY_CAMERAS_DEBUG
+
 }  // namespace UE::Cameras
diff --git a/Engine/Plugins/Cameras/GameplayCameras/Source/GameplayCameras/Private/Core/CameraSystemEvaluator.cpp b/Engine/Plugins/Cameras/GameplayCameras/Source/GameplayCameras/Private/Core/CameraSystemEvaluator.cpp
--- a/Engine/Plugins/Cameras/GameplayCameras/Source/GameplayCameras/Private/Core/CameraSystemEvaluator.cpp
+++ b/Engine/Plugins/Cameras/GameplayCameras/Source/GameplayCameras/Private/Core/CameraSystemEvaluator.cpp
@@ -131,6 +131,11 @@ void FCameraSystemEvaluator::UnregisterEvaluationService(TSharedRef<FCameraEvalu
 	EvaluationServices.Remove(EvaluationService);
 }
 
+void FCameraSystemEvaluator::GetEvaluationServices(TArray<TSharedPtr<FCameraEvaluationService>>& OutEvaluationServices) const
+{
+	OutEvaluationServices = EvaluationServices;
+}
+
 void FCameraSystemEvaluator::NotifyRootCameraNodeEvent(const FRootCameraNodeCameraRigEvent& InEvent)
 {
 	for (TSharedPtr<FCameraEvaluationService> EvaluationService : EvaluationServices)
diff --git a/Engine/Plugins/Cameras/GameplayCameras/Source/GameplayCameras/Private/Debug/CameraDebugCategories.cpp b/Engine/Plugins/Cameras/GameplayCameras/Source/GameplayCameras/Private/Debug/CameraDebugCategories.cpp
--- a/Engine/Plugins/Cameras/GameplayCameras/Source/GameplayCameras/Private/Debug/CameraDebugCategories.cpp
+++ b/Engine/Plugins/Cameras/GameplayCameras/Source/GameplayCameras/Private/Debug/CameraDebugCategories.cpp
@@ -12,6 +12,7 @@ namespace UE::Cameras
 const FString FCameraDebugCategories::NodeTree("nodetree");
 const FString FCameraDebugCategories::DirectorTree("directortree");
 const FString FCameraDebugCategories::BlendStacks("blendstacks");
+const FString FCameraDebugCategories::Services("services");
 const FString FCameraDebugCategories::PoseStats("posestats");
 const FString FCameraDebugCategories::Viewfinder("viewfinder");
 
diff --git a/Engine/Plugins/Cameras/GameplayCameras/Source/GameplayCameras/Private/Debug/CameraEvaluationServiceDebugBlock.cpp b/Engine/Plugins/Cameras/GameplayCameras/Source/GameplayCameras/Private/Debug/CameraEvaluationServiceDebugBlock.cpp
new file mode 100644
--- /dev/null
+++ b/Engine/Plugins/Cameras/GameplayCameras/Source/GameplayCameras/Private/Debug/CameraEvaluationServiceDebugBlock.cpp
@@ -0,0 +1,40 @@
+// Copyright Epic Games, Inc. All Rights Reserved.
+
+#include "Debug/CameraEvaluationServiceDebugBlock.h"
+
+#include "Core/CameraEvaluationService.h"
+#include "Core/CameraObjectRtti.h"
+#include "Debug/CameraDebugRenderer.h"
+
+#if UE_GAMEPLAY_CAMERAS_DEBUG
+
+namespace UE::Cameras
+{
+
+UE_DEFINE_CAMERA_DEBUG_BLOCK(FCameraEvaluationServiceDebugBlock)
+
+FCameraEvaluationServiceDebugBlock::FCameraEvaluationServiceDebugBlock()
+{
+}
+
+FCameraEvaluationServiceDebugBlock::FCameraEvaluationServiceDebugBlock(TSharedPtr<const FCameraEvaluationService> InEvaluationService)
+{
+	FCameraObjectTypeID TypeID = InEvaluationService->GetTypeID();
+	const FCameraObjectTypeInfo* TypeInfo = FCameraObjectTypeRegistry::Get().GetTypeInfo(TypeID);
+	ServiceClassName = TypeInfo ? TypeInfo->TypeName.ToString() : TEXT("<no type info>");
+}
+
+void FCameraEvaluationServiceDebugBlock::OnDebugDraw(const FCameraDebugBlockDrawParams& Params, FCameraDebugRenderer& Renderer)
+{
+	Renderer.AddText(TEXT("{cam_passive}[%s]{cam_default} "), *ServiceClassName);
+}
+
+void FCameraEvaluationServiceDebugBlock::OnSerialize(FArchive& Ar)
+{
+	Ar << ServiceClassName;
+}
+
+}  // namespace UE::Cameras
+
+#endif  // UE_GAMEPLAY_CAMERAS_DEBUG
+
diff --git a/Engine/Plugins/Cameras/GameplayCameras/Source/GameplayCameras/Private/Debug/CameraNodeEvaluationResultDebugBlock.cpp b/Engine/Plugins/Cameras/GameplayCameras/Source/GameplayCameras/Private/Debug/CameraNodeEvaluationResultDebugBlock.cpp
--- a/Engine/Plugins/Cameras/GameplayCameras/Source/GameplayCameras/Private/Debug/CameraNodeEvaluationResultDebugBlock.cpp
+++ b/Engine/Plugins/Cameras/GameplayCameras/Source/GameplayCameras/Private/Debug/CameraNodeEvaluationResultDebugBlock.cpp
@@ -3,6 +3,7 @@
 #include "Debug/CameraNodeEvaluationResultDebugBlock.h"
 
 #include "Core/CameraNodeEvaluator.h"
+#include "Core/CameraSystemEvaluator.h"
 #include "Debug/CameraDebugBlockBuilder.h"
 #include "Debug/CameraDebugColors.h"
 #include "Debug/CameraDebugRenderer.h"
@@ -31,6 +32,15 @@ void FCameraNodeEvaluationResultDebugBlock::Initialize(const FCameraNodeEvaluati
 	AddChild(&Builder.BuildDebugBlock<FVariableTableDebugBlock>(InResult.VariableTable));
 }
 
+void FCameraNodeEvaluationResultDebugBlock::Initialize(const FCameraSystemEvaluationResult& InResult, FCameraDebugBlockBuilder& Builder)
+{
+	bIsCameraCut = InResult.bIsCameraCut;
+	bIsValid = InResult.bIsValid;
+
+	AddChild(&Builder.BuildDebugBlock<FCameraPoseDebugBlock>(InResult.CameraPose));
+	AddChild(&Builder.BuildDebugBlock<FVariableTableDebugBlock>(InResult.VariableTable));
+}
+
 FCameraPoseDebugBlock* FCameraNodeEvaluationResultDebugBlock::GetCameraPoseDebugBlock()
 {
 	TArrayView<FCameraDebugBlock*> ChildrenView(GetChildren());
diff --git a/Engine/Plugins/Cameras/GameplayCameras/Source/GameplayCameras/Private/Debug/RootCameraDebugBlock.cpp b/Engine/Plugins/Cameras/GameplayCameras/Source/GameplayCameras/Private/Debug/RootCameraDebugBlock.cpp
--- a/Engine/Plugins/Cameras/GameplayCameras/Source/GameplayCameras/Private/Debug/RootCameraDebugBlock.cpp
+++ b/Engine/Plugins/Cameras/GameplayCameras/Source/GameplayCameras/Private/Debug/RootCameraDebugBlock.cpp
@@ -3,6 +3,7 @@
 #include "Debug/RootCameraDebugBlock.h"
 
 #include "Core/CameraEvaluationContextStack.h"
+#include "Core/CameraEvaluationService.h"
 #include "Core/CameraSystemEvaluator.h"
 #include "Core/RootCameraNode.h"
 #include "Debug/CameraDebugBlockBuilder.h"
@@ -10,9 +11,12 @@
 #include "Debug/CameraDebugColors.h"
 #include "Debug/CameraDebugRenderer.h"
 #include "Debug/CameraDirectorTreeDebugBlock.h"
+#include "Debug/CameraEvaluationServiceDebugBlock.h"
+#include "Debug/CameraNodeEvaluationResultDebugBlock.h"
 #include "Debug/CameraNodeEvaluatorDebugBlock.h"
 #include "Debug/CameraPoseDebugBlock.h"
 #include "Debug/CategoryTitleDebugBlock.h"
+#include "Debug/VariableTableDebugBlock.h"
 #include "Debug/ViewfinderDebugBlock.h"
 #include "HAL/IConsoleManager.h"
 #include "String/ParseTokens.h"
@@ -90,15 +94,38 @@ void FRootCameraDebugBlock::BuildDebugBlocks(const FCameraSystemEvaluator& Camer
 	}
 	Builder.EndChildDebugBlock();
 
+	// Debug block for the evaluation services.
+	FCategoryTitleDebugBlock& ServicesCategory = Builder.StartChildDebugBlock<FCategoryTitleDebugBlock>();
+	{
+		ServicesCategory.Title = TEXT("Services");
+		ServicesCategory.Category = FCameraDebugCategories::Services;
+
+		TArray<TSharedPtr<FCameraEvaluationService>> EvaluationServices;
+		CameraSystem.GetEvaluationServices(EvaluationServices);
+		for (TSharedPtr<FCameraEvaluationService> EvaluationService : EvaluationServices)
+		{
+			FCameraEvaluationServiceDebugBlock& ServiceDebugBlock = Builder.StartChildDebugBlock<FCameraEvaluationServiceDebugBlock>(EvaluationService);
+			{
+				EvaluationService->BuildDebugBlocks(Params, Builder);
+			}
+			Builder.EndChildDebugBlock();
+		}
+	}
+	Builder.EndChildDebugBlock();
+
 	// Debug block for showing the final evaluated camera.
 	FCategoryTitleDebugBlock& PoseStatsCategory = Builder.StartChildDebugBlock<FCategoryTitleDebugBlock>();
 	{
 		PoseStatsCategory.Title = TEXT("Evaluated Camera");
 		PoseStatsCategory.Category = FCameraDebugCategories::PoseStats;
 
-		const FCameraSystemEvaluationResult& Result = CameraSystem.GetEvaluatedResult();
-		Builder.AttachDebugBlock<FCameraPoseDebugBlock>(Result.CameraPose)
-			.WithShowUnchangedCVar(TEXT("GameplayCameras.Debug.PoseStats.ShowUnchanged"));
+		FCameraNodeEvaluationResultDebugBlock& ResultDebugBlock = Builder.BuildDebugBlock<FCameraNodeEvaluationResultDebugBlock>();
+		PoseStatsCategory.AddChild(&ResultDebugBlock);
+		{
+			ResultDebugBlock.Initialize(CameraSystem.GetEvaluatedResult(), Builder);
+			ResultDebugBlock.GetCameraPoseDebugBlock()->WithShowUnchangedCVar(TEXT("GameplayCameras.Debug.PoseStats.ShowUnchanged"));
+			ResultDebugBlock.GetVariableTableDebugBlock()->WithShowVariableIDsCVar(TEXT("GameplayCameras.Debug.PoseStats.ShowUnchanged"));
+		}
 	}
 	Builder.EndChildDebugBlock();
 	
diff --git a/Engine/Plugins/Cameras/GameplayCameras/Source/GameplayCameras/Private/Services/AutoResetCameraVariableService.cpp b/Engine/Plugins/Cameras/GameplayCameras/Source/GameplayCameras/Private/Services/AutoResetCameraVariableService.cpp
--- a/Engine/Plugins/Cameras/GameplayCameras/Source/GameplayCameras/Private/Services/AutoResetCameraVariableService.cpp
+++ b/Engine/Plugins/Cameras/GameplayCameras/Source/GameplayCameras/Private/Services/AutoResetCameraVariableService.cpp
@@ -11,6 +11,8 @@
 namespace UE::Cameras
 {
 
+UE_DEFINE_CAMERA_EVALUATION_SERVICE(FAutoResetCameraVariableService)
+
 void FAutoResetCameraVariableService::OnInitialize(const FCameraEvaluationServiceInitializeParams& Params)
 {
 	SetEvaluationServiceFlags(
diff --git a/Engine/Plugins/Cameras/GameplayCameras/Source/GameplayCameras/Private/Services/AutoResetCameraVariableService.h b/Engine/Plugins/Cameras/GameplayCameras/Source/GameplayCameras/Private/Services/AutoResetCameraVariableService.h
--- a/Engine/Plugins/Cameras/GameplayCameras/Source/GameplayCameras/Private/Services/AutoResetCameraVariableService.h
+++ b/Engine/Plugins/Cameras/GameplayCameras/Source/GameplayCameras/Private/Services/AutoResetCameraVariableService.h
@@ -12,6 +12,8 @@ namespace UE::Cameras
 
 class FAutoResetCameraVariableService : public FCameraEvaluationService
 {
+	UE_DECLARE_CAMERA_EVALUATION_SERVICE(GAMEPLAYCAMERAS_API, FAutoResetCameraVariableService)
+
 protected:
 
 	// FCameraEvaluationService interface.
diff --git a/Engine/Plugins/Cameras/GameplayCameras/Source/GameplayCameras/Private/Services/OrientationInitializationService.cpp b/Engine/Plugins/Cameras/GameplayCameras/Source/GameplayCameras/Private/Services/OrientationInitializationService.cpp
--- a/Engine/Plugins/Cameras/GameplayCameras/Source/GameplayCameras/Private/Services/OrientationInitializationService.cpp
+++ b/Engine/Plugins/Cameras/GameplayCameras/Source/GameplayCameras/Private/Services/OrientationInitializationService.cpp
@@ -14,6 +14,8 @@
 namespace UE::Cameras
 {
 
+UE_DEFINE_CAMERA_EVALUATION_SERVICE(FOrientationInitializationService)
+
 void FOrientationInitializationService::OnInitialize(const FCameraEvaluationServiceInitializeParams& Params)
 {
 	SetEvaluationServiceFlags(
diff --git a/Engine/Plugins/Cameras/GameplayCameras/Source/GameplayCameras/Private/Services/OrientationInitializationService.h b/Engine/Plugins/Cameras/GameplayCameras/Source/GameplayCameras/Private/Services/OrientationInitializationService.h
--- a/Engine/Plugins/Cameras/GameplayCameras/Source/GameplayCameras/Private/Services/OrientationInitializationService.h
+++ b/Engine/Plugins/Cameras/GameplayCameras/Source/GameplayCameras/Private/Services/OrientationInitializationService.h
@@ -13,6 +13,8 @@ struct FCameraRigEvaluationInfo;
 
 class FOrientationInitializationService : public FCameraEvaluationService
 {
+	UE_DECLARE_CAMERA_EVALUATION_SERVICE(GAMEPLAYCAMERAS_API, FOrientationInitializationService)
+
 protected:
 
 	// FCameraEvaluationService interface.
diff --git a/Engine/Plugins/Cameras/GameplayCameras/Source/GameplayCameras/Public/Core/CameraEvaluationService.h b/Engine/Plugins/Cameras/GameplayCameras/Source/GameplayCameras/Public/Core/CameraEvaluationService.h
--- a/Engine/Plugins/Cameras/GameplayCameras/Source/GameplayCameras/Public/Core/CameraEvaluationService.h
+++ b/Engine/Plugins/Cameras/GameplayCameras/Source/GameplayCameras/Public/Core/CameraEvaluationService.h
@@ -3,6 +3,7 @@
 #pragma once
 
 #include "Core/CameraObjectRtti.h"
+#include "Debug/CameraDebugBlockFwd.h"
 #include "GameplayCameras.h"
 #include "Templates/SharedPointer.h"
 
@@ -79,6 +80,11 @@ public:
 	/** Tears down the evaluation service. */
 	void Teardown(const FCameraEvaluationServiceTeardownParams& Params);
 
+#if UE_GAMEPLAY_CAMERAS_DEBUG
+	/** Called to create debug blocks for this node evaluator. */
+	void BuildDebugBlocks(const FCameraDebugBlockBuildParams& Params, FCameraDebugBlockBuilder& Builder);
+#endif  // UE_GAMEPLAY_CAMERAS_DEBUG
+
 public:
 
 	// Internal API.
@@ -102,6 +108,11 @@ protected:
 	/** Called when the root camera node experiences an event. */
 	virtual void OnRootCameraNodeEvent(const FRootCameraNodeCameraRigEvent& InEvent) {}
 
+#if UE_GAMEPLAY_CAMERAS_DEBUG
+	/** Called to create debug blocks for this node evaluator. */
+	virtual void OnBuildDebugBlocks(const FCameraDebugBlockBuildParams& Params, FCameraDebugBlockBuilder& Builder) {}
+#endif  // UE_GAMEPLAY_CAMERAS_DEBUG
+
 protected:
 
 	/** Sets the flags on this service. */
@@ -115,3 +126,14 @@ private:
 
 }  // namespace UE::Cameras
 
+// Utility macros for declaring and defining camera evaluation services.
+//
+#define UE_DECLARE_CAMERA_EVALUATION_SERVICE(ApiDeclSpec, ClassName)\
+	UE_GAMEPLAY_CAMERAS_DECLARE_RTTI(ApiDeclSpec, ClassName, FCameraEvaluationService)
+
+#define UE_DECLARE_CAMERA_EVALUATION_SERVICE_EX(ApiDeclSpec, ClassName, BaseClassName)\
+	UE_GAMEPLAY_CAMERAS_DECLARE_RTTI(ApiDeclSpec, ClassName, BaseClassName)
+
+#define UE_DEFINE_CAMERA_EVALUATION_SERVICE(ClassName)\
+	UE_GAMEPLAY_CAMERAS_DEFINE_RTTI(ClassName)
+
diff --git a/Engine/Plugins/Cameras/GameplayCameras/Source/GameplayCameras/Public/Core/CameraNodeEvaluator.h b/Engine/Plugins/Cameras/GameplayCameras/Source/GameplayCameras/Public/Core/CameraNodeEvaluator.h
--- a/Engine/Plugins/Cameras/GameplayCameras/Source/GameplayCameras/Public/Core/CameraNodeEvaluator.h
+++ b/Engine/Plugins/Cameras/GameplayCameras/Source/GameplayCameras/Public/Core/CameraNodeEvaluator.h
@@ -9,6 +9,7 @@
 #include "Core/CameraVariableTable.h"
 #include "Core/ObjectChildrenView.h"
 #include "CoreTypes.h"
+#include "Debug/CameraDebugBlockFwd.h"
 #include "Debug/RootCameraDebugBlock.h"
 #include "GameplayCameras.h"
 #include "UObject/ObjectPtr.h"
@@ -28,10 +29,6 @@ struct FCameraNodeEvaluationParams;
 struct FCameraNodeEvaluatorBuilder;
 struct FCameraRigEvaluationInfo;
 
-#if UE_GAMEPLAY_CAMERAS_DEBUG
-struct FCameraDebugBlockBuilder;
-#endif  // UE_GAMEPLAY_CAMERAS_DEBUG
-
 /**
  * Flags describing the needs of a camera node evaluator.
  */
@@ -187,18 +184,6 @@ struct GAMEPLAYCAMERAS_API FCameraNodeEvaluatorSerializeParams
 {
 };
 
-#if UE_GAMEPLAY_CAMERAS_DEBUG
-
-/**
- * Structure for creating the node evaluator's debug blocks.
- */
-struct FCameraDebugBlockBuildParams
-{
-	// Empty for now, but defined for later API changes.
-};
-
-#endif  // UE_GAMEPLAY_CAMERAS_DEBUG
-
 /**
  * Base class for objects responsible for running a camera node.
  */
diff --git a/Engine/Plugins/Cameras/GameplayCameras/Source/GameplayCameras/Public/Core/CameraSystemEvaluator.h b/Engine/Plugins/Cameras/GameplayCameras/Source/GameplayCameras/Public/Core/CameraSystemEvaluator.h
--- a/Engine/Plugins/Cameras/GameplayCameras/Source/GameplayCameras/Public/Core/CameraSystemEvaluator.h
+++ b/Engine/Plugins/Cameras/GameplayCameras/Source/GameplayCameras/Public/Core/CameraSystemEvaluator.h
@@ -124,6 +124,8 @@ public:
 	GAMEPLAYCAMERAS_API void RegisterEvaluationService(TSharedRef<FCameraEvaluationService> EvaluationService);
 	/** Unregisters an evaluation service from this camera system. */
 	GAMEPLAYCAMERAS_API void UnregisterEvaluationService(TSharedRef<FCameraEvaluationService> EvaluationService);
+	/** Get currently registered evaluation services. */
+	GAMEPLAYCAMERAS_API void GetEvaluationServices(TArray<TSharedPtr<FCameraEvaluationService>>& OutEvaluationServices) const;
 
 public:
 
diff --git a/Engine/Plugins/Cameras/GameplayCameras/Source/GameplayCameras/Public/Debug/CameraDebugBlockFwd.h b/Engine/Plugins/Cameras/GameplayCameras/Source/GameplayCameras/Public/Debug/CameraDebugBlockFwd.h
new file mode 100644
--- /dev/null
+++ b/Engine/Plugins/Cameras/GameplayCameras/Source/GameplayCameras/Public/Debug/CameraDebugBlockFwd.h
@@ -0,0 +1,25 @@
+// Copyright Epic Games, Inc. All Rights Reserved.
+
+#pragma once
+
+#include "GameplayCameras.h"
+
+namespace UE::Cameras
+{
+
+#if UE_GAMEPLAY_CAMERAS_DEBUG
+
+struct FCameraDebugBlockBuilder;
+
+/**
+ * Structure for creating the node evaluator's debug blocks.
+ */
+struct FCameraDebugBlockBuildParams
+{
+	// Empty for now, but defined for later API changes.
+};
+
+#endif  // UE_GAMEPLAY_CAMERAS_DEBUG
+
+}  // namespace UE::Cameras
+
diff --git a/Engine/Plugins/Cameras/GameplayCameras/Source/GameplayCameras/Public/Debug/CameraDebugCategories.h b/Engine/Plugins/Cameras/GameplayCameras/Source/GameplayCameras/Public/Debug/CameraDebugCategories.h
--- a/Engine/Plugins/Cameras/GameplayCameras/Source/GameplayCameras/Public/Debug/CameraDebugCategories.h
+++ b/Engine/Plugins/Cameras/GameplayCameras/Source/GameplayCameras/Public/Debug/CameraDebugCategories.h
@@ -18,6 +18,7 @@ struct GAMEPLAYCAMERAS_API FCameraDebugCategories
 	static const FString NodeTree;
 	static const FString DirectorTree;
 	static const FString BlendStacks;
+	static const FString Services;
 	static const FString PoseStats;
 	static const FString Viewfinder;
 };
diff --git a/Engine/Plugins/Cameras/GameplayCameras/Source/GameplayCameras/Public/Debug/CameraEvaluationServiceDebugBlock.h b/Engine/Plugins/Cameras/GameplayCameras/Source/GameplayCameras/Public/Debug/CameraEvaluationServiceDebugBlock.h
new file mode 100644
--- /dev/null
+++ b/Engine/Plugins/Cameras/GameplayCameras/Source/GameplayCameras/Public/Debug/CameraEvaluationServiceDebugBlock.h
@@ -0,0 +1,44 @@
+// Copyright Epic Games, Inc. All Rights Reserved.
+
+#pragma once
+
+#include "CoreTypes.h"
+#include "Debug/CameraDebugBlock.h"
+#include "GameplayCameras.h"
+
+#if UE_GAMEPLAY_CAMERAS_DEBUG
+
+namespace UE::Cameras
+{
+
+class FCameraEvaluationService;
+
+/**
+ * Basic debug block for an evaluation service.
+ */
+class FCameraEvaluationServiceDebugBlock : public FCameraDebugBlock
+{
+	UE_DECLARE_CAMERA_DEBUG_BLOCK(GAMEPLAYCAMERAS_API, FCameraEvaluationServiceDebugBlock)
+
+public:
+
+	/** Constructs a new evaluation service debug block. */
+	FCameraEvaluationServiceDebugBlock();
+	/** Constructs a new evaluation service debug block. */
+	FCameraEvaluationServiceDebugBlock(TSharedPtr<const FCameraEvaluationService> InEvaluationService);
+
+protected:
+
+	// FCameraDebugBlock interface.
+	virtual void OnDebugDraw(const FCameraDebugBlockDrawParams& Params, FCameraDebugRenderer& Renderer) override;
+	virtual void OnSerialize(FArchive& Ar) override;
+
+private:
+
+	FString ServiceClassName;
+};
+
+}  // namespace UE::Cameras
+
+#endif  // UE_GAMEPLAY_CAMERAS_DEBUG
+
diff --git a/Engine/Plugins/Cameras/GameplayCameras/Source/GameplayCameras/Public/Debug/CameraNodeEvaluationResultDebugBlock.h b/Engine/Plugins/Cameras/GameplayCameras/Source/GameplayCameras/Public/Debug/CameraNodeEvaluationResultDebugBlock.h
--- a/Engine/Plugins/Cameras/GameplayCameras/Source/GameplayCameras/Public/Debug/CameraNodeEvaluationResultDebugBlock.h
+++ b/Engine/Plugins/Cameras/GameplayCameras/Source/GameplayCameras/Public/Debug/CameraNodeEvaluationResultDebugBlock.h
@@ -13,6 +13,7 @@ class FCameraPoseDebugBlock;
 class FVariableTableDebugBlock;
 struct FCameraDebugBlockBuilder;
 struct FCameraNodeEvaluationResult;
+struct FCameraSystemEvaluationResult;
 
 /**
  * A debug block that prints an evaluation result.
@@ -29,13 +30,21 @@ public:
 	FCameraNodeEvaluationResultDebugBlock();
 
 	/**
-	 * Initializes this block with the given result.
+	 * Initializes this block with the given camera node result.
 	 *
 	 * @param InResult   The result to display
 	 * @param Builder    The builder to use to make the children debug blocks
 	 */
 	void Initialize(const FCameraNodeEvaluationResult& InResult, FCameraDebugBlockBuilder& Builder);
 
+	/**
+	 * Initializes this block with the given camera system result.
+	 *
+	 * @param InResult   The result to display
+	 * @param Builder    The builder to use to make the children debug blocks
+	 */
+	void Initialize(const FCameraSystemEvaluationResult& InResult, FCameraDebugBlockBuilder& Builder);
+
 	/** Gets the pose stats debug block. */
 	FCameraPoseDebugBlock* GetCameraPoseDebugBlock();
 	/** Gets the variable table debug block. */
diff --git a/Engine/Plugins/Cameras/GameplayCameras/Source/GameplayCamerasEditor/Private/GameplayCamerasEditorModule.cpp b/Engine/Plugins/Cameras/GameplayCameras/Source/GameplayCamerasEditor/Private/GameplayCamerasEditorModule.cpp
--- a/Engine/Plugins/Cameras/GameplayCameras/Source/GameplayCamerasEditor/Private/GameplayCamerasEditorModule.cpp
+++ b/Engine/Plugins/Cameras/GameplayCameras/Source/GameplayCamerasEditor/Private/GameplayCamerasEditorModule.cpp
@@ -380,6 +380,12 @@ private:
 				LOCTEXT("BlendStacksDebugCategoryToolTip", "Shows a summary of the blend stacks"),
 				FSlateIcon(GameplayCamerasEditorStyleName, "DebugCategory.BlendStacks.Icon")
 			});
+		RegisterDebugCategory(FCameraDebugCategoryInfo{
+				FCameraDebugCategories::Services,
+				LOCTEXT("ServicesDebugCategory", "Services"),
+				LOCTEXT("ServicesDebugCategoryToolTip", "Shows the debug information from evaluation services"),
+				FSlateIcon(GameplayCamerasEditorStyleName, "DebugCategory.Services.Icon")
+			});
 		RegisterDebugCategory(FCameraDebugCategoryInfo{
 				FCameraDebugCategories::PoseStats,
 				LOCTEXT("PoseStatsDebugCategory", "Pose Stats"),
