Cameras: add icon for "services" debug category, revamp all debugger png icons into svg ones

diff --git a/Engine/Plugins/Cameras/GameplayCameras/Source/GameplayCamerasEditor/Private/Debugger/SGameplayCamerasDebugger.cpp b/Engine/Plugins/Cameras/GameplayCameras/Source/GameplayCamerasEditor/Private/Debugger/SGameplayCamerasDebugger.cpp
--- a/Engine/Plugins/Cameras/GameplayCameras/Source/GameplayCamerasEditor/Private/Debugger/SGameplayCamerasDebugger.cpp
+++ b/Engine/Plugins/Cameras/GameplayCameras/Source/GameplayCamerasEditor/Private/Debugger/SGameplayCamerasDebugger.cpp
@@ -389,11 +389,11 @@ FSlateIcon SGameplayCamerasDebugger::GetToggleDebugDrawIcon() const
 {
 	if (GGameplayCamerasDebugEnable)
 	{
-		return FSlateIcon(GameplayCamerasEditorStyleName, "Debugger.EnableDebugInfo.Icon");
+		return FSlateIcon(GameplayCamerasEditorStyleName, "Debugger.DebugInfoEnabled.Icon");
 	}
 	else
 	{
-		return FSlateIcon(GameplayCamerasEditorStyleName, "Debugger.DisableDebugInfo.Icon");
+		return FSlateIcon(GameplayCamerasEditorStyleName, "Debugger.DebugInfoDisabled.Icon");
 	}
 }
 
diff --git a/Engine/Plugins/Cameras/GameplayCameras/Source/GameplayCamerasEditor/Private/Styles/GameplayCamerasEditorStyle.cpp b/Engine/Plugins/Cameras/GameplayCameras/Source/GameplayCamerasEditor/Private/Styles/GameplayCamerasEditorStyle.cpp
--- a/Engine/Plugins/Cameras/GameplayCameras/Source/GameplayCamerasEditor/Private/Styles/GameplayCamerasEditorStyle.cpp
+++ b/Engine/Plugins/Cameras/GameplayCameras/Source/GameplayCamerasEditor/Private/Styles/GameplayCamerasEditorStyle.cpp
@@ -97,14 +97,15 @@ FGameplayCamerasEditorStyle::FGameplayCamerasEditorStyle()
 	Set("CameraVariableCollectionEditor.Entry.Value", FTextBlockStyle(NormalText) .SetFont(DEFAULT_FONT("Regular", 10)));
 
 	// Debugger tool icons.
-	Set("Debugger.EnableDebugInfo.Icon", new IMAGE_BRUSH("Icons/EnableDebugInfo", Icon16x16));
-	Set("Debugger.DisableDebugInfo.Icon", new IMAGE_BRUSH("Icons/DisableDebugInfo", Icon16x16));
-
-	Set("DebugCategory.NodeTree.Icon", new IMAGE_BRUSH("Icons/DebugCategory-NodeTree", Icon16x16));
-	Set("DebugCategory.DirectorTree.Icon", new IMAGE_BRUSH("Icons/DebugCategory-DirectorTree", Icon16x16));
-	Set("DebugCategory.BlendStacks.Icon", new IMAGE_BRUSH("Icons/DebugCategory-BlendStacks", Icon16x16));
-	Set("DebugCategory.PoseStats.Icon", new IMAGE_BRUSH("Icons/DebugCategory-PoseStats", Icon16x16));
-	Set("DebugCategory.Viewfinder.Icon", new IMAGE_BRUSH("Icons/DebugCategory-Viewfinder", Icon16x16));
+	Set("Debugger.DebugInfoEnabled.Icon", new IMAGE_BRUSH_SVG("Icons/DebugInfo-ToggleCheck", Icon16x16, FStyleColors::AccentGreen));
+	Set("Debugger.DebugInfoDisabled.Icon", new IMAGE_BRUSH_SVG("Icons/DebugInfo-ToggleCheck", Icon16x16, FStyleColors::AccentGray));
+
+	Set("DebugCategory.NodeTree.Icon", new IMAGE_BRUSH_SVG("Icons/DebugCategory-NodeTree", Icon16x16));
+	Set("DebugCategory.DirectorTree.Icon", new IMAGE_BRUSH_SVG("Icons/DebugCategory-DirectorTree", Icon16x16));
+	Set("DebugCategory.BlendStacks.Icon", new IMAGE_BRUSH_SVG("Icons/DebugCategory-BlendStacks", Icon16x16));
+	Set("DebugCategory.Services.Icon", new IMAGE_BRUSH_SVG("Icons/DebugCategory-Services", Icon16x16));
+	Set("DebugCategory.PoseStats.Icon", new IMAGE_BRUSH_SVG("Icons/DebugCategory-PoseStats", Icon16x16));
+	Set("DebugCategory.Viewfinder.Icon", new IMAGE_BRUSH_SVG("Icons/DebugCategory-Viewfinder", Icon16x16));
 
 	// Graph editor brushes.
 	Set("Graph.CameraRigParameterNode.Body", new BOX_BRUSH("Graph/CameraRigParameterNode_Body", FMargin(16.f/64.f, 12.f/28.f)));
