Cameras: setup variable table during parameter pre-blend phase

This makes variables correctly propagate.

diff --git a/Engine/Plugins/Cameras/GameplayCameras/Source/GameplayCameras/Private/Core/BlendStackCameraNode.cpp b/Engine/Plugins/Cameras/GameplayCameras/Source/GameplayCameras/Private/Core/BlendStackCameraNode.cpp
--- a/Engine/Plugins/Cameras/GameplayCameras/Source/GameplayCameras/Private/Core/BlendStackCameraNode.cpp
+++ b/Engine/Plugins/Cameras/GameplayCameras/Source/GameplayCameras/Private/Core/BlendStackCameraNode.cpp
@@ -335,6 +335,13 @@ void FBlendStackCameraNodeEvaluator::OnRun(const FCameraNodeEvaluationParams& Pa
 
 		FCameraNodeEvaluationResult& CurResult(Entry.Result);
 
+		// Start with the input given to us.
+		CurResult.VariableTable.OverrideAll(OutResult.VariableTable);
+
+		// Override it with whatever the evaluation context has set on its result.
+		const FCameraNodeEvaluationResult& ContextResult(ResolvedEntry.Context->GetInitialResult());
+		CurResult.VariableTable.OverrideAll(ContextResult.VariableTable);
+
 		CurResult.VariableTable.ClearAllWrittenThisFrameFlags();
 
 		// Gather input parameters.
@@ -413,12 +420,10 @@ void FBlendStackCameraNodeEvaluator::OnRun(const FCameraNodeEvaluationParams& Pa
 
 		// Start with the input given to us.
 		CurResult.CameraPose = OutResult.CameraPose;
-		CurResult.VariableTable.OverrideAll(OutResult.VariableTable);
 
 		// Override it with whatever the evaluation context has set on its result.
 		const FCameraNodeEvaluationResult& ContextResult(ResolvedEntry.Context->GetInitialResult());
 		CurResult.CameraPose.OverrideChanged(ContextResult.CameraPose);
-		CurResult.VariableTable.OverrideAll(ContextResult.VariableTable);
 		CurResult.bIsCameraCut = OutResult.bIsCameraCut || ContextResult.bIsCameraCut;
 		CurResult.bIsValid = true;
 
