Cameras: fix comment/doc

diff --git a/Engine/Plugins/Cameras/GameplayCameras/Source/GameplayCameras/Public/Debug/CameraDebugBlock.h b/Engine/Plugins/Cameras/GameplayCameras/Source/GameplayCameras/Public/Debug/CameraDebugBlock.h
--- a/Engine/Plugins/Cameras/GameplayCameras/Source/GameplayCameras/Public/Debug/CameraDebugBlock.h
+++ b/Engine/Plugins/Cameras/GameplayCameras/Source/GameplayCameras/Public/Debug/CameraDebugBlock.h
@@ -150,7 +150,7 @@ private:
 //			UE_DECLARE_CAMERA_DEBUG_BLOCK_FIELD(bool, bSomethingElse)
 //	   UE_DECLARE_CAMERA_DEBUG_BLOCK_END()
 //
-//	   UE_DEFINE_CAMERA_DEBUG_BLOCK(FMyDebugBlock)
+//	   UE_DEFINE_CAMERA_DEBUG_BLOCK_WITH_FIELDS(FMyDebugBlock)
 //
 #define UE_DECLARE_CAMERA_DEBUG_BLOCK_START(ApiDeclSpec, ClassName)\
 	class ClassName : public ::UE::Cameras::FCameraDebugBlock\
