Cameras: make it easier to get/set variables in a variable table using a variable definition struct

diff --git a/Engine/Plugins/Cameras/GameplayCameras/Source/GameplayCameras/Public/Core/CameraVariableTable.h b/Engine/Plugins/Cameras/GameplayCameras/Source/GameplayCameras/Public/Core/CameraVariableTable.h
--- a/Engine/Plugins/Cameras/GameplayCameras/Source/GameplayCameras/Public/Core/CameraVariableTable.h
+++ b/Engine/Plugins/Cameras/GameplayCameras/Source/GameplayCameras/Public/Core/CameraVariableTable.h
@@ -116,6 +116,12 @@ public:
 	template<typename ValueType>
 	bool TrySetValue(FCameraVariableID VariableID, typename TCallTraits<ValueType>::ParamType Value);
 
+	template<typename ValueType>
+	void SetValue(
+			const FCameraVariableDefinition& VariableDefinition, 
+			typename TCallTraits<ValueType>::ParamType Value,
+			bool bCreateIfMissing = false);
+
 	template<typename VariableAssetType>
 	void SetValue(
 			const VariableAssetType* VariableAsset, 
@@ -296,6 +302,36 @@ bool FCameraVariableTable::TrySetValue(FCameraVariableID VariableID, typename TC
 	return false;
 }
 
+template<typename ValueType>
+void FCameraVariableTable::SetValue(
+		const FCameraVariableDefinition& VariableDefinition, 
+		typename TCallTraits<ValueType>::ParamType Value,
+		bool bCreateIfMissing)
+{
+	const bool bDidSet = TrySetValue<ValueType>(VariableDefinition.VariableID, Value);
+#if WITH_EDITORONLY_DATA
+	ensureMsgf(
+			bDidSet || bCreateIfMissing, 
+			TEXT("Can't set camera variable '%s' (ID '%d') because it doesn't exist in the table."),
+			*VariableDefinition.VariableName, VariableDefinition.VariableID.GetValue());
+#else
+	ensureMsgf(
+			bDidSet || bCreateIfMissing, 
+			TEXT("Can't set camera variable '%s' (ID '%d') because it doesn't exist in the table."),
+			*LexToString(VariableDefinition.VariableID.GetValue()), VariableDefinition.VariableID.GetValue());
+#endif
+	if (bDidSet)
+	{
+		return;
+	}
+
+	if (bCreateIfMissing)
+	{
+		AddVariable(VariableDefinition);
+		SetValue<ValueType>(VariableDefinition.VariableID, Value);
+	}
+}
+
 template<typename VariableAssetType>
 void FCameraVariableTable::SetValue(
 		const VariableAssetType* VariableAsset, 
@@ -304,7 +340,12 @@ void FCameraVariableTable::SetValue(
 {
 	if (ensure(VariableAsset))
 	{
-		if (TrySetValue<typename VariableAssetType::ValueType>(VariableAsset->GetVariableID(), Value))
+		const bool bDidSet = TrySetValue<typename VariableAssetType::ValueType>(VariableAsset->GetVariableID(), Value);
+		ensureMsgf(
+				bDidSet || bCreateIfMissing, 
+				TEXT("Can't set camera variable '%s' (ID '%d') because it doesn't exist in the table."),
+				*GetNameSafe(VariableAsset), VariableAsset->GetVariableID().GetValue());
+		if (bDidSet)
 		{
 			return;
 		}
diff --git a/Engine/Plugins/Cameras/GameplayCameras/Source/GameplayCameras/Public/Core/CameraVariableTableFwd.h b/Engine/Plugins/Cameras/GameplayCameras/Source/GameplayCameras/Public/Core/CameraVariableTableFwd.h
--- a/Engine/Plugins/Cameras/GameplayCameras/Source/GameplayCameras/Public/Core/CameraVariableTableFwd.h
+++ b/Engine/Plugins/Cameras/GameplayCameras/Source/GameplayCameras/Public/Core/CameraVariableTableFwd.h
@@ -145,11 +145,19 @@ struct FCameraVariableDefinition
 	FString VariableName;
 #endif
 
+	/** Returns whether this definition has a valid variable ID. */
 	bool IsValid() const
 	{
 		return VariableID.IsValid();
 	}
 
+	/** Implicit conversion to a camera variable ID. */
+	operator FCameraVariableID() const
+	{
+		return VariableID;
+	}
+
+	/** Creates a variant of this camera variable definition. */
 	FCameraVariableDefinition CreateVariant(const FString& VariantID) const
 	{
 		FCameraVariableDefinition VariantDefinition(*this);
