Cameras: add reference collection API to evaluation services

diff --git a/Engine/Plugins/Cameras/GameplayCameras/Source/GameplayCameras/Private/Core/CameraEvaluationService.cpp b/Engine/Plugins/Cameras/GameplayCameras/Source/GameplayCameras/Private/Core/CameraEvaluationService.cpp
--- a/Engine/Plugins/Cameras/GameplayCameras/Source/GameplayCameras/Private/Core/CameraEvaluationService.cpp
+++ b/Engine/Plugins/Cameras/GameplayCameras/Source/GameplayCameras/Private/Core/CameraEvaluationService.cpp
@@ -35,6 +35,11 @@ void FCameraEvaluationService::Teardown(const FCameraEvaluationServiceTeardownPa
 	OnTeardown(Params);
 }
 
+void FCameraEvaluationService::AddReferencedObjects(FReferenceCollector& Collector)
+{
+	OnAddReferencedObjects(Collector);
+}
+
 void FCameraEvaluationService::NotifyRootCameraNodeEvent(const FRootCameraNodeCameraRigEvent& InEvent)
 {
 	OnRootCameraNodeEvent(InEvent);
diff --git a/Engine/Plugins/Cameras/GameplayCameras/Source/GameplayCameras/Private/Core/CameraSystemEvaluator.cpp b/Engine/Plugins/Cameras/GameplayCameras/Source/GameplayCameras/Private/Core/CameraSystemEvaluator.cpp
--- a/Engine/Plugins/Cameras/GameplayCameras/Source/GameplayCameras/Private/Core/CameraSystemEvaluator.cpp
+++ b/Engine/Plugins/Cameras/GameplayCameras/Source/GameplayCameras/Private/Core/CameraSystemEvaluator.cpp
@@ -94,6 +94,10 @@ void FCameraSystemEvaluator::AddReferencedObjects(FReferenceCollector& Collector
 	{
 		RootEvaluator->AddReferencedObjects(Collector);
 	}
+	for (TSharedPtr<FCameraEvaluationService> EvaluationService : EvaluationServices)
+	{
+		EvaluationService->AddReferencedObjects(Collector);
+	}
 }
 
 void FCameraSystemEvaluator::PushEvaluationContext(TSharedRef<FCameraEvaluationContext> EvaluationContext)
diff --git a/Engine/Plugins/Cameras/GameplayCameras/Source/GameplayCameras/Public/Core/CameraEvaluationService.h b/Engine/Plugins/Cameras/GameplayCameras/Source/GameplayCameras/Public/Core/CameraEvaluationService.h
--- a/Engine/Plugins/Cameras/GameplayCameras/Source/GameplayCameras/Public/Core/CameraEvaluationService.h
+++ b/Engine/Plugins/Cameras/GameplayCameras/Source/GameplayCameras/Public/Core/CameraEvaluationService.h
@@ -7,6 +7,8 @@
 #include "GameplayCameras.h"
 #include "Templates/SharedPointer.h"
 
+class FReferenceCollector;
+
 namespace UE::Cameras
 {
 
@@ -80,6 +82,9 @@ public:
 	/** Tears down the evaluation service. */
 	void Teardown(const FCameraEvaluationServiceTeardownParams& Params);
 
+	/** Collect referenced UObjects. */
+	void AddReferencedObjects(FReferenceCollector& Collector);
+
 #if UE_GAMEPLAY_CAMERAS_DEBUG
 	/** Called to create debug blocks for this node evaluator. */
 	void BuildDebugBlocks(const FCameraDebugBlockBuildParams& Params, FCameraDebugBlockBuilder& Builder);
@@ -105,6 +110,9 @@ protected:
 	/** Tears down the evaluation service. */
 	virtual void OnTeardown(const FCameraEvaluationServiceTeardownParams& Params) {}
 
+	/** Collect referenced UObjects. */
+	virtual void OnAddReferencedObjects(FReferenceCollector& Collector) {}
+
 	/** Called when the root camera node experiences an event. */
 	virtual void OnRootCameraNodeEvent(const FRootCameraNodeCameraRigEvent& InEvent) {}
 
