Cameras: auto-build camera asset when a camera component begins play

This prevents errors where the camera asset doesn't have up-to-date allocation information.

diff --git a/Engine/Plugins/Cameras/GameplayCameras/Source/GameplayCameras/Private/GameFramework/GameplayCameraComponent.cpp b/Engine/Plugins/Cameras/GameplayCameras/Source/GameplayCameras/Private/GameFramework/GameplayCameraComponent.cpp
--- a/Engine/Plugins/Cameras/GameplayCameras/Source/GameplayCameras/Private/GameFramework/GameplayCameraComponent.cpp
+++ b/Engine/Plugins/Cameras/GameplayCameras/Source/GameplayCameras/Private/GameFramework/GameplayCameraComponent.cpp
@@ -3,6 +3,8 @@
 #include "GameFramework/GameplayCameraComponent.h"
 
 #include "Components/StaticMeshComponent.h"
+#include "Core/CameraAssetBuilder.h"
+#include "Core/CameraBuildLog.h"
 #include "Core/CameraSystemEvaluator.h"
 #include "Engine/StaticMesh.h"
 #include "Engine/World.h"
@@ -180,6 +182,16 @@ void UGameplayCameraComponent::BeginPlay()
 {
 	Super::BeginPlay();
 
+#if WITH_EDITOR
+	if (Camera)
+	{
+		using namespace UE::Cameras;
+		FCameraBuildLog BuildLog;
+		FCameraAssetBuilder Builder(BuildLog);
+		Builder.BuildCamera(Camera);
+	}
+#endif
+
 	if (AutoActivateForPlayer != EAutoReceiveInput::Disabled && GetNetMode() != NM_DedicatedServer)
 	{
 		const int32 PlayerIndex = AutoActivateForPlayer.GetIntValue() - 1;
