Cameras: add control rotation handling, improve auto-rotate node

The auto-rotate node brings considerations of player control rotation to the camera system, so this change implements an example of how to handle this.

- Remove option to set control rotation directly from the camera system component.
- Add a new "player control rotation" component that reads information from the camera system and figures out what to do with control rotation.
- New options and better logic for the auto-rotate node.

The new player control rotation component can freeze the control rotation when needed, based on hints given by nodes like the auto-rotate node.

diff --git a/Engine/Plugins/Cameras/GameplayCameras/Source/GameplayCameras/Private/Core/BuiltInCameraVariables.cpp b/Engine/Plugins/Cameras/GameplayCameras/Source/GameplayCameras/Private/Core/BuiltInCameraVariables.cpp
--- a/Engine/Plugins/Cameras/GameplayCameras/Source/GameplayCameras/Private/Core/BuiltInCameraVariables.cpp
+++ b/Engine/Plugins/Cameras/GameplayCameras/Source/GameplayCameras/Private/Core/BuiltInCameraVariables.cpp
@@ -17,6 +17,9 @@ struct FBuiltInCameraVariablesGuids
 
 	FGuid YawPitchGuid;  // {67B7828D-C645-4907-92BE-B40DEB8C838C}
 
+	FGuid FreezeControlRotationGuid; // {39EDC23E-70F2-4041-971B-7BD680F640D1}
+	FGuid ControlRotationGuid;  // {2AB5FCA4-D62C-494C-B700-C8F98CD2D8E8}
+
 	FBuiltInCameraVariablesGuids()
 	{
 		YawGuid = FGuid(0x6E23348F, 0x290E460F, 0x9432ED80, 0xE7CA03F2);
@@ -25,6 +28,9 @@ struct FBuiltInCameraVariablesGuids
 		ZoomGuid = FGuid(0x56BF6A94, 0x87B34648, 0x84A4391B, 0xDAD27061);
 
 		YawPitchGuid = FGuid(0x67B7828D, 0xC6454907, 0x92BEB40D, 0xEB8C838C);
+
+		FreezeControlRotationGuid = FGuid(0x39EDC23E, 0x70F24041, 0x971B7BD6, 0x80F640D1);
+		ControlRotationGuid = FGuid(0x2AB5FCA4, 0xD62C494C, 0xB700C8F9, 0x8CD2D8E8);
 	}
 };
 
@@ -61,6 +67,9 @@ FBuiltInCameraVariables::FBuiltInCameraVariables()
 	MakeCameraVariableDefinition(ZoomDefinition, KnownGuids.ZoomGuid, ECameraVariableType::Double, true, TEXT("Zoom"));
 
 	MakeCameraVariableDefinition(YawPitchDefinition, KnownGuids.YawPitchGuid, ECameraVariableType::Vector2d, true, TEXT("YawPitch"));
+
+	MakeCameraVariableDefinition(FreezeControlRotationDefinition, KnownGuids.FreezeControlRotationGuid, ECameraVariableType::Boolean, true, TEXT("FreezeControlRotation"));
+	MakeCameraVariableDefinition(ControlRotationDefinition, KnownGuids.ControlRotationGuid, ECameraVariableType::Rotator3d, true, TEXT("ControlRotation"));
 }
 
 const FCameraVariableDefinition& FBuiltInCameraVariables::GetDefinition(EBuiltInDoubleCameraVariable BuiltInVariable) const
@@ -93,5 +102,17 @@ const FCameraVariableDefinition& FBuiltInCameraVariables::GetDefinition(EBuiltIn
 	}
 }
 
+const FCameraVariableDefinition& FBuiltInCameraVariables::GetDefinition(EBuiltInRotator3dCameraVariable BuiltInVariable) const
+{
+	static const FCameraVariableDefinition Invalid;
+	switch (BuiltInVariable)
+	{
+		case EBuiltInRotator3dCameraVariable::ControlRotation:
+			return ControlRotationDefinition;
+		default:
+			return Invalid;
+	}
+}
+
 }  // namespace UE::Cameras
 
diff --git a/Engine/Plugins/Cameras/GameplayCameras/Source/GameplayCameras/Private/GameFramework/GameplayCameraComponent.cpp b/Engine/Plugins/Cameras/GameplayCameras/Source/GameplayCameras/Private/GameFramework/GameplayCameraComponent.cpp
--- a/Engine/Plugins/Cameras/GameplayCameras/Source/GameplayCameras/Private/GameFramework/GameplayCameraComponent.cpp
+++ b/Engine/Plugins/Cameras/GameplayCameras/Source/GameplayCameras/Private/GameFramework/GameplayCameraComponent.cpp
@@ -39,6 +39,11 @@ UGameplayCameraComponent::UGameplayCameraComponent(const FObjectInitializer& Obj
 #endif  // WITH_EDITORONLY_DATA
 }
 
+TSharedPtr<UE::Cameras::FCameraEvaluationContext> UGameplayCameraComponent::GetEvaluationContext()
+{
+	return EvaluationContext;
+}
+
 void UGameplayCameraComponent::ActivateCamera(int32 PlayerIndex)
 {
 	if (ActivatedForPlayerIndex == PlayerIndex)
diff --git a/Engine/Plugins/Cameras/GameplayCameras/Source/GameplayCameras/Private/GameFramework/GameplayCameraSystemComponent.cpp b/Engine/Plugins/Cameras/GameplayCameras/Source/GameplayCameras/Private/GameFramework/GameplayCameraSystemComponent.cpp
--- a/Engine/Plugins/Cameras/GameplayCameras/Source/GameplayCameras/Private/GameFramework/GameplayCameraSystemComponent.cpp
+++ b/Engine/Plugins/Cameras/GameplayCameras/Source/GameplayCameras/Private/GameFramework/GameplayCameraSystemComponent.cpp
@@ -43,15 +43,6 @@ void UGameplayCameraSystemComponent::GetCameraView(float DeltaTime, FMinimalView
 		UpdateParams.DeltaTime = DeltaTime;
 		Evaluator->Update(UpdateParams);
 
-		if (bSetPlayerControllerRotation && ActivatedForPlayerIndex != INDEX_NONE)
-		{
-			APlayerController* PC = UGameplayStatics::GetPlayerController(this, ActivatedForPlayerIndex);
-			if (PC)
-			{
-				PC->SetControlRotation(Evaluator->GetEvaluatedResult().CameraPose.GetRotation());
-			}
-		}
-
 		Evaluator->GetEvaluatedCameraView(DesiredView);
 	}
 }
diff --git a/Engine/Plugins/Cameras/GameplayCameras/Source/GameplayCameras/Private/GameFramework/GameplayControlRotationComponent.cpp b/Engine/Plugins/Cameras/GameplayCameras/Source/GameplayCameras/Private/GameFramework/GameplayControlRotationComponent.cpp
new file mode 100644
--- /dev/null
+++ b/Engine/Plugins/Cameras/GameplayCameras/Source/GameplayCameras/Private/GameFramework/GameplayControlRotationComponent.cpp
@@ -0,0 +1,140 @@
+// Copyright Epic Games, Inc. All Rights Reserved.
+
+#include "GameFramework/GameplayControlRotationComponent.h"
+
+#include "Core/CameraEvaluationContext.h"
+#include "Debug/DebugDrawService.h"
+#include "Engine/Engine.h"
+#include "Engine/World.h"
+#include "EnhancedInputComponent.h"
+#include "GameFramework/Controller.h"
+#include "GameFramework/GameplayCameraComponent.h"
+#include "GameFramework/GameplayCameraSystemComponent.h"
+#include "GameFramework/Pawn.h"
+#include "GameplayCameras.h"
+#include "Kismet/GameplayStatics.h"
+#include "Math/ColorList.h"
+#include "Services/PlayerControlRotationService.h"
+
+#include UE_INLINE_GENERATED_CPP_BY_NAME(GameplayControlRotationComponent)
+
+#define LOCTEXT_NAMESPACE "GameplayControlRotationComponent"
+
+UGameplayControlRotationComponent::UGameplayControlRotationComponent(const FObjectInitializer& ObjectInit)
+	: Super(ObjectInit)
+{
+	bAutoActivate = true;
+	PrimaryComponentTick.bCanEverTick = true;
+}
+
+void UGameplayControlRotationComponent::BeginPlay()
+{
+	Super::BeginPlay();
+
+	using namespace UE::Cameras;
+
+	// We need to grab our owner actor, and then find the gameplay camera component and the enhanced
+	// input component that we will be coordinating.
+	AActor* OwnerActor = GetOwner();
+	if (!OwnerActor)
+	{
+		UE_LOG(LogCameraSystem, Error, 
+				TEXT("GameplayControlRotationComponent '%s' doesn't belong to any actor"),
+				*GetNameSafe(this));
+		return;
+	}
+
+	UGameplayCameraComponent* GameplayCameraComponent = OwnerActor->FindComponentByClass<UGameplayCameraComponent>();
+	if (!GameplayCameraComponent)
+	{
+		UE_LOG(LogCameraSystem, Error,
+				TEXT("GameplayControlRotationComponent '%s' couldn't find a GameplayCameraComponent on owner '%s'"),
+				*GetNameSafe(this), *GetNameSafe(OwnerActor));
+		return;
+	}
+
+	// Find where the camera system is running from.
+	PlayerIndex = GameplayCameraComponent->GetPlayerIndex();
+	if (!ensure(PlayerIndex != INDEX_NONE))
+	{
+		UE_LOG(LogCameraSystem, Error,
+				TEXT("GameplayCameraComponent '%s' hasn't activated for any player yet"),
+				*GetNameSafe(GameplayCameraComponent));
+		return;
+	}
+
+	APlayerController* PlayerController = UGameplayStatics::GetPlayerController(this, PlayerIndex);
+	if (!PlayerController)
+	{
+		UE_LOG(LogCameraSystem, Error, TEXT("Can't find player controller for player index %d"), PlayerIndex);
+		return;
+	}
+
+	AActor* ViewTarget = PlayerController->PlayerCameraManager->GetViewTarget();
+	UGameplayCameraSystemComponent* CameraSystemComponent = ViewTarget->FindComponentByClass<UGameplayCameraSystemComponent>();
+	if (!CameraSystemComponent)
+	{
+		UE_LOG(LogCameraSystem, Error, TEXT("Can't find camera system on the player camera manager."));
+		return;
+	}
+
+	TSharedPtr<FCameraSystemEvaluator> CameraSystem = CameraSystemComponent->GetCameraSystemEvaluator();
+	if (!CameraSystem)
+	{
+		UE_LOG(LogCameraSystem, Error, TEXT("No camera system has been initialized"));
+		return;
+	}
+
+	// Create the evaluation service, with a copy of our parameters.
+	FPlayerControlRotationParams ServiceParams;
+	ServiceParams.AxisActionAngularSpeedThreshold = AxisActionAngularSpeedThreshold;
+	ServiceParams.AxisActionMagnitudeThreshold = AxisActionMagnitudeThreshold;
+	ServiceParams.AxisActions = AxisActions;
+	// We will set the control rotation ourselves.
+	ServiceParams.bApplyControlRotation = false;
+
+	ControlRotationService = MakeShared<FPlayerControlRotationEvaluationService>(ServiceParams);
+	CameraSystem->RegisterEvaluationService(ControlRotationService.ToSharedRef());
+
+	WeakCameraSystem = CameraSystem;
+}
+
+void UGameplayControlRotationComponent::EndPlay(const EEndPlayReason::Type EndPlayReason)
+{
+	using namespace UE::Cameras;
+
+	if (TSharedPtr<FCameraSystemEvaluator> CameraSystem = WeakCameraSystem.Pin())
+	{
+		CameraSystem->UnregisterEvaluationService(ControlRotationService.ToSharedRef());
+	}
+
+	ControlRotationService.Reset();
+
+	Super::EndPlay(EndPlayReason);
+}
+
+void UGameplayControlRotationComponent::TickComponent(float DeltaTime, ELevelTick TickType, FActorComponentTickFunction* ThisTickFunction)
+{
+	Super::TickComponent(DeltaTime, TickType, ThisTickFunction);
+
+	APlayerController* PlayerController = GetPlayerController();
+	if (PlayerController)
+	{
+		// This may be technically one frame late (i.e. we set the control rotation computed 
+		// late last tick) unless the camera system is setup to process player input into camera
+		// rotation early in the frame.
+		PlayerController->SetControlRotation(ControlRotationService->GetCurrentControlRotation());
+	}
+}
+
+APlayerController* UGameplayControlRotationComponent::GetPlayerController()
+{
+	if (PlayerIndex != INDEX_NONE)
+	{
+		return UGameplayStatics::GetPlayerController(this, PlayerIndex);
+	}
+	return nullptr;
+}
+
+#undef LOCTEXT_NAMESPACE
+
diff --git a/Engine/Plugins/Cameras/GameplayCameras/Source/GameplayCameras/Private/Nodes/Input/AutoRotateInput2DCameraNode.cpp b/Engine/Plugins/Cameras/GameplayCameras/Source/GameplayCameras/Private/Nodes/Input/AutoRotateInput2DCameraNode.cpp
--- a/Engine/Plugins/Cameras/GameplayCameras/Source/GameplayCameras/Private/Nodes/Input/AutoRotateInput2DCameraNode.cpp
+++ b/Engine/Plugins/Cameras/GameplayCameras/Source/GameplayCameras/Private/Nodes/Input/AutoRotateInput2DCameraNode.cpp
@@ -2,19 +2,30 @@
 
 #include "Nodes/Input/AutoRotateInput2DCameraNode.h"
 
+#include "Core/BuiltInCameraVariables.h"
 #include "Core/CameraEvaluationContext.h"
 #include "Core/CameraOperation.h"
 #include "Core/CameraParameterReader.h"
+#include "Core/CameraRigAsset.h"
+#include "Core/CameraRigBuildContext.h"
 #include "Core/CameraValueInterpolator.h"
+#include "Core/CameraVariableTableFwd.h"
 #include "Debug/CameraDebugBlock.h"
 #include "Debug/CameraDebugBlockBuilder.h"
 #include "Debug/CameraDebugRenderer.h"
+#include "HAL/IConsoleManager.h"
 
 #include UE_INLINE_GENERATED_CPP_BY_NAME(AutoRotateInput2DCameraNode)
 
 namespace UE::Cameras
 {
 
+float GGameplayCamerasAutoRotateSnapThreshold = 0.5f;
+static FAutoConsoleVariableRef CVarGameplayCamerasAutoRotateSnapThreshold(
+	TEXT("GameplayCameras.AutoRotate.SnapThreshold"),
+	GGameplayCamerasAutoRotateSnapThreshold,
+	TEXT(""));
+
 class FAutoRotateInput2DCameraNodeEvaluator : public FInput2DCameraNodeEvaluator
 {
 	UE_DECLARE_CAMERA_NODE_EVALUATOR_EX(GAMEPLAYCAMERAS_API, FAutoRotateInput2DCameraNodeEvaluator, FInput2DCameraNodeEvaluator)
@@ -34,18 +45,27 @@ protected:
 	virtual void OnBuildDebugBlocks(const FCameraDebugBlockBuildParams& Params, FCameraDebugBlockBuilder& Builder) override;
 #endif  // UE_GAMEPLAY_CAMERAS_DEBUG
 
+private:
+
+	void DeactivateAutoRotate(FCameraNodeEvaluationResult& OutResult);
+
 private:
 
 	TCameraParameterReader<float> WaitTimeReader;
 	TCameraParameterReader<float> DeactivationThresholdReader;
+	TCameraParameterReader<bool> FreezeControlRotationReader;
 	TCameraParameterReader<bool> EnableAutoRotateReader;
+	TCameraParameterReader<bool> AutoRotateYawReader;
+	TCameraParameterReader<bool> AutoRotatePitchReader;
 
 	TUniquePtr<FCameraDoubleValueInterpolator> Interpolator;
 
 	FInput2DCameraNodeEvaluator* InputNodeEvaluator = nullptr;
 
-	FVector3d LastContextLocation;
 	FVector2d LastInputValue;
+
+	FVector3d LastContextLocation;
+	FVector2d OriginalInputValue;
 	double RemainingWaitTime = 0.f;
 	bool bIsAutoRotating = false;
 };
@@ -94,7 +114,10 @@ void FAutoRotateInput2DCameraNodeEvaluator::OnInitialize(const FCameraNodeEvalua
 
 	WaitTimeReader.Initialize(AutoRotateNode->WaitTime);
 	DeactivationThresholdReader.Initialize(AutoRotateNode->DeactivationThreshold);
+	FreezeControlRotationReader.Initialize(AutoRotateNode->FreezeControlRotation);
 	EnableAutoRotateReader.Initialize(AutoRotateNode->EnableAutoRotate);
+	AutoRotateYawReader.Initialize(AutoRotateNode->AutoRotateYaw);
+	AutoRotatePitchReader.Initialize(AutoRotateNode->AutoRotatePitch);
 
 	RemainingWaitTime = WaitTimeReader.Get(OutResult.VariableTable);
 	bIsAutoRotating = false;
@@ -114,15 +137,11 @@ void FAutoRotateInput2DCameraNodeEvaluator::OnRun(const FCameraNodeEvaluationPar
 
 	// Bail out if auto-rotate is disabled.
 	const bool bAutoRotateEnabled = EnableAutoRotateReader.Get(OutResult.VariableTable);
-	if (!bAutoRotateEnabled)
+	const bool bAutoRotateYaw = AutoRotateYawReader.Get(OutResult.VariableTable);
+	const bool bAutoRotatePitch = AutoRotatePitchReader.Get(OutResult.VariableTable);
+	if (!bAutoRotateEnabled || (!bAutoRotateYaw && !bAutoRotatePitch))
 	{
-		// If we were disabled while auto-rotating, tear down everything we had.
-		if (bIsAutoRotating)
-		{
-			bIsAutoRotating = false;
-			Interpolator.Reset();
-			RemainingWaitTime = WaitTimeReader.Get(OutResult.VariableTable);
-		}
+		DeactivateAutoRotate(OutResult);
 		return;
 	}
 
@@ -142,24 +161,12 @@ void FAutoRotateInput2DCameraNodeEvaluator::OnRun(const FCameraNodeEvaluationPar
 	const double PitchChange = FMath::Abs(InputValue.Y - LastInputValue.Y);
 	const float DeactivationThreshold = DeactivationThresholdReader.Get(OutResult.VariableTable);
 	LastInputValue = InputValue;
-	if (YawChange >= DeactivationThreshold || PitchChange >= DeactivationThreshold)
+	if ((bAutoRotateYaw && YawChange >= DeactivationThreshold) || (bAutoRotatePitch && PitchChange >= DeactivationThreshold))
 	{
-		bIsAutoRotating = false;
-		Interpolator.Reset();
-		RemainingWaitTime = WaitTimeReader.Get(OutResult.VariableTable);
+		DeactivateAutoRotate(OutResult);
 		return;
 	}
 
-	// We are almost good to auto-rotate... but maybe we need to wait a bit longer.
-	if (RemainingWaitTime > 0.f)
-	{
-		RemainingWaitTime -= Params.DeltaTime;
-		if (RemainingWaitTime > 0.f)
-		{
-			return;
-		}
-	}
-
 	// Figure out which direction we should auto-rotate towards.
 	FVector3d AutoRotateDir = InitialResult.CameraPose.GetAimDir();
 	if (AutoRotateNode->Direction == ECameraAutoRotateDirection::Movement)
@@ -169,13 +176,25 @@ void FAutoRotateInput2DCameraNodeEvaluator::OnRun(const FCameraNodeEvaluationPar
 	const FRotator3d AutoRotateRot = AutoRotateDir.ToOrientationRotator();
 
 	// Figure out how much work we have to do.
-	const double DeltaYaw = AutoRotateRot.Yaw - InputValue.X;
-	const double DeltaPitch = AutoRotateRot.Pitch - InputValue.Y;
-	if (FMath::Abs(DeltaYaw) < UE_DOUBLE_SMALL_NUMBER && FMath::Abs(DeltaPitch) < UE_DOUBLE_SMALL_NUMBER)
+	const double DeltaYaw = FRotator::NormalizeAxis(bAutoRotateYaw ? (AutoRotateRot.Yaw - InputValue.X) : 0.0);
+	const double DeltaPitch = FRotator::NormalizeAxis(bAutoRotatePitch ? (AutoRotateRot.Pitch - InputValue.Y) : 0.0);
+	const double DeltaThreshold = GGameplayCamerasAutoRotateSnapThreshold;
+	if (FMath::Abs(DeltaYaw) < DeltaThreshold && FMath::Abs(DeltaPitch) < DeltaThreshold)
 	{
+		DeactivateAutoRotate(OutResult);
 		return;
 	}
 
+	// We are almost good to auto-rotate... but maybe we need to wait a bit longer.
+	if (RemainingWaitTime > 0.f)
+	{
+		RemainingWaitTime -= Params.DeltaTime;
+		if (RemainingWaitTime > 0.f)
+		{
+			return;
+		}
+	}
+
 	// We will interpolate the length of the vector that represents the delta yaw/pitch
 	// that we need to compensate with.
 	const FVector2d DeltaVector = FVector2d(DeltaYaw, DeltaPitch);
@@ -192,6 +211,13 @@ void FAutoRotateInput2DCameraNodeEvaluator::OnRun(const FCameraNodeEvaluationPar
 		{
 			Interpolator = MakeUnique<TPopValueInterpolator<double>>();
 		}
+
+		OriginalInputValue = FVector2d::ZeroVector;
+		if (InputNodeEvaluator)
+		{
+			OriginalInputValue = InputNodeEvaluator->GetInputValue();
+		}
+
 		bIsAutoRotating = true;
 	}
 	check(Interpolator);
@@ -228,10 +254,32 @@ void FAutoRotateInput2DCameraNodeEvaluator::OnRun(const FCameraNodeEvaluationPar
 	}
 
 	if (Interpolator->IsFinished() || bDeactivateAutoRotate)
+	{
+		DeactivateAutoRotate(OutResult);
+	}
+	else if (FreezeControlRotationReader.Get(OutResult.VariableTable))
+	{
+		const FBuiltInCameraVariables& BuiltInVariables = FBuiltInCameraVariables::Get();
+		const FRotator3d OriginalControlRotation(OriginalInputValue.Y, OriginalInputValue.X, 0.0);
+		OutResult.VariableTable.SetValue<FRotator3d>(BuiltInVariables.ControlRotationDefinition, OriginalControlRotation);
+		OutResult.VariableTable.SetValue<bool>(BuiltInVariables.FreezeControlRotationDefinition, true);
+	}
+}
+
+void FAutoRotateInput2DCameraNodeEvaluator::DeactivateAutoRotate(FCameraNodeEvaluationResult& OutResult)
+{
+	if (bIsAutoRotating)
 	{
 		bIsAutoRotating = false;
 		Interpolator.Reset();
 		RemainingWaitTime = WaitTimeReader.Get(OutResult.VariableTable);
+
+		if (FreezeControlRotationReader.Get(OutResult.VariableTable))
+		{
+			// TODO: need Deinit callback to set this back to false also
+			const FBuiltInCameraVariables& BuiltInVariables = FBuiltInCameraVariables::Get();
+			OutResult.VariableTable.SetValue<bool>(BuiltInVariables.FreezeControlRotationDefinition, false);
+		}
 	}
 }
 
@@ -273,6 +321,14 @@ FCameraNodeChildrenView UAutoRotateInput2DCameraNode::OnGetChildren()
 	return FCameraNodeChildrenView({ InputNode });
 }
 
+void UAutoRotateInput2DCameraNode::OnBuild(FCameraRigBuildContext& BuildContext)
+{
+	using namespace UE::Cameras;
+	const FBuiltInCameraVariables& BuiltInVariables = FBuiltInCameraVariables::Get();
+	BuildContext.AllocationInfo.VariableTableInfo.VariableDefinitions.Add(BuiltInVariables.ControlRotationDefinition);
+	BuildContext.AllocationInfo.VariableTableInfo.VariableDefinitions.Add(BuiltInVariables.FreezeControlRotationDefinition);
+}
+
 FCameraNodeEvaluatorPtr UAutoRotateInput2DCameraNode::OnBuildEvaluator(FCameraNodeEvaluatorBuilder& Builder) const
 {
 	using namespace UE::Cameras;
diff --git a/Engine/Plugins/Cameras/GameplayCameras/Source/GameplayCameras/Private/Services/PlayerControlRotationService.cpp b/Engine/Plugins/Cameras/GameplayCameras/Source/GameplayCameras/Private/Services/PlayerControlRotationService.cpp
new file mode 100644
--- /dev/null
+++ b/Engine/Plugins/Cameras/GameplayCameras/Source/GameplayCameras/Private/Services/PlayerControlRotationService.cpp
@@ -0,0 +1,359 @@
+// Copyright Epic Games, Inc. All Rights Reserved.
+
+#include "Services/PlayerControlRotationService.h"
+
+#include "Core/BuiltInCameraVariables.h"
+#include "Core/CameraEvaluationContext.h"
+#include "Core/CameraEvaluationContextStack.h"
+#include "Core/CameraSystemEvaluator.h"
+#include "Debug/CameraDebugBlock.h"
+#include "Debug/CameraDebugBlockBuilder.h"
+#include "Debug/CameraDebugRenderer.h"
+#include "EnhancedInputComponent.h"
+#include "GameFramework/Pawn.h"
+#include "GameFramework/PlayerController.h"
+#include "HAL/IConsoleManager.h"
+
+namespace UE::Cameras
+{
+
+float GGameplayCamerasControlRotationDebugArrowLength = 200.f;
+static FAutoConsoleVariableRef CVarGameplayCamerasControlRotationDebugArrowLength(
+	TEXT("GameplayControlRotation.DebugArrowLength"),
+	GGameplayCamerasControlRotationDebugArrowLength,
+	TEXT(""));
+
+float GGameplayCamerasControlRotationDebugAxisBindingAngleSpeedHistory = 1.f;
+static FAutoConsoleVariableRef CVarGameplayCamerasControlRotationDebugAxisBindingAngleSpeedHistory(
+	TEXT("GameplayControlRotation.DebugAxisBindingAngleSpeedHistory"),
+	GGameplayCamerasControlRotationDebugAxisBindingAngleSpeedHistory,
+	TEXT(""));
+
+UE_DEFINE_CAMERA_EVALUATION_SERVICE(FPlayerControlRotationEvaluationService)
+
+UE_DECLARE_CAMERA_DEBUG_BLOCK_START(GAMEPLAYCAMERAS_API, FPlayerControlRotationDebugBlock)
+	UE_DECLARE_CAMERA_DEBUG_BLOCK_FIELD(FTransform, PawnTransform)
+	UE_DECLARE_CAMERA_DEBUG_BLOCK_FIELD(FRotator3d, ControlRotation)
+	UE_DECLARE_CAMERA_DEBUG_BLOCK_FIELD(FRotator3d, CameraRotation)
+	UE_DECLARE_CAMERA_DEBUG_BLOCK_FIELD(bool, bIsFrozen)
+	UE_DECLARE_CAMERA_DEBUG_BLOCK_FIELD(FString, FreezeReason)
+UE_DECLARE_CAMERA_DEBUG_BLOCK_END()
+
+UE_DEFINE_CAMERA_DEBUG_BLOCK_WITH_FIELDS(FPlayerControlRotationDebugBlock)
+
+FPlayerControlRotationEvaluationService::FPlayerControlRotationEvaluationService()
+	: PreviousAxisBindingValue(EForceInit::ForceInit)
+{
+}
+
+FPlayerControlRotationEvaluationService::FPlayerControlRotationEvaluationService(const FPlayerControlRotationParams& InParams)
+	: ServiceParams(InParams)
+	, PreviousAxisBindingValue(EForceInit::ForceInit)
+{
+	SetEvaluationServiceFlags(ECameraEvaluationServiceFlags::NeedsPostUpdate);
+}
+
+void FPlayerControlRotationEvaluationService::OnAddReferencedObjects(FReferenceCollector& Collector)
+{
+	Collector.AddReferencedObjects(ServiceParams.AxisActions);
+}
+
+void FPlayerControlRotationEvaluationService::MonitorActiveContext(TSharedPtr<const FCameraEvaluationContext> ActiveContext)
+{
+	// Try to find the currently active input component. If we fail at any point, get rid of any
+	// previous bindings we had.
+	if (!ActiveContext)
+	{
+		UnbindActionValues();
+		return;
+	}
+
+	APlayerController* PlayerController = ActiveContext->GetPlayerController();
+	if (!PlayerController)
+	{
+		UnbindActionValues();
+		return;
+	}
+	
+	APawn* Pawn = PlayerController->GetPawn();
+	if (!Pawn)
+	{
+		UnbindActionValues();
+		return;
+	}
+
+	UEnhancedInputComponent* InputComponent = Cast<UEnhancedInputComponent>(Pawn->InputComponent);
+	if (!InputComponent)
+	{
+		UnbindActionValues();
+		return;
+	}
+
+	// If the input component has changed this we last checked, re-bind our input actions to this new one.
+	if (InputComponent != WeakInputComponent)
+	{
+		BindActionValues(InputComponent);
+	}
+}
+
+void FPlayerControlRotationEvaluationService::BindActionValues(UEnhancedInputComponent* InputComponent)
+{
+	UE_LOG(LogCameraSystem, Verbose, TEXT("FPlayerControlRotationEvaluationService: binding to input actions on '%s'"), *GetNameSafe(InputComponent));
+
+	WeakInputComponent = InputComponent;
+
+	AxisBindings.Reset();
+
+	for (TObjectPtr<UInputAction> AxisAction : ServiceParams.AxisActions)
+	{
+		FEnhancedInputActionValueBinding& AxisBinding = InputComponent->BindActionValue(AxisAction);
+		AxisBindings.Add(&AxisBinding);
+	}
+}
+
+void FPlayerControlRotationEvaluationService::UnbindActionValues()
+{
+	UE_LOG(LogCameraSystem, Verbose, TEXT("FPlayerControlRotationEvaluationService: unbinding from input actions"));
+	WeakInputComponent.Reset();
+	AxisBindings.Reset();
+}
+
+void FPlayerControlRotationEvaluationService::OnPostUpdate(const FCameraEvaluationServiceUpdateParams& Params, FCameraEvaluationServiceUpdateResult& OutResult)
+{
+	// Grab the new camera orientation and set is as the control rotation by default.
+	CameraRotation = OutResult.EvaluationResult.CameraPose.GetRotation();
+	CurrentControlRotation = CameraRotation;
+#if UE_GAMEPLAY_CAMERAS_DEBUG
+	bDebugDidApplyControlRotation = false;
+#endif
+
+	// Get the active context.
+	FCameraEvaluationContextStack& EvaluationContextStack = Params.Evaluator->GetEvaluationContextStack();
+	TSharedPtr<const FCameraEvaluationContext> ActiveContext = EvaluationContextStack.GetActiveContext();
+
+	// Check if we need to bind to a new input component, or abandon the one we had.
+	MonitorActiveContext(ActiveContext);
+
+	// If there's no active context, we are done, control rotation isn't frozen.
+	if (!ActiveContext.IsValid())
+	{
+		bIsFrozen = false;
+#if UE_GAMEPLAY_CAMERAS_DEBUG
+		DebugFreezeReason = TEXT("no active context");
+#endif
+		return;
+	}
+
+	// Find the player controller whose control rotation we need to manage.
+	APlayerController* PlayerController = ActiveContext->GetPlayerController();
+	if (!PlayerController)
+	{
+		bIsFrozen = false;
+#if UE_GAMEPLAY_CAMERAS_DEBUG
+		DebugFreezeReason = TEXT("no player controller on active context");
+#endif
+		return;
+	}
+
+	// Get the pawn rotation for the debug information.
+#if UE_GAMEPLAY_CAMERAS_DEBUG
+	{
+		APawn* PlayerPawn = PlayerController->GetPawn();
+		if (PlayerPawn)
+		{
+			DebugPawnTransform = PlayerPawn->GetActorTransform();
+		}
+	}
+#endif
+
+	// Update our current control rotation, and apply it if allowed.
+	if (!AxisBindings.IsEmpty())
+	{
+		UpdateControlRotation(Params, OutResult);
+	}
+	else
+	{
+		bIsFrozen = false;
+#if UE_GAMEPLAY_CAMERAS_DEBUG
+		DebugFreezeReason = TEXT("no input bindings defined");
+#endif
+	}
+	if (ServiceParams.bApplyControlRotation)
+	{
+		PlayerController->SetControlRotation(CurrentControlRotation);
+	}
+#if UE_GAMEPLAY_CAMERAS_DEBUG
+	bDebugDidApplyControlRotation = ServiceParams.bApplyControlRotation;
+#endif
+}
+
+void FPlayerControlRotationEvaluationService::UpdateControlRotation(const FCameraEvaluationServiceUpdateParams& Params, FCameraEvaluationServiceUpdateResult& OutResult)
+{
+	// If we were not already frozen, see if a camera rig is requesting it.
+	const FBuiltInCameraVariables& BuiltInVariables = FBuiltInCameraVariables::Get();
+	FCameraVariableTable& OutVariableTable = OutResult.EvaluationResult.VariableTable;
+	if (!bIsFrozen)
+	{
+		const bool bFreezeControlRotation = OutVariableTable.GetValue<bool>(
+				BuiltInVariables.FreezeControlRotationDefinition, false);
+		const bool bHasCustomControlRotation = OutVariableTable.TryGetValue<FRotator3d>(
+				BuiltInVariables.ControlRotationDefinition, FrozenControlRotation);
+
+		bIsFrozen = (bFreezeControlRotation && bHasCustomControlRotation);
+	}
+
+	// See if the player is actively using the controls.
+	FVector2d MaxAxisBindingValue(EForceInit::ForceInit);
+	double MaxAxisBindingValueSquaredLength = 0.0;
+	for (FEnhancedInputActionValueBinding* AxisBinding : AxisBindings)
+	{
+		const FVector2d Value = AxisBinding->GetValue().Get<FVector2D>();
+		const double ValueSquaredLength = Value.SquaredLength();
+		if (ValueSquaredLength > MaxAxisBindingValueSquaredLength)
+		{
+			MaxAxisBindingValue = Value;
+			MaxAxisBindingValueSquaredLength = ValueSquaredLength;
+		}
+	}
+
+	// Compute input direction speed change.
+	const FVector2d PreviousAxisDir = PreviousAxisBindingValue.GetSafeNormal();
+	const FVector2d CurrentAxisDir = MaxAxisBindingValue.GetSafeNormal();
+	const double AxisActionAngle = (!PreviousAxisDir.IsZero() && !CurrentAxisDir.IsZero()) ?
+		FMath::RadiansToDegrees(FMath::Acos(PreviousAxisDir.Dot(CurrentAxisDir))) :
+		0.0;
+	const double AxisActionAngularSpeed = AxisActionAngle / (Params.DeltaTime > 0 ? Params.DeltaTime : 1.f);
+	PreviousAxisBindingValue = MaxAxisBindingValue;
+
+#if UE_GAMEPLAY_CAMERAS_DEBUG
+	double DebugHighestRecentAxisBindingAngleSpeed = 0.0;
+	if (Params.DeltaTime > 0)
+	{
+		DebugAxisBindingAngleSpeedHistory.Add({ AxisActionAngularSpeed, Params.DeltaTime });
+
+		double HighestValue = 0.f;
+		float AccumulatedDeltaTime = 0.f;
+		for (int32 Index = DebugAxisBindingAngleSpeedHistory.Num() - 1; Index >= 0; --Index)
+		{
+			AccumulatedDeltaTime += DebugAxisBindingAngleSpeedHistory[Index].DeltaTime;
+			HighestValue = FMath::Max(HighestValue, DebugAxisBindingAngleSpeedHistory[Index].Value);
+			if (AccumulatedDeltaTime > GGameplayCamerasControlRotationDebugAxisBindingAngleSpeedHistory)
+			{
+				DebugAxisBindingAngleSpeedHistory.RemoveAt(0, Index);
+				break;
+			}
+		}
+		DebugHighestRecentAxisBindingAngleSpeed = HighestValue;
+	}
+#endif
+
+	// If we are not feezing the control rotation, we are done.
+	// (our CurrentControlRotation is already set to the camera rotation)
+	if (!bIsFrozen)
+	{
+#if UE_GAMEPLAY_CAMERAS_DEBUG
+		DebugFreezeReason = FString::Printf(
+				TEXT("no freeze request (input speed %7.2fdeg/s, recent max %7.2fdeg/s)"),
+				AxisActionAngularSpeed, DebugHighestRecentAxisBindingAngleSpeed);
+#endif
+		return;
+	}
+
+	// If the player isn't using the controls, we can immediately unfreeze the control rotation,
+	// and we are done. (our CurrentControlRotation is already set to the camera rotation)
+	const double SquaredThreshold = FMath::Square(ServiceParams.AxisActionMagnitudeThreshold);
+	if (MaxAxisBindingValueSquaredLength <= SquaredThreshold)
+	{
+		bIsFrozen = false;
+#if UE_GAMEPLAY_CAMERAS_DEBUG
+		DebugFreezeReason = FString::Printf(
+				TEXT("no input (magnitude squared %7.2f < %7.f)"),
+				MaxAxisBindingValueSquaredLength, SquaredThreshold);
+#endif
+		return;
+	}
+
+	// If the player is using the controls, but is turning more than our defined threshold, immediately
+	// unfreeze the control rotation, and we are done.
+	// (our CurrentControlRotation is already set to the camera rotation)
+	if (AxisActionAngularSpeed >= ServiceParams.AxisActionAngularSpeedThreshold)
+	{
+		bIsFrozen = false;
+#if UE_GAMEPLAY_CAMERAS_DEBUG
+		DebugFreezeReason = FString::Printf(
+				TEXT("changed input (%7.2fdeg/s > %7.2fdeg/s, recent max %7.2fdeg/s)"),
+				AxisActionAngularSpeed, ServiceParams.AxisActionAngularSpeedThreshold, DebugHighestRecentAxisBindingAngleSpeed);
+#endif
+		return;
+	}
+
+	// We keep the CurrentControlRotation frozen for one more frame.
+	ensure(bIsFrozen);
+	CurrentControlRotation = FrozenControlRotation;
+#if UE_GAMEPLAY_CAMERAS_DEBUG
+	DebugFreezeReason = FString::Printf(
+			TEXT("unchanged input (%7.2fdeg/s < %7.2fdeg/s, recent max %7.2fdeg/s)"), 
+			AxisActionAngularSpeed, ServiceParams.AxisActionAngularSpeedThreshold, DebugHighestRecentAxisBindingAngleSpeed);
+#endif
+}
+
+#if UE_GAMEPLAY_CAMERAS_DEBUG
+
+void FPlayerControlRotationEvaluationService::OnBuildDebugBlocks(const FCameraDebugBlockBuildParams& Params, FCameraDebugBlockBuilder& Builder)
+{
+	FPlayerControlRotationDebugBlock& DebugBlock = Builder.AttachDebugBlock<FPlayerControlRotationDebugBlock>();
+
+	DebugBlock.PawnTransform = DebugPawnTransform;
+	DebugBlock.ControlRotation = CurrentControlRotation;
+	DebugBlock.CameraRotation = CameraRotation;
+	DebugBlock.bIsFrozen = bIsFrozen;
+	DebugBlock.FreezeReason = DebugFreezeReason;
+}
+
+void FPlayerControlRotationDebugBlock::OnDebugDraw(const FCameraDebugBlockDrawParams& Params, FCameraDebugRenderer& Renderer)
+{
+	const FVector3d PawnLocation = PawnTransform.GetLocation();
+
+	const FRotator3d CameraYaw(0, CameraRotation.Yaw, 0);
+	const FRotator3d ControlYaw(0, ControlRotation.Yaw, 0);
+
+	const FVector3d ForwardArrowEnd = FVector3d::ForwardVector * GGameplayCamerasControlRotationDebugArrowLength;
+
+	// Pawn orientation
+	Renderer.DrawDirectionalArrow(
+			PawnLocation, PawnLocation + PawnTransform.TransformVectorNoScale(ForwardArrowEnd),
+			5.f, FColorList::MandarianOrange, 1.f);
+
+	// Camera rotation
+	Renderer.DrawDirectionalArrow(
+			PawnLocation, PawnLocation + CameraYaw.RotateVector(ForwardArrowEnd),
+			5.f, FColorList::PaleGreen, 1.f);
+
+	// Control rotation
+	FVector3d ControlYawArrowEnd = PawnLocation + ControlYaw.RotateVector(ForwardArrowEnd);
+	Renderer.DrawDirectionalArrow(
+			PawnLocation, ControlYawArrowEnd,
+			5.f, FColorList::Green, 2.f);
+
+	const FColor TextColor(255, 255, 255, 192);
+	const FVector2d TextOffset(5.f, 5.f);
+	if (bIsFrozen)
+	{
+		FString DebugText = FString::Printf(
+				TEXT("camera: %+7.2f, frozen: %+7.2f\n%s"), 
+				CameraRotation.Yaw, ControlRotation.Yaw, *FreezeReason);
+		Renderer.DrawText(ControlYawArrowEnd, TextOffset, DebugText, TextColor);
+	}
+	else
+	{
+		FString DebugText = FString::Printf(
+				TEXT("camera: %+7.2f\n%s"), 
+				CameraRotation.Yaw, *FreezeReason);
+		Renderer.DrawText(ControlYawArrowEnd, TextOffset, DebugText, TextColor);
+	}
+}
+
+#endif  // UE_GAMEPLAY_CAMERAS_DEBUG
+
+}  // namespace UE::Cameras
+
diff --git a/Engine/Plugins/Cameras/GameplayCameras/Source/GameplayCameras/Public/Core/BuiltInCameraVariables.h b/Engine/Plugins/Cameras/GameplayCameras/Source/GameplayCameras/Public/Core/BuiltInCameraVariables.h
--- a/Engine/Plugins/Cameras/GameplayCameras/Source/GameplayCameras/Public/Core/BuiltInCameraVariables.h
+++ b/Engine/Plugins/Cameras/GameplayCameras/Source/GameplayCameras/Public/Core/BuiltInCameraVariables.h
@@ -33,6 +33,13 @@ enum class EBuiltInVector2dCameraVariable
 	YawPitch
 };
 
+UENUM()
+enum class EBuiltInRotator3dCameraVariable
+{
+	None,
+	ControlRotation
+};
+
 namespace UE::Cameras
 {
 
@@ -50,6 +57,8 @@ public:
 	const FCameraVariableDefinition& GetDefinition(EBuiltInDoubleCameraVariable BuiltInVariable) const;
 	/** Get the definition of a built-in camera variable. */
 	const FCameraVariableDefinition& GetDefinition(EBuiltInVector2dCameraVariable BuiltInVariable) const;
+	/** Get the definition of a built-in camera variable. */
+	const FCameraVariableDefinition& GetDefinition(EBuiltInRotator3dCameraVariable BuiltInVariable) const;
 
 public:
 
@@ -60,6 +69,9 @@ public:
 
 	FCameraVariableDefinition YawPitchDefinition;
 
+	FCameraVariableDefinition FreezeControlRotationDefinition;
+	FCameraVariableDefinition ControlRotationDefinition;
+
 protected:
 
 	FBuiltInCameraVariables();
diff --git a/Engine/Plugins/Cameras/GameplayCameras/Source/GameplayCameras/Public/GameFramework/GameplayCameraComponent.h b/Engine/Plugins/Cameras/GameplayCameras/Source/GameplayCameras/Public/GameFramework/GameplayCameraComponent.h
--- a/Engine/Plugins/Cameras/GameplayCameras/Source/GameplayCameras/Public/GameFramework/GameplayCameraComponent.h
+++ b/Engine/Plugins/Cameras/GameplayCameras/Source/GameplayCameras/Public/GameFramework/GameplayCameraComponent.h
@@ -30,8 +30,15 @@ class UGameplayCameraComponent : public USceneComponent
 
 public:
 
+	/** Create a new camera component. */
 	UGameplayCameraComponent(const FObjectInitializer& ObjectInit);
 
+	/** Get the camera evaluation context used by this component. */
+	TSharedPtr<UE::Cameras::FCameraEvaluationContext> GetEvaluationContext();
+
+	/** Get the player controller index that this component has been activated for. */
+	int32 GetPlayerIndex() const { return ActivatedForPlayerIndex; }
+
 public:
 
 	/** 
diff --git a/Engine/Plugins/Cameras/GameplayCameras/Source/GameplayCameras/Public/GameFramework/GameplayCameraSystemComponent.h b/Engine/Plugins/Cameras/GameplayCameras/Source/GameplayCameras/Public/GameFramework/GameplayCameraSystemComponent.h
--- a/Engine/Plugins/Cameras/GameplayCameras/Source/GameplayCameras/Public/GameFramework/GameplayCameraSystemComponent.h
+++ b/Engine/Plugins/Cameras/GameplayCameras/Source/GameplayCameras/Public/GameFramework/GameplayCameraSystemComponent.h
@@ -89,13 +89,6 @@ public:
 	UPROPERTY(EditAnywhere, Category=Camera)
 	TEnumAsByte<EAutoReceiveInput::Type> AutoActivateForPlayer;
 
-	/**
-	 * If enabled, sets the evaluated camera orientation as the player controller rotation every frame.
-	 * This is set on the player controller that this component was activated for.
-	 */
-	UPROPERTY(EditAnywhere, Category=Camera)
-	bool bSetPlayerControllerRotation = true;
-
 private:
 	
 	TSharedPtr<FCameraSystemEvaluator> Evaluator;
diff --git a/Engine/Plugins/Cameras/GameplayCameras/Source/GameplayCameras/Public/GameFramework/GameplayControlRotationComponent.h b/Engine/Plugins/Cameras/GameplayCameras/Source/GameplayCameras/Public/GameFramework/GameplayControlRotationComponent.h
new file mode 100644
--- /dev/null
+++ b/Engine/Plugins/Cameras/GameplayCameras/Source/GameplayCameras/Public/GameFramework/GameplayControlRotationComponent.h
@@ -0,0 +1,71 @@
+// Copyright Epic Games, Inc. All Rights Reserved.
+
+#pragma once
+
+#include "Components/ActorComponent.h"
+#include "Core/CameraSystemEvaluator.h"
+
+#include "GameplayControlRotationComponent.generated.h"
+
+class APlayerController;
+class UCanvas;
+class UInputAction;
+
+namespace UE::Cameras
+{
+	class FCameraSystemEvaluator; 
+	class FPlayerControlRotationEvaluationService;
+}
+
+/**
+ * An example component that works with the GameplayCameraComponent to manage a player's
+ * control rotation when the camera changes or moves in a way that was not initiated 
+ * by the player themselves.
+ */
+UCLASS(BlueprintType, MinimalAPI, ClassGroup=Camera, HideCategories=(Mobility, Rendering, LOD), meta=(BlueprintSpawnableComponent))
+class UGameplayControlRotationComponent : public UActorComponent
+{
+	GENERATED_BODY()
+
+public:
+
+	/** Create a new control rotation component. */
+	UGameplayControlRotationComponent(const FObjectInitializer& ObjectInit);
+
+public:
+
+	// UActorComponent interface
+	virtual void BeginPlay() override;
+	virtual void EndPlay(const EEndPlayReason::Type EndPlayReason) override;
+	virtual void TickComponent(float DeltaTime, ELevelTick TickType, FActorComponentTickFunction* ThisTickFunction) override;
+
+private:
+
+	APlayerController* GetPlayerController();
+
+public:
+
+	/** The axis input action(s) to read from. */
+	UPROPERTY(EditAnywhere, Category="Input")
+	TArray<TObjectPtr<UInputAction>> AxisActions;
+
+	/**
+	 * The angular speed, in degrees per second, past which a change in the player input
+	 * will thaw a frozen control rotation.
+	 */
+	UPROPERTY(EditAnywhere, Category="Input")
+	float AxisActionAngularSpeedThreshold = 20.f;
+
+	/**
+	 * The player input magnitude under which the frozen control rotation is thawed.
+	 */
+	UPROPERTY(EditAnywhere, Category="Input")
+	float AxisActionMagnitudeThreshold = 0.1f;
+
+private:
+
+	TSharedPtr<UE::Cameras::FPlayerControlRotationEvaluationService> ControlRotationService;
+	TWeakPtr<UE::Cameras::FCameraSystemEvaluator> WeakCameraSystem;
+	int32 PlayerIndex = INDEX_NONE;
+};
+
diff --git a/Engine/Plugins/Cameras/GameplayCameras/Source/GameplayCameras/Public/Nodes/Input/AutoRotateInput2DCameraNode.h b/Engine/Plugins/Cameras/GameplayCameras/Source/GameplayCameras/Public/Nodes/Input/AutoRotateInput2DCameraNode.h
--- a/Engine/Plugins/Cameras/GameplayCameras/Source/GameplayCameras/Public/Nodes/Input/AutoRotateInput2DCameraNode.h
+++ b/Engine/Plugins/Cameras/GameplayCameras/Source/GameplayCameras/Public/Nodes/Input/AutoRotateInput2DCameraNode.h
@@ -39,7 +39,7 @@ public:
 	UPROPERTY(EditAnywhere, Category="Auto-Rotate")
 	FFloatCameraParameter WaitTime = 1.f;
 
-	/** The minimum manual rotation, in degrees, to deactivate auto-rtation. */
+	/** The minimum player-induced/manual rotation, in degrees, to deactivate auto-rotation. */
 	UPROPERTY(EditAnywhere, Category="Auto-Rotate")
 	FFloatCameraParameter DeactivationThreshold = 0.01f;
 
@@ -47,10 +47,22 @@ public:
 	UPROPERTY(EditAnywhere, Category="Auto-Rotate")
 	TObjectPtr<UCameraValueInterpolator> Interpolator;
 
-	/** Whether to enable auto-rotation. */
+	/** Whether to suggest freezing the input control rotation. */
 	UPROPERTY(EditAnywhere, Category="Auto-Rotate")
+	FBooleanCameraParameter FreezeControlRotation = false;
+
+	/** Whether to enable auto-rotation. */
+	UPROPERTY(EditAnywhere, Category="Auto-Rotate Toggle")
 	FBooleanCameraParameter EnableAutoRotate = true;
 
+	/** Whether to auto-rotate yaw. */
+	UPROPERTY(EditAnywhere, Category="Auto-Rotate Toggle")
+	FBooleanCameraParameter AutoRotateYaw = true;
+
+	/** Whether to auto-rotate pitch. */
+	UPROPERTY(EditAnywhere, Category="Auto-Rotate Toggle")
+	FBooleanCameraParameter AutoRotatePitch = true;
+
 	/** The underlying input node. */
 	UPROPERTY()
 	TObjectPtr<UInput2DCameraNode> InputNode;
@@ -63,6 +75,7 @@ protected:
 
 	// UCameraNode interface.
 	virtual FCameraNodeChildrenView OnGetChildren() override;
+	virtual void OnBuild(FCameraRigBuildContext& BuildContext) override;
 	virtual FCameraNodeEvaluatorPtr OnBuildEvaluator(FCameraNodeEvaluatorBuilder& Builder) const override;
 };
 
diff --git a/Engine/Plugins/Cameras/GameplayCameras/Source/GameplayCameras/Public/Services/PlayerControlRotationService.h b/Engine/Plugins/Cameras/GameplayCameras/Source/GameplayCameras/Public/Services/PlayerControlRotationService.h
new file mode 100644
--- /dev/null
+++ b/Engine/Plugins/Cameras/GameplayCameras/Source/GameplayCameras/Public/Services/PlayerControlRotationService.h
@@ -0,0 +1,111 @@
+// Copyright Epic Games, Inc. All Rights Reserved.
+
+#pragma once
+
+#include "Core/CameraEvaluationService.h"
+#include "Math/Rotator.h"
+#include "Math/Transform.h"
+#include "Math/Vector2D.h"
+#include "Templates/SharedPointerFwd.h"
+#include "UObject/WeakObjectPtr.h"
+
+class APlayerController;
+class UEnhancedInputComponent;
+class UInputAction;
+struct FEnhancedInputActionValueBinding;
+
+namespace UE::Cameras
+{
+
+class FCameraEvaluationContext;
+class FCameraSystemEvaluator;
+
+/**
+ * Parameter structure for the player control rotation service.
+ */
+struct FPlayerControlRotationParams
+{
+	/** The input magnitude below which we can change control rotation. */
+	double AxisActionMagnitudeThreshold = 0;
+	/** The input angular change speed over which we can change control rotation. */
+	double AxisActionAngularSpeedThreshold = 0;
+	/** Whether the service should set control rotation on the active context's player controller. */
+	bool bApplyControlRotation = true;
+	/** Input actions representing how the player can move their pawn. */
+	TArray<TObjectPtr<UInputAction>> AxisActions;
+};
+
+/**
+ * An evaluation service that manages the player's control rotation based on 
+ * what's going on with cameras.
+ */
+class FPlayerControlRotationEvaluationService : public FCameraEvaluationService
+{
+	UE_DECLARE_CAMERA_EVALUATION_SERVICE(GAMEPLAYCAMERAS_API, FPlayerControlRotationEvaluationService)
+
+public:
+
+	/** Create a new player control rotation service. */
+	GAMEPLAYCAMERAS_API FPlayerControlRotationEvaluationService();
+
+	/** Create a new player control rotation service, setting its parameters immediately. */
+	GAMEPLAYCAMERAS_API FPlayerControlRotationEvaluationService(const FPlayerControlRotationParams& InParams);
+
+public:
+
+	/** Gets the parameters for managing player control rotation. */
+	const FPlayerControlRotationParams& GetParameters() const { return ServiceParams; }
+	/** Sets the parameters for managing player control rotation. */
+	void SetParameters(const FPlayerControlRotationParams& InParams) { ServiceParams = InParams; }
+
+	/** Gets the last evaluated control rotation. */
+	const FRotator3d& GetCurrentControlRotation() const { return CurrentControlRotation; }
+	/** Gets whether the control rotation was last frozen. */
+	bool IsControlRotationFrozen() const { return bIsFrozen; }
+
+protected:
+
+	virtual void OnPostUpdate(const FCameraEvaluationServiceUpdateParams& Params, FCameraEvaluationServiceUpdateResult& OutResult) override;
+	virtual void OnAddReferencedObjects(FReferenceCollector& Collector) override;
+
+#if UE_GAMEPLAY_CAMERAS_DEBUG
+	virtual void OnBuildDebugBlocks(const FCameraDebugBlockBuildParams& Params, FCameraDebugBlockBuilder& Builder) override;
+#endif  // UE_GAMEPLAY_CAMERAS_DEBUG
+
+private:
+
+	void MonitorActiveContext(TSharedPtr<const FCameraEvaluationContext> ActiveContext);
+	void BindActionValues(UEnhancedInputComponent* InputComponent);
+	void UnbindActionValues();
+
+	void UpdateControlRotation(const FCameraEvaluationServiceUpdateParams& Params, FCameraEvaluationServiceUpdateResult& OutResult);
+
+private:
+
+	FPlayerControlRotationParams ServiceParams;
+
+	TWeakObjectPtr<UEnhancedInputComponent> WeakInputComponent;
+	TArray<FEnhancedInputActionValueBinding*> AxisBindings;
+
+	FVector2d PreviousAxisBindingValue;
+	FRotator3d CameraRotation;
+	FRotator3d FrozenControlRotation;
+	FRotator3d CurrentControlRotation;
+	bool bIsFrozen = false;
+
+#if UE_GAMEPLAY_CAMERAS_DEBUG
+	struct FDebugAxisBindingAngleSpeedHistoryEntry
+	{
+		double Value;
+		float DeltaTime;
+	};
+	TArray<FDebugAxisBindingAngleSpeedHistoryEntry> DebugAxisBindingAngleSpeedHistory;
+
+	FTransform DebugPawnTransform;
+	FString DebugFreezeReason;
+	bool bDebugDidApplyControlRotation;
+#endif
+};
+
+}  // namespace UE::Cameras
+
