Cameras: add debug clocks and graphs to the debug info framework

- Clocks show a circle with an arrow representing a 2D vector or an angle.
- Graphs show one or more graph lines that extend back in time up to a certain limit.
- They all get displayed on the right side of the screen inside rectangular "cards".

diff --git a/Engine/Plugins/Cameras/GameplayCameras/Source/GameplayCameras/Private/Debug/CameraDebugClock.cpp b/Engine/Plugins/Cameras/GameplayCameras/Source/GameplayCameras/Private/Debug/CameraDebugClock.cpp
new file mode 100644
--- /dev/null
+++ b/Engine/Plugins/Cameras/GameplayCameras/Source/GameplayCameras/Private/Debug/CameraDebugClock.cpp
@@ -0,0 +1,219 @@
+// Copyright Epic Games, Inc. All Rights Reserved.
+
+#include "Debug/CameraDebugClock.h"
+
+#if UE_GAMEPLAY_CAMERAS_DEBUG
+
+#include "Debug/CameraDebugColors.h"
+#include "Engine/Canvas.h"
+#include "Engine/Engine.h"
+#include "Engine/Font.h"
+#include "HAL/IConsoleManager.h"
+
+#define LOCTEXT_NAMESPACE "CameraDebugClock"
+
+namespace UE::Cameras
+{
+	
+float GGameplayCamerasDebugClockPadding = 10.f;
+FAutoConsoleVariableRef CVarGameplayCamerasDebugClockPadding(
+		TEXT("GameplayCameras.DebugClock.Margin"),
+		GGameplayCamerasDebugClockPadding,
+		TEXT("Default: 10px. The uniform padding inside the debug clock card."));
+
+float GGameplayCamerasDebugClockArrowThickness = 3.f;
+FAutoConsoleVariableRef CVarGameplayCamerasDebugClockArrowThickness(
+		TEXT("GameplayCameras.DebugClock.ArrowThickness"),
+		GGameplayCamerasDebugClockArrowThickness,
+		TEXT("Default: 3px. The thickness of the clock arrow."));
+
+extern float GGameplayCamerasDebugBackgroundOpacity;
+
+FCameraDebugClockDrawParams::FCameraDebugClockDrawParams()
+{
+	const FCameraDebugColors& ColorScheme = FCameraDebugColors::Get();
+	ClockBackgroundColor = ColorScheme.Background.WithAlpha((uint8)(GGameplayCamerasDebugBackgroundOpacity * 255));
+	ClockNameColor = ColorScheme.Title;
+	ClockFaceColor = ColorScheme.Passive;
+	ClockValueLineColor = ColorScheme.Warning;
+}
+
+namespace Internal
+{
+
+FCameraDebugClockRenderer::FCameraDebugClockRenderer(FCanvas* InCanvas, const FCameraDebugClockDrawParams& InDrawParams)
+	: Canvas(InCanvas)
+	, DrawParams(InDrawParams)
+{
+}
+
+void FCameraDebugClockRenderer::DrawVectorClock(const FVector2d& Value, double MaxLength) const
+{
+	DrawFrame();
+
+	DrawCurrentValue(FText::Format(LOCTEXT("CurrentValueFmt", "{0}"), FText::FromString(Value.ToString())));
+
+	double ClockRadius;
+	FVector2f ClockCenter;
+	GetClockFaceParams(ClockCenter, ClockRadius);
+
+	const double ValueToPixels = MaxLength != 0 ? ClockRadius / MaxLength : 1.f;
+	FVector2f ScreenValue(Value * ValueToPixels);
+	ScreenValue.Y = -ScreenValue.Y;  // Flip Y for UI-space coordinates.
+
+	FCanvasLineItem ValueLineItem(FVector2D(ClockCenter), FVector2D(ClockCenter + ScreenValue));
+	ValueLineItem.SetColor(DrawParams.ClockValueLineColor);
+	ValueLineItem.LineThickness = GGameplayCamerasDebugClockArrowThickness;
+	Canvas->DrawItem(ValueLineItem);
+}
+
+void FCameraDebugClockRenderer::DrawAngleClock(double Angle) const
+{
+	DrawFrame();
+
+	DrawCurrentValue(FText::Format(LOCTEXT("CurrentValueFmt", "{0}"), Angle));
+
+	double ClockRadius;
+	FVector2f ClockCenter;
+	GetClockFaceParams(ClockCenter, ClockRadius);
+
+	FVector2f ScreenValue(ClockRadius * FMath::Cos(Angle), ClockRadius * FMath::Sin(Angle));
+	ScreenValue.Y = -ScreenValue.Y;  // Flip Y for UI-space coordinates.
+
+	FCanvasLineItem ValueLineItem(FVector2D(ClockCenter), FVector2D(ClockCenter + ScreenValue));
+	ValueLineItem.SetColor(DrawParams.ClockValueLineColor);
+	ValueLineItem.LineThickness = 3.f;
+	Canvas->DrawItem(ValueLineItem);
+}
+
+void FCameraDebugClockRenderer::DrawFrame() const
+{
+	// Draw the background tile.
+	{
+		FCanvasTileItem TileItem(
+				FVector2D(DrawParams.ClockPosition),
+				FVector2D(DrawParams.ClockSize),
+				DrawParams.ClockBackgroundColor);
+		TileItem.BlendMode = SE_BLEND_Translucent;
+		Canvas->DrawItem(TileItem);
+	}
+
+	// Draw the clock name.
+	if (!DrawParams.ClockName.IsEmpty())
+	{
+		UFont* SmallFont = GEngine->GetSmallFont();
+		const float MaxSmallFontCharHeight = SmallFont->GetMaxCharHeight();
+
+		const FVector2f ClockNamePosition(DrawParams.ClockPosition + 
+				FVector2f(
+					GGameplayCamerasDebugClockPadding, 
+					DrawParams.ClockSize.Y - GGameplayCamerasDebugClockPadding - MaxSmallFontCharHeight));
+		FCanvasTextItem ClockNameItem(
+				FVector2D(ClockNamePosition), 
+				DrawParams.ClockName, 
+				SmallFont, 
+				DrawParams.ClockNameColor);
+		Canvas->DrawItem(ClockNameItem);
+	}
+
+	// Draw the clock face.
+	double ClockRadius;
+	FVector2f ClockCenter;
+	GetClockFaceParams(ClockCenter, ClockRadius);
+	{
+		const int32 NumSides = FMath::Max(20, (int)ClockRadius / 25);
+		const float	AngleDelta = 2.0f * UE_PI / NumSides;
+		const FVector2f AxisX(1.f, 0.f);
+		const FVector2f AxisY(0.f, -1.f);
+		FVector2f LastVertex = ClockCenter + AxisX * ClockRadius;
+
+		for (int32 SideIndex = 0; SideIndex < NumSides; SideIndex++)
+		{
+			const float CurAngle = AngleDelta * (SideIndex + 1);
+			const FVector2f Vertex = ClockCenter + (AxisX * FMath::Cos(CurAngle) + AxisY * FMath::Sin(CurAngle)) * ClockRadius;
+
+			FCanvasLineItem LineItem{ FVector2D(LastVertex), FVector2D(Vertex) };
+			LineItem.SetColor(DrawParams.ClockFaceColor);
+			Canvas->DrawItem(LineItem);
+
+			LastVertex = Vertex;
+		}
+	}
+}
+
+void FCameraDebugClockRenderer::DrawCurrentValue(const FText& CurrentValueStr) const
+{
+	UFont* TinyFont = GEngine->GetTinyFont();
+	const float MaxTinyFontCharHeight = TinyFont->GetMaxCharHeight();
+
+	const FVector2f CurrentValuePosition = DrawParams.ClockPosition + FVector2f(GGameplayCamerasDebugClockPadding);
+	FCanvasTextItem TextItem(
+			FVector2D(CurrentValuePosition),
+			CurrentValueStr,
+			TinyFont,
+			DrawParams.ClockValueLineColor);
+	Canvas->DrawItem(TextItem);
+}
+
+void FCameraDebugClockRenderer::GetClockFaceParams(FVector2f& OutClockCenter, double& OutClockRadius) const
+{
+	UFont* SmallFont = GEngine->GetSmallFont();
+	const float MaxSmallFontCharHeight = SmallFont->GetMaxCharHeight();
+	const double ClockAreaHeight = DrawParams.ClockSize.Y - 3 * GGameplayCamerasDebugClockPadding - MaxSmallFontCharHeight;
+
+	OutClockRadius = (FMath::Min(DrawParams.ClockSize.X, DrawParams.ClockSize.Y) - 2 * GGameplayCamerasDebugClockPadding) / 2.f;
+	OutClockCenter = DrawParams.ClockPosition + FVector2f(
+			GGameplayCamerasDebugClockPadding + OutClockRadius,
+			GGameplayCamerasDebugClockPadding + ClockAreaHeight / 2.f);
+}
+
+}  // namespace Internal
+
+void FCameraDebugClock::Update(double InAngle)
+{
+	Value = FVariant(TInPlaceType<FAngleValue>(), FAngleValue{ InAngle });
+}
+
+void FCameraDebugClock::Update(const FVector2d& InValue)
+{
+	if (Value.IsType<FVectorValue>())
+	{
+		const FVectorValue PreviousValue = Value.Get<FVectorValue>();
+		const double NextMaxLength = FMath::Max(PreviousValue.CurrentMaxLength, InValue.Length());
+		Value = FVariant(TInPlaceType<FVectorValue>(), FVectorValue{ InValue, NextMaxLength });
+	}
+}
+
+void FCameraDebugClock::Draw(FCanvas* Canvas, const FCameraDebugClockDrawParams& DrawParams)
+{
+	using namespace Internal;
+
+	FCameraDebugClockRenderer Renderer(Canvas, DrawParams);
+	switch (Value.GetIndex())
+	{
+		case FVariant::IndexOfType<FVectorValue>():
+			{
+				const FVectorValue& VectorValue = Value.Get<FVectorValue>();
+				Renderer.DrawVectorClock(VectorValue.Vector, VectorValue.CurrentMaxLength);
+			}
+			break;
+		case FVariant::IndexOfType<FAngleValue>():
+			{
+				const FAngleValue& AngleValue = Value.Get<FAngleValue>();
+				Renderer.DrawAngleClock(AngleValue.Angle);
+			}
+			break;
+	}
+}
+
+void FCameraDebugClock::Serialize(FArchive& Ar)
+{
+	Ar << Value;
+}
+
+}  // namespace UE::Cameras
+
+#undef LOCTEXT_NAMESPACE
+
+#endif
+
diff --git a/Engine/Plugins/Cameras/GameplayCameras/Source/GameplayCameras/Private/Debug/CameraDebugGraph.cpp b/Engine/Plugins/Cameras/GameplayCameras/Source/GameplayCameras/Private/Debug/CameraDebugGraph.cpp
new file mode 100644
--- /dev/null
+++ b/Engine/Plugins/Cameras/GameplayCameras/Source/GameplayCameras/Private/Debug/CameraDebugGraph.cpp
@@ -0,0 +1,191 @@
+// Copyright Epic Games, Inc. All Rights Reserved.
+
+#include "Debug/CameraDebugGraph.h"
+
+#if UE_GAMEPLAY_CAMERAS_DEBUG
+
+#include "Engine/Canvas.h"
+#include "Engine/Engine.h"
+#include "Engine/Font.h"
+#include "HAL/IConsoleManager.h"
+
+#define LOCTEXT_NAMESPACE "CameraDebugGraph"
+
+namespace UE::Cameras
+{
+
+float GGameplayCamerasDebugGraphPadding = 10.f;
+FAutoConsoleVariableRef CVarGameplayCamerasDebugGraphPadding(
+		TEXT("GameplayCameras.DebugGraph.Padding"),
+		GGameplayCamerasDebugGraphPadding,
+		TEXT("Default: 10px. The uniform padding inside the debug graph card."));
+
+float GGameplayCamerasDebugGraphDefaultHistoryTime = 3.f;
+FAutoConsoleVariableRef CVarGameplayCamerasDebugGraphDefaultHistoryTime(
+		TEXT("GameplayCameras.DebugGraph.DefaultHistoryTime"),
+		GGameplayCamerasDebugGraphDefaultHistoryTime,
+		TEXT("Default: 2s. The default amount of the time in seconds displayed by a debug graph."));
+
+FCameraDebugGraphDrawParams::FCameraDebugGraphDrawParams()
+	: HistoryTime(0.f)
+{
+	const FCameraDebugColors& ColorScheme = FCameraDebugColors::Get();
+	GraphBackgroundColor = ColorScheme.Background.WithAlpha((uint8)(GGameplayCamerasDebugBackgroundOpacity * 255));
+	GraphNameColor = ColorScheme.Title;
+	GraphLineColors = { ColorScheme.Notice };
+	// Leave history time at zero so that we can dynamically adjust the history time with the CVar.
+}
+
+float FCameraDebugGraphDrawParams::GetDefaultMaxHistoryTime()
+{
+	return GGameplayCamerasDebugGraphDefaultHistoryTime;
+}
+
+namespace Internal
+{
+
+FCameraDebugGraphRenderer::FCameraDebugGraphRenderer(FCanvas* InCanvas, const FCameraDebugGraphDrawParams& InDrawParams)
+	: Canvas(InCanvas)
+	, DrawParams(InDrawParams)
+{
+}
+	
+void FCameraDebugGraphRenderer::DrawEmptyFrame() const
+{
+	DrawFrameImpl();
+}
+
+void FCameraDebugGraphRenderer::DrawFrame(TArrayView<float> CurrentValues) const
+{
+	DrawFrameImpl();
+
+	// Draw the current values on top of each other in the top-left corner of the card.
+	UFont* TinyFont = GEngine->GetTinyFont();
+	const float MaxTinyFontCharHeight = TinyFont->GetMaxCharHeight();
+
+	FVector2f CurrentValuePosition = DrawParams.GraphPosition +
+		FVector2f(GGameplayCamerasDebugGraphPadding, GGameplayCamerasDebugGraphPadding);
+	for (int32 Index = 0; Index < CurrentValues.Num(); ++Index)
+	{
+		FCanvasTextItem TextItem(
+				FVector2D(CurrentValuePosition), 
+				FText::Format(LOCTEXT("CurrentValueFmt", "{0}"), CurrentValues[Index]),
+				TinyFont, 
+				DrawParams.GraphLineColors[Index]);
+		Canvas->DrawItem(TextItem);
+
+		CurrentValuePosition.Y += MaxTinyFontCharHeight + 2.f;
+	}
+}
+
+void FCameraDebugGraphRenderer::DrawFrameImpl() const
+{
+	// Draw the card's background tile.
+	{
+		FCanvasTileItem TileItem(
+				FVector2D(DrawParams.GraphPosition), 
+				FVector2D(DrawParams.GraphSize), 
+				DrawParams.GraphBackgroundColor);
+		TileItem.BlendMode = SE_BLEND_Translucent;
+		Canvas->DrawItem(TileItem);
+	}
+
+	// Add the graph name at the bottom of the card.
+	if (!DrawParams.GraphName.IsEmpty())
+	{
+		UFont* SmallFont = GEngine->GetSmallFont();
+		const float MaxSmallFontCharHeight = SmallFont->GetMaxCharHeight();
+
+		const FVector2D GraphNamePosition(
+				DrawParams.GraphPosition + 
+				FVector2f(
+					GGameplayCamerasDebugGraphPadding, 
+					DrawParams.GraphSize.Y - GGameplayCamerasDebugGraphPadding - MaxSmallFontCharHeight));
+		FCanvasTextItem GraphNameItem(GraphNamePosition, DrawParams.GraphName, SmallFont, DrawParams.GraphNameColor);
+		Canvas->DrawItem(GraphNameItem);
+	}
+}
+
+void FCameraDebugGraphRenderer::DrawGraphLine(const FLineDrawParams& LineDrawParams, TStridedView<float> Times, TStridedView<float> Values) const
+{
+	UFont* SmallFont = GEngine->GetSmallFont();
+	const float MaxSmallFontCharHeight = SmallFont->GetMaxCharHeight();
+
+	// Compute the actual area for the graph inside the card, based on the current margins and paddings.
+	const FVector2D GraphAreaSize(
+			DrawParams.GraphSize.X - 2 * GGameplayCamerasDebugGraphPadding,
+			DrawParams.GraphSize.Y - 3 * GGameplayCamerasDebugGraphPadding - MaxSmallFontCharHeight);
+	const FVector2D GraphAreaPosition(DrawParams.GraphPosition + FVector2f(GGameplayCamerasDebugGraphPadding));
+	const double GraphAreaRight = GraphAreaPosition.X + GraphAreaSize.X;
+	const double GraphAreaBottom = GraphAreaPosition.Y + GraphAreaSize.Y;
+
+	// Figure out how much history we're showing, and compute the conversion ratio between seconds
+	// and pixels.
+	float HistoryTime = DrawParams.HistoryTime;
+	if (HistoryTime <= 0)
+	{
+		HistoryTime = FCameraDebugGraphDrawParams::GetDefaultMaxHistoryTime();
+	}
+	const float TimeToPixels = GraphAreaSize.X / HistoryTime;
+	const float LatestTime = Times[Times.Num() - 1];
+
+	// Figure out how much vertical range we're showing, and compute the conversion ratio between
+	// value units and pixels.
+	float ValueRange = LineDrawParams.MaxValue - LineDrawParams.MinValue;
+	if (ValueRange <= 0)
+	{
+		ValueRange = 1.f;
+	}
+	const float ValueToPixels = GraphAreaSize.Y / ValueRange;
+	const float LowestValue = LineDrawParams.MinValue;
+
+	// Draw the lines! We start drawing from the right-side of the card, so that the latest value is
+	// always exactly on the edge of the card's graph area. The oldest values, on the left-side, may
+	// overflow into the padding area, so we have to possibly cut that short.
+	const int32 Num = FMath::Min(Times.Num(), Values.Num());
+	for (int32 Index = Num - 1; Index > 0; --Index)
+	{
+		const float NextTime = Times[Index];
+		const float PrevTime = Times[Index - 1];
+
+		const float NextValue = Values[Index];
+		const float PrevValue = Values[Index - 1];
+
+		const FVector2D NextGraphPoint(
+				GraphAreaRight - (LatestTime - NextTime) * TimeToPixels,
+				GraphAreaBottom + (LowestValue - NextValue) * ValueToPixels);
+		const FVector2D PrevGraphPoint(
+				GraphAreaRight - (LatestTime - PrevTime) * TimeToPixels,
+				GraphAreaBottom + (LowestValue - PrevValue) * ValueToPixels);
+
+		if (PrevGraphPoint.X >= GraphAreaPosition.X)
+		{
+			FCanvasLineItem LineItem(NextGraphPoint, PrevGraphPoint);
+			LineItem.SetColor(LineDrawParams.LineColor);
+			Canvas->DrawItem(LineItem);
+		}
+		else
+		{
+			const float SegmentTimeRange = (NextGraphPoint.X - PrevGraphPoint.X);
+			if (SegmentTimeRange > 0)
+			{
+				const float Factor = (GraphAreaPosition.X - PrevGraphPoint.X) / SegmentTimeRange;
+				const FVector2D EdgePoint = FMath::Lerp(PrevGraphPoint, NextGraphPoint, Factor);
+
+				FCanvasLineItem LineItem(NextGraphPoint, EdgePoint);
+				LineItem.SetColor(LineDrawParams.LineColor);
+				Canvas->DrawItem(LineItem);
+			}
+			break;
+		}
+	}
+}
+
+}  // namespace Internal
+
+}  // namespace UE::Cameras
+
+#undef LOCTEXT_NAMESPACE
+
+#endif
+
diff --git a/Engine/Plugins/Cameras/GameplayCameras/Source/GameplayCameras/Private/Debug/CameraDebugRenderer.cpp b/Engine/Plugins/Cameras/GameplayCameras/Source/GameplayCameras/Private/Debug/CameraDebugRenderer.cpp
--- a/Engine/Plugins/Cameras/GameplayCameras/Source/GameplayCameras/Private/Debug/CameraDebugRenderer.cpp
+++ b/Engine/Plugins/Cameras/GameplayCameras/Source/GameplayCameras/Private/Debug/CameraDebugRenderer.cpp
@@ -6,6 +6,7 @@
 #include "CanvasItem.h"
 #include "CanvasTypes.h"
 #include "Components/LineBatchComponent.h"
+#include "Debug/CameraDebugClock.h"
 #include "Debug/CameraDebugColors.h"
 #include "Debug/DebugTextRenderer.h"
 #include "Engine/Canvas.h"
@@ -33,6 +34,12 @@ static FAutoConsoleVariableRef CVarGameplayCamerasDebugTopMargin(
 	GGameplayCamerasDebugTopMargin,
 	TEXT("(Default: 10px. The top margin for rendering Gameplay Cameras debug text."));
 
+int32 GGameplayCamerasDebugRightMargin = 10;
+static FAutoConsoleVariableRef CVarGameplayCamerasDebugRightMargin(
+	TEXT("GameplayCameras.Debug.RightMargin"),
+	GGameplayCamerasDebugRightMargin,
+	TEXT("(Default: 10px. The right margin for rendering Gameplay Cameras debug text."));
+
 int32 GGameplayCamerasDebugInnerMargin = 5;
 static FAutoConsoleVariableRef CVarGameplayCamerasDebugInnerMargin(
 	TEXT("GameplayCameras.Debug.InnerMargin"),
@@ -49,7 +56,31 @@ int32 GGameplayCamerasDebugBackgroundDepthSortKey = 1;
 static FAutoConsoleVariableRef CVarGameplayCamerasDebugBackgroundDepthSortKey(
 	TEXT("GameplayCameras.Debug.BackgroundDepthSortKey"),
 	GGameplayCamerasDebugBackgroundDepthSortKey,
-	TEXT(""));
+	TEXT("Default: 1. The sort key for drawing the background behind debug text and debug cards."));
+
+int32 GGameplayCamerasDebugCardWidth = 200;
+static FAutoConsoleVariableRef CVarGameplayCamerasDebugCardWidth(
+	TEXT("GameplayCameras.Debug.CardWidth"),
+	GGameplayCamerasDebugCardWidth,
+	TEXT("Default: 200px. The width of the debug cards (e.g. graphs, clocks, etc.)"));
+
+int32 GGameplayCamerasDebugCardHeight = 250;
+static FAutoConsoleVariableRef CVarGameplayCamerasDebugCardHeight(
+	TEXT("GameplayCameras.Debug.CardHeight"),
+	GGameplayCamerasDebugCardHeight,
+	TEXT("Default: 250px. The height of the debug cards (e.g. graphs, clocks, etc.)"));
+
+int32 GGameplayCamerasDebugCardGap = 10;
+static FAutoConsoleVariableRef CVarGameplayCamerasDebugCardGap(
+	TEXT("GameplayCameras.Debug.CardGap"),
+	GGameplayCamerasDebugCardGap,
+	TEXT("Default: 10px. The gap between the debug cards (e.g. graphs, clocks, etc.)"));
+
+int32 GGameplayCamerasDebugMaxCardColumns = 2;
+static FAutoConsoleVariableRef CVarGameplayCamerasDebugMaxCardColumns(
+	TEXT("GameplayCameras.Debug.MaxCardColumns"),
+	GGameplayCamerasDebugMaxCardColumns,
+	TEXT("Default: 2. The number of columns to layout the debug cards (e.g. graphs, clocks, etc.)"));
 
 FCameraDebugRenderer::FCameraDebugRenderer(UWorld* InWorld, UCanvas* InCanvasObject)
 	: World(InWorld)
@@ -60,6 +91,15 @@ FCameraDebugRenderer::FCameraDebugRenderer(UWorld* InWorld, UCanvas* InCanvasObj
 	MaxCharHeight = RenderFont->GetMaxCharHeight();
 
 	NextDrawPosition = FVector2f{ (float)GGameplayCamerasDebugLeftMargin, (float)GGameplayCamerasDebugTopMargin };
+
+	NextCardPosition = FVector2f::ZeroVector;
+	NextCardColumn = 0;
+	if (CanvasObject)
+	{
+		NextCardPosition = FVector2f{ 
+			CanvasObject->SizeX - (float)GGameplayCamerasDebugCardWidth - (float)GGameplayCamerasDebugRightMargin,
+			(float)GGameplayCamerasDebugTopMargin };
+	}
 }
 
 FCameraDebugRenderer::~FCameraDebugRenderer()
@@ -216,6 +256,44 @@ void FCameraDebugRenderer::DrawTextBackgroundTile(float Opacity)
 	}
 }
 
+void FCameraDebugRenderer::DrawClock(FCameraDebugClock& InClock, const FText& InClockName)
+{
+	FCameraDebugClockDrawParams DrawParams;
+	DrawParams.ClockName = InClockName;
+	DrawParams.ClockPosition = GetNextCardPosition();
+	DrawParams.ClockSize = FVector2f(GGameplayCamerasDebugCardWidth, GGameplayCamerasDebugCardHeight);
+	InClock.Draw(GetCanvas(), DrawParams);
+}
+
+FVector2f FCameraDebugRenderer::GetNextCardPosition()
+{
+	const FVector2f Result(NextCardPosition);
+
+	++NextCardColumn;
+	if (NextCardColumn >= GGameplayCamerasDebugMaxCardColumns)
+	{
+		// We went over the number of columns we're supposed to stick to.
+		// Place the next card below the previous cards, at the right-side edge of the canvas.
+		NextCardColumn = 0;
+		NextCardPosition.X = CanvasObject->SizeX - (float)GGameplayCamerasDebugCardWidth - (float)GGameplayCamerasDebugRightMargin;
+		NextCardPosition.Y += GGameplayCamerasDebugCardHeight + GGameplayCamerasDebugCardGap;
+	}
+	else
+	{
+		// We can go to the next column. Place the next card to the left of the previous card.
+		NextCardPosition.X -= (float)GGameplayCamerasDebugCardWidth + (float)GGameplayCamerasDebugCardGap;
+	}
+
+	return Result;
+}
+
+void FCameraDebugRenderer::GetNextDrawGraphParams(FCameraDebugGraphDrawParams& OutDrawParams, const FText& InGraphName)
+{
+	OutDrawParams.GraphName = InGraphName;
+	OutDrawParams.GraphPosition = GetNextCardPosition();
+	OutDrawParams.GraphSize = FVector2f(GGameplayCamerasDebugCardWidth, GGameplayCamerasDebugCardHeight);
+}
+
 void FCameraDebugRenderer::Draw2DLine(const FVector2D& Start, const FVector2D& End, const FLinearColor& LineColor, float LineThickness)
 {
 	if (FCanvas* Canvas = GetCanvas())
diff --git a/Engine/Plugins/Cameras/GameplayCameras/Source/GameplayCameras/Public/Debug/CameraDebugClock.h b/Engine/Plugins/Cameras/GameplayCameras/Source/GameplayCameras/Public/Debug/CameraDebugClock.h
new file mode 100644
--- /dev/null
+++ b/Engine/Plugins/Cameras/GameplayCameras/Source/GameplayCameras/Public/Debug/CameraDebugClock.h
@@ -0,0 +1,142 @@
+// Copyright Epic Games, Inc. All Rights Reserved.
+
+#pragma once
+
+#include "CoreTypes.h"
+#include "GameplayCameras.h"
+#include "Internationalization/Text.h"
+#include "Math/Vector2D.h"
+#include "Misc/TVariant.h"
+#include "Serialization/Archive.h"
+
+#if UE_GAMEPLAY_CAMERAS_DEBUG
+
+class FCanvas;
+
+namespace UE::Cameras
+{
+
+class FCameraDebugClock;
+
+/**
+ * Parameter structure for drawing a debug clock.
+ */
+struct FCameraDebugClockDrawParams
+{
+	FCameraDebugClockDrawParams();
+
+	/** The position of the clock card on screen. */
+	FVector2f ClockPosition;
+	/** The total size of the clock card on screen. */
+	FVector2f ClockSize;
+	/** The color of the clock card's background. */
+	FLinearColor ClockBackgroundColor;
+
+	/** The name of the clock, displayed at the bottom of the card. */
+	FText ClockName;
+	/** The color of the clock name text. */
+	FLinearColor ClockNameColor;
+
+	/** The color of the clock face circle. */
+	FLinearColor ClockFaceColor;
+	/** The color of the arrow inside the clock. */
+	FLinearColor ClockValueLineColor;
+};
+
+namespace Internal
+{
+
+class FCameraDebugClockRenderer
+{
+public:
+
+	FCameraDebugClockRenderer(FCanvas* InCanvas, const FCameraDebugClockDrawParams& InDrawParams);
+
+	void DrawVectorClock(const FVector2d& Value, double MaxLength) const;
+	void DrawAngleClock(double Angle) const;
+
+private:
+
+	void DrawFrame() const;
+	void DrawCurrentValue(const FText& CurrentValueStr) const;
+	void GetClockFaceParams(FVector2f& OutClockCenter, double& OutClockRadius) const;
+
+private:
+
+	FCanvas* Canvas = nullptr;
+	FCameraDebugClockDrawParams DrawParams;
+};
+
+}  // namespace Internal
+
+/**
+ * A debug clock, for showing a real-time angle or 2D vector in a graphical way.
+ */
+class GAMEPLAYCAMERAS_API FCameraDebugClock
+{
+public:
+
+	/**
+	 * Update the clock with the given angle. The clock's arrow will reach the
+	 * edge of the clock face and oriented according to this angle, relative to
+	 * the direction pointing straight upwards on the screen.
+	 */
+	void Update(double InAngle);
+
+	/**
+	 * Update the clock with the given 2D vector. The clock's arrow will represent
+	 * this vector, relative to the direction pointing straight updwards on the
+	 * screen.
+	 */
+	void Update(const FVector2d& InValue);
+
+	/**
+	 * Draw the debug clock onto the given canvas.
+	 */
+	void Draw(FCanvas* Canvas, const FCameraDebugClockDrawParams& DrawParams);
+
+public:
+
+	void Serialize(FArchive& Ar);
+
+	friend FCameraDebugClock& operator<< (FArchive& Ar, FCameraDebugClock& This)
+	{
+		This.Serialize(Ar);
+		return This;
+	}
+
+private:
+
+	struct FVectorValue
+	{
+		FVector2d Vector = FVector2d::ZeroVector;
+		double CurrentMaxLength = 0.f;
+	};
+	struct FAngleValue
+	{
+		double Angle = 0.f;
+	};
+	using FVariant = TVariant<FVectorValue, FAngleValue>;
+	FVariant Value;
+
+	friend FVectorValue& operator<< (FArchive& Ar, FVectorValue& This);
+	friend FAngleValue& operator<< (FArchive& Ar, FAngleValue& This);
+};
+
+inline FCameraDebugClock::FVectorValue& operator<< (FArchive& Ar, FCameraDebugClock::FVectorValue& This)
+{
+	Ar << This.Vector;
+	Ar << This.CurrentMaxLength;
+	return This;
+}
+
+inline FCameraDebugClock::FAngleValue& operator<< (FArchive& Ar, FCameraDebugClock::FAngleValue& This)
+{
+	Ar << This.Angle;
+	return This;
+}
+
+}  // namespace UE::Cameras
+
+#endif
+
diff --git a/Engine/Plugins/Cameras/GameplayCameras/Source/GameplayCameras/Public/Debug/CameraDebugGraph.h b/Engine/Plugins/Cameras/GameplayCameras/Source/GameplayCameras/Public/Debug/CameraDebugGraph.h
new file mode 100644
--- /dev/null
+++ b/Engine/Plugins/Cameras/GameplayCameras/Source/GameplayCameras/Public/Debug/CameraDebugGraph.h
@@ -0,0 +1,345 @@
+// Copyright Epic Games, Inc. All Rights Reserved.
+
+#pragma once
+
+#include "Containers/StridedView.h"
+#include "CoreTypes.h"
+#include "Debug/CameraDebugColors.h"
+#include "GameplayCameras.h"
+#include "Internationalization/Text.h"
+#include "Math/Vector2D.h"
+#include "Serialization/Archive.h"
+
+#if UE_GAMEPLAY_CAMERAS_DEBUG
+
+class FCanvas;
+
+namespace UE::Cameras
+{
+
+extern float GGameplayCamerasDebugBackgroundOpacity;
+
+/**
+ * Parameter structure for drawing a debug graph.
+ */
+struct FCameraDebugGraphDrawParams
+{
+	/** Gets the default maximum time shown on the graph. */
+	static GAMEPLAYCAMERAS_API float GetDefaultMaxHistoryTime();
+
+	FCameraDebugGraphDrawParams();
+
+	/** The position of the graph card on screen. */
+	FVector2f GraphPosition;
+	/** The total size of the graph card on screen. */
+	FVector2f GraphSize;
+	/** The color of the card's background. */
+	FLinearColor GraphBackgroundColor;
+
+	/** The name of the graph, displayed at the bottom of the card. */
+	FText GraphName;
+	/** The color of the graph name text. */
+	FLinearColor GraphNameColor;
+
+	/** How far back in the past the graph goes. */
+	float HistoryTime;
+	/** The colors for the lines of the graph. */
+	TArray<FLinearColor, TInlineAllocator<4>> GraphLineColors;
+
+	/** Setup line colors using the active color scheme. */
+	template<uint8 NumValues>
+	void SetupDefaultLineColors();
+};
+
+template<uint8 NumValues>
+inline void FCameraDebugGraphDrawParams::SetupDefaultLineColors()
+{
+	const FCameraDebugColors& ColorScheme = FCameraDebugColors::Get();
+	GraphLineColors.Reset();
+	GraphLineColors.Add(ColorScheme.Error);
+	GraphLineColors.Add(ColorScheme.Good);
+	GraphLineColors.Add(ColorScheme.Notice);
+	GraphLineColors.Add(ColorScheme.Notice2);
+}
+
+template<>
+inline void FCameraDebugGraphDrawParams::SetupDefaultLineColors<2>()
+{
+	const FCameraDebugColors& ColorScheme = FCameraDebugColors::Get();
+	GraphLineColors.Reset();
+	GraphLineColors.Add(ColorScheme.Good);
+	GraphLineColors.Add(ColorScheme.Notice);
+}
+
+template<>
+inline void FCameraDebugGraphDrawParams::SetupDefaultLineColors<1>()
+{
+	const FCameraDebugColors& ColorScheme = FCameraDebugColors::Get();
+	GraphLineColors.Reset();
+	GraphLineColors.Add(ColorScheme.Notice);
+}
+
+namespace Internal
+{
+
+class FCameraDebugGraphRenderer
+{
+public:
+
+	struct FLineDrawParams
+	{
+		FLinearColor LineColor;
+		float MaxValue;
+		float MinValue;
+	};
+
+	FCameraDebugGraphRenderer(FCanvas* InCanvas, const FCameraDebugGraphDrawParams& InDrawParams);
+
+	void DrawEmptyFrame() const;
+	void DrawFrame(TArrayView<float> CurrentValues) const;
+	void DrawGraphLine(const FLineDrawParams& LineDrawParams, TStridedView<float> Times, TStridedView<float> Values) const;
+
+private:
+
+	void DrawFrameImpl() const;
+
+private:
+
+	FCanvas* Canvas = nullptr;
+	FCameraDebugGraphDrawParams DrawParams;
+};
+
+}   // namespace Internal
+
+/**
+ * An entry on a debug graph, defined by a timestamp and one or more graph values.
+ */
+template<uint8 NumValues>
+struct TCameraDebugGraphEntry
+{
+	/** The absolute time of the entry, from an arbitrary start time. */
+	float Time;
+	/** The values of each line on the owning graph. */
+	float Values[NumValues];
+
+	/** Creates a new uninitialized entry. */
+	TCameraDebugGraphEntry() = default;
+
+	/** Creates new entry, zero-initialized. */
+	TCameraDebugGraphEntry(EForceInit)
+		: Time(0)
+	{
+		for (int32 Index = 0; Index < NumValues; ++Index)
+		{
+			Values[Index] = 0.f;
+		}
+	}
+
+	/** Creates a new entry given an absolute time and some graph values. */
+	template<typename... InValueTypes>
+	TCameraDebugGraphEntry(float InTime, float InFirstValue, InValueTypes... InOtherValues)
+		: Time(InTime)
+	{
+		SetValues<0u>(InFirstValue, InOtherValues...);
+	}
+
+	/** Sets the values on this entry. */
+	template<typename... InValueTypes>
+	void SetValues(float InFirstValue, InValueTypes... InOtherValues)
+	{
+		SetValues<0u>(InFirstValue, InOtherValues...);
+	}
+
+public:
+
+	void Serialize(FArchive& Ar)
+	{
+		Ar << Time;
+		for (int32 Index = 0; Index < NumValues; ++Index)
+		{
+			Ar << Values[Index];
+		}
+	}
+
+	friend TCameraDebugGraphEntry<NumValues>& operator<< (FArchive& Ar, TCameraDebugGraphEntry<NumValues>& This)
+	{
+		This.Serialize(Ar);
+		return This;
+	}
+
+private:
+
+	template<uint8 InValueIndex, typename... InValueTypes>
+	void SetValues(float InNextValue, InValueTypes... InOtherValues)
+	{
+		static_assert(InValueIndex < NumValues, "Incorrect number of values passed");
+		Values[InValueIndex] = InNextValue;
+		SetValues<InValueIndex + 1>(InOtherValues...);
+	}
+	
+	template<uint8 InEndIndex>
+	void SetValues()
+	{
+		static_assert(InEndIndex == NumValues, "Incorrect number of values passed");
+	}
+};
+
+/**
+ * A debug graph, showing one or more lines. The lines progress as new timestamped values
+ * are added to the graph, with older values being discarded when they go past the maximum
+ * history time of the grpah.
+ */
+template<uint8 NumValues>
+class TCameraDebugGraph
+{
+	using FGraphEntry = TCameraDebugGraphEntry<NumValues>;
+
+public:
+
+	/** Creates a new debug graph. */
+	TCameraDebugGraph()
+	{
+		Entries.Reserve(20);
+	}
+
+	/**
+	 * Adds a new entry to the graph, timestamped relative to the last added entry.
+	 *
+	 * @param InDeltaTime  The delta-time elapsed since the last entry was added.
+	 * @param InValues     The values for the new entry.
+	 */
+	template<typename... InValueTypes>
+	void Add(float InDeltaTime, InValueTypes... InValues)
+	{
+		const float Time = GetNewEntryTime(InDeltaTime);
+		Add(FGraphEntry(Time, InValues...));
+	}
+
+	/**
+	 * Draw this debug graph to the given canvas.
+	 */
+	void Draw(FCanvas* InCanvas, const FCameraDebugGraphDrawParams& InDrawParams)
+	{
+		using namespace Internal;
+
+		Update(InDrawParams.HistoryTime);
+
+		FCameraDebugGraphRenderer Renderer(InCanvas, InDrawParams);
+		if (Entries.IsEmpty())
+		{
+			Renderer.DrawEmptyFrame();
+		}
+		else
+		{
+			Renderer.DrawFrame(MakeArrayView(Entries.Last().Values));
+
+			const int32 NumLineColors = InDrawParams.GraphLineColors.Num();
+			TStridedView<float> TimesView = MakeStridedView(Entries, &FGraphEntry::Time);
+			const int32 StrideBetweenValues = sizeof(float) + NumValues * sizeof(float);
+			for (int32 Index = 0; Index < NumValues; ++Index)
+			{
+				const int32 LineColorIndex = Index % NumLineColors;
+				const FLinearColor& LineColor = InDrawParams.GraphLineColors[LineColorIndex];
+				TStridedView<float> ValuesView = MakeStridedView(
+						StrideBetweenValues, 
+						reinterpret_cast<float*>(
+							reinterpret_cast<uint8*>(&Entries[0]) + sizeof(float) + Index * sizeof(float)),
+						Entries.Num());
+
+				FCameraDebugGraphRenderer::FLineDrawParams LineDrawParams;
+				LineDrawParams.LineColor = LineColor;
+				LineDrawParams.MinValue = CurrentMinValue;
+				LineDrawParams.MaxValue = CurrentMaxValue;
+				Renderer.DrawGraphLine(LineDrawParams, TimesView, ValuesView);
+			}
+		}
+	}
+
+public:
+
+	void Serialize(FArchive& Ar)
+	{
+		Ar << CurrentMinValue;
+		Ar << CurrentMaxValue;
+		Ar << Entries;
+	}
+
+	friend TCameraDebugGraph<NumValues>& operator<< (FArchive& Ar, TCameraDebugGraph<NumValues>& This)
+	{
+		This.Serialize(Ar);
+		return This;
+	}
+
+private:
+
+	void Add(const FGraphEntry& InEntry)
+	{
+		Entries.Add(InEntry);
+	}
+
+	float GetNewEntryTime(float InDeltaTime) const
+	{
+		if (!Entries.IsEmpty())
+		{
+			return Entries.Last().Time + InDeltaTime;
+		}
+		return 0.f;
+	}
+
+	void Update(float InMaxHistoryTime)
+	{
+		float MaxHistoryTime = InMaxHistoryTime;
+		if (InMaxHistoryTime <= 0.f)
+		{
+			MaxHistoryTime = FCameraDebugGraphDrawParams::GetDefaultMaxHistoryTime();
+		}
+
+		if (!Entries.IsEmpty())
+		{
+			int32 TrimBeforeIndex = 0;
+			const float NewestTime = Entries.Last().Time;
+
+			CurrentMinValue = TNumericLimits<float>::Max();
+			CurrentMaxValue = TNumericLimits<float>::Lowest();
+
+			float HistoryTime = 0.f;
+			for (int32 EntryIndex = Entries.Num() - 1; EntryIndex >= 0; --EntryIndex)
+			{
+				const FGraphEntry& Entry = Entries[EntryIndex];
+				for (int32 ValueIndex = 0; ValueIndex < NumValues; ++ValueIndex)
+				{
+					CurrentMinValue = FMath::Min(CurrentMinValue, Entry.Values[ValueIndex]);
+					CurrentMaxValue = FMath::Max(CurrentMaxValue, Entry.Values[ValueIndex]);
+				}
+
+				if (NewestTime - Entry.Time >= MaxHistoryTime)
+				{
+					TrimBeforeIndex = EntryIndex;
+					break;
+				}
+			}
+
+			if (TrimBeforeIndex > 0)
+			{
+				Entries.RemoveAt(0, TrimBeforeIndex, EAllowShrinking::No);
+			}
+		}
+		else
+		{
+			CurrentMinValue = 0.f;
+			CurrentMaxValue = 0.f;
+		}
+	}
+
+private:
+
+	float CurrentMinValue = 0;
+	float CurrentMaxValue = 0;
+
+	TArray<FGraphEntry> Entries;
+};
+
+}  // namespace UE::Cameras
+
+#endif
+
diff --git a/Engine/Plugins/Cameras/GameplayCameras/Source/GameplayCameras/Public/Debug/CameraDebugRenderer.h b/Engine/Plugins/Cameras/GameplayCameras/Source/GameplayCameras/Public/Debug/CameraDebugRenderer.h
--- a/Engine/Plugins/Cameras/GameplayCameras/Source/GameplayCameras/Public/Debug/CameraDebugRenderer.h
+++ b/Engine/Plugins/Cameras/GameplayCameras/Source/GameplayCameras/Public/Debug/CameraDebugRenderer.h
@@ -3,6 +3,7 @@
 #pragma once
 
 #include "CoreTypes.h"
+#include "Debug/CameraDebugGraph.h"
 #include "GameplayCameras.h"
 #include "Math/Color.h"
 #include "Math/MathFwd.h"
@@ -21,6 +22,8 @@ namespace UE::Cameras
 {
 
 class FDebugTextRenderer;
+class FCameraDebugClock;
+template<uint8> class TCameraDebugGraph;
 
 enum class ECameraDebugDrawVisitFlags
 {
@@ -63,8 +66,20 @@ public:
 	/** Decreases the indent of the next text wall entry. This will make a new line. */
 	GAMEPLAYCAMERAS_API void RemoveIndent();
 
-	/** Draws a translucent background behind the text. */
-	GAMEPLAYCAMERAS_API void DrawTextBackgroundTile(float Opacity);
+public:
+
+	/**
+	 * Draw a debug clock showing an angle or 2D vector at the next position available
+	 * for a "card" debug item.
+	 */
+	GAMEPLAYCAMERAS_API void DrawClock(FCameraDebugClock& InClock, const FText& InClockName);
+
+	/**
+	 * Draw a debug graph showing one or more graph lines at the next position available
+	 * for a "card" debug item.
+	 */
+	template<uint8 NumValues>
+	void DrawGraph(TCameraDebugGraph<NumValues>& InGraph, const FText& InGraphName);
 
 public:
 
@@ -109,8 +124,14 @@ public:
 	/** Gets the size of the canvas. */
 	FVector2D GetCanvasSize() const;
 
+	/** Returns whether this renderer has a valid canvas to draw upon. */
 	bool HasCanvas() const { return CanvasObject != nullptr; }
 
+public:
+
+	// Internal API.
+	void DrawTextBackgroundTile(float Opacity);
+
 private:
 
 	void AddTextFmtImpl(const TCHAR* Fmt, va_list Args);
@@ -119,6 +140,9 @@ private:
 	float GetIndentMargin() const;
 	void FlushText();
 
+	FVector2f GetNextCardPosition();
+	void GetNextDrawGraphParams(FCameraDebugGraphDrawParams& OutDrawParams, const FText& InGraphName);
+
 	ULineBatchComponent* GetDebugLineBatcher() const;
 
 private:
@@ -147,10 +171,24 @@ private:
 	/** The maximum horizontal extent of the text rendered so far. */
 	float RightMargin = 0;
 
+	/** The next available position for a card item. */
+	FVector2f NextCardPosition;
+	/** The index of the column for displaying the next card item. */
+	int8 NextCardColumn;
+
 	/** How to visit the next debug blocks. */
 	ECameraDebugDrawVisitFlags VisitFlags;
 };
 
+template<uint8 NumValues>
+inline void FCameraDebugRenderer::DrawGraph(TCameraDebugGraph<NumValues>& InGraph, const FText& InGraphName)
+{
+	FCameraDebugGraphDrawParams DrawParams;
+	GetNextDrawGraphParams(DrawParams, InGraphName);
+	DrawParams.SetupDefaultLineColors<NumValues>();
+	InGraph.Draw(GetCanvas(), DrawParams);
+}
+
 }  // namespace UE::Cameras
 
 #endif  // UE_GAMEPLAY_CAMERAS_DEBUG
