Cameras: use the new clocks and graphs for the debug info of the control rotation service.

diff --git a/Engine/Plugins/Cameras/GameplayCameras/Source/GameplayCameras/Private/Services/PlayerControlRotationService.cpp b/Engine/Plugins/Cameras/GameplayCameras/Source/GameplayCameras/Private/Services/PlayerControlRotationService.cpp
--- a/Engine/Plugins/Cameras/GameplayCameras/Source/GameplayCameras/Private/Services/PlayerControlRotationService.cpp
+++ b/Engine/Plugins/Cameras/GameplayCameras/Source/GameplayCameras/Private/Services/PlayerControlRotationService.cpp
@@ -14,6 +14,8 @@
 #include "GameFramework/PlayerController.h"
 #include "HAL/IConsoleManager.h"
 
+#define LOCTEXT_NAMESPACE "PlayerControlRotationService"
+
 namespace UE::Cameras
 {
 
@@ -23,12 +25,6 @@ static FAutoConsoleVariableRef CVarGameplayCamerasControlRotationDebugArrowLengt
 	GGameplayCamerasControlRotationDebugArrowLength,
 	TEXT(""));
 
-float GGameplayCamerasControlRotationDebugAxisBindingAngleSpeedHistory = 1.f;
-static FAutoConsoleVariableRef CVarGameplayCamerasControlRotationDebugAxisBindingAngleSpeedHistory(
-	TEXT("GameplayControlRotation.DebugAxisBindingAngleSpeedHistory"),
-	GGameplayCamerasControlRotationDebugAxisBindingAngleSpeedHistory,
-	TEXT(""));
-
 UE_DEFINE_CAMERA_EVALUATION_SERVICE(FPlayerControlRotationEvaluationService)
 
 UE_DECLARE_CAMERA_DEBUG_BLOCK_START(GAMEPLAYCAMERAS_API, FPlayerControlRotationDebugBlock)
@@ -37,6 +33,8 @@ UE_DECLARE_CAMERA_DEBUG_BLOCK_START(GAMEPLAYCAMERAS_API, FPlayerControlRotationD
 	UE_DECLARE_CAMERA_DEBUG_BLOCK_FIELD(FRotator3d, CameraRotation)
 	UE_DECLARE_CAMERA_DEBUG_BLOCK_FIELD(bool, bIsFrozen)
 	UE_DECLARE_CAMERA_DEBUG_BLOCK_FIELD(FString, FreezeReason)
+	UE_DECLARE_CAMERA_DEBUG_BLOCK_FIELD(TCameraDebugGraph<1>, AxisActionAngularSpeedGraph)
+	UE_DECLARE_CAMERA_DEBUG_BLOCK_FIELD(FCameraDebugClock, AxisActionValueClock)
 UE_DECLARE_CAMERA_DEBUG_BLOCK_END()
 
 UE_DEFINE_CAMERA_DEBUG_BLOCK_WITH_FIELDS(FPlayerControlRotationDebugBlock)
@@ -219,32 +217,15 @@ void FPlayerControlRotationEvaluationService::UpdateControlRotation(const FCamer
 	// Compute input direction speed change.
 	const FVector2d PreviousAxisDir = PreviousAxisBindingValue.GetSafeNormal();
 	const FVector2d CurrentAxisDir = MaxAxisBindingValue.GetSafeNormal();
-	const double AxisActionAngle = (!PreviousAxisDir.IsZero() && !CurrentAxisDir.IsZero()) ?
+	const double AxisActionAngleChange = (!PreviousAxisDir.IsZero() && !CurrentAxisDir.IsZero()) ?
 		FMath::RadiansToDegrees(FMath::Acos(PreviousAxisDir.Dot(CurrentAxisDir))) :
 		0.0;
-	const double AxisActionAngularSpeed = AxisActionAngle / (Params.DeltaTime > 0 ? Params.DeltaTime : 1.f);
+	const double AxisActionAngularSpeed = AxisActionAngleChange / (Params.DeltaTime > 0 ? Params.DeltaTime : 1.f);
 	PreviousAxisBindingValue = MaxAxisBindingValue;
 
 #if UE_GAMEPLAY_CAMERAS_DEBUG
-	double DebugHighestRecentAxisBindingAngleSpeed = 0.0;
-	if (Params.DeltaTime > 0)
-	{
-		DebugAxisBindingAngleSpeedHistory.Add({ AxisActionAngularSpeed, Params.DeltaTime });
-
-		double HighestValue = 0.f;
-		float AccumulatedDeltaTime = 0.f;
-		for (int32 Index = DebugAxisBindingAngleSpeedHistory.Num() - 1; Index >= 0; --Index)
-		{
-			AccumulatedDeltaTime += DebugAxisBindingAngleSpeedHistory[Index].DeltaTime;
-			HighestValue = FMath::Max(HighestValue, DebugAxisBindingAngleSpeedHistory[Index].Value);
-			if (AccumulatedDeltaTime > GGameplayCamerasControlRotationDebugAxisBindingAngleSpeedHistory)
-			{
-				DebugAxisBindingAngleSpeedHistory.RemoveAt(0, Index);
-				break;
-			}
-		}
-		DebugHighestRecentAxisBindingAngleSpeed = HighestValue;
-	}
+	AxisActionAngularSpeedGraph.Add(Params.DeltaTime, (float)AxisActionAngularSpeed);
+	AxisActionValueClock.Update(MaxAxisBindingValue);
 #endif
 
 	// If we are not feezing the control rotation, we are done.
@@ -253,8 +234,8 @@ void FPlayerControlRotationEvaluationService::UpdateControlRotation(const FCamer
 	{
 #if UE_GAMEPLAY_CAMERAS_DEBUG
 		DebugFreezeReason = FString::Printf(
-				TEXT("no freeze request (input speed %7.2fdeg/s, recent max %7.2fdeg/s)"),
-				AxisActionAngularSpeed, DebugHighestRecentAxisBindingAngleSpeed);
+				TEXT("no freeze request (input speed %7.2fdeg/s)"),
+				AxisActionAngularSpeed);
 #endif
 		return;
 	}
@@ -281,8 +262,8 @@ void FPlayerControlRotationEvaluationService::UpdateControlRotation(const FCamer
 		bIsFrozen = false;
 #if UE_GAMEPLAY_CAMERAS_DEBUG
 		DebugFreezeReason = FString::Printf(
-				TEXT("changed input (%7.2fdeg/s > %7.2fdeg/s, recent max %7.2fdeg/s)"),
-				AxisActionAngularSpeed, ServiceParams.AxisActionAngularSpeedThreshold, DebugHighestRecentAxisBindingAngleSpeed);
+				TEXT("changed input (%7.2fdeg/s > %7.2fdeg/s)"),
+				AxisActionAngularSpeed, ServiceParams.AxisActionAngularSpeedThreshold);
 #endif
 		return;
 	}
@@ -292,8 +273,8 @@ void FPlayerControlRotationEvaluationService::UpdateControlRotation(const FCamer
 	CurrentControlRotation = FrozenControlRotation;
 #if UE_GAMEPLAY_CAMERAS_DEBUG
 	DebugFreezeReason = FString::Printf(
-			TEXT("unchanged input (%7.2fdeg/s < %7.2fdeg/s, recent max %7.2fdeg/s)"), 
-			AxisActionAngularSpeed, ServiceParams.AxisActionAngularSpeedThreshold, DebugHighestRecentAxisBindingAngleSpeed);
+			TEXT("unchanged input (%7.2fdeg/s < %7.2fdeg/s)"), 
+			AxisActionAngularSpeed, ServiceParams.AxisActionAngularSpeedThreshold);
 #endif
 }
 
@@ -308,6 +289,8 @@ void FPlayerControlRotationEvaluationService::OnBuildDebugBlocks(const FCameraDe
 	DebugBlock.CameraRotation = CameraRotation;
 	DebugBlock.bIsFrozen = bIsFrozen;
 	DebugBlock.FreezeReason = DebugFreezeReason;
+	DebugBlock.AxisActionAngularSpeedGraph = AxisActionAngularSpeedGraph;
+	DebugBlock.AxisActionValueClock = AxisActionValueClock;
 }
 
 void FPlayerControlRotationDebugBlock::OnDebugDraw(const FCameraDebugBlockDrawParams& Params, FCameraDebugRenderer& Renderer)
@@ -319,17 +302,17 @@ void FPlayerControlRotationDebugBlock::OnDebugDraw(const FCameraDebugBlockDrawPa
 
 	const FVector3d ForwardArrowEnd = FVector3d::ForwardVector * GGameplayCamerasControlRotationDebugArrowLength;
 
-	// Pawn orientation
+	// Pawn orientation.
 	Renderer.DrawDirectionalArrow(
 			PawnLocation, PawnLocation + PawnTransform.TransformVectorNoScale(ForwardArrowEnd),
 			5.f, FColorList::MandarianOrange, 1.f);
 
-	// Camera rotation
+	// Camera rotation.
 	Renderer.DrawDirectionalArrow(
 			PawnLocation, PawnLocation + CameraYaw.RotateVector(ForwardArrowEnd),
 			5.f, FColorList::PaleGreen, 1.f);
 
-	// Control rotation
+	// Control rotation.
 	FVector3d ControlYawArrowEnd = PawnLocation + ControlYaw.RotateVector(ForwardArrowEnd);
 	Renderer.DrawDirectionalArrow(
 			PawnLocation, ControlYawArrowEnd,
@@ -351,9 +334,15 @@ void FPlayerControlRotationDebugBlock::OnDebugDraw(const FCameraDebugBlockDrawPa
 				CameraRotation.Yaw, *FreezeReason);
 		Renderer.DrawText(ControlYawArrowEnd, TextOffset, DebugText, TextColor);
 	}
+
+	// Value clock and angular speed graph.
+	Renderer.DrawClock(AxisActionValueClock, LOCTEXT("AxisBindingValue", "AxisBindingValue"));
+	Renderer.DrawGraph(AxisActionAngularSpeedGraph, LOCTEXT("AxisBindingAngularSpeed", "AxisBindingAngularSpeed"));
 }
 
 #endif  // UE_GAMEPLAY_CAMERAS_DEBUG
 
 }  // namespace UE::Cameras
 
+#undef LOCTEXT_NAMESPACE
+
diff --git a/Engine/Plugins/Cameras/GameplayCameras/Source/GameplayCameras/Public/Services/PlayerControlRotationService.h b/Engine/Plugins/Cameras/GameplayCameras/Source/GameplayCameras/Public/Services/PlayerControlRotationService.h
--- a/Engine/Plugins/Cameras/GameplayCameras/Source/GameplayCameras/Public/Services/PlayerControlRotationService.h
+++ b/Engine/Plugins/Cameras/GameplayCameras/Source/GameplayCameras/Public/Services/PlayerControlRotationService.h
@@ -3,6 +3,8 @@
 #pragma once
 
 #include "Core/CameraEvaluationService.h"
+#include "Debug/CameraDebugClock.h"
+#include "Debug/CameraDebugGraph.h"
 #include "Math/Rotator.h"
 #include "Math/Transform.h"
 #include "Math/Vector2D.h"
@@ -94,16 +96,11 @@ private:
 	bool bIsFrozen = false;
 
 #if UE_GAMEPLAY_CAMERAS_DEBUG
-	struct FDebugAxisBindingAngleSpeedHistoryEntry
-	{
-		double Value;
-		float DeltaTime;
-	};
-	TArray<FDebugAxisBindingAngleSpeedHistoryEntry> DebugAxisBindingAngleSpeedHistory;
-
 	FTransform DebugPawnTransform;
 	FString DebugFreezeReason;
 	bool bDebugDidApplyControlRotation;
+	TCameraDebugGraph<1> AxisActionAngularSpeedGraph;
+	FCameraDebugClock AxisActionValueClock;
 #endif
 };
 
