Cameras: rewrite Blueprint interop for the camera evaluation result

BP interop was previously using a UObject that wraps an evaluation result, but that proved to be awkward and cumbersome in the code.
This change instead creates BP structs for camera poses and variable tables, accompanied by BP function libraries that provide all kinds of methods to operate on them.
This is much more lightweight and simple to manage.

diff --git a/Engine/Plugins/Cameras/GameplayCameras/Source/GameplayCameras/Private/Core/CameraPose.cpp b/Engine/Plugins/Cameras/GameplayCameras/Source/GameplayCameras/Private/Core/CameraPose.cpp
--- a/Engine/Plugins/Cameras/GameplayCameras/Source/GameplayCameras/Private/Core/CameraPose.cpp
+++ b/Engine/Plugins/Cameras/GameplayCameras/Source/GameplayCameras/Private/Core/CameraPose.cpp
@@ -114,6 +114,11 @@ void FCameraPose::SetTransform(FTransform3d Transform)
 }
 
 double FCameraPose::GetEffectiveFieldOfView() const
+{
+	return GetEffectiveFieldOfView(FocalLength, FieldOfView, SensorWidth, SensorHeight, SqueezeFactor);
+}
+
+double FCameraPose::GetEffectiveFieldOfView(float FocalLength, float FieldOfView, float SensorWidth, float SensorHeight, float SqueezeFactor)
 {
 	checkf((FocalLength > 0.f || FieldOfView > 0.f), TEXT("FocalLength or FieldOfView must have a valid, positive value."));
 
@@ -130,18 +135,18 @@ double FCameraPose::GetEffectiveFieldOfView() const
 	if (FocalLength > 0.f)
 	{
 		// Compute FOV with similar code to UCineCameraComponent...
-		double CropedSensorWidth = SensorWidth * SqueezeFactor;
-		const double AspectRatio = GetSensorAspectRatio();
+		double CroppedSensorWidth = SensorWidth * SqueezeFactor;
+		const double AspectRatio = GetSensorAspectRatio(SensorWidth, SensorHeight);
 		if (AspectRatio > 0.0)
 		{
-			double DesqueezeAspectRatio = SensorWidth * SqueezeFactor / SensorHeight;
+			double DesqueezeAspectRatio = AspectRatio * SqueezeFactor;
 			if (AspectRatio < DesqueezeAspectRatio)
 			{
-				CropedSensorWidth *= AspectRatio / DesqueezeAspectRatio;
+				CroppedSensorWidth *= AspectRatio / DesqueezeAspectRatio;
 			}
 		}
 
-		return FMath::RadiansToDegrees(2.0 * FMath::Atan(CropedSensorWidth / (2.0 * FocalLength)));
+		return FMath::RadiansToDegrees(2.0 * FMath::Atan(CroppedSensorWidth / (2.0 * FocalLength)));
 	}
 	else
 	{
@@ -151,6 +156,11 @@ double FCameraPose::GetEffectiveFieldOfView() const
 }
 
 double FCameraPose::GetSensorAspectRatio() const
+{
+	return GetSensorAspectRatio(SensorWidth, SensorHeight);
+}
+
+double FCameraPose::GetSensorAspectRatio(float SensorWidth, float SensorHeight)
 {
 	return (SensorHeight > 0.f) ? (SensorWidth / SensorHeight) : 0.0;
 }
diff --git a/Engine/Plugins/Cameras/GameplayCameras/Source/GameplayCameras/Private/GameFramework/BlueprintCameraPose.cpp b/Engine/Plugins/Cameras/GameplayCameras/Source/GameplayCameras/Private/GameFramework/BlueprintCameraPose.cpp
new file mode 100644
--- /dev/null
+++ b/Engine/Plugins/Cameras/GameplayCameras/Source/GameplayCameras/Private/GameFramework/BlueprintCameraPose.cpp
@@ -0,0 +1,159 @@
+// Copyright Epic Games, Inc. All Rights Reserved.
+
+#include "GameFramework/BlueprintCameraPose.h"
+
+#include "Camera/CameraComponent.h"
+#include "CineCameraComponent.h"
+
+#include UE_INLINE_GENERATED_CPP_BY_NAME(BlueprintCameraPose)
+
+FBlueprintCameraPose FBlueprintCameraPose::FromCameraPose(const FCameraPose& InCameraPose)
+{
+	FBlueprintCameraPose Result;
+#define UE_CAMERA_POSE_FOR_PROPERTY(PropType, PropName)\
+	Result.PropName = InCameraPose.Get##PropName();
+UE_CAMERA_POSE_FOR_ALL_PROPERTIES()
+#undef UE_CAMERA_POSE_FOR_PROPERTY
+	return Result;
+}
+
+void FBlueprintCameraPose::ApplyTo(FCameraPose& OutCameraPose) const
+{
+#define UE_CAMERA_POSE_FOR_PROPERTY(PropType, PropName)\
+	OutCameraPose.Set##PropName(PropName, false);
+UE_CAMERA_POSE_FOR_ALL_PROPERTIES()
+#undef UE_CAMERA_POSE_FOR_PROPERTY
+}
+
+FBlueprintCameraPose UBlueprintCameraPoseFunctionLibrary::SetLocation(const FBlueprintCameraPose& CameraPose, const FVector& Location)
+{
+	FBlueprintCameraPose Result(CameraPose);
+	Result.Location = Location;
+	return Result;
+}
+
+FBlueprintCameraPose UBlueprintCameraPoseFunctionLibrary::SetRotation(const FBlueprintCameraPose& CameraPose, const FRotator& Rotation)
+{
+	FBlueprintCameraPose Result(CameraPose);
+	Result.Rotation = Rotation;
+	return Result;
+}
+
+FBlueprintCameraPose UBlueprintCameraPoseFunctionLibrary::SetTargetDistance(const FBlueprintCameraPose& CameraPose, double TargetDistance)
+{
+	FBlueprintCameraPose Result(CameraPose);
+	Result.TargetDistance = TargetDistance;
+	return Result;
+}
+
+FBlueprintCameraPose UBlueprintCameraPoseFunctionLibrary::SetFieldOfView(const FBlueprintCameraPose& CameraPose, float FieldOfView)
+{
+	FBlueprintCameraPose Result(CameraPose);
+	Result.FieldOfView = FieldOfView;
+	return Result;
+}
+
+FBlueprintCameraPose UBlueprintCameraPoseFunctionLibrary::SetFocalLength(const FBlueprintCameraPose& CameraPose, float FocalLength)
+{
+	FBlueprintCameraPose Result(CameraPose);
+	Result.FocalLength = FocalLength;
+	return Result;
+}
+
+FTransform UBlueprintCameraPoseFunctionLibrary::GetTransform(const FBlueprintCameraPose& CameraPose)
+{
+	FTransform Transform;
+	Transform.SetLocation(CameraPose.Location);
+	Transform.SetRotation(CameraPose.Rotation.Quaternion());
+	return Transform;
+}
+
+double UBlueprintCameraPoseFunctionLibrary::GetEffectiveFieldOfView(const FBlueprintCameraPose& CameraPose)
+{
+	return FCameraPose::GetEffectiveFieldOfView(
+			CameraPose.FocalLength, 
+			CameraPose.FieldOfView, 
+			CameraPose.SensorWidth, 
+			CameraPose.SensorHeight, 
+			CameraPose.SqueezeFactor);
+}
+
+double UBlueprintCameraPoseFunctionLibrary::GetSensorAspectRatio(const FBlueprintCameraPose& CameraPose)
+{
+	return FCameraPose::GetSensorAspectRatio(
+			CameraPose.SensorWidth,
+			CameraPose.SensorHeight);
+}
+
+FRay UBlueprintCameraPoseFunctionLibrary::GetAimRay(const FBlueprintCameraPose& CameraPose)
+{
+	const bool bDirectionIsNormalized = true;
+	const FVector TargetDir{ 1, 0, 0 };
+	return FRay(CameraPose.Location, CameraPose.Rotation.RotateVector(TargetDir), bDirectionIsNormalized);
+}
+
+FVector UBlueprintCameraPoseFunctionLibrary::GetAimDir(const FBlueprintCameraPose& CameraPose)
+{
+	return CameraPose.Rotation.RotateVector(FVector{ 1, 0, 0 });
+}
+
+FVector UBlueprintCameraPoseFunctionLibrary::GetTarget(const FBlueprintCameraPose& CameraPose)
+{
+	return CameraPose.Location + CameraPose.TargetDistance * GetAimDir(CameraPose);
+}
+
+FVector UBlueprintCameraPoseFunctionLibrary::GetTargetAtDistance(const FBlueprintCameraPose& CameraPose, double TargetDistance)
+{
+	return CameraPose.Location + TargetDistance * GetAimDir(CameraPose);
+}
+
+FBlueprintCameraPose UBlueprintCameraPoseFunctionLibrary::SetTransform(const FBlueprintCameraPose& CameraPose, const FTransform& Transform)
+{
+	FBlueprintCameraPose Result(CameraPose);
+	Result.Location = Transform.GetLocation();
+	Result.Rotation = Transform.GetRotation().Rotator();
+	return Result;
+}
+
+FBlueprintCameraPose UBlueprintCameraPoseFunctionLibrary::MakeCameraPoseFromCameraComponent(const UCameraComponent* CameraComponent)
+{
+	FBlueprintCameraPose Result;
+	if (CameraComponent)
+	{
+		const FTransform& CameraComponentTransform = CameraComponent->GetComponentTransform();
+
+		FTransform AdditiveTransform;
+		float AdditiveFOV;
+		CameraComponent->GetAdditiveOffset(AdditiveTransform, AdditiveFOV);
+
+		Result.Location = CameraComponentTransform.GetLocation() + AdditiveTransform.GetLocation();
+		Result.Rotation = (AdditiveTransform.GetRotation() * CameraComponentTransform.GetRotation()).Rotator();
+		Result.FieldOfView = CameraComponent->FieldOfView + AdditiveFOV;
+		Result.ConstrainAspectRatio = CameraComponent->bConstrainAspectRatio;
+		Result.AspectRatioAxisConstraint = CameraComponent->AspectRatioAxisConstraint;
+		Result.SensorWidth = CameraComponent->AspectRatio * Result.SensorHeight;
+	}
+	else
+	{
+		FFrame::KismetExecutionMessage(TEXT("Invalid camera component was given"), ELogVerbosity::Error);
+	}
+	return Result;
+}
+
+FBlueprintCameraPose UBlueprintCameraPoseFunctionLibrary::MakeCameraPoseFromCineCameraComponent(const UCineCameraComponent* CameraComponent)
+{
+	FBlueprintCameraPose Result = MakeCameraPoseFromCameraComponent(CameraComponent);
+	if (CameraComponent)
+	{
+		Result.TargetDistance = CameraComponent->CurrentFocusDistance;
+		Result.FieldOfView = -1.f;
+		Result.FocalLength = CameraComponent->CurrentFocalLength;
+		Result.Aperture = CameraComponent->CurrentAperture;
+		Result.SensorWidth = CameraComponent->Filmback.SensorWidth;
+		Result.SensorHeight = CameraComponent->Filmback.SensorHeight;
+		Result.SqueezeFactor = CameraComponent->LensSettings.SqueezeFactor;
+	}
+	// Error message already emitted in MakeCameraPoseFromCameraComponent.
+	return Result;
+}
+
diff --git a/Engine/Plugins/Cameras/GameplayCameras/Source/GameplayCameras/Private/GameFramework/BlueprintCameraVariableTable.cpp b/Engine/Plugins/Cameras/GameplayCameras/Source/GameplayCameras/Private/GameFramework/BlueprintCameraVariableTable.cpp
new file mode 100644
--- /dev/null
+++ b/Engine/Plugins/Cameras/GameplayCameras/Source/GameplayCameras/Private/GameFramework/BlueprintCameraVariableTable.cpp
@@ -0,0 +1,134 @@
+// Copyright Epic Games, Inc. All Rights Reserved.
+
+#include "GameFramework/BlueprintCameraVariableTable.h"
+
+#include "Core/CameraVariableAssets.h"
+#include "Core/CameraVariableTable.h"
+
+#include UE_INLINE_GENERATED_CPP_BY_NAME(BlueprintCameraVariableTable)
+
+#define UE_PRIVATE_BLUEPRINT_CAMERA_VARIABLE_TABLE_VALIDATE(ErrorResult)\
+	if (!VariableTable.IsValid())\
+	{\
+		FFrame::KismetExecutionMessage(TEXT("No camera variable table has been set"), ELogVerbosity::Error);\
+		return ErrorResult;\
+	}
+
+#define UE_PRIVATE_BLUEPRINT_CAMERA_VARIABLE_TABLE_VALIDATE_PARAM(ErrorResult)\
+	if (!Variable)\
+	{\
+		FFrame::KismetExecutionMessage(TEXT("No camera variable asset was given"), ELogVerbosity::Error);\
+		return ErrorResult;\
+	}\
+
+#define UE_PRIVATE_BLUEPRINT_CAMERA_VARIABLE_TABLE_GET_VARIABLE(VariableType)\
+	static VariableType ErrorResult {};\
+	UE_PRIVATE_BLUEPRINT_CAMERA_VARIABLE_TABLE_VALIDATE(ErrorResult)\
+	UE_PRIVATE_BLUEPRINT_CAMERA_VARIABLE_TABLE_VALIDATE_PARAM(ErrorResult)\
+	return VariableTable.GetVariableTable()->GetValue<VariableType>(Variable->GetVariableID());
+
+#define UE_PRIVATE_BLUEPRINT_CAMERA_VARIABLE_TABLE_SET_VARIABLE(VariableType)\
+	UE_PRIVATE_BLUEPRINT_CAMERA_VARIABLE_TABLE_VALIDATE()\
+	UE_PRIVATE_BLUEPRINT_CAMERA_VARIABLE_TABLE_VALIDATE_PARAM()\
+	VariableTable.GetVariableTable()->SetValue(Variable, Value, true);
+
+FBlueprintCameraVariableTable::FBlueprintCameraVariableTable()
+	: PrivateVariableTable(nullptr)
+{
+}
+
+FBlueprintCameraVariableTable::FBlueprintCameraVariableTable(FCameraVariableTable* InVariableTable)
+	: PrivateVariableTable(InVariableTable)
+{
+}
+
+bool UBlueprintCameraVariableTableFunctionLibrary::GetBooleanCameraVariable(const FBlueprintCameraVariableTable& VariableTable, UBooleanCameraVariable* Variable)
+{
+	UE_PRIVATE_BLUEPRINT_CAMERA_VARIABLE_TABLE_GET_VARIABLE(bool);
+}
+
+int32 UBlueprintCameraVariableTableFunctionLibrary::GetInteger32CameraVariable(const FBlueprintCameraVariableTable& VariableTable, UInteger32CameraVariable* Variable)
+{
+	UE_PRIVATE_BLUEPRINT_CAMERA_VARIABLE_TABLE_GET_VARIABLE(int32);
+}
+
+float UBlueprintCameraVariableTableFunctionLibrary::GetFloatCameraVariable(const FBlueprintCameraVariableTable& VariableTable, UFloatCameraVariable* Variable)
+{
+	UE_PRIVATE_BLUEPRINT_CAMERA_VARIABLE_TABLE_GET_VARIABLE(float);
+}
+
+double UBlueprintCameraVariableTableFunctionLibrary::GetDoubleCameraVariable(const FBlueprintCameraVariableTable& VariableTable, UDoubleCameraVariable* Variable)
+{
+	UE_PRIVATE_BLUEPRINT_CAMERA_VARIABLE_TABLE_GET_VARIABLE(double);
+}
+
+FVector2D UBlueprintCameraVariableTableFunctionLibrary::GetVector2CameraVariable(const FBlueprintCameraVariableTable& VariableTable, UVector2dCameraVariable* Variable)
+{
+	UE_PRIVATE_BLUEPRINT_CAMERA_VARIABLE_TABLE_GET_VARIABLE(FVector2d);
+}
+
+FVector UBlueprintCameraVariableTableFunctionLibrary::GetVector3CameraVariable(const FBlueprintCameraVariableTable& VariableTable, UVector3dCameraVariable* Variable)
+{
+	UE_PRIVATE_BLUEPRINT_CAMERA_VARIABLE_TABLE_GET_VARIABLE(FVector3d);
+}
+
+FVector4 UBlueprintCameraVariableTableFunctionLibrary::GetVector4CameraVariable(const FBlueprintCameraVariableTable& VariableTable, UVector4dCameraVariable* Variable)
+{
+	UE_PRIVATE_BLUEPRINT_CAMERA_VARIABLE_TABLE_GET_VARIABLE(FVector4d);
+}
+
+FRotator UBlueprintCameraVariableTableFunctionLibrary::GetRotatorCameraVariable(const FBlueprintCameraVariableTable& VariableTable, URotator3dCameraVariable* Variable)
+{
+	UE_PRIVATE_BLUEPRINT_CAMERA_VARIABLE_TABLE_GET_VARIABLE(FRotator3d);
+}
+
+FTransform UBlueprintCameraVariableTableFunctionLibrary::GetTransformCameraVariable(const FBlueprintCameraVariableTable& VariableTable, UTransform3dCameraVariable* Variable)
+{
+	UE_PRIVATE_BLUEPRINT_CAMERA_VARIABLE_TABLE_GET_VARIABLE(FTransform3d);
+}
+
+void UBlueprintCameraVariableTableFunctionLibrary::SetBooleanCameraVariable(const FBlueprintCameraVariableTable& VariableTable, UBooleanCameraVariable* Variable, bool Value)
+{
+	UE_PRIVATE_BLUEPRINT_CAMERA_VARIABLE_TABLE_SET_VARIABLE(bool);
+}
+
+void UBlueprintCameraVariableTableFunctionLibrary::SetInteger32CameraVariable(const FBlueprintCameraVariableTable& VariableTable, UInteger32CameraVariable* Variable, int32 Value)
+{
+	UE_PRIVATE_BLUEPRINT_CAMERA_VARIABLE_TABLE_SET_VARIABLE(int32);
+}
+
+void UBlueprintCameraVariableTableFunctionLibrary::SetFloatCameraVariable(const FBlueprintCameraVariableTable& VariableTable, UFloatCameraVariable* Variable, float Value)
+{
+	UE_PRIVATE_BLUEPRINT_CAMERA_VARIABLE_TABLE_SET_VARIABLE(float);
+}
+
+void UBlueprintCameraVariableTableFunctionLibrary::SetDoubleCameraVariable(const FBlueprintCameraVariableTable& VariableTable, UDoubleCameraVariable* Variable, double Value)
+{
+	UE_PRIVATE_BLUEPRINT_CAMERA_VARIABLE_TABLE_SET_VARIABLE(double);
+}
+
+void UBlueprintCameraVariableTableFunctionLibrary::SetVector2CameraVariable(const FBlueprintCameraVariableTable& VariableTable, UVector2dCameraVariable* Variable, const FVector2D& Value)
+{
+	UE_PRIVATE_BLUEPRINT_CAMERA_VARIABLE_TABLE_SET_VARIABLE(FVector2d);
+}
+
+void UBlueprintCameraVariableTableFunctionLibrary::SetVector3CameraVariable(const FBlueprintCameraVariableTable& VariableTable, UVector3dCameraVariable* Variable, const FVector& Value)
+{
+	UE_PRIVATE_BLUEPRINT_CAMERA_VARIABLE_TABLE_SET_VARIABLE(FVector3d);
+}
+
+void UBlueprintCameraVariableTableFunctionLibrary::SetVector4CameraVariable(const FBlueprintCameraVariableTable& VariableTable, UVector4dCameraVariable* Variable, const FVector4& Value)
+{
+	UE_PRIVATE_BLUEPRINT_CAMERA_VARIABLE_TABLE_SET_VARIABLE(FVector4d);
+}
+
+void UBlueprintCameraVariableTableFunctionLibrary::SetRotatorCameraVariable(const FBlueprintCameraVariableTable& VariableTable, URotator3dCameraVariable* Variable, const FRotator& Value)
+{
+	UE_PRIVATE_BLUEPRINT_CAMERA_VARIABLE_TABLE_SET_VARIABLE(FRotator3d);
+}
+
+void UBlueprintCameraVariableTableFunctionLibrary::SetTransformCameraVariable(const FBlueprintCameraVariableTable& VariableTable, UTransform3dCameraVariable* Variable, const FTransform& Value)
+{
+	UE_PRIVATE_BLUEPRINT_CAMERA_VARIABLE_TABLE_SET_VARIABLE(FTransform3d);
+}
+
diff --git a/Engine/Plugins/Cameras/GameplayCameras/Source/GameplayCameras/Private/GameFramework/CameraEvaluationResultInterop.cpp b/Engine/Plugins/Cameras/GameplayCameras/Source/GameplayCameras/Private/GameFramework/CameraEvaluationResultInterop.cpp
deleted file mode 100644
diff --git a/Engine/Plugins/Cameras/GameplayCameras/Source/GameplayCameras/Private/GameFramework/CameraRigParameterInterop.cpp b/Engine/Plugins/Cameras/GameplayCameras/Source/GameplayCameras/Private/GameFramework/CameraRigParameterInterop.cpp
--- a/Engine/Plugins/Cameras/GameplayCameras/Source/GameplayCameras/Private/GameFramework/CameraRigParameterInterop.cpp
+++ b/Engine/Plugins/Cameras/GameplayCameras/Source/GameplayCameras/Private/GameFramework/CameraRigParameterInterop.cpp
@@ -5,7 +5,7 @@
 #include "Core/CameraRigAsset.h"
 #include "Core/CameraNodeEvaluator.h"
 #include "Core/CameraVariableTable.h"
-#include "GameFramework/CameraEvaluationResultInterop.h"
+#include "GameFramework/BlueprintCameraVariableTable.h"
 
 #define LOCTEXT_NAMESPACE "CameraRigParameterInterop"
 
@@ -15,11 +15,11 @@ namespace UE::Cameras::Private
 {
 
 template<typename VariableAssetType>
-void SetCameraRigParameter(UCameraEvaluationResultInterop* ResultInterop, VariableAssetType* PrivateVariable, typename VariableAssetType::ValueType Value)
+void SetCameraRigParameter(FBlueprintCameraVariableTable& VariableTable, VariableAssetType* PrivateVariable, typename VariableAssetType::ValueType Value)
 {
-	if (ResultInterop == nullptr)
+	if (!VariableTable.IsValid())
 	{
-		FFrame::KismetExecutionMessage(TEXT("No camera evaluation result was passed."), ELogVerbosity::Error);
+		FFrame::KismetExecutionMessage(TEXT("Invalid camera variable table was passed."), ELogVerbosity::Error);
 		return;
 	}
 	if (PrivateVariable == nullptr)
@@ -28,14 +28,7 @@ void SetCameraRigParameter(UCameraEvaluationResultInterop* ResultInterop, Variab
 		return;
 	}
 
-	FCameraNodeEvaluationResult* Result = ResultInterop->GetEvaluationResult();
-	if (Result == nullptr)
-	{
-		FFrame::KismetExecutionMessage(TEXT("The given camera evaluation result is invalid."), ELogVerbosity::Error);
-		return;
-	}
-
-	Result->VariableTable.SetValue(PrivateVariable, Value, true);
+	VariableTable.GetVariableTable()->SetValue(PrivateVariable, Value, true);
 }
 
 }  // namespace UE::Cameras::Private
@@ -45,74 +38,74 @@ UCameraRigParameterInterop::UCameraRigParameterInterop(const FObjectInitializer&
 {
 }
 
-void UCameraRigParameterInterop::SetBooleanParameter(UCameraEvaluationResultInterop* ResultInterop, UCameraRigAsset* CameraRig, const FString& ParameterName, bool bParameterValue)
+void UCameraRigParameterInterop::SetBooleanParameter(FBlueprintCameraVariableTable& VariableTable, UCameraRigAsset* CameraRig, const FString& ParameterName, bool bParameterValue)
 {
 	UE::Cameras::Private::SetCameraRigParameter(
-			ResultInterop, 
+			VariableTable, 
 			Cast<UBooleanCameraVariable>(GetParameterPrivateVariable(CameraRig, ParameterName)), 
 			bParameterValue);
 }
 
-void UCameraRigParameterInterop::SetIntegerParameter(UCameraEvaluationResultInterop* ResultInterop, UCameraRigAsset* CameraRig, const FString& ParameterName, int32 ParameterValue)
+void UCameraRigParameterInterop::SetIntegerParameter(FBlueprintCameraVariableTable& VariableTable, UCameraRigAsset* CameraRig, const FString& ParameterName, int32 ParameterValue)
 {
 	UE::Cameras::Private::SetCameraRigParameter(
-			ResultInterop, 
+			VariableTable, 
 			Cast<UInteger32CameraVariable>(GetParameterPrivateVariable(CameraRig, ParameterName)), 
 			ParameterValue);
 }
 
-void UCameraRigParameterInterop::SetFloatParameter(UCameraEvaluationResultInterop* ResultInterop, UCameraRigAsset* CameraRig, const FString& ParameterName, double ParameterValue)
+void UCameraRigParameterInterop::SetFloatParameter(FBlueprintCameraVariableTable& VariableTable, UCameraRigAsset* CameraRig, const FString& ParameterName, double ParameterValue)
 {
 	UE::Cameras::Private::SetCameraRigParameter(
-			ResultInterop, 
+			VariableTable, 
 			Cast<UFloatCameraVariable>(GetParameterPrivateVariable(CameraRig, ParameterName)),
 			(float)ParameterValue);
 }
 
-void UCameraRigParameterInterop::SetDoubleParameter(UCameraEvaluationResultInterop* ResultInterop, UCameraRigAsset* CameraRig, const FString& ParameterName, double ParameterValue)
+void UCameraRigParameterInterop::SetDoubleParameter(FBlueprintCameraVariableTable& VariableTable, UCameraRigAsset* CameraRig, const FString& ParameterName, double ParameterValue)
 {
 	UE::Cameras::Private::SetCameraRigParameter(
-			ResultInterop, 
+			VariableTable, 
 			Cast<UDoubleCameraVariable>(GetParameterPrivateVariable(CameraRig, ParameterName)),
 			ParameterValue);
 }
 
-void UCameraRigParameterInterop::SetVector2Parameter(UCameraEvaluationResultInterop* ResultInterop, UCameraRigAsset* CameraRig, const FString& ParameterName, FVector2D ParameterValue)
+void UCameraRigParameterInterop::SetVector2Parameter(FBlueprintCameraVariableTable& VariableTable, UCameraRigAsset* CameraRig, const FString& ParameterName, FVector2D ParameterValue)
 {
 	UE::Cameras::Private::SetCameraRigParameter(
-			ResultInterop, 
+			VariableTable, 
 			Cast<UVector2dCameraVariable>(GetParameterPrivateVariable(CameraRig, ParameterName)),
 			ParameterValue);
 }
 
-void UCameraRigParameterInterop::SetVector3Parameter(UCameraEvaluationResultInterop* ResultInterop, UCameraRigAsset* CameraRig, const FString& ParameterName, FVector ParameterValue)
+void UCameraRigParameterInterop::SetVector3Parameter(FBlueprintCameraVariableTable& VariableTable, UCameraRigAsset* CameraRig, const FString& ParameterName, FVector ParameterValue)
 {
 	UE::Cameras::Private::SetCameraRigParameter(
-			ResultInterop, 
+			VariableTable, 
 			Cast<UVector3dCameraVariable>(GetParameterPrivateVariable(CameraRig, ParameterName)),
 			ParameterValue);
 }
 
-void UCameraRigParameterInterop::SetVector4Parameter(UCameraEvaluationResultInterop* ResultInterop, UCameraRigAsset* CameraRig, const FString& ParameterName, FVector4 ParameterValue)
+void UCameraRigParameterInterop::SetVector4Parameter(FBlueprintCameraVariableTable& VariableTable, UCameraRigAsset* CameraRig, const FString& ParameterName, FVector4 ParameterValue)
 {
 	UE::Cameras::Private::SetCameraRigParameter(
-			ResultInterop, 
+			VariableTable, 
 			Cast<UVector4dCameraVariable>(GetParameterPrivateVariable(CameraRig, ParameterName)),
 			ParameterValue);
 }
 
-void UCameraRigParameterInterop::SetRotatorParameter(UCameraEvaluationResultInterop* ResultInterop, UCameraRigAsset* CameraRig, const FString& ParameterName, FRotator ParameterValue)
+void UCameraRigParameterInterop::SetRotatorParameter(FBlueprintCameraVariableTable& VariableTable, UCameraRigAsset* CameraRig, const FString& ParameterName, FRotator ParameterValue)
 {
 	UE::Cameras::Private::SetCameraRigParameter(
-			ResultInterop, 
+			VariableTable, 
 			Cast<URotator3dCameraVariable>(GetParameterPrivateVariable(CameraRig, ParameterName)),
 			ParameterValue);
 }
 
-void UCameraRigParameterInterop::SetTransformParameter(UCameraEvaluationResultInterop* ResultInterop, UCameraRigAsset* CameraRig, const FString& ParameterName, FTransform ParameterValue)
+void UCameraRigParameterInterop::SetTransformParameter(FBlueprintCameraVariableTable& VariableTable, UCameraRigAsset* CameraRig, const FString& ParameterName, FTransform ParameterValue)
 {
 	UE::Cameras::Private::SetCameraRigParameter(
-			ResultInterop, 
+			VariableTable, 
 			Cast<UTransform3dCameraVariable>(GetParameterPrivateVariable(CameraRig, ParameterName)),
 			ParameterValue);
 }
diff --git a/Engine/Plugins/Cameras/GameplayCameras/Source/GameplayCameras/Public/Core/CameraPose.h b/Engine/Plugins/Cameras/GameplayCameras/Source/GameplayCameras/Public/Core/CameraPose.h
--- a/Engine/Plugins/Cameras/GameplayCameras/Source/GameplayCameras/Public/Core/CameraPose.h
+++ b/Engine/Plugins/Cameras/GameplayCameras/Source/GameplayCameras/Public/Core/CameraPose.h
@@ -79,7 +79,7 @@ public:
  * Fields are private and can only be accessed via the getters and setters.
  * The ChangedFlags structure keeps track of which fields were changed via the setters.
  */
-USTRUCT(BlueprintType)
+USTRUCT()
 struct GAMEPLAYCAMERAS_API FCameraPose
 {
 	GENERATED_BODY()
@@ -161,6 +161,14 @@ public:
 	/** Gets the location of the camera's target given a specific distance. */
 	FVector3d GetTarget(double InTargetDistance) const;
 
+public:
+
+	/** Computes the horizontal field of view of a camera. */
+	static double GetEffectiveFieldOfView(float FocalLength, float FieldOfView, float SensorWidth, float SensorHeight, float SqueezeFactor);
+	
+	/** Computes the aspect ratio of a camera sensor. */
+	static double GetSensorAspectRatio(float SensorWidth, float SensorHeight);
+
 public:
 
 	// Interpolation
diff --git a/Engine/Plugins/Cameras/GameplayCameras/Source/GameplayCameras/Public/GameFramework/BlueprintCameraPose.h b/Engine/Plugins/Cameras/GameplayCameras/Source/GameplayCameras/Public/GameFramework/BlueprintCameraPose.h
new file mode 100644
--- /dev/null
+++ b/Engine/Plugins/Cameras/GameplayCameras/Source/GameplayCameras/Public/GameFramework/BlueprintCameraPose.h
@@ -0,0 +1,194 @@
+// Copyright Epic Games, Inc. All Rights Reserved.
+
+#pragma once
+
+#include "Core/CameraPose.h"
+#include "CoreTypes.h"
+#include "Kismet/BlueprintFunctionLibrary.h"
+#include "UObject/ObjectMacros.h"
+
+#include "BlueprintCameraPose.generated.h"
+
+class UCameraComponent;
+class UCineCameraComponent;
+
+namespace UE::Cameras
+{
+
+class FCameraVariableTable;
+struct FCameraNodeEvaluationResult;
+
+}  // namespace UE::Cameras
+
+/** Represents a camera pose. */
+USTRUCT(BlueprintType, meta=(DisplayName="Camera Pose"))
+struct GAMEPLAYCAMERAS_API FBlueprintCameraPose
+{
+	GENERATED_BODY()
+
+public:
+
+	/** The location of the camera. */
+	UPROPERTY(BlueprintReadWrite, EditAnywhere, Category="Camera")
+	FVector Location = {0, 0, 0};
+
+	/** The rotation of the camera. */
+	UPROPERTY(BlueprintReadWrite, EditAnywhere, Category="Camera")
+	FRotator Rotation = {0, 0, 0};
+
+	/** The distance of the target of the camera. */
+	UPROPERTY(BlueprintReadWrite, EditAnywhere, Category="Camera")
+	double TargetDistance = 1000.0;
+
+	/** The field of view of the camera. May be negative if driven by focal length. */
+	UPROPERTY(BlueprintReadWrite, EditAnywhere, Category="Camera")
+	float FieldOfView = -1.f;
+
+	/** The focal length of the camera. May be negative if driven directly by field of view. */
+	UPROPERTY(BlueprintReadWrite, EditAnywhere, Category="Camera")
+	float FocalLength = 35.f;
+
+	/** The aperture (f-stop) of the camera's lens. */
+	UPROPERTY(BlueprintReadWrite, EditAnywhere, Category="Camera")
+	float Aperture = 2.8f;
+
+	/** The focus distance of the camera, if different from target distance. */
+	UPROPERTY(BlueprintReadWrite, EditAnywhere, Category="Camera")
+	float FocusDistance = -1.f;
+
+		/** The width of the camera's sensor. */
+	UPROPERTY(BlueprintReadWrite, EditAnywhere, Category="Camera")
+	float SensorWidth = 24.89f;
+
+	/** The height of the camera's sensor. */
+	UPROPERTY(BlueprintReadWrite, EditAnywhere, Category="Camera")
+	float SensorHeight = 18.67f;
+
+	/** The squeeze factor of the camera's lens. */
+	UPROPERTY(BlueprintReadWrite, EditAnywhere, Category="Camera")
+	float SqueezeFactor = 1.f;
+
+	/** The camera's near clipping plane. */
+	UPROPERTY(BlueprintReadWrite, EditAnywhere, Category="Camera")
+	float NearClippingPlane = 10.f;
+
+	/** The camera's far clipping plane. */
+	UPROPERTY(BlueprintReadWrite, EditAnywhere, Category="Camera")
+	float FarClippingPlane = -1.f;
+
+	/** Whether the camera should constrain aspect ratio. */
+	UPROPERTY(BlueprintReadWrite, EditAnywhere, Category="Camera")
+	bool ConstrainAspectRatio = false;
+
+	/** The aspect ratio axis constraint to use if aspect ratio is constrained. */
+	UPROPERTY(BlueprintReadWrite, EditAnywhere, Category="Camera")
+	TEnumAsByte<EAspectRatioAxisConstraint> AspectRatioAxisConstraint = EAspectRatioAxisConstraint::AspectRatio_MaintainYFOV;
+
+public:
+
+	/** Creates a new BP camera pose from a raw camera pose structure. */
+	static FBlueprintCameraPose FromCameraPose(const FCameraPose& InCameraPose);
+
+	/** Applies this BP camera pose's values to the given raw camera pose. */
+	void ApplyTo(FCameraPose& OutCameraPose) const;
+};
+
+/**
+ * Utility Blueprint functions for camera poses.
+ */
+UCLASS()
+class UBlueprintCameraPoseFunctionLibrary : public UBlueprintFunctionLibrary
+{
+	GENERATED_BODY()
+
+public:
+
+	/** Gets the location of the camera pose. */
+	UFUNCTION(BlueprintPure, Category="Camera")
+	static inline FVector GetLocation(const FBlueprintCameraPose& CameraPose) { return CameraPose.Location; }
+
+	/** Gets the rotation of the camera pose. */
+	UFUNCTION(BlueprintPure, Category="Camera")
+	static inline FRotator GetRotation(const FBlueprintCameraPose& CameraPose) { return CameraPose.Rotation; }
+
+	/** Gets the target distance of the camera pose. */
+	UFUNCTION(BlueprintPure, Category="Camera")
+	static inline double GetTargetDistance(const FBlueprintCameraPose& CameraPose) { return CameraPose.TargetDistance; }
+
+	/** Gets the field of view of the camera pose. */
+	UFUNCTION(BlueprintPure, Category="Camera")
+	static inline double GetFieldOfView(const FBlueprintCameraPose& CameraPose) { return CameraPose.FieldOfView; }
+
+	/** Gets the focal length of the camera pose. */
+	UFUNCTION(BlueprintPure, Category="Camera")
+	static inline double GetFocalLength(const FBlueprintCameraPose& CameraPose) { return CameraPose.FocalLength; }
+	
+public:
+
+	/** Creates a copy of the given camera pose with the given location. */
+	UFUNCTION(BlueprintCallable, Category="Camera")
+	static FBlueprintCameraPose SetLocation(const FBlueprintCameraPose& CameraPose, const FVector& Location);
+
+	/** Creates a copy of the given camera pose with the given rotation. */
+	UFUNCTION(BlueprintCallable, Category="Camera")
+	static FBlueprintCameraPose SetRotation(const FBlueprintCameraPose& CameraPose, const FRotator& Rotation);
+
+	/** Creates a copy of the given camera pose with the given target distance. */
+	UFUNCTION(BlueprintCallable, Category="Camera")
+	static FBlueprintCameraPose SetTargetDistance(const FBlueprintCameraPose& CameraPose, double TargetDistance);
+
+	/** Creates a copy of the given camera pose with the given field of view. */
+	UFUNCTION(BlueprintCallable, Category="Camera")
+	static FBlueprintCameraPose SetFieldOfView(const FBlueprintCameraPose& CameraPose, float FieldOfView);
+
+	/** Creates a copy of the given camera pose with the given focal length. */
+	UFUNCTION(BlueprintCallable, Category="Camera")
+	static FBlueprintCameraPose SetFocalLength(const FBlueprintCameraPose& CameraPose, float FocalLength);
+
+public:
+
+	/** Gets the transform matrix of the camera pose. */
+	UFUNCTION(BlueprintPure, Category="Camera")
+	static FTransform GetTransform(const FBlueprintCameraPose& CameraPose);
+
+	/** Gets the effective field of view of the camera pose, possibly computed from focal length. */
+	UFUNCTION(BlueprintPure, Category="Camera")
+	static double GetEffectiveFieldOfView(const FBlueprintCameraPose& CameraPose);
+
+	/** Gets the effective aspect ratio of the camera pose, computed from the sensor size. */
+	UFUNCTION(BlueprintPure, Category="Camera")
+	static double GetSensorAspectRatio(const FBlueprintCameraPose& CameraPose);
+
+	/** Gets the aim ray of the camera pose. */
+	UFUNCTION(BlueprintPure, Category="Camera")
+	static FRay GetAimRay(const FBlueprintCameraPose& CameraPose);
+
+	/** Gets the facing direction of the camera pose. */
+	UFUNCTION(BlueprintPure, Category="Camera")
+	static FVector GetAimDir(const FBlueprintCameraPose& CameraPose);
+
+	/** Gets the target of the camera pose. */
+	UFUNCTION(BlueprintPure, Category="Camera")
+	static FVector GetTarget(const FBlueprintCameraPose& CameraPose);
+
+	/** Gets the target of the camera pose given a specific target distance. */
+	UFUNCTION(BlueprintPure, Category="Camera")
+	static FVector GetTargetAtDistance(const FBlueprintCameraPose& CameraPose, double TargetDistance);
+
+public:
+
+	/** Creates a copy of the given camera pose with the given location and rotation. */
+	UFUNCTION(BlueprintCallable, Category="Camera")
+	static FBlueprintCameraPose SetTransform(const FBlueprintCameraPose& CameraPose, const FTransform& Transform);
+
+public:
+
+	/** Creates a new camera pose given a camera component. */
+	UFUNCTION(BlueprintCallable, Category="Camera")
+	static FBlueprintCameraPose MakeCameraPoseFromCameraComponent(const UCameraComponent* CameraComponent);
+
+	/** Creates a new camera pose given a cine-camera component. */
+	UFUNCTION(BlueprintCallable, Category="Camera")
+	static FBlueprintCameraPose MakeCameraPoseFromCineCameraComponent(const UCineCameraComponent* CameraComponent);
+};
+
diff --git a/Engine/Plugins/Cameras/GameplayCameras/Source/GameplayCameras/Public/GameFramework/BlueprintCameraVariableTable.h b/Engine/Plugins/Cameras/GameplayCameras/Source/GameplayCameras/Public/GameFramework/BlueprintCameraVariableTable.h
new file mode 100644
--- /dev/null
+++ b/Engine/Plugins/Cameras/GameplayCameras/Source/GameplayCameras/Public/GameFramework/BlueprintCameraVariableTable.h
@@ -0,0 +1,137 @@
+// Copyright Epic Games, Inc. All Rights Reserved.
+
+#pragma once
+
+#include "CoreTypes.h"
+#include "Kismet/BlueprintFunctionLibrary.h"
+#include "Math/MathFwd.h"
+#include "UObject/ObjectMacros.h"
+
+#include "BlueprintCameraVariableTable.generated.h"
+
+class UBooleanCameraVariable;
+class UDoubleCameraVariable;
+class UFloatCameraVariable;
+class UInteger32CameraVariable;
+class URotator3dCameraVariable;
+class UTransform3dCameraVariable;
+class UVector2dCameraVariable;
+class UVector3dCameraVariable;
+class UVector4dCameraVariable;
+
+namespace UE::Cameras
+{
+
+class FCameraVariableTable;
+
+}  // namespace UE::Cameras
+
+/** Provides access to a camera variable table. */
+USTRUCT(BlueprintType, DisplayName="Camera Variable Table")
+struct GAMEPLAYCAMERAS_API FBlueprintCameraVariableTable
+{
+	GENERATED_BODY()
+
+public:
+
+	using FCameraVariableTable = UE::Cameras::FCameraVariableTable;
+
+	FBlueprintCameraVariableTable();
+	FBlueprintCameraVariableTable(FCameraVariableTable* InVariableTable);
+
+	/** Gets the underlying variable table. */
+	FCameraVariableTable* GetVariableTable() const { return PrivateVariableTable; }
+
+	/** Returns whether this variable table is valid. */
+	bool IsValid() const { return PrivateVariableTable != nullptr; }
+
+private:
+
+	FCameraVariableTable* PrivateVariableTable = nullptr;
+};
+
+/**
+ * Utility Blueprint functions for camera variable tables.
+ */
+UCLASS()
+class UBlueprintCameraVariableTableFunctionLibrary : public UBlueprintFunctionLibrary
+{
+	GENERATED_BODY()
+
+public:
+
+	/** Gets a camera variable's value from the given table. */
+	UFUNCTION(BlueprintCallable, Category=Camera, meta=(CompactNodeTitle="Get Bool"))
+	static bool GetBooleanCameraVariable(const FBlueprintCameraVariableTable& VariableTable, UBooleanCameraVariable* Variable);
+
+	/** Gets a camera variable's value from the given table. */
+	UFUNCTION(BlueprintCallable, Category=Camera, meta=(CompactNodeTitle="Get Int"))
+	static int32 GetInteger32CameraVariable(const FBlueprintCameraVariableTable& VariableTable, UInteger32CameraVariable* Variable);
+
+	/** Gets a camera variable's value from the given table. */
+	UFUNCTION(BlueprintCallable, Category=Camera, meta=(CompactNodeTitle="Get Float"))
+	static float GetFloatCameraVariable(const FBlueprintCameraVariableTable& VariableTable, UFloatCameraVariable* Variable);
+
+	/** Gets a camera variable's value from the given table. */
+	UFUNCTION(BlueprintCallable, Category=Camera, meta=(CompactNodeTitle="Get Double"))
+	static double GetDoubleCameraVariable(const FBlueprintCameraVariableTable& VariableTable, UDoubleCameraVariable* Variable);
+
+	/** Gets a camera variable's value from the given table. */
+	UFUNCTION(BlueprintCallable, Category=Camera, meta=(CompactNodeTitle="Get Vec2"))
+	static FVector2D GetVector2CameraVariable(const FBlueprintCameraVariableTable& VariableTable, UVector2dCameraVariable* Variable);
+
+	/** Gets a camera variable's value from the given table. */
+	UFUNCTION(BlueprintCallable, Category=Camera, meta=(CompactNodeTitle="Get Vec3"))
+	static FVector GetVector3CameraVariable(const FBlueprintCameraVariableTable& VariableTable, UVector3dCameraVariable* Variable);
+
+	/** Gets a camera variable's value from the given table. */
+	UFUNCTION(BlueprintCallable, Category=Camera, meta=(CompactNodeTitle="Get Vec4"))
+	static FVector4 GetVector4CameraVariable(const FBlueprintCameraVariableTable& VariableTable, UVector4dCameraVariable* Variable);
+
+	/** Gets a camera variable's value from the given table. */
+	UFUNCTION(BlueprintCallable, Category=Camera, meta=(CompactNodeTitle="Get Rotator"))
+	static FRotator GetRotatorCameraVariable(const FBlueprintCameraVariableTable& VariableTable, URotator3dCameraVariable* Variable);
+
+	/** Gets a camera variable's value from the given table. */
+	UFUNCTION(BlueprintCallable, Category=Camera, meta=(CompactNodeTitle="Get Transform"))
+	static FTransform GetTransformCameraVariable(const FBlueprintCameraVariableTable& VariableTable, UTransform3dCameraVariable* Variable);
+
+public:
+
+	/** Sets a camera variable's value in the given table. */
+	UFUNCTION(BlueprintCallable, Category=Camera)
+	static void SetBooleanCameraVariable(const FBlueprintCameraVariableTable& VariableTable, UBooleanCameraVariable* Variable, bool Value);
+
+	/** Sets a camera variable's value in the given table. */
+	UFUNCTION(BlueprintCallable, Category=Camera)
+	static void SetInteger32CameraVariable(const FBlueprintCameraVariableTable& VariableTable, UInteger32CameraVariable* Variable, int32 Value);
+
+	/** Sets a camera variable's value in the given table. */
+	UFUNCTION(BlueprintCallable, Category=Camera)
+	static void SetFloatCameraVariable(const FBlueprintCameraVariableTable& VariableTable, UFloatCameraVariable* Variable, float Value);
+
+	/** Sets a camera variable's value in the given table. */
+	UFUNCTION(BlueprintCallable, Category=Camera)
+	static void SetDoubleCameraVariable(const FBlueprintCameraVariableTable& VariableTable, UDoubleCameraVariable* Variable, double Value);
+
+	/** Sets a camera variable's value in the given table. */
+	UFUNCTION(BlueprintCallable, Category=Camera)
+	static void SetVector2CameraVariable(const FBlueprintCameraVariableTable& VariableTable, UVector2dCameraVariable* Variable, const FVector2D& Value);
+
+	/** Sets a camera variable's value in the given table. */
+	UFUNCTION(BlueprintCallable, Category=Camera)
+	static void SetVector3CameraVariable(const FBlueprintCameraVariableTable& VariableTable, UVector3dCameraVariable* Variable, const FVector& Value);
+
+	/** Sets a camera variable's value in the given table. */
+	UFUNCTION(BlueprintCallable, Category=Camera)
+	static void SetVector4CameraVariable(const FBlueprintCameraVariableTable& VariableTable, UVector4dCameraVariable* Variable, const FVector4& Value);
+
+	/** Sets a camera variable's value in the given table. */
+	UFUNCTION(BlueprintCallable, Category=Camera)
+	static void SetRotatorCameraVariable(const FBlueprintCameraVariableTable& VariableTable, URotator3dCameraVariable* Variable, const FRotator& Value);
+
+	/** Sets a camera variable's value in the given table. */
+	UFUNCTION(BlueprintCallable, Category=Camera)
+	static void SetTransformCameraVariable(const FBlueprintCameraVariableTable& VariableTable, UTransform3dCameraVariable* Variable, const FTransform& Value);
+};
+
diff --git a/Engine/Plugins/Cameras/GameplayCameras/Source/GameplayCameras/Public/GameFramework/CameraEvaluationResultInterop.h b/Engine/Plugins/Cameras/GameplayCameras/Source/GameplayCameras/Public/GameFramework/CameraEvaluationResultInterop.h
deleted file mode 100644
diff --git a/Engine/Plugins/Cameras/GameplayCameras/Source/GameplayCameras/Public/GameFramework/CameraRigParameterInterop.h b/Engine/Plugins/Cameras/GameplayCameras/Source/GameplayCameras/Public/GameFramework/CameraRigParameterInterop.h
--- a/Engine/Plugins/Cameras/GameplayCameras/Source/GameplayCameras/Public/GameFramework/CameraRigParameterInterop.h
+++ b/Engine/Plugins/Cameras/GameplayCameras/Source/GameplayCameras/Public/GameFramework/CameraRigParameterInterop.h
@@ -11,8 +11,8 @@
 #include "CameraRigParameterInterop.generated.h"
 
 class UCameraRigAsset;
-class UCameraEvaluationResultInterop;
 class UCameraVariableAsset;
+struct FBlueprintCameraVariableTable;
 
 /**
  * Blueprint internal methods to set values on a camera rig's exposed parameters.
@@ -32,31 +32,31 @@ public:
 public:
 
 	UFUNCTION(BlueprintCallable, Category="Camera", meta=(BlueprintInternalUseOnly="true"))
-	static void SetBooleanParameter(UCameraEvaluationResultInterop* ResultInterop, UCameraRigAsset* CameraRig, const FString& ParameterName, bool bParameterValue);
+	static void SetBooleanParameter(FBlueprintCameraVariableTable& VariableTable, UCameraRigAsset* CameraRig, const FString& ParameterName, bool bParameterValue);
 
 	UFUNCTION(BlueprintCallable, Category="Camera", meta=(BlueprintInternalUseOnly="true"))
-	static void SetIntegerParameter(UCameraEvaluationResultInterop* ResultInterop, UCameraRigAsset* CameraRig, const FString& ParameterName, int32 ParameterValue);
+	static void SetIntegerParameter(FBlueprintCameraVariableTable& VariableTable, UCameraRigAsset* CameraRig, const FString& ParameterName, int32 ParameterValue);
 
 	UFUNCTION(BlueprintCallable, Category="Camera", meta=(BlueprintInternalUseOnly="true"))
-	static void SetFloatParameter(UCameraEvaluationResultInterop* ResultInterop, UCameraRigAsset* CameraRig, const FString& ParameterName, double ParameterValue);
+	static void SetFloatParameter(FBlueprintCameraVariableTable& VariableTable, UCameraRigAsset* CameraRig, const FString& ParameterName, double ParameterValue);
 
 	UFUNCTION(BlueprintCallable, Category="Camera", meta=(BlueprintInternalUseOnly="true"))
-	static void SetDoubleParameter(UCameraEvaluationResultInterop* ResultInterop, UCameraRigAsset* CameraRig, const FString& ParameterName, double ParameterValue);
+	static void SetDoubleParameter(FBlueprintCameraVariableTable& VariableTable, UCameraRigAsset* CameraRig, const FString& ParameterName, double ParameterValue);
 
 	UFUNCTION(BlueprintCallable, Category="Camera", meta=(BlueprintInternalUseOnly="true"))
-	static void SetVector2Parameter(UCameraEvaluationResultInterop* ResultInterop, UCameraRigAsset* CameraRig, const FString& ParameterName, FVector2D ParameterValue);
+	static void SetVector2Parameter(FBlueprintCameraVariableTable& VariableTable, UCameraRigAsset* CameraRig, const FString& ParameterName, FVector2D ParameterValue);
 
 	UFUNCTION(BlueprintCallable, Category="Camera", meta=(BlueprintInternalUseOnly="true"))
-	static void SetVector3Parameter(UCameraEvaluationResultInterop* ResultInterop, UCameraRigAsset* CameraRig, const FString& ParameterName, FVector ParameterValue);
+	static void SetVector3Parameter(FBlueprintCameraVariableTable& VariableTable, UCameraRigAsset* CameraRig, const FString& ParameterName, FVector ParameterValue);
 
 	UFUNCTION(BlueprintCallable, Category="Camera", meta=(BlueprintInternalUseOnly="true"))
-	static void SetVector4Parameter(UCameraEvaluationResultInterop* ResultInterop, UCameraRigAsset* CameraRig, const FString& ParameterName, FVector4 ParameterValue);
+	static void SetVector4Parameter(FBlueprintCameraVariableTable& VariableTable, UCameraRigAsset* CameraRig, const FString& ParameterName, FVector4 ParameterValue);
 
 	UFUNCTION(BlueprintCallable, Category="Camera", meta=(BlueprintInternalUseOnly="true"))
-	static void SetRotatorParameter(UCameraEvaluationResultInterop* ResultInterop, UCameraRigAsset* CameraRig, const FString& ParameterName, FRotator ParameterValue);
+	static void SetRotatorParameter(FBlueprintCameraVariableTable& VariableTable, UCameraRigAsset* CameraRig, const FString& ParameterName, FRotator ParameterValue);
 
 	UFUNCTION(BlueprintCallable, Category="Camera", meta=(BlueprintInternalUseOnly="true"))
-	static void SetTransformParameter(UCameraEvaluationResultInterop* ResultInterop, UCameraRigAsset* CameraRig, const FString& ParameterName, FTransform ParameterValue);
+	static void SetTransformParameter(FBlueprintCameraVariableTable& VariableTable, UCameraRigAsset* CameraRig, const FString& ParameterName, FTransform ParameterValue);
 
 private:
 
diff --git a/Engine/Plugins/Cameras/GameplayCameras/Source/GameplayCamerasUncookedOnly/Private/BlueprintGraph/K2Node_SetCameraRigParameters.cpp b/Engine/Plugins/Cameras/GameplayCameras/Source/GameplayCamerasUncookedOnly/Private/BlueprintGraph/K2Node_SetCameraRigParameters.cpp
--- a/Engine/Plugins/Cameras/GameplayCameras/Source/GameplayCamerasUncookedOnly/Private/BlueprintGraph/K2Node_SetCameraRigParameters.cpp
+++ b/Engine/Plugins/Cameras/GameplayCameras/Source/GameplayCamerasUncookedOnly/Private/BlueprintGraph/K2Node_SetCameraRigParameters.cpp
@@ -8,7 +8,7 @@
 #include "Core/CameraVariableAssets.h"
 #include "EdGraphSchema_K2.h"
 #include "EditorCategoryUtils.h"
-#include "GameFramework/CameraEvaluationResultInterop.h"
+#include "GameFramework/BlueprintCameraVariableTable.h"
 #include "GameFramework/CameraRigParameterInterop.h"
 #include "K2Node_CallFunction.h"
 #include "KismetCompiler.h"
@@ -17,7 +17,7 @@
 #define LOCTEXT_NAMESPACE "K2Node_SetCameraRigParameters"
 
 const FName UK2Node_SetCameraRigParameters::CameraRigPinName(TEXT("CameraRig"));
-const FName UK2Node_SetCameraRigParameters::CameraEvaluationResultPinName(TEXT("CameraEvaluationResult"));
+const FName UK2Node_SetCameraRigParameters::CameraVariableTablePinName(TEXT("CameraVariableTable"));
 
 UK2Node_SetCameraRigParameters::UK2Node_SetCameraRigParameters(const FObjectInitializer& ObjectInit)
 	: Super(ObjectInit)
@@ -31,7 +31,7 @@ void UK2Node_SetCameraRigParameters::AllocateDefaultPins()
 	CreatePin(EGPD_Output, UEdGraphSchema_K2::PC_Exec, UEdGraphSchema_K2::PN_Then);
 
 	// Add evalation result pin.
-	CreatePin(EGPD_Input, UEdGraphSchema_K2::PC_Object, UCameraEvaluationResultInterop::StaticClass(), CameraEvaluationResultPinName);
+	CreatePin(EGPD_Input, UEdGraphSchema_K2::PC_Object, FBlueprintCameraVariableTable::StaticStruct(), CameraVariableTablePinName);
 
 	// Add camera rig pin.
 	CreatePin(EGPD_Input, UEdGraphSchema_K2::PC_Object, UCameraRigAsset::StaticClass(), CameraRigPinName);
@@ -123,7 +123,7 @@ void UK2Node_SetCameraRigParameters::ExpandNode(class FKismetCompilerContext& Co
 	GetCameraRigParameterPins(RigParameterPins);
 
 	UEdGraphPin* const CameraRigPin = FindPinChecked(CameraRigPinName);
-	UEdGraphPin* const EvaluationResultPin = FindPinChecked(CameraEvaluationResultPinName);
+	UEdGraphPin* const CameraVariableTablePin = FindPinChecked(CameraVariableTablePinName);
 
 	UEdGraphPin* OriginalThenPin = GetThenPin();
 	UEdGraphPin* PreviousThenPin = nullptr;
@@ -186,9 +186,9 @@ void UK2Node_SetCameraRigParameters::ExpandNode(class FKismetCompilerContext& Co
 		CallSetParameter->FunctionReference.SetExternalMember(CallSetParameterFuncName, UCameraRigParameterInterop::StaticClass());
 		CallSetParameter->AllocateDefaultPins();
 
-		// Connect the evaluation result pin that specifies where the parameter should be overriden.
-		UEdGraphPin* CallSetParameterResultInteropPin = CallSetParameter->FindPinChecked(TEXT("ResultInterop"));
-		CompilerContext.CopyPinLinksToIntermediate(*EvaluationResultPin, *CallSetParameterResultInteropPin);
+		// Connect the variable table pin that specifies where the parameter should be overriden.
+		UEdGraphPin* CallSetParameterVariableTablePin = CallSetParameter->FindPinChecked(TEXT("VariableTable"));
+		CompilerContext.CopyPinLinksToIntermediate(*CameraVariableTablePin, *CallSetParameterVariableTablePin);
 
 		// Connect the camera rig argument.
 		UEdGraphPin* CallSetParameterCameraRigPin = CallSetParameter->FindPinChecked(TEXT("CameraRig"));
@@ -287,7 +287,7 @@ bool UK2Node_SetCameraRigParameters::IsCameraRigParameterPin(UEdGraphPin* Pin) c
 		Pin->PinName != UEdGraphSchema_K2::PN_Then &&
 		Pin->PinName != UEdGraphSchema_K2::PN_ReturnValue &&
 		Pin->PinName != CameraRigPinName &&
-		Pin->PinName != CameraEvaluationResultPinName;
+		Pin->PinName != CameraVariableTablePinName;
 }
 
 void UK2Node_SetCameraRigParameters::CreatePinsForCameraRig(UCameraRigAsset* CameraRig, TArray<UEdGraphPin*>* CreatedPins)
@@ -390,23 +390,6 @@ UCameraRigAsset* UK2Node_SetCameraRigParameters::GetCameraRig(TArrayView<UEdGrap
 	return nullptr;
 }
 
-UCameraEvaluationResultInterop* UK2Node_SetCameraRigParameters::GetCameraEvaluationResult() const
-{
-	UEdGraphPin* ResultPin = GetCameraEvaluationResultPin();
-	if (ResultPin && ResultPin->DefaultObject && ResultPin->LinkedTo.Num() == 0)
-	{
-		return CastChecked<UCameraEvaluationResultInterop>(ResultPin->DefaultObject);
-	}
-	else if (ResultPin && ResultPin->LinkedTo.Num() > 0)
-	{
-		if (UEdGraphPin* ResultSource = ResultPin->LinkedTo[0])
-		{
-			return Cast<UCameraEvaluationResultInterop>(ResultSource->PinType.PinSubCategoryObject.Get());
-		}
-	}
-	return nullptr;
-}
-
 void UK2Node_SetCameraRigParameters::OnCameraRigChanged()
 {
 	TArray<UEdGraphPin*> OldPins = Pins;
diff --git a/Engine/Plugins/Cameras/GameplayCameras/Source/GameplayCamerasUncookedOnly/Private/BlueprintGraph/K2Node_SetCameraRigParameters.h b/Engine/Plugins/Cameras/GameplayCameras/Source/GameplayCamerasUncookedOnly/Private/BlueprintGraph/K2Node_SetCameraRigParameters.h
--- a/Engine/Plugins/Cameras/GameplayCameras/Source/GameplayCamerasUncookedOnly/Private/BlueprintGraph/K2Node_SetCameraRigParameters.h
+++ b/Engine/Plugins/Cameras/GameplayCameras/Source/GameplayCamerasUncookedOnly/Private/BlueprintGraph/K2Node_SetCameraRigParameters.h
@@ -10,7 +10,6 @@
 
 #include "K2Node_SetCameraRigParameters.generated.h"
 
-class UCameraEvaluationResultInterop;
 class UCameraRigAsset;
 
 /**
@@ -47,7 +46,7 @@ public:
 protected:
 
 	static const FName CameraRigPinName;
-	static const FName CameraEvaluationResultPinName;
+	static const FName CameraVariableTablePinName;
 
 	UEdGraphPin* GetCameraRigPin(TArrayView<UEdGraphPin* const>* InPinsToSearch = nullptr) const;
 	UEdGraphPin* GetCameraEvaluationResultPin() const;
@@ -56,7 +55,6 @@ protected:
 	void CreatePinsForCameraRig(UCameraRigAsset* CameraRig, TArray<UEdGraphPin*>* CreatedPins = nullptr);
 
 	UCameraRigAsset* GetCameraRig(TArrayView<UEdGraphPin* const>* InPinsToSearch = nullptr) const;
-	UCameraEvaluationResultInterop* GetCameraEvaluationResult() const;
 	void OnCameraRigChanged();
 };
 
