Cameras: fix missing dependency on CinematicCamera

diff --git a/Engine/Plugins/Cameras/GameplayCameras/Source/GameplayCameras/GameplayCameras.Build.cs b/Engine/Plugins/Cameras/GameplayCameras/Source/GameplayCameras/GameplayCameras.Build.cs
--- a/Engine/Plugins/Cameras/GameplayCameras/Source/GameplayCameras/GameplayCameras.Build.cs
+++ b/Engine/Plugins/Cameras/GameplayCameras/Source/GameplayCameras/GameplayCameras.Build.cs
@@ -11,6 +11,7 @@ public class GameplayCameras : ModuleRules
 
 		PublicDependencyModuleNames.AddRange(
 			new string[] {
+				"CinematicCamera",
 				"Core",
 				"CoreUObject",
 				"Engine",
