Cameras: simplify and fix the activation/deactivation of the camera/camera system components

Activation/deactivation of components is left alone. This is just about activating/deactivating the cameras themselves.

diff --git a/Engine/Plugins/Cameras/GameplayCameras/Source/GameplayCameras/Private/GameFramework/GameplayCameraComponent.cpp b/Engine/Plugins/Cameras/GameplayCameras/Source/GameplayCameras/Private/GameFramework/GameplayCameraComponent.cpp
--- a/Engine/Plugins/Cameras/GameplayCameras/Source/GameplayCameras/Private/GameFramework/GameplayCameraComponent.cpp
+++ b/Engine/Plugins/Cameras/GameplayCameras/Source/GameplayCameras/Private/GameFramework/GameplayCameraComponent.cpp
@@ -6,15 +6,17 @@
 #include "Core/CameraAssetBuilder.h"
 #include "Core/CameraBuildLog.h"
 #include "Core/CameraSystemEvaluator.h"
+#include "Engine/EngineTypes.h"
 #include "Engine/StaticMesh.h"
 #include "Engine/World.h"
+#include "GameDelegates.h"
 #include "GameFramework/Controller.h"
-#include "GameFramework/CameraEvaluationResultInterop.h"
 #include "GameFramework/GameplayCameraSystemActor.h"
 #include "GameFramework/GameplayCameraSystemComponent.h"
 #include "GameplayCameras.h"
 #include "Kismet/GameplayStatics.h"
 #include "Logging/MessageLog.h"
+#include "TimerManager.h"
 #include "UObject/ConstructorHelpers.h"
 #include "UObject/UnrealType.h"
 
@@ -27,8 +29,6 @@ UGameplayCameraComponent::UGameplayCameraComponent(const FObjectInitializer& Obj
 {
 	PrimaryComponentTick.bCanEverTick = true;
 
-	InitialResultInterop = ObjectInit.CreateDefaultSubobject<UCameraEvaluationResultInterop>(this, "InitialResultInterop");
-
 #if WITH_EDITORONLY_DATA
 	if (GIsEditor && !IsRunningCommandlet())
 	{
@@ -45,6 +45,16 @@ TSharedPtr<UE::Cameras::FCameraEvaluationContext> UGameplayCameraComponent::GetE
 }
 
 void UGameplayCameraComponent::ActivateCamera(int32 PlayerIndex)
+{
+	ActivateCameraEvaluationContext(PlayerIndex);
+}
+
+void UGameplayCameraComponent::DeactivateCamera()
+{
+	DeactivateCameraEvaluationContext();
+}
+
+void UGameplayCameraComponent::ActivateCameraEvaluationContext(int32 PlayerIndex)
 {
 	if (ActivatedForPlayerIndex == PlayerIndex)
 	{
@@ -53,19 +63,17 @@ void UGameplayCameraComponent::ActivateCamera(int32 PlayerIndex)
 
 	if (ActivatedForPlayerIndex >= 0)
 	{
-		DeactivateCamera();
+		DeactivateCameraEvaluationContext();
 	}
 
-	Activate();
-
 	if (APlayerController* PC = UGameplayStatics::GetPlayerController(this, PlayerIndex))
 	{
-		ActivateCamera(PC);
+		ActivateCameraEvaluationContext(PC);
 		ActivatedForPlayerIndex = PlayerIndex;
 	}
 }
 
-void UGameplayCameraComponent::DeactivateCamera()
+void UGameplayCameraComponent::DeactivateCameraEvaluationContext()
 {
 	if (ActivatedForPlayerIndex < 0)
 	{
@@ -74,15 +82,13 @@ void UGameplayCameraComponent::DeactivateCamera()
 
 	if (APlayerController* PC = UGameplayStatics::GetPlayerController(this, ActivatedForPlayerIndex))
 	{
-		DeactivateCamera(PC);
+		DeactivateCameraEvaluationContext(PC);
 	}
 
 	ActivatedForPlayerIndex = INDEX_NONE;
-
-	Deactivate();
 }
 
-void UGameplayCameraComponent::ActivateCamera(APlayerController* PlayerController)
+void UGameplayCameraComponent::ActivateCameraEvaluationContext(APlayerController* PlayerController)
 {
 	using namespace UE::Cameras;
 
@@ -92,20 +98,58 @@ void UGameplayCameraComponent::ActivateCamera(APlayerController* PlayerControlle
 	{
 		return;
 	}
+
+	if (!Camera)
+	{
+		UE_LOG(LogCameraSystem, Error, TEXT("Can't activate gameplay camera component: no camera asset was set!"));
+		return;
+	}
 	
 	AGameplayCameraSystemActor* CameraSystem = Cast<AGameplayCameraSystemActor>(PlayerController->PlayerCameraManager->GetViewTarget());
 	if (!CameraSystem)
 	{
-		UE_LOG(LogCameraSystem, Error, TEXT("Can't activate gameplay camera component: no camera system found on the view target!"));
+		FGameDelegates::Get().GetViewTargetChangedDelegate().AddUObject(
+				this, &UGameplayCameraComponent::DelayedActivateCameraEvaluationContext);
+
+		GetWorld()->GetTimerManager().SetTimerForNextTick(
+				this, &UGameplayCameraComponent::AbortDelayedActivateCameraEvaluationContext);
+
+		UE_LOG(LogCameraSystem, Verbose, 
+				TEXT("Can't activate gameplay camera component: no camera system found on the view target. "
+					"Waiting until next tick to see if it is added after us."));
 		return;
 	}
 
-	if (!Camera)
+	DoActivateCameraEvaluationContext(PlayerController, CameraSystem->GetCameraSystemComponent()->GetCameraSystemEvaluator());
+}
+
+void UGameplayCameraComponent::DelayedActivateCameraEvaluationContext(APlayerController* PlayerController, AActor* OldViewTarget, AActor* NewViewTarget)
+{
+	FGameDelegates::Get().GetViewTargetChangedDelegate().RemoveAll(this);
+
+	AGameplayCameraSystemActor* CameraSystem = Cast<AGameplayCameraSystemActor>(NewViewTarget);
+	if (!CameraSystem)
 	{
-		UE_LOG(LogCameraSystem, Error, TEXT("Can't activate gameplay camera component: no camera asset was set!"));
+		UE_LOG(LogCameraSystem, Error, TEXT("Can't activate gameplay camera component: no camera system found on the view target!"));
 		return;
 	}
 
+	DoActivateCameraEvaluationContext(PlayerController, CameraSystem->GetCameraSystemComponent()->GetCameraSystemEvaluator());
+}
+
+void UGameplayCameraComponent::AbortDelayedActivateCameraEvaluationContext()
+{
+	if (ActivatedForPlayerIndex == INDEX_NONE)
+	{
+		UE_LOG(LogCameraSystem, Error, TEXT("Can't activate gameplay camera component: no camera system found after a full tick!"));
+		FGameDelegates::Get().GetViewTargetChangedDelegate().RemoveAll(this);
+	}
+}
+
+void UGameplayCameraComponent::DoActivateCameraEvaluationContext(APlayerController* PlayerController, TSharedPtr<UE::Cameras::FCameraSystemEvaluator> CameraSystemEvaluator)
+{
+	using namespace UE::Cameras;
+
 	if (!EvaluationContext.IsValid())
 	{
 		EvaluationContext = MakeShared<FGameplayCameraComponentEvaluationContext>();
@@ -115,17 +159,12 @@ void UGameplayCameraComponent::ActivateCamera(APlayerController* PlayerControlle
 		InitParams.CameraAsset = Camera;
 		InitParams.PlayerController = PlayerController;
 		EvaluationContext->Initialize(InitParams);
-
-		InitialResultInterop->Setup(&EvaluationContext->GetInitialResult());
 	}
 
-	TSharedPtr<FCameraSystemEvaluator> Evaluator = CameraSystem->GetCameraSystemComponent()->GetCameraSystemEvaluator();
-	Evaluator->PushEvaluationContext(EvaluationContext.ToSharedRef());
-
-	Activate();
+	CameraSystemEvaluator->PushEvaluationContext(EvaluationContext.ToSharedRef());
 }
 
-void UGameplayCameraComponent::DeactivateCamera(APlayerController* PlayerController)
+void UGameplayCameraComponent::DeactivateCameraEvaluationContext(APlayerController* PlayerController)
 {
 	using namespace UE::Cameras;
 
@@ -145,13 +184,22 @@ void UGameplayCameraComponent::DeactivateCamera(APlayerController* PlayerControl
 		TSharedPtr<FCameraSystemEvaluator> Evaluator = CameraSystem->GetCameraSystemComponent()->GetCameraSystemEvaluator();
 		Evaluator->RemoveEvaluationContext(EvaluationContext.ToSharedRef());
 	}
+}
 
-	Deactivate();
+FBlueprintCameraPose UGameplayCameraComponent::GetInitialPose() const
+{
+	return FBlueprintCameraPose::FromCameraPose(EvaluationContext->GetInitialResult().CameraPose);
 }
 
-UCameraEvaluationResultInterop* UGameplayCameraComponent::GetInitialResult() const
+void UGameplayCameraComponent::SetInitialPose(const FBlueprintCameraPose& CameraPose)
 {
-	return InitialResultInterop;
+	FCameraPose InitialPose = EvaluationContext->GetInitialResult().CameraPose;
+	CameraPose.ApplyTo(InitialPose);
+}
+
+FBlueprintCameraVariableTable UGameplayCameraComponent::GetInitialVariableTable() const
+{
+	return FBlueprintCameraVariableTable(&EvaluationContext->GetInitialResult().VariableTable);
 }
 
 void UGameplayCameraComponent::OnRegister()
@@ -176,13 +224,6 @@ void UGameplayCameraComponent::OnRegister()
 #endif	// WITH_EDITORONLY_DATA
 }
 
-void UGameplayCameraComponent::Deactivate()
-{
-	DeactivateCamera();
-
-	Super::Deactivate();
-}
-
 void UGameplayCameraComponent::BeginPlay()
 {
 	Super::BeginPlay();
@@ -204,6 +245,13 @@ void UGameplayCameraComponent::BeginPlay()
 	}
 }
 
+void UGameplayCameraComponent::EndPlay(const EEndPlayReason::Type EndPlayReason)
+{
+	DeactivateCameraEvaluationContext();
+
+	Super::EndPlay(EndPlayReason);
+}
+
 void UGameplayCameraComponent::TickComponent(float DeltaTime, enum ELevelTick TickType, FActorComponentTickFunction *ThisTickFunction)
 {
 	Super::TickComponent(DeltaTime, TickType, ThisTickFunction);
diff --git a/Engine/Plugins/Cameras/GameplayCameras/Source/GameplayCameras/Private/GameFramework/GameplayCameraSystemComponent.cpp b/Engine/Plugins/Cameras/GameplayCameras/Source/GameplayCameras/Private/GameFramework/GameplayCameraSystemComponent.cpp
--- a/Engine/Plugins/Cameras/GameplayCameras/Source/GameplayCameras/Private/GameFramework/GameplayCameraSystemComponent.cpp
+++ b/Engine/Plugins/Cameras/GameplayCameras/Source/GameplayCameras/Private/GameFramework/GameplayCameraSystemComponent.cpp
@@ -90,13 +90,6 @@ void UGameplayCameraSystemComponent::OnRegister()
 #endif	// WITH_EDITORONLY_DATA
 }
 
-void UGameplayCameraSystemComponent::Deactivate()
-{
-	DeactivateCameraSystem();
-
-	Super::Deactivate();
-}
-
 void UGameplayCameraSystemComponent::ActivateCameraSystem(int32 PlayerIndex)
 {
 	if (ActivatedForPlayerIndex == PlayerIndex)
@@ -123,8 +116,6 @@ void UGameplayCameraSystemComponent::ActivateCameraSystem(int32 PlayerIndex)
 		return;
 	}
 
-	Activate();
-
 	PC->SetViewTarget(OwningActor);
 	ActivatedForPlayerIndex = PlayerIndex;
 }
@@ -143,8 +134,6 @@ void UGameplayCameraSystemComponent::DeactivateCameraSystem(AActor* NextViewTarg
 	}
 
 	ActivatedForPlayerIndex = INDEX_NONE;
-
-	Deactivate();
 }
 
 void UGameplayCameraSystemComponent::BeginPlay()
@@ -158,6 +147,13 @@ void UGameplayCameraSystemComponent::BeginPlay()
 	}
 }
 
+void UGameplayCameraSystemComponent::EndPlay(const EEndPlayReason::Type EndPlayReason)
+{
+	DeactivateCameraSystem();
+
+	Super::EndPlay(EndPlayReason);
+}
+
 void UGameplayCameraSystemComponent::TickComponent(float DeltaTime, ELevelTick TickType, FActorComponentTickFunction* ThisTickFunction)
 {
 	Super::TickComponent(DeltaTime, TickType, ThisTickFunction);
diff --git a/Engine/Plugins/Cameras/GameplayCameras/Source/GameplayCameras/Public/GameFramework/GameplayCameraComponent.h b/Engine/Plugins/Cameras/GameplayCameras/Source/GameplayCameras/Public/GameFramework/GameplayCameraComponent.h
--- a/Engine/Plugins/Cameras/GameplayCameras/Source/GameplayCameras/Public/GameFramework/GameplayCameraComponent.h
+++ b/Engine/Plugins/Cameras/GameplayCameras/Source/GameplayCameras/Public/GameFramework/GameplayCameraComponent.h
@@ -5,17 +5,20 @@
 #include "Components/SceneComponent.h"
 #include "Core/CameraEvaluationContext.h"
 #include "CoreMinimal.h"
+#include "Engine/TimerHandle.h"
+#include "GameFramework/BlueprintCameraPose.h"
+#include "GameFramework/BlueprintCameraVariableTable.h"
 #include "UObject/ObjectMacros.h"
 #include "UObject/ScriptInterface.h"
 
 #include "GameplayCameraComponent.generated.h"
 
 class UCameraAsset;
-class UCameraEvaluationResultInterop;
 
 namespace UE::Cameras
 {
 
+class FCameraSystemEvaluator;
 class FGameplayCameraComponentEvaluationContext;
 
 }  // namespace UE::Cameras
@@ -43,9 +46,9 @@ public:
 
 	/** 
 	 * Activates the camera for the given player.
-	 * This looks up the current player camera manager and/or view target in order to find
-	 * the active camera system for the given player. If found, it adds its own camera asset
-	 * as the active one.
+	 * This looks up the given player's camera manager and/or view target in order to find
+	 * the active camera system. If found, this component adds its camera asset as the active one.
+	 * If this component was already active for another player, it will be first deactivated.
 	 */
 	UFUNCTION(BlueprintCallable, Category=Camera)
 	GAMEPLAYCAMERAS_API void ActivateCamera(int32 PlayerIndex = 0);
@@ -54,23 +57,39 @@ public:
 	UFUNCTION(BlueprintCallable, Category=Camera)
 	GAMEPLAYCAMERAS_API void DeactivateCamera();
 
-	/** Gets the initial evaluation result for this component's context. */
+	/** Gets the initial camera pose for this component's camera evaluation context. */
 	UFUNCTION(BlueprintPure, Category=Camera)
-	GAMEPLAYCAMERAS_API UCameraEvaluationResultInterop* GetInitialResult() const;
+	GAMEPLAYCAMERAS_API FBlueprintCameraPose GetInitialPose() const;
+
+	/** Sets the initial camera pose for this component's camera evaluation context. */
+	UFUNCTION(BlueprintCallable, Category=Camera)
+	GAMEPLAYCAMERAS_API void SetInitialPose(const FBlueprintCameraPose& CameraPose);
+
+	/** Gets the initial camera variable table for this component's camera evaluation context. */
+	UFUNCTION(BlueprintPure, Category=Camera)
+	GAMEPLAYCAMERAS_API FBlueprintCameraVariableTable GetInitialVariableTable() const;
 
 public:
 
 	// UActorComponent interface
 	virtual void OnRegister() override;
-	virtual void Deactivate() override;
 	virtual void BeginPlay() override;
+	virtual void EndPlay(const EEndPlayReason::Type EndPlayReason) override;
 	virtual void TickComponent(float DeltaTime, enum ELevelTick TickType, FActorComponentTickFunction *ThisTickFunction) override;
 	virtual void OnComponentDestroyed(bool bDestroyingHierarchy) override;
 
 private:
 
-	void ActivateCamera(APlayerController* PlayerController);
-	void DeactivateCamera(APlayerController* PlayerController);
+	void ActivateCameraEvaluationContext(int32 PlayerIndex);
+	void DeactivateCameraEvaluationContext();
+
+	void ActivateCameraEvaluationContext(APlayerController* PlayerController);
+	void DeactivateCameraEvaluationContext(APlayerController* PlayerController);
+
+	void DelayedActivateCameraEvaluationContext(APlayerController* PlayerController, AActor* OldViewTarget, AActor* NewViewTarget);
+	void AbortDelayedActivateCameraEvaluationContext();
+
+	void DoActivateCameraEvaluationContext(APlayerController* PlayerController, TSharedPtr<UE::Cameras::FCameraSystemEvaluator> CameraSystemEvaluator);
 
 #if WITH_EDITORONLY_DATA
 
@@ -97,9 +116,6 @@ protected:
 
 	TSharedPtr<FGameplayCameraComponentEvaluationContext> EvaluationContext;
 
-	UPROPERTY(Transient)
-	TObjectPtr<UCameraEvaluationResultInterop> InitialResultInterop;
-	
 #if WITH_EDITORONLY_DATA
 
 	UPROPERTY(Transient)
diff --git a/Engine/Plugins/Cameras/GameplayCameras/Source/GameplayCameras/Public/GameFramework/GameplayCameraSystemComponent.h b/Engine/Plugins/Cameras/GameplayCameras/Source/GameplayCameras/Public/GameFramework/GameplayCameraSystemComponent.h
--- a/Engine/Plugins/Cameras/GameplayCameras/Source/GameplayCameras/Public/GameFramework/GameplayCameraSystemComponent.h
+++ b/Engine/Plugins/Cameras/GameplayCameras/Source/GameplayCameras/Public/GameFramework/GameplayCameraSystemComponent.h
@@ -55,8 +55,8 @@ public:
 
 	// UActorComponent interface
 	virtual void OnRegister() override;
-	virtual void Deactivate() override;
 	virtual void BeginPlay() override;
+	virtual void EndPlay(const EEndPlayReason::Type EndPlayReason) override;
 	virtual void TickComponent(float DeltaTime, ELevelTick TickType, FActorComponentTickFunction* ThisTickFunction) override;
 	virtual void OnComponentDestroyed(bool bDestroyingHierarchy) override;
 
