Cameras: add OnActivate/OnDeactivate API on camera director evaluators

diff --git a/Engine/Plugins/Cameras/GameplayCameras/Source/GameplayCameras/Private/Core/CameraDirectorEvaluator.cpp b/Engine/Plugins/Cameras/GameplayCameras/Source/GameplayCameras/Private/Core/CameraDirectorEvaluator.cpp
--- a/Engine/Plugins/Cameras/GameplayCameras/Source/GameplayCameras/Private/Core/CameraDirectorEvaluator.cpp
+++ b/Engine/Plugins/Cameras/GameplayCameras/Source/GameplayCameras/Private/Core/CameraDirectorEvaluator.cpp
@@ -31,6 +31,20 @@ void FCameraDirectorEvaluator::Initialize(const FCameraDirectorInitializeParams&
 	OnInitialize(Params);
 }
 
+void FCameraDirectorEvaluator::Activate(const FCameraDirectorActivateParams& Params)
+{
+	WeakOwnerContext = Params.OwnerContext;
+
+	OnActivate(Params);
+}
+
+void FCameraDirectorEvaluator::Deactivate(const FCameraDirectorDeactivateParams& Params)
+{
+	OnDeactivate(Params);
+
+	WeakOwnerContext.Reset();
+}
+
 void FCameraDirectorEvaluator::Run(const FCameraDirectorEvaluationParams& Params, FCameraDirectorEvaluationResult& OutResult)
 {
 	OnRun(Params, OutResult);
diff --git a/Engine/Plugins/Cameras/GameplayCameras/Source/GameplayCameras/Private/Core/CameraEvaluationContext.cpp b/Engine/Plugins/Cameras/GameplayCameras/Source/GameplayCameras/Private/Core/CameraEvaluationContext.cpp
--- a/Engine/Plugins/Cameras/GameplayCameras/Source/GameplayCameras/Private/Core/CameraEvaluationContext.cpp
+++ b/Engine/Plugins/Cameras/GameplayCameras/Source/GameplayCameras/Private/Core/CameraEvaluationContext.cpp
@@ -85,7 +85,6 @@ void FCameraEvaluationContext::AutoCreateDirectorEvaluator()
 
 		FCameraDirectorInitializeParams InitParams;
 		InitParams.OwnerContext = SharedThis(this);
-		InitParams.CameraDirector = CameraDirector;
 		DirectorEvaluator->Initialize(InitParams);
 	}
 }
@@ -105,6 +104,13 @@ void FCameraEvaluationContext::Activate(const FCameraEvaluationContextActivatePa
 
 	AutoCreateDirectorEvaluator();
 
+	if (ensure(DirectorEvaluator))
+	{
+		FCameraDirectorActivateParams DirectorParams;
+		DirectorParams.OwnerContext = SharedThis(this);
+		DirectorEvaluator->Activate(DirectorParams);
+	}
+
 	bActivated = true;
 }
 
@@ -115,6 +121,13 @@ void FCameraEvaluationContext::Deactivate(const FCameraEvaluationContextDeactiva
 		return;
 	}
 
+	if (ensure(DirectorEvaluator))
+	{
+		FCameraDirectorDeactivateParams DirectorParams;
+		DirectorParams.OwnerContext = SharedThis(this);
+		DirectorEvaluator->Deactivate(DirectorParams);
+	}
+
 	// Don't destroy the camera director evaluator, it could still be useful. We only destroy it
 	// along with this context.
 
diff --git a/Engine/Plugins/Cameras/GameplayCameras/Source/GameplayCameras/Public/Core/CameraDirectorEvaluator.h b/Engine/Plugins/Cameras/GameplayCameras/Source/GameplayCameras/Public/Core/CameraDirectorEvaluator.h
--- a/Engine/Plugins/Cameras/GameplayCameras/Source/GameplayCameras/Public/Core/CameraDirectorEvaluator.h
+++ b/Engine/Plugins/Cameras/GameplayCameras/Source/GameplayCameras/Public/Core/CameraDirectorEvaluator.h
@@ -18,10 +18,31 @@ class FCameraDirectorEvaluator;
 class FCameraDirectorEvaluatorStorage;
 class FCameraEvaluationContext;
 
+/**
+ * Parameter structure for initializing a newly created camera director evaluator.
+ */
 struct FCameraDirectorInitializeParams
 {
+	/** The evaluation context that owns the camera director. */
+	TSharedPtr<FCameraEvaluationContext> OwnerContext;
+};
+
+/**
+ * Parameter structure for activating a camera director evaluator.
+ * A camera director evaluator can be activated and deactivated multiple times.
+ * This means that OnActivated/OnDeactivated can be called multiple times in pairs,
+ * unlike OnInitialize which is only called once.
+ */
+struct FCameraDirectorActivateParams
+{
+	/** The evaluation context that owns the camera director. */
+	TSharedPtr<FCameraEvaluationContext> OwnerContext;
+};
+
+struct FCameraDirectorDeactivateParams
+{
+	/** The evaluation context that owns the camera director. */
 	TSharedPtr<FCameraEvaluationContext> OwnerContext;
-	TObjectPtr<const UCameraDirector> CameraDirector;
 };
 
 /**
@@ -109,11 +130,19 @@ public:
 	GAMEPLAYCAMERAS_API FCameraDirectorEvaluator();
 	GAMEPLAYCAMERAS_API virtual ~FCameraDirectorEvaluator() {}
 
-	void Initialize(const FCameraDirectorInitializeParams& Params);
+	/** Initializes a camera director evalutor. */
+	GAMEPLAYCAMERAS_API void Initialize(const FCameraDirectorInitializeParams& Params);
+
+	/** Activates the camera director evaluator. */
+	GAMEPLAYCAMERAS_API void Activate(const FCameraDirectorActivateParams& Params);
+
+	/** Deactivates the camera director evaluator. */
+	GAMEPLAYCAMERAS_API void Deactivate(const FCameraDirectorDeactivateParams& Params);
 	
 	/** Runs the camera director to determine what camera rig(s) should be active this frame. */
-	void Run(const FCameraDirectorEvaluationParams& Params, FCameraDirectorEvaluationResult& OutResult);
+	GAMEPLAYCAMERAS_API void Run(const FCameraDirectorEvaluationParams& Params, FCameraDirectorEvaluationResult& OutResult);
 
+	/** Gets the camera director. */
 	const UCameraDirector* GetCameraDirector() const { return PrivateCameraDirector; }
 
 	/** Gets the camera director. */
@@ -126,7 +155,7 @@ public:
 	GAMEPLAYCAMERAS_API bool AddChildEvaluationContext(TSharedRef<FCameraEvaluationContext> InContext);
 	GAMEPLAYCAMERAS_API bool RemoveChildEvaluationContext(TSharedRef<FCameraEvaluationContext> InContext);
 
-	void AddReferencedObjects(FReferenceCollector& Collector);
+	GAMEPLAYCAMERAS_API void AddReferencedObjects(FReferenceCollector& Collector);
 
 public:
 
@@ -155,8 +184,15 @@ protected:
 #endif  // UE_GAMEPLAY_CAMERAS_DEBUG
 	};
 
+	/** Initializes a camera director evalutor. Only called once after construction. */
 	virtual void OnInitialize(const FCameraDirectorInitializeParams& Params) {}
 
+	/** Activates the camera director evaluator. May be called multiple times, in pair with OnDeactivate. */
+	virtual void OnActivate(const FCameraDirectorActivateParams& Params) {}
+
+	/** Deactivates the camera director evaluator. May be called multiple times, in pair with OnActivate. */
+	virtual void OnDeactivate(const FCameraDirectorDeactivateParams& Params) {}
+
 	/** Runs the camera director to determine what camera rig(s) should be active this frame. */
 	virtual void OnRun(const FCameraDirectorEvaluationParams& Params, FCameraDirectorEvaluationResult& OutResult) {}
 
