Cameras: add StateTree camera director

- Add a StateTree camera director and evaluator.
- Add a basic task to activate a camera rig.
- Add a Details View customization for picking camera rigs from a referenced camera asset.

diff --git a/Engine/Plugins/Cameras/GameplayCameras/GameplayCameras.uplugin b/Engine/Plugins/Cameras/GameplayCameras/GameplayCameras.uplugin
--- a/Engine/Plugins/Cameras/GameplayCameras/GameplayCameras.uplugin
+++ b/Engine/Plugins/Cameras/GameplayCameras/GameplayCameras.uplugin
@@ -20,6 +20,10 @@
 			"Name": "EnhancedInput",
 			"Enabled": true
 		},
+		{
+			"Name": "StateTree",
+			"Enabled": true
+		},
 		{
 			"Name": "TemplateSequence",
 			"Enabled": true
diff --git a/Engine/Plugins/Cameras/GameplayCameras/Source/GameplayCameras/GameplayCameras.Build.cs b/Engine/Plugins/Cameras/GameplayCameras/Source/GameplayCameras/GameplayCameras.Build.cs
--- a/Engine/Plugins/Cameras/GameplayCameras/Source/GameplayCameras/GameplayCameras.Build.cs
+++ b/Engine/Plugins/Cameras/GameplayCameras/Source/GameplayCameras/GameplayCameras.Build.cs
@@ -20,6 +20,7 @@ public class GameplayCameras : ModuleRules
 				"HeadMountedDisplay",
 				"MovieScene",
 				"MovieSceneTracks",
+				"StateTreeModule",
 				"TemplateSequence",
 				"TraceLog"
 			}
diff --git a/Engine/Plugins/Cameras/GameplayCameras/Source/GameplayCameras/Private/Directors/CameraDirectorStateTreeSchema.cpp b/Engine/Plugins/Cameras/GameplayCameras/Source/GameplayCameras/Private/Directors/CameraDirectorStateTreeSchema.cpp
new file mode 100644
--- /dev/null
+++ b/Engine/Plugins/Cameras/GameplayCameras/Source/GameplayCameras/Private/Directors/CameraDirectorStateTreeSchema.cpp
@@ -0,0 +1,55 @@
+// Copyright Epic Games, Inc. All Rights Reserved.
+
+#include "Directors/CameraDirectorStateTreeSchema.h"
+
+#include "StateTreeConditionBase.h"
+#include "StateTreeConsiderationBase.h"
+#include "StateTreeEvaluatorBase.h"
+#include "StateTreeTaskBase.h"
+
+#include UE_INLINE_GENERATED_CPP_BY_NAME(CameraDirectorStateTreeSchema)
+
+namespace UE::Cameras
+{
+
+const FName FStateTreeContextDataNames::ContextOwner = TEXT("ContextOwner");
+
+}  // namespace UE::Cameras
+
+UCameraDirectorStateTreeSchema::UCameraDirectorStateTreeSchema()
+{
+	using namespace UE::Cameras;
+
+	ContextDataDescs.Append({
+			// EvaluationContextOwner: {A474F4B3-A82F-45C2-9405-3535F711751B}
+			FStateTreeExternalDataDesc(
+					FStateTreeContextDataNames::ContextOwner,
+					UObject::StaticClass(), 
+					FGuid(0xA474F4B3, 0xA82F45C2, 0x94053535, 0xF711751B))
+			});
+}
+
+bool UCameraDirectorStateTreeSchema::IsStructAllowed(const UScriptStruct* InScriptStruct) const
+{
+	return
+		// Common structs.
+		InScriptStruct->IsChildOf(FStateTreeConditionCommonBase::StaticStruct())
+		|| InScriptStruct->IsChildOf(FStateTreeConsiderationCommonBase::StaticStruct())
+		|| InScriptStruct->IsChildOf(FStateTreeEvaluatorCommonBase::StaticStruct())
+		|| InScriptStruct->IsChildOf(FStateTreeTaskCommonBase::StaticStruct())
+		// GameplayCameras structs.
+		|| InScriptStruct->IsChildOf(FGameplayCamerasStateTreeTask::StaticStruct())
+		|| InScriptStruct->IsChildOf(FGameplayCamerasStateTreeCondition::StaticStruct())
+		;
+}
+
+bool UCameraDirectorStateTreeSchema::IsClassAllowed(const UClass* InClass) const
+{
+	return IsChildOfBlueprintBase(InClass);
+}
+
+bool UCameraDirectorStateTreeSchema::IsExternalItemAllowed(const UStruct& InStruct) const
+{
+	return true;
+}
+
diff --git a/Engine/Plugins/Cameras/GameplayCameras/Source/GameplayCameras/Private/Directors/StateTreeCameraDirector.cpp b/Engine/Plugins/Cameras/GameplayCameras/Source/GameplayCameras/Private/Directors/StateTreeCameraDirector.cpp
new file mode 100644
--- /dev/null
+++ b/Engine/Plugins/Cameras/GameplayCameras/Source/GameplayCameras/Private/Directors/StateTreeCameraDirector.cpp
@@ -0,0 +1,214 @@
+// Copyright Epic Games, Inc. All Rights Reserved.
+
+#include "Directors/StateTreeCameraDirector.h"
+
+#include "Core/CameraAsset.h"
+#include "Core/CameraBuildLog.h"
+#include "Core/CameraEvaluationContext.h"
+#include "Directors/CameraDirectorStateTreeSchema.h"
+#include "GameplayCameras.h"
+#include "Logging/TokenizedMessage.h"
+#include "StateTreeExecutionContext.h"
+#include "StateTreeInstanceData.h"
+
+#include UE_INLINE_GENERATED_CPP_BY_NAME(StateTreeCameraDirector)
+
+#define LOCTEXT_NAMESPACE "StateTreeCameraDirector"
+
+namespace UE::Cameras
+{
+
+class FStateTreeCameraDirectorEvaluator : public FCameraDirectorEvaluator
+{
+	UE_DECLARE_CAMERA_DIRECTOR_EVALUATOR(GAMEPLAYCAMERAS_API, FStateTreeCameraDirectorEvaluator)
+
+public:
+
+protected:
+
+	// FCameraDirectorEvaluator interface.
+	virtual void OnActivate(const FCameraDirectorActivateParams& Params) override;
+	virtual void OnDeactivate(const FCameraDirectorDeactivateParams& Params) override;
+	virtual void OnRun(const FCameraDirectorEvaluationParams& Params, FCameraDirectorEvaluationResult& OutResult) override;
+	virtual void OnAddReferencedObjects(FReferenceCollector& Collector) override;
+
+private:
+
+	bool SetContextRequirements(TSharedPtr<const FCameraEvaluationContext> OwnerContext, FStateTreeExecutionContext& StateTreeContext);
+
+private:
+
+	FStateTreeInstanceData StateTreeInstanceData;
+
+	FCameraDirectorStateTreeEvaluationData EvaluationData;
+};
+
+UE_DEFINE_CAMERA_DIRECTOR_EVALUATOR(FStateTreeCameraDirectorEvaluator)
+
+void FStateTreeCameraDirectorEvaluator::OnActivate(const FCameraDirectorActivateParams& Params)
+{
+	const UStateTreeCameraDirector* StateTreeDirector = GetCameraDirectorAs<UStateTreeCameraDirector>();
+	const FStateTreeReference& StateTreeReference = StateTreeDirector->StateTreeReference;
+	const UStateTree* StateTree = StateTreeReference.GetStateTree();
+
+	if (!StateTree)
+	{
+		UE_LOG(LogCameraSystem, Error,
+			TEXT("Can't activate camera director '%s': it doesn't have a valid StateTree asset specified."),
+			*GetNameSafe(StateTreeDirector));
+		return;
+	}
+
+	UObject* ContextOwner = Params.OwnerContext->GetOwner();
+	if (!ContextOwner)
+	{
+		UE_LOG(LogCameraSystem, Error,
+			TEXT("Can't activate camera director '%s': the evaluation context doesn't have a valid owner."),
+			*GetNameSafe(StateTreeDirector));
+		return;
+	}
+
+	FStateTreeExecutionContext StateTreeContext(*ContextOwner, *StateTree, StateTreeInstanceData);
+
+	if (!StateTreeContext.IsValid())
+	{
+		UE_LOG(LogCameraSystem, Error,
+			TEXT("Can't activate camera director '%s': initialization of execution context for StateTree asset '%s' "
+				"and context owner '%s' failed."),
+			*GetNameSafe(StateTreeDirector), *GetNameSafe(StateTree), *GetNameSafe(ContextOwner));
+		return;
+	}
+
+	// TODO: validate schema.
+	
+	if (!SetContextRequirements(Params.OwnerContext, StateTreeContext))
+	{
+		UE_LOG(LogCameraSystem, Error,
+			TEXT("Can't activate camera director '%s': failed to setup external data views for StateTree asset '%s'."),
+			*GetNameSafe(StateTreeDirector), *GetNameSafe(StateTree));
+		return;
+	}
+
+	StateTreeContext.Start(&StateTreeReference.GetParameters());
+}
+
+void FStateTreeCameraDirectorEvaluator::OnDeactivate(const FCameraDirectorDeactivateParams& Params)
+{
+	const UStateTreeCameraDirector* StateTreeDirector = GetCameraDirectorAs<UStateTreeCameraDirector>();
+	const FStateTreeReference& StateTreeReference = StateTreeDirector->StateTreeReference;
+	const UStateTree* StateTree = StateTreeReference.GetStateTree();
+
+	UObject* ContextOwner = Params.OwnerContext->GetOwner();
+	if (!ContextOwner)
+	{
+		UE_LOG(LogCameraSystem, Error,
+			TEXT("Can't deactivate camera director '%s': the evaluation context doesn't have a valid owner."),
+			*GetNameSafe(StateTreeDirector));
+		return;
+	}
+
+	if (!StateTree)
+	{
+		UE_LOG(LogCameraSystem, Error,
+			TEXT("Can't deactivate camera director '%s': it doesn't have a valid StateTree asset specified."),
+			*GetNameSafe(StateTreeDirector));
+		return;
+	}
+	
+	FStateTreeExecutionContext StateTreeContext(*ContextOwner, *StateTree, StateTreeInstanceData);
+
+	if (SetContextRequirements(Params.OwnerContext, StateTreeContext))
+	{
+		StateTreeContext.Stop();
+	}
+}
+
+void FStateTreeCameraDirectorEvaluator::OnRun(const FCameraDirectorEvaluationParams& Params, FCameraDirectorEvaluationResult& OutResult)
+{
+	const UStateTreeCameraDirector* StateTreeDirector = GetCameraDirectorAs<UStateTreeCameraDirector>();
+	const FStateTreeReference& StateTreeReference = StateTreeDirector->StateTreeReference;
+	const UStateTree* StateTree = StateTreeReference.GetStateTree();
+
+	UObject* ContextOwner = Params.OwnerContext->GetOwner();
+
+	if (!StateTree || !ContextOwner)
+	{
+		// Fail silently... we already emitted errors during OnActivate.
+		return;
+	}
+	
+	FStateTreeExecutionContext StateTreeContext(*ContextOwner, *StateTree, StateTreeInstanceData);
+
+	if (SetContextRequirements(Params.OwnerContext, StateTreeContext))
+	{
+		StateTreeContext.Tick(Params.DeltaTime);
+
+		for (TObjectPtr<const UCameraRigAsset> CameraRig : EvaluationData.ActiveCameraRigs)
+		{
+			FActiveCameraRigInfo CameraRigInfo;
+			CameraRigInfo.CameraRig = CameraRig;
+			CameraRigInfo.EvaluationContext = Params.OwnerContext;
+			OutResult.ActiveCameraRigs.Add(CameraRigInfo);
+		}
+	}
+}
+
+bool FStateTreeCameraDirectorEvaluator::SetContextRequirements(TSharedPtr<const FCameraEvaluationContext> OwnerContext, FStateTreeExecutionContext& StateTreeContext)
+{
+	UObject* ContextOwner = OwnerContext->GetOwner();
+	StateTreeContext.SetContextDataByName(
+			FStateTreeContextDataNames::ContextOwner, 
+			FStateTreeDataView(ContextOwner));
+
+	EvaluationData.ActiveCameraRigs.Reset();
+
+	StateTreeContext.SetCollectExternalDataCallback(FOnCollectStateTreeExternalData::CreateLambda(
+		[this](const FStateTreeExecutionContext& Context, const UStateTree* StateTree, 
+			TArrayView<const FStateTreeExternalDataDesc> ExternalDescs, TArrayView<FStateTreeDataView> OutDataViews)
+		{
+			for (int32 Index = 0; Index < ExternalDescs.Num(); Index++)
+			{
+				const FStateTreeExternalDataDesc& ExternalDesc(ExternalDescs[Index]);
+				if (ExternalDesc.Struct == FCameraDirectorStateTreeEvaluationData::StaticStruct())
+				{
+					OutDataViews[Index] = FStateTreeDataView(FStructView::Make(EvaluationData));
+				}
+			}
+			return true;
+		}));
+
+	return true;
+}
+
+void FStateTreeCameraDirectorEvaluator::OnAddReferencedObjects(FReferenceCollector& Collector)
+{
+	StateTreeInstanceData.AddStructReferencedObjects(Collector);
+}
+
+}  // namespace UE::Cameras
+
+UStateTreeCameraDirector::UStateTreeCameraDirector(const FObjectInitializer& ObjectInit)
+	: Super(ObjectInit)
+{
+}
+
+FCameraDirectorEvaluatorPtr UStateTreeCameraDirector::OnBuildEvaluator(FCameraDirectorEvaluatorBuilder& Builder) const
+{
+	using namespace UE::Cameras;
+	return Builder.BuildEvaluator<FStateTreeCameraDirectorEvaluator>();
+}
+
+void UStateTreeCameraDirector::OnBuildCameraDirector(UE::Cameras::FCameraBuildLog& BuildLog)
+{
+}
+
+#if WITH_EDITOR
+
+void UStateTreeCameraDirector::OnFactoryCreateAsset(const FCameraDirectorFactoryCreateParams& InParams)
+{
+}
+
+#endif
+
+#undef LOCTEXT_NAMESPACE
+
diff --git a/Engine/Plugins/Cameras/GameplayCameras/Source/GameplayCameras/Private/Directors/StateTreeCameraDirectorTasks.cpp b/Engine/Plugins/Cameras/GameplayCameras/Source/GameplayCameras/Private/Directors/StateTreeCameraDirectorTasks.cpp
new file mode 100644
--- /dev/null
+++ b/Engine/Plugins/Cameras/GameplayCameras/Source/GameplayCameras/Private/Directors/StateTreeCameraDirectorTasks.cpp
@@ -0,0 +1,47 @@
+// Copyright Epic Games, Inc. All Rights Reserved.
+
+#include "Directors/StateTreeCameraDirectorTasks.h"
+
+#include "StateTreeExecutionContext.h"
+#include "StateTreeLinker.h"
+
+#include UE_INLINE_GENERATED_CPP_BY_NAME(StateTreeCameraDirectorTasks)
+
+#define LOCTEXT_NAMESPACE "StateTreeCameraDirectorTasks"
+
+bool FGameplayCamerasActivateCameraRigTask::Link(FStateTreeLinker& Linker)
+{
+	Linker.LinkExternalData(CameraDirectorEvaluationDataHandle);
+	return true;
+}
+
+EStateTreeRunStatus FGameplayCamerasActivateCameraRigTask::EnterState(FStateTreeExecutionContext& Context, const FStateTreeTransitionResult& Transition) const
+{
+	if (UpdateResult(Context))
+	{
+		return bRunOnce ? EStateTreeRunStatus::Succeeded : EStateTreeRunStatus::Running;
+	}
+	return EStateTreeRunStatus::Failed;
+}
+
+EStateTreeRunStatus FGameplayCamerasActivateCameraRigTask::Tick(FStateTreeExecutionContext& Context, const float DeltaTime) const
+{
+	UpdateResult(Context);
+
+	return EStateTreeRunStatus::Running;
+}
+
+bool FGameplayCamerasActivateCameraRigTask::UpdateResult(FStateTreeExecutionContext& Context) const
+{
+	FInstanceDataType& InstanceData = Context.GetInstanceData(*this);
+	if (InstanceData.CameraRig)
+	{
+		FCameraDirectorStateTreeEvaluationData& EvaluationData = Context.GetExternalData(CameraDirectorEvaluationDataHandle);
+		EvaluationData.ActiveCameraRigs.Add(InstanceData.CameraRig);
+		return true;
+	}
+	return false;
+}
+
+#undef LOCTEXT_NAMESPACE
+
diff --git a/Engine/Plugins/Cameras/GameplayCameras/Source/GameplayCameras/Public/Directors/CameraDirectorStateTreeSchema.h b/Engine/Plugins/Cameras/GameplayCameras/Source/GameplayCameras/Public/Directors/CameraDirectorStateTreeSchema.h
new file mode 100644
--- /dev/null
+++ b/Engine/Plugins/Cameras/GameplayCameras/Source/GameplayCameras/Public/Directors/CameraDirectorStateTreeSchema.h
@@ -0,0 +1,78 @@
+// Copyright Epic Games, Inc. All Rights Reserved.
+
+#pragma once
+
+#include "Containers/Array.h"
+#include "StateTreeConditionBase.h"
+#include "StateTreeSchema.h"
+#include "StateTreeTaskBase.h"
+#include "StateTreeTypes.h"
+#include "UObject/NameTypes.h"
+#include "UObject/ObjectPtr.h"
+
+#include "CameraDirectorStateTreeSchema.generated.h"
+
+class UCameraRigAsset;
+struct FStateTreeExternalDataDesc;
+
+namespace UE::Cameras
+{
+
+struct FStateTreeContextDataNames
+{
+	const static FName ContextOwner;
+};
+
+}  // namespace UE::Cameras
+
+/**
+ * The schema of the StateTree for a StateTree camera director.
+ */
+UCLASS(BlueprintType, EditInlineNew, CollapseCategories, meta=(DisplayName="Gameplay Camera Director"))
+class GAMEPLAYCAMERAS_API UCameraDirectorStateTreeSchema : public UStateTreeSchema
+{
+	GENERATED_BODY()
+
+public:
+
+	UCameraDirectorStateTreeSchema();
+
+protected:
+
+	// UStateTreeSchema interface.
+	virtual bool IsStructAllowed(const UScriptStruct* InScriptStruct) const override;
+	virtual bool IsClassAllowed(const UClass* InScriptStruct) const override;
+	virtual bool IsExternalItemAllowed(const UStruct& InStruct) const override;
+	virtual TConstArrayView<FStateTreeExternalDataDesc> GetContextDataDescs() const override { return ContextDataDescs; }
+
+private:
+
+	UPROPERTY()
+	TArray<FStateTreeExternalDataDesc> ContextDataDescs;
+};
+
+/** The evaluation data for the StateTree camera director. */
+USTRUCT(BlueprintType)
+struct GAMEPLAYCAMERAS_API FCameraDirectorStateTreeEvaluationData
+{
+	GENERATED_BODY()
+
+	/** Camera rigs activated during a StateTree's execution frame. */
+	UPROPERTY()
+	TArray<TObjectPtr<UCameraRigAsset>> ActiveCameraRigs;
+};
+
+/** Base classs for camera director StateTree tasks. */
+USTRUCT(meta = (Hidden))
+struct GAMEPLAYCAMERAS_API FGameplayCamerasStateTreeTask : public FStateTreeTaskBase
+{
+	GENERATED_BODY()
+};
+
+/** Base classs for camera director StateTree conditions. */
+USTRUCT(meta = (Hidden))
+struct GAMEPLAYCAMERAS_API FGameplayCamerasStateTreeCondition : public FStateTreeConditionBase
+{
+	GENERATED_BODY()
+};
+
diff --git a/Engine/Plugins/Cameras/GameplayCameras/Source/GameplayCameras/Public/Directors/StateTreeCameraDirector.h b/Engine/Plugins/Cameras/GameplayCameras/Source/GameplayCameras/Public/Directors/StateTreeCameraDirector.h
new file mode 100644
--- /dev/null
+++ b/Engine/Plugins/Cameras/GameplayCameras/Source/GameplayCameras/Public/Directors/StateTreeCameraDirector.h
@@ -0,0 +1,39 @@
+// Copyright Epic Games, Inc. All Rights Reserved.
+
+#pragma once
+
+#include "Core/CameraDirector.h"
+#include "StateTreeReference.h"
+
+#include "StateTreeCameraDirector.generated.h"
+
+/**
+ * A camera director that runs a StateTree to specify which camera rigs should be active
+ * any given frame.
+ */
+UCLASS(MinimalAPI, EditInlineNew)
+class UStateTreeCameraDirector : public UCameraDirector
+{
+	GENERATED_BODY()
+
+public:
+
+	UStateTreeCameraDirector(const FObjectInitializer& ObjectInit);
+
+protected:
+
+	// UCameraDirector interface.
+	virtual FCameraDirectorEvaluatorPtr OnBuildEvaluator(FCameraDirectorEvaluatorBuilder& Builder) const override;
+	virtual void OnBuildCameraDirector(UE::Cameras::FCameraBuildLog& BuildLog) override;
+#if WITH_EDITOR
+	virtual void OnFactoryCreateAsset(const FCameraDirectorFactoryCreateParams& InParams) override;
+#endif
+
+public:
+
+	/** The StateTree to execute. Must have been created with the CameraDirectorStateTreeSchema. */
+	UPROPERTY(EditAnywhere, Category="StateTree",
+			meta=(Schema="/Script/GameplayCameras.CameraDirectorStateTreeSchema"))
+	FStateTreeReference StateTreeReference;
+};
+
diff --git a/Engine/Plugins/Cameras/GameplayCameras/Source/GameplayCameras/Public/Directors/StateTreeCameraDirectorTasks.h b/Engine/Plugins/Cameras/GameplayCameras/Source/GameplayCameras/Public/Directors/StateTreeCameraDirectorTasks.h
new file mode 100644
--- /dev/null
+++ b/Engine/Plugins/Cameras/GameplayCameras/Source/GameplayCameras/Public/Directors/StateTreeCameraDirectorTasks.h
@@ -0,0 +1,53 @@
+// Copyright Epic Games, Inc. All Rights Reserved.
+
+#pragma once
+
+#include "Core/CameraRigAsset.h"
+#include "Directors/CameraDirectorStateTreeSchema.h"
+
+#include "StateTreeCameraDirectorTasks.generated.h"
+
+/** Instance data for the "Activate Camera Rig" task. */
+USTRUCT()
+struct FGameplayCamerasActivateCameraRigTaskInstanceData
+{
+	GENERATED_BODY()
+
+	/** The camera rig to activate. */
+	UPROPERTY(EditAnywhere, Category="Cameras", meta=(UseCameraRigPicker=true))
+	TObjectPtr<UCameraRigAsset> CameraRig;
+};
+
+/**
+ * A task that activates a given camera rig inside a StateTreeCameraDirector.
+ */
+USTRUCT(meta=(DisplayName="Activate Camera Rig", Category="Cameras"))
+struct FGameplayCamerasActivateCameraRigTask : public FGameplayCamerasStateTreeTask
+{
+	GENERATED_BODY()
+
+	using FInstanceDataType = FGameplayCamerasActivateCameraRigTaskInstanceData;
+
+protected:
+
+	// FStateTreeTaskBase interface.
+	virtual const UStruct* GetInstanceDataType() const override { return FInstanceDataType::StaticStruct(); }
+	virtual bool Link(FStateTreeLinker& Linker) override;
+	virtual EStateTreeRunStatus EnterState(FStateTreeExecutionContext& Context, const FStateTreeTransitionResult& Transition) const override;
+	virtual EStateTreeRunStatus Tick(FStateTreeExecutionContext& Context, const float DeltaTime) const override;
+
+private:
+
+	bool UpdateResult(FStateTreeExecutionContext& Context) const;
+
+private:
+
+	TStateTreeExternalDataHandle<FCameraDirectorStateTreeEvaluationData> CameraDirectorEvaluationDataHandle;
+
+public:
+
+	/** If true, the task will complete immediately. If false, the task will run until a transition triggers. */
+	UPROPERTY(EditAnywhere, Category="State Tree")
+	bool bRunOnce = false;
+};
+
diff --git a/Engine/Plugins/Cameras/GameplayCameras/Source/GameplayCamerasEditor/Private/Customizations/CameraRigPtrDetailsCustomization.cpp b/Engine/Plugins/Cameras/GameplayCameras/Source/GameplayCamerasEditor/Private/Customizations/CameraRigPtrDetailsCustomization.cpp
new file mode 100644
--- /dev/null
+++ b/Engine/Plugins/Cameras/GameplayCameras/Source/GameplayCamerasEditor/Private/Customizations/CameraRigPtrDetailsCustomization.cpp
@@ -0,0 +1,163 @@
+// Copyright Epic Games, Inc. All Rights Reserved.
+
+#include "Customizations/CameraRigPtrDetailsCustomization.h"
+
+#include "AssetRegistry/AssetRegistryModule.h"
+#include "Core/CameraAsset.h"
+#include "Core/CameraDirector.h"
+#include "Core/CameraRigAsset.h"
+#include "DetailCategoryBuilder.h"
+#include "DetailWidgetRow.h"
+#include "Editors/CameraRigPickerConfig.h"
+#include "Helpers/CameraDirectorHelper.h"
+#include "IDetailChildrenBuilder.h"
+#include "IGameplayCamerasEditorModule.h"
+#include "PropertyEditorModule.h"
+#include "Widgets/Input/SComboButton.h"
+
+#define LOCTEXT_NAMESPACE "CameraRigPtrDetailsCustomization"
+
+namespace UE::Cameras
+{
+
+TSharedRef<IPropertyTypeCustomization> FCameraRigPtrDetailsCustomization::MakeInstance()
+{
+	return MakeShared<FCameraRigPtrDetailsCustomization>();
+}
+
+void FCameraRigPtrDetailsCustomization::CustomizeHeader(TSharedRef<IPropertyHandle> StructPropertyHandle, FDetailWidgetRow& HeaderRow, IPropertyTypeCustomizationUtils& StructCustomizationUtils)
+{
+	CameraRigPropertyHandle = StructPropertyHandle;
+
+	FProperty* StructProperty = StructPropertyHandle->GetProperty();
+	const bool bUseCameraRigPicker = StructProperty->GetBoolMetaData("UseCameraRigPicker");
+
+	TSharedPtr<SWidget> ValueContentWidget;
+	if (bUseCameraRigPicker)
+	{
+		CameraRigPickerButton = SNew(SComboButton)
+		.ButtonStyle(FAppStyle::Get(), "PropertyEditor.AssetComboStyle")
+		.ContentPadding(FMargin(2.f, 2.f, 2.f, 1.f))
+		.ButtonContent()
+		[
+			SNew(STextBlock)
+			.TextStyle(FAppStyle::Get(), "PropertyEditor.AssetClass")
+			.Font(FAppStyle::GetFontStyle("PropertyWindow.NormalFont"))
+			.Text(this, &FCameraRigPtrDetailsCustomization::OnGetComboText)
+			.ToolTipText(this, &FCameraRigPtrDetailsCustomization::OnGetComboToolTipText)
+		]
+		.OnGetMenuContent(this, &FCameraRigPtrDetailsCustomization::OnBuildCameraRigNamePicker);
+
+		ValueContentWidget = CameraRigPickerButton;
+	}
+	else
+	{
+		ValueContentWidget = StructPropertyHandle->CreatePropertyValueWidget();
+	}
+
+	HeaderRow
+		.NameContent()
+		[
+			StructPropertyHandle->CreatePropertyNameWidget()
+		]
+		.ValueContent()
+		[
+			ValueContentWidget.ToSharedRef()
+		];
+}
+
+void FCameraRigPtrDetailsCustomization::CustomizeChildren(TSharedRef<IPropertyHandle> StructPropertyHandle, IDetailChildrenBuilder& StructBuilder, IPropertyTypeCustomizationUtils& StructCustomizationUtils)
+{
+}
+
+FText FCameraRigPtrDetailsCustomization::OnGetComboText() const
+{
+	UObject* Value;
+	FPropertyAccess::Result PropertyAccess = CameraRigPropertyHandle->GetValue(Value);
+	if (PropertyAccess == FPropertyAccess::Success)
+	{
+		if (const UCameraRigAsset* CameraRig = Cast<const UCameraRigAsset>(Value))
+		{
+			return FText::FromString(CameraRig->GetDisplayName());
+		}
+	}
+	else if (PropertyAccess == FPropertyAccess::MultipleValues)
+	{
+		return LOCTEXT("MultipleValues", "Multiple values");
+	}
+	return LOCTEXT("ErrorValue", "Error reading camera rig value");
+}
+
+FText FCameraRigPtrDetailsCustomization::OnGetComboToolTipText() const
+{
+	return LOCTEXT("ComboToolTipText", "The name of the camera rig to activate.");
+}
+
+TSharedRef<SWidget> FCameraRigPtrDetailsCustomization::OnBuildCameraRigNamePicker()
+{
+	FCameraRigPickerConfig CameraRigPickerConfig;
+	CameraRigPickerConfig.bCanSelectCameraAsset = false;
+	CameraRigPickerConfig.bFocusCameraRigSearchBoxWhenOpened = true;
+	CameraRigPickerConfig.OnCameraRigSelected = FOnCameraRigSelected::CreateSP(this, &FCameraRigPtrDetailsCustomization::OnPickerAssetSelected);
+
+	TArray<UObject*> OuterObjects;
+	CameraRigPropertyHandle->GetOuterObjects(OuterObjects);
+
+	TSet<UObject*> CameraDirectorObjects;
+	for (UObject* OuterObject : OuterObjects)
+	{
+		if (UObject* OutermostObject = OuterObject->GetOutermostObject())
+		{
+			CameraDirectorObjects.Add(OutermostObject);
+		}
+	}
+
+	TArray<UCameraAsset*> ReferencingCameraAssets;
+	for (UObject* CameraDirectorObject : CameraDirectorObjects)
+	{
+		FCameraDirectorHelper::GetReferencingCameraAssets(CameraDirectorObject, ReferencingCameraAssets);
+	}
+
+	if (ReferencingCameraAssets.Num() == 0)
+	{
+		CameraRigPickerConfig.WarningMessage = LOCTEXT("NoReferencingCameraAssetWarning",
+				"No camera asset references this asset, so no camera rig list can be displayed. "
+				"Make a camera asset use this asset as its camera director evaluator, or use "
+				"ActivateCameraRigViaProxy.");
+	}
+	else
+	{
+		CameraRigPickerConfig.InitialCameraAssetSelection = ReferencingCameraAssets[0];
+
+		if (ReferencingCameraAssets.Num() > 1)
+		{
+			CameraRigPickerConfig.WarningMessage = LOCTEXT("ManyReferencingCameraAssetsWarning",
+				"More than one camera asset references this asset. Only camera rigs from the first "
+				"one will be displayed. Even then, shared camera director assets should use "
+				"ActivateCameraRigViaProxy instead.");
+		}
+	}
+
+	UObject* SelectedCameraRig;
+	if (CameraRigPropertyHandle->GetValue(SelectedCameraRig) == FPropertyAccess::Success)
+	{
+		CameraRigPickerConfig.InitialCameraRigSelection = Cast<UCameraRigAsset>(SelectedCameraRig);
+	}
+
+	IGameplayCamerasEditorModule& CamerasEditorModule = IGameplayCamerasEditorModule::Get();
+	return CamerasEditorModule.CreateCameraRigPicker(CameraRigPickerConfig);
+}
+
+void FCameraRigPtrDetailsCustomization::OnPickerAssetSelected(UCameraRigAsset* SelectedItem)
+{
+	if (SelectedItem)
+	{
+		CameraRigPickerButton->SetIsOpen(false);
+		CameraRigPropertyHandle->SetValue(SelectedItem);
+	}
+}
+
+}  // namespace UE::Cameras
+
+#undef LOCTEXT_NAMESPACE
+
diff --git a/Engine/Plugins/Cameras/GameplayCameras/Source/GameplayCamerasEditor/Private/GameplayCamerasEditorModule.cpp b/Engine/Plugins/Cameras/GameplayCameras/Source/GameplayCamerasEditor/Private/GameplayCamerasEditorModule.cpp
--- a/Engine/Plugins/Cameras/GameplayCameras/Source/GameplayCamerasEditor/Private/GameplayCamerasEditorModule.cpp
+++ b/Engine/Plugins/Cameras/GameplayCameras/Source/GameplayCamerasEditor/Private/GameplayCamerasEditorModule.cpp
@@ -14,6 +14,7 @@
 #include "Commands/GameplayCamerasDebuggerCommands.h"
 #include "Customizations/CameraParameterDetailsCustomizations.h"
 #include "Customizations/CameraRigCameraNodeDetailsCustomization.h"
+#include "Customizations/CameraRigPtrDetailsCustomization.h"
 #include "Customizations/CameraProxyTableDetailsCustomization.h"
 #include "Customizations/SingleCameraDirectorDetailsCustomization.h"
 #include "Debug/CameraDebugCategories.h"
@@ -459,13 +460,18 @@ private:
 		using namespace UE::Cameras;
 
 		FPropertyEditorModule& PropertyEditorModule = FModuleManager::GetModuleChecked<FPropertyEditorModule>("PropertyEditor");
+
 		FCameraParameterDetailsCustomization::Register(PropertyEditorModule);
+
 		PropertyEditorModule.RegisterCustomClassLayout("CameraRigCameraNode", FOnGetDetailCustomizationInstance::CreateStatic(
 					&FCameraRigCameraNodeDetailsCustomization::MakeInstance));
-		PropertyEditorModule.RegisterCustomPropertyTypeLayout("CameraRigProxyTableEntry", FOnGetPropertyTypeCustomizationInstance::CreateStatic(
-					&FCameraProxyTableEntryDetailsCustomization::MakeInstance));
 		PropertyEditorModule.RegisterCustomClassLayout("SingleCameraDirector", FOnGetDetailCustomizationInstance::CreateStatic(
 					&FSingleCameraDirectorDetailsCustomization::MakeInstance));
+
+		PropertyEditorModule.RegisterCustomPropertyTypeLayout("CameraRigProxyTableEntry", FOnGetPropertyTypeCustomizationInstance::CreateStatic(
+					&FCameraProxyTableEntryDetailsCustomization::MakeInstance));
+		PropertyEditorModule.RegisterCustomPropertyTypeLayout("CameraRigAsset", FOnGetPropertyTypeCustomizationInstance::CreateStatic(
+					&FCameraRigPtrDetailsCustomization::MakeInstance));
 	}
 
 	void UnregisterDetailsCustomizations()
diff --git a/Engine/Plugins/Cameras/GameplayCameras/Source/GameplayCamerasEditor/Public/Customizations/CameraRigPtrDetailsCustomization.h b/Engine/Plugins/Cameras/GameplayCameras/Source/GameplayCamerasEditor/Public/Customizations/CameraRigPtrDetailsCustomization.h
new file mode 100644
--- /dev/null
+++ b/Engine/Plugins/Cameras/GameplayCameras/Source/GameplayCamerasEditor/Public/Customizations/CameraRigPtrDetailsCustomization.h
@@ -0,0 +1,44 @@
+// Copyright Epic Games, Inc. All Rights Reserved.
+
+#pragma once
+
+#include "IPropertyTypeCustomization.h"
+
+class FDetailWidgetRow;
+class IDetailChildrenBuilder;
+class IPropertyHandle;
+class SComboButton;
+class SWidget;
+class UCameraRigAsset;
+struct FAssetData;
+
+namespace UE::Cameras
+{
+
+class FCameraRigPtrDetailsCustomization : public IPropertyTypeCustomization
+{
+public:
+
+	static TSharedRef<IPropertyTypeCustomization> MakeInstance();
+
+public:
+
+	// IPropertyTypeCustomization interface
+	virtual void CustomizeHeader(TSharedRef<IPropertyHandle> StructPropertyHandle, FDetailWidgetRow& HeaderRow, IPropertyTypeCustomizationUtils& StructCustomizationUtils) override;
+	virtual void CustomizeChildren(TSharedRef<IPropertyHandle> StructPropertyHandle, IDetailChildrenBuilder& StructBuilder, IPropertyTypeCustomizationUtils& StructCustomizationUtils) override;
+
+private:
+
+	FText OnGetComboText() const;
+	FText OnGetComboToolTipText() const;
+	TSharedRef<SWidget> OnBuildCameraRigNamePicker();
+	void OnPickerAssetSelected(UCameraRigAsset* SelectedItem);
+
+private:
+
+	TSharedPtr<IPropertyHandle> CameraRigPropertyHandle;
+	TSharedPtr<SComboButton> CameraRigPickerButton;
+};
+
+}  // namespace UE::Cameras
+
