Cameras: add a Blueprint camera node

This allows designers to prototype camera behaviors in Blueprint.

diff --git a/Engine/Plugins/Cameras/GameplayCameras/Source/GameplayCameras/Private/Nodes/Utility/BlueprintCameraNode.cpp b/Engine/Plugins/Cameras/GameplayCameras/Source/GameplayCameras/Private/Nodes/Utility/BlueprintCameraNode.cpp
new file mode 100644
--- /dev/null
+++ b/Engine/Plugins/Cameras/GameplayCameras/Source/GameplayCameras/Private/Nodes/Utility/BlueprintCameraNode.cpp
@@ -0,0 +1,103 @@
+// Copyright Epic Games, Inc. All Rights Reserved.
+
+#include "Nodes/Utility/BlueprintCameraNode.h"
+
+#include "Core/CameraBuildLog.h"
+#include "Core/CameraEvaluationContext.h"
+#include "Core/CameraNodeEvaluator.h"
+#include "Core/CameraRigBuildContext.h"
+#include "Core/CameraVariableTable.h"
+#include "GameplayCameras.h"
+#include "Templates/UnrealTemplate.h"
+
+#include UE_INLINE_GENERATED_CPP_BY_NAME(BlueprintCameraNode)
+
+#define LOCTEXT_NAMESPACE "BlueprintCameraNode"
+
+namespace UE::Cameras
+{
+
+class FBlueprintCameraNodeEvaluator : public FCameraNodeEvaluator
+{
+	UE_DECLARE_CAMERA_NODE_EVALUATOR(GAMEPLAYCAMERAS_API, FBlueprintCameraNodeEvaluator)
+
+protected:
+
+	// FCameraNodeEvaluator interface.
+	virtual void OnInitialize(const FCameraNodeEvaluatorInitializeParams& Params, FCameraNodeEvaluationResult& OutResult) override;
+	virtual void OnRun(const FCameraNodeEvaluationParams& Params, FCameraNodeEvaluationResult& OutResult) override;
+	virtual void OnAddReferencedObjects(FReferenceCollector& Collector) override;
+
+private:
+
+	TObjectPtr<UBlueprintCameraNodeEvaluator> EvaluatorBlueprint;
+};
+
+UE_DEFINE_CAMERA_NODE_EVALUATOR(FBlueprintCameraNodeEvaluator)
+
+void FBlueprintCameraNodeEvaluator::OnInitialize(const FCameraNodeEvaluatorInitializeParams& Params, FCameraNodeEvaluationResult& OutResult)
+{
+	const UBlueprintCameraNode* BlueprintNode = GetCameraNodeAs<UBlueprintCameraNode>();
+	if (!ensure(BlueprintNode))
+	{
+		return;
+	}
+
+	if (BlueprintNode->CameraNodeEvaluatorClass)
+	{
+		UObject* Outer = Params.EvaluationContext->GetOwner();
+		EvaluatorBlueprint = NewObject<UBlueprintCameraNodeEvaluator>(Outer, BlueprintNode->CameraNodeEvaluatorClass);
+	}
+	else
+	{
+		UE_LOG(LogCameraSystem, Error, TEXT("No Blueprint class set on camera node '%s'."), *GetNameSafe(BlueprintNode));
+	}
+}
+
+void FBlueprintCameraNodeEvaluator::OnRun(const FCameraNodeEvaluationParams& Params, FCameraNodeEvaluationResult& OutResult)
+{
+	if (EvaluatorBlueprint)
+	{
+		EvaluatorBlueprint->NativeRunCameraNode(Params, OutResult);
+	}
+}
+
+void FBlueprintCameraNodeEvaluator::OnAddReferencedObjects(FReferenceCollector& Collector)
+{
+	Collector.AddReferencedObject(EvaluatorBlueprint);
+}
+
+}  // namespace UE::Cameras
+
+void UBlueprintCameraNodeEvaluator::NativeRunCameraNode(const FCameraNodeEvaluationParams& Params, FCameraNodeEvaluationResult& OutResult)
+{
+	bIsFirstFrame = Params.bIsFirstFrame;
+	EvaluationContextOwner = Params.EvaluationContext->GetOwner();
+	CameraPose = FBlueprintCameraPose::FromCameraPose(OutResult.CameraPose);
+
+	ensure(CurrentResult == nullptr);
+	TGuardValue<FCameraNodeEvaluationResult*> CurrentResultGuard(CurrentResult, &OutResult);
+
+	TickCameraNode(Params.DeltaTime);
+
+	CameraPose.ApplyTo(OutResult.CameraPose);
+}
+
+void UBlueprintCameraNode::OnBuild(FCameraRigBuildContext& BuildContext)
+{
+	if (!CameraNodeEvaluatorClass)
+	{
+		BuildContext.BuildLog.AddMessage(
+				EMessageSeverity::Error, this,
+				LOCTEXT("MissingBlueprintClass", "No evaluator Blueprint class is set."));
+	}
+}
+
+FCameraNodeEvaluatorPtr UBlueprintCameraNode::OnBuildEvaluator(FCameraNodeEvaluatorBuilder& Builder) const
+{
+	using namespace UE::Cameras;
+	return Builder.BuildEvaluator<FBlueprintCameraNodeEvaluator>();
+}
+
+#undef LOCTEXT_NAMESPACE
+
diff --git a/Engine/Plugins/Cameras/GameplayCameras/Source/GameplayCameras/Public/Nodes/Utility/BlueprintCameraNode.h b/Engine/Plugins/Cameras/GameplayCameras/Source/GameplayCameras/Public/Nodes/Utility/BlueprintCameraNode.h
new file mode 100644
--- /dev/null
+++ b/Engine/Plugins/Cameras/GameplayCameras/Source/GameplayCameras/Public/Nodes/Utility/BlueprintCameraNode.h
@@ -0,0 +1,87 @@
+// Copyright Epic Games, Inc. All Rights Reserved.
+
+#pragma once
+
+#include "Core/CameraNode.h"
+#include "Core/CameraRigJoints.h"
+#include "Core/CameraVariableAssets.h"
+#include "GameFramework/BlueprintCameraPose.h"
+#include "GameFramework/BlueprintCameraVariableTable.h"
+#include "Templates/SubclassOf.h"
+
+#include "BlueprintCameraNode.generated.h"
+
+namespace UE::Cameras
+{
+	class FCameraVariableTable;
+	struct FBlueprintCameraPose;
+	struct FCameraNodeEvaluationParams;
+	struct FCameraNodeEvaluationResult;
+};
+
+/**
+ * The base class for Blueprint camera node evaluators.
+ */
+UCLASS(MinimalAPI, Blueprintable, Abstract)
+class UBlueprintCameraNodeEvaluator : public UObject
+{
+	GENERATED_BODY()
+
+public:
+
+	/** The main execution callback for the camera node. Call SetCameraPose to affect the result. */
+	UFUNCTION(BlueprintImplementableEvent, Category="Evaluation")
+	void TickCameraNode(float DeltaTime);
+
+public:
+
+	using FCameraNodeEvaluationParams = UE::Cameras::FCameraNodeEvaluationParams;
+	using FCameraNodeEvaluationResult = UE::Cameras::FCameraNodeEvaluationResult;
+
+	/** Runs this camera node. */
+	void NativeRunCameraNode(const FCameraNodeEvaluationParams& Params, FCameraNodeEvaluationResult& OutResult);
+
+protected:
+
+	/** Whether this is the first frame of this camera node's lifetime. */
+	UPROPERTY(BlueprintReadOnly, Category="Evaluation")
+	bool bIsFirstFrame = false;
+
+	/** The owner object of this camera node's evaluation context. */
+	UPROPERTY(BlueprintReadOnly, Category="Evaluation")
+	TObjectPtr<UObject> EvaluationContextOwner;
+
+	/** The input/output camera pose for this frame. */
+	UPROPERTY(BlueprintReadWrite, Category="Evaluation")
+	FBlueprintCameraPose CameraPose;
+
+	/** The input/output camera variable table for this frame. */
+	UPROPERTY(BlueprintReadWrite, Category="Evaluation")
+	FBlueprintCameraVariableTable VariableTable;
+
+private:
+
+	FCameraNodeEvaluationResult* CurrentResult = nullptr;
+};
+
+/**
+ * A camera node that runs arbitrary Blueprint logic.
+ */
+UCLASS(MinimalAPI, meta=(CameraNodeCategories="Common,Transform"))
+class UBlueprintCameraNode : public UCameraNode
+{
+	GENERATED_BODY()
+
+protected:
+
+	// UCameraNode interface.
+	virtual void OnBuild(FCameraRigBuildContext& BuildContext) override;
+	virtual FCameraNodeEvaluatorPtr OnBuildEvaluator(FCameraNodeEvaluatorBuilder& Builder) const override;
+
+public:
+
+	/** The camera node evaluator class to instantiate and run. */
+	UPROPERTY(EditAnywhere, Category=Common)
+	TSubclassOf<UBlueprintCameraNodeEvaluator> CameraNodeEvaluatorClass;
+};
+
