Cameras: rename "array" node to "sequence" in the editor

diff --git a/Engine/Plugins/Cameras/GameplayCameras/Source/GameplayCameras/Public/Nodes/Common/ArrayCameraNode.h b/Engine/Plugins/Cameras/GameplayCameras/Source/GameplayCameras/Public/Nodes/Common/ArrayCameraNode.h
--- a/Engine/Plugins/Cameras/GameplayCameras/Source/GameplayCameras/Public/Nodes/Common/ArrayCameraNode.h
+++ b/Engine/Plugins/Cameras/GameplayCameras/Source/GameplayCameras/Public/Nodes/Common/ArrayCameraNode.h
@@ -10,7 +10,7 @@
 /**
  * A camera node that runs a list of other camera nodes.
  */
-UCLASS(MinimalAPI, meta=(CameraNodeCategories="Common,Utility"))
+UCLASS(MinimalAPI, meta=(DisplayName="Sequence", CameraNodeCategories="Common,Utility"))
 class UArrayCameraNode : public UCameraNode
 {
 	GENERATED_BODY()
