Cameras: move all config files to GameplayCameras.ini

diff --git a/Engine/Plugins/Cameras/GameplayCameras/Source/GameplayCameras/GameplayCameras.Build.cs b/Engine/Plugins/Cameras/GameplayCameras/Source/GameplayCameras/GameplayCameras.Build.cs
--- a/Engine/Plugins/Cameras/GameplayCameras/Source/GameplayCameras/GameplayCameras.Build.cs
+++ b/Engine/Plugins/Cameras/GameplayCameras/Source/GameplayCameras/GameplayCameras.Build.cs
@@ -14,6 +14,7 @@ public class GameplayCameras : ModuleRules
 				"CinematicCamera",
 				"Core",
 				"CoreUObject",
+				"DeveloperSettings",
 				"Engine",
 				"EnhancedInput",
 				"GameplayTags",
diff --git a/Engine/Plugins/Cameras/GameplayCameras/Source/GameplayCameras/Private/GameplayCamerasSettings.cpp b/Engine/Plugins/Cameras/GameplayCameras/Source/GameplayCameras/Private/GameplayCamerasSettings.cpp
new file mode 100644
--- /dev/null
+++ b/Engine/Plugins/Cameras/GameplayCameras/Source/GameplayCameras/Private/GameplayCamerasSettings.cpp
@@ -0,0 +1,11 @@
+// Copyright Epic Games, Inc. All Rights Reserved.
+
+#include "GameplayCamerasSettings.h"
+
+#include UE_INLINE_GENERATED_CPP_BY_NAME(GameplayCamerasSettings)
+
+FName UGameplayCamerasSettings::GetCategoryName() const
+{
+	return TEXT("Plugins");
+}
+
diff --git a/Engine/Plugins/Cameras/GameplayCameras/Source/GameplayCameras/Public/GameplayCamerasSettings.h b/Engine/Plugins/Cameras/GameplayCameras/Source/GameplayCameras/Public/GameplayCamerasSettings.h
--- a/Engine/Plugins/Cameras/GameplayCameras/Source/GameplayCameras/Public/GameplayCamerasSettings.h
+++ b/Engine/Plugins/Cameras/GameplayCameras/Source/GameplayCameras/Public/GameplayCamerasSettings.h
@@ -3,6 +3,7 @@
 #pragma once
 
 #include "CoreTypes.h"
+#include "Engine/DeveloperSettings.h"
 #include "UObject/ObjectMacros.h"
 #include "UObject/Object.h"
 
@@ -11,8 +12,8 @@
 /**
  * The settings for the Gameplay Cameras runtime.
  */
-UCLASS(config=Project, MinimalAPI)
-class UGameplayCamerasSettings : public UObject
+UCLASS(Config=GameplayCameras, DefaultConfig, MinimalAPI, meta=(DisplayName="Gameplay Cameras"))
+class UGameplayCamerasSettings : public UDeveloperSettings
 {
 	GENERATED_BODY()
 
@@ -33,5 +34,10 @@ public:
 	/** The distance below which any IK aiming operation is disabled. */
 	UPROPERTY(EditAnywhere, Config, Category="IK Aiming")
 	double DefaultIKAimingMinDistance = 100.0;  // 1m
+
+protected:
+
+	// UDeveloperSettings interface.
+	virtual FName GetCategoryName() const override;
 };
 
diff --git a/Engine/Plugins/Cameras/GameplayCameras/Source/GameplayCamerasEditor/GameplayCamerasEditor.Build.cs b/Engine/Plugins/Cameras/GameplayCameras/Source/GameplayCamerasEditor/GameplayCamerasEditor.Build.cs
--- a/Engine/Plugins/Cameras/GameplayCameras/Source/GameplayCamerasEditor/GameplayCamerasEditor.Build.cs
+++ b/Engine/Plugins/Cameras/GameplayCameras/Source/GameplayCamerasEditor/GameplayCamerasEditor.Build.cs
@@ -25,6 +25,7 @@ public class GameplayCamerasEditor : ModuleRules
 				"BlueprintGraph",
 				"Core",
 				"CoreUObject",
+				"DeveloperSettings",
 				"EditorFramework",
 				"EditorSubsystem",
 				"Engine",
diff --git a/Engine/Plugins/Cameras/GameplayCameras/Source/GameplayCamerasEditor/Private/GameplayCamerasEditorSettings.cpp b/Engine/Plugins/Cameras/GameplayCameras/Source/GameplayCamerasEditor/Private/GameplayCamerasEditorSettings.cpp
--- a/Engine/Plugins/Cameras/GameplayCameras/Source/GameplayCamerasEditor/Private/GameplayCamerasEditorSettings.cpp
+++ b/Engine/Plugins/Cameras/GameplayCameras/Source/GameplayCamerasEditor/Private/GameplayCamerasEditorSettings.cpp
@@ -15,3 +15,8 @@ UGameplayCamerasEditorSettings::UGameplayCamerasEditorSettings(const FObjectInit
 	CameraBlendNodeTitleColor = FLinearColor(0.6f, 0.0f, 1.0f, 1.0f);                // Purple
 }
 
+FName UGameplayCamerasEditorSettings::GetCategoryName() const
+{
+	return TEXT("Plugins");
+}
+
diff --git a/Engine/Plugins/Cameras/GameplayCameras/Source/GameplayCamerasEditor/Public/GameplayCamerasEditorSettings.h b/Engine/Plugins/Cameras/GameplayCameras/Source/GameplayCamerasEditor/Public/GameplayCamerasEditorSettings.h
--- a/Engine/Plugins/Cameras/GameplayCameras/Source/GameplayCamerasEditor/Public/GameplayCamerasEditorSettings.h
+++ b/Engine/Plugins/Cameras/GameplayCameras/Source/GameplayCamerasEditor/Public/GameplayCamerasEditorSettings.h
@@ -3,6 +3,7 @@
 #pragma once
 
 #include "CoreTypes.h"
+#include "Engine/DeveloperSettings.h"
 #include "UObject/ObjectMacros.h"
 #include "UObject/Object.h"
 #include "UObject/Package.h"
@@ -10,8 +11,8 @@
 
 #include "GameplayCamerasEditorSettings.generated.h"
 
-UCLASS(config=EditorPerProjectUserSettings)
-class GAMEPLAYCAMERASEDITOR_API UGameplayCamerasEditorSettings : public UObject
+UCLASS(Config=GameplayCameras, DefaultConfig, MinimalAPI, meta=(DisplayName="Gameplay Cameras Editor"))
+class UGameplayCamerasEditorSettings : public UDeveloperSettings
 {
 	GENERATED_BODY()
 
@@ -43,5 +44,10 @@ public:
 
 	UPROPERTY()
 	FName LastCameraAssetToolkitModeName;
+
+protected:
+
+	// UDeveloperSettings interface.
+	virtual FName GetCategoryName() const override;
 };
 
