Cameras: auto-spawn the camera system actor by default

The camera system actor is dynamically spawned and activated when a gameplay camera activates and doesn't find it.
This automatic spawning can be disabled by a setting.

diff --git a/Engine/Plugins/Cameras/GameplayCameras/Source/GameplayCameras/Private/GameFramework/GameplayCameraComponent.cpp b/Engine/Plugins/Cameras/GameplayCameras/Source/GameplayCameras/Private/GameFramework/GameplayCameraComponent.cpp
--- a/Engine/Plugins/Cameras/GameplayCameras/Source/GameplayCameras/Private/GameFramework/GameplayCameraComponent.cpp
+++ b/Engine/Plugins/Cameras/GameplayCameras/Source/GameplayCameras/Private/GameFramework/GameplayCameraComponent.cpp
@@ -14,6 +14,7 @@
 #include "GameFramework/GameplayCameraSystemActor.h"
 #include "GameFramework/GameplayCameraSystemComponent.h"
 #include "GameplayCameras.h"
+#include "GameplayCamerasSettings.h"
 #include "Kismet/GameplayStatics.h"
 #include "Logging/MessageLog.h"
 #include "TimerManager.h"
@@ -44,11 +45,16 @@ TSharedPtr<UE::Cameras::FCameraEvaluationContext> UGameplayCameraComponent::GetE
 	return EvaluationContext;
 }
 
-void UGameplayCameraComponent::ActivateCamera(int32 PlayerIndex)
+void UGameplayCameraComponent::ActivateCameraForPlayerIndex(int32 PlayerIndex)
 {
 	ActivateCameraEvaluationContext(PlayerIndex);
 }
 
+void UGameplayCameraComponent::ActivateCameraForPlayerController(APlayerController* PlayerController)
+{
+	ActivateCameraEvaluationContext(PlayerController);
+}
+
 void UGameplayCameraComponent::DeactivateCamera()
 {
 	DeactivateCameraEvaluationContext();
@@ -56,36 +62,49 @@ void UGameplayCameraComponent::DeactivateCamera()
 
 void UGameplayCameraComponent::ActivateCameraEvaluationContext(int32 PlayerIndex)
 {
-	if (ActivatedForPlayerIndex == PlayerIndex)
-	{
-		return;
-	}
-
-	if (ActivatedForPlayerIndex >= 0)
+	if (WeakPlayerController.IsValid())
 	{
 		DeactivateCameraEvaluationContext();
 	}
 
-	if (APlayerController* PC = UGameplayStatics::GetPlayerController(this, PlayerIndex))
+	APlayerController* PlayerController = UGameplayStatics::GetPlayerController(this, PlayerIndex);
+	if (!PlayerController)
 	{
-		ActivateCameraEvaluationContext(PC);
-		ActivatedForPlayerIndex = PlayerIndex;
+		UE_LOG(LogCameraSystem, Error, TEXT("Can't activate gameplay camera: no player controller found!"));
+		return;
 	}
+
+	ActivateCameraEvaluationContext(PlayerController);
 }
 
 void UGameplayCameraComponent::DeactivateCameraEvaluationContext()
 {
-	if (ActivatedForPlayerIndex < 0)
+	using namespace UE::Cameras;
+
+	APlayerController* PlayerController = WeakPlayerController.Get();
+	if (!PlayerController)
 	{
 		return;
 	}
 
-	if (APlayerController* PC = UGameplayStatics::GetPlayerController(this, ActivatedForPlayerIndex))
+	if (!ensure(PlayerController->PlayerCameraManager))
 	{
-		DeactivateCameraEvaluationContext(PC);
+		return;
+	}
+	
+	AGameplayCameraSystemActor* CameraSystem = Cast<AGameplayCameraSystemActor>(PlayerController->PlayerCameraManager->GetViewTarget());
+	if (!ensure(CameraSystem))
+	{
+		return;
+	}
+
+	if (EvaluationContext.IsValid())
+	{
+		TSharedPtr<FCameraSystemEvaluator> Evaluator = CameraSystem->GetCameraSystemComponent()->GetCameraSystemEvaluator();
+		Evaluator->RemoveEvaluationContext(EvaluationContext.ToSharedRef());
 	}
 
-	ActivatedForPlayerIndex = INDEX_NONE;
+	WeakPlayerController.Reset();
 }
 
 void UGameplayCameraComponent::ActivateCameraEvaluationContext(APlayerController* PlayerController)
@@ -107,6 +126,16 @@ void UGameplayCameraComponent::ActivateCameraEvaluationContext(APlayerController
 	
 	AGameplayCameraSystemActor* CameraSystem = Cast<AGameplayCameraSystemActor>(PlayerController->PlayerCameraManager->GetViewTarget());
 	if (!CameraSystem)
+	{
+		const UGameplayCamerasSettings* Settings = GetDefault<UGameplayCamerasSettings>();
+		if (Settings->bAutoSpawnCameraSystemActor)
+		{
+			FActorSpawnParameters SpawnParams;
+			CameraSystem = GetWorld()->SpawnActor<AGameplayCameraSystemActor>(SpawnParams);
+			CameraSystem->GetCameraSystemComponent()->ActivateCameraSystemForPlayerController(PlayerController);
+		}
+	}
+	if (!CameraSystem)
 	{
 		FGameDelegates::Get().GetViewTargetChangedDelegate().AddUObject(
 				this, &UGameplayCameraComponent::DelayedActivateCameraEvaluationContext);
@@ -121,6 +150,8 @@ void UGameplayCameraComponent::ActivateCameraEvaluationContext(APlayerController
 	}
 
 	DoActivateCameraEvaluationContext(PlayerController, CameraSystem->GetCameraSystemComponent()->GetCameraSystemEvaluator());
+
+	WeakPlayerController = PlayerController;
 }
 
 void UGameplayCameraComponent::DelayedActivateCameraEvaluationContext(APlayerController* PlayerController, AActor* OldViewTarget, AActor* NewViewTarget)
@@ -135,11 +166,13 @@ void UGameplayCameraComponent::DelayedActivateCameraEvaluationContext(APlayerCon
 	}
 
 	DoActivateCameraEvaluationContext(PlayerController, CameraSystem->GetCameraSystemComponent()->GetCameraSystemEvaluator());
+
+	WeakPlayerController = PlayerController;
 }
 
 void UGameplayCameraComponent::AbortDelayedActivateCameraEvaluationContext()
 {
-	if (ActivatedForPlayerIndex == INDEX_NONE)
+	if (!WeakPlayerController.IsValid())
 	{
 		UE_LOG(LogCameraSystem, Error, TEXT("Can't activate gameplay camera component: no camera system found after a full tick!"));
 		FGameDelegates::Get().GetViewTargetChangedDelegate().RemoveAll(this);
@@ -164,28 +197,6 @@ void UGameplayCameraComponent::DoActivateCameraEvaluationContext(APlayerControll
 	CameraSystemEvaluator->PushEvaluationContext(EvaluationContext.ToSharedRef());
 }
 
-void UGameplayCameraComponent::DeactivateCameraEvaluationContext(APlayerController* PlayerController)
-{
-	using namespace UE::Cameras;
-
-	if (!ensure(PlayerController && PlayerController->PlayerCameraManager))
-	{
-		return;
-	}
-	
-	AGameplayCameraSystemActor* CameraSystem = Cast<AGameplayCameraSystemActor>(PlayerController->PlayerCameraManager->GetViewTarget());
-	if (!ensure(CameraSystem))
-	{
-		return;
-	}
-
-	if (EvaluationContext.IsValid())
-	{
-		TSharedPtr<FCameraSystemEvaluator> Evaluator = CameraSystem->GetCameraSystemComponent()->GetCameraSystemEvaluator();
-		Evaluator->RemoveEvaluationContext(EvaluationContext.ToSharedRef());
-	}
-}
-
 FBlueprintCameraPose UGameplayCameraComponent::GetInitialPose() const
 {
 	return FBlueprintCameraPose::FromCameraPose(EvaluationContext->GetInitialResult().CameraPose);
@@ -241,7 +252,7 @@ void UGameplayCameraComponent::BeginPlay()
 	if (AutoActivateForPlayer != EAutoReceiveInput::Disabled && GetNetMode() != NM_DedicatedServer)
 	{
 		const int32 PlayerIndex = AutoActivateForPlayer.GetIntValue() - 1;
-		ActivateCamera(PlayerIndex);
+		ActivateCameraForPlayerIndex(PlayerIndex);
 	}
 }
 
diff --git a/Engine/Plugins/Cameras/GameplayCameras/Source/GameplayCameras/Private/GameFramework/GameplayCameraSystemComponent.cpp b/Engine/Plugins/Cameras/GameplayCameras/Source/GameplayCameras/Private/GameFramework/GameplayCameraSystemComponent.cpp
--- a/Engine/Plugins/Cameras/GameplayCameras/Source/GameplayCameras/Private/GameFramework/GameplayCameraSystemComponent.cpp
+++ b/Engine/Plugins/Cameras/GameplayCameras/Source/GameplayCameras/Private/GameFramework/GameplayCameraSystemComponent.cpp
@@ -90,14 +90,21 @@ void UGameplayCameraSystemComponent::OnRegister()
 #endif	// WITH_EDITORONLY_DATA
 }
 
-void UGameplayCameraSystemComponent::ActivateCameraSystem(int32 PlayerIndex)
+void UGameplayCameraSystemComponent::ActivateCameraSystemForPlayerIndex(int32 PlayerIndex)
 {
-	if (ActivatedForPlayerIndex == PlayerIndex)
+	APlayerController* PlayerController = UGameplayStatics::GetPlayerController(this, PlayerIndex);
+	if (!PlayerController)
 	{
+		UE_LOG(LogCameraSystem, Error, TEXT("Can't activate gameplay camera system: no player controller found!"));
 		return;
 	}
 
-	if (ActivatedForPlayerIndex >= 0)
+	ActivateCameraSystemForPlayerController(PlayerController);
+}
+
+void UGameplayCameraSystemComponent::ActivateCameraSystemForPlayerController(APlayerController* PlayerController)
+{
+	if (WeakPlayerController.IsValid())
 	{
 		DeactivateCameraSystem();
 	}
@@ -109,31 +116,20 @@ void UGameplayCameraSystemComponent::ActivateCameraSystem(int32 PlayerIndex)
 		return;
 	}
 
-	APlayerController* PC = UGameplayStatics::GetPlayerController(this, PlayerIndex);
-	if (!PC)
-	{
-		UE_LOG(LogCameraSystem, Error, TEXT("Can't activate gameplay camera system: no player controller found!"));
-		return;
-	}
-
-	PC->SetViewTarget(OwningActor);
-	ActivatedForPlayerIndex = PlayerIndex;
+	PlayerController->SetViewTarget(OwningActor);
+	WeakPlayerController = PlayerController;
 }
 
 void UGameplayCameraSystemComponent::DeactivateCameraSystem(AActor* NextViewTarget)
 {
-	if (ActivatedForPlayerIndex < 0)
+	APlayerController* PlayerController = WeakPlayerController.Get();
+	if (!PlayerController)
 	{
 		return;
 	}
 
-	APlayerController* PC = UGameplayStatics::GetPlayerController(this, ActivatedForPlayerIndex);
-	if (PC)
-	{
-		PC->SetViewTarget(NextViewTarget);
-	}
-
-	ActivatedForPlayerIndex = INDEX_NONE;
+	PlayerController->SetViewTarget(NextViewTarget);
+	WeakPlayerController.Reset();
 }
 
 void UGameplayCameraSystemComponent::BeginPlay()
@@ -143,7 +139,7 @@ void UGameplayCameraSystemComponent::BeginPlay()
 	if (AutoActivateForPlayer != EAutoReceiveInput::Disabled && GetNetMode() != NM_DedicatedServer)
 	{
 		const int32 PlayerIndex = AutoActivateForPlayer.GetIntValue() - 1;
-		ActivateCameraSystem(PlayerIndex);
+		ActivateCameraSystemForPlayerIndex(PlayerIndex);
 	}
 }
 
diff --git a/Engine/Plugins/Cameras/GameplayCameras/Source/GameplayCameras/Private/GameFramework/GameplayControlRotationComponent.cpp b/Engine/Plugins/Cameras/GameplayCameras/Source/GameplayCameras/Private/GameFramework/GameplayControlRotationComponent.cpp
--- a/Engine/Plugins/Cameras/GameplayCameras/Source/GameplayCameras/Private/GameFramework/GameplayControlRotationComponent.cpp
+++ b/Engine/Plugins/Cameras/GameplayCameras/Source/GameplayCameras/Private/GameFramework/GameplayControlRotationComponent.cpp
@@ -54,8 +54,8 @@ void UGameplayControlRotationComponent::BeginPlay()
 	}
 
 	// Find where the camera system is running from.
-	PlayerIndex = GameplayCameraComponent->GetPlayerIndex();
-	if (!ensure(PlayerIndex != INDEX_NONE))
+	PlayerController = GameplayCameraComponent->GetPlayerController();
+	if (!ensure(PlayerController))
 	{
 		UE_LOG(LogCameraSystem, Error,
 				TEXT("GameplayCameraComponent '%s' hasn't activated for any player yet"),
@@ -63,13 +63,6 @@ void UGameplayControlRotationComponent::BeginPlay()
 		return;
 	}
 
-	APlayerController* PlayerController = UGameplayStatics::GetPlayerController(this, PlayerIndex);
-	if (!PlayerController)
-	{
-		UE_LOG(LogCameraSystem, Error, TEXT("Can't find player controller for player index %d"), PlayerIndex);
-		return;
-	}
-
 	AActor* ViewTarget = PlayerController->PlayerCameraManager->GetViewTarget();
 	UGameplayCameraSystemComponent* CameraSystemComponent = ViewTarget->FindComponentByClass<UGameplayCameraSystemComponent>();
 	if (!CameraSystemComponent)
@@ -117,7 +110,6 @@ void UGameplayControlRotationComponent::TickComponent(float DeltaTime, ELevelTic
 {
 	Super::TickComponent(DeltaTime, TickType, ThisTickFunction);
 
-	APlayerController* PlayerController = GetPlayerController();
 	if (PlayerController)
 	{
 		// This may be technically one frame late (i.e. we set the control rotation computed 
@@ -127,14 +119,5 @@ void UGameplayControlRotationComponent::TickComponent(float DeltaTime, ELevelTic
 	}
 }
 
-APlayerController* UGameplayControlRotationComponent::GetPlayerController()
-{
-	if (PlayerIndex != INDEX_NONE)
-	{
-		return UGameplayStatics::GetPlayerController(this, PlayerIndex);
-	}
-	return nullptr;
-}
-
 #undef LOCTEXT_NAMESPACE
 
diff --git a/Engine/Plugins/Cameras/GameplayCameras/Source/GameplayCameras/Public/GameFramework/GameplayCameraComponent.h b/Engine/Plugins/Cameras/GameplayCameras/Source/GameplayCameras/Public/GameFramework/GameplayCameraComponent.h
--- a/Engine/Plugins/Cameras/GameplayCameras/Source/GameplayCameras/Public/GameFramework/GameplayCameraComponent.h
+++ b/Engine/Plugins/Cameras/GameplayCameras/Source/GameplayCameras/Public/GameFramework/GameplayCameraComponent.h
@@ -13,6 +13,7 @@
 
 #include "GameplayCameraComponent.generated.h"
 
+class APlayerController;
 class UCameraAsset;
 
 namespace UE::Cameras
@@ -39,19 +40,18 @@ public:
 	/** Get the camera evaluation context used by this component. */
 	TSharedPtr<UE::Cameras::FCameraEvaluationContext> GetEvaluationContext();
 
-	/** Get the player controller index that this component has been activated for. */
-	int32 GetPlayerIndex() const { return ActivatedForPlayerIndex; }
+	/** Get the player controller this component is currently activated for (if any). */
+	APlayerController* GetPlayerController() const { return WeakPlayerController.Get(); }
 
 public:
 
-	/** 
-	 * Activates the camera for the given player.
-	 * This looks up the given player's camera manager and/or view target in order to find
-	 * the active camera system. If found, this component adds its camera asset as the active one.
-	 * If this component was already active for another player, it will be first deactivated.
-	 */
+	/** Activates the camera for the given player. */
+	UFUNCTION(BlueprintCallable, Category=Camera)
+	GAMEPLAYCAMERAS_API void ActivateCameraForPlayerIndex(int32 PlayerIndex);
+
+	/** Activates the camera for the given player. */
 	UFUNCTION(BlueprintCallable, Category=Camera)
-	GAMEPLAYCAMERAS_API void ActivateCamera(int32 PlayerIndex = 0);
+	GAMEPLAYCAMERAS_API void ActivateCameraForPlayerController(APlayerController* PlayerController);
 
 	/** Deactivates the camera for the last player it was activated for. */
 	UFUNCTION(BlueprintCallable, Category=Camera)
@@ -81,10 +81,8 @@ public:
 private:
 
 	void ActivateCameraEvaluationContext(int32 PlayerIndex);
-	void DeactivateCameraEvaluationContext();
-
 	void ActivateCameraEvaluationContext(APlayerController* PlayerController);
-	void DeactivateCameraEvaluationContext(APlayerController* PlayerController);
+	void DeactivateCameraEvaluationContext();
 
 	void DelayedActivateCameraEvaluationContext(APlayerController* PlayerController, AActor* OldViewTarget, AActor* NewViewTarget);
 	void AbortDelayedActivateCameraEvaluationContext();
@@ -128,7 +126,7 @@ protected:
 	
 private:
 
-	int32 ActivatedForPlayerIndex = INDEX_NONE;
+	TWeakObjectPtr<APlayerController> WeakPlayerController;
 };
 
 namespace UE::Cameras
diff --git a/Engine/Plugins/Cameras/GameplayCameras/Source/GameplayCameras/Public/GameFramework/GameplayCameraSystemComponent.h b/Engine/Plugins/Cameras/GameplayCameras/Source/GameplayCameras/Public/GameFramework/GameplayCameraSystemComponent.h
--- a/Engine/Plugins/Cameras/GameplayCameras/Source/GameplayCameras/Public/GameFramework/GameplayCameraSystemComponent.h
+++ b/Engine/Plugins/Cameras/GameplayCameras/Source/GameplayCameras/Public/GameFramework/GameplayCameraSystemComponent.h
@@ -45,7 +45,11 @@ public:
 
 	/** Sets this component's actor as the view target for the given player. */
 	UFUNCTION(BlueprintCallable, Category=Camera)
-	void ActivateCameraSystem(int32 PlayerIndex = 0);
+	void ActivateCameraSystemForPlayerIndex(int32 PlayerIndex);
+
+	/** Sets this component's actor as the view target for the given player. */
+	UFUNCTION(BlueprintCallable, Category=Camera)
+	void ActivateCameraSystemForPlayerController(APlayerController* PlayerController);
 
 	/** Removes this component's actor from being the view target. */
 	UFUNCTION(BlueprintCallable, Category=Camera)
@@ -93,7 +97,7 @@ private:
 	
 	TSharedPtr<FCameraSystemEvaluator> Evaluator;
 
-	int32 ActivatedForPlayerIndex = INDEX_NONE;
+	TWeakObjectPtr<APlayerController> WeakPlayerController;
 
 #if UE_GAMEPLAY_CAMERAS_DEBUG
 	FDelegateHandle DebugDrawDelegateHandle;
diff --git a/Engine/Plugins/Cameras/GameplayCameras/Source/GameplayCameras/Public/GameFramework/GameplayControlRotationComponent.h b/Engine/Plugins/Cameras/GameplayCameras/Source/GameplayCameras/Public/GameFramework/GameplayControlRotationComponent.h
--- a/Engine/Plugins/Cameras/GameplayCameras/Source/GameplayCameras/Public/GameFramework/GameplayControlRotationComponent.h
+++ b/Engine/Plugins/Cameras/GameplayCameras/Source/GameplayCameras/Public/GameFramework/GameplayControlRotationComponent.h
@@ -39,10 +39,6 @@ public:
 	virtual void EndPlay(const EEndPlayReason::Type EndPlayReason) override;
 	virtual void TickComponent(float DeltaTime, ELevelTick TickType, FActorComponentTickFunction* ThisTickFunction) override;
 
-private:
-
-	APlayerController* GetPlayerController();
-
 public:
 
 	/** The axis input action(s) to read from. */
@@ -64,8 +60,10 @@ public:
 
 private:
 
+	UPROPERTY()
+	TObjectPtr<APlayerController> PlayerController;
+
 	TSharedPtr<UE::Cameras::FPlayerControlRotationEvaluationService> ControlRotationService;
 	TWeakPtr<UE::Cameras::FCameraSystemEvaluator> WeakCameraSystem;
-	int32 PlayerIndex = INDEX_NONE;
 };
 
diff --git a/Engine/Plugins/Cameras/GameplayCameras/Source/GameplayCameras/Public/GameplayCamerasSettings.h b/Engine/Plugins/Cameras/GameplayCameras/Source/GameplayCameras/Public/GameplayCamerasSettings.h
--- a/Engine/Plugins/Cameras/GameplayCameras/Source/GameplayCameras/Public/GameplayCamerasSettings.h
+++ b/Engine/Plugins/Cameras/GameplayCameras/Source/GameplayCameras/Public/GameplayCamerasSettings.h
@@ -17,6 +17,16 @@ class UGameplayCamerasSettings : public UDeveloperSettings
 {
 	GENERATED_BODY()
 
+public:
+
+	/**
+	 * Automatically spawn a camera system actor when any gameplay camera activates and no camera system
+	 * is found on the player controller's camera manager, or as a view target. This camera system actor
+	 * will be spawned and set as the view target automatically.
+	 */
+	UPROPERTY(EditAnywhere, Config, Category="General")
+	bool bAutoSpawnCameraSystemActor = true;
+
 public:
 
 	/** The default angle tolerance to accept an aiming operation. */
