Cameras: ObjectTreeGraphNode now shows its underlying object class' tooltip text as its own tooltip

diff --git a/Engine/Plugins/Cameras/GameplayCameras/Source/GameplayCamerasEditor/Private/Editors/ObjectTreeGraphNode.cpp b/Engine/Plugins/Cameras/GameplayCameras/Source/GameplayCamerasEditor/Private/Editors/ObjectTreeGraphNode.cpp
--- a/Engine/Plugins/Cameras/GameplayCameras/Source/GameplayCamerasEditor/Private/Editors/ObjectTreeGraphNode.cpp
+++ b/Engine/Plugins/Cameras/GameplayCameras/Source/GameplayCamerasEditor/Private/Editors/ObjectTreeGraphNode.cpp
@@ -61,6 +61,15 @@ FLinearColor UObjectTreeGraphNode::GetNodeBodyTintColor() const
 	return NodeContext.ObjectClassConfig.NodeBodyTintColor().Get(NodeContext.GraphConfig.DefaultGraphNodeBodyTintColor);
 }
 
+FText UObjectTreeGraphNode::GetTooltipText() const
+{
+	if (Object)
+	{
+		return Object->GetClass()->GetToolTipText();
+	}
+	return FText::GetEmpty();
+}
+
 void UObjectTreeGraphNode::AllocateDefaultPins()
 {
 	if (!ensure(Object))
diff --git a/Engine/Plugins/Cameras/GameplayCameras/Source/GameplayCamerasEditor/Public/Editors/ObjectTreeGraphNode.h b/Engine/Plugins/Cameras/GameplayCameras/Source/GameplayCamerasEditor/Public/Editors/ObjectTreeGraphNode.h
--- a/Engine/Plugins/Cameras/GameplayCameras/Source/GameplayCamerasEditor/Public/Editors/ObjectTreeGraphNode.h
+++ b/Engine/Plugins/Cameras/GameplayCameras/Source/GameplayCamerasEditor/Public/Editors/ObjectTreeGraphNode.h
@@ -65,6 +65,7 @@ public:
 	virtual TSharedPtr<SGraphNode> CreateVisualWidget() override;
 	virtual FLinearColor GetNodeTitleColor() const override;
 	virtual FLinearColor GetNodeBodyTintColor() const override;
+	virtual FText GetTooltipText() const override;
 	virtual void AllocateDefaultPins() override;
 	virtual void PostPlacedNewNode() override;
 	virtual void AutowireNewNode(UEdGraphPin* FromPin) override;
