Cameras: bubble-up change events from camera rigs and directors so we can dirty the build status of the parent camera asset

diff --git a/Engine/Plugins/Cameras/GameplayCameras/Source/GameplayCamerasEditor/Private/Toolkits/CameraDirectorAssetEditorMode.cpp b/Engine/Plugins/Cameras/GameplayCameras/Source/GameplayCamerasEditor/Private/Toolkits/CameraDirectorAssetEditorMode.cpp
--- a/Engine/Plugins/Cameras/GameplayCameras/Source/GameplayCamerasEditor/Private/Toolkits/CameraDirectorAssetEditorMode.cpp
+++ b/Engine/Plugins/Cameras/GameplayCameras/Source/GameplayCamerasEditor/Private/Toolkits/CameraDirectorAssetEditorMode.cpp
@@ -39,6 +39,7 @@ void FCameraDirectorAssetEditorMode::OnActivateMode(const FAssetEditorModeActiva
 		FDetailsViewArgs DetailsViewArgs;
 		DetailsViewArgs.NameAreaSettings = FDetailsViewArgs::HideNameArea;
 		DetailsViewArgs.bHideSelectionTip = true;
+		DetailsViewArgs.NotifyHook = this;
 		DetailsView = PropertyEditorModule.CreateDetailView(DetailsViewArgs);
 
 		bInitializedToolkit = true;
@@ -67,6 +68,14 @@ void FCameraDirectorAssetEditorMode::OnDeactivateMode(const FAssetEditorModeDeac
 	InParams.TabManager->UnregisterTabSpawner(DirectorEditorTabId);
 }
 
+void FCameraDirectorAssetEditorMode::NotifyPostChange(const FPropertyChangedEvent& PropertyChangedEvent, FProperty* PropertyThatChanged)
+{
+	if (CameraAsset)
+	{
+		CameraAsset->DirtyBuildStatus();
+	}
+}
+
 bool FCameraDirectorAssetEditorMode::JumpToObject(UObject* InObject, FName InPropertyName)
 {
 	if (InObject == CameraAsset->GetCameraDirector())
diff --git a/Engine/Plugins/Cameras/GameplayCameras/Source/GameplayCamerasEditor/Private/Toolkits/CameraRigAssetEditorToolkitBase.cpp b/Engine/Plugins/Cameras/GameplayCameras/Source/GameplayCamerasEditor/Private/Toolkits/CameraRigAssetEditorToolkitBase.cpp
--- a/Engine/Plugins/Cameras/GameplayCameras/Source/GameplayCamerasEditor/Private/Toolkits/CameraRigAssetEditorToolkitBase.cpp
+++ b/Engine/Plugins/Cameras/GameplayCameras/Source/GameplayCamerasEditor/Private/Toolkits/CameraRigAssetEditorToolkitBase.cpp
@@ -203,6 +203,8 @@ void FCameraRigAssetEditorToolkitBase::OnAnyGraphChanged(const FEdGraphEditActio
 	if (CameraRigAsset)
 	{
 		CameraRigAsset->BuildStatus = ECameraBuildStatus::Dirty;
+
+		CameraRigBuildStatusDirtied.Broadcast();
 	}
 }
 
diff --git a/Engine/Plugins/Cameras/GameplayCameras/Source/GameplayCamerasEditor/Private/Toolkits/CameraRigsAssetEditorMode.cpp b/Engine/Plugins/Cameras/GameplayCameras/Source/GameplayCamerasEditor/Private/Toolkits/CameraRigsAssetEditorMode.cpp
--- a/Engine/Plugins/Cameras/GameplayCameras/Source/GameplayCamerasEditor/Private/Toolkits/CameraRigsAssetEditorMode.cpp
+++ b/Engine/Plugins/Cameras/GameplayCameras/Source/GameplayCamerasEditor/Private/Toolkits/CameraRigsAssetEditorMode.cpp
@@ -74,6 +74,8 @@ void FCameraRigsAssetEditorMode::OnActivateMode(const FAssetEditorModeActivatePa
 	Impl->BuildToolbarMenu(ToolbarMenu);
 
 	Impl->BindCommands(InParams.CommandList.ToSharedRef());
+
+	Impl->OnCameraRigBuildStatusDirtied().AddSP(this, &FCameraRigsAssetEditorMode::OnCameraRigBuildStatusDirtied);
 }
 
 TSharedRef<SDockTab> FCameraRigsAssetEditorMode::SpawnTab_CameraRigs(const FSpawnTabArgs& Args)
@@ -89,6 +91,8 @@ TSharedRef<SDockTab> FCameraRigsAssetEditorMode::SpawnTab_CameraRigs(const FSpaw
 
 void FCameraRigsAssetEditorMode::OnDeactivateMode(const FAssetEditorModeDeactivateParams& InParams)
 {
+	Impl->OnCameraRigBuildStatusDirtied().RemoveAll(this);
+
 	Impl->UnregisterTabSpawners(InParams.TabManager.ToSharedRef());
 
 	InParams.TabManager->UnregisterTabSpawner(CameraRigsTabId);
@@ -113,6 +117,14 @@ void FCameraRigsAssetEditorMode::OnCameraRigDeleted(const TArray<UCameraRigAsset
 	}
 }
 
+void FCameraRigsAssetEditorMode::OnCameraRigBuildStatusDirtied()
+{
+	if (CameraAsset)
+	{
+		CameraAsset->DirtyBuildStatus();
+	}
+}
+
 void FCameraRigsAssetEditorMode::OnGetRootObjectsToSearch(TArray<FFindInObjectTreeGraphSource>& OutSources)
 {
 	for (UCameraRigAsset* CameraRig : CameraAsset->GetCameraRigs())
diff --git a/Engine/Plugins/Cameras/GameplayCameras/Source/GameplayCamerasEditor/Public/Toolkits/CameraDirectorAssetEditorMode.h b/Engine/Plugins/Cameras/GameplayCameras/Source/GameplayCamerasEditor/Public/Toolkits/CameraDirectorAssetEditorMode.h
--- a/Engine/Plugins/Cameras/GameplayCameras/Source/GameplayCamerasEditor/Public/Toolkits/CameraDirectorAssetEditorMode.h
+++ b/Engine/Plugins/Cameras/GameplayCameras/Source/GameplayCamerasEditor/Public/Toolkits/CameraDirectorAssetEditorMode.h
@@ -3,6 +3,7 @@
 #pragma once
 
 #include "CoreTypes.h"
+#include "Misc/NotifyHook.h"
 #include "Toolkits/AssetEditorMode.h"
 
 class IDetailsView;
@@ -16,6 +17,7 @@ class FStandardToolkitLayout;
 
 class FCameraDirectorAssetEditorMode
 	: public FAssetEditorMode
+	, public FNotifyHook
 {
 public:
 
@@ -27,9 +29,13 @@ public:
 
 protected:
 
+	// FAssetEditorMode interface.
 	virtual void OnActivateMode(const FAssetEditorModeActivateParams& InParams) override;
 	virtual void OnDeactivateMode(const FAssetEditorModeDeactivateParams& InParams) override;
 
+	// FNotifyHook interface.
+	virtual void NotifyPostChange(const FPropertyChangedEvent& PropertyChangedEvent, FProperty* PropertyThatChanged) override;
+
 private:
 
 	TSharedRef<SDockTab> SpawnTab_DirectorEditor(const FSpawnTabArgs& Args);
diff --git a/Engine/Plugins/Cameras/GameplayCameras/Source/GameplayCamerasEditor/Public/Toolkits/CameraRigAssetEditorToolkitBase.h b/Engine/Plugins/Cameras/GameplayCameras/Source/GameplayCamerasEditor/Public/Toolkits/CameraRigAssetEditorToolkitBase.h
--- a/Engine/Plugins/Cameras/GameplayCameras/Source/GameplayCamerasEditor/Public/Toolkits/CameraRigAssetEditorToolkitBase.h
+++ b/Engine/Plugins/Cameras/GameplayCameras/Source/GameplayCamerasEditor/Public/Toolkits/CameraRigAssetEditorToolkitBase.h
@@ -50,6 +50,8 @@ public:
 	void BuildToolbarMenu(UToolMenu* ToolbarMenu);
 	void BindCommands(TSharedRef<FUICommandList> CommandList);
 
+	FSimpleMulticastDelegate& OnCameraRigBuildStatusDirtied() { return CameraRigBuildStatusDirtied; }
+
 protected:
 
 	// FGCObject interface
@@ -95,6 +97,9 @@ private:
 
 	/** Toolbox widget */
 	TSharedPtr<SObjectTreeGraphToolbox> ToolboxWidget;
+
+	/** A multicast delegate that is invoked when the camera rig asset is dirtied */
+	FSimpleMulticastDelegate CameraRigBuildStatusDirtied;
 };
 
 }  // namespace UE::Cameras
diff --git a/Engine/Plugins/Cameras/GameplayCameras/Source/GameplayCamerasEditor/Public/Toolkits/CameraRigsAssetEditorMode.h b/Engine/Plugins/Cameras/GameplayCameras/Source/GameplayCamerasEditor/Public/Toolkits/CameraRigsAssetEditorMode.h
--- a/Engine/Plugins/Cameras/GameplayCameras/Source/GameplayCamerasEditor/Public/Toolkits/CameraRigsAssetEditorMode.h
+++ b/Engine/Plugins/Cameras/GameplayCameras/Source/GameplayCamerasEditor/Public/Toolkits/CameraRigsAssetEditorMode.h
@@ -44,6 +44,8 @@ private:
 	void OnCameraRigEditRequested(UCameraRigAsset* InCameraRig);
 	void OnCameraRigDeleted(const TArray<UCameraRigAsset*>& InCameraRigs);
 
+	void OnCameraRigBuildStatusDirtied();
+
 private:
 
 	static const FName CameraRigsTabId;
