Cameras: don't dirty camera rigs and camera assets if building them doesn't change anything

diff --git a/Engine/Plugins/Cameras/GameplayCameras/Source/GameplayCameras/Private/Core/CameraAssetBuilder.cpp b/Engine/Plugins/Cameras/GameplayCameras/Source/GameplayCameras/Private/Core/CameraAssetBuilder.cpp
--- a/Engine/Plugins/Cameras/GameplayCameras/Source/GameplayCameras/Private/Core/CameraAssetBuilder.cpp
+++ b/Engine/Plugins/Cameras/GameplayCameras/Source/GameplayCameras/Private/Core/CameraAssetBuilder.cpp
@@ -77,11 +77,8 @@ void FCameraAssetBuilder::UpdateBuildStatus()
 		BuildStatus = ECameraBuildStatus::CleanWithWarnings;
 	}
 
-	if (CameraAsset->GetBuildStatus() != BuildStatus)
-	{
-		CameraAsset->Modify();
-		CameraAsset->SetBuildStatus(BuildStatus);
-	}
+	// Don't modify the camera rig: BuildStatus is transient.
+	CameraAsset->SetBuildStatus(BuildStatus);
 }
 
 }  // namespace UE::Cameras
diff --git a/Engine/Plugins/Cameras/GameplayCameras/Source/GameplayCameras/Private/Core/CameraRigAsset.cpp b/Engine/Plugins/Cameras/GameplayCameras/Source/GameplayCameras/Private/Core/CameraRigAsset.cpp
--- a/Engine/Plugins/Cameras/GameplayCameras/Source/GameplayCameras/Private/Core/CameraRigAsset.cpp
+++ b/Engine/Plugins/Cameras/GameplayCameras/Source/GameplayCameras/Private/Core/CameraRigAsset.cpp
@@ -24,6 +24,12 @@ void FCameraRigAllocationInfo::Append(const FCameraRigAllocationInfo& OtherAlloc
 	VariableTableInfo.VariableDefinitions.Append(OtherVariableTableInfo.VariableDefinitions);
 }
 
+bool operator==(const FCameraRigAllocationInfo& A, const FCameraRigAllocationInfo& B)
+{
+	return A.EvaluatorInfo == B.EvaluatorInfo
+		&& A.VariableTableInfo == B.VariableTableInfo;
+}
+
 #if WITH_EDITOR
 
 void UCameraRigInterfaceParameter::GetGraphNodePosition(FName InGraphName, int32& NodePosX, int32& NodePosY) const
diff --git a/Engine/Plugins/Cameras/GameplayCameras/Source/GameplayCameras/Private/Core/CameraRigAssetBuilder.cpp b/Engine/Plugins/Cameras/GameplayCameras/Source/GameplayCameras/Private/Core/CameraRigAssetBuilder.cpp
--- a/Engine/Plugins/Cameras/GameplayCameras/Source/GameplayCameras/Private/Core/CameraRigAssetBuilder.cpp
+++ b/Engine/Plugins/Cameras/GameplayCameras/Source/GameplayCameras/Private/Core/CameraRigAssetBuilder.cpp
@@ -167,7 +167,9 @@ void SetupPrivateVariable(
 		{
 			PrivateVariable->Modify();
 		}
-		PrivateVariable->Rename(*VariableName);
+		// Rename non-transactionally because we might be simply setting the variable's name back to what it
+		// always was. We don't want to dirty the package for no-op builds.
+		PrivateVariable->Rename(*VariableName, nullptr, REN_NonTransactional);
 	}
 	else
 	{
@@ -325,7 +327,9 @@ UE_CAMERA_VARIABLE_FOR_ALL_TYPES()
 		TStringBuilder<256> StringBuilder;
 		StringBuilder.Append("REUSABLE_");
 		StringBuilder.Append(Pair.Value->GetName());
-		Pair.Value->Rename(StringBuilder.ToString());
+		// Rename non-transactionally because if nothing has change, we will rename it back
+		// later and we don't want to dirty the package for nothing.
+		Pair.Value->Rename(StringBuilder.ToString(), nullptr, REN_NonTransactional);
 	}
 }
 
@@ -491,7 +495,11 @@ void FCameraRigAssetBuilder::BuildAllocationInfo()
 	}
 
 	// Set it on the camera rig asset.
-	CameraRig->AllocationInfo = AllocationInfo;
+	if (CameraRig->AllocationInfo != AllocationInfo)
+	{
+		CameraRig->Modify();
+		CameraRig->AllocationInfo = AllocationInfo;
+	}
 }
 
 void FCameraRigAssetBuilder::BuildAllocationInfo(UCameraNode* CameraNode)
@@ -547,11 +555,8 @@ void FCameraRigAssetBuilder::UpdateBuildStatus()
 		BuildStatus = ECameraBuildStatus::CleanWithWarnings;
 	}
 
-	if (CameraRig->BuildStatus != BuildStatus)
-	{
-		CameraRig->Modify();
-		CameraRig->BuildStatus = BuildStatus;
-	}
+	// Don't modify the camera rig: BuildStatus is transient.
+	CameraRig->BuildStatus = BuildStatus;
 }
 
 }  // namespace UE::Cameras
diff --git a/Engine/Plugins/Cameras/GameplayCameras/Source/GameplayCameras/Private/Core/CameraVariableTable.cpp b/Engine/Plugins/Cameras/GameplayCameras/Source/GameplayCameras/Private/Core/CameraVariableTable.cpp
--- a/Engine/Plugins/Cameras/GameplayCameras/Source/GameplayCameras/Private/Core/CameraVariableTable.cpp
+++ b/Engine/Plugins/Cameras/GameplayCameras/Source/GameplayCameras/Private/Core/CameraVariableTable.cpp
@@ -532,3 +532,21 @@ UE_CAMERA_VARIABLE_FOR_ALL_TYPES()
 
 }  // namespace UE::Cameras
 
+bool operator==(const FCameraVariableDefinition& A, const FCameraVariableDefinition& B)
+{
+	return A.VariableID == B.VariableID
+		&& A.VariableType == B.VariableType
+		&& A.bIsPrivate == B.bIsPrivate
+		&& A.bIsInput == B.bIsInput
+#if WITH_EDITORONLY_DATA
+		&& A.VariableName == B.VariableName
+#endif
+		;
+}
+
+bool operator==(const FCameraVariableTableAllocationInfo& A, const FCameraVariableTableAllocationInfo& B)
+{
+	return A.VariableDefinitions == B.VariableDefinitions
+		&& A.AutoResetVariables == B.AutoResetVariables;
+}
+
diff --git a/Engine/Plugins/Cameras/GameplayCameras/Source/GameplayCameras/Public/Core/CameraNodeEvaluatorFwd.h b/Engine/Plugins/Cameras/GameplayCameras/Source/GameplayCameras/Public/Core/CameraNodeEvaluatorFwd.h
--- a/Engine/Plugins/Cameras/GameplayCameras/Source/GameplayCameras/Public/Core/CameraNodeEvaluatorFwd.h
+++ b/Engine/Plugins/Cameras/GameplayCameras/Source/GameplayCameras/Public/Core/CameraNodeEvaluatorFwd.h
@@ -31,6 +31,24 @@ struct FCameraNodeEvaluatorAllocationInfo
 	/** Maximum required alignment for the node evaluators. */
 	UPROPERTY()
 	int16 MaxAlignof = 0;
+
+public:
+
+	friend bool operator==(const FCameraNodeEvaluatorAllocationInfo& A, const FCameraNodeEvaluatorAllocationInfo& B)
+	{
+		return A.TotalSizeof == B.TotalSizeof &&
+			A.MaxAlignof == B.MaxAlignof;
+	}
+};
+
+template<>
+struct TStructOpsTypeTraits<FCameraNodeEvaluatorAllocationInfo> : public TStructOpsTypeTraitsBase2<FCameraNodeEvaluatorAllocationInfo>
+{
+	enum
+	{
+		WithCopy = true,
+		WithIdenticalViaEquality = true
+	};
 };
 
 /** Allocation information for a node evaluator, auto-setup for a given type. */
diff --git a/Engine/Plugins/Cameras/GameplayCameras/Source/GameplayCameras/Public/Core/CameraRigAsset.h b/Engine/Plugins/Cameras/GameplayCameras/Source/GameplayCameras/Public/Core/CameraRigAsset.h
--- a/Engine/Plugins/Cameras/GameplayCameras/Source/GameplayCameras/Public/Core/CameraRigAsset.h
+++ b/Engine/Plugins/Cameras/GameplayCameras/Source/GameplayCameras/Public/Core/CameraRigAsset.h
@@ -42,7 +42,19 @@ struct FCameraRigAllocationInfo
 
 public:
 
-	void Append(const FCameraRigAllocationInfo& OtherAllocationInfo);
+	GAMEPLAYCAMERAS_API void Append(const FCameraRigAllocationInfo& OtherAllocationInfo);
+
+	GAMEPLAYCAMERAS_API friend bool operator==(const FCameraRigAllocationInfo& A, const FCameraRigAllocationInfo& B);
+};
+
+template<>
+struct TStructOpsTypeTraits<FCameraRigAllocationInfo> : public TStructOpsTypeTraitsBase2<FCameraRigAllocationInfo>
+{
+	enum
+	{
+		WithCopy = true,
+		WithIdenticalViaEquality = true
+	};
 };
 
 /**
diff --git a/Engine/Plugins/Cameras/GameplayCameras/Source/GameplayCameras/Public/Core/CameraVariableTableFwd.h b/Engine/Plugins/Cameras/GameplayCameras/Source/GameplayCameras/Public/Core/CameraVariableTableFwd.h
--- a/Engine/Plugins/Cameras/GameplayCameras/Source/GameplayCameras/Public/Core/CameraVariableTableFwd.h
+++ b/Engine/Plugins/Cameras/GameplayCameras/Source/GameplayCameras/Public/Core/CameraVariableTableFwd.h
@@ -171,6 +171,18 @@ struct FCameraVariableDefinition
 #endif
 		return VariantDefinition;
 	}
+
+	GAMEPLAYCAMERAS_API friend bool operator==(const FCameraVariableDefinition& A, const FCameraVariableDefinition& B);
+};
+
+template<>
+struct TStructOpsTypeTraits<FCameraVariableDefinition> : public TStructOpsTypeTraitsBase2<FCameraVariableDefinition>
+{
+	enum
+	{
+		WithCopy = true,
+		WithIdenticalViaEquality = true
+	};
 };
 
 /**
@@ -188,5 +200,17 @@ struct FCameraVariableTableAllocationInfo
 	/** The list of variables that should be auto-reset to their default value every frame. */
 	UPROPERTY()
 	TArray<TObjectPtr<UCameraVariableAsset>> AutoResetVariables;
+
+	GAMEPLAYCAMERAS_API friend bool operator==(const FCameraVariableTableAllocationInfo& A, const FCameraVariableTableAllocationInfo& B);
+};
+
+template<>
+struct TStructOpsTypeTraits<FCameraVariableTableAllocationInfo> : public TStructOpsTypeTraitsBase2<FCameraVariableTableAllocationInfo>
+{
+	enum
+	{
+		WithCopy = true,
+		WithIdenticalViaEquality = true
+	};
 };
 
