Cameras: fix missing calls to NotifyGraphChanged in schemas

diff --git a/Engine/Plugins/Cameras/GameplayCameras/Source/GameplayCamerasEditor/Private/Editors/CameraNodeGraphSchema.cpp b/Engine/Plugins/Cameras/GameplayCameras/Source/GameplayCamerasEditor/Private/Editors/CameraNodeGraphSchema.cpp
--- a/Engine/Plugins/Cameras/GameplayCameras/Source/GameplayCamerasEditor/Private/Editors/CameraNodeGraphSchema.cpp
+++ b/Engine/Plugins/Cameras/GameplayCameras/Source/GameplayCamerasEditor/Private/Editors/CameraNodeGraphSchema.cpp
@@ -224,6 +224,9 @@ bool UCameraNodeGraphSchema::OnCreateConnection(UEdGraphPin* A, UEdGraphPin* B)
 		RigParameter->InterfaceParameterName = StructProperty->GetName();
 	}
 
+	UEdGraph* Graph = CameraNodeNode->GetGraph();
+	Graph->NotifyGraphChanged();
+
 	return true;
 }
 
@@ -251,6 +254,9 @@ bool UCameraNodeGraphSchema::OnBreakPinLinks(UEdGraphPin& TargetPin, bool bSends
 			RigParameter->TargetPropertyName = NAME_None;
 			RigParameter->PrivateVariable = nullptr;
 
+			UEdGraph* Graph = RigParameterNode->GetGraph();
+			Graph->NotifyGraphChanged();
+
 			return true;
 		}
 	}
@@ -260,14 +266,14 @@ bool UCameraNodeGraphSchema::OnBreakPinLinks(UEdGraphPin& TargetPin, bool bSends
 
 bool UCameraNodeGraphSchema::OnBreakSinglePinLink(UEdGraphPin* SourcePin, UEdGraphPin* TargetPin) const
 {
-	UCameraRigInterfaceParameter* RigParameter = nullptr;
+	UObjectTreeGraphNode* RigParameterNode = nullptr;
 
 	if (SourcePin->PinType.PinCategory == PC_Self && TargetPin->PinType.PinCategory == PC_CameraParameter)
 	{
 		UObjectTreeGraphNode* SourceNode = Cast<UObjectTreeGraphNode>(SourcePin->GetOwningNode());
 		if (SourceNode)
 		{
-			RigParameter = SourceNode->CastObject<UCameraRigInterfaceParameter>();
+			RigParameterNode = SourceNode;
 		}
 	}
 	else if (SourcePin->PinType.PinCategory == PC_CameraParameter && TargetPin->PinType.PinCategory == PC_Self)
@@ -275,10 +281,16 @@ bool UCameraNodeGraphSchema::OnBreakSinglePinLink(UEdGraphPin* SourcePin, UEdGra
 		UObjectTreeGraphNode* TargetNode = Cast<UObjectTreeGraphNode>(TargetPin->GetOwningNode());
 		if (TargetNode)
 		{ 
-			RigParameter = TargetNode->CastObject<UCameraRigInterfaceParameter>();
+			RigParameterNode = TargetNode;
 		}
 	}
 
+	UCameraRigInterfaceParameter* RigParameter = nullptr;
+	if (RigParameterNode)
+	{
+		RigParameter = RigParameterNode->CastObject<UCameraRigInterfaceParameter>();
+	}
+
 	if (RigParameter)
 	{
 		const FScopedTransaction Transaction(LOCTEXT("BreakSinglePinLink", "Break Pin Link"));
@@ -289,6 +301,9 @@ bool UCameraNodeGraphSchema::OnBreakSinglePinLink(UEdGraphPin* SourcePin, UEdGra
 		RigParameter->TargetPropertyName = NAME_None;
 		RigParameter->PrivateVariable = nullptr;
 
+		UEdGraph* Graph = RigParameterNode->GetGraph();
+		Graph->NotifyGraphChanged();
+
 		return true;
 	}
 	
diff --git a/Engine/Plugins/Cameras/GameplayCameras/Source/GameplayCamerasEditor/Private/Editors/ObjectTreeGraphSchema.cpp b/Engine/Plugins/Cameras/GameplayCameras/Source/GameplayCamerasEditor/Private/Editors/ObjectTreeGraphSchema.cpp
--- a/Engine/Plugins/Cameras/GameplayCameras/Source/GameplayCamerasEditor/Private/Editors/ObjectTreeGraphSchema.cpp
+++ b/Engine/Plugins/Cameras/GameplayCameras/Source/GameplayCamerasEditor/Private/Editors/ObjectTreeGraphSchema.cpp
@@ -663,6 +663,9 @@ bool UObjectTreeGraphSchema::TryCreateConnection(UEdGraphPin* A, UEdGraphPin* B)
 		ObjectA->PostEditChangeProperty(PropertyChangedEvent);
 	}
 
+	UEdGraph* Graph = NodeA->GetGraph();
+	Graph->NotifyGraphChanged();
+
 	return true;
 }
 
@@ -737,6 +740,9 @@ void UObjectTreeGraphSchema::BreakPinLinks(UEdGraphPin& TargetPin, bool bSendsNo
 		PropertyOwningNode->RemoveItemPin(PropertyPin);
 		PropertyOwningNode->GetGraph()->NotifyNodeChanged(PropertyOwningNode);
 	}
+
+	UEdGraph* Graph = PropertyOwningNode->GetGraph();
+	Graph->NotifyGraphChanged();
 }
 
 bool UObjectTreeGraphSchema::OnBreakPinLinks(UEdGraphPin& TargetPin, bool bSendsNodeNotification) const
@@ -795,6 +801,9 @@ void UObjectTreeGraphSchema::BreakSinglePinLink(UEdGraphPin* SourcePin, UEdGraph
 		PropertyOwningNode->RemoveItemPin(PropertyPin);
 		PropertyOwningNode->GetGraph()->NotifyNodeChanged(PropertyOwningNode);
 	}
+
+	UEdGraph* Graph = PropertyOwningNode->GetGraph();
+	Graph->NotifyGraphChanged();
 }
 
 bool UObjectTreeGraphSchema::OnBreakSinglePinLink(UEdGraphPin* SourcePin, UEdGraphPin* TargetPin) const
@@ -983,6 +992,8 @@ void UObjectTreeGraphSchema::ImportNodesFromText(UObjectTreeGraph* InGraph, cons
 	{
 		OutPastedNodes.Add(Pair.Value);
 	}
+
+	InGraph->NotifyGraphChanged();
 }
 
 bool UObjectTreeGraphSchema::CanImportNodesFromText(UObjectTreeGraph* InGraph, const FString& TextToImport) const
