Cameras: clean-up how camera rig picker widgets work

- Make a combo button widget that shows the camera rig picker.
- The camera rig pin widget uses that button instead of setting things up itself.
- A new subclass of the camera rig pin widget handles specific logic and messages for the Blueprint camera director.
- Make the pin factory handle the Blueprint camera director pin widget.
- Make the pin factory handle generic rig picker pins, not just parameters to function nodes.
- Remove "name" mode for the pin, it's not used and not desirable anyway.
- Default to a full camera rig picker that lets the user pick the camera asset.
- Miscellaneous code clean-up and renaming.

diff --git a/Engine/Plugins/Cameras/GameplayCameras/Source/GameplayCameras/Public/Directors/BlueprintCameraDirector.h b/Engine/Plugins/Cameras/GameplayCameras/Source/GameplayCameras/Public/Directors/BlueprintCameraDirector.h
--- a/Engine/Plugins/Cameras/GameplayCameras/Source/GameplayCameras/Public/Directors/BlueprintCameraDirector.h
+++ b/Engine/Plugins/Cameras/GameplayCameras/Source/GameplayCameras/Public/Directors/BlueprintCameraDirector.h
@@ -64,7 +64,9 @@ public:
 
 	/** Specifies a camera rig to be active this frame. */
 	UFUNCTION(BlueprintCallable, Category="Evaluation")
-	void ActivateCameraRig(UPARAM(meta=(UseCameraRigPicker=true)) UCameraRigAsset* CameraRig);
+	void ActivateCameraRig(
+			UPARAM(meta=(UseBlueprintCameraDirectorRigPicker=true))
+			UCameraRigAsset* CameraRig);
 
 	/**
 	 * Specifies a camera rig to be active this frame, via a proxy which is later resolved
diff --git a/Engine/Plugins/Cameras/GameplayCameras/Source/GameplayCamerasEditor/Private/Builders/BlueprintCameraDirectorEditorBuilder.cpp b/Engine/Plugins/Cameras/GameplayCameras/Source/GameplayCamerasEditor/Private/Builders/BlueprintCameraDirectorEditorBuilder.cpp
--- a/Engine/Plugins/Cameras/GameplayCameras/Source/GameplayCamerasEditor/Private/Builders/BlueprintCameraDirectorEditorBuilder.cpp
+++ b/Engine/Plugins/Cameras/GameplayCameras/Source/GameplayCamerasEditor/Private/Builders/BlueprintCameraDirectorEditorBuilder.cpp
@@ -7,7 +7,7 @@
 #include "Directors/BlueprintCameraDirector.h"
 #include "EdGraph/EdGraph.h"
 #include "Engine/Blueprint.h"
-#include "Helpers/CameraDirectorHelper.h"
+#include "Helpers/CameraAssetReferenceGatherer.h"
 #include "K2Node_CallFunction.h"
 
 #define LOCTEXT_NAMESPACE "BlueprintCameraDirectorEditorBuilder"
@@ -35,7 +35,7 @@ void FBlueprintCameraDirectorEditorBuilder::OnBuildCameraAsset(UCameraAsset* Cam
 	// camera rig proxies for activation.
 	TArray<UCameraAsset*> ReferencingCameraAssets;
 	UBlueprint* Blueprint = CastChecked<UBlueprint>(CameraDirectorEvaluatorClass->ClassGeneratedBy);
-	FCameraDirectorHelper::GetReferencingCameraAssets(Blueprint, ReferencingCameraAssets);
+	FCameraAssetReferenceGatherer::GetReferencingCameraAssets(Blueprint, ReferencingCameraAssets);
 
 	UCameraAsset* ThisCameraAsset = BlueprintCameraDirector->GetTypedOuter<UCameraAsset>();
 	ensure(ThisCameraAsset);
diff --git a/Engine/Plugins/Cameras/GameplayCameras/Source/GameplayCamerasEditor/Private/Customizations/CameraRigPtrDetailsCustomization.cpp b/Engine/Plugins/Cameras/GameplayCameras/Source/GameplayCamerasEditor/Private/Customizations/CameraRigPtrDetailsCustomization.cpp
--- a/Engine/Plugins/Cameras/GameplayCameras/Source/GameplayCamerasEditor/Private/Customizations/CameraRigPtrDetailsCustomization.cpp
+++ b/Engine/Plugins/Cameras/GameplayCameras/Source/GameplayCamerasEditor/Private/Customizations/CameraRigPtrDetailsCustomization.cpp
@@ -8,8 +8,8 @@
 #include "Core/CameraRigAsset.h"
 #include "DetailCategoryBuilder.h"
 #include "DetailWidgetRow.h"
+#include "Helpers/CameraAssetReferenceGatherer.h"
 #include "Editors/CameraRigPickerConfig.h"
-#include "Helpers/CameraDirectorHelper.h"
 #include "IDetailChildrenBuilder.h"
 #include "IGameplayCamerasEditorModule.h"
 #include "PropertyEditorModule.h"
@@ -115,7 +115,7 @@ TSharedRef<SWidget> FCameraRigPtrDetailsCustomization::OnBuildCameraRigNamePicke
 	TArray<UCameraAsset*> ReferencingCameraAssets;
 	for (UObject* CameraDirectorObject : CameraDirectorObjects)
 	{
-		FCameraDirectorHelper::GetReferencingCameraAssets(CameraDirectorObject, ReferencingCameraAssets);
+		FCameraAssetReferenceGatherer::GetReferencingCameraAssets(CameraDirectorObject, ReferencingCameraAssets);
 	}
 
 	if (ReferencingCameraAssets.Num() == 0)
diff --git a/Engine/Plugins/Cameras/GameplayCameras/Source/GameplayCamerasEditor/Private/Editors/GameplayCamerasGraphPanelPinFactory.cpp b/Engine/Plugins/Cameras/GameplayCameras/Source/GameplayCamerasEditor/Private/Editors/GameplayCamerasGraphPanelPinFactory.cpp
--- a/Engine/Plugins/Cameras/GameplayCameras/Source/GameplayCamerasEditor/Private/Editors/GameplayCamerasGraphPanelPinFactory.cpp
+++ b/Engine/Plugins/Cameras/GameplayCameras/Source/GameplayCamerasEditor/Private/Editors/GameplayCamerasGraphPanelPinFactory.cpp
@@ -2,27 +2,40 @@
 
 #include "Editors/GameplayCamerasGraphPanelPinFactory.h"
 
+#include "Core/CameraRigAsset.h"
+#include "EdGraphSchema_K2.h"
+#include "Editors/SBlueprintCameraDirectorRigNameGraphPin.h"
 #include "Editors/SCameraRigNameGraphPin.h"
 #include "K2Node_CallFunction.h"
 
 namespace UE::Cameras
 {
 
-TSharedPtr<class SGraphPin> FGameplayCamerasGraphPanelPinFactory::CreatePin(UEdGraphPin* Pin) const
+TSharedPtr<SGraphPin> FGameplayCamerasGraphPanelPinFactory::CreatePin(UEdGraphPin* Pin) const
 {
 	if (!Pin)
 	{
 		return nullptr;
 	}
 
-	UK2Node_CallFunction* OwningNode = Cast<UK2Node_CallFunction>(Pin->GetOwningNode());
-	if (!OwningNode)
+	if (UK2Node_CallFunction* CallFunctionNode = Cast<UK2Node_CallFunction>(Pin->GetOwningNode()))
 	{
-		return nullptr;
+		return CreateFunctionParameterPin(Pin, CallFunctionNode);
+	}
+
+	if (Pin->PinType.PinCategory == UEdGraphSchema_K2::PC_Object && 
+			Pin->PinType.PinSubCategoryObject == UCameraRigAsset::StaticClass())
+	{
+		return CreateCustomPin(Pin);
 	}
 
-	UClass* BlueprintClass = OwningNode->GetBlueprintClassFromNode();
-	UFunction* ReferencedFunction = OwningNode->FunctionReference.ResolveMember<UFunction>(BlueprintClass);
+	return nullptr;
+}
+
+TSharedPtr<SGraphPin> FGameplayCamerasGraphPanelPinFactory::CreateFunctionParameterPin(UEdGraphPin* Pin, UK2Node_CallFunction* CallFunctionNode) const
+{
+	UClass* BlueprintClass = CallFunctionNode->GetBlueprintClassFromNode();
+	UFunction* ReferencedFunction = CallFunctionNode->FunctionReference.ResolveMember<UFunction>(BlueprintClass);
 	if (!ReferencedFunction)
 	{
 		return nullptr;
@@ -34,20 +47,23 @@ TSharedPtr<class SGraphPin> FGameplayCamerasGraphPanelPinFactory::CreatePin(UEdG
 		return nullptr;
 	}
 
-	if (ParameterProperty->HasMetaData(TEXT("UseCameraRigNamePicker")))
+	if (ParameterProperty->HasMetaData(TEXT("UseBlueprintCameraDirectorRigPicker")))
 	{
-		return SNew(SCameraRigNameGraphPin, Pin)
-			.PinMode(ECameraRigNameGraphPinMode::NamePin);
+		return SNew(SBlueprintCameraDirectorRigNameGraphPin, Pin);
 	}
 
 	if (ParameterProperty->HasMetaData(TEXT("UseCameraRigPicker")))
 	{
-		return SNew(SCameraRigNameGraphPin, Pin)
-			.PinMode(ECameraRigNameGraphPinMode::ReferencePin);
+		return SNew(SCameraRigNameGraphPin, Pin);
 	}
 
 	return nullptr;
 }
 
+TSharedPtr<SGraphPin> FGameplayCamerasGraphPanelPinFactory::CreateCustomPin(UEdGraphPin* Pin) const
+{
+	return SNew(SCameraRigNameGraphPin, Pin);
+}
+
 }  // namespace UE::Cameras
 
diff --git a/Engine/Plugins/Cameras/GameplayCameras/Source/GameplayCamerasEditor/Private/Editors/SBlueprintCameraDirectorRigNameGraphPin.cpp b/Engine/Plugins/Cameras/GameplayCameras/Source/GameplayCamerasEditor/Private/Editors/SBlueprintCameraDirectorRigNameGraphPin.cpp
new file mode 100644
--- /dev/null
+++ b/Engine/Plugins/Cameras/GameplayCameras/Source/GameplayCamerasEditor/Private/Editors/SBlueprintCameraDirectorRigNameGraphPin.cpp
@@ -0,0 +1,62 @@
+// Copyright Epic Games, Inc. All Rights Reserved.
+
+#include "Editors/SBlueprintCameraDirectorRigNameGraphPin.h"
+
+#include "Core/CameraAsset.h"
+#include "Editors/CameraRigPickerConfig.h"
+#include "Helpers/CameraAssetReferenceGatherer.h"
+#include "Kismet2/BlueprintEditorUtils.h"
+
+#define LOCTEXT_NAMESPACE "SBlueprintCameraDirectorRigNameGraphPin"
+
+namespace UE::Cameras
+{
+
+void SBlueprintCameraDirectorRigNameGraphPin::Construct(const FArguments& InArgs, UEdGraphPin* InGraphPinObj)
+{
+	SCameraRigNameGraphPin::Construct(
+			SCameraRigNameGraphPin::FArguments(),
+			InGraphPinObj);
+}
+
+void SBlueprintCameraDirectorRigNameGraphPin::OnCustomizeCameraRigPickerConfig(FCameraRigPickerConfig& CameraRigPickerConfig) const
+{
+	TSharedPtr<SGraphNode> OwnerNodeWidget = OwnerNodePtr.Pin();
+	check(OwnerNodeWidget);
+	
+	UEdGraphNode* OwnerNode = OwnerNodeWidget->GetNodeObj();
+	UBlueprint* OwnerBlueprint = FBlueprintEditorUtils::FindBlueprintForNode(OwnerNode);
+
+	TArray<UCameraAsset*> ReferencingCameraAssets;
+	FCameraAssetReferenceGatherer::GetReferencingCameraAssets(OwnerBlueprint, ReferencingCameraAssets);
+	
+	// Don't let the user choose which camera asset to activate rigs from. Automatically use the camera asset
+	// that references this Blueprint camera director. Display warnings if we have zero referencers, or more
+	// than one.
+	CameraRigPickerConfig.bCanSelectCameraAsset = false;
+
+	if (ReferencingCameraAssets.Num() == 0)
+	{
+		CameraRigPickerConfig.WarningMessage = LOCTEXT("NoReferencingCameraAssetWarning",
+				"No camera asset references this Blueprint, so no camera rig list can be displayed. "
+				"Make a camera asset use this Blueprint as its camera director evaluator, or use "
+				"ActivateCameraRigViaProxy.");
+	}
+	else
+	{
+		CameraRigPickerConfig.InitialCameraAssetSelection = ReferencingCameraAssets[0];
+
+		if (ReferencingCameraAssets.Num() > 1)
+		{
+			CameraRigPickerConfig.WarningMessage = LOCTEXT("ManyReferencingCameraAssetsWarning",
+				"More than one camera asset references this Blueprint. Only camera rigs from the first "
+				"one will be displayed. Even then, shared camera director Blueprints should use "
+				"ActivateCameraRigViaProxy instead.");
+		}
+	}
+}
+
+}  // namespace UE::Cameras
+
+#undef LOCTEXT_NAMESPACE
+
diff --git a/Engine/Plugins/Cameras/GameplayCameras/Source/GameplayCamerasEditor/Private/Editors/SCameraRigNameGraphPin.cpp b/Engine/Plugins/Cameras/GameplayCameras/Source/GameplayCamerasEditor/Private/Editors/SCameraRigNameGraphPin.cpp
--- a/Engine/Plugins/Cameras/GameplayCameras/Source/GameplayCamerasEditor/Private/Editors/SCameraRigNameGraphPin.cpp
+++ b/Engine/Plugins/Cameras/GameplayCameras/Source/GameplayCamerasEditor/Private/Editors/SCameraRigNameGraphPin.cpp
@@ -2,26 +2,19 @@
 
 #include "Editors/SCameraRigNameGraphPin.h"
 
-#include "AssetRegistry/AssetRegistryModule.h"
 #include "Core/CameraAsset.h"
 #include "Core/CameraRigAsset.h"
 #include "Editors/CameraRigPickerConfig.h"
-#include "Framework/Views/ITypedTableView.h"
-#include "Helpers/CameraDirectorHelper.h"
+#include "Editors/SCameraRigPickerButton.h"
 #include "IContentBrowserSingleton.h"
 #include "IGameplayCamerasEditorModule.h"
-#include "Kismet2/BlueprintEditorUtils.h"
-#include "Modules/ModuleManager.h"
 #include "ScopedTransaction.h"
 #include "SGraphPin.h"
-#include "Styles/GameplayCamerasEditorStyle.h"
 #include "Widgets/Images/SImage.h"
 #include "Widgets/Input/SButton.h"
 #include "Widgets/Input/SEditableTextBox.h"
-#include "Widgets/Layout/SBorder.h"
 #include "Widgets/SBoxPanel.h"
 #include "Widgets/SNullWidget.h"
-#include "Widgets/Text/STextBlock.h"
 
 #define LOCTEXT_NAMESPACE "SCameraRigNameGraphPin"
 
@@ -31,8 +24,6 @@ namespace UE::Cameras
 void SCameraRigNameGraphPin::Construct(const FArguments& InArgs, UEdGraphPin* InGraphPinObj)
 {
 	SGraphPin::Construct(SGraphPin::FArguments(), InGraphPinObj);
-
-	PinMode = InArgs._PinMode;
 }
 
 TSharedRef<SWidget>	SCameraRigNameGraphPin::GetDefaultValueWidget()
@@ -49,23 +40,14 @@ TSharedRef<SWidget>	SCameraRigNameGraphPin::GetDefaultValueWidget()
 		.Padding(2.f)
 		.MaxWidth(200.f)
 		[
-			SAssignNew(CameraRigPickerButton, SComboButton)
-			.ButtonStyle(FAppStyle::Get(), "PropertyEditor.AssetComboStyle")
-			.ContentPadding(FMargin(2.f, 2.f, 2.f, 1.f))
-			.ForegroundColor(this, &SCameraRigNameGraphPin::OnGetComboForeground)
+			SAssignNew(CameraRigPickerButton, SCameraRigPickerButton)
+			.CameraRigPickerConfig(this, &SCameraRigNameGraphPin::OnCreateCameraRigPickerConfig)
+			.SelectedCameraRigName(this, &SCameraRigNameGraphPin::OnGetSelectedCameraRigName)
+			.ButtonToolTipText(this, &SCameraRigNameGraphPin::OnGetCameraRigPickerToolTipText)
+			.ButtonForegroundColor(this, &SCameraRigNameGraphPin::OnGetComboForeground)
 			.ButtonColorAndOpacity(this, &SCameraRigNameGraphPin::OnGetWidgetBackground)
-			.MenuPlacement(MenuPlacement_BelowAnchor)
+			.PickerMenuPlacement(MenuPlacement_BelowAnchor)
 			.IsEnabled(this, &SGraphPin::IsEditingEnabled)
-			.ButtonContent()
-			[
-				SNew(STextBlock)
-				.ColorAndOpacity(this, &SCameraRigNameGraphPin::OnGetComboForeground)
-				.TextStyle(FAppStyle::Get(), "PropertyEditor.AssetClass")
-				.Font(FAppStyle::GetFontStyle("PropertyWindow.NormalFont"))
-				.Text(this, &SCameraRigNameGraphPin::OnGetComboText)
-				.ToolTipText(this, &SCameraRigNameGraphPin::OnGetComboToolTipText)
-			]
-			.OnGetMenuContent(this, &SCameraRigNameGraphPin::OnBuildCameraRigNamePicker)
 		]
 		// Reset button
 		+SHorizontalBox::Slot()
@@ -111,108 +93,48 @@ FSlateColor SCameraRigNameGraphPin::OnGetWidgetBackground() const
 	return FSlateColor(FLinearColor(1.f, 1.f, 1.f, Alpha));
 }
 
-FText SCameraRigNameGraphPin::GetDefaultComboText() const
-{
-	return LOCTEXT("DefaultComboText", "Select Camera Rig");
-}
-
-FText SCameraRigNameGraphPin::OnGetComboText() const
+FText SCameraRigNameGraphPin::OnGetSelectedCameraRigName() const
 {
-	FText Value = GetDefaultComboText();
+	FText Value;
 	
 	if (GraphPinObj != nullptr)
 	{
-		switch (PinMode)
+		if (const UCameraRigAsset* CameraRig = Cast<const UCameraRigAsset>(GraphPinObj->DefaultObject))
 		{
-			case ECameraRigNameGraphPinMode::NamePin:
-				{
-					Value = FText::FromString(GraphPinObj->DefaultValue);
-				}
-				break;
-			case ECameraRigNameGraphPinMode::ReferencePin:
-				if (const UCameraRigAsset* CameraRig = Cast<const UCameraRigAsset>(GraphPinObj->DefaultObject))
-				{
-					Value = FText::FromString(CameraRig->GetDisplayName());
-				}
-				break;
+			Value = FText::FromString(CameraRig->GetDisplayName());
 		}
 	}
 	return Value;
 }
 
-FText SCameraRigNameGraphPin::OnGetComboToolTipText() const
+FText SCameraRigNameGraphPin::OnGetCameraRigPickerToolTipText() const
 {
-	return LOCTEXT("ComboToolTipText", "The name of the camera rig to activate.");
+	return LOCTEXT("ComboToolTipText", "The name of the camera rig.");
 }
 
-TSharedRef<SWidget> SCameraRigNameGraphPin::OnBuildCameraRigNamePicker()
+FCameraRigPickerConfig SCameraRigNameGraphPin::OnCreateCameraRigPickerConfig() const
 {
 	FCameraRigPickerConfig CameraRigPickerConfig;
-	CameraRigPickerConfig.bCanSelectCameraAsset = false;
 	CameraRigPickerConfig.bFocusCameraRigSearchBoxWhenOpened = true;
 	CameraRigPickerConfig.OnCameraRigSelected = FOnCameraRigSelected::CreateSP(this, &SCameraRigNameGraphPin::OnPickerAssetSelected);
 
-	TSharedPtr<SGraphNode> OwnerNodeWidget = OwnerNodePtr.Pin();
-	check(OwnerNodeWidget);
-	
-	UEdGraphNode* OwnerNode = OwnerNodeWidget->GetNodeObj();
-	UBlueprint* OwnerBlueprint = FBlueprintEditorUtils::FindBlueprintForNode(OwnerNode);
-
-	TArray<UCameraAsset*> ReferencingCameraAssets;
-	FCameraDirectorHelper::GetReferencingCameraAssets(OwnerBlueprint, ReferencingCameraAssets);
-
-	if (ReferencingCameraAssets.Num() == 0)
-	{
-		CameraRigPickerConfig.WarningMessage = LOCTEXT("NoReferencingCameraAssetWarning",
-				"No camera asset references this Blueprint, so no camera rig list can be displayed. "
-				"Make a camera asset use this Blueprint as its camera director evaluator, or use "
-				"ActivateCameraRigViaProxy.");
-	}
-	else
-	{
-		CameraRigPickerConfig.InitialCameraAssetSelection = ReferencingCameraAssets[0];
-
-		if (ReferencingCameraAssets.Num() > 1)
-		{
-			CameraRigPickerConfig.WarningMessage = LOCTEXT("ManyReferencingCameraAssetsWarning",
-				"More than one camera asset references this Blueprint. Only camera rigs from the first "
-				"one will be displayed. Even then, shared camera director Blueprints should use "
-				"ActivateCameraRigViaProxy instead.");
-		}
-	}
-
+	// Find the already specified camera rig, if any.
 	if (GraphPinObj)
 	{
-		switch (PinMode)
+		UCameraRigAsset* DefaultCameraRig = Cast<UCameraRigAsset>(GraphPinObj->DefaultObject);
+		if (DefaultCameraRig)
 		{
-			case ECameraRigNameGraphPinMode::NamePin:
-				if (!ReferencingCameraAssets.IsEmpty())
-				{
-					TArrayView<const TObjectPtr<UCameraRigAsset>> CameraRigs = ReferencingCameraAssets[0]->GetCameraRigs();
-					const TObjectPtr<UCameraRigAsset>* FoundItem = CameraRigs.FindByPredicate(
-							[this](UCameraRigAsset* Item)
-							{
-							return Item->GetDisplayName() == GraphPinObj->DefaultValue;
-							});
-					if (FoundItem)
-					{
-						CameraRigPickerConfig.InitialCameraRigSelection = *FoundItem;
-					}
-				}
-				break;
-			case ECameraRigNameGraphPinMode::ReferencePin:
-				{
-					CameraRigPickerConfig.InitialCameraRigSelection = Cast<UCameraRigAsset>(GraphPinObj->DefaultObject);
-				}
-				break;
+			CameraRigPickerConfig.InitialCameraAssetSelection = DefaultCameraRig->GetTypedOuter<UCameraAsset>();
+			CameraRigPickerConfig.InitialCameraRigSelection = DefaultCameraRig;
 		}
 	}
 
-	IGameplayCamerasEditorModule& CamerasEditorModule = FModuleManager::LoadModuleChecked<IGameplayCamerasEditorModule>("GameplayCamerasEditor");
-	return CamerasEditorModule.CreateCameraRigPicker(CameraRigPickerConfig);
+	OnCustomizeCameraRigPickerConfig(CameraRigPickerConfig);
+
+	return CameraRigPickerConfig;
 }
 
-void SCameraRigNameGraphPin::OnPickerAssetSelected(UCameraRigAsset* SelectedItem)
+void SCameraRigNameGraphPin::OnPickerAssetSelected(UCameraRigAsset* SelectedItem) const
 {
 	if (SelectedItem)
 	{
@@ -228,21 +150,12 @@ FReply SCameraRigNameGraphPin::OnResetButtonClicked()
 	return FReply::Handled();
 }
 
-void SCameraRigNameGraphPin::SetCameraRig(UCameraRigAsset* SelectedCameraRig)
+void SCameraRigNameGraphPin::SetCameraRig(UCameraRigAsset* SelectedCameraRig) const
 {
 	const FScopedTransaction Transaction(LOCTEXT("ChangeObjectPinValue", "Change Object Pin Value"));
 
 	GraphPinObj->Modify();
-
-	switch (PinMode)
-	{
-		case ECameraRigNameGraphPinMode::NamePin:
-			GraphPinObj->GetSchema()->TrySetDefaultValue(*GraphPinObj, SelectedCameraRig->GetDisplayName());
-			break;
-		case ECameraRigNameGraphPinMode::ReferencePin:
-			GraphPinObj->GetSchema()->TrySetDefaultObject(*GraphPinObj, SelectedCameraRig);
-			break;
-	}
+	GraphPinObj->GetSchema()->TrySetDefaultObject(*GraphPinObj, SelectedCameraRig);
 }
 
 }  // namespace UE::Cameras
diff --git a/Engine/Plugins/Cameras/GameplayCameras/Source/GameplayCamerasEditor/Private/Editors/SCameraRigPickerButton.cpp b/Engine/Plugins/Cameras/GameplayCameras/Source/GameplayCamerasEditor/Private/Editors/SCameraRigPickerButton.cpp
new file mode 100644
--- /dev/null
+++ b/Engine/Plugins/Cameras/GameplayCameras/Source/GameplayCamerasEditor/Private/Editors/SCameraRigPickerButton.cpp
@@ -0,0 +1,88 @@
+// Copyright Epic Games, Inc. All Rights Reserved.
+
+#include "Editors/SCameraRigPickerButton.h"
+
+#include "Editors/CameraRigPickerConfig.h"
+#include "IGameplayCamerasEditorModule.h"
+#include "Widgets/Input/SComboButton.h"
+
+#define LOCTEXT_NAMESPACE "SCameraRigPickerButton"
+
+namespace UE::Cameras
+{
+
+void SCameraRigPickerButton::Construct(const FArguments& InArgs)
+{
+	CameraRigPickerConfigAttribute = InArgs._CameraRigPickerConfig;
+	OnGetSelectedCameraRigName = InArgs._SelectedCameraRigName;
+
+	TAttribute<FText> ComboToolTipTextAttribute = InArgs._ButtonToolTipText;
+	if (!ComboToolTipTextAttribute.IsSet())
+	{
+		ComboToolTipTextAttribute = GetDefaultComboToolTipText();
+	}
+
+	ChildSlot
+	[
+		SAssignNew(CameraRigPickerButton, SComboButton)
+			.ButtonStyle(FAppStyle::Get(), "PropertyEditor.AssetComboStyle")
+			.ContentPadding(FMargin(2.f, 2.f, 2.f, 1.f))
+			.ForegroundColor(InArgs._ButtonForegroundColor)
+			.ButtonColorAndOpacity(InArgs._ButtonColorAndOpacity)
+			.MenuPlacement(InArgs._PickerMenuPlacement)
+			.IsEnabled(InArgs._IsEnabled)
+			.ButtonContent()
+			[
+				SNew(STextBlock)
+					.ColorAndOpacity(InArgs._ButtonForegroundColor)
+					.TextStyle(FAppStyle::Get(), "PropertyEditor.AssetClass")
+					.Font(FAppStyle::GetFontStyle("PropertyWindow.NormalFont"))
+					.Text(this, &SCameraRigPickerButton::OnGetComboText)
+					.ToolTipText(ComboToolTipTextAttribute)
+			]
+			.OnGetMenuContent(this, &SCameraRigPickerButton::OnBuildCameraRigNamePicker)
+	];
+}
+
+bool SCameraRigPickerButton::IsOpen() const
+{
+	return CameraRigPickerButton->IsOpen();
+}
+
+void SCameraRigPickerButton::SetIsOpen(bool bIsOpen)
+{
+	CameraRigPickerButton->SetIsOpen(bIsOpen);
+}
+
+FText SCameraRigPickerButton::GetDefaultComboText() const
+{
+	return LOCTEXT("DefaultComboText", "Select Camera Rig");
+}
+
+FText SCameraRigPickerButton::GetDefaultComboToolTipText() const
+{
+	return LOCTEXT("ComboToolTipText", "The selected camera rig, if any.");
+}
+
+FText SCameraRigPickerButton::OnGetComboText() const
+{
+	FText Value = OnGetSelectedCameraRigName.Get();
+	if (Value.IsEmpty())
+	{
+		Value = GetDefaultComboText();
+	}
+	return Value;
+}
+
+TSharedRef<SWidget> SCameraRigPickerButton::OnBuildCameraRigNamePicker()
+{
+	FCameraRigPickerConfig CameraRigPickerConfig = CameraRigPickerConfigAttribute.Get();
+
+	IGameplayCamerasEditorModule& CamerasEditorModule = IGameplayCamerasEditorModule::Get();
+	return CamerasEditorModule.CreateCameraRigPicker(CameraRigPickerConfig);
+}
+
+}  // namespace UE::Cameras
+
+#undef LOCTEXT_NAMESPACE
+
diff --git a/Engine/Plugins/Cameras/GameplayCameras/Source/GameplayCamerasEditor/Private/Helpers/CameraAssetReferenceGatherer.cpp b/Engine/Plugins/Cameras/GameplayCameras/Source/GameplayCamerasEditor/Private/Helpers/CameraAssetReferenceGatherer.cpp
new file mode 100644
--- /dev/null
+++ b/Engine/Plugins/Cameras/GameplayCameras/Source/GameplayCamerasEditor/Private/Helpers/CameraAssetReferenceGatherer.cpp
@@ -0,0 +1,66 @@
+// Copyright Epic Games, Inc. All Rights Reserved.
+
+#include "Helpers/CameraAssetReferenceGatherer.h"
+
+#include "AssetRegistry/AssetData.h"
+#include "AssetRegistry/IAssetRegistry.h"
+#include "Core/CameraAsset.h"
+#include "Misc/AssetRegistryInterface.h"
+#include "UObject/Object.h"
+#include "UObject/Package.h"
+#include "UObject/ReferencerFinder.h"
+
+namespace UE::Cameras
+{
+
+void FCameraAssetReferenceGatherer::GetReferencingCameraAssets(UObject* ReferencedObject, TArray<UCameraAsset*>& OutReferencers)
+{
+	TSet<UCameraAsset*> UniqueReferencers;
+	UPackage* ReferencedObjectPackage = ReferencedObject->GetOutermost();
+
+	// Assume the asset registry module is already loaded.
+	IAssetRegistry* AssetRegistry = IAssetRegistry::Get();
+
+	// Get on-disk referencers.
+	TArray<FAssetIdentifier> ReferencerIds;
+	FAssetIdentifier AssetIdentifier(ReferencedObjectPackage->GetFName());
+	AssetRegistry->GetReferencers(AssetIdentifier, ReferencerIds);
+
+	TArray<FAssetData> AllAssetData;
+	for (FAssetIdentifier ReferencerId : ReferencerIds)
+	{
+		AssetRegistry->GetAssetsByPackageName(ReferencerId.PackageName, AllAssetData);
+	}
+	for (const FAssetData& AssetData : AllAssetData)
+	{
+		UObject* Asset = AssetData.GetAsset();
+		if (UCameraAsset* CameraAsset = Cast<UCameraAsset>(Asset))
+		{
+			UniqueReferencers.Add(CameraAsset);
+		}
+	}
+
+	// Get in-memory referencers.
+	TArray<UObject*> ReferencedObjects;
+	GetObjectsWithPackage(ReferencedObjectPackage, ReferencedObjects);
+	ReferencedObjects.Add(ReferencedObjectPackage);
+
+	TArray<UObject*> AllReferencers = FReferencerFinder::GetAllReferencers(
+			ReferencedObjects, nullptr, EReferencerFinderFlags::SkipWeakReferences);
+	for (UObject* Referencer : AllReferencers)
+	{
+		if (UCameraAsset* CameraAsset = Cast<UCameraAsset>(Referencer))
+		{
+			UniqueReferencers.Add(CameraAsset);
+		}
+		else if (UCameraAsset* OuterCameraAsset = Referencer->GetTypedOuter<UCameraAsset>())
+		{
+			UniqueReferencers.Add(OuterCameraAsset);
+		}
+	}
+
+	OutReferencers.Append(UniqueReferencers.Array());
+}
+
+}  // namespace UE::Cameras
+
diff --git a/Engine/Plugins/Cameras/GameplayCameras/Source/GameplayCamerasEditor/Private/Helpers/CameraDirectorHelper.cpp b/Engine/Plugins/Cameras/GameplayCameras/Source/GameplayCamerasEditor/Private/Helpers/CameraDirectorHelper.cpp
deleted file mode 100644
diff --git a/Engine/Plugins/Cameras/GameplayCameras/Source/GameplayCamerasEditor/Public/Editors/CameraRigPickerConfig.h b/Engine/Plugins/Cameras/GameplayCameras/Source/GameplayCamerasEditor/Public/Editors/CameraRigPickerConfig.h
--- a/Engine/Plugins/Cameras/GameplayCameras/Source/GameplayCamerasEditor/Public/Editors/CameraRigPickerConfig.h
+++ b/Engine/Plugins/Cameras/GameplayCameras/Source/GameplayCamerasEditor/Public/Editors/CameraRigPickerConfig.h
@@ -38,7 +38,7 @@ struct FCameraRigPickerConfig
 	 * determines the only asset to be used. If the latter (when the value is false),
 	 * then no camera asset picker is shown. Only the list of rigs is shown.
 	 */
-	bool bCanSelectCameraAsset = false;
+	bool bCanSelectCameraAsset = true;
 
 	/** Asset picker view type for the camera asset picker. */
 	EAssetViewType::Type CameraAssetViewType = EAssetViewType::List;
diff --git a/Engine/Plugins/Cameras/GameplayCameras/Source/GameplayCamerasEditor/Public/Editors/GameplayCamerasGraphPanelPinFactory.h b/Engine/Plugins/Cameras/GameplayCameras/Source/GameplayCamerasEditor/Public/Editors/GameplayCamerasGraphPanelPinFactory.h
--- a/Engine/Plugins/Cameras/GameplayCameras/Source/GameplayCamerasEditor/Public/Editors/GameplayCamerasGraphPanelPinFactory.h
+++ b/Engine/Plugins/Cameras/GameplayCameras/Source/GameplayCamerasEditor/Public/Editors/GameplayCamerasGraphPanelPinFactory.h
@@ -4,6 +4,9 @@
 
 #include "EdGraphUtilities.h"
 
+class SGraphPin;
+class UK2Node_CallFunction;
+
 namespace UE::Cameras
 {
 
@@ -15,7 +18,12 @@ struct FGameplayCamerasGraphPanelPinFactory : public FGraphPanelPinFactory
 public:
 
 	// FGraphPanelPinFactory interface.
-	virtual TSharedPtr<class SGraphPin> CreatePin(UEdGraphPin* Pin) const override;
+	virtual TSharedPtr<SGraphPin> CreatePin(UEdGraphPin* Pin) const override;
+
+private:
+
+	TSharedPtr<SGraphPin> CreateFunctionParameterPin(UEdGraphPin* Pin, UK2Node_CallFunction* CallFunctionNode) const;
+	TSharedPtr<SGraphPin> CreateCustomPin(UEdGraphPin* Pin) const;
 };
 
 }  // namespace UE::Cameras
diff --git a/Engine/Plugins/Cameras/GameplayCameras/Source/GameplayCamerasEditor/Public/Editors/SBlueprintCameraDirectorRigNameGraphPin.h b/Engine/Plugins/Cameras/GameplayCameras/Source/GameplayCamerasEditor/Public/Editors/SBlueprintCameraDirectorRigNameGraphPin.h
new file mode 100644
--- /dev/null
+++ b/Engine/Plugins/Cameras/GameplayCameras/Source/GameplayCamerasEditor/Public/Editors/SBlueprintCameraDirectorRigNameGraphPin.h
@@ -0,0 +1,30 @@
+// Copyright Epic Games, Inc. All Rights Reserved.
+
+#pragma once
+
+#include "Editors/SCameraRigNameGraphPin.h"
+
+namespace UE::Cameras
+{
+
+/**
+ * A version of the camera rig picker pin widget to be used specifically with 
+ * the ActivateCameraRig method of a BlueprintCameraDirector.
+ */
+class SBlueprintCameraDirectorRigNameGraphPin : public SCameraRigNameGraphPin
+{
+public:
+
+	SLATE_BEGIN_ARGS(SBlueprintCameraDirectorRigNameGraphPin)
+	{}
+	SLATE_END_ARGS()
+
+	void Construct(const FArguments& InArgs, UEdGraphPin* InGraphPinObj);
+
+protected:
+
+	virtual void OnCustomizeCameraRigPickerConfig(FCameraRigPickerConfig& CameraRigPickerConfig) const override;
+};
+
+}  // namespace UE::Cameras
+
diff --git a/Engine/Plugins/Cameras/GameplayCameras/Source/GameplayCamerasEditor/Public/Editors/SCameraRigNameGraphPin.h b/Engine/Plugins/Cameras/GameplayCameras/Source/GameplayCamerasEditor/Public/Editors/SCameraRigNameGraphPin.h
--- a/Engine/Plugins/Cameras/GameplayCameras/Source/GameplayCamerasEditor/Public/Editors/SCameraRigNameGraphPin.h
+++ b/Engine/Plugins/Cameras/GameplayCameras/Source/GameplayCamerasEditor/Public/Editors/SCameraRigNameGraphPin.h
@@ -2,21 +2,16 @@
 
 #pragma once
 
-#include "SGraphPin.h"
-
 #include "ContentBrowserModule.h"
-#include "Widgets/Input/SComboButton.h"
+#include "Editors/CameraRigPickerConfig.h"
+#include "SGraphPin.h"
 
 class UCameraRigAsset;
 
 namespace UE::Cameras
 {
 
-enum class ECameraRigNameGraphPinMode
-{
-	NamePin,
-	ReferencePin
-};
+class SCameraRigPickerButton;
 
 /**
  * A custom widget for a graph editor pin that shows a camera rig picker dialog.
@@ -26,9 +21,7 @@ class SCameraRigNameGraphPin : public SGraphPin
 public:
 
 	SLATE_BEGIN_ARGS(SCameraRigNameGraphPin)
-		: _PinMode(ECameraRigNameGraphPinMode::NamePin)
 	{}
-		SLATE_ARGUMENT(ECameraRigNameGraphPinMode, PinMode)
 	SLATE_END_ARGS()
 
 	void Construct(const FArguments& InArgs, UEdGraphPin* InGraphPinObj);
@@ -39,6 +32,8 @@ protected:
 	virtual TSharedRef<SWidget>	GetDefaultValueWidget() override;
 	virtual bool DoesWidgetHandleSettingEditingEnabled() const override;
 
+	virtual void OnCustomizeCameraRigPickerConfig(FCameraRigPickerConfig& CameraRigPickerConfig) const {}
+
 private:
 
 	static constexpr float ActiveComboAlpha = 1.f;
@@ -52,21 +47,18 @@ private:
 	FSlateColor OnGetWidgetForeground() const;
 	FSlateColor OnGetWidgetBackground() const;
 
-	FText GetDefaultComboText() const;
-	FText OnGetComboText() const;
-	FText OnGetComboToolTipText() const;
-
-	TSharedRef<SWidget> OnBuildCameraRigNamePicker();
-	void OnPickerAssetSelected(UCameraRigAsset* SelectedItem);
+	FCameraRigPickerConfig OnCreateCameraRigPickerConfig() const;
+	FText OnGetSelectedCameraRigName() const;
+	FText OnGetCameraRigPickerToolTipText() const;
+	void OnPickerAssetSelected(UCameraRigAsset* SelectedItem) const;
 
 	FReply OnResetButtonClicked();
 
-	void SetCameraRig(UCameraRigAsset* SelectedCameraRig);
+	void SetCameraRig(UCameraRigAsset* SelectedCameraRig) const;
 
-private:
+protected:
 
-	TSharedPtr<SComboButton> CameraRigPickerButton;
-	ECameraRigNameGraphPinMode PinMode;
+	TSharedPtr<SCameraRigPickerButton> CameraRigPickerButton;
 };
 
 }  // namespace UE::Cameras
diff --git a/Engine/Plugins/Cameras/GameplayCameras/Source/GameplayCamerasEditor/Public/Editors/SCameraRigPickerButton.h b/Engine/Plugins/Cameras/GameplayCameras/Source/GameplayCamerasEditor/Public/Editors/SCameraRigPickerButton.h
new file mode 100644
--- /dev/null
+++ b/Engine/Plugins/Cameras/GameplayCameras/Source/GameplayCamerasEditor/Public/Editors/SCameraRigPickerButton.h
@@ -0,0 +1,56 @@
+// Copyright Epic Games, Inc. All Rights Reserved.
+
+#pragma once
+
+#include "Editors/CameraRigPickerConfig.h"
+#include "Widgets/SCompoundWidget.h"
+
+class SComboButton;
+class UCameraRigAsset;
+
+namespace UE::Cameras
+{
+
+/**
+ * A simple combo button that shows a camera rig picker dialog.
+ */
+class SCameraRigPickerButton : public SCompoundWidget
+{
+public:
+
+	SLATE_BEGIN_ARGS(SCameraRigPickerButton)
+	{}
+		SLATE_ATTRIBUTE(FCameraRigPickerConfig, CameraRigPickerConfig)
+		SLATE_ATTRIBUTE(FText, SelectedCameraRigName)
+		SLATE_ATTRIBUTE(FText, ButtonToolTipText)
+		SLATE_ATTRIBUTE(FSlateColor, ButtonForegroundColor)
+		SLATE_ATTRIBUTE(FSlateColor, ButtonColorAndOpacity)
+		SLATE_ATTRIBUTE(EMenuPlacement, PickerMenuPlacement)
+		SLATE_ATTRIBUTE(bool, IsEnabled)
+		SLATE_ATTRIBUTE(UObject*, ShowOnlyCameraAssetsReferencingObject)
+	SLATE_END_ARGS()
+
+	void Construct(const FArguments& InArgs);
+
+	bool IsOpen() const;
+	void SetIsOpen(bool bIsOpen);
+
+private:
+
+	FText GetDefaultComboText() const;
+	FText GetDefaultComboToolTipText() const;
+	FText OnGetComboText() const;
+
+	TSharedRef<SWidget> OnBuildCameraRigNamePicker();
+
+private:
+
+	TSharedPtr<SComboButton> CameraRigPickerButton;
+
+	TAttribute<FCameraRigPickerConfig> CameraRigPickerConfigAttribute;
+	TAttribute<FText> OnGetSelectedCameraRigName;
+	TAttribute<UObject*> ShowOnlyCameraAssetsReferencingObjectAttribute;
+};
+
+}  // namespace UE::Cameras
+
diff --git a/Engine/Plugins/Cameras/GameplayCameras/Source/GameplayCamerasEditor/Public/Helpers/CameraAssetReferenceGatherer.h b/Engine/Plugins/Cameras/GameplayCameras/Source/GameplayCamerasEditor/Public/Helpers/CameraAssetReferenceGatherer.h
new file mode 100644
--- /dev/null
+++ b/Engine/Plugins/Cameras/GameplayCameras/Source/GameplayCamerasEditor/Public/Helpers/CameraAssetReferenceGatherer.h
@@ -0,0 +1,27 @@
+// Copyright Epic Games, Inc. All Rights Reserved.
+
+#pragma once
+
+#include "CoreTypes.h"
+#include "Containers/Array.h"
+
+class UCameraAsset;
+class UObject;
+
+namespace UE::Cameras
+{
+
+/**
+ * A helper class for gathering camera assets referencing a given object, such as a
+ * camera director dependency.
+ */
+class FCameraAssetReferenceGatherer
+{
+public:
+
+	/** Gets the list of camera assets that reference the given object. */
+	static void GetReferencingCameraAssets(UObject* ReferencedObject, TArray<UCameraAsset*>& OutReferencers);
+};
+
+}  // namespace UE::Cameras
+
diff --git a/Engine/Plugins/Cameras/GameplayCameras/Source/GameplayCamerasEditor/Public/Helpers/CameraDirectorHelper.h b/Engine/Plugins/Cameras/GameplayCameras/Source/GameplayCamerasEditor/Public/Helpers/CameraDirectorHelper.h
deleted file mode 100644
diff --git a/Engine/Plugins/Cameras/GameplayCameras/Source/GameplayCamerasUncookedOnly/Private/BlueprintGraph/K2Node_SetCameraRigParameters.cpp b/Engine/Plugins/Cameras/GameplayCameras/Source/GameplayCamerasUncookedOnly/Private/BlueprintGraph/K2Node_SetCameraRigParameters.cpp
--- a/Engine/Plugins/Cameras/GameplayCameras/Source/GameplayCamerasUncookedOnly/Private/BlueprintGraph/K2Node_SetCameraRigParameters.cpp
+++ b/Engine/Plugins/Cameras/GameplayCameras/Source/GameplayCamerasUncookedOnly/Private/BlueprintGraph/K2Node_SetCameraRigParameters.cpp
@@ -26,12 +26,14 @@ UK2Node_SetCameraRigParameters::UK2Node_SetCameraRigParameters(const FObjectInit
 
 void UK2Node_SetCameraRigParameters::AllocateDefaultPins()
 {
+	using namespace UE::Cameras;
+
 	// Add execution pins.
 	CreatePin(EGPD_Input, UEdGraphSchema_K2::PC_Exec, UEdGraphSchema_K2::PN_Execute);
 	CreatePin(EGPD_Output, UEdGraphSchema_K2::PC_Exec, UEdGraphSchema_K2::PN_Then);
 
 	// Add evalation result pin.
-	CreatePin(EGPD_Input, UEdGraphSchema_K2::PC_Object, FBlueprintCameraVariableTable::StaticStruct(), CameraVariableTablePinName);
+	CreatePin(EGPD_Input, UEdGraphSchema_K2::PC_Struct, FBlueprintCameraVariableTable::StaticStruct(), CameraVariableTablePinName);
 
 	// Add camera rig pin.
 	CreatePin(EGPD_Input, UEdGraphSchema_K2::PC_Object, UCameraRigAsset::StaticClass(), CameraRigPinName);
