Cameras: add collision push camera node

diff --git a/Engine/Plugins/Cameras/GameplayCameras/Source/GameplayCameras/Private/Nodes/Collision/CollisionPushCameraNode.cpp b/Engine/Plugins/Cameras/GameplayCameras/Source/GameplayCameras/Private/Nodes/Collision/CollisionPushCameraNode.cpp
new file mode 100644
--- /dev/null
+++ b/Engine/Plugins/Cameras/GameplayCameras/Source/GameplayCameras/Private/Nodes/Collision/CollisionPushCameraNode.cpp
@@ -0,0 +1,360 @@
+// Copyright Epic Games, Inc. All Rights Reserved.
+
+#include "Nodes/Collision/CollisionPushCameraNode.h"
+
+#include "CollisionQueryParams.h"
+#include "Core/BuiltInCameraVariables.h"
+#include "Core/CameraEvaluationContext.h"
+#include "Core/CameraParameterReader.h"
+#include "Core/CameraSystemEvaluator.h"
+#include "Core/CameraValueInterpolator.h"
+#include "Debug/CameraDebugBlock.h"
+#include "Debug/CameraDebugBlockBuilder.h"
+#include "Debug/CameraDebugRenderer.h"
+#include "Engine/HitResult.h"
+#include "Engine/World.h"
+#include "GameFramework/Pawn.h"
+#include "GameFramework/PlayerController.h"
+#include "GameplayCameras.h"
+#include "Misc/AssertionMacros.h"
+#include "WorldCollision.h"
+
+#include UE_INLINE_GENERATED_CPP_BY_NAME(CollisionPushCameraNode)
+
+namespace UE::Cameras
+{
+
+class FCollisionPushCameraNodeEvaluator : public FCameraNodeEvaluator
+{
+	UE_DECLARE_CAMERA_NODE_EVALUATOR(GAMEPLAYCAMERAS_API, FCollisionPushCameraNodeEvaluator)
+
+protected:
+
+	virtual void OnInitialize(const FCameraNodeEvaluatorInitializeParams& Params, FCameraNodeEvaluationResult& OutResult) override;
+	virtual void OnRun(const FCameraNodeEvaluationParams& Params, FCameraNodeEvaluationResult& OutResult) override;
+
+#if UE_GAMEPLAY_CAMERAS_DEBUG
+	virtual void OnBuildDebugBlocks(const FCameraDebugBlockBuildParams& Params, FCameraDebugBlockBuilder& Builder) override;
+#endif  // UE_GAMEPLAY_CAMERAS_DEBUG
+
+private:
+
+	TOptional<FVector3d> GetSafePosition(const APawn* Pawn, const FCameraNodeEvaluationResult& OutResult);
+	void RunCollisionTrace(UWorld* World, APlayerController* PlayerController, const FVector3d& SafePosition, const FCameraNodeEvaluationParams& Params, FCameraNodeEvaluationResult& OutResult);
+	void HandleAsyncCollisionTraceResult(UWorld* World, const FVector3d& SafePosition, const FCameraNodeEvaluationParams& Params, FCameraNodeEvaluationResult& OutResult);
+	void HandleCollisionTraceResult(UWorld* World, TArrayView<const FHitResult> HitResults, const FVector3d& SafePosition, const FCameraNodeEvaluationParams& Params, FCameraNodeEvaluationResult& OutResult);
+
+private:
+
+	TCameraParameterReader<float> CollisionSphereRadiusReader;
+	TCameraParameterReader<FVector3d> SafePositionOffsetReader;
+
+	TUniquePtr<FCameraDoubleValueInterpolator> PushInterpolator;
+	TUniquePtr<FCameraDoubleValueInterpolator> PullInterpolator;
+
+	FTraceHandle CollisionTraceHandle;
+
+	float LastPushFactor = 0.f;
+	float LastDampedPushFactor = 0.f;
+
+	enum class ECameraCollisionDirection { Pushing, Pulling };
+	ECameraCollisionDirection LastDirection = ECameraCollisionDirection::Pushing;
+
+#if UE_GAMEPLAY_CAMERAS_DEBUG
+	bool bDebugFoundHit = false;
+#endif
+};
+
+UE_DEFINE_CAMERA_NODE_EVALUATOR(FCollisionPushCameraNodeEvaluator)
+
+UE_DECLARE_CAMERA_DEBUG_BLOCK_START(GAMEPLAYCAMERAS_API, FCollisionPushCameraDebugBlock)
+	UE_DECLARE_CAMERA_DEBUG_BLOCK_FIELD(float, PushFactor);
+	UE_DECLARE_CAMERA_DEBUG_BLOCK_FIELD(float, DampedPushFactor);
+	UE_DECLARE_CAMERA_DEBUG_BLOCK_FIELD(bool, bIsPulling)
+	UE_DECLARE_CAMERA_DEBUG_BLOCK_FIELD(bool, bFoundHit)
+UE_DECLARE_CAMERA_DEBUG_BLOCK_END()
+
+UE_DEFINE_CAMERA_DEBUG_BLOCK_WITH_FIELDS(FCollisionPushCameraDebugBlock)
+
+void FCollisionPushCameraNodeEvaluator::OnInitialize(const FCameraNodeEvaluatorInitializeParams& Params, FCameraNodeEvaluationResult& OutResult)
+{
+	const UCollisionPushCameraNode* CollisionPushNode = GetCameraNodeAs<UCollisionPushCameraNode>();
+
+	CollisionSphereRadiusReader.Initialize(CollisionPushNode->CollisionSphereRadius);
+	SafePositionOffsetReader.Initialize(CollisionPushNode->SafePositionOffset);
+
+	PushInterpolator = CollisionPushNode->PushInterpolator ?
+		CollisionPushNode->PushInterpolator->BuildDoubleInterpolator() :
+		MakeUnique<TPopValueInterpolator<double>>();
+	PullInterpolator = CollisionPushNode->PullInterpolator ?
+		CollisionPushNode->PullInterpolator->BuildDoubleInterpolator() :
+		MakeUnique<TPopValueInterpolator<double>>();
+}
+
+void FCollisionPushCameraNodeEvaluator::OnRun(const FCameraNodeEvaluationParams& Params, FCameraNodeEvaluationResult& OutResult)
+{
+	if (!ensure(Params.EvaluationContext))
+	{
+		return;
+	}
+
+	UWorld* World = Params.EvaluationContext->GetWorld();
+	APlayerController* PlayerController = Params.EvaluationContext->GetPlayerController();
+	if (!World || !PlayerController)
+	{
+		return;
+	}
+
+	APawn* Pawn = PlayerController->GetPawn();
+	TOptional<FVector3d> SafePosition = GetSafePosition(Pawn, OutResult);
+	if (!SafePosition.IsSet())
+	{
+		const UCameraNode* ThisNode = GetCameraNode();
+		UE_LOG(LogCameraSystem, Error, 
+				TEXT("Can't find safe position for CollisionPushCameraNode '%s'."),
+				*GetNameSafe(ThisNode));
+		return;
+	}
+
+	if (Params.EvaluationType != ECameraNodeEvaluationType::Standard)
+	{
+		// Don't run collision traces during IK/stateless updates.
+		// Push the camera by the same amount as last time we updated properly, if possible.
+		if (SafePosition.IsSet())
+		{
+			const FVector3d CameraPoseLocation = OutResult.CameraPose.GetLocation();
+			const FVector3d PushedLocation = CameraPoseLocation + (SafePosition.GetValue() - CameraPoseLocation) * LastDampedPushFactor;
+			OutResult.CameraPose.SetLocation(PushedLocation);
+		}
+		return;
+	}
+
+	HandleAsyncCollisionTraceResult(World, SafePosition.GetValue(), Params, OutResult);
+	RunCollisionTrace(World, PlayerController, SafePosition.GetValue(), Params, OutResult);
+}
+
+TOptional<FVector3d> FCollisionPushCameraNodeEvaluator::GetSafePosition(const APawn* Pawn, const FCameraNodeEvaluationResult& OutResult)
+{
+	FVector3d SafePosition;
+	bool bGotSafePosition = false;
+
+	// See if we can use the camera rig pivot as the safe position.
+	// Otherwise, use the pawn actor's position.
+	const FBuiltInCameraVariables& BuiltInVariables = FBuiltInCameraVariables::Get();
+	TArrayView<const FCameraRigJoint> Joints = OutResult.CameraRigJoints.GetJoints();
+	const FCameraRigJoint* PivotJoint = Joints.FindByPredicate([&BuiltInVariables](const FCameraRigJoint& Joint)
+			{
+				return Joint.VariableID == BuiltInVariables.YawPitchDefinition;
+			});
+	if (PivotJoint)
+	{
+		SafePosition = PivotJoint->Transform.GetLocation();
+		bGotSafePosition = true;
+	}
+	else if (Pawn)
+	{
+		SafePosition = Pawn->GetActorLocation();
+		bGotSafePosition = true;
+	}
+
+	if (bGotSafePosition)
+	{
+		const FVector3d SafePositionOffset = SafePositionOffsetReader.Get(OutResult.VariableTable);
+		return SafePosition + SafePositionOffset;
+	}
+	return TOptional<FVector3d>();
+}
+
+void FCollisionPushCameraNodeEvaluator::RunCollisionTrace(UWorld* World, APlayerController* PlayerController, const FVector3d& SafePosition, const FCameraNodeEvaluationParams& Params, FCameraNodeEvaluationResult& OutResult)
+{
+	static FName CollisionTraceTag(TEXT("CameraCollision"));
+	static FName CollisionTraceOwnerTag(TEXT("CollisionPushCameraNode"));
+
+	const UCollisionPushCameraNode* CollisionPushNode = GetCameraNodeAs<UCollisionPushCameraNode>();
+	ECollisionChannel CollisionChannel = CollisionPushNode->CollisionChannel;
+
+	const float CollisionSphereRadius = CollisionSphereRadiusReader.Get(OutResult.VariableTable);
+	const FVector3d SafePositionOffset = SafePositionOffsetReader.Get(OutResult.VariableTable);
+
+	const FVector3d TraceStart(SafePosition);
+	const FVector3d TraceEnd(OutResult.CameraPose.GetLocation());
+
+	double TraceLength = FVector3d::Distance(TraceStart, TraceEnd);
+	if (!ensure(TraceLength > 0))
+	{
+		return;
+	}
+
+	FCollisionShape SweepShape = FCollisionShape::MakeSphere(CollisionSphereRadius);
+	// Ignore the player pawn by default.
+	APawn* Pawn = PlayerController->GetPawn();
+	FCollisionQueryParams QueryParams(SCENE_QUERY_STAT(StartCollisionSweep), false, Pawn);
+	QueryParams.TraceTag = CollisionTraceTag;
+	QueryParams.OwnerTag = CollisionTraceOwnerTag;
+
+	if (CollisionPushNode->bRunAsyncCollision)
+	{
+		CollisionTraceHandle = World->AsyncSweepByChannel(
+			EAsyncTraceType::Multi,
+			TraceStart, TraceEnd, FQuat::Identity,
+			CollisionChannel,
+			SweepShape,
+			QueryParams,
+			FCollisionResponseParams::DefaultResponseParam);
+	}
+	else
+	{
+		TArray<FHitResult> HitResults;
+		World->SweepMultiByChannel(
+			HitResults,
+			TraceStart, TraceEnd, FQuat::Identity,
+			CollisionChannel,
+			SweepShape,
+			QueryParams,
+			FCollisionResponseParams::DefaultResponseParam);
+		HandleCollisionTraceResult(World, HitResults, SafePosition, Params, OutResult);
+	}
+}
+
+void FCollisionPushCameraNodeEvaluator::HandleAsyncCollisionTraceResult(UWorld* World, const FVector3d& SafePosition, const FCameraNodeEvaluationParams& Params, FCameraNodeEvaluationResult& OutResult)
+{
+	const UCollisionPushCameraNode* CollisionPushNode = GetCameraNodeAs<UCollisionPushCameraNode>();
+	if (!CollisionPushNode->bRunAsyncCollision)
+	{
+		return;
+	}
+
+	if (!CollisionTraceHandle.IsValid())
+	{
+		return;
+	}
+
+	FTraceDatum TraceDatum;
+	if (!World->QueryTraceData(CollisionTraceHandle, TraceDatum))
+	{
+		return;
+	}
+
+	HandleCollisionTraceResult(World, TraceDatum.OutHits, SafePosition, Params, OutResult);
+}
+
+void FCollisionPushCameraNodeEvaluator::HandleCollisionTraceResult(UWorld* World, TArrayView<const FHitResult> HitResults, const FVector3d& SafePosition, const FCameraNodeEvaluationParams& Params, FCameraNodeEvaluationResult& OutResult)
+{
+	// Find any relevant hit in the trace results.
+	bool bFoundHit = false;
+	float CurrentPushFactor = 0.f;
+
+	for (const FHitResult& Hit : HitResults)
+	{
+		double TraceLength = FVector3d::Distance(Hit.TraceStart, Hit.TraceEnd);
+		double DistanceToHit = FVector3d::Distance(Hit.TraceEnd, Hit.Location);
+		if (ensure(TraceLength > 0))
+		{
+			CurrentPushFactor = (float)(DistanceToHit / TraceLength);
+			bFoundHit = true;
+			break;
+		}
+	}
+
+#if UE_GAMEPLAY_CAMERAS_DEBUG
+	bDebugFoundHit = bFoundHit;
+#endif
+
+	// If we weren't pushed last frame, and we have no reason to push this frame either, then
+	// we can bail out early.
+	if (!bFoundHit && LastDampedPushFactor < UE_SMALL_NUMBER)
+	{
+		LastDirection = ECameraCollisionDirection::Pushing;
+		LastPushFactor = 0;
+		LastDampedPushFactor = 0;
+		return;
+	}
+
+	// Figure out if we're pulling or pushing towards the safe position.
+	// If we had no hit, the current push factor is zero.
+	ECameraCollisionDirection CurrentDirection = LastDirection;
+	if (LastPushFactor < CurrentPushFactor)
+	{
+		CurrentDirection = ECameraCollisionDirection::Pushing;
+	}
+	else if (LastPushFactor > CurrentPushFactor)
+	{
+		CurrentDirection = ECameraCollisionDirection::Pulling;
+	}
+	// else, push factor hasn't changed, keep the same direction as before.
+
+	// Interpolate the push factor to make camera movements smoother.
+	float CurrentDampedPushFactor = LastDampedPushFactor;
+
+	FCameraValueInterpolationParams InterpParams;
+	InterpParams.bIsCameraCut = Params.bIsFirstFrame;
+	InterpParams.DeltaTime = Params.DeltaTime;
+	FCameraValueInterpolationResult InterpResult(OutResult.VariableTable);
+	switch (CurrentDirection)
+	{
+		case ECameraCollisionDirection::Pushing:
+			PushInterpolator->Reset(LastDampedPushFactor, CurrentPushFactor);
+			CurrentDampedPushFactor = PushInterpolator->Run(InterpParams, InterpResult);
+			break;
+		case ECameraCollisionDirection::Pulling:
+			PullInterpolator->Reset(LastDampedPushFactor, CurrentPushFactor);
+			CurrentDampedPushFactor = PullInterpolator->Run(InterpParams, InterpResult);
+			break;
+	}
+
+	// Push the camera!
+	if (CurrentDampedPushFactor > 0)
+	{
+		const FVector3d CameraPoseLocation = OutResult.CameraPose.GetLocation();
+		const FVector3d PushedLocation = CameraPoseLocation + (SafePosition - CameraPoseLocation) * CurrentDampedPushFactor;
+		OutResult.CameraPose.SetLocation(PushedLocation);
+	}
+
+	LastPushFactor = CurrentPushFactor;
+	LastDampedPushFactor = CurrentDampedPushFactor;
+	LastDirection = CurrentDirection;
+}
+
+#if UE_GAMEPLAY_CAMERAS_DEBUG
+
+void FCollisionPushCameraNodeEvaluator::OnBuildDebugBlocks(const FCameraDebugBlockBuildParams& Params, FCameraDebugBlockBuilder& Builder)
+{
+	FCollisionPushCameraDebugBlock& DebugBlock = Builder.AttachDebugBlock<FCollisionPushCameraDebugBlock>();
+
+	DebugBlock.bFoundHit = bDebugFoundHit;
+	DebugBlock.bIsPulling = (LastDirection == ECameraCollisionDirection::Pulling);
+	DebugBlock.PushFactor = LastPushFactor;
+	DebugBlock.DampedPushFactor = LastDampedPushFactor;
+}
+
+void FCollisionPushCameraDebugBlock::OnDebugDraw(const FCameraDebugBlockDrawParams& Params, FCameraDebugRenderer& Renderer)
+{
+	if (DampedPushFactor > 0)
+	{
+		Renderer.AddText(TEXT("need to push by %.2f%%, currently %.2f%% [%s] %s"),
+				PushFactor, DampedPushFactor,
+				bIsPulling ? TEXT("pulling") : TEXT("pushing"),
+				bFoundHit ? TEXT("(colliding!)") : TEXT(""));
+	}
+	else
+	{
+		Renderer.AddText(TEXT("not pushing"));
+	}
+}
+
+#endif  // UE_GAMEPLAY_CAMERAS_DEBUG
+
+}  // namespace UE::Cameras
+
+UCollisionPushCameraNode::UCollisionPushCameraNode(const FObjectInitializer& ObjectInit)
+	: Super(ObjectInit)
+{
+	CollisionSphereRadius.Value = 10.f;
+}
+
+FCameraNodeEvaluatorPtr UCollisionPushCameraNode::OnBuildEvaluator(FCameraNodeEvaluatorBuilder& Builder) const
+{
+	using namespace UE::Cameras;
+	return Builder.BuildEvaluator<FCollisionPushCameraNodeEvaluator>();
+}
+
diff --git a/Engine/Plugins/Cameras/GameplayCameras/Source/GameplayCameras/Public/Nodes/Collision/CollisionPushCameraNode.h b/Engine/Plugins/Cameras/GameplayCameras/Source/GameplayCameras/Public/Nodes/Collision/CollisionPushCameraNode.h
new file mode 100644
--- /dev/null
+++ b/Engine/Plugins/Cameras/GameplayCameras/Source/GameplayCameras/Public/Nodes/Collision/CollisionPushCameraNode.h
@@ -0,0 +1,61 @@
+// Copyright Epic Games, Inc. All Rights Reserved.
+
+#pragma once
+
+#include "Core/CameraNode.h"
+#include "Core/CameraParameters.h"
+#include "Engine/EngineTypes.h"
+
+#include "CollisionPushCameraNode.generated.h"
+
+class UCameraValueInterpolator;
+
+/**
+ * A node that pushes the camera towards a "safe position" when it is colliding with 
+ * the environment. By default, the "safe position" is the pivot of the camera (if any) 
+ * or the position of the player pawn.
+ */
+UCLASS(MinimalAPI, meta=(CameraNodeCategories="Collision"))
+class UCollisionPushCameraNode : public UCameraNode
+{
+	GENERATED_BODY()
+
+public:
+
+	/** Radius of the sphere used for collision testing. */
+	UPROPERTY(EditAnywhere, Category="Collision")
+	FFloatCameraParameter CollisionSphereRadius;
+
+	/** Collision channel to use for the line trace. */
+	UPROPERTY(EditAnywhere, Category="Collision")
+	TEnumAsByte<ECollisionChannel> CollisionChannel;
+
+	/** World-space offset from the target to the line trace's end. */
+	UPROPERTY(EditAnywhere, Category="Occlusion")
+	FVector3dCameraParameter SafePositionOffset;
+
+	/** The interpolation to use when pushing the camera towards the safe position. */
+	UPROPERTY(EditAnywhere, Category="Collision")
+	TObjectPtr<UCameraValueInterpolator> PushInterpolator;
+
+	/** The interpolation to use when pulling the camera back to its ideal position. */
+	UPROPERTY(EditAnywhere, Category="Collision")
+	TObjectPtr<UCameraValueInterpolator> PullInterpolator;
+
+	/**
+	 * Whether to run the collision asynchrnously. 
+	 * This is better for performance, but results in collision handling being one frame late.
+	 */
+	UPROPERTY(EditAnywhere, Category="Collision")
+	bool bRunAsyncCollision = false;
+
+public:
+
+	UCollisionPushCameraNode(const FObjectInitializer& ObjectInit);
+
+protected:
+
+	// UCameraNode interface.
+	virtual FCameraNodeEvaluatorPtr OnBuildEvaluator(FCameraNodeEvaluatorBuilder& Builder) const override;
+};
+
