Cameras: make a standalone camera variable picker widget, use it for both camera parameters and graph pins

The camera variable picker widget was implemented in the camera parameter details customization. But we also need it for Blueprint nodes, for instance, if they have properties of type UCameraVariableAsset.
So this change extracts the widget code into its own files, and the graph pin factory can use it to make a graph pin that shows the camera variable picker.

diff --git a/Engine/Plugins/Cameras/GameplayCameras/Source/GameplayCamerasEditor/Private/Customizations/CameraParameterDetailsCustomizations.cpp b/Engine/Plugins/Cameras/GameplayCameras/Source/GameplayCamerasEditor/Private/Customizations/CameraParameterDetailsCustomizations.cpp
--- a/Engine/Plugins/Cameras/GameplayCameras/Source/GameplayCamerasEditor/Private/Customizations/CameraParameterDetailsCustomizations.cpp
+++ b/Engine/Plugins/Cameras/GameplayCameras/Source/GameplayCamerasEditor/Private/Customizations/CameraParameterDetailsCustomizations.cpp
@@ -8,191 +8,28 @@
 #include "Core/CameraVariableCollection.h"
 #include "Customizations/MathStructCustomizations.h"
 #include "DetailLayoutBuilder.h"
-#include "Framework/Application/SlateApplication.h"
+#include "Editors/CameraVariablePickerConfig.h"
 #include "Framework/MultiBox/MultiBoxBuilder.h"
-#include "Framework/Views/ITypedTableView.h"
 #include "IContentBrowserSingleton.h"
 #include "IDetailChildrenBuilder.h"
 #include "IDetailPropertyRow.h"
+#include "IGameplayCamerasEditorModule.h"
 #include "IPropertyUtilities.h"
-#include "Modules/ModuleManager.h"
 #include "PropertyCustomizationHelpers.h"
 #include "PropertyEditorModule.h"
 #include "ScopedTransaction.h"
 #include "Styles/GameplayCamerasEditorStyle.h"
 #include "Styling/StyleColors.h"
 #include "Widgets/Images/SImage.h"
-#include "Widgets/Input/SButton.h"
 #include "Widgets/Input/SComboButton.h"
 #include "Widgets/SBoxPanel.h"
 #include "Widgets/Text/STextBlock.h"
-#include "Widgets/Views/SListView.h"
 
 #define LOCTEXT_NAMESPACE "CameraParameterDetailsCustomization"
 
 namespace UE::Cameras
 {
 
-DECLARE_DELEGATE_OneParam(FExecuteSetParameterVariable, UCameraVariableAsset*);
-
-class SCameraVariableBrowser : public SCompoundWidget
-{
-public:
-
-	SLATE_BEGIN_ARGS(SCameraVariableBrowser)
-	{}
-		SLATE_ARGUMENT(UClass*, VariableClass)
-		SLATE_ARGUMENT(UCameraVariableAsset*, InitialVariable)
-		SLATE_EVENT(FExecuteSetParameterVariable, RequestSetVariable)
-	SLATE_END_ARGS()
-	
-	void Construct(const FArguments& InArgs)
-	{
-		VariableClass = InArgs._VariableClass;
-		InitialVariable = InArgs._InitialVariable;
-
-		RequestSetVariable = InArgs._RequestSetVariable;
-
-		ChildSlot
-		[
-			SNew(SBorder)
-			.BorderImage(FAppStyle::GetBrush("ToolPanel.GroupBorder"))
-			[
-				SNew(SVerticalBox)
-				+SVerticalBox::Slot()
-				.FillHeight(0.5f)
-				[
-					BuildVariableCollectionAssetPicker()
-				]
-				+SVerticalBox::Slot()
-				.FillHeight(0.5f)
-				[
-					SAssignNew(CameraVariableListView, SListView<UCameraVariableAsset*>)
-					.ListItemsSource(&CameraVariableItemsSource)
-					.OnGenerateRow(this, &SCameraVariableBrowser::OnVariableListGenerateRow)
-					.OnSelectionChanged(this, &SCameraVariableBrowser::OnVariableListSelectionChanged)
-				]
-			]
-		];
-
-		if (InitialVariable)
-		{
-			bSuppressVariableListSelectionChanged = true;
-			CameraVariableListView->SetSelection(InitialVariable);
-			CameraVariableListView->RequestScrollIntoView(InitialVariable);
-			bSuppressVariableListSelectionChanged = false;
-		}
-	}
-
-private:
-
-	TSharedRef<SWidget> BuildVariableCollectionAssetPicker()
-	{
-		FContentBrowserModule& ContentBrowserModule = FModuleManager::Get().LoadModuleChecked<FContentBrowserModule>(TEXT("ContentBrowser"));
-
-		FAssetPickerConfig AssetPickerConfig;
-
-		FARFilter ARFilter;
-		ARFilter.ClassPaths.Add(FTopLevelAssetPath(UCameraVariableCollection::StaticClass()->GetPathName()));
-
-		UCameraVariableCollection* InitialVariableCollection = nullptr;
-		if (InitialVariable)
-		{
-			InitialVariableCollection = InitialVariable->GetTypedOuter<UCameraVariableCollection>();
-		}
-
-		AssetPickerConfig.bAllowDragging = false;
-		AssetPickerConfig.bCanShowClasses = false;
-		AssetPickerConfig.bAllowNullSelection = false;
-		AssetPickerConfig.bShowBottomToolbar = true;
-		AssetPickerConfig.bFocusSearchBoxWhenOpened = true;
-		AssetPickerConfig.SelectionMode = ESelectionMode::Single;
-		AssetPickerConfig.Filter = ARFilter;
-		AssetPickerConfig.SaveSettingsName = TEXT("CameraParameterVariablePropertyPicker");
-		AssetPickerConfig.InitialAssetViewType = EAssetViewType::List;
-		AssetPickerConfig.InitialAssetSelection = FAssetData(InitialVariableCollection);
-		AssetPickerConfig.OnAssetSelected = FOnAssetSelected::CreateSP(this, &SCameraVariableBrowser::OnAssetSelected);
-		AssetPickerConfig.GetCurrentSelectionDelegates.Add(&GetCurrentAssetPickerSelection);
-
-		return ContentBrowserModule.Get().CreateAssetPicker(AssetPickerConfig);
-	}
-
-	void OnAssetSelected(const FAssetData& SelectedAsset)
-	{
-		UpdateVariableListItemsSource();
-	}
-
-	void UpdateVariableListItemsSource()
-	{
-		TArray<FAssetData> SelectedAssets;
-		if (GetCurrentAssetPickerSelection.IsBound())
-		{
-			SelectedAssets = GetCurrentAssetPickerSelection.Execute();
-		}
-
-		CameraVariableItemsSource.Reset();
-		if (!SelectedAssets.IsEmpty())
-		{
-			if (UCameraVariableCollection* VariableCollection = Cast<UCameraVariableCollection>(SelectedAssets[0].GetAsset()))
-			{
-				CameraVariableItemsSource = VariableCollection->Variables.FilterByPredicate(
-						[this](UCameraVariableAsset* Item) { return Item->GetClass() == VariableClass; });
-			}
-		}
-
-		CameraVariableListView->RequestListRefresh();
-	}
-
-	TSharedRef<ITableRow> OnVariableListGenerateRow(UCameraVariableAsset* Item, const TSharedRef<STableViewBase>& OwnerTable)
-	{
-		TSharedRef<FGameplayCamerasEditorStyle> GameplayCamerasStyle = FGameplayCamerasEditorStyle::Get();
-
-		const FText DisplayName = Item->DisplayName.IsEmpty() ?
-			FText::FromName(Item->GetFName()) : FText::FromString(Item->DisplayName);
-
-		return SNew(STableRow<UCameraVariableAsset*>, OwnerTable)
-			[
-				SNew(SHorizontalBox)
-				+SHorizontalBox::Slot()
-				.AutoWidth()
-				[
-					SNew(SImage)
-					.ColorAndOpacity(FSlateColor::UseForeground())
-					.Image(GameplayCamerasStyle->GetBrush("CameraParameter.VariableBrowser"))
-				]
-				+SHorizontalBox::Slot()
-				.FillWidth(1.f)
-				.Padding(4.f, 2.f)
-				[
-					SNew(STextBlock)
-					.Text(DisplayName)
-				]
-			];
-	}
-
-	void OnVariableListSelectionChanged(UCameraVariableAsset* Item, ESelectInfo::Type SelectInfo) const
-	{
-		if (!bSuppressVariableListSelectionChanged && Item)
-		{
-			RequestSetVariable.ExecuteIfBound(Item);
-		}
-		// else, the user should click the "Clear" button to actually remove the variable.
-	}
-
-private:
-
-	TSharedPtr<SListView<UCameraVariableAsset*>> CameraVariableListView;
-	TArray<UCameraVariableAsset*> CameraVariableItemsSource;
-
-	UClass* VariableClass = nullptr;
-	UCameraVariableAsset* InitialVariable = nullptr;
-
-	FGetCurrentSelectionDelegate GetCurrentAssetPickerSelection;
-	bool bSuppressVariableListSelectionChanged = false;
-
-	FExecuteSetParameterVariable RequestSetVariable;
-};
-
 void FCameraParameterDetailsCustomization::Register(FPropertyEditorModule& PropertyEditorModule)
 {
 #define UE_CAMERA_VARIABLE_FOR_TYPE(ValueType, ValueName)\
@@ -360,17 +197,22 @@ TSharedRef<SWidget> FCameraParameterDetailsCustomization::BuildCameraVariableBro
 	}
 	MenuBuilder.EndSection();
 
+	FCameraVariablePickerConfig PickerConfig;
+	PickerConfig.CameraVariableClass = VariableClass;
+	PickerConfig.InitialCameraVariableSelection = CommonVariable;
+	PickerConfig.CameraVariableCollectionSaveSettingsName = TEXT("CameraParameterVariablePropertyPicker");
+	PickerConfig.OnCameraVariableSelected = FOnCameraVariableSelected::CreateSP(
+			this, &FCameraParameterDetailsCustomization::OnSetVariable);
+	IGameplayCamerasEditorModule& GameplayCamerasEditorModule = IGameplayCamerasEditorModule::Get();
+	TSharedRef<SWidget> PickerWidget = GameplayCamerasEditorModule.CreateCameraVariablePicker(PickerConfig);
+
 	MenuBuilder.BeginSection(NAME_None, LOCTEXT("CameraVariableBrowser", "Browse"));
 	{
 		TSharedRef<SWidget> VariableBrowser = SNew(SBox)
 			.MinDesiredWidth(300.f)
 			.MinDesiredHeight(300.f)
 			[
-				SNew(SCameraVariableBrowser)
-				.VariableClass(VariableClass)
-				.InitialVariable(CommonVariable)
-				.RequestSetVariable(
-						FExecuteSetParameterVariable::CreateSP(this, &FCameraParameterDetailsCustomization::OnSetVariable))
+				PickerWidget
 			];
 		MenuBuilder.AddWidget(VariableBrowser, FText(), true, false);
 	}
diff --git a/Engine/Plugins/Cameras/GameplayCameras/Source/GameplayCamerasEditor/Private/Editors/GameplayCamerasGraphPanelPinFactory.cpp b/Engine/Plugins/Cameras/GameplayCameras/Source/GameplayCamerasEditor/Private/Editors/GameplayCamerasGraphPanelPinFactory.cpp
--- a/Engine/Plugins/Cameras/GameplayCameras/Source/GameplayCamerasEditor/Private/Editors/GameplayCamerasGraphPanelPinFactory.cpp
+++ b/Engine/Plugins/Cameras/GameplayCameras/Source/GameplayCamerasEditor/Private/Editors/GameplayCamerasGraphPanelPinFactory.cpp
@@ -3,9 +3,11 @@
 #include "Editors/GameplayCamerasGraphPanelPinFactory.h"
 
 #include "Core/CameraRigAsset.h"
+#include "Core/CameraVariableAssets.h"
 #include "EdGraphSchema_K2.h"
 #include "Editors/SBlueprintCameraDirectorRigNameGraphPin.h"
 #include "Editors/SCameraRigNameGraphPin.h"
+#include "Editors/SCameraVariableNameGraphPin.h"
 #include "K2Node_CallFunction.h"
 
 namespace UE::Cameras
@@ -20,13 +22,33 @@ TSharedPtr<SGraphPin> FGameplayCamerasGraphPanelPinFactory::CreatePin(UEdGraphPi
 
 	if (UK2Node_CallFunction* CallFunctionNode = Cast<UK2Node_CallFunction>(Pin->GetOwningNode()))
 	{
-		return CreateFunctionParameterPin(Pin, CallFunctionNode);
+		if (TSharedPtr<SGraphPin> PinWidget = CreateFunctionParameterPin(Pin, CallFunctionNode))
+		{
+			return PinWidget;
+		}
 	}
 
 	if (Pin->PinType.PinCategory == UEdGraphSchema_K2::PC_Object && 
 			Pin->PinType.PinSubCategoryObject == UCameraRigAsset::StaticClass())
 	{
-		return CreateCustomPin(Pin);
+		if (TSharedPtr<SGraphPin> PinWidget = CreateCameraRigPickerPin(Pin))
+		{
+			return PinWidget;
+		}
+	}
+
+	const FEdGraphPinType& PinType = Pin->PinType;
+	const UClass* PinPropertyClass = Cast<const UClass>(PinType.PinSubCategoryObject);
+	if (PinType.PinCategory == UEdGraphSchema_K2::PC_Object && PinPropertyClass)
+	{
+		if (PinPropertyClass == UCameraRigAsset::StaticClass())
+		{
+			return CreateCameraRigPickerPin(Pin);
+		}
+		if (PinPropertyClass->IsChildOf<UCameraVariableAsset>())
+		{
+			return CreateCameraVariablePickerPin(Pin);
+		}
 	}
 
 	return nullptr;
@@ -60,10 +82,15 @@ TSharedPtr<SGraphPin> FGameplayCamerasGraphPanelPinFactory::CreateFunctionParame
 	return nullptr;
 }
 
-TSharedPtr<SGraphPin> FGameplayCamerasGraphPanelPinFactory::CreateCustomPin(UEdGraphPin* Pin) const
+TSharedPtr<SGraphPin> FGameplayCamerasGraphPanelPinFactory::CreateCameraRigPickerPin(UEdGraphPin* Pin) const
 {
 	return SNew(SCameraRigNameGraphPin, Pin);
 }
 
+TSharedPtr<SGraphPin> FGameplayCamerasGraphPanelPinFactory::CreateCameraVariablePickerPin(UEdGraphPin* Pin) const
+{
+	return SNew(SCameraVariableNameGraphPin, Pin);
+}
+
 }  // namespace UE::Cameras
 
diff --git a/Engine/Plugins/Cameras/GameplayCameras/Source/GameplayCamerasEditor/Private/Editors/SCameraVariableNameGraphPin.cpp b/Engine/Plugins/Cameras/GameplayCameras/Source/GameplayCamerasEditor/Private/Editors/SCameraVariableNameGraphPin.cpp
new file mode 100644
--- /dev/null
+++ b/Engine/Plugins/Cameras/GameplayCameras/Source/GameplayCamerasEditor/Private/Editors/SCameraVariableNameGraphPin.cpp
@@ -0,0 +1,175 @@
+// Copyright Epic Games, Inc. All Rights Reserved.
+
+#include "Editors/SCameraVariableNameGraphPin.h"
+
+#include "Core/CameraVariableAssets.h"
+#include "Editors/CameraVariablePickerConfig.h"
+#include "IContentBrowserSingleton.h"
+#include "IGameplayCamerasEditorModule.h"
+#include "ScopedTransaction.h"
+#include "Widgets/Images/SImage.h"
+#include "Widgets/Input/SButton.h"
+#include "Widgets/Input/SComboButton.h"
+#include "Widgets/SBoxPanel.h"
+
+#define LOCTEXT_NAMESPACE "SCameraVariableNameGraphPin"
+
+namespace UE::Cameras
+{
+
+void SCameraVariableNameGraphPin::Construct(const FArguments& InArgs, UEdGraphPin* InGraphPinObj)
+{
+	SGraphPin::Construct(SGraphPin::FArguments(), InGraphPinObj);
+}
+
+TSharedRef<SWidget>	SCameraVariableNameGraphPin::GetDefaultValueWidget()
+{
+	if (!GraphPinObj)
+	{
+		return SNullWidget::NullWidget;
+	}
+
+	return SNew(SHorizontalBox)
+		.Visibility(this, &SGraphPin::GetDefaultValueVisibility)
+		+SHorizontalBox::Slot()
+		.AutoWidth()
+		.Padding(2.f)
+		.MaxWidth(200.f)
+		[
+			SAssignNew(CameraVariablePickerButton, SComboButton)
+			.ButtonStyle(FAppStyle::Get(), "PropertyEditor.AssetComboStyle")
+			.ContentPadding(FMargin(2.f, 2.f, 2.f, 1.f))
+			.ForegroundColor(this, &SCameraVariableNameGraphPin::OnGetComboForeground)
+			.ButtonColorAndOpacity(this, &SCameraVariableNameGraphPin::OnGetWidgetBackground)
+			.MenuPlacement(MenuPlacement_BelowAnchor)
+			.IsEnabled(this, &SCameraVariableNameGraphPin::IsEditingEnabled)
+			.ButtonContent()
+			[
+				SNew(STextBlock)
+					.TextStyle(FAppStyle::Get(), "PropertyEditor.AssetClass")
+					.Font(FAppStyle::GetFontStyle("PropertyWindow.NormalFont"))
+					.ColorAndOpacity(this, &SCameraVariableNameGraphPin::OnGetComboForeground)
+					.Text(this, &SCameraVariableNameGraphPin::OnGetSelectedCameraVariableName)
+					.ToolTipText(this, &SCameraVariableNameGraphPin::OnGetCameraVariablePickerToolTipText)
+			]
+			.OnGetMenuContent(this, &SCameraVariableNameGraphPin::OnBuildCameraVariablePicker)
+		]
+		// Reset button
+		+SHorizontalBox::Slot()
+		.AutoWidth()
+		.Padding(1,0)
+		.VAlign(VAlign_Center)
+		[
+			SNew(SButton)
+			.ButtonStyle(FAppStyle::Get(), "NoBorder")
+			.ButtonColorAndOpacity(this, &SCameraVariableNameGraphPin::OnGetWidgetBackground)
+			.OnClicked(this, &SCameraVariableNameGraphPin::OnResetButtonClicked)
+			.ContentPadding(1.f)
+			.ToolTipText(LOCTEXT("ResetButtonToolTip", "Reset the camera rig reference."))
+			.IsEnabled(this, &SGraphPin::IsEditingEnabled)
+			[
+				SNew(SImage)
+				.ColorAndOpacity(this, &SCameraVariableNameGraphPin::OnGetWidgetForeground)
+				.Image(FAppStyle::GetBrush(TEXT("Icons.CircleArrowLeft")))
+			]
+		];
+}
+
+bool SCameraVariableNameGraphPin::DoesWidgetHandleSettingEditingEnabled() const
+{
+	return true;
+}
+
+FSlateColor SCameraVariableNameGraphPin::OnGetComboForeground() const
+{
+	float Alpha = (IsHovered() || bOnlyShowDefaultValue) ? ActiveComboAlpha : InactiveComboAlpha;
+	return FSlateColor(FLinearColor(1.f, 1.f, 1.f, Alpha));
+}
+
+FSlateColor SCameraVariableNameGraphPin::OnGetWidgetForeground() const
+{
+	float Alpha = (IsHovered() || bOnlyShowDefaultValue) ? ActivePinForegroundAlpha : InactivePinForegroundAlpha;
+	return FSlateColor(FLinearColor(1.f, 1.f, 1.f, Alpha));
+}
+
+FSlateColor SCameraVariableNameGraphPin::OnGetWidgetBackground() const
+{
+	float Alpha = (IsHovered() || bOnlyShowDefaultValue) ? ActivePinBackgroundAlpha : InactivePinBackgroundAlpha;
+	return FSlateColor(FLinearColor(1.f, 1.f, 1.f, Alpha));
+}
+
+FText SCameraVariableNameGraphPin::OnGetSelectedCameraVariableName() const
+{
+	FText Value;
+	if (GraphPinObj != nullptr)
+	{
+		if (const UCameraVariableAsset* CameraVariable = Cast<const UCameraVariableAsset>(GraphPinObj->DefaultObject))
+		{
+			Value = FText::FromString(CameraVariable->GetDisplayName());
+		}
+	}
+	if (Value.IsEmpty())
+	{
+		return LOCTEXT("DefaultComboText", "Select Camera Variable");
+	}
+	return Value;
+}
+
+FText SCameraVariableNameGraphPin::OnGetCameraVariablePickerToolTipText() const
+{
+	return LOCTEXT("ComboToolTipText", "The name of the camera variable.");
+}
+
+TSharedRef<SWidget> SCameraVariableNameGraphPin::OnBuildCameraVariablePicker() const
+{
+	FCameraVariablePickerConfig CameraVariablePickerConfig;
+	CameraVariablePickerConfig.CameraAssetViewType = EAssetViewType::List;
+	CameraVariablePickerConfig.CameraVariableCollectionSaveSettingsName = TEXT("CameraVariableNameGraphPinAssetPicker");
+	CameraVariablePickerConfig.OnCameraVariableSelected = FOnCameraVariableSelected::CreateSP(
+			this, &SCameraVariableNameGraphPin::OnPickerAssetSelected);
+
+	// Find the already specified camera variable, if any.
+	if (GraphPinObj)
+	{
+		UClass* VariableClass = Cast<UClass>(GraphPinObj->PinType.PinSubCategoryObject);
+		CameraVariablePickerConfig.CameraVariableClass = VariableClass;
+
+		UCameraVariableAsset* DefaultCameraVariable = Cast<UCameraVariableAsset>(GraphPinObj->DefaultObject);
+		if (DefaultCameraVariable)
+		{
+			CameraVariablePickerConfig.InitialCameraVariableSelection = DefaultCameraVariable;
+		}
+	}
+
+	IGameplayCamerasEditorModule& CamerasEditorModule = IGameplayCamerasEditorModule::Get();
+	return CamerasEditorModule.CreateCameraVariablePicker(CameraVariablePickerConfig);
+}
+
+void SCameraVariableNameGraphPin::OnPickerAssetSelected(UCameraVariableAsset* SelectedItem) const
+{
+	if (SelectedItem)
+	{
+		CameraVariablePickerButton->SetIsOpen(false);
+		SetCameraVariable(SelectedItem);
+	}
+}
+
+FReply SCameraVariableNameGraphPin::OnResetButtonClicked()
+{
+	CameraVariablePickerButton->SetIsOpen(false);
+	SetCameraVariable(nullptr);
+	return FReply::Handled();
+}
+
+void SCameraVariableNameGraphPin::SetCameraVariable(UCameraVariableAsset* SelectedCameraVariable) const
+{
+	const FScopedTransaction Transaction(LOCTEXT("ChangeObjectPinValue", "Change Object Pin Value"));
+
+	GraphPinObj->Modify();
+	GraphPinObj->GetSchema()->TrySetDefaultObject(*GraphPinObj, SelectedCameraVariable);
+}
+
+}  // namespace UE::Cameras
+
+#undef LOCTEXT_NAMESPACE
+
diff --git a/Engine/Plugins/Cameras/GameplayCameras/Source/GameplayCamerasEditor/Private/Editors/SCameraVariablePicker.cpp b/Engine/Plugins/Cameras/GameplayCameras/Source/GameplayCamerasEditor/Private/Editors/SCameraVariablePicker.cpp
new file mode 100644
--- /dev/null
+++ b/Engine/Plugins/Cameras/GameplayCameras/Source/GameplayCamerasEditor/Private/Editors/SCameraVariablePicker.cpp
@@ -0,0 +1,209 @@
+// Copyright Epic Games, Inc. All Rights Reserved.
+
+#include "Editors/SCameraVariablePicker.h"
+
+#include "Core/CameraVariableAssets.h"
+#include "Core/CameraVariableCollection.h"
+#include "Styles/GameplayCamerasEditorStyle.h"
+
+#define LOCTEXT_NAMESPACE "SCameraVariablePicker"
+
+namespace UE::Cameras
+{
+
+void SCameraVariablePicker::Construct(const FArguments& InArgs)
+{
+	const FCameraVariablePickerConfig& PickerConfig = InArgs._CameraVariablePickerConfig;
+
+	VariableClass = PickerConfig.CameraVariableClass;
+	OnCameraVariableSelected = PickerConfig.OnCameraVariableSelected;
+
+	ChildSlot
+	[
+		SNew(SBox)
+		.HeightOverride(400)
+		.WidthOverride(350)
+		[
+			SNew(SBorder)
+			.BorderImage(FAppStyle::GetBrush("Menu.Background"))
+			[
+				SNew(SVerticalBox)
+				// Camera variable collection asset picker.
+				+SVerticalBox::Slot()
+				.FillHeight(0.55f)
+				[
+					BuildVariableCollectionAssetPicker(PickerConfig)
+				]
+				// Camera variable list.
+				+SVerticalBox::Slot()
+				.FillHeight(0.45f)
+				.Padding(0.f, 3.f)
+				[
+					SAssignNew(CameraVariableListView, SListView<UCameraVariableAsset*>)
+					.ListItemsSource(&CameraVariableItemsSource)
+					.OnGenerateRow(this, &SCameraVariablePicker::OnVariableListGenerateRow)
+					.OnSelectionChanged(this, &SCameraVariablePicker::OnVariableListSelectionChanged)
+				]
+				// Number of items in the camera variable list.
+				+SVerticalBox::Slot()
+				.AutoHeight()
+				[
+					SNew(SHorizontalBox)
+					+SHorizontalBox::Slot()
+					.FillWidth(1.f)
+					.VAlign(VAlign_Center)
+					.Padding(8, 5)
+					[
+						SNew(STextBlock)
+						.Text(this, &SCameraVariablePicker::GetCameraVariableCountText)
+					]
+				]
+			]
+		]
+	];
+
+	if (PickerConfig.InitialCameraVariableSelection)
+	{
+		SetupInitialSelections(PickerConfig.InitialCameraVariableSelection);
+	}
+}
+
+void SCameraVariablePicker::SetupInitialSelections(UCameraVariableAsset* InSelectedCameraVariable)
+{
+	UCameraVariableCollection* InitialVariableCollection = nullptr;
+	if (InSelectedCameraVariable)
+	{
+		InitialVariableCollection = InSelectedCameraVariable->GetTypedOuter<UCameraVariableCollection>();
+	}
+
+	UpdateVariableListItemsSource(InitialVariableCollection);
+
+	CameraVariableListView->SetSelection(InSelectedCameraVariable);
+	CameraVariableListView->RequestScrollIntoView(InSelectedCameraVariable);
+}
+
+TSharedRef<SWidget> SCameraVariablePicker::BuildVariableCollectionAssetPicker(const FCameraVariablePickerConfig& InPickerConfig)
+{
+	FContentBrowserModule& ContentBrowserModule = FModuleManager::Get().LoadModuleChecked<FContentBrowserModule>(TEXT("ContentBrowser"));
+
+	FAssetPickerConfig AssetPickerConfig;
+
+	FARFilter ARFilter;
+	ARFilter.ClassPaths.Add(FTopLevelAssetPath(UCameraVariableCollection::StaticClass()->GetPathName()));
+
+	FAssetData InitialVariableCollection = InPickerConfig.InitialCameraVariableCollectionSelection;
+	if (InPickerConfig.InitialCameraVariableSelection)
+	{
+		InitialVariableCollection = InPickerConfig.InitialCameraVariableSelection->GetTypedOuter<UCameraVariableCollection>();
+	}
+
+	AssetPickerConfig.bAllowDragging = false;
+	AssetPickerConfig.bCanShowClasses = false;
+	AssetPickerConfig.bAllowNullSelection = false;
+	AssetPickerConfig.bShowBottomToolbar = true;
+	AssetPickerConfig.bFocusSearchBoxWhenOpened = true;
+	AssetPickerConfig.SelectionMode = ESelectionMode::Single;
+	AssetPickerConfig.Filter = ARFilter;
+	AssetPickerConfig.SaveSettingsName = InPickerConfig.CameraVariableCollectionSaveSettingsName;
+	AssetPickerConfig.InitialAssetViewType = EAssetViewType::List;
+	AssetPickerConfig.InitialAssetSelection = InitialVariableCollection;
+	AssetPickerConfig.OnAssetSelected = FOnAssetSelected::CreateSP(this, &SCameraVariablePicker::OnAssetSelected);
+	AssetPickerConfig.GetCurrentSelectionDelegates.Add(&GetCurrentAssetPickerSelection);
+
+	return ContentBrowserModule.Get().CreateAssetPicker(AssetPickerConfig);
+}
+
+void SCameraVariablePicker::OnAssetSelected(const FAssetData& SelectedAsset)
+{
+	UpdateVariableListItemsSource();
+}
+
+void SCameraVariablePicker::UpdateVariableListItemsSource(UCameraVariableCollection* InCameraVariableCollection)
+{
+	UCameraVariableCollection* CameraVariableCollection = InCameraVariableCollection;
+	if (!CameraVariableCollection)
+	{
+		TArray<FAssetData> SelectedAssets;
+		if (GetCurrentAssetPickerSelection.IsBound())
+		{
+			SelectedAssets = GetCurrentAssetPickerSelection.Execute();
+		}
+
+		if (!SelectedAssets.IsEmpty())
+		{
+			CameraVariableCollection = Cast<UCameraVariableCollection>(SelectedAssets[0].GetAsset());
+		}
+	}
+
+	CameraVariableItemsSource.Reset();
+	if (CameraVariableCollection)
+	{
+		if (VariableClass)
+		{
+			CameraVariableItemsSource = CameraVariableCollection->Variables.FilterByPredicate(
+					[this](UCameraVariableAsset* Item) { return Item->GetClass() == VariableClass; });
+		}
+		else
+		{
+			CameraVariableItemsSource = CameraVariableCollection->Variables;
+		}
+	}
+
+	CameraVariableListView->RequestListRefresh();
+}
+
+TSharedRef<ITableRow> SCameraVariablePicker::OnVariableListGenerateRow(UCameraVariableAsset* Item, const TSharedRef<STableViewBase>& OwnerTable)
+{
+	TSharedRef<FGameplayCamerasEditorStyle> GameplayCamerasStyle = FGameplayCamerasEditorStyle::Get();
+
+	const FText DisplayName = Item->DisplayName.IsEmpty() ?
+		FText::FromName(Item->GetFName()) : FText::FromString(Item->DisplayName);
+
+	return SNew(STableRow<UCameraVariableAsset*>, OwnerTable)
+		[
+			SNew(SHorizontalBox)
+			+SHorizontalBox::Slot()
+			.AutoWidth()
+			[
+				SNew(SImage)
+				.ColorAndOpacity(FSlateColor::UseForeground())
+				.Image(GameplayCamerasStyle->GetBrush("CameraParameter.VariableBrowser"))
+			]
+			+SHorizontalBox::Slot()
+			.FillWidth(1.f)
+			.Padding(4.f, 2.f)
+			[
+				SNew(STextBlock)
+				.Text(DisplayName)
+			]
+		];
+}
+
+void SCameraVariablePicker::OnVariableListSelectionChanged(UCameraVariableAsset* Item, ESelectInfo::Type SelectInfo) const
+{
+	if (SelectInfo != ESelectInfo::Direct)
+	{
+		OnCameraVariableSelected.ExecuteIfBound(Item);
+	}
+}
+
+FText SCameraVariablePicker::GetCameraVariableCountText() const
+{
+	const int32 NumCameraVariables = CameraVariableItemsSource.Num();
+
+	FText CountText = FText::GetEmpty();
+	if (NumCameraVariables == 1)
+	{
+		CountText = LOCTEXT("CameraVariableCountTextSingular", "1 item");
+	}
+	else
+	{
+		CountText = FText::Format(LOCTEXT("CameraVariableCountTextPlural", "{0} items"), FText::AsNumber(NumCameraVariables));
+	}
+	return CountText;
+}
+
+}  // namespace UE::Cameras
+
+#undef LOCTEXT_NAMESPACE
+
diff --git a/Engine/Plugins/Cameras/GameplayCameras/Source/GameplayCamerasEditor/Private/Editors/SCameraVariablePicker.h b/Engine/Plugins/Cameras/GameplayCameras/Source/GameplayCamerasEditor/Private/Editors/SCameraVariablePicker.h
new file mode 100644
--- /dev/null
+++ b/Engine/Plugins/Cameras/GameplayCameras/Source/GameplayCamerasEditor/Private/Editors/SCameraVariablePicker.h
@@ -0,0 +1,58 @@
+// Copyright Epic Games, Inc. All Rights Reserved.
+
+#pragma once
+
+#include "ContentBrowserModule.h"
+#include "Editors/CameraVariablePickerConfig.h"
+#include "Widgets/SCompoundWidget.h"
+#include "Widgets/Views/SListView.h"
+
+class SSearchBox;
+class UCameraVariableAsset;
+class UCameraVariableCollection;
+
+namespace UE::Cameras
+{
+
+/**
+ * A picker widget for selecting a camera variable.
+ */
+class SCameraVariablePicker : public SCompoundWidget
+{
+public:
+
+	SLATE_BEGIN_ARGS(SCameraVariablePicker)
+	{}
+		SLATE_ARGUMENT(FCameraVariablePickerConfig, CameraVariablePickerConfig)
+	SLATE_END_ARGS()
+
+	void Construct(const FArguments& InArgs);
+
+private:
+
+	void SetupInitialSelections(UCameraVariableAsset* InSelectedCameraVariable);
+
+	TSharedRef<SWidget> BuildVariableCollectionAssetPicker(const FCameraVariablePickerConfig& InPickerConfig);
+
+	void OnAssetSelected(const FAssetData& SelectedAsset);
+
+	void UpdateVariableListItemsSource(UCameraVariableCollection* InCameraVariableCollection = nullptr);
+	TSharedRef<ITableRow> OnVariableListGenerateRow(UCameraVariableAsset* Item, const TSharedRef<STableViewBase>& OwnerTable);
+	void OnVariableListSelectionChanged(UCameraVariableAsset* Item, ESelectInfo::Type SelectInfo) const;
+
+	FText GetCameraVariableCountText() const;
+
+private:
+
+	TSharedPtr<SListView<UCameraVariableAsset*>> CameraVariableListView;
+	TArray<UCameraVariableAsset*> CameraVariableItemsSource;
+
+	UClass* VariableClass = nullptr;
+
+	FGetCurrentSelectionDelegate GetCurrentAssetPickerSelection;
+
+	FOnCameraVariableSelected OnCameraVariableSelected;
+};
+
+}  // namespace UE::Cameras
+
diff --git a/Engine/Plugins/Cameras/GameplayCameras/Source/GameplayCamerasEditor/Private/GameplayCamerasEditorModule.cpp b/Engine/Plugins/Cameras/GameplayCameras/Source/GameplayCamerasEditor/Private/GameplayCamerasEditorModule.cpp
--- a/Engine/Plugins/Cameras/GameplayCameras/Source/GameplayCamerasEditor/Private/GameplayCamerasEditorModule.cpp
+++ b/Engine/Plugins/Cameras/GameplayCameras/Source/GameplayCamerasEditor/Private/GameplayCamerasEditorModule.cpp
@@ -25,6 +25,7 @@
 #include "EdGraph/EdGraph.h"
 #include "Editors/GameplayCamerasGraphPanelPinFactory.h"
 #include "Editors/SCameraRigPicker.h"
+#include "Editors/SCameraVariablePicker.h"
 #include "Features/IModularFeatures.h"
 #include "GameplayCameras.h"
 #include "GameplayCamerasEditorSettings.h"
@@ -163,6 +164,14 @@ public:
 			.CameraRigPickerConfig(InPickerConfig);
 	}
 
+	virtual TSharedRef<SWidget> CreateCameraVariablePicker(const FCameraVariablePickerConfig& InPickerConfig) override
+	{
+		using namespace UE::Cameras;
+
+		return SNew(SCameraVariablePicker)
+			.CameraVariablePickerConfig(InPickerConfig);
+	}
+
 	virtual FDelegateHandle RegisterCameraDirectorEditor(FOnCreateCameraDirectorAssetEditorMode InOnCreateEditor) override
 	{
 		CameraDirectorEditorCreators.Add(InOnCreateEditor);
diff --git a/Engine/Plugins/Cameras/GameplayCameras/Source/GameplayCamerasEditor/Public/Editors/CameraVariablePickerConfig.h b/Engine/Plugins/Cameras/GameplayCameras/Source/GameplayCamerasEditor/Public/Editors/CameraVariablePickerConfig.h
new file mode 100644
--- /dev/null
+++ b/Engine/Plugins/Cameras/GameplayCameras/Source/GameplayCamerasEditor/Public/Editors/CameraVariablePickerConfig.h
@@ -0,0 +1,49 @@
+// Copyright Epic Games, Inc. All Rights Reserved.
+
+#pragma once
+
+#include "IContentBrowserSingleton.h"
+
+class IPropertyHandle;
+class UCameraVariableAsset;
+class UCameraVariableCollection;
+
+namespace UE::Cameras
+{
+
+DECLARE_DELEGATE_OneParam(FOnCameraVariableSelected, UCameraVariableAsset*);
+
+/**
+ * Configuration structure for a camera variable picker widget.
+ * This is a widget that shows the list of camera variables inside a given
+ * camera variable collection. It can be filtered by variable type.
+ *
+ * See IGameplayCamerasEditorModule for creating that widget.
+ */
+struct FCameraVariablePickerConfig
+{
+	/** The initially selected camera variable collection, if any. */
+	FAssetData InitialCameraVariableCollectionSelection;
+
+	/** 
+	 * The initially selected camera variable, if any.
+	 * When set, InitialCameraVariableCollectionSelection is ignored, and the outer
+	 * collection of the selected variable is used instead.
+	 */
+	UCameraVariableAsset* InitialCameraVariableSelection = nullptr;
+
+	/** The type of variable to select. */
+	UClass* CameraVariableClass = nullptr;
+
+	/** Asset picker view type for the camera variable collection picker. */
+	EAssetViewType::Type CameraAssetViewType = EAssetViewType::List;
+
+	/** Asset picker settings name for the camera variable collection picker. */
+	FString CameraVariableCollectionSaveSettingsName;
+
+	/** Callback for when a camera variable has been selected. */
+	FOnCameraVariableSelected OnCameraVariableSelected;
+};
+
+}  // namespace UE::Cameras
+
diff --git a/Engine/Plugins/Cameras/GameplayCameras/Source/GameplayCamerasEditor/Public/Editors/GameplayCamerasGraphPanelPinFactory.h b/Engine/Plugins/Cameras/GameplayCameras/Source/GameplayCamerasEditor/Public/Editors/GameplayCamerasGraphPanelPinFactory.h
--- a/Engine/Plugins/Cameras/GameplayCameras/Source/GameplayCamerasEditor/Public/Editors/GameplayCamerasGraphPanelPinFactory.h
+++ b/Engine/Plugins/Cameras/GameplayCameras/Source/GameplayCamerasEditor/Public/Editors/GameplayCamerasGraphPanelPinFactory.h
@@ -23,7 +23,9 @@ public:
 private:
 
 	TSharedPtr<SGraphPin> CreateFunctionParameterPin(UEdGraphPin* Pin, UK2Node_CallFunction* CallFunctionNode) const;
-	TSharedPtr<SGraphPin> CreateCustomPin(UEdGraphPin* Pin) const;
+
+	TSharedPtr<SGraphPin> CreateCameraRigPickerPin(UEdGraphPin* Pin) const;
+	TSharedPtr<SGraphPin> CreateCameraVariablePickerPin(UEdGraphPin* Pin) const;
 };
 
 }  // namespace UE::Cameras
diff --git a/Engine/Plugins/Cameras/GameplayCameras/Source/GameplayCamerasEditor/Public/Editors/SCameraVariableNameGraphPin.h b/Engine/Plugins/Cameras/GameplayCameras/Source/GameplayCamerasEditor/Public/Editors/SCameraVariableNameGraphPin.h
new file mode 100644
--- /dev/null
+++ b/Engine/Plugins/Cameras/GameplayCameras/Source/GameplayCamerasEditor/Public/Editors/SCameraVariableNameGraphPin.h
@@ -0,0 +1,63 @@
+// Copyright Epic Games, Inc. All Rights Reserved.
+
+#pragma once
+
+#include "ContentBrowserModule.h"
+#include "SGraphPin.h"
+
+class SComboButton;
+class UCameraVariableAsset;
+
+namespace UE::Cameras
+{
+
+class SCameraVariablePicker;
+
+/**
+ * A custom widget for a graph editor pin that shows a camera rig picker dialog.
+ */
+class SCameraVariableNameGraphPin : public SGraphPin
+{
+public:
+
+	SLATE_BEGIN_ARGS(SCameraVariableNameGraphPin)
+	{}
+	SLATE_END_ARGS()
+
+	void Construct(const FArguments& InArgs, UEdGraphPin* InGraphPinObj);
+
+protected:
+
+	// SGraphPin interface.
+	virtual TSharedRef<SWidget>	GetDefaultValueWidget() override;
+	virtual bool DoesWidgetHandleSettingEditingEnabled() const override;
+
+private:
+
+	static constexpr float ActiveComboAlpha = 1.f;
+	static constexpr float InactiveComboAlpha = 0.6f;
+	static constexpr float ActivePinForegroundAlpha = 1.f;
+	static constexpr float InactivePinForegroundAlpha = 0.15f;
+	static constexpr float ActivePinBackgroundAlpha = 0.8f;
+	static constexpr float InactivePinBackgroundAlpha = 0.4f;
+
+	FSlateColor OnGetComboForeground() const;
+	FSlateColor OnGetWidgetForeground() const;
+	FSlateColor OnGetWidgetBackground() const;
+
+	FText OnGetSelectedCameraVariableName() const;
+	FText OnGetCameraVariablePickerToolTipText() const;
+	TSharedRef<SWidget> OnBuildCameraVariablePicker() const;
+	void OnPickerAssetSelected(UCameraVariableAsset* SelectedItem) const;
+
+	FReply OnResetButtonClicked();
+
+	void SetCameraVariable(UCameraVariableAsset* SelectedCameraVariable) const;
+
+private:
+
+	TSharedPtr<SComboButton> CameraVariablePickerButton;
+};
+
+}  // namespace UE::Cameras
+
diff --git a/Engine/Plugins/Cameras/GameplayCameras/Source/GameplayCamerasEditor/Public/IGameplayCamerasEditorModule.h b/Engine/Plugins/Cameras/GameplayCameras/Source/GameplayCamerasEditor/Public/IGameplayCamerasEditorModule.h
--- a/Engine/Plugins/Cameras/GameplayCameras/Source/GameplayCamerasEditor/Public/IGameplayCamerasEditorModule.h
+++ b/Engine/Plugins/Cameras/GameplayCameras/Source/GameplayCamerasEditor/Public/IGameplayCamerasEditorModule.h
@@ -24,6 +24,7 @@ namespace UE::Cameras
 class FCameraBuildLog;
 class FCameraDirectorAssetEditorMode;
 struct FCameraRigPickerConfig;
+struct FCameraVariablePickerConfig;
 
 struct FCameraDebugCategoryInfo
 {
@@ -57,6 +58,7 @@ public:
 public:
 
 	using FCameraRigPickerConfig = UE::Cameras::FCameraRigPickerConfig;
+	using FCameraVariablePickerConfig = UE::Cameras::FCameraVariablePickerConfig;
 
 	/** Creates an editor for the given camera asset */
 	virtual UCameraAssetEditor* CreateCameraAssetEditor(const EToolkitMode::Type Mode, const TSharedPtr<IToolkitHost>& InitToolkitHost, UCameraAsset* CameraAsset) = 0;
@@ -73,6 +75,9 @@ public:
 	/** Creates a new camera rig asset picker widget */
 	virtual TSharedRef<SWidget> CreateCameraRigPicker(const FCameraRigPickerConfig& InPickerConfig) = 0;
 
+	/** Creates a new camera varable asset picker widget */
+	virtual TSharedRef<SWidget> CreateCameraVariablePicker(const FCameraVariablePickerConfig& InPickerConfig) = 0;
+
 public:
 
 	/** Registers a new camera director editor creator. */
