Cameras: fix camera rig picker not correctly setting things up when the camera asset picker is visible and an initial camera rig selection is specified.

The problem was that the asset picker populates asynchronously, so we can rely on its selection to set things up in the first tick.

diff --git a/Engine/Plugins/Cameras/GameplayCameras/Source/GameplayCamerasEditor/Private/Editors/SCameraRigPicker.cpp b/Engine/Plugins/Cameras/GameplayCameras/Source/GameplayCamerasEditor/Private/Editors/SCameraRigPicker.cpp
--- a/Engine/Plugins/Cameras/GameplayCameras/Source/GameplayCamerasEditor/Private/Editors/SCameraRigPicker.cpp
+++ b/Engine/Plugins/Cameras/GameplayCameras/Source/GameplayCamerasEditor/Private/Editors/SCameraRigPicker.cpp
@@ -183,24 +183,54 @@ void SCameraRigPicker::Construct(const FArguments& InArgs)
 		]
 	];
 
-	// If we have an initially selected camera asset, set it up immediately.
 	if (!PickerConfig.bCanSelectCameraAsset)
 	{
 		FixedCameraAssetSelection = PickerConfig.InitialCameraAssetSelection;
 	}
-	UpdateCameraRigItemsSource();
+
+	// If we have an initially selected assets, register a timer to do that next frame.
+	if (PickerConfig.InitialCameraRigSelection)
+	{
+		TVariant<UCameraRigAsset*, FGuid> SelectedCameraRig(TInPlaceType<UCameraRigAsset*>(), PickerConfig.InitialCameraRigSelection);
+		SetupInitialSelections(PickerConfig.InitialCameraAssetSelection, SelectedCameraRig);
+	}
+	else if (PickerConfig.InitialCameraAssetSelection.IsValid() && PickerConfig.InitialCameraRigSelectionGuid.IsValid())
+	{
+		TVariant<UCameraRigAsset*, FGuid> SelectedCameraRig(TInPlaceType<FGuid>(), PickerConfig.InitialCameraRigSelectionGuid);
+		SetupInitialSelections(PickerConfig.InitialCameraAssetSelection, SelectedCameraRig);
+	}
+
+	// If we need to focus the search box, register a timer to do that next frame.
+	if (PickerConfig.bFocusCameraRigSearchBoxWhenOpened)
+	{
+		RegisterActiveTimer(0.f, FWidgetActiveTimerDelegate::CreateSP(this, &SCameraRigPicker::FocusCameraRigSearchBox));
+	}
+
+	// Keep track of miscellaneous stuff.
+	OnCameraRigSelected = PickerConfig.OnCameraRigSelected;
+	PropertyToSet = PickerConfig.PropertyToSet;
+}
+
+void SCameraRigPicker::SetupInitialSelections(const FAssetData& InSelectedCameraAssetData, TVariant<UCameraRigAsset*, FGuid> InSelectedCameraRig)
+{
+	UCameraAsset* SelectedCameraAsset = Cast<UCameraAsset>(InSelectedCameraAssetData.GetAsset());
+	UpdateCameraRigItemsSource(SelectedCameraAsset);
 	UpdateCameraRigFilteredItemsSource();
 
-	// If we have an initially selected camera rig, select it in the list immediately too.
-	UCameraRigAsset* InitialCameraRigSelection = PickerConfig.InitialCameraRigSelection;
-	if (!InitialCameraRigSelection && PickerConfig.InitialCameraRigSelectionGuid.IsValid())
+	UCameraRigAsset* InitialCameraRigSelection = nullptr;
+	if (InSelectedCameraRig.IsType<UCameraRigAsset*>())
+	{
+		InitialCameraRigSelection = InSelectedCameraRig.Get<UCameraRigAsset*>();
+	}
+	if (InSelectedCameraRig.IsType<FGuid>())
 	{
 		if (UCameraAsset* CameraAsset = GetSelectedCameraAsset())
 		{
+			const FGuid CameraRigGuid = InSelectedCameraRig.Get<FGuid>();
 			const TObjectPtr<UCameraRigAsset>* FoundItem = CameraAsset->GetCameraRigs().FindByPredicate(
-					[&PickerConfig](UCameraRigAsset* Item)
+					[&CameraRigGuid](UCameraRigAsset* Item)
 					{
-						return Item->GetGuid() == PickerConfig.InitialCameraRigSelectionGuid;
+						return Item->GetGuid() == CameraRigGuid;
 					});
 			if (FoundItem)
 			{
@@ -213,16 +243,6 @@ void SCameraRigPicker::Construct(const FArguments& InArgs)
 		CameraRigListView->RequestScrollIntoView(InitialCameraRigSelection);
 		CameraRigListView->SetSelection(InitialCameraRigSelection);
 	}
-
-	// If we need to focus the search box, register a time to do that next frame.
-	if (PickerConfig.bFocusCameraRigSearchBoxWhenOpened)
-	{
-		RegisterActiveTimer(0.f, FWidgetActiveTimerDelegate::CreateSP(this, &SCameraRigPicker::FocusCameraRigSearchBox));
-	}
-
-	// Keep track of miscellaneous stuff.
-	OnCameraRigSelected = PickerConfig.OnCameraRigSelected;
-	PropertyToSet = PickerConfig.PropertyToSet;
 }
 
 EActiveTimerReturnType SCameraRigPicker::FocusCameraRigSearchBox(double InCurrentTime, float InDeltaTime)
@@ -370,9 +390,10 @@ void SCameraRigPicker::OnCameraRigListSelectionChanged(UCameraRigAsset* Item, ES
 	}
 }
 
-void SCameraRigPicker::UpdateCameraRigItemsSource()
+void SCameraRigPicker::UpdateCameraRigItemsSource(UCameraAsset* InCameraAsset)
 {
-	if (UCameraAsset* CameraAsset = GetSelectedCameraAsset())
+	UCameraAsset* CameraAsset = InCameraAsset ? InCameraAsset : GetSelectedCameraAsset();
+	if (CameraAsset)
 	{
 		CameraRigItemsSource = CameraAsset->GetCameraRigs();
 	}
diff --git a/Engine/Plugins/Cameras/GameplayCameras/Source/GameplayCamerasEditor/Private/Editors/SCameraRigPicker.h b/Engine/Plugins/Cameras/GameplayCameras/Source/GameplayCamerasEditor/Private/Editors/SCameraRigPicker.h
--- a/Engine/Plugins/Cameras/GameplayCameras/Source/GameplayCamerasEditor/Private/Editors/SCameraRigPicker.h
+++ b/Engine/Plugins/Cameras/GameplayCameras/Source/GameplayCamerasEditor/Private/Editors/SCameraRigPicker.h
@@ -32,6 +32,8 @@ protected:
 
 private:
 
+	void SetupInitialSelections(const FAssetData& InSelectedCameraAssetData, TVariant<UCameraRigAsset*, FGuid> InSelectedCameraRig);
+
 	EActiveTimerReturnType FocusCameraRigSearchBox(double InCurrentTime, float InDeltaTime);
 
 	UCameraAsset* GetSelectedCameraAsset() const;
@@ -42,7 +44,7 @@ private:
 
 	TSharedRef<ITableRow> OnCameraRigListGenerateRow(UCameraRigAsset* Item, const TSharedRef<STableViewBase>& OwnerTable);
 	void OnCameraRigListSelectionChanged(UCameraRigAsset* Item, ESelectInfo::Type SelectInfo);
-	void UpdateCameraRigItemsSource();
+	void UpdateCameraRigItemsSource(UCameraAsset* InCameraAsset = nullptr);
 	void UpdateCameraRigFilteredItemsSource();
 	FText GetCameraRigCountText() const;
 
