Cameras: update GameplayCameras plugin description

diff --git a/Engine/Plugins/Cameras/GameplayCameras/GameplayCameras.uplugin b/Engine/Plugins/Cameras/GameplayCameras/GameplayCameras.uplugin
--- a/Engine/Plugins/Cameras/GameplayCameras/GameplayCameras.uplugin
+++ b/Engine/Plugins/Cameras/GameplayCameras/GameplayCameras.uplugin
@@ -3,7 +3,7 @@
 	"Version" : 1,
 	"VersionName" : "0.1",
 	"FriendlyName" : "Gameplay Cameras",
-	"Description" : "Default gameplay camera classes and systems",
+	"Description" : "A modular and data-driven camera system for Unreal",
 	"Category" : "Cameras",
 	"CreatedBy" : "Epic Games, Inc.",
 	"CreatedByURL" : "https://epicgames.com",
@@ -12,7 +12,7 @@
 	"SupportURL" : "",
 	"EnabledByDefault" : true,
 	"CanContainContent" : true,
-	"IsBetaVersion" : true,
+	"IsExperimentalVersion" : true,
 	"Installed" : false,
 	"Plugins":
 	[
