Cameras: display the initial evaluation context's transform in the debug view

diff --git a/Engine/Plugins/Cameras/GameplayCameras/Source/GameplayCameras/Private/Debug/CameraDebugRenderer.cpp b/Engine/Plugins/Cameras/GameplayCameras/Source/GameplayCameras/Private/Debug/CameraDebugRenderer.cpp
--- a/Engine/Plugins/Cameras/GameplayCameras/Source/GameplayCameras/Private/Debug/CameraDebugRenderer.cpp
+++ b/Engine/Plugins/Cameras/GameplayCameras/Source/GameplayCameras/Private/Debug/CameraDebugRenderer.cpp
@@ -82,6 +82,12 @@ static FAutoConsoleVariableRef CVarGameplayCamerasDebugMaxCardColumns(
 	GGameplayCamerasDebugMaxCardColumns,
 	TEXT("Default: 2. The number of columns to layout the debug cards (e.g. graphs, clocks, etc.)"));
 
+float GGameplayCamerasDebugDefaultCoordinateSystemAxesLength = 100.f;
+static FAutoConsoleVariableRef CVarGameplayCamerasDebugDefaultCoordinateSystemAxesLength(
+	TEXT("GameplayCameras.Debug.DefaultCoordinateSystemAxesLength"),
+	GGameplayCamerasDebugDefaultCoordinateSystemAxesLength,
+	TEXT("Default: 100. The default length of coordinate system axes."));
+
 FCameraDebugRenderer::FCameraDebugRenderer(UWorld* InWorld, UCanvas* InCanvasObject)
 	: World(InWorld)
 	, CanvasObject(InCanvasObject)
@@ -372,6 +378,41 @@ void FCameraDebugRenderer::DrawDirectionalArrow(const FVector3d& Start, const FV
 	}
 }
 
+void FCameraDebugRenderer::DrawCoordinateSystem(const FVector3d& Location, const FRotator3d& Rotation, float AxesLength)
+{
+	if (ULineBatchComponent* LineBatcher = GetDebugLineBatcher())
+	{
+		if (AxesLength <= 0.f)
+		{
+			AxesLength = GGameplayCamerasDebugDefaultCoordinateSystemAxesLength;
+		}
+
+		LineBatcher->DrawLine(
+				Location, 
+				Location + Rotation.RotateVector(FVector3d::ForwardVector * AxesLength),
+				FLinearColor::Red,
+				SDPG_Foreground,
+				0.f);
+		LineBatcher->DrawLine(
+				Location, 
+				Location + Rotation.RotateVector(FVector3d::RightVector * AxesLength),
+				FLinearColor::Green,
+				SDPG_Foreground,
+				0.f);
+		LineBatcher->DrawLine(
+				Location, 
+				Location + Rotation.RotateVector(FVector3d::UpVector * AxesLength),
+				FLinearColor::Blue,
+				SDPG_Foreground,
+				0.f);
+	}
+}
+
+void FCameraDebugRenderer::DrawCoordinateSystem(const FTransform3d& Transform, float AxesLength)
+{
+	DrawCoordinateSystem(Transform.GetLocation(), Transform.GetRotation().Rotator(), AxesLength);
+}
+
 void FCameraDebugRenderer::DrawText(const FVector3d& WorldPosition, const FString& Text, const FLinearColor& TextColor, UFont* TextFont)
 {
 	DrawText(WorldPosition, FVector2d::ZeroVector, Text, TextColor, TextFont);
diff --git a/Engine/Plugins/Cameras/GameplayCameras/Source/GameplayCameras/Private/Debug/CameraDirectorTreeDebugBlock.cpp b/Engine/Plugins/Cameras/GameplayCameras/Source/GameplayCameras/Private/Debug/CameraDirectorTreeDebugBlock.cpp
--- a/Engine/Plugins/Cameras/GameplayCameras/Source/GameplayCameras/Private/Debug/CameraDirectorTreeDebugBlock.cpp
+++ b/Engine/Plugins/Cameras/GameplayCameras/Source/GameplayCameras/Private/Debug/CameraDirectorTreeDebugBlock.cpp
@@ -37,12 +37,29 @@ void FCameraDirectorTreeDebugBlock::Initialize(const FCameraEvaluationContextSta
 		TSharedPtr<FCameraEvaluationContext> Context = Entry.WeakContext.Pin();
 
 		FDirectorDebugInfo EntryDebugInfo;
-		EntryDebugInfo.CameraAssetName = Context ? Context->GetCameraAsset()->GetName() : TEXT("<no camera asset>");
+		if (Context)
+		{
+			EntryDebugInfo.CameraAssetName = GetNameSafe(Context->GetCameraAsset());
+			EntryDebugInfo.InitialContextTransform = Context->GetInitialResult().CameraPose.GetTransform();
+			EntryDebugInfo.bIsValid = true;
+		}
+		else
+		{
+			EntryDebugInfo.bIsValid = false;
+		}
 		CameraDirectors.Add(EntryDebugInfo);
 
-		const FCameraNodeEvaluationResult& InitialResult = Context->GetInitialResult();
-		AddChild(&Builder.BuildDebugBlock<FCameraPoseDebugBlock>(InitialResult.CameraPose)
-				.WithShowUnchangedCVar(TEXT("GameplayCameras.Debug.ContextInitialResult.ShowUnchanged")));
+		if (Context)
+		{
+			const FCameraNodeEvaluationResult& InitialResult = Context->GetInitialResult();
+			AddChild(&Builder.BuildDebugBlock<FCameraPoseDebugBlock>(InitialResult.CameraPose)
+					.WithShowUnchangedCVar(TEXT("GameplayCameras.Debug.ContextInitialResult.ShowUnchanged")));
+		}
+		else
+		{
+			// Dummy debug block.
+			AddChild(&Builder.BuildDebugBlock<FCameraDebugBlock>());
+		}
 	}
 }
 
@@ -73,7 +90,15 @@ void FCameraDirectorTreeDebugBlock::OnDebugDraw(const FCameraDebugBlockDrawParam
 		}
 
 		const FDirectorDebugInfo& EntryDebugInfo(CameraDirectors[Index]);
-		Renderer.AddText(TEXT("Camera asset: {cam_notice}%s{cam_default}\n"), *EntryDebugInfo.CameraAssetName);
+		if (EntryDebugInfo.bIsValid)
+		{
+			Renderer.AddText(TEXT("Camera asset: {cam_notice}%s{cam_default}\n"), *EntryDebugInfo.CameraAssetName);
+			Renderer.DrawCoordinateSystem(EntryDebugInfo.InitialContextTransform);
+		}
+		else
+		{
+			Renderer.AddText(TEXT("{cam_error}Invalid context!{cam_default}\n"));
+		}
 
 		Renderer.AddIndent();
 		ChildrenView[Index]->DebugDraw(Params, Renderer);
diff --git a/Engine/Plugins/Cameras/GameplayCameras/Source/GameplayCameras/Public/Debug/CameraDebugRenderer.h b/Engine/Plugins/Cameras/GameplayCameras/Source/GameplayCameras/Public/Debug/CameraDebugRenderer.h
--- a/Engine/Plugins/Cameras/GameplayCameras/Source/GameplayCameras/Public/Debug/CameraDebugRenderer.h
+++ b/Engine/Plugins/Cameras/GameplayCameras/Source/GameplayCameras/Public/Debug/CameraDebugRenderer.h
@@ -98,6 +98,10 @@ public:
 	GAMEPLAYCAMERAS_API void DrawSphere(const FVector3d& Center, float Radius, int32 Segments, const FLinearColor& LineColor, float LineThickness);
 	/** Draws a 3D arrow. */
 	GAMEPLAYCAMERAS_API void DrawDirectionalArrow(const FVector3d& Start, const FVector3d& End, float ArrowSize, const FLinearColor& LineColor, float LineThickness);
+	/** Draws a 3D coordinate system. */
+	GAMEPLAYCAMERAS_API void DrawCoordinateSystem(const FVector3d& Location, const FRotator3d& Rotation, float AxesLength = 0.f);
+	/** Draws a 3D coordinate system. */
+	GAMEPLAYCAMERAS_API void DrawCoordinateSystem(const FTransform3d& Transform, float AxesLength = 0.f);
 	/** Draws text at a projected 3D position. */
 	GAMEPLAYCAMERAS_API void DrawText(const FVector3d& WorldPosition, const FString& Text, const FLinearColor& TextColor, UFont* TextFont = nullptr);
 	/** Draws text at a projected 3D position, with an added screen-space offset. */
diff --git a/Engine/Plugins/Cameras/GameplayCameras/Source/GameplayCameras/Public/Debug/CameraDirectorTreeDebugBlock.h b/Engine/Plugins/Cameras/GameplayCameras/Source/GameplayCameras/Public/Debug/CameraDirectorTreeDebugBlock.h
--- a/Engine/Plugins/Cameras/GameplayCameras/Source/GameplayCameras/Public/Debug/CameraDirectorTreeDebugBlock.h
+++ b/Engine/Plugins/Cameras/GameplayCameras/Source/GameplayCameras/Public/Debug/CameraDirectorTreeDebugBlock.h
@@ -3,6 +3,7 @@
 #pragma once
 
 #include "Debug/CameraDebugBlock.h"
+#include "Math/Transform.h"
 
 #if UE_GAMEPLAY_CAMERAS_DEBUG
 
@@ -35,6 +36,8 @@ private:
 	struct FDirectorDebugInfo
 	{
 		FString CameraAssetName;
+		FTransform3d InitialContextTransform = FTransform3d::Identity;
+		bool bIsValid = false;
 	};
 	TArray<FDirectorDebugInfo> CameraDirectors;
 
