Cameras: dampen position node now uses camera parameters, has options for the space in which to dampen, and renders more debugging info

diff --git a/Engine/Plugins/Cameras/GameplayCameras/Source/GameplayCameras/Private/Nodes/Common/DampenPositionCameraNode.cpp b/Engine/Plugins/Cameras/GameplayCameras/Source/GameplayCameras/Private/Nodes/Common/DampenPositionCameraNode.cpp
--- a/Engine/Plugins/Cameras/GameplayCameras/Source/GameplayCameras/Private/Nodes/Common/DampenPositionCameraNode.cpp
+++ b/Engine/Plugins/Cameras/GameplayCameras/Source/GameplayCameras/Private/Nodes/Common/DampenPositionCameraNode.cpp
@@ -4,6 +4,7 @@
 
 #include "Core/CameraEvaluationContext.h"
 #include "Core/CameraNodeEvaluator.h"
+#include "Core/CameraParameterReader.h"
 #include "Debug/CameraDebugBlock.h"
 #include "Debug/CameraDebugBlockBuilder.h"
 #include "Debug/CameraDebugRenderer.h"
@@ -32,6 +33,10 @@ protected:
 
 private:
 
+	TCameraParameterReader<float> ForwardDampingFactorReader;
+	TCameraParameterReader<float> LateralDampingFactorReader;
+	TCameraParameterReader<float> VerticalDampingFactorReader;
+
 	FCriticalDamper ForwardDamper;
 	FCriticalDamper LateralDamper;
 	FCriticalDamper VerticalDamper;
@@ -39,8 +44,9 @@ private:
 	FVector3d PreviousLocation;
 
 #if UE_GAMEPLAY_CAMERAS_DEBUG
-	FVector3d LastUndampedPosition;
-	FVector3d LastDampedPosition;
+	FVector3d DebugLastUndampedPosition;
+	FVector3d DebugLastDampedPosition;
+	FRotator3d DebugLastDampingRotation;
 #endif  // UE_GAMEPLAY_CAMERAS_DEBUG
 };
 
@@ -55,6 +61,7 @@ UE_DECLARE_CAMERA_DEBUG_BLOCK_START(GAMEPLAYCAMERAS_API, FDampenPositionCameraDe
 	UE_DECLARE_CAMERA_DEBUG_BLOCK_FIELD(float, VerticalDampingFactor);
 	UE_DECLARE_CAMERA_DEBUG_BLOCK_FIELD(FVector3d, UndampedPosition);
 	UE_DECLARE_CAMERA_DEBUG_BLOCK_FIELD(FVector3d, DampedPosition);
+	UE_DECLARE_CAMERA_DEBUG_BLOCK_FIELD(FRotator3d, DampingRotation);
 UE_DECLARE_CAMERA_DEBUG_BLOCK_END()
 
 UE_DEFINE_CAMERA_DEBUG_BLOCK_WITH_FIELDS(FDampenPositionCameraDebugBlock)
@@ -63,13 +70,17 @@ void FDampenPositionCameraNodeEvaluator::OnInitialize(const FCameraNodeEvaluator
 {
 	const UDampenPositionCameraNode* DampenNode = GetCameraNodeAs<UDampenPositionCameraNode>();
 
-	ForwardDamper.SetW0(DampenNode->ForwardDampingFactor);
+	ForwardDampingFactorReader.Initialize(DampenNode->ForwardDampingFactor);
+	LateralDampingFactorReader.Initialize(DampenNode->LateralDampingFactor);
+	VerticalDampingFactorReader.Initialize(DampenNode->VerticalDampingFactor);
+
+	ForwardDamper.SetW0(ForwardDampingFactorReader.Get(OutResult.VariableTable));
 	ForwardDamper.Reset(0, 0);
 
-	LateralDamper.SetW0(DampenNode->LateralDampingFactor);
+	LateralDamper.SetW0(LateralDampingFactorReader.Get(OutResult.VariableTable));
 	LateralDamper.Reset(0, 0);
 
-	VerticalDamper.SetW0(DampenNode->VerticalDampingFactor);
+	VerticalDamper.SetW0(VerticalDampingFactorReader.Get(OutResult.VariableTable));
 	VerticalDamper.Reset(0, 0);
 
 	const FCameraNodeEvaluationResult& InitialResult = Params.EvaluationContext->GetInitialResult();
@@ -78,30 +89,84 @@ void FDampenPositionCameraNodeEvaluator::OnInitialize(const FCameraNodeEvaluator
 
 void FDampenPositionCameraNodeEvaluator::OnRun(const FCameraNodeEvaluationParams& Params, FCameraNodeEvaluationResult& OutResult)
 {
-#if WITH_EDITOR
 	const UDampenPositionCameraNode* DampenNode = GetCameraNodeAs<UDampenPositionCameraNode>();
-	ForwardDamper.SetW0(DampenNode->ForwardDampingFactor);
-	LateralDamper.SetW0(DampenNode->LateralDampingFactor);
-	VerticalDamper.SetW0(DampenNode->VerticalDampingFactor);
+
+	// Update our damping factors if they're driven by a camera variable (which means they
+	// could change every frame). Also update them if we are in the editor, since the 
+	// user might tweak them live.
+#if WITH_EDITOR
+	const bool bUpdateForwardDampingFactor = true;
+	const bool bUpdateLateralDampingFactor = true;
+	const bool bUpdateVerticalDampingFactor = true;
+#else
+	const bool bUpdateForwardDampingFactor = ForwardDampingFactorReader.IsDriven();
+	const bool bUpdateLateralDampingFactor = LateralDampingFactorReader.IsDriven();
+	const bool bUpdateVerticalDampingFactor = VerticalDampingFactorReader.IsDriven();
 #endif
+	if (bUpdateForwardDampingFactor)
+	{
+		ForwardDamper.SetW0(ForwardDampingFactorReader.Get(OutResult.VariableTable));
+	}
+	if (bUpdateLateralDampingFactor)
+	{
+		LateralDamper.SetW0(LateralDampingFactorReader.Get(OutResult.VariableTable));
+	}
+	if (bUpdateVerticalDampingFactor)
+	{
+		VerticalDamper.SetW0(VerticalDampingFactorReader.Get(OutResult.VariableTable));
+	}
 
 	// We want the dampen the given camera position, which means it's trying
 	// to converge towards the one given in the result (which we set as our 
 	// next target), but will be lagging behind.
 	const FVector3d NextTarget = OutResult.CameraPose.GetLocation();
 	FVector3d NextLocation = NextTarget;
+#if UE_GAMEPLAY_CAMERAS_DEBUG
+	DebugLastDampingRotation = FRotator3d::ZeroRotator;
+#endif
 
 	if (!Params.bIsFirstFrame && !OutResult.bIsCameraCut)
 	{
-		FTransform3d Transform = OutResult.CameraPose.GetTransform();
-		FRotator3d Rotation = OutResult.CameraPose.GetRotation();
+		// Figure out the desired coordinate system.
+		FRotator3d AxesRotation(EForceInit::ForceInit);
+		switch (DampenNode->DampenSpace)
+		{
+			case ECameraNodeSpace::CameraPose:
+				{
+					const FRotator3d CameraRotation = OutResult.CameraPose.GetRotation();
+					AxesRotation = CameraRotation;
+				}
+				break;
+			case ECameraNodeSpace::Context:
+				if (Params.EvaluationContext)
+				{
+					const FCameraNodeEvaluationResult& InitialResult = Params.EvaluationContext->GetInitialResult();
+					ensureMsgf(InitialResult.bIsValid,
+							TEXT("DampenPositionCameraNode: using invalid context result as damping space!"));
+
+					const FRotator3d ContextRotation = InitialResult.CameraPose.GetRotation();
+					AxesRotation = ContextRotation;
+				}
+				else
+				{
+					UE_LOG(LogCameraSystem, Error,
+							TEXT("DampenPositionCameraNode: cannot dampen in context space when there is "
+								"no current context set."));
+				}
+				break;
+			case ECameraNodeSpace::World:
+				break;
+		}
+#if UE_GAMEPLAY_CAMERAS_DEBUG
+		DebugLastDampingRotation = AxesRotation;
+#endif
 
 		using FAxisDamper = TTuple<FVector3d, FCriticalDamper*>;
 		FAxisDamper AxisDampers[3]
 		{
-			{ Rotation.RotateVector(FVector3d::ForwardVector), &ForwardDamper },
-			{ Rotation.RotateVector(FVector3d::RightVector),& LateralDamper },
-			{ Rotation.RotateVector(FVector3d::UpVector), &VerticalDamper }
+			{ AxesRotation.RotateVector(FVector3d::ForwardVector), &ForwardDamper },
+			{ AxesRotation.RotateVector(FVector3d::RightVector),& LateralDamper },
+			{ AxesRotation.RotateVector(FVector3d::UpVector), &VerticalDamper }
 		};
 
 		// The next target has moved further away compared to the previous target,
@@ -129,8 +194,8 @@ void FDampenPositionCameraNodeEvaluator::OnRun(const FCameraNodeEvaluationParams
 	}
 
 #if UE_GAMEPLAY_CAMERAS_DEBUG
-	LastUndampedPosition = NextTarget;
-	LastDampedPosition = NextLocation;
+	DebugLastUndampedPosition = NextTarget;
+	DebugLastDampedPosition = NextLocation;
 #endif  // UE_GAMEPLAY_CAMERAS_DEBUG
 
 	PreviousLocation = NextLocation;
@@ -147,8 +212,9 @@ void FDampenPositionCameraNodeEvaluator::OnSerialize(const FCameraNodeEvaluatorS
 	Ar << PreviousLocation;
 
 #if UE_GAMEPLAY_CAMERAS_DEBUG
-	Ar << LastUndampedPosition;
-	Ar << LastDampedPosition;
+	Ar << DebugLastUndampedPosition;
+	Ar << DebugLastDampedPosition;
+	Ar << DebugLastDampingRotation;
 #endif  // UE_GAMEPLAY_CAMERAS_DEBUG
 }
 
@@ -166,8 +232,9 @@ void FDampenPositionCameraNodeEvaluator::OnBuildDebugBlocks(const FCameraDebugBl
 	DebugBlock.LateralDampingFactor = LateralDamper.GetW0();
 	DebugBlock.VerticalDampingFactor = VerticalDamper.GetW0();
 
-	DebugBlock.UndampedPosition = LastUndampedPosition;
-	DebugBlock.DampedPosition = LastDampedPosition;
+	DebugBlock.UndampedPosition = DebugLastUndampedPosition;
+	DebugBlock.DampedPosition = DebugLastDampedPosition;
+	DebugBlock.DampingRotation = DebugLastDampingRotation;
 }
 
 void FDampenPositionCameraDebugBlock::OnDebugDraw(const FCameraDebugBlockDrawParams& Params, FCameraDebugRenderer& Renderer)
@@ -177,6 +244,22 @@ void FDampenPositionCameraDebugBlock::OnDebugDraw(const FCameraDebugBlockDrawPar
 			ForwardX0, ForwardDampingFactor,
 			LateralX0, LateralDampingFactor,
 			VerticalX0, VerticalDampingFactor);
+
+	const double DampingAxesLength = 100.f;
+	Renderer.DrawLine(
+			UndampedPosition, 
+			UndampedPosition + DampingRotation.RotateVector(FVector3d::ForwardVector * DampingAxesLength),
+			FLinearColor::Red);
+	Renderer.DrawLine(
+			UndampedPosition, 
+			UndampedPosition + DampingRotation.RotateVector(FVector3d::RightVector * DampingAxesLength),
+			FLinearColor::Green);
+	Renderer.DrawLine(
+			UndampedPosition, 
+			UndampedPosition + DampingRotation.RotateVector(FVector3d::UpVector * DampingAxesLength),
+			FLinearColor::Blue);
+
+	Renderer.DrawLine(UndampedPosition, DampedPosition, FLinearColor::Yellow);
 }
 
 #endif  // UE_GAMEPLAY_CAMERAS_DEBUG
diff --git a/Engine/Plugins/Cameras/GameplayCameras/Source/GameplayCameras/Public/Core/CameraParameterReader.h b/Engine/Plugins/Cameras/GameplayCameras/Source/GameplayCameras/Public/Core/CameraParameterReader.h
--- a/Engine/Plugins/Cameras/GameplayCameras/Source/GameplayCameras/Public/Core/CameraParameterReader.h
+++ b/Engine/Plugins/Cameras/GameplayCameras/Source/GameplayCameras/Public/Core/CameraParameterReader.h
@@ -65,6 +65,14 @@ public:
 		}
 	}
 
+	/**
+	 * Returns whether the parameter is driven by a variable.
+	 */
+	bool IsDriven() const
+	{
+		return VariableID.IsValid();
+	}
+
 private:
 
 	/** Pointer to the value in the parameter. */
diff --git a/Engine/Plugins/Cameras/GameplayCameras/Source/GameplayCameras/Public/Nodes/Common/DampenPositionCameraNode.h b/Engine/Plugins/Cameras/GameplayCameras/Source/GameplayCameras/Public/Nodes/Common/DampenPositionCameraNode.h
--- a/Engine/Plugins/Cameras/GameplayCameras/Source/GameplayCameras/Public/Nodes/Common/DampenPositionCameraNode.h
+++ b/Engine/Plugins/Cameras/GameplayCameras/Source/GameplayCameras/Public/Nodes/Common/DampenPositionCameraNode.h
@@ -3,6 +3,7 @@
 #pragma once
 
 #include "Core/CameraNode.h"
+#include "Core/CameraParameters.h"
 #include "Nodes/CameraNodeTypes.h"
 
 #include "DampenPositionCameraNode.generated.h"
@@ -24,14 +25,18 @@ public:
 
 	/** Damping factor for forward/backward camera movement. */
 	UPROPERTY(EditAnywhere, Category=Damping)
-	float ForwardDampingFactor = 0.f;
+	FFloatCameraParameter ForwardDampingFactor = 0.f;
 
 	/** Damping factor for left/right camera movement. */
 	UPROPERTY(EditAnywhere, Category=Damping)
-	float LateralDampingFactor = 0.f;
+	FFloatCameraParameter LateralDampingFactor = 0.f;
 
 	/** Damping factor for up/down camera movement. */
 	UPROPERTY(EditAnywhere, Category=Damping)
-	float VerticalDampingFactor = 0.f;
+	FFloatCameraParameter VerticalDampingFactor = 0.f;
+
+	/** What space the damping should be in. */
+	UPROPERTY(EditAnywhere, Category=Damping)
+	ECameraNodeSpace DampenSpace = ECameraNodeSpace::CameraPose;
 };
 
