Cameras: add option on collision push node to specify the space in which the safe position offset is applied

diff --git a/Engine/Plugins/Cameras/GameplayCameras/Source/GameplayCameras/Private/Nodes/Collision/CollisionPushCameraNode.cpp b/Engine/Plugins/Cameras/GameplayCameras/Source/GameplayCameras/Private/Nodes/Collision/CollisionPushCameraNode.cpp
--- a/Engine/Plugins/Cameras/GameplayCameras/Source/GameplayCameras/Private/Nodes/Collision/CollisionPushCameraNode.cpp
+++ b/Engine/Plugins/Cameras/GameplayCameras/Source/GameplayCameras/Private/Nodes/Collision/CollisionPushCameraNode.cpp
@@ -11,6 +11,7 @@
 #include "Debug/CameraDebugBlock.h"
 #include "Debug/CameraDebugBlockBuilder.h"
 #include "Debug/CameraDebugRenderer.h"
+#include "Engine/Engine.h"
 #include "Engine/HitResult.h"
 #include "Engine/World.h"
 #include "GameFramework/Pawn.h"
@@ -39,7 +40,7 @@ protected:
 
 private:
 
-	TOptional<FVector3d> GetSafePosition(const APawn* Pawn, const FCameraNodeEvaluationResult& OutResult);
+	TOptional<FVector3d> GetSafePosition(const FCameraNodeEvaluationParams& Params, const FCameraNodeEvaluationResult& OutResult);
 	void RunCollisionTrace(UWorld* World, APlayerController* PlayerController, const FVector3d& SafePosition, const FCameraNodeEvaluationParams& Params, FCameraNodeEvaluationResult& OutResult);
 	void HandleAsyncCollisionTraceResult(UWorld* World, const FVector3d& SafePosition, const FCameraNodeEvaluationParams& Params, FCameraNodeEvaluationResult& OutResult);
 	void HandleCollisionTraceResult(UWorld* World, TArrayView<const FHitResult> HitResults, const FVector3d& SafePosition, const FCameraNodeEvaluationParams& Params, FCameraNodeEvaluationResult& OutResult);
@@ -62,6 +63,8 @@ private:
 
 #if UE_GAMEPLAY_CAMERAS_DEBUG
 	bool bDebugFoundHit = false;
+	FString DebugHitObjectName;
+	FVector3d DebugSafePosition;
 #endif
 };
 
@@ -72,6 +75,8 @@ UE_DECLARE_CAMERA_DEBUG_BLOCK_START(GAMEPLAYCAMERAS_API, FCollisionPushCameraDeb
 	UE_DECLARE_CAMERA_DEBUG_BLOCK_FIELD(float, DampedPushFactor);
 	UE_DECLARE_CAMERA_DEBUG_BLOCK_FIELD(bool, bIsPulling)
 	UE_DECLARE_CAMERA_DEBUG_BLOCK_FIELD(bool, bFoundHit)
+	UE_DECLARE_CAMERA_DEBUG_BLOCK_FIELD(FString, HitObjectName)
+	UE_DECLARE_CAMERA_DEBUG_BLOCK_FIELD(FVector3d, SafePosition)
 UE_DECLARE_CAMERA_DEBUG_BLOCK_END()
 
 UE_DEFINE_CAMERA_DEBUG_BLOCK_WITH_FIELDS(FCollisionPushCameraDebugBlock)
@@ -105,8 +110,7 @@ void FCollisionPushCameraNodeEvaluator::OnRun(const FCameraNodeEvaluationParams&
 		return;
 	}
 
-	APawn* Pawn = PlayerController->GetPawn();
-	TOptional<FVector3d> SafePosition = GetSafePosition(Pawn, OutResult);
+	TOptional<FVector3d> SafePosition = GetSafePosition(Params, OutResult);
 	if (!SafePosition.IsSet())
 	{
 		const UCameraNode* ThisNode = GetCameraNode();
@@ -133,13 +137,13 @@ void FCollisionPushCameraNodeEvaluator::OnRun(const FCameraNodeEvaluationParams&
 	RunCollisionTrace(World, PlayerController, SafePosition.GetValue(), Params, OutResult);
 }
 
-TOptional<FVector3d> FCollisionPushCameraNodeEvaluator::GetSafePosition(const APawn* Pawn, const FCameraNodeEvaluationResult& OutResult)
+TOptional<FVector3d> FCollisionPushCameraNodeEvaluator::GetSafePosition(const FCameraNodeEvaluationParams& Params, const FCameraNodeEvaluationResult& OutResult)
 {
-	FVector3d SafePosition;
+	FTransform3d SafePositionTransform;
 	bool bGotSafePosition = false;
 
 	// See if we can use the camera rig pivot as the safe position.
-	// Otherwise, use the pawn actor's position.
+	// Otherwise, use the context's initial position.
 	const FBuiltInCameraVariables& BuiltInVariables = FBuiltInCameraVariables::Get();
 	TArrayView<const FCameraRigJoint> Joints = OutResult.CameraRigJoints.GetJoints();
 	const FCameraRigJoint* PivotJoint = Joints.FindByPredicate([&BuiltInVariables](const FCameraRigJoint& Joint)
@@ -148,19 +152,83 @@ TOptional<FVector3d> FCollisionPushCameraNodeEvaluator::GetSafePosition(const AP
 			});
 	if (PivotJoint)
 	{
-		SafePosition = PivotJoint->Transform.GetLocation();
+		SafePositionTransform = PivotJoint->Transform;
 		bGotSafePosition = true;
 	}
-	else if (Pawn)
+	else if (Params.EvaluationContext)
 	{
-		SafePosition = Pawn->GetActorLocation();
-		bGotSafePosition = true;
+		const FCameraNodeEvaluationResult& InitialResult = Params.EvaluationContext->GetInitialResult();
+		if (InitialResult.bIsValid)
+		{
+			SafePositionTransform = InitialResult.CameraPose.GetTransform();
+			bGotSafePosition = true;
+		}
 	}
 
+	// Apply the offset in the specified space.
 	if (bGotSafePosition)
 	{
+		const UCollisionPushCameraNode* ThisNode = GetCameraNodeAs<UCollisionPushCameraNode>();
 		const FVector3d SafePositionOffset = SafePositionOffsetReader.Get(OutResult.VariableTable);
-		return SafePosition + SafePositionOffset;
+
+		FVector3d WorldSafePositionOffset = SafePositionOffset;
+		switch (ThisNode->SafePositionOffsetSpace)
+		{
+			case ECollisionSafePositionOffsetSpace::ActiveContext:
+				if (Params.Evaluator)
+				{
+					const FCameraEvaluationContextStack& ContextStack = Params.Evaluator->GetEvaluationContextStack();
+					TSharedPtr<const FCameraEvaluationContext> ActiveContext = ContextStack.GetActiveContext();
+					if (ActiveContext)
+					{
+						const FCameraNodeEvaluationResult& InitialResult = Params.EvaluationContext->GetInitialResult();
+						ensure(InitialResult.bIsValid);
+						WorldSafePositionOffset = InitialResult.CameraPose.GetRotation().RotateVector(SafePositionOffset);
+					}
+					else
+					{
+						UE_LOG(LogCameraSystem, Warning, 
+								TEXT("Can't offset safe collision position in active context space because "
+									"no active context is active on the camera system evaluator for node '%s'!"),
+								*GetNameSafe(ThisNode));
+					}
+				}
+				else
+				{
+					UE_LOG(LogCameraSystem, Warning, 
+							TEXT("Can't offset safe collision position in active context space because "
+								"no camera system evaluator was given for node '%s'!"),
+							*GetNameSafe(ThisNode));
+				}
+			case ECollisionSafePositionOffsetSpace::OwningContext:
+				if (Params.EvaluationContext)
+				{
+					const FCameraNodeEvaluationResult& InitialResult = Params.EvaluationContext->GetInitialResult();
+					ensure(InitialResult.bIsValid);
+					WorldSafePositionOffset = InitialResult.CameraPose.GetRotation().RotateVector(SafePositionOffset);
+				}
+				else
+				{
+					UE_LOG(LogCameraSystem, Warning, 
+							TEXT("Can't offset safe collision position in owning context space because "
+								"there is no owning context for node '%s'!"),
+							*GetNameSafe(ThisNode));
+				}
+				break;
+			case ECollisionSafePositionOffsetSpace::Pivot:
+				if (PivotJoint)
+				{
+					WorldSafePositionOffset = PivotJoint->Transform.TransformVectorNoScale(SafePositionOffset);
+				}
+				break;
+			case ECollisionSafePositionOffsetSpace::CameraPose:
+				{
+					WorldSafePositionOffset = OutResult.CameraPose.GetRotation().RotateVector(SafePositionOffset);
+				}
+				break;
+		}
+
+		return SafePositionTransform.GetLocation() + WorldSafePositionOffset;
 	}
 	return TOptional<FVector3d>();
 }
@@ -244,6 +312,11 @@ void FCollisionPushCameraNodeEvaluator::HandleCollisionTraceResult(UWorld* World
 	bool bFoundHit = false;
 	float CurrentPushFactor = 0.f;
 
+#if UE_GAMEPLAY_CAMERAS_DEBUG
+	DebugHitObjectName.Empty();
+	DebugSafePosition = SafePosition;
+#endif
+
 	for (const FHitResult& Hit : HitResults)
 	{
 		double TraceLength = FVector3d::Distance(Hit.TraceStart, Hit.TraceEnd);
@@ -252,6 +325,16 @@ void FCollisionPushCameraNodeEvaluator::HandleCollisionTraceResult(UWorld* World
 		{
 			CurrentPushFactor = (float)(DistanceToHit / TraceLength);
 			bFoundHit = true;
+#if UE_GAMEPLAY_CAMERAS_DEBUG
+			if (UObject* HitPhysicsObjectOwner = Hit.PhysicsObjectOwner.Get())
+			{
+				DebugHitObjectName = GetNameSafe(HitPhysicsObjectOwner);
+			}
+			else
+			{
+				DebugHitObjectName = TEXT("<no physics object owner>");
+			}
+#endif
 			break;
 		}
 	}
@@ -322,6 +405,8 @@ void FCollisionPushCameraNodeEvaluator::OnBuildDebugBlocks(const FCameraDebugBlo
 	FCollisionPushCameraDebugBlock& DebugBlock = Builder.AttachDebugBlock<FCollisionPushCameraDebugBlock>();
 
 	DebugBlock.bFoundHit = bDebugFoundHit;
+	DebugBlock.HitObjectName = DebugHitObjectName;
+	DebugBlock.SafePosition = DebugSafePosition;
 	DebugBlock.bIsPulling = (LastDirection == ECameraCollisionDirection::Pulling);
 	DebugBlock.PushFactor = LastPushFactor;
 	DebugBlock.DampedPushFactor = LastDampedPushFactor;
@@ -331,15 +416,20 @@ void FCollisionPushCameraDebugBlock::OnDebugDraw(const FCameraDebugBlockDrawPara
 {
 	if (DampedPushFactor > 0)
 	{
-		Renderer.AddText(TEXT("need to push by %.2f%%, currently %.2f%% [%s] %s"),
+		Renderer.AddText(TEXT("need to push by %.2f%%, currently %.2f%% [%s]"),
 				PushFactor, DampedPushFactor,
-				bIsPulling ? TEXT("pulling") : TEXT("pushing"),
-				bFoundHit ? TEXT("(colliding!)") : TEXT(""));
+				bIsPulling ? TEXT("pulling") : TEXT("pushing"));
+		if (bFoundHit)
+		{
+			Renderer.AddText(TEXT(" (colliding with '%s')"), *HitObjectName);
+		}
 	}
 	else
 	{
 		Renderer.AddText(TEXT("not pushing"));
 	}
+
+	Renderer.DrawText(SafePosition, TEXT("Safe Position"), FLinearColor::Gray, GEngine->GetTinyFont());
 }
 
 #endif  // UE_GAMEPLAY_CAMERAS_DEBUG
diff --git a/Engine/Plugins/Cameras/GameplayCameras/Source/GameplayCameras/Public/Nodes/Collision/CollisionPushCameraNode.h b/Engine/Plugins/Cameras/GameplayCameras/Source/GameplayCameras/Public/Nodes/Collision/CollisionPushCameraNode.h
--- a/Engine/Plugins/Cameras/GameplayCameras/Source/GameplayCameras/Public/Nodes/Collision/CollisionPushCameraNode.h
+++ b/Engine/Plugins/Cameras/GameplayCameras/Source/GameplayCameras/Public/Nodes/Collision/CollisionPushCameraNode.h
@@ -4,12 +4,22 @@
 
 #include "Core/CameraNode.h"
 #include "Core/CameraParameters.h"
+#include "Nodes/CameraNodeTypes.h"
 #include "Engine/EngineTypes.h"
 
 #include "CollisionPushCameraNode.generated.h"
 
 class UCameraValueInterpolator;
 
+UENUM()
+enum class ECollisionSafePositionOffsetSpace
+{
+	ActiveContext,
+	OwningContext,
+	Pivot,
+	CameraPose
+};
+
 /**
  * A node that pushes the camera towards a "safe position" when it is colliding with 
  * the environment. By default, the "safe position" is the pivot of the camera (if any) 
@@ -34,6 +44,10 @@ public:
 	UPROPERTY(EditAnywhere, Category="Occlusion")
 	FVector3dCameraParameter SafePositionOffset;
 
+	/** What space the safe position offset should be in. */
+	UPROPERTY(EditAnywhere, Category=Damping)
+	ECollisionSafePositionOffsetSpace SafePositionOffsetSpace = ECollisionSafePositionOffsetSpace::Pivot;
+
 	/** The interpolation to use when pushing the camera towards the safe position. */
 	UPROPERTY(EditAnywhere, Category="Collision")
 	TObjectPtr<UCameraValueInterpolator> PushInterpolator;
