Cameras: properly blend/override the list of camera joints

diff --git a/Engine/Plugins/Cameras/GameplayCameras/Source/GameplayCameras/Private/Core/BlendStackCameraNode.cpp b/Engine/Plugins/Cameras/GameplayCameras/Source/GameplayCameras/Private/Core/BlendStackCameraNode.cpp
--- a/Engine/Plugins/Cameras/GameplayCameras/Source/GameplayCameras/Private/Core/BlendStackCameraNode.cpp
+++ b/Engine/Plugins/Cameras/GameplayCameras/Source/GameplayCameras/Private/Core/BlendStackCameraNode.cpp
@@ -420,6 +420,7 @@ void FBlendStackCameraNodeEvaluator::OnRun(const FCameraNodeEvaluationParams& Pa
 
 		// Start with the input given to us.
 		CurResult.CameraPose = OutResult.CameraPose;
+		CurResult.CameraRigJoints.OverrideAll(OutResult.CameraRigJoints);
 
 		// Override it with whatever the evaluation context has set on its result.
 		const FCameraNodeEvaluationResult& ContextResult(ResolvedEntry.Context->GetInitialResult());
diff --git a/Engine/Plugins/Cameras/GameplayCameras/Source/GameplayCameras/Private/Core/CameraRigJoints.cpp b/Engine/Plugins/Cameras/GameplayCameras/Source/GameplayCameras/Private/Core/CameraRigJoints.cpp
--- a/Engine/Plugins/Cameras/GameplayCameras/Source/GameplayCameras/Private/Core/CameraRigJoints.cpp
+++ b/Engine/Plugins/Cameras/GameplayCameras/Source/GameplayCameras/Private/Core/CameraRigJoints.cpp
@@ -27,5 +27,60 @@ void FCameraRigJoints::Reset()
 	Joints.Reset();
 }
 
+void FCameraRigJoints::OverrideAll(const FCameraRigJoints& OtherJoints)
+{
+	Joints = OtherJoints.Joints;
+}
+
+void FCameraRigJoints::LerpAll(const FCameraRigJoints& ToJoints, float BlendFactor)
+{
+	const bool bFlipOldNewJoints = (BlendFactor >= 0.5f);
+
+	// Gather our existing joints. Flag them as "keep" if we are still below 50% blend, which means
+	// that we are not switching over to the other joints yet.
+	using FFlaggedJoints = TTuple<FCameraRigJoint, bool>;
+	TMap<FCameraVariableID, FFlaggedJoints> JointsPerVariable;
+	const bool bKeepExistingJoints = !bFlipOldNewJoints;
+	for (const FCameraRigJoint& Joint : Joints)
+	{
+		JointsPerVariable.Add(Joint.VariableID, { Joint, bKeepExistingJoints });
+	}
+
+	// Look at the other joints... if we have a similar joint, blend the transform. If this is a
+	// new joint, only add it if we are over 50% blend (i.e. if we are switching over).
+	for (const FCameraRigJoint& OtherJoint : ToJoints.Joints)
+	{
+		if (FFlaggedJoints* FlaggedJoint = JointsPerVariable.Find(OtherJoint.VariableID))
+		{
+			FCameraRigJoint& Joint = FlaggedJoint->Key;
+
+			FTransform3d OutTransform;
+			OutTransform.SetLocation(FMath::Lerp(Joint.Transform.GetLocation(), OtherJoint.Transform.GetLocation(), BlendFactor));
+			OutTransform.SetRotation(FMath::Lerp(Joint.Transform.GetRotation(), OtherJoint.Transform.GetRotation(), BlendFactor));
+			OutTransform.SetScale3D(FMath::Lerp(Joint.Transform.GetScale3D(), OtherJoint.Transform.GetScale3D(), BlendFactor));
+
+			Joint.Transform = OutTransform;
+
+			FlaggedJoint->Value = true;
+		}
+		else if (bFlipOldNewJoints)
+		{
+			JointsPerVariable.Add(OtherJoint.VariableID, { OtherJoint, true });
+		}
+	}
+
+	// Reset our list of joints to whatever joints we have in our map that have been flagged
+	// for keeping.
+	Joints.Reset();
+	for (auto& Pair : JointsPerVariable)
+	{
+		const FFlaggedJoints& FlaggedJoint(Pair.Value);
+		if (FlaggedJoint.Value)
+		{
+			Joints.Add(FlaggedJoint.Key);
+		}
+	}
+}
+
 }  // namespace UE::Cameras
 
diff --git a/Engine/Plugins/Cameras/GameplayCameras/Source/GameplayCameras/Private/Nodes/Blends/PopBlendCameraNode.cpp b/Engine/Plugins/Cameras/GameplayCameras/Source/GameplayCameras/Private/Nodes/Blends/PopBlendCameraNode.cpp
--- a/Engine/Plugins/Cameras/GameplayCameras/Source/GameplayCameras/Private/Nodes/Blends/PopBlendCameraNode.cpp
+++ b/Engine/Plugins/Cameras/GameplayCameras/Source/GameplayCameras/Private/Nodes/Blends/PopBlendCameraNode.cpp
@@ -36,6 +36,8 @@ void FPopBlendCameraNodeEvaluator::OnBlendResults(const FCameraNodeBlendParams&
 	
 	BlendedResult.CameraPose.OverrideAll(ChildResult.CameraPose);
 	BlendedResult.VariableTable.OverrideAll(ChildResult.VariableTable);
+	BlendedResult.CameraRigJoints.OverrideAll(ChildResult.CameraRigJoints);
+
 	if (ChildResult.bIsCameraCut || Params.ChildParams.bIsFirstFrame)
 	{
 		BlendedResult.bIsCameraCut = true;
diff --git a/Engine/Plugins/Cameras/GameplayCameras/Source/GameplayCameras/Private/Nodes/Blends/SimpleBlendCameraNode.cpp b/Engine/Plugins/Cameras/GameplayCameras/Source/GameplayCameras/Private/Nodes/Blends/SimpleBlendCameraNode.cpp
--- a/Engine/Plugins/Cameras/GameplayCameras/Source/GameplayCameras/Private/Nodes/Blends/SimpleBlendCameraNode.cpp
+++ b/Engine/Plugins/Cameras/GameplayCameras/Source/GameplayCameras/Private/Nodes/Blends/SimpleBlendCameraNode.cpp
@@ -41,6 +41,9 @@ void FSimpleBlendCameraNodeEvaluator::OnBlendResults(const FCameraNodeBlendParam
 	BlendedResult.CameraPose.LerpAll(ChildResult.CameraPose, BlendFactor);
 	BlendedResult.VariableTable.LerpAll(ChildResult.VariableTable, BlendFactor);
 
+	// Merge/blend the joints.
+	BlendedResult.CameraRigJoints.LerpAll(ChildResult.CameraRigJoints, BlendFactor);
+
 	// If we have even a fraction of a camera cut, we need to make the
 	// whole result into a camera cut.
 	if (BlendFactor > 0.f && ChildResult.bIsCameraCut)
diff --git a/Engine/Plugins/Cameras/GameplayCameras/Source/GameplayCameras/Public/Core/CameraRigJoints.h b/Engine/Plugins/Cameras/GameplayCameras/Source/GameplayCameras/Public/Core/CameraRigJoints.h
--- a/Engine/Plugins/Cameras/GameplayCameras/Source/GameplayCameras/Public/Core/CameraRigJoints.h
+++ b/Engine/Plugins/Cameras/GameplayCameras/Source/GameplayCameras/Public/Core/CameraRigJoints.h
@@ -41,6 +41,14 @@ public:
 	/** Removes all previously added joints. */
 	void Reset();
 
+public:
+
+	/** Override the joints with another set of joints. */
+	void OverrideAll(const FCameraRigJoints& OtherJoints);
+
+	/** Interpolate the joints towards anoter set of joints. */
+	void LerpAll(const FCameraRigJoints& ToJoints, float BlendFactor);
+
 private:
 
 	using FJointArray = TArray<FCameraRigJoint, TInlineAllocator<2>>;
