Cameras: rename "context space" to "owning context space"

This is because some nodes (like the collision push node) need to differentiate between their owning context and the main active context.

diff --git a/Engine/Plugins/Cameras/GameplayCameras/Source/GameplayCameras/Private/Nodes/Common/DampenPositionCameraNode.cpp b/Engine/Plugins/Cameras/GameplayCameras/Source/GameplayCameras/Private/Nodes/Common/DampenPositionCameraNode.cpp
--- a/Engine/Plugins/Cameras/GameplayCameras/Source/GameplayCameras/Private/Nodes/Common/DampenPositionCameraNode.cpp
+++ b/Engine/Plugins/Cameras/GameplayCameras/Source/GameplayCameras/Private/Nodes/Common/DampenPositionCameraNode.cpp
@@ -137,7 +137,7 @@ void FDampenPositionCameraNodeEvaluator::OnRun(const FCameraNodeEvaluationParams
 					AxesRotation = CameraRotation;
 				}
 				break;
-			case ECameraNodeSpace::Context:
+			case ECameraNodeSpace::OwningContext:
 				if (Params.EvaluationContext)
 				{
 					const FCameraNodeEvaluationResult& InitialResult = Params.EvaluationContext->GetInitialResult();
diff --git a/Engine/Plugins/Cameras/GameplayCameras/Source/GameplayCameras/Private/Nodes/Common/OffsetCameraNode.cpp b/Engine/Plugins/Cameras/GameplayCameras/Source/GameplayCameras/Private/Nodes/Common/OffsetCameraNode.cpp
--- a/Engine/Plugins/Cameras/GameplayCameras/Source/GameplayCameras/Private/Nodes/Common/OffsetCameraNode.cpp
+++ b/Engine/Plugins/Cameras/GameplayCameras/Source/GameplayCameras/Private/Nodes/Common/OffsetCameraNode.cpp
@@ -58,7 +58,7 @@ void FOffsetCameraNodeEvaluator::OnRun(const FCameraNodeEvaluationParams& Params
 				OutResult.CameraPose.SetTransform(Transform);
 			}
 			break;
-		case ECameraNodeSpace::Context:
+		case ECameraNodeSpace::OwningContext:
 			if (Params.EvaluationContext)
 			{ 
 				const FCameraNodeEvaluationResult& InitialResult = Params.EvaluationContext->GetInitialResult();
diff --git a/Engine/Plugins/Cameras/GameplayCameras/Source/GameplayCameras/Public/Nodes/CameraNodeTypes.h b/Engine/Plugins/Cameras/GameplayCameras/Source/GameplayCameras/Public/Nodes/CameraNodeTypes.h
--- a/Engine/Plugins/Cameras/GameplayCameras/Source/GameplayCameras/Public/Nodes/CameraNodeTypes.h
+++ b/Engine/Plugins/Cameras/GameplayCameras/Source/GameplayCameras/Public/Nodes/CameraNodeTypes.h
@@ -15,7 +15,7 @@ enum class ECameraNodeSpace
 	/** Operates in the local camera pose space. */
 	CameraPose,
 	/** Operates in the space of the owning evaluation context's initial result. */
-	Context,
+	OwningContext,
 	/** Operates in world space. */
 	World
 };
