Cameras: refactor how the camera system evaluator is created and found

- A new "camera system host" object actually holds the evaluator.
- The host object has an API to find it or create it for a given player controller.
- All other components use this API, so random activation ordering doesn't matter, the first call will create the host.

diff --git a/Engine/Plugins/Cameras/GameplayCameras/Source/GameplayCameras/Private/GameFramework/ActivateCameraRigFunctions.cpp b/Engine/Plugins/Cameras/GameplayCameras/Source/GameplayCameras/Private/GameFramework/ActivateCameraRigFunctions.cpp
--- a/Engine/Plugins/Cameras/GameplayCameras/Source/GameplayCameras/Private/GameFramework/ActivateCameraRigFunctions.cpp
+++ b/Engine/Plugins/Cameras/GameplayCameras/Source/GameplayCameras/Private/GameFramework/ActivateCameraRigFunctions.cpp
@@ -10,6 +10,7 @@
 #include "Core/CameraEvaluationContextStack.h"
 #include "GameFramework/GameplayCameraSystemActor.h"
 #include "GameFramework/GameplayCameraSystemComponent.h"
+#include "GameFramework/GameplayCameraSystemHost.h"
 #include "GameFramework/PlayerController.h"
 #include "GameplayCameras.h"
 #include "Templates/SharedPointer.h"
@@ -93,9 +94,8 @@ void UControllerGameplayCameraEvaluationComponent::ActivateCameraRigs()
 {
 	using namespace UE::Cameras;
 
-	APlayerController* PlayerController = GetOwner<APlayerController>();
-	FCameraSystemEvaluator* SystemEvaluator = FindCameraSystemEvaluator(PlayerController);
-	if (!SystemEvaluator)
+	EnsureCameraSystemHost();
+	if (!CameraSystemHost)
 	{
 		return;
 	}
@@ -106,6 +106,7 @@ void UControllerGameplayCameraEvaluationComponent::ActivateCameraRigs()
 		return;
 	}
 
+	TSharedPtr<FCameraSystemEvaluator> SystemEvaluator = CameraSystemHost->GetCameraSystemEvaluator();
 	FRootCameraNodeEvaluator* RootNodeEvaluator = SystemEvaluator->GetRootNodeEvaluator();
 
 	for (FCameraRigInfo& CameraRigInfo : CameraRigInfos)
@@ -115,7 +116,7 @@ void UControllerGameplayCameraEvaluationComponent::ActivateCameraRigs()
 			FActivateCameraRigParams Params;
 			Params.CameraRig = CameraRigInfo.CameraRig;
 			Params.EvaluationContext = EvaluationContext;
-			Params.Evaluator = SystemEvaluator;
+			Params.Evaluator = SystemEvaluator.Get();
 			Params.Layer = CameraRigInfo.EvaluationLayer;
 			RootNodeEvaluator->ActivateCameraRig(Params);
 
@@ -140,16 +141,13 @@ void UControllerGameplayCameraEvaluationComponent::EnsureEvaluationContext()
 	}
 }
 
-UE::Cameras::FCameraSystemEvaluator* UControllerGameplayCameraEvaluationComponent::FindCameraSystemEvaluator(APlayerController* PlayerController)
+void UControllerGameplayCameraEvaluationComponent::EnsureCameraSystemHost()
 {
-	if (PlayerController && PlayerController->PlayerCameraManager)
+	if (!CameraSystemHost)
 	{
-		AActor* ViewTarget = PlayerController->PlayerCameraManager->GetViewTarget();
-		if (AGameplayCameraSystemActor* SystemActor = Cast<AGameplayCameraSystemActor>(ViewTarget))
-		{
-			return SystemActor->GetCameraSystemComponent()->GetCameraSystemEvaluator().Get();
-		}
+		APlayerController* PlayerController = GetOwner<APlayerController>();
+
+		CameraSystemHost = UGameplayCameraSystemHost::FindOrCreateHost(PlayerController);
 	}
-	return nullptr;
 }
 
diff --git a/Engine/Plugins/Cameras/GameplayCameras/Source/GameplayCameras/Private/GameFramework/GameplayCameraComponent.cpp b/Engine/Plugins/Cameras/GameplayCameras/Source/GameplayCameras/Private/GameFramework/GameplayCameraComponent.cpp
--- a/Engine/Plugins/Cameras/GameplayCameras/Source/GameplayCameras/Private/GameFramework/GameplayCameraComponent.cpp
+++ b/Engine/Plugins/Cameras/GameplayCameras/Source/GameplayCameras/Private/GameFramework/GameplayCameraComponent.cpp
@@ -9,17 +9,10 @@
 #include "Engine/EngineTypes.h"
 #include "Engine/StaticMesh.h"
 #include "Engine/World.h"
-#include "GameDelegates.h"
-#include "GameFramework/Controller.h"
-#include "GameFramework/GameplayCameraSystemActor.h"
-#include "GameFramework/GameplayCameraSystemComponent.h"
+#include "GameFramework/GameplayCameraSystemHost.h"
 #include "GameplayCameras.h"
-#include "GameplayCamerasSettings.h"
 #include "Kismet/GameplayStatics.h"
-#include "Logging/MessageLog.h"
-#include "TimerManager.h"
 #include "UObject/ConstructorHelpers.h"
-#include "UObject/UnrealType.h"
 
 #include UE_INLINE_GENERATED_CPP_BY_NAME(GameplayCameraComponent)
 
@@ -87,20 +80,14 @@ void UGameplayCameraComponent::DeactivateCameraEvaluationContext()
 		return;
 	}
 
-	if (!ensure(PlayerController->PlayerCameraManager))
-	{
-		return;
-	}
-	
-	AGameplayCameraSystemActor* CameraSystem = Cast<AGameplayCameraSystemActor>(PlayerController->PlayerCameraManager->GetViewTarget());
-	if (!ensure(CameraSystem))
+	if (!CameraSystemHost)
 	{
 		return;
 	}
 
 	if (EvaluationContext.IsValid())
 	{
-		TSharedPtr<FCameraSystemEvaluator> Evaluator = CameraSystem->GetCameraSystemComponent()->GetCameraSystemEvaluator();
+		TSharedPtr<FCameraSystemEvaluator> Evaluator = CameraSystemHost->GetCameraSystemEvaluator();
 		Evaluator->RemoveEvaluationContext(EvaluationContext.ToSharedRef());
 	}
 
@@ -111,9 +98,7 @@ void UGameplayCameraComponent::ActivateCameraEvaluationContext(APlayerController
 {
 	using namespace UE::Cameras;
 
-	if (!ensureMsgf(
-				PlayerController && PlayerController->PlayerCameraManager,
-				TEXT("Can't activate gameplay camera component: invalid player controller!")))
+	if (!ensureMsgf(PlayerController, TEXT("Can't activate gameplay camera component: invalid player controller!")))
 	{
 		return;
 	}
@@ -124,65 +109,13 @@ void UGameplayCameraComponent::ActivateCameraEvaluationContext(APlayerController
 		return;
 	}
 	
-	AGameplayCameraSystemActor* CameraSystem = Cast<AGameplayCameraSystemActor>(PlayerController->PlayerCameraManager->GetViewTarget());
-	if (!CameraSystem)
-	{
-		const UGameplayCamerasSettings* Settings = GetDefault<UGameplayCamerasSettings>();
-		if (Settings->bAutoSpawnCameraSystemActor)
-		{
-			FActorSpawnParameters SpawnParams;
-			CameraSystem = GetWorld()->SpawnActor<AGameplayCameraSystemActor>(SpawnParams);
-			CameraSystem->GetCameraSystemComponent()->ActivateCameraSystemForPlayerController(PlayerController);
-		}
-	}
-	if (!CameraSystem)
-	{
-		FGameDelegates::Get().GetViewTargetChangedDelegate().AddUObject(
-				this, &UGameplayCameraComponent::DelayedActivateCameraEvaluationContext);
-
-		GetWorld()->GetTimerManager().SetTimerForNextTick(
-				this, &UGameplayCameraComponent::AbortDelayedActivateCameraEvaluationContext);
-
-		UE_LOG(LogCameraSystem, Verbose, 
-				TEXT("Can't activate gameplay camera component: no camera system found on the view target. "
-					"Waiting until next tick to see if it is added after us."));
-		return;
-	}
-
-	DoActivateCameraEvaluationContext(PlayerController, CameraSystem->GetCameraSystemComponent()->GetCameraSystemEvaluator());
-
-	WeakPlayerController = PlayerController;
-}
-
-void UGameplayCameraComponent::DelayedActivateCameraEvaluationContext(APlayerController* PlayerController, AActor* OldViewTarget, AActor* NewViewTarget)
-{
-	FGameDelegates::Get().GetViewTargetChangedDelegate().RemoveAll(this);
-
-	AGameplayCameraSystemActor* CameraSystem = Cast<AGameplayCameraSystemActor>(NewViewTarget);
-	if (!CameraSystem)
+	CameraSystemHost = UGameplayCameraSystemHost::FindOrCreateHost(PlayerController);
+	if (!CameraSystemHost)
 	{
-		UE_LOG(LogCameraSystem, Error, TEXT("Can't activate gameplay camera component: no camera system found on the view target!"));
+		UE_LOG(LogCameraSystem, Error, TEXT("Can't activate gameplay camera component: no camera system host found!"));
 		return;
 	}
 
-	DoActivateCameraEvaluationContext(PlayerController, CameraSystem->GetCameraSystemComponent()->GetCameraSystemEvaluator());
-
-	WeakPlayerController = PlayerController;
-}
-
-void UGameplayCameraComponent::AbortDelayedActivateCameraEvaluationContext()
-{
-	if (!WeakPlayerController.IsValid())
-	{
-		UE_LOG(LogCameraSystem, Error, TEXT("Can't activate gameplay camera component: no camera system found after a full tick!"));
-		FGameDelegates::Get().GetViewTargetChangedDelegate().RemoveAll(this);
-	}
-}
-
-void UGameplayCameraComponent::DoActivateCameraEvaluationContext(APlayerController* PlayerController, TSharedPtr<UE::Cameras::FCameraSystemEvaluator> CameraSystemEvaluator)
-{
-	using namespace UE::Cameras;
-
 	if (!EvaluationContext.IsValid())
 	{
 		EvaluationContext = MakeShared<FGameplayCameraComponentEvaluationContext>();
@@ -194,7 +127,10 @@ void UGameplayCameraComponent::DoActivateCameraEvaluationContext(APlayerControll
 		EvaluationContext->Initialize(InitParams);
 	}
 
+	TSharedPtr<FCameraSystemEvaluator> CameraSystemEvaluator = CameraSystemHost->GetCameraSystemEvaluator();
 	CameraSystemEvaluator->PushEvaluationContext(EvaluationContext.ToSharedRef());
+
+	WeakPlayerController = PlayerController;
 }
 
 FBlueprintCameraPose UGameplayCameraComponent::GetInitialPose() const
@@ -242,6 +178,7 @@ void UGameplayCameraComponent::BeginPlay()
 #if WITH_EDITOR
 	if (Camera)
 	{
+		// Auto-build the camera asset on begin play to make sure we've got the latest user edits.
 		using namespace UE::Cameras;
 		FCameraBuildLog BuildLog;
 		FCameraAssetBuilder Builder(BuildLog);
diff --git a/Engine/Plugins/Cameras/GameplayCameras/Source/GameplayCameras/Private/GameFramework/GameplayCameraSystemComponent.cpp b/Engine/Plugins/Cameras/GameplayCameras/Source/GameplayCameras/Private/GameFramework/GameplayCameraSystemComponent.cpp
--- a/Engine/Plugins/Cameras/GameplayCameras/Source/GameplayCameras/Private/GameFramework/GameplayCameraSystemComponent.cpp
+++ b/Engine/Plugins/Cameras/GameplayCameras/Source/GameplayCameras/Private/GameFramework/GameplayCameraSystemComponent.cpp
@@ -3,15 +3,15 @@
 #include "GameFramework/GameplayCameraSystemComponent.h"
 
 #include "Components/StaticMeshComponent.h"
+#include "Core/CameraSystemEvaluator.h"
 #include "Debug/DebugDrawService.h"
 #include "Engine/Canvas.h"
 #include "Engine/StaticMesh.h"
 #include "Engine/World.h"
 #include "GameFramework/Controller.h"
+#include "GameFramework/GameplayCameraSystemHost.h"
 #include "Kismet/GameplayStatics.h"
-#include "Logging/MessageLog.h"
 #include "UObject/ConstructorHelpers.h"
-#include "UObject/UnrealType.h"
 
 #include UE_INLINE_GENERATED_CPP_BY_NAME(GameplayCameraSystemComponent)
 
@@ -33,10 +33,22 @@ UGameplayCameraSystemComponent::UGameplayCameraSystemComponent(const FObjectInit
 #endif  // WITH_EDITORONLY_DATA
 }
 
+TSharedPtr<UE::Cameras::FCameraSystemEvaluator> UGameplayCameraSystemComponent::GetCameraSystemEvaluator(bool bEnsureIfNull)
+{
+	UGameplayCameraSystemHost* HostPtr = CameraSystemHost.Get();
+	ensureMsgf(HostPtr || !bEnsureIfNull, TEXT("Accessing camera system evaluator when we haven't found or created a host for one."));
+	if (HostPtr)
+	{
+		return HostPtr->GetCameraSystemEvaluator();
+	}
+	return nullptr;
+}
+
 void UGameplayCameraSystemComponent::GetCameraView(float DeltaTime, FMinimalViewInfo& DesiredView)
 {
 	using namespace UE::Cameras;
 
+	TSharedPtr<FCameraSystemEvaluator> Evaluator = GetCameraSystemEvaluator();
 	if (Evaluator.IsValid())
 	{
 		FCameraSystemEvaluationParams UpdateParams;
@@ -59,12 +71,6 @@ void UGameplayCameraSystemComponent::OnRegister()
 		return;
 	}
 
-	if (!Evaluator.IsValid())
-	{
-		Evaluator = MakeShared<FCameraSystemEvaluator>();
-		Evaluator->Initialize(this);
-	}
-
 #if UE_GAMEPLAY_CAMERAS_DEBUG
 	UWorld* World = GetWorld();
 	if (World && World->IsGameWorld())
@@ -116,6 +122,13 @@ void UGameplayCameraSystemComponent::ActivateCameraSystemForPlayerController(APl
 		return;
 	}
 
+	CameraSystemHost = UGameplayCameraSystemHost::FindOrCreateHost(PlayerController);
+	if (!CameraSystemHost)
+	{
+		UE_LOG(LogCameraSystem, Error, TEXT("can't create camera system host!"));
+		return;
+	}
+
 	PlayerController->SetViewTarget(OwningActor);
 	WeakPlayerController = PlayerController;
 }
@@ -173,11 +186,6 @@ void UGameplayCameraSystemComponent::OnComponentDestroyed(bool bDestroyingHierar
 		DebugDrawDelegateHandle.Reset();
 	}
 #endif  // UE_GAMEPLAY_CAMERAS_DEBUG
-
-	if (Evaluator.IsValid())
-	{
-		Evaluator.Reset();
-	}
 }
 
 #if WITH_EDITOR
@@ -194,17 +202,6 @@ bool UGameplayCameraSystemComponent::GetEditorPreviewInfo(float DeltaTime, FMini
 
 #endif  // WITH_EDITOR
 
-void UGameplayCameraSystemComponent::AddReferencedObjects(UObject* InThis, FReferenceCollector& Collector)
-{
-	Super::AddReferencedObjects(InThis, Collector);
-
-	UGameplayCameraSystemComponent* TypedThis = CastChecked<UGameplayCameraSystemComponent>(InThis);
-	if (TypedThis->Evaluator.IsValid())
-	{
-		TypedThis->Evaluator->AddReferencedObjects(Collector);
-	}
-}
-
 void UGameplayCameraSystemComponent::OnBecomeViewTarget()
 {
 }
@@ -219,6 +216,7 @@ void UGameplayCameraSystemComponent::DebugDraw(UCanvas* Canvas, APlayerControlle
 {
 	using namespace UE::Cameras;
 
+	TSharedPtr<FCameraSystemEvaluator> Evaluator = GetCameraSystemEvaluator();
 	if (Evaluator.IsValid())
 	{
 		FCameraSystemDebugUpdateParams DebugUpdateParams;
diff --git a/Engine/Plugins/Cameras/GameplayCameras/Source/GameplayCameras/Private/GameFramework/GameplayCameraSystemHost.cpp b/Engine/Plugins/Cameras/GameplayCameras/Source/GameplayCameras/Private/GameFramework/GameplayCameraSystemHost.cpp
new file mode 100644
--- /dev/null
+++ b/Engine/Plugins/Cameras/GameplayCameras/Source/GameplayCameras/Private/GameFramework/GameplayCameraSystemHost.cpp
@@ -0,0 +1,108 @@
+// Copyright Epic Games, Inc. All Rights Reserved.
+
+#include "GameFramework/GameplayCameraSystemHost.h"
+
+#include "Core/CameraSystemEvaluator.h"
+#include "Engine/World.h"
+#include "GameFramework/GameplayCameraSystemActor.h"
+#include "GameFramework/GameplayCameraSystemComponent.h"
+#include "GameFramework/PlayerController.h"
+#include "GameplayCamerasSettings.h"
+#include "Templates/UnrealTemplate.h"
+
+const TCHAR* UGameplayCameraSystemHost::DefaultHostName = TEXT("GameplayCameraSystemHost");
+
+UGameplayCameraSystemHost* UGameplayCameraSystemHost::FindOrCreateHost(APlayerController* PlayerController, const TCHAR* HostName)
+{
+	static bool bIsCreatingHost = false;
+
+	if (!ensure(PlayerController))
+	{
+		return nullptr;
+	}
+
+	if (UGameplayCameraSystemHost* ExistingHost = FindHost(PlayerController, HostName, true))
+	{
+		return ExistingHost;
+	}
+	
+	if (!ensureMsgf(!bIsCreatingHost, TEXT("Detected reentrant call to UGameplayCameraSystemHost::FindOrCreateHost!")))
+	{
+		return nullptr;
+	}
+
+	TGuardValue<bool> IsCreatingHostGuard(bIsCreatingHost, true);
+
+	if (!HostName)
+	{
+		HostName = DefaultHostName;
+	}
+
+	UGameplayCameraSystemHost* NewHost = NewObject<UGameplayCameraSystemHost>(PlayerController, HostName);
+
+	const UGameplayCamerasSettings* Settings = GetDefault<UGameplayCamerasSettings>();
+	if (Settings->bAutoSpawnCameraSystemActor)
+	{
+		UWorld* World = PlayerController->GetWorld();
+
+		FActorSpawnParameters SpawnParams;
+		AGameplayCameraSystemActor* CameraSystemActor = World->SpawnActor<AGameplayCameraSystemActor>(SpawnParams);
+		CameraSystemActor->GetCameraSystemComponent()->ActivateCameraSystemForPlayerController(PlayerController);
+	}
+
+	return NewHost;
+}
+
+UGameplayCameraSystemHost* UGameplayCameraSystemHost::FindHost(APlayerController* PlayerController, const TCHAR* HostName, bool bAllowNull)
+{
+	if (!PlayerController)
+	{
+		ensureMsgf(bAllowNull, TEXT("Can't find gameplay camera system host: null player controller provided!"));
+		return nullptr;
+	}
+
+	if (!HostName)
+	{
+		HostName = DefaultHostName;
+	}
+	
+	UGameplayCameraSystemHost* Host = FindObject<UGameplayCameraSystemHost>(PlayerController, HostName);
+	ensureMsgf(Host || bAllowNull, 
+			TEXT("Can't find gameplay camera system host named '%s' under player controller '%s'."),
+			HostName, *GetNameSafe(PlayerController));
+	return Host;
+}
+
+UGameplayCameraSystemHost::UGameplayCameraSystemHost(const FObjectInitializer& ObjectInitializer)
+	: Super(ObjectInitializer)
+{
+	if (!HasAnyFlags(RF_ClassDefaultObject))
+	{
+		Evaluator = MakeShared<FCameraSystemEvaluator>();
+		Evaluator->Initialize(this);
+	}
+}
+
+void UGameplayCameraSystemHost::BeginDestroy()
+{
+	Super::BeginDestroy();
+
+	Evaluator.Reset();
+}
+
+void UGameplayCameraSystemHost::AddReferencedObjects(UObject* Object, FReferenceCollector& Collector)
+{
+	Super::AddReferencedObjects(Object, Collector);
+
+	UGameplayCameraSystemHost* This = CastChecked<UGameplayCameraSystemHost>(Object);
+	if (This->Evaluator.IsValid())
+	{
+		This->Evaluator->AddReferencedObjects(Collector);
+	}
+}
+
+TSharedPtr<UE::Cameras::FCameraSystemEvaluator> UGameplayCameraSystemHost::GetCameraSystemEvaluator()
+{
+	return Evaluator;
+}
+
diff --git a/Engine/Plugins/Cameras/GameplayCameras/Source/GameplayCameras/Public/GameFramework/ActivateCameraRigFunctions.h b/Engine/Plugins/Cameras/GameplayCameras/Source/GameplayCameras/Public/GameFramework/ActivateCameraRigFunctions.h
--- a/Engine/Plugins/Cameras/GameplayCameras/Source/GameplayCameras/Public/GameFramework/ActivateCameraRigFunctions.h
+++ b/Engine/Plugins/Cameras/GameplayCameras/Source/GameplayCameras/Public/GameFramework/ActivateCameraRigFunctions.h
@@ -15,6 +15,7 @@
 
 class APlayerController;
 class UCameraRigAsset;
+class UGameplayCameraSystemHost;
 
 namespace UE::Cameras
 {
@@ -78,9 +79,9 @@ public:
 private:
 
 	void ActivateCameraRigs();
+	void ActivateCameraRig(int32 Index);
 	void EnsureEvaluationContext();
-
-	static UE::Cameras::FCameraSystemEvaluator* FindCameraSystemEvaluator(APlayerController* PlayerController);
+	void EnsureCameraSystemHost();
 
 private:
 
@@ -92,6 +93,10 @@ private:
 	};
 
 	TArray<FCameraRigInfo> CameraRigInfos;
+
 	TSharedPtr<UE::Cameras::FCameraEvaluationContext> EvaluationContext;
+
+	UPROPERTY()
+	TObjectPtr<UGameplayCameraSystemHost> CameraSystemHost;
 };
 
diff --git a/Engine/Plugins/Cameras/GameplayCameras/Source/GameplayCameras/Public/GameFramework/GameplayCameraComponent.h b/Engine/Plugins/Cameras/GameplayCameras/Source/GameplayCameras/Public/GameFramework/GameplayCameraComponent.h
--- a/Engine/Plugins/Cameras/GameplayCameras/Source/GameplayCameras/Public/GameFramework/GameplayCameraComponent.h
+++ b/Engine/Plugins/Cameras/GameplayCameras/Source/GameplayCameras/Public/GameFramework/GameplayCameraComponent.h
@@ -4,17 +4,15 @@
 
 #include "Components/SceneComponent.h"
 #include "Core/CameraEvaluationContext.h"
-#include "CoreMinimal.h"
-#include "Engine/TimerHandle.h"
 #include "GameFramework/BlueprintCameraPose.h"
 #include "GameFramework/BlueprintCameraVariableTable.h"
 #include "UObject/ObjectMacros.h"
-#include "UObject/ScriptInterface.h"
 
 #include "GameplayCameraComponent.generated.h"
 
 class APlayerController;
 class UCameraAsset;
+class UGameplayCameraSystemHost;
 
 namespace UE::Cameras
 {
@@ -84,11 +82,6 @@ private:
 	void ActivateCameraEvaluationContext(APlayerController* PlayerController);
 	void DeactivateCameraEvaluationContext();
 
-	void DelayedActivateCameraEvaluationContext(APlayerController* PlayerController, AActor* OldViewTarget, AActor* NewViewTarget);
-	void AbortDelayedActivateCameraEvaluationContext();
-
-	void DoActivateCameraEvaluationContext(APlayerController* PlayerController, TSharedPtr<UE::Cameras::FCameraSystemEvaluator> CameraSystemEvaluator);
-
 #if WITH_EDITORONLY_DATA
 
 	void UpdatePreviewMeshTransform();
@@ -114,6 +107,9 @@ protected:
 
 	TSharedPtr<FGameplayCameraComponentEvaluationContext> EvaluationContext;
 
+	UPROPERTY()
+	TObjectPtr<UGameplayCameraSystemHost> CameraSystemHost;
+
 #if WITH_EDITORONLY_DATA
 
 	UPROPERTY(Transient)
diff --git a/Engine/Plugins/Cameras/GameplayCameras/Source/GameplayCameras/Public/GameFramework/GameplayCameraSystemComponent.h b/Engine/Plugins/Cameras/GameplayCameras/Source/GameplayCameras/Public/GameFramework/GameplayCameraSystemComponent.h
--- a/Engine/Plugins/Cameras/GameplayCameras/Source/GameplayCameras/Public/GameFramework/GameplayCameraSystemComponent.h
+++ b/Engine/Plugins/Cameras/GameplayCameras/Source/GameplayCameras/Public/GameFramework/GameplayCameraSystemComponent.h
@@ -2,11 +2,8 @@
 
 #pragma once
 
-#include "CoreMinimal.h"
-#include "Core/CameraSystemEvaluator.h"
 #include "GameplayCameras.h"
 #include "UObject/ObjectMacros.h"
-#include "UObject/ScriptInterface.h"
 #include "Components/SceneComponent.h"
 
 #include "GameplayCameraSystemComponent.generated.h"
@@ -14,14 +11,13 @@
 class APlayerController;
 class UCameraRigAsset;
 class UCanvas;
+class UGameplayCameraSystemHost;
 struct FMinimalViewInfo;
 
 namespace UE::Cameras
 {
-
-class FCameraSystemEvaluator;
-
-}  // namespace UE::Cameras
+	class FCameraSystemEvaluator;
+}
 
 /**
  * A component that hosts a camera system.
@@ -38,7 +34,7 @@ public:
 	UGameplayCameraSystemComponent(const FObjectInitializer& ObjectInit);
 
 	/** Gets the camera system evaluator. */
-	TSharedPtr<FCameraSystemEvaluator> GetCameraSystemEvaluator() { return Evaluator; }
+	GAMEPLAYCAMERAS_API TSharedPtr<FCameraSystemEvaluator> GetCameraSystemEvaluator(bool bEnsureIfNull = true);
 
 	/** Updates the camera system and returns the computed view. */
 	GAMEPLAYCAMERAS_API void GetCameraView(float DeltaTime, FMinimalViewInfo& DesiredView);
@@ -68,9 +64,6 @@ public:
 #if WITH_EDITOR
 	virtual bool GetEditorPreviewInfo(float DeltaTime, FMinimalViewInfo& ViewOut) override;
 #endif  // WITH_EDITOR
-	
-	// UObject interface
-	static void AddReferencedObjects(UObject* InThis, FReferenceCollector& Collector);
 
 public:
 
@@ -95,7 +88,8 @@ public:
 
 private:
 	
-	TSharedPtr<FCameraSystemEvaluator> Evaluator;
+	UPROPERTY()
+	TObjectPtr<UGameplayCameraSystemHost> CameraSystemHost;
 
 	TWeakObjectPtr<APlayerController> WeakPlayerController;
 
diff --git a/Engine/Plugins/Cameras/GameplayCameras/Source/GameplayCameras/Public/GameFramework/GameplayCameraSystemHost.h b/Engine/Plugins/Cameras/GameplayCameras/Source/GameplayCameras/Public/GameFramework/GameplayCameraSystemHost.h
new file mode 100644
--- /dev/null
+++ b/Engine/Plugins/Cameras/GameplayCameras/Source/GameplayCameras/Public/GameFramework/GameplayCameraSystemHost.h
@@ -0,0 +1,71 @@
+// Copyright Epic Games, Inc. All Rights Reserved.
+
+#pragma once
+
+#include "UObject/Object.h"
+#include "Templates/SharedPointerFwd.h"
+
+#include "GameplayCameraSystemHost.generated.h"
+
+class APlayerController;
+
+namespace UE::Cameras
+{
+	class FCameraSystemEvaluator;
+}
+
+/**
+ * A class that host a camera system evaluator so that it can be accessed in a game world.
+ *
+ * The host doesn't stay alive very long if nothing references it. Gameplay camera components and actors
+ * are meant to hold a reference to it while they use it. When nobody uses it, the host is meant to be
+ * collectable by the GC.
+ */
+UCLASS(MinimalAPI)
+class UGameplayCameraSystemHost : public UObject
+{
+	GENERATED_BODY()
+
+public:
+
+	/** Creates a new camera system host. */
+	GAMEPLAYCAMERAS_API UGameplayCameraSystemHost(const FObjectInitializer& ObjectInitializer);
+
+	// UObject interface.
+	static GAMEPLAYCAMERAS_API void AddReferencedObjects(UObject* InThis, FReferenceCollector& Collector);
+
+public:
+
+	using FCameraSystemEvaluator = UE::Cameras::FCameraSystemEvaluator;
+
+	/** Gets the camera system evaluator. */
+	GAMEPLAYCAMERAS_API TSharedPtr<FCameraSystemEvaluator> GetCameraSystemEvaluator();
+
+protected:
+
+	// UObject interface.
+	GAMEPLAYCAMERAS_API virtual void BeginDestroy() override;
+
+public:
+
+	/**
+	 * Finds a camera system host under the given player controller, or creates one if none was found.
+	 * If the "auto-spawn camera system actor" setting is enabled, this will also spawn an actor that 
+	 * references the newly created host, and sets itself as the active view-target.
+	 */
+	static GAMEPLAYCAMERAS_API UGameplayCameraSystemHost* FindOrCreateHost(APlayerController* PlayerController, const TCHAR* HostName = nullptr);
+
+	/**
+	 * Finds a camera system host under the given player controller.
+	 */
+	static GAMEPLAYCAMERAS_API UGameplayCameraSystemHost* FindHost(APlayerController* PlayerController, const TCHAR* HostName = nullptr, bool bAllowNull = true);
+
+private:
+
+	/** Default host name to use when creating a new host. */
+	static const TCHAR* DefaultHostName;
+
+	/** The camera system evaluator. */
+	TSharedPtr<FCameraSystemEvaluator> Evaluator;
+};
+
