Cameras: more tweaks of activation/deactivation of various components

This is a follow-up to the previous change. In particular it makes the control rotation component use the new host API.

diff --git a/Engine/Plugins/Cameras/GameplayCameras/Source/GameplayCameras/Private/GameFramework/GameplayCameraComponent.cpp b/Engine/Plugins/Cameras/GameplayCameras/Source/GameplayCameras/Private/GameFramework/GameplayCameraComponent.cpp
--- a/Engine/Plugins/Cameras/GameplayCameras/Source/GameplayCameras/Private/GameFramework/GameplayCameraComponent.cpp
+++ b/Engine/Plugins/Cameras/GameplayCameras/Source/GameplayCameras/Private/GameFramework/GameplayCameraComponent.cpp
@@ -38,6 +38,15 @@ TSharedPtr<UE::Cameras::FCameraEvaluationContext> UGameplayCameraComponent::GetE
 	return EvaluationContext;
 }
 
+APlayerController* UGameplayCameraComponent::GetPlayerController() const
+{
+	if (CameraSystemHost)
+	{
+		return CameraSystemHost->GetPlayerController();
+	}
+	return nullptr;
+}
+
 void UGameplayCameraComponent::ActivateCameraForPlayerIndex(int32 PlayerIndex)
 {
 	ActivateCameraEvaluationContext(PlayerIndex);
@@ -55,10 +64,7 @@ void UGameplayCameraComponent::DeactivateCamera()
 
 void UGameplayCameraComponent::ActivateCameraEvaluationContext(int32 PlayerIndex)
 {
-	if (WeakPlayerController.IsValid())
-	{
-		DeactivateCameraEvaluationContext();
-	}
+	DeactivateCameraEvaluationContext();
 
 	APlayerController* PlayerController = UGameplayStatics::GetPlayerController(this, PlayerIndex);
 	if (!PlayerController)
@@ -74,12 +80,6 @@ void UGameplayCameraComponent::DeactivateCameraEvaluationContext()
 {
 	using namespace UE::Cameras;
 
-	APlayerController* PlayerController = WeakPlayerController.Get();
-	if (!PlayerController)
-	{
-		return;
-	}
-
 	if (!CameraSystemHost)
 	{
 		return;
@@ -90,8 +90,6 @@ void UGameplayCameraComponent::DeactivateCameraEvaluationContext()
 		TSharedPtr<FCameraSystemEvaluator> Evaluator = CameraSystemHost->GetCameraSystemEvaluator();
 		Evaluator->RemoveEvaluationContext(EvaluationContext.ToSharedRef());
 	}
-
-	WeakPlayerController.Reset();
 }
 
 void UGameplayCameraComponent::ActivateCameraEvaluationContext(APlayerController* PlayerController)
@@ -129,8 +127,6 @@ void UGameplayCameraComponent::ActivateCameraEvaluationContext(APlayerController
 
 	TSharedPtr<FCameraSystemEvaluator> CameraSystemEvaluator = CameraSystemHost->GetCameraSystemEvaluator();
 	CameraSystemEvaluator->PushEvaluationContext(EvaluationContext.ToSharedRef());
-
-	WeakPlayerController = PlayerController;
 }
 
 FBlueprintCameraPose UGameplayCameraComponent::GetInitialPose() const
diff --git a/Engine/Plugins/Cameras/GameplayCameras/Source/GameplayCameras/Private/GameFramework/GameplayCameraSystemHost.cpp b/Engine/Plugins/Cameras/GameplayCameras/Source/GameplayCameras/Private/GameFramework/GameplayCameraSystemHost.cpp
--- a/Engine/Plugins/Cameras/GameplayCameras/Source/GameplayCameras/Private/GameFramework/GameplayCameraSystemHost.cpp
+++ b/Engine/Plugins/Cameras/GameplayCameras/Source/GameplayCameras/Private/GameFramework/GameplayCameraSystemHost.cpp
@@ -101,6 +101,11 @@ void UGameplayCameraSystemHost::AddReferencedObjects(UObject* Object, FReference
 	}
 }
 
+APlayerController* UGameplayCameraSystemHost::GetPlayerController()
+{
+	return GetTypedOuter<APlayerController>();
+}
+
 TSharedPtr<UE::Cameras::FCameraSystemEvaluator> UGameplayCameraSystemHost::GetCameraSystemEvaluator()
 {
 	return Evaluator;
diff --git a/Engine/Plugins/Cameras/GameplayCameras/Source/GameplayCameras/Private/GameFramework/GameplayControlRotationComponent.cpp b/Engine/Plugins/Cameras/GameplayCameras/Source/GameplayCameras/Private/GameFramework/GameplayControlRotationComponent.cpp
--- a/Engine/Plugins/Cameras/GameplayCameras/Source/GameplayCameras/Private/GameFramework/GameplayControlRotationComponent.cpp
+++ b/Engine/Plugins/Cameras/GameplayCameras/Source/GameplayCameras/Private/GameFramework/GameplayControlRotationComponent.cpp
@@ -3,13 +3,13 @@
 #include "GameFramework/GameplayControlRotationComponent.h"
 
 #include "Core/CameraEvaluationContext.h"
+#include "Core/CameraSystemEvaluator.h"
 #include "Debug/DebugDrawService.h"
 #include "Engine/Engine.h"
 #include "Engine/World.h"
-#include "EnhancedInputComponent.h"
 #include "GameFramework/Controller.h"
 #include "GameFramework/GameplayCameraComponent.h"
-#include "GameFramework/GameplayCameraSystemComponent.h"
+#include "GameFramework/GameplayCameraSystemHost.h"
 #include "GameFramework/Pawn.h"
 #include "GameplayCameras.h"
 #include "Kismet/GameplayStatics.h"
@@ -63,18 +63,10 @@ void UGameplayControlRotationComponent::BeginPlay()
 		return;
 	}
 
-	AActor* ViewTarget = PlayerController->PlayerCameraManager->GetViewTarget();
-	UGameplayCameraSystemComponent* CameraSystemComponent = ViewTarget->FindComponentByClass<UGameplayCameraSystemComponent>();
-	if (!CameraSystemComponent)
+	CameraSystemHost = UGameplayCameraSystemHost::FindHost(PlayerController);
+	if (!CameraSystemHost)
 	{
-		UE_LOG(LogCameraSystem, Error, TEXT("Can't find camera system on the player camera manager."));
-		return;
-	}
-
-	TSharedPtr<FCameraSystemEvaluator> CameraSystem = CameraSystemComponent->GetCameraSystemEvaluator();
-	if (!CameraSystem)
-	{
-		UE_LOG(LogCameraSystem, Error, TEXT("No camera system has been initialized"));
+		UE_LOG(LogCameraSystem, Error, TEXT("Can't find camera system host on the player controller."));
 		return;
 	}
 
@@ -87,17 +79,18 @@ void UGameplayControlRotationComponent::BeginPlay()
 	ServiceParams.bApplyControlRotation = false;
 
 	ControlRotationService = MakeShared<FPlayerControlRotationEvaluationService>(ServiceParams);
-	CameraSystem->RegisterEvaluationService(ControlRotationService.ToSharedRef());
 
-	WeakCameraSystem = CameraSystem;
+	TSharedPtr<FCameraSystemEvaluator> CameraSystem = CameraSystemHost->GetCameraSystemEvaluator();
+	CameraSystem->RegisterEvaluationService(ControlRotationService.ToSharedRef());
 }
 
 void UGameplayControlRotationComponent::EndPlay(const EEndPlayReason::Type EndPlayReason)
 {
 	using namespace UE::Cameras;
 
-	if (TSharedPtr<FCameraSystemEvaluator> CameraSystem = WeakCameraSystem.Pin())
+	if (CameraSystemHost)
 	{
+		TSharedPtr<FCameraSystemEvaluator> CameraSystem = CameraSystemHost->GetCameraSystemEvaluator();
 		CameraSystem->UnregisterEvaluationService(ControlRotationService.ToSharedRef());
 	}
 
diff --git a/Engine/Plugins/Cameras/GameplayCameras/Source/GameplayCameras/Public/GameFramework/ActivateCameraRigFunctions.h b/Engine/Plugins/Cameras/GameplayCameras/Source/GameplayCameras/Public/GameFramework/ActivateCameraRigFunctions.h
--- a/Engine/Plugins/Cameras/GameplayCameras/Source/GameplayCameras/Public/GameFramework/ActivateCameraRigFunctions.h
+++ b/Engine/Plugins/Cameras/GameplayCameras/Source/GameplayCameras/Public/GameFramework/ActivateCameraRigFunctions.h
@@ -79,7 +79,6 @@ public:
 private:
 
 	void ActivateCameraRigs();
-	void ActivateCameraRig(int32 Index);
 	void EnsureEvaluationContext();
 	void EnsureCameraSystemHost();
 
diff --git a/Engine/Plugins/Cameras/GameplayCameras/Source/GameplayCameras/Public/GameFramework/GameplayCameraComponent.h b/Engine/Plugins/Cameras/GameplayCameras/Source/GameplayCameras/Public/GameFramework/GameplayCameraComponent.h
--- a/Engine/Plugins/Cameras/GameplayCameras/Source/GameplayCameras/Public/GameFramework/GameplayCameraComponent.h
+++ b/Engine/Plugins/Cameras/GameplayCameras/Source/GameplayCameras/Public/GameFramework/GameplayCameraComponent.h
@@ -33,13 +33,13 @@ class UGameplayCameraComponent : public USceneComponent
 public:
 
 	/** Create a new camera component. */
-	UGameplayCameraComponent(const FObjectInitializer& ObjectInit);
+	GAMEPLAYCAMERAS_API UGameplayCameraComponent(const FObjectInitializer& ObjectInit);
 
 	/** Get the camera evaluation context used by this component. */
-	TSharedPtr<UE::Cameras::FCameraEvaluationContext> GetEvaluationContext();
+	GAMEPLAYCAMERAS_API TSharedPtr<UE::Cameras::FCameraEvaluationContext> GetEvaluationContext();
 
 	/** Get the player controller this component is currently activated for (if any). */
-	APlayerController* GetPlayerController() const { return WeakPlayerController.Get(); }
+	GAMEPLAYCAMERAS_API APlayerController* GetPlayerController() const;
 
 public:
 
@@ -119,10 +119,6 @@ protected:
 	TObjectPtr<UStaticMeshComponent> PreviewMeshComponent;
 
 #endif	// WITH_EDITORONLY_DATA
-	
-private:
-
-	TWeakObjectPtr<APlayerController> WeakPlayerController;
 };
 
 namespace UE::Cameras
diff --git a/Engine/Plugins/Cameras/GameplayCameras/Source/GameplayCameras/Public/GameFramework/GameplayCameraSystemHost.h b/Engine/Plugins/Cameras/GameplayCameras/Source/GameplayCameras/Public/GameFramework/GameplayCameraSystemHost.h
--- a/Engine/Plugins/Cameras/GameplayCameras/Source/GameplayCameras/Public/GameFramework/GameplayCameraSystemHost.h
+++ b/Engine/Plugins/Cameras/GameplayCameras/Source/GameplayCameras/Public/GameFramework/GameplayCameraSystemHost.h
@@ -38,6 +38,9 @@ public:
 
 	using FCameraSystemEvaluator = UE::Cameras::FCameraSystemEvaluator;
 
+	/** Get the player controller that this host is hooked up to. */
+	GAMEPLAYCAMERAS_API APlayerController* GetPlayerController();
+
 	/** Gets the camera system evaluator. */
 	GAMEPLAYCAMERAS_API TSharedPtr<FCameraSystemEvaluator> GetCameraSystemEvaluator();
 
diff --git a/Engine/Plugins/Cameras/GameplayCameras/Source/GameplayCameras/Public/GameFramework/GameplayControlRotationComponent.h b/Engine/Plugins/Cameras/GameplayCameras/Source/GameplayCameras/Public/GameFramework/GameplayControlRotationComponent.h
--- a/Engine/Plugins/Cameras/GameplayCameras/Source/GameplayCameras/Public/GameFramework/GameplayControlRotationComponent.h
+++ b/Engine/Plugins/Cameras/GameplayCameras/Source/GameplayCameras/Public/GameFramework/GameplayControlRotationComponent.h
@@ -3,11 +3,11 @@
 #pragma once
 
 #include "Components/ActorComponent.h"
-#include "Core/CameraSystemEvaluator.h"
 
 #include "GameplayControlRotationComponent.generated.h"
 
 class APlayerController;
+class UGameplayCameraSystemHost;
 class UCanvas;
 class UInputAction;
 
@@ -63,7 +63,9 @@ private:
 	UPROPERTY()
 	TObjectPtr<APlayerController> PlayerController;
 
+	UPROPERTY()
+	TObjectPtr<UGameplayCameraSystemHost> CameraSystemHost;
+
 	TSharedPtr<UE::Cameras::FPlayerControlRotationEvaluationService> ControlRotationService;
-	TWeakPtr<UE::Cameras::FCameraSystemEvaluator> WeakCameraSystem;
 };
 
