Cameras: fix camera rig picker details customization widget

- Allow for pickers that let the user select the camera asset or not.
- Use specific UProperty metadata for a picker that shows camera director warning messages.
- Fix initial population of camera rig list view.

diff --git a/Engine/Plugins/Cameras/GameplayCameras/Source/GameplayCameras/Public/Directors/StateTreeCameraDirectorTasks.h b/Engine/Plugins/Cameras/GameplayCameras/Source/GameplayCameras/Public/Directors/StateTreeCameraDirectorTasks.h
--- a/Engine/Plugins/Cameras/GameplayCameras/Source/GameplayCameras/Public/Directors/StateTreeCameraDirectorTasks.h
+++ b/Engine/Plugins/Cameras/GameplayCameras/Source/GameplayCameras/Public/Directors/StateTreeCameraDirectorTasks.h
@@ -14,7 +14,7 @@ struct FGameplayCamerasActivateCameraRigTaskInstanceData
 	GENERATED_BODY()
 
 	/** The camera rig to activate. */
-	UPROPERTY(EditAnywhere, Category="Cameras", meta=(UseCameraRigPicker=true))
+	UPROPERTY(EditAnywhere, Category="Cameras", meta=(UseCameraDirectorRigPicker=true))
 	TObjectPtr<UCameraRigAsset> CameraRig;
 };
 
diff --git a/Engine/Plugins/Cameras/GameplayCameras/Source/GameplayCamerasEditor/Private/Customizations/CameraRigPtrDetailsCustomization.cpp b/Engine/Plugins/Cameras/GameplayCameras/Source/GameplayCamerasEditor/Private/Customizations/CameraRigPtrDetailsCustomization.cpp
--- a/Engine/Plugins/Cameras/GameplayCameras/Source/GameplayCamerasEditor/Private/Customizations/CameraRigPtrDetailsCustomization.cpp
+++ b/Engine/Plugins/Cameras/GameplayCameras/Source/GameplayCamerasEditor/Private/Customizations/CameraRigPtrDetailsCustomization.cpp
@@ -31,13 +31,25 @@ void FCameraRigPtrDetailsCustomization::CustomizeHeader(TSharedRef<IPropertyHand
 
 	FProperty* StructProperty = StructPropertyHandle->GetProperty();
 	const bool bUseCameraRigPicker = StructProperty->GetBoolMetaData("UseCameraRigPicker");
+	const bool bUseCameraDirectorRigPicker = StructProperty->GetBoolMetaData("UseCameraDirectorRigPicker");
 
 	TSharedPtr<SWidget> ValueContentWidget;
-	if (bUseCameraRigPicker)
+	if (bUseCameraRigPicker || bUseCameraDirectorRigPicker)
 	{
+		FOnGetContent OnGetComboMenuContent;
+		if (bUseCameraRigPicker)
+		{
+			OnGetComboMenuContent = FOnGetContent::CreateSP(this, &FCameraRigPtrDetailsCustomization::OnBuildCameraRigNamePicker);
+		}
+		else if (bUseCameraDirectorRigPicker)
+		{
+			OnGetComboMenuContent = FOnGetContent::CreateSP(this, &FCameraRigPtrDetailsCustomization::OnBuildCameraDirectorRigNamePicker);
+		}
+
 		CameraRigPickerButton = SNew(SComboButton)
 		.ButtonStyle(FAppStyle::Get(), "PropertyEditor.AssetComboStyle")
 		.ContentPadding(FMargin(2.f, 2.f, 2.f, 1.f))
+		.OnGetMenuContent(OnGetComboMenuContent)
 		.ButtonContent()
 		[
 			SNew(STextBlock)
@@ -45,8 +57,7 @@ void FCameraRigPtrDetailsCustomization::CustomizeHeader(TSharedRef<IPropertyHand
 			.Font(FAppStyle::GetFontStyle("PropertyWindow.NormalFont"))
 			.Text(this, &FCameraRigPtrDetailsCustomization::OnGetComboText)
 			.ToolTipText(this, &FCameraRigPtrDetailsCustomization::OnGetComboToolTipText)
-		]
-		.OnGetMenuContent(this, &FCameraRigPtrDetailsCustomization::OnBuildCameraRigNamePicker);
+		];
 
 		ValueContentWidget = CameraRigPickerButton;
 	}
@@ -80,6 +91,7 @@ FText FCameraRigPtrDetailsCustomization::OnGetComboText() const
 		{
 			return FText::FromString(CameraRig->GetDisplayName());
 		}
+		return LOCTEXT("NoCameraRigValue", "Select camera rig");
 	}
 	else if (PropertyAccess == FPropertyAccess::MultipleValues)
 	{
@@ -93,13 +105,12 @@ FText FCameraRigPtrDetailsCustomization::OnGetComboToolTipText() const
 	return LOCTEXT("ComboToolTipText", "The name of the camera rig to activate.");
 }
 
-TSharedRef<SWidget> FCameraRigPtrDetailsCustomization::OnBuildCameraRigNamePicker()
+TSharedRef<SWidget> FCameraRigPtrDetailsCustomization::OnBuildCameraDirectorRigNamePicker()
 {
-	FCameraRigPickerConfig CameraRigPickerConfig;
-	CameraRigPickerConfig.bCanSelectCameraAsset = false;
-	CameraRigPickerConfig.bFocusCameraRigSearchBoxWhenOpened = true;
-	CameraRigPickerConfig.OnCameraRigSelected = FOnCameraRigSelected::CreateSP(this, &FCameraRigPtrDetailsCustomization::OnPickerAssetSelected);
+	FCameraRigPickerConfig PickerConfig;
+	PickerConfig.bCanSelectCameraAsset = false;
 
+	// Figure out what camera asset is referencing our outermost objects (main package objects).
 	TArray<UObject*> OuterObjects;
 	CameraRigPropertyHandle->GetOuterObjects(OuterObjects);
 
@@ -112,6 +123,9 @@ TSharedRef<SWidget> FCameraRigPtrDetailsCustomization::OnBuildCameraRigNamePicke
 		}
 	}
 
+	// Automatically show the rigs of the referencing camera asset. Show a warning if none or multiple
+	// camera assets are referencing us, similar to the warning for the Blueprint camera director
+	// rig activation picker.
 	TArray<UCameraAsset*> ReferencingCameraAssets;
 	for (UObject* CameraDirectorObject : CameraDirectorObjects)
 	{
@@ -120,32 +134,48 @@ TSharedRef<SWidget> FCameraRigPtrDetailsCustomization::OnBuildCameraRigNamePicke
 
 	if (ReferencingCameraAssets.Num() == 0)
 	{
-		CameraRigPickerConfig.WarningMessage = LOCTEXT("NoReferencingCameraAssetWarning",
-				"No camera asset references this asset, so no camera rig list can be displayed. "
-				"Make a camera asset use this asset as its camera director evaluator, or use "
-				"ActivateCameraRigViaProxy.");
+		PickerConfig.WarningMessage = LOCTEXT("NoReferencingCameraAssetWarning",
+				"No camera asset references this camera director, so no camera rig list can be displayed. "
+				"Make a camera asset use this asset as its camera director evaluator, or use camera rig"
+				"proxy assets instead.");
 	}
 	else
 	{
-		CameraRigPickerConfig.InitialCameraAssetSelection = ReferencingCameraAssets[0];
+		PickerConfig.InitialCameraAssetSelection = ReferencingCameraAssets[0];
 
 		if (ReferencingCameraAssets.Num() > 1)
 		{
-			CameraRigPickerConfig.WarningMessage = LOCTEXT("ManyReferencingCameraAssetsWarning",
-				"More than one camera asset references this asset. Only camera rigs from the first "
-				"one will be displayed. Even then, shared camera director assets should use "
-				"ActivateCameraRigViaProxy instead.");
+			PickerConfig.WarningMessage = LOCTEXT("ManyReferencingCameraAssetsWarning",
+				"More than one camera asset references this camera director. Only camera rigs from the first "
+				"one will be displayed. Even then, shared camera director assets should use camera rig"
+				"proxy assets instead.");
 		}
 	}
 
+	return BuildCameraRigNamePickerImpl(PickerConfig);
+}
+
+TSharedRef<SWidget> FCameraRigPtrDetailsCustomization::OnBuildCameraRigNamePicker()
+{
+	FCameraRigPickerConfig PickerConfig;
+	PickerConfig.bCanSelectCameraAsset = true;
+
+	return BuildCameraRigNamePickerImpl(PickerConfig);
+}
+
+TSharedRef<SWidget> FCameraRigPtrDetailsCustomization::BuildCameraRigNamePickerImpl(FCameraRigPickerConfig& PickerConfig)
+{
+	PickerConfig.bFocusCameraRigSearchBoxWhenOpened = true;
+	PickerConfig.OnCameraRigSelected = FOnCameraRigSelected::CreateSP(this, &FCameraRigPtrDetailsCustomization::OnPickerAssetSelected);
+
 	UObject* SelectedCameraRig;
 	if (CameraRigPropertyHandle->GetValue(SelectedCameraRig) == FPropertyAccess::Success)
 	{
-		CameraRigPickerConfig.InitialCameraRigSelection = Cast<UCameraRigAsset>(SelectedCameraRig);
+		PickerConfig.InitialCameraRigSelection = Cast<UCameraRigAsset>(SelectedCameraRig);
 	}
 
 	IGameplayCamerasEditorModule& CamerasEditorModule = IGameplayCamerasEditorModule::Get();
-	return CamerasEditorModule.CreateCameraRigPicker(CameraRigPickerConfig);
+	return CamerasEditorModule.CreateCameraRigPicker(PickerConfig);
 }
 
 void FCameraRigPtrDetailsCustomization::OnPickerAssetSelected(UCameraRigAsset* SelectedItem)
diff --git a/Engine/Plugins/Cameras/GameplayCameras/Source/GameplayCamerasEditor/Public/Customizations/CameraRigPtrDetailsCustomization.h b/Engine/Plugins/Cameras/GameplayCameras/Source/GameplayCamerasEditor/Public/Customizations/CameraRigPtrDetailsCustomization.h
--- a/Engine/Plugins/Cameras/GameplayCameras/Source/GameplayCamerasEditor/Public/Customizations/CameraRigPtrDetailsCustomization.h
+++ b/Engine/Plugins/Cameras/GameplayCameras/Source/GameplayCamerasEditor/Public/Customizations/CameraRigPtrDetailsCustomization.h
@@ -15,6 +15,8 @@ struct FAssetData;
 namespace UE::Cameras
 {
 
+struct FCameraRigPickerConfig;
+
 class FCameraRigPtrDetailsCustomization : public IPropertyTypeCustomization
 {
 public:
@@ -32,6 +34,8 @@ private:
 	FText OnGetComboText() const;
 	FText OnGetComboToolTipText() const;
 	TSharedRef<SWidget> OnBuildCameraRigNamePicker();
+	TSharedRef<SWidget> OnBuildCameraDirectorRigNamePicker();
+	TSharedRef<SWidget> BuildCameraRigNamePickerImpl(FCameraRigPickerConfig& PickerConfig);
 	void OnPickerAssetSelected(UCameraRigAsset* SelectedItem);
 
 private:
