Cameras: fix initial population of camera rig picker widget

diff --git a/Engine/Plugins/Cameras/GameplayCameras/Source/GameplayCamerasEditor/Private/Editors/SCameraRigPicker.cpp b/Engine/Plugins/Cameras/GameplayCameras/Source/GameplayCamerasEditor/Private/Editors/SCameraRigPicker.cpp
--- a/Engine/Plugins/Cameras/GameplayCameras/Source/GameplayCamerasEditor/Private/Editors/SCameraRigPicker.cpp
+++ b/Engine/Plugins/Cameras/GameplayCameras/Source/GameplayCamerasEditor/Private/Editors/SCameraRigPicker.cpp
@@ -189,6 +189,8 @@ void SCameraRigPicker::Construct(const FArguments& InArgs)
 	}
 
 	// If we have an initially selected assets, register a timer to do that next frame.
+	// When there is only an initially selected camera asset, do the setup with a null camera rig in order
+	// to make sure the list of camera rigs is populated in the list view.
 	if (PickerConfig.InitialCameraRigSelection)
 	{
 		TVariant<UCameraRigAsset*, FGuid> SelectedCameraRig(TInPlaceType<UCameraRigAsset*>(), PickerConfig.InitialCameraRigSelection);
@@ -199,6 +201,11 @@ void SCameraRigPicker::Construct(const FArguments& InArgs)
 		TVariant<UCameraRigAsset*, FGuid> SelectedCameraRig(TInPlaceType<FGuid>(), PickerConfig.InitialCameraRigSelectionGuid);
 		SetupInitialSelections(PickerConfig.InitialCameraAssetSelection, SelectedCameraRig);
 	}
+	else if (PickerConfig.InitialCameraAssetSelection.IsValid())
+	{
+		TVariant<UCameraRigAsset*, FGuid> NullCameraRig(TInPlaceType<UCameraRigAsset*>(), nullptr);
+		SetupInitialSelections(PickerConfig.InitialCameraAssetSelection, NullCameraRig);
+	}
 
 	// If we need to focus the search box, register a timer to do that next frame.
 	if (PickerConfig.bFocusCameraRigSearchBoxWhenOpened)
