Cameras: tweak default camera rig name display in the graph pin picker

diff --git a/Engine/Plugins/Cameras/GameplayCameras/Source/GameplayCamerasEditor/Private/Editors/SCameraRigNameGraphPin.cpp b/Engine/Plugins/Cameras/GameplayCameras/Source/GameplayCamerasEditor/Private/Editors/SCameraRigNameGraphPin.cpp
--- a/Engine/Plugins/Cameras/GameplayCameras/Source/GameplayCamerasEditor/Private/Editors/SCameraRigNameGraphPin.cpp
+++ b/Engine/Plugins/Cameras/GameplayCameras/Source/GameplayCamerasEditor/Private/Editors/SCameraRigNameGraphPin.cpp
@@ -95,15 +95,15 @@ FSlateColor SCameraRigNameGraphPin::OnGetWidgetBackground() const
 
 FText SCameraRigNameGraphPin::OnGetSelectedCameraRigName() const
 {
-	FText Value;
 	if (GraphPinObj != nullptr)
 	{
 		if (const UCameraRigAsset* CameraRig = Cast<const UCameraRigAsset>(GraphPinObj->DefaultObject))
 		{
-			Value = FText::FromString(CameraRig->GetDisplayName());
+			return FText::FromString(CameraRig->GetDisplayName());
 		}
+		return LOCTEXT("NoCameraRigName", "Select camera rig");
 	}
-	return Value;
+	return LOCTEXT("InvalidGraphPin", "Invalid graph pin");
 }
 
 FText SCameraRigNameGraphPin::OnGetCameraRigPickerToolTipText() const
