Cameras: fix undefined behavior/crash when use a blend node with zero blend time

diff --git a/Engine/Plugins/Cameras/GameplayCameras/Source/GameplayCameras/Private/Nodes/Blends/SimpleBlendCameraNode.cpp b/Engine/Plugins/Cameras/GameplayCameras/Source/GameplayCameras/Private/Nodes/Blends/SimpleBlendCameraNode.cpp
--- a/Engine/Plugins/Cameras/GameplayCameras/Source/GameplayCameras/Private/Nodes/Blends/SimpleBlendCameraNode.cpp
+++ b/Engine/Plugins/Cameras/GameplayCameras/Source/GameplayCameras/Private/Nodes/Blends/SimpleBlendCameraNode.cpp
@@ -94,7 +94,11 @@ void FSimpleFixedTimeBlendCameraNodeEvaluator::OnRun(const FCameraNodeEvaluation
 float FSimpleFixedTimeBlendCameraNodeEvaluator::GetTimeFactor() const
 {
 	const USimpleFixedTimeBlendCameraNode* BlendNode = GetCameraNodeAs<USimpleFixedTimeBlendCameraNode>();
-	return CurrentTime / BlendNode->BlendTime;
+	if (BlendNode->BlendTime > 0.f)
+	{
+		return CurrentTime / BlendNode->BlendTime;
+	}
+	return 1.f;
 }
 
 }  // namespace UE::Cameras
