Cameras: add support for camera rig proxies in the StateTree director

diff --git a/Engine/Plugins/Cameras/GameplayCameras/Source/GameplayCameras/Private/Directors/CameraDirectorStateTreeSchema.cpp b/Engine/Plugins/Cameras/GameplayCameras/Source/GameplayCameras/Private/Directors/CameraDirectorStateTreeSchema.cpp
--- a/Engine/Plugins/Cameras/GameplayCameras/Source/GameplayCameras/Private/Directors/CameraDirectorStateTreeSchema.cpp
+++ b/Engine/Plugins/Cameras/GameplayCameras/Source/GameplayCameras/Private/Directors/CameraDirectorStateTreeSchema.cpp
@@ -53,3 +53,9 @@ bool UCameraDirectorStateTreeSchema::IsExternalItemAllowed(const UStruct& InStru
 	return true;
 }
 
+void FCameraDirectorStateTreeEvaluationData::Reset()
+{
+	ActiveCameraRigs.Reset();
+	ActiveCameraRigProxies.Reset();
+}
+
diff --git a/Engine/Plugins/Cameras/GameplayCameras/Source/GameplayCameras/Private/Directors/StateTreeCameraDirector.cpp b/Engine/Plugins/Cameras/GameplayCameras/Source/GameplayCameras/Private/Directors/StateTreeCameraDirector.cpp
--- a/Engine/Plugins/Cameras/GameplayCameras/Source/GameplayCameras/Private/Directors/StateTreeCameraDirector.cpp
+++ b/Engine/Plugins/Cameras/GameplayCameras/Source/GameplayCameras/Private/Directors/StateTreeCameraDirector.cpp
@@ -5,6 +5,8 @@
 #include "Core/CameraAsset.h"
 #include "Core/CameraBuildLog.h"
 #include "Core/CameraEvaluationContext.h"
+#include "Core/CameraRigProxyAsset.h"
+#include "Core/CameraRigProxyTable.h"
 #include "Directors/CameraDirectorStateTreeSchema.h"
 #include "GameplayCameras.h"
 #include "Logging/TokenizedMessage.h"
@@ -35,6 +37,7 @@ protected:
 private:
 
 	bool SetContextRequirements(TSharedPtr<const FCameraEvaluationContext> OwnerContext, FStateTreeExecutionContext& StateTreeContext);
+	const UCameraRigAsset* FindCameraRigByProxy(const UCameraRigProxyAsset* InProxy);
 
 private:
 
@@ -143,12 +146,42 @@ void FStateTreeCameraDirectorEvaluator::OnRun(const FCameraDirectorEvaluationPar
 	{
 		StateTreeContext.Tick(Params.DeltaTime);
 
-		for (TObjectPtr<const UCameraRigAsset> CameraRig : EvaluationData.ActiveCameraRigs)
+		TArray<const UCameraRigAsset*, TInlineAllocator<2>> CameraRigs;
+		const UCameraAsset* CameraAsset = Params.OwnerContext->GetCameraAsset();
+
+		// Gather camera rigs.
+		for (const UCameraRigAsset* ActiveCameraRig : EvaluationData.ActiveCameraRigs)
+		{
+			if (ActiveCameraRig)
+			{
+				CameraRigs.Add(ActiveCameraRig);
+			}
+			else
+			{
+				UE_LOG(LogCameraSystem, Error, TEXT("Null camera rig specified in camera director '%s'."), 
+						*StateTree->GetPathName());
+			}
+		}
+
+		// Resolve camera rig proxies.
+		for (const UCameraRigProxyAsset* ActiveCameraRigProxy : EvaluationData.ActiveCameraRigProxies)
 		{
-			FActiveCameraRigInfo CameraRigInfo;
-			CameraRigInfo.CameraRig = CameraRig;
-			CameraRigInfo.EvaluationContext = Params.OwnerContext;
-			OutResult.ActiveCameraRigs.Add(CameraRigInfo);
+			const UCameraRigAsset* ActiveCameraRig = FindCameraRigByProxy(ActiveCameraRigProxy);
+			if (ActiveCameraRig)
+			{
+				CameraRigs.Add(ActiveCameraRig);
+			}
+			else
+			{
+				UE_LOG(LogCameraSystem, Error, TEXT("No camera rig found mapped to proxy '%s' in camera '%s'."),
+						*ActiveCameraRigProxy->GetPathName(), *CameraAsset->GetPathName());
+			}
+		}
+
+		// Set all collected camera rigs as our active rigs this frame.
+		for (const UCameraRigAsset* CameraRig : CameraRigs)
+		{
+			OutResult.Add(Params.OwnerContext, CameraRig);
 		}
 	}
 }
@@ -160,7 +193,7 @@ bool FStateTreeCameraDirectorEvaluator::SetContextRequirements(TSharedPtr<const
 			FStateTreeContextDataNames::ContextOwner, 
 			FStateTreeDataView(ContextOwner));
 
-	EvaluationData.ActiveCameraRigs.Reset();
+	EvaluationData.Reset();
 
 	StateTreeContext.SetCollectExternalDataCallback(FOnCollectStateTreeExternalData::CreateLambda(
 		[this](const FStateTreeExecutionContext& Context, const UStateTree* StateTree, 
@@ -180,6 +213,25 @@ bool FStateTreeCameraDirectorEvaluator::SetContextRequirements(TSharedPtr<const
 	return true;
 }
 
+const UCameraRigAsset* FStateTreeCameraDirectorEvaluator::FindCameraRigByProxy(const UCameraRigProxyAsset* InProxy)
+{
+	const UStateTreeCameraDirector* Director = GetCameraDirectorAs<UStateTreeCameraDirector>();
+	if (!ensure(Director))
+	{
+		return nullptr;
+	}
+
+	const UCameraRigProxyTable* ProxyTable = Director->CameraRigProxyTable;
+	if (!ensureMsgf(ProxyTable, TEXT("No proxy table set on StateTree director '%s'."), *Director->GetPathName()))
+	{
+		return nullptr;
+	}
+
+	FCameraRigProxyTableResolveParams ResolveParams;
+	ResolveParams.CameraRigProxy = InProxy;
+	return Director->CameraRigProxyTable->ResolveProxy(ResolveParams);
+}
+
 void FStateTreeCameraDirectorEvaluator::OnAddReferencedObjects(FReferenceCollector& Collector)
 {
 	StateTreeInstanceData.AddStructReferencedObjects(Collector);
@@ -200,12 +252,24 @@ FCameraDirectorEvaluatorPtr UStateTreeCameraDirector::OnBuildEvaluator(FCameraDi
 
 void UStateTreeCameraDirector::OnBuildCameraDirector(UE::Cameras::FCameraBuildLog& BuildLog)
 {
+	using namespace UE::Cameras;
+
+	// Check that a state tree was specified.
+	if (!StateTreeReference.IsValid())
+	{
+		BuildLog.AddMessage(EMessageSeverity::Error, LOCTEXT("MissingStateTree", "No state tree reference is set."));
+		return;
+	}
 }
 
 #if WITH_EDITOR
 
 void UStateTreeCameraDirector::OnFactoryCreateAsset(const FCameraDirectorFactoryCreateParams& InParams)
 {
+	if (!CameraRigProxyTable)
+	{
+		CameraRigProxyTable = NewObject<UCameraRigProxyTable>(this);
+	}
 }
 
 #endif
diff --git a/Engine/Plugins/Cameras/GameplayCameras/Source/GameplayCameras/Private/Directors/StateTreeCameraDirectorTasks.cpp b/Engine/Plugins/Cameras/GameplayCameras/Source/GameplayCameras/Private/Directors/StateTreeCameraDirectorTasks.cpp
--- a/Engine/Plugins/Cameras/GameplayCameras/Source/GameplayCameras/Private/Directors/StateTreeCameraDirectorTasks.cpp
+++ b/Engine/Plugins/Cameras/GameplayCameras/Source/GameplayCameras/Private/Directors/StateTreeCameraDirectorTasks.cpp
@@ -43,5 +43,39 @@ bool FGameplayCamerasActivateCameraRigTask::UpdateResult(FStateTreeExecutionCont
 	return false;
 }
 
+bool FGameplayCamerasActivateCameraRigViaProxyTask::Link(FStateTreeLinker& Linker)
+{
+	Linker.LinkExternalData(CameraDirectorEvaluationDataHandle);
+	return true;
+}
+
+EStateTreeRunStatus FGameplayCamerasActivateCameraRigViaProxyTask::EnterState(FStateTreeExecutionContext& Context, const FStateTreeTransitionResult& Transition) const
+{
+	if (UpdateResult(Context))
+	{
+		return bRunOnce ? EStateTreeRunStatus::Succeeded : EStateTreeRunStatus::Running;
+	}
+	return EStateTreeRunStatus::Failed;
+}
+
+EStateTreeRunStatus FGameplayCamerasActivateCameraRigViaProxyTask::Tick(FStateTreeExecutionContext& Context, const float DeltaTime) const
+{
+	UpdateResult(Context);
+
+	return EStateTreeRunStatus::Running;
+}
+
+bool FGameplayCamerasActivateCameraRigViaProxyTask::UpdateResult(FStateTreeExecutionContext& Context) const
+{
+	FInstanceDataType& InstanceData = Context.GetInstanceData(*this);
+	if (InstanceData.CameraRigProxy)
+	{
+		FCameraDirectorStateTreeEvaluationData& EvaluationData = Context.GetExternalData(CameraDirectorEvaluationDataHandle);
+		EvaluationData.ActiveCameraRigProxies.Add(InstanceData.CameraRigProxy);
+		return true;
+	}
+	return false;
+}
+
 #undef LOCTEXT_NAMESPACE
 
diff --git a/Engine/Plugins/Cameras/GameplayCameras/Source/GameplayCameras/Public/Directors/CameraDirectorStateTreeSchema.h b/Engine/Plugins/Cameras/GameplayCameras/Source/GameplayCameras/Public/Directors/CameraDirectorStateTreeSchema.h
--- a/Engine/Plugins/Cameras/GameplayCameras/Source/GameplayCameras/Public/Directors/CameraDirectorStateTreeSchema.h
+++ b/Engine/Plugins/Cameras/GameplayCameras/Source/GameplayCameras/Public/Directors/CameraDirectorStateTreeSchema.h
@@ -13,6 +13,7 @@
 #include "CameraDirectorStateTreeSchema.generated.h"
 
 class UCameraRigAsset;
+class UCameraRigProxyAsset;
 struct FStateTreeExternalDataDesc;
 
 namespace UE::Cameras
@@ -60,6 +61,15 @@ struct GAMEPLAYCAMERAS_API FCameraDirectorStateTreeEvaluationData
 	/** Camera rigs activated during a StateTree's execution frame. */
 	UPROPERTY()
 	TArray<TObjectPtr<UCameraRigAsset>> ActiveCameraRigs;
+
+	/** Camera rig proxies activated during a StateTree's execution frame. */
+	UPROPERTY()
+	TArray<TObjectPtr<UCameraRigProxyAsset>> ActiveCameraRigProxies;
+
+public:
+
+	/** Reset this evaluation data for a new frame. */
+	void Reset();
 };
 
 /** Base classs for camera director StateTree tasks. */
diff --git a/Engine/Plugins/Cameras/GameplayCameras/Source/GameplayCameras/Public/Directors/StateTreeCameraDirector.h b/Engine/Plugins/Cameras/GameplayCameras/Source/GameplayCameras/Public/Directors/StateTreeCameraDirector.h
--- a/Engine/Plugins/Cameras/GameplayCameras/Source/GameplayCameras/Public/Directors/StateTreeCameraDirector.h
+++ b/Engine/Plugins/Cameras/GameplayCameras/Source/GameplayCameras/Public/Directors/StateTreeCameraDirector.h
@@ -7,6 +7,9 @@
 
 #include "StateTreeCameraDirector.generated.h"
 
+class UCameraRigProxyAsset;
+class UCameraRigProxyTable;
+
 /**
  * A camera director that runs a StateTree to specify which camera rigs should be active
  * any given frame.
@@ -35,5 +38,12 @@ public:
 	UPROPERTY(EditAnywhere, Category="StateTree",
 			meta=(Schema="/Script/GameplayCameras.CameraDirectorStateTreeSchema"))
 	FStateTreeReference StateTreeReference;
+
+	/** 
+	 * The table that maps camera rig proxies (used in the evaluator State Tree's tasks)
+	 * to actual camera rigs.
+	 */
+	UPROPERTY(EditAnywhere, Instanced, Category="Evaluation")
+	TObjectPtr<UCameraRigProxyTable> CameraRigProxyTable;
 };
 
diff --git a/Engine/Plugins/Cameras/GameplayCameras/Source/GameplayCameras/Public/Directors/StateTreeCameraDirectorTasks.h b/Engine/Plugins/Cameras/GameplayCameras/Source/GameplayCameras/Public/Directors/StateTreeCameraDirectorTasks.h
--- a/Engine/Plugins/Cameras/GameplayCameras/Source/GameplayCameras/Public/Directors/StateTreeCameraDirectorTasks.h
+++ b/Engine/Plugins/Cameras/GameplayCameras/Source/GameplayCameras/Public/Directors/StateTreeCameraDirectorTasks.h
@@ -51,3 +51,49 @@ public:
 	bool bRunOnce = false;
 };
 
+/** Instance data for the "Activate Camera Rig via Proxy" task. */
+USTRUCT()
+struct FGameplayCamerasActivateCameraRigViaProxyTaskInstanceData
+{
+	GENERATED_BODY()
+
+	/** The camera rig proxy to activate. */
+	UPROPERTY(EditAnywhere, Category="Cameras")
+	TObjectPtr<UCameraRigProxyAsset> CameraRigProxy;
+};
+
+/**
+ * A task that activates a camera rig via the given camera rig proxy inside a
+ * StateTreeCameraDirector. The given proxy must be mapped to an actual proxy
+ * in each camera asset using this StateTree.
+ */
+USTRUCT(meta=(DisplayName="Activate Camera via Proxy", Category="Cameras"))
+struct FGameplayCamerasActivateCameraRigViaProxyTask : public FGameplayCamerasStateTreeTask
+{
+	GENERATED_BODY()
+
+	using FInstanceDataType = FGameplayCamerasActivateCameraRigViaProxyTaskInstanceData;
+
+protected:
+
+	// FStateTreeTaskBase interface.
+	virtual const UStruct* GetInstanceDataType() const override { return FInstanceDataType::StaticStruct(); }
+	virtual bool Link(FStateTreeLinker& Linker) override;
+	virtual EStateTreeRunStatus EnterState(FStateTreeExecutionContext& Context, const FStateTreeTransitionResult& Transition) const override;
+	virtual EStateTreeRunStatus Tick(FStateTreeExecutionContext& Context, const float DeltaTime) const override;
+
+private:
+
+	bool UpdateResult(FStateTreeExecutionContext& Context) const;
+
+private:
+
+	TStateTreeExternalDataHandle<FCameraDirectorStateTreeEvaluationData> CameraDirectorEvaluationDataHandle;
+
+public:
+
+	/** If true, the task will complete immediately. If false, the task will run until a transition triggers. */
+	UPROPERTY(EditAnywhere, Category="State Tree")
+	bool bRunOnce = false;
+};
+
