Cameras: extract the delete camera variable dialog and make it re-usable

Also improve the dialog and fix various problems with undo/redo of deletions.

diff --git a/Engine/Plugins/Cameras/GameplayCameras/Source/GameplayCamerasEditor/Private/Toolkits/CameraVariableCollectionEditorToolkit.cpp b/Engine/Plugins/Cameras/GameplayCameras/Source/GameplayCamerasEditor/Private/Toolkits/CameraVariableCollectionEditorToolkit.cpp
--- a/Engine/Plugins/Cameras/GameplayCameras/Source/GameplayCamerasEditor/Private/Toolkits/CameraVariableCollectionEditorToolkit.cpp
+++ b/Engine/Plugins/Cameras/GameplayCameras/Source/GameplayCamerasEditor/Private/Toolkits/CameraVariableCollectionEditorToolkit.cpp
@@ -2,7 +2,6 @@
 
 #include "Toolkits/CameraVariableCollectionEditorToolkit.h"
 
-#include "AssetRegistry/AssetRegistryModule.h"
 #include "AssetTools/CameraVariableCollectionEditor.h"
 #include "Commands/CameraVariableCollectionEditorCommands.h"
 #include "ContentBrowserModule.h"
@@ -22,7 +21,7 @@
 #include "ToolMenus.h"
 #include "UObject/UObjectIterator.h"
 #include "Widgets/Docking/SDockTab.h"
-#include "Widgets/Input/SButton.h"
+#include "Widgets/SDeleteCameraObjectDialog.h"
 
 #include UE_INLINE_GENERATED_CPP_BY_NAME(CameraVariableCollectionEditorToolkit)
 
@@ -34,232 +33,6 @@ namespace UE::Cameras
 const FName FCameraVariableCollectionEditorToolkit::VariableCollectionEditorTabId(TEXT("CameraVariableCollectionEditor_VariableCollectionEditor"));
 const FName FCameraVariableCollectionEditorToolkit::DetailsViewTabId(TEXT("CameraVariableCollectionEditor_DetailsView"));
 
-class SDeleteVariableDialog : public SCompoundWidget
-{
-public:
-
-	SLATE_BEGIN_ARGS(SDeleteVariableDialog)
-	{}
-		SLATE_ARGUMENT(TWeakPtr<SWindow>, ParentWindow)
-		SLATE_ARGUMENT(TSet<FName>, ReferencingPackages)
-	SLATE_END_ARGS()
-
-	void Construct(const FArguments& InArgs)
-	{
-		WeakParentWindow = InArgs._ParentWindow;
-
-		ReferencingPackages = InArgs._ReferencingPackages;
-
-		ChildSlot
-		[
-			SNew(SVerticalBox)
-
-			+ SVerticalBox::Slot()
-			.AutoHeight()
-			.Padding(5.0f)
-			[
-				SNew(SBorder)
-				.BorderBackgroundColor(FLinearColor::Green)
-				.BorderImage(FAppStyle::GetBrush("ToolPanel.GroupBorder"))
-				.Visibility(this, &SDeleteVariableDialog::GetNoReferencesVisibility)
-				.Padding(5.0f)
-				[
-					SNew(STextBlock)
-					.Text(LOCTEXT(
-								"VariablesOkToDelete", 
-								"No assets reference the variables being deleted."))
-				]
-			]
-
-			+ SVerticalBox::Slot()
-			.AutoHeight()
-			.Padding(5.0f)
-			[
-				SNew(SBorder)
-				.BorderBackgroundColor(FLinearColor::Red)
-				.BorderImage(FAppStyle::GetBrush("ToolPanel.GroupBorder"))
-				.Visibility(this, &SDeleteVariableDialog::GetReferencesVisiblity)
-				.Padding(5.0f)
-				[
-					SNew(STextBlock)
-					.Text(LOCTEXT(
-								"VariablesPendingDeleteAreInUse", 
-								"Some of the camera variables being deleted are referenced by camera assets."))
-				]
-			]
-
-			+ SVerticalBox::Slot()
-			.FillHeight(1.0f)
-			.Padding(5.0f)
-			[
-				SNew(SBorder)
-				.BorderImage(FAppStyle::GetBrush("ToolPanel.GroupBorder"))
-				.Padding(FMargin(0, 0, 0, 3))
-				.Visibility(this, &SDeleteVariableDialog::GetReferencesVisiblity)
-				[
-					SNew(SVerticalBox)
-
-					+ SVerticalBox::Slot()
-					.AutoHeight()
-					[
-						SNew(SBorder)
-						.BorderImage(FAppStyle::GetBrush("DetailsView.CategoryTop"))
-						.BorderBackgroundColor(FLinearColor(.6, .6, .6, 1.0f))
-						.Padding(3.0f)
-						[
-							SNew(STextBlock)
-							.Text(LOCTEXT(
-										"AssetsReferencingVariablesPendingDelete", 
-										"Assets Referencing the Camera Variables to Delete"))
-							.Font(FAppStyle::GetFontStyle("BoldFont"))
-							.ShadowOffset(FVector2D(1.0f, 1.0f))
-						]
-					]
-
-					+ SVerticalBox::Slot()
-					.FillHeight(1.0f)
-					[
-						SDeleteVariableDialog::BuildReferencerAssetPicker()
-					]
-				]
-			]
-
-			+ SVerticalBox::Slot()
-			.AutoHeight()
-			.Padding(0.f, 4.f)
-			[
-				SNew(SHorizontalBox)
-
-				+ SHorizontalBox::Slot()
-				.FillWidth(1.0f)
-				.Padding(6, 0)
-				[
-					SNew(SBorder)
-					.BorderImage(FAppStyle::GetBrush("NoBorder"))
-					[
-						SNew(SButton)
-						.HAlign(HAlign_Center)
-						.Text(LOCTEXT("Delete", "Delete"))
-						.ToolTipText(LOCTEXT("DeleteTooltipText", "Perform the delete"))
-						.ButtonStyle(FAppStyle::Get(), "FlatButton.Danger")
-						.TextStyle(FAppStyle::Get(), "FlatButton.DefaultTextStyle")
-						.OnClicked(this, &SDeleteVariableDialog::OnDeleteClicked)
-					]
-				]
-
-				+ SHorizontalBox::Slot()
-				.FillWidth(1.0f)
-				.Padding(6, 0)
-				[
-					SNew(SBorder)
-					.BorderImage(FAppStyle::GetBrush("NoBorder"))
-					[
-						SNew(SButton)
-						.HAlign(HAlign_Center)
-						.Text(LOCTEXT("Cancel", "Cancel"))
-						.ToolTipText(LOCTEXT("CancelDeleteTooltipText", "Cancel the delete"))
-						.ButtonStyle(FAppStyle::Get(), "FlatButton.Default")
-						.TextStyle(FAppStyle::Get(), "FlatButton.DefaultTextStyle")
-						.OnClicked(this, &SDeleteVariableDialog::OnCancelClicked)
-					]
-				]
-			]
-		];
-	}
-
-	bool ShouldPerformDelete() const
-	{
-		return bPerformDelete;
-	}
-	
-protected:
-
-	virtual FReply OnKeyDown(const FGeometry& MyGeometry, const FKeyEvent& InKeyEvent) override
-	{
-		if (InKeyEvent.GetKey() == EKeys::Escape)
-		{
-			OnCancelClicked();
-			return FReply::Handled();
-		}
-		return FReply::Unhandled();
-	}
-
-private:
-
-	EVisibility GetNoReferencesVisibility() const
-	{
-		return ReferencingPackages.IsEmpty() ? EVisibility::Visible : EVisibility::Collapsed;
-	}
-
-	EVisibility GetReferencesVisiblity() const
-	{
-		return ReferencingPackages.IsEmpty() ? EVisibility::Collapsed : EVisibility::Visible;
-	}
-
-	TSharedRef<SWidget> BuildReferencerAssetPicker()
-	{
-		FContentBrowserModule& ContentBrowserModule = FModuleManager::Get().LoadModuleChecked<FContentBrowserModule>(TEXT("ContentBrowser"));
-
-		FARFilter ARFilter;
-		ARFilter.PackageNames = ReferencingPackages.Array();
-
-		FAssetPickerConfig AssetPickerConfig;
-		AssetPickerConfig.bAllowDragging = false;
-		AssetPickerConfig.bCanShowClasses = false;
-		AssetPickerConfig.bAllowNullSelection = false;
-		AssetPickerConfig.bShowBottomToolbar = false;
-		AssetPickerConfig.bAutohideSearchBar = true;
-		AssetPickerConfig.Filter = ARFilter;
-		AssetPickerConfig.InitialAssetViewType = EAssetViewType::Tile;
-		AssetPickerConfig.OnAssetsActivated = FOnAssetsActivated::CreateSP(this, &SDeleteVariableDialog::OnAssetsActivated);
-
-		return ContentBrowserModule.Get().CreateAssetPicker(AssetPickerConfig);
-	}
-
-	void OnAssetsActivated(const TArray<FAssetData>& ActivatedAssets, EAssetTypeActivationMethod::Type ActivationMethod)
-	{
-		if (ActivationMethod == EAssetTypeActivationMethod::DoubleClicked || ActivationMethod == EAssetTypeActivationMethod::Opened)
-		{
-			CloseWindow();
-
-			for (const FAssetData& ActivatedAsset : ActivatedAssets)
-			{
-				GEditor->GetEditorSubsystem<UAssetEditorSubsystem>()->OpenEditorForAsset(ActivatedAsset.GetAsset());
-			}
-		}
-	}
-
-	void CloseWindow()
-	{
-		if (TSharedPtr<SWindow> ParentWindow = WeakParentWindow.Pin())
-		{
-			ParentWindow->RequestDestroyWindow();
-		}
-	}
-
-	FReply OnDeleteClicked()
-	{
-		bPerformDelete = true;
-		CloseWindow();
-		return FReply::Handled();
-	}
-
-	FReply OnCancelClicked()
-	{
-		bPerformDelete = false;
-		CloseWindow();
-		return FReply::Handled();
-	}
-
-private:
-
-	TWeakPtr<SWindow> WeakParentWindow;
-
-	TSet<FName> ReferencingPackages;
-
-	bool bPerformDelete = false;
-};
-
 FCameraVariableCollectionEditorToolkit::FCameraVariableCollectionEditorToolkit(UCameraVariableCollectionEditor* InOwningAssetEditor)
 	: FBaseAssetToolkit(InOwningAssetEditor)
 	, VariableCollection(InOwningAssetEditor->GetVariableCollection())
@@ -285,10 +58,19 @@ FCameraVariableCollectionEditorToolkit::FCameraVariableCollectionEditorToolkit(U
 				->AddTab(DetailsViewTabId, ETabState::OpenedTab)
 			)
 		);
+
+	if (GEditor)
+	{
+		GEditor->RegisterForUndo(this);
+	}
 }
 
 FCameraVariableCollectionEditorToolkit::~FCameraVariableCollectionEditorToolkit()
 {
+	if (GEditor)
+	{
+		GEditor->UnregisterForUndo(this);
+	}
 }
 
 void FCameraVariableCollectionEditorToolkit::RegisterTabSpawners(const TSharedRef<class FTabManager>& InTabManager)
@@ -508,80 +290,38 @@ void FCameraVariableCollectionEditorToolkit::OnDeleteVariable()
 		return;
 	}
 
-	GWarn->BeginSlowTask(LOCTEXT("PreDeleteScanning", "Scanning assets before deleting camera variables"), true);
-
-	UPackage* VariableCollectionPackage = VariableCollection->GetOutermost();
-	const FName VariableCollectionPackageName = VariableCollectionPackage->GetFName();
-
-	TSet<FName> AllReferencers;
-
-	{
-		FAssetRegistryModule& AssetRegistryModule = FModuleManager::LoadModuleChecked<FAssetRegistryModule>("AssetRegistry");
-		IAssetRegistry& AssetRegistry = AssetRegistryModule.Get();
-
-		TArray<FName> OnDiskReferencers;
-		AssetRegistry.GetReferencers(VariableCollectionPackageName, OnDiskReferencers);
-		AllReferencers.Append(OnDiskReferencers);
-	}
-
-	for (UCameraVariableAsset* CameraVariable : Selection)
-	{
-		bool bIsReferencedInMemoryByNonUndo = false;
-		bool bIsReferencedInMemoryByUndo = false;
-		FReferencerInformationList MemoryReferences;
-		ObjectTools::GatherObjectReferencersForDeletion(
-				CameraVariable, bIsReferencedInMemoryByNonUndo, bIsReferencedInMemoryByUndo, &MemoryReferences);
-
-		UPackage* TransientPackage = GetTransientPackage();
-		for (const FReferencerInformation& ExternalReference : MemoryReferences.ExternalReferences)
-		{
-			UPackage* ExternalReferencePackage = ExternalReference.Referencer->GetOutermost();
-			if (ExternalReferencePackage != VariableCollectionPackage && ExternalReferencePackage != TransientPackage)
-			{
-				AllReferencers.Add(ExternalReferencePackage->GetFName());
-			}
-		}
-	}
-
-	GWarn->EndSlowTask();
-
-	bool bPerformDelete = false;
-	{
-		TSharedRef<SWindow> DeleteVariableWindow = SNew(SWindow)
-			.Title(LOCTEXT("DeleteVariablesWindowTitle", "Delete Variables"))
-			.ClientSize(FVector2D(600, 700));
-
-		TSharedRef<SDeleteVariableDialog> DeleteVariableDialog = SNew(SDeleteVariableDialog)
-			.ParentWindow(DeleteVariableWindow)
-			.ReferencingPackages(AllReferencers);
-		DeleteVariableWindow->SetContent(DeleteVariableDialog);
-
-		GEditor->EditorAddModalWindow(DeleteVariableWindow);
-
-		bPerformDelete = DeleteVariableDialog->ShouldPerformDelete();
-	}
-
+	TSharedRef<SWindow> DeleteVariableWindow = SNew(SWindow)
+		.Title(LOCTEXT("DeleteVariablesWindowTitle", "Delete Variables"))
+		.ClientSize(FVector2D(600, 700));
+
+	TSharedRef<SDeleteCameraObjectDialog> DeleteVariableDialog = SNew(SDeleteCameraObjectDialog)
+		.ParentWindow(DeleteVariableWindow)
+		.ObjectsToDelete(TArray<UObject*>(Selection))
+		.OnDeletedObject_Lambda([](UObject* Obj)
+					{
+						if (UCameraVariableAsset* TrashVariable = Cast<UCameraVariableAsset>(Obj))
+						{
+							SDeleteCameraObjectDialog::RenameObjectAsTrash(TrashVariable->DisplayName);
+						}
+					});
+	DeleteVariableWindow->SetContent(DeleteVariableDialog);
+
+	GEditor->EditorAddModalWindow(DeleteVariableWindow);
+
+	const bool bPerformDelete = DeleteVariableDialog->ShouldPerformDelete();
 	if (bPerformDelete)
 	{
 		FScopedTransaction DeleteTransaction(LOCTEXT("DeleteVariable", "Delete camera variable"));
 
 		VariableCollection->Modify();
-
-		TArray<UObject*> ObjectsToReplace(Selection);
-		ObjectTools::ForceReplaceReferences(nullptr, ObjectsToReplace);
-
-		TStringBuilder<256> StringBuilder;
 		for (UCameraVariableAsset* VariableToDelete : Selection)
 		{
+			VariableToDelete->Modify();
 			VariableCollection->Variables.Remove(VariableToDelete);
-
-			StringBuilder.Reset();
-			StringBuilder.Append("TRASH_");
-			StringBuilder.Append(VariableToDelete->GetName());
-			VariableToDelete->Rename(StringBuilder.ToString());
-			VariableToDelete->MarkAsGarbage();
 		}
 
+		DeleteVariableDialog->PerformReferenceReplacement();
+
 		VariableCollectionEditorWidget->RequestListRefresh();
 	}
 }
@@ -593,6 +333,26 @@ bool FCameraVariableCollectionEditorToolkit::CanDeleteVariable()
 	return !Selection.IsEmpty();
 }
 
+void FCameraVariableCollectionEditorToolkit::PostUndo(bool bSuccess)
+{
+	VariableCollectionEditorWidget->RequestListRefresh();
+}
+
+void FCameraVariableCollectionEditorToolkit::PostRedo(bool bSuccess)
+{
+	VariableCollectionEditorWidget->RequestListRefresh();
+}
+
+void FCameraVariableCollectionEditorToolkit::AddReferencedObjects(FReferenceCollector& Collector)
+{
+	Collector.AddReferencedObject(VariableCollection);
+}
+
+FString FCameraVariableCollectionEditorToolkit::GetReferencerName() const
+{
+	return TEXT("FCameraVariableCollectionEditorToolkit");
+}
+
 }  // namespace UE::Cameras
 
 #undef LOCTEXT_NAMESPACE
diff --git a/Engine/Plugins/Cameras/GameplayCameras/Source/GameplayCamerasEditor/Private/Widgets/SDeleteCameraObjectDialog.cpp b/Engine/Plugins/Cameras/GameplayCameras/Source/GameplayCamerasEditor/Private/Widgets/SDeleteCameraObjectDialog.cpp
new file mode 100644
--- /dev/null
+++ b/Engine/Plugins/Cameras/GameplayCameras/Source/GameplayCamerasEditor/Private/Widgets/SDeleteCameraObjectDialog.cpp
@@ -0,0 +1,524 @@
+// Copyright Epic Games, Inc. All Rights Reserved.
+
+#include "Widgets/SDeleteCameraObjectDialog.h"
+
+#include "AssetRegistry/AssetRegistryModule.h"
+#include "ContentBrowserModule.h"
+#include "Editor.h"
+#include "FileHelpers.h"
+#include "IContentBrowserSingleton.h"
+#include "ObjectTools.h"
+#include "Serialization/FindReferencersArchive.h"
+#include "Subsystems/AssetEditorSubsystem.h"
+#include "Widgets/Input/SButton.h"
+#include "Widgets/Layout/SBorder.h"
+#include "Widgets/Notifications/SNotificationList.h"
+#include "Widgets/Notifications/SProgressBar.h"
+#include "Widgets/Text/STextBlock.h"
+
+#define LOCTEXT_NAMESPACE "SDeleteCameraObjectDialog"
+
+namespace UE::Cameras
+{
+
+void SDeleteCameraObjectDialog::Construct(const FArguments& InArgs)
+{
+	WeakParentWindow = InArgs._ParentWindow;
+
+	ObjectsToDelete = InArgs._ObjectsToDelete;
+	ObjectsToDelete.Remove(nullptr);
+
+	State = (ObjectsToDelete.Num() > 0) ? EState::StartScanning : EState::Waiting;
+
+	bRenameObjectsAsTrash = InArgs._RenameObjectsAsTrash.Get();
+
+	ChildSlot
+	[
+		SNew(SVerticalBox)
+
+		+ SVerticalBox::Slot()
+		.AutoHeight()
+		.Padding(5.0f)
+		[
+			SNew(SBorder)
+			.BorderBackgroundColor(FLinearColor::Green)
+			.BorderImage(FAppStyle::GetBrush("ToolPanel.GroupBorder"))
+			.Visibility(this, &SDeleteCameraObjectDialog::GetNoReferencesVisibility)
+			.Padding(5.0f)
+			[
+				SNew(STextBlock)
+				.Text(LOCTEXT(
+							"ObjectsOkToDelete", 
+							"No assets reference the objects being deleted."))
+			]
+		]
+
+		+ SVerticalBox::Slot()
+		.AutoHeight()
+		.Padding(5.0f)
+		[
+			SNew(SBorder)
+			.BorderBackgroundColor(FLinearColor::Red)
+			.BorderImage(FAppStyle::GetBrush("ToolPanel.GroupBorder"))
+			.Visibility(this, &SDeleteCameraObjectDialog::GetReferencesVisiblity)
+			.Padding(5.0f)
+			[
+				SNew(STextBlock)
+				.Text(LOCTEXT(
+							"ObjectsPendingDeleteAreInUse", 
+							"Some of the objects being deleted are referenced by other assets."))
+			]
+		]
+
+		+ SVerticalBox::Slot()
+		.FillHeight(1.0f)
+		.Padding(5.0f)
+		[
+			SNew(SBorder)
+			.BorderImage(FAppStyle::GetBrush("ToolPanel.GroupBorder"))
+			.Padding(FMargin(0, 0, 0, 3))
+			.Visibility(this, &SDeleteCameraObjectDialog::GetReferencesVisiblity)
+			[
+				SNew(SVerticalBox)
+
+				+ SVerticalBox::Slot()
+				.AutoHeight()
+				[
+					SNew(SBorder)
+					.BorderImage(FAppStyle::GetBrush("DetailsView.CategoryTop"))
+					.BorderBackgroundColor(FLinearColor(.6, .6, .6, 1.0f))
+					.Padding(3.0f)
+					[
+						SNew(STextBlock)
+						.Text(LOCTEXT(
+									"AssetsReferencingObjectsPendingDelete", 
+									"Assets Referencing the Objects to Delete"))
+						.Font(FAppStyle::GetFontStyle("BoldFont"))
+						.ShadowOffset(FVector2D(1.0f, 1.0f))
+					]
+				]
+
+				+ SVerticalBox::Slot()
+				.FillHeight(1.0f)
+				[
+					SDeleteCameraObjectDialog::BuildReferencerAssetPicker()
+				]
+			]
+		]
+
+		+ SVerticalBox::Slot()
+		.AutoHeight()
+		.Padding(0.f, 4.f)
+		[
+			SNew(SHorizontalBox)
+
+			+ SHorizontalBox::Slot()
+			.FillWidth(1.0f)
+			.Padding(6, 0)
+			[
+				SNew(SBorder)
+				.BorderImage(FAppStyle::GetBrush("NoBorder"))
+				[
+					SNew(SButton)
+					.HAlign(HAlign_Center)
+					.Text(LOCTEXT("Delete", "Delete"))
+					.ToolTipText(LOCTEXT("DeleteTooltipText", "Perform the delete"))
+					.ButtonStyle(FAppStyle::Get(), "FlatButton.Danger")
+					.TextStyle(FAppStyle::Get(), "FlatButton.DefaultTextStyle")
+					.IsEnabled(this, &SDeleteCameraObjectDialog::IsDeleteEnabled)
+					.OnClicked(this, &SDeleteCameraObjectDialog::OnDeleteClicked)
+				]
+			]
+
+			+ SHorizontalBox::Slot()
+			.FillWidth(1.0f)
+			.Padding(6, 0)
+			[
+				SNew(SBorder)
+				.BorderImage(FAppStyle::GetBrush("NoBorder"))
+				[
+					SNew(SButton)
+					.HAlign(HAlign_Center)
+					.Text(LOCTEXT("Cancel", "Cancel"))
+					.ToolTipText(LOCTEXT("CancelDeleteTooltipText", "Cancel the delete"))
+					.ButtonStyle(FAppStyle::Get(), "FlatButton.Default")
+					.TextStyle(FAppStyle::Get(), "FlatButton.DefaultTextStyle")
+					.OnClicked(this, &SDeleteCameraObjectDialog::OnCancelClicked)
+				]
+			]
+		]
+
+		+SVerticalBox::Slot()
+		.AutoHeight()
+		.Padding(5.0f)
+		[
+			SNew(SProgressBar)
+			.Visibility(this, &SDeleteCameraObjectDialog::GetProgressBarVisibility)
+			.Percent(this, &SDeleteCameraObjectDialog::GetProgressBarPercent)
+		]
+	];
+}
+
+bool SDeleteCameraObjectDialog::ShouldPerformDelete() const
+{
+	return bPerformDelete;
+}
+
+void SDeleteCameraObjectDialog::PerformReferenceReplacement()
+{
+	// Replace references to the objects we want to delete.
+	//
+	// We need to specify ObjectsToReplaceWithin otherwise the replacement also occurs inside the undo buffer!
+	// Note that we also need to pass the assets, not the packages, in that list, otherwise ForceReplaceReferences
+	// fails to recurse into the package (that's because it's using object references, and packages don't really
+	// reference their asset directly).
+	TArray<UPackage*> DirtiedPackages;
+	TSet<UObject*> ObjectsToReplaceWithin;
+	for (const FName PackageName : ReferencingPackages)
+	{
+		UPackage* Package = FindPackage(nullptr, *PackageName.ToString());
+		if (ensure(Package))
+		{
+			DirtiedPackages.Add(Package);
+
+			UObject* PackageAsset = Package->FindAssetInPackage();
+			if (ensure(PackageAsset))
+			{
+				ObjectsToReplaceWithin.Add(PackageAsset);
+			}
+		}
+	}
+	for (UObject* Object : ObjectsToDelete)
+	{
+		UPackage* Package = Object->GetOutermost();
+		if (ensure(Package))
+		{
+			DirtiedPackages.Add(Package);
+
+			UObject* PackageAsset = Package->FindAssetInPackage();
+			if (ensure(PackageAsset))
+			{
+				ObjectsToReplaceWithin.Add(PackageAsset);
+			}
+		}
+	}
+	ensure(ObjectsToReplaceWithin.Num() > 0);
+	ObjectTools::ForceReplaceReferences(nullptr, ObjectsToDelete, ObjectsToReplaceWithin);
+
+	// Prompt for checking out and saving changed packages.
+	if (DirtiedPackages.Num() > 0)
+	{
+		FEditorFileUtils::FPromptForCheckoutAndSaveParams SaveParams;
+		SaveParams.bCheckDirty = false;
+		SaveParams.bPromptToSave = true;
+		SaveParams.bCanBeDeclined = false;
+		SaveParams.bIsExplicitSave = true;
+
+		FEditorFileUtils::PromptForCheckoutAndSave(ObjectPtrDecay(DirtiedPackages), SaveParams);
+	}
+
+	// Optionally rename objects with a TRASH prefix, helpful for debugging sometimes.
+	if (bRenameObjectsAsTrash)
+	{
+		TStringBuilder<256> StringBuilder;
+		for (UObject* Object : ObjectsToDelete)
+		{
+			StringBuilder.Reset();
+			StringBuilder.Append("TRASH_");
+			StringBuilder.Append(Object->GetName());
+			Object->Rename(StringBuilder.ToString());
+		}
+	}
+
+	// Call custom callback.
+	if (OnDeletedObject.IsBound())
+	{
+		for (UObject* Object : ObjectsToDelete)
+		{
+			OnDeletedObject.Execute(Object);
+		}
+	}
+}
+
+void SDeleteCameraObjectDialog::Tick(const FGeometry& AllottedGeometry, const double InCurrentTime, const float InDeltaTime)
+{
+	static const double MaxTickTime = 0.1;  // 100ms
+
+	FAssetRegistryModule& AssetRegistryModule = FModuleManager::LoadModuleChecked<FAssetRegistryModule>("AssetRegistry");
+	IAssetRegistry& AssetRegistry = AssetRegistryModule.Get();
+
+	const double StartTickTime = FPlatformTime::Seconds();
+	while (State != EState::Waiting && FPlatformTime::Seconds() - StartTickTime < MaxTickTime)
+	{
+		switch (State)
+		{
+			case EState::Waiting:
+				// We should never get here but some compiler emit warnings for missing case statements.
+				ensure(false);
+				break;
+			case EState::StartScanning:
+				StartScanning();
+				break;
+			case EState::ScanNextObject:
+				ScanNextObject(AssetRegistry);
+				break;
+			case EState::ScanNextReferencingPackage:
+				ScanNextReferencingPackage();
+				break;
+			case EState::FinishScanning:
+				FinishScanning();
+				break;
+		}
+	}
+
+	SCompoundWidget::Tick(AllottedGeometry, InCurrentTime, InDeltaTime);
+}
+
+FReply SDeleteCameraObjectDialog::OnKeyDown(const FGeometry& MyGeometry, const FKeyEvent& InKeyEvent)
+{
+	if (InKeyEvent.GetKey() == EKeys::Escape)
+	{
+		OnCancelClicked();
+		return FReply::Handled();
+	}
+
+	return SCompoundWidget::OnKeyDown(MyGeometry, InKeyEvent);
+}
+
+void SDeleteCameraObjectDialog::StartScanning()
+{
+	ReferencingPackages.Reset();
+	bIsAnyReferencedInMemoryByNonUndo = false;
+	bIsAnyReferencedInMemoryByUndo = false;
+
+	PossiblyReferencingPackages.Reset();
+	NextObjectToScan = 0;
+	NextPackageToScan = 0;
+
+	State = ObjectsToDelete.Num() > 0 ? EState::ScanNextObject : EState::FinishScanning;
+}
+
+void SDeleteCameraObjectDialog::ScanNextObject(IAssetRegistry& AssetRegistry)
+{
+	if (!ensure(ObjectsToDelete.IsValidIndex(NextObjectToScan)))
+	{
+		State = EState::ScanNextReferencingPackage;
+		return;
+	}
+
+	UObject* ObjectToScan = ObjectsToDelete[NextObjectToScan];
+	NextObjectToScan++;
+	if (NextObjectToScan >= ObjectsToDelete.Num())
+	{
+		State = EState::ScanNextReferencingPackage;
+	}
+
+	if (!ensure(ObjectToScan))
+	{
+		return;
+	}
+
+	UPackage* Package = ObjectToScan->GetOutermost();
+
+	// Check on-disk references to the object's outer package.
+	// We will have to later scan them properly to check if they are actually referencing the exact
+	// object we are trying to delete.
+	if (Package)
+	{
+		TArray<FName> OnDiskReferencers;
+		AssetRegistry.GetReferencers(Package->GetFName(), OnDiskReferencers);
+		PossiblyReferencingPackages.Append(OnDiskReferencers);
+	}
+
+	// Check in-memory references to the object.
+	{
+		bool bIsReferencedInMemoryByNonUndo = false;
+		bool bIsReferencedInMemoryByUndo = false;
+		FReferencerInformationList MemoryReferences;
+		ObjectTools::GatherObjectReferencersForDeletion(
+				ObjectToScan, bIsReferencedInMemoryByNonUndo, bIsReferencedInMemoryByUndo, &MemoryReferences);
+
+		UPackage* TransientPackage = GetTransientPackage();
+		for (const FReferencerInformation& ExternalReference : MemoryReferences.ExternalReferences)
+		{
+			UPackage* ExternalReferencePackage = ExternalReference.Referencer->GetOutermost();
+			if (ExternalReferencePackage != Package && ExternalReferencePackage != TransientPackage)
+			{
+				ReferencingPackages.Add(ExternalReferencePackage->GetFName());
+			}
+		}
+
+		bIsAnyReferencedInMemoryByNonUndo |= bIsReferencedInMemoryByNonUndo;
+		bIsAnyReferencedInMemoryByUndo |= bIsReferencedInMemoryByUndo;
+	}
+}
+
+void SDeleteCameraObjectDialog::ScanNextReferencingPackage()
+{
+	if (ReferencingPackagesToScan.IsEmpty() && !PossiblyReferencingPackages.IsEmpty())
+	{
+		ReferencingPackagesToScan = PossiblyReferencingPackages.Array();
+	}
+
+	if (!ensure(ReferencingPackagesToScan.IsValidIndex(NextPackageToScan)))
+	{
+		State = EState::FinishScanning;
+		return;
+	}
+
+	const FName ReferencingPackageName = ReferencingPackagesToScan[NextPackageToScan];
+	NextPackageToScan++;
+	if (NextPackageToScan >= ReferencingPackagesToScan.Num())
+	{
+		State = EState::FinishScanning;
+	}
+
+	UPackage* Package = FindPackage(nullptr, *ReferencingPackageName.ToString());
+	if (!ensure(Package))
+	{
+		return;
+	}
+
+	UObject* PackageAsset = Package->FindAssetInPackage();
+	if (!ensure(PackageAsset))
+	{
+		return;
+	}
+
+	FFindReferencersArchive FindRefsArchive(PackageAsset, ObjectsToDelete);
+
+	TMap<UObject*, int32> NumReferencesMap;
+	TMultiMap<UObject*, FProperty*> ReferencingPropertiesMMap;
+	if (FindRefsArchive.GetReferenceCounts(NumReferencesMap, ReferencingPropertiesMMap) > 0)
+	{
+		// That package does reference one of the objects we want to delete. Keep it.
+		ReferencingPackages.Add(ReferencingPackageName);
+	}
+}
+
+void SDeleteCameraObjectDialog::FinishScanning()
+{
+	State = EState::Waiting;
+}
+
+TOptional<float> SDeleteCameraObjectDialog::GetProgressBarPercent() const
+{
+	switch (State)
+	{
+		case EState::Waiting:
+		default:
+			return TOptional<float>();
+		case EState::StartScanning:
+			return 0.f;
+		case EState::ScanNextObject:
+			return ((float)NextObjectToScan / (float)FMath::Max(1, ObjectsToDelete.Num())) * 0.5f;
+		case EState::ScanNextReferencingPackage:
+			return ((float)NextPackageToScan / (float)FMath::Max(1, ReferencingPackagesToScan.Num())) * 0.5f + 0.5f;
+		case EState::FinishScanning:
+			return 1.f;
+	}
+}
+
+EVisibility SDeleteCameraObjectDialog::GetNoReferencesVisibility() const
+{
+	return ReferencingPackages.IsEmpty() ? EVisibility::Visible : EVisibility::Collapsed;
+}
+
+EVisibility SDeleteCameraObjectDialog::GetReferencesVisiblity() const
+{
+	return ReferencingPackages.IsEmpty() ? EVisibility::Collapsed : EVisibility::Visible;
+}
+
+EVisibility SDeleteCameraObjectDialog::GetProgressBarVisibility() const
+{
+	return State == EState::Waiting ? EVisibility::Collapsed : EVisibility::Visible;
+}
+
+TSharedRef<SWidget> SDeleteCameraObjectDialog::BuildReferencerAssetPicker()
+{
+	FContentBrowserModule& ContentBrowserModule = FModuleManager::Get().LoadModuleChecked<FContentBrowserModule>(TEXT("ContentBrowser"));
+
+	FAssetPickerConfig AssetPickerConfig;
+	AssetPickerConfig.bAllowDragging = false;
+	AssetPickerConfig.bCanShowClasses = false;
+	AssetPickerConfig.bAllowNullSelection = false;
+	AssetPickerConfig.bShowBottomToolbar = false;
+	AssetPickerConfig.bAutohideSearchBar = true;
+	AssetPickerConfig.InitialAssetViewType = EAssetViewType::Tile;
+	AssetPickerConfig.OnShouldFilterAsset = FOnShouldFilterAsset::CreateSP(this, &SDeleteCameraObjectDialog::OnShouldFilterReferencerAsset);
+	AssetPickerConfig.OnAssetsActivated = FOnAssetsActivated::CreateSP(this, &SDeleteCameraObjectDialog::OnAssetsActivated);
+
+	return ContentBrowserModule.Get().CreateAssetPicker(AssetPickerConfig);
+}
+
+bool SDeleteCameraObjectDialog::OnShouldFilterReferencerAsset(const FAssetData& InAssetData) const
+{
+	if (ReferencingPackages.Contains(InAssetData.PackageName))
+	{
+		return false;
+	}
+	return true;
+}
+
+void SDeleteCameraObjectDialog::OnAssetsActivated(const TArray<FAssetData>& ActivatedAssets, EAssetTypeActivationMethod::Type ActivationMethod)
+{
+	if (ActivationMethod == EAssetTypeActivationMethod::DoubleClicked || ActivationMethod == EAssetTypeActivationMethod::Opened)
+	{
+		CloseWindow();
+
+		for (const FAssetData& ActivatedAsset : ActivatedAssets)
+		{
+			GEditor->GetEditorSubsystem<UAssetEditorSubsystem>()->OpenEditorForAsset(ActivatedAsset.GetAsset());
+		}
+	}
+}
+
+bool SDeleteCameraObjectDialog::IsDeleteEnabled() const
+{
+	return State == EState::Waiting;
+}
+
+FReply SDeleteCameraObjectDialog::OnDeleteClicked()
+{
+	bPerformDelete = true;
+	CloseWindow();
+	return FReply::Handled();
+}
+
+FReply SDeleteCameraObjectDialog::OnCancelClicked()
+{
+	bPerformDelete = false;
+	CloseWindow();
+	return FReply::Handled();
+}
+
+void SDeleteCameraObjectDialog::CloseWindow()
+{
+	if (TSharedPtr<SWindow> ParentWindow = WeakParentWindow.Pin())
+	{
+		ParentWindow->RequestDestroyWindow();
+	}
+}
+
+void SDeleteCameraObjectDialog::RenameObjectAsTrash(FString& InOutName)
+{
+	if (!InOutName.IsEmpty())
+	{
+		TStringBuilder<256> StringBuilder;
+		StringBuilder.Append(TEXT("TRASH_"));
+		StringBuilder.Append(InOutName);
+		InOutName = StringBuilder.ToString();
+	}
+}
+
+FString SDeleteCameraObjectDialog::MakeTrashName(const FString& InName)
+{
+	FString Result(InName);
+	RenameObjectAsTrash(Result);
+	return Result;
+}
+
+}  // namespace UE::Cameras
+
+#undef LOCTEXT_NAMESPACE
+
diff --git a/Engine/Plugins/Cameras/GameplayCameras/Source/GameplayCamerasEditor/Public/Toolkits/CameraVariableCollectionEditorToolkit.h b/Engine/Plugins/Cameras/GameplayCameras/Source/GameplayCamerasEditor/Public/Toolkits/CameraVariableCollectionEditorToolkit.h
--- a/Engine/Plugins/Cameras/GameplayCameras/Source/GameplayCamerasEditor/Public/Toolkits/CameraVariableCollectionEditorToolkit.h
+++ b/Engine/Plugins/Cameras/GameplayCameras/Source/GameplayCamerasEditor/Public/Toolkits/CameraVariableCollectionEditorToolkit.h
@@ -3,6 +3,7 @@
 #pragma once
 
 #include "CoreTypes.h"
+#include "EditorUndoClient.h"
 #include "Tools/BaseAssetToolkit.h"
 #include "UObject/GCObject.h"
 
@@ -22,6 +23,7 @@ class SCameraVariableCollectionEditor;
  */
 class FCameraVariableCollectionEditorToolkit
 	: public FBaseAssetToolkit
+	, public FEditorUndoClient
 	, public FGCObject
 {
 public:
@@ -45,15 +47,13 @@ protected:
 	virtual FString GetWorldCentricTabPrefix() const override;
 	virtual FLinearColor GetWorldCentricTabColorScale() const override;
 
+	// FEditorUndoClient interface
+	virtual void PostUndo(bool bSuccess) override;
+	virtual void PostRedo(bool bSuccess) override;
+
 	// FGCObject interface
-	virtual void AddReferencedObjects(FReferenceCollector& Collector) override
-	{
-		Collector.AddReferencedObject(VariableCollection);
-	}
-	virtual FString GetReferencerName() const override
-	{
-		return TEXT("FCameraVariableCollectionEditorToolkit");
-	}
+	virtual void AddReferencedObjects(FReferenceCollector& Collector) override;
+	virtual FString GetReferencerName() const override;
 
 private:
 
diff --git a/Engine/Plugins/Cameras/GameplayCameras/Source/GameplayCamerasEditor/Public/Widgets/SDeleteCameraObjectDialog.h b/Engine/Plugins/Cameras/GameplayCameras/Source/GameplayCamerasEditor/Public/Widgets/SDeleteCameraObjectDialog.h
new file mode 100644
--- /dev/null
+++ b/Engine/Plugins/Cameras/GameplayCameras/Source/GameplayCamerasEditor/Public/Widgets/SDeleteCameraObjectDialog.h
@@ -0,0 +1,99 @@
+// Copyright Epic Games, Inc. All Rights Reserved.
+
+#pragma once
+
+#include "IAssetTypeActions.h"
+#include "Widgets/SCompoundWidget.h"
+
+class IAssetRegistry;
+
+namespace UE::Cameras
+{
+
+DECLARE_DELEGATE_OneParam(FOnDeletedObject, UObject*);
+
+/**
+ * Re-usable dialog for deleting camera objects that are sub-objects of an asset.
+ */
+class SDeleteCameraObjectDialog : public SCompoundWidget
+{
+public:
+
+	SLATE_BEGIN_ARGS(SDeleteCameraObjectDialog)
+		: _RenameObjectsAsTrash(false)
+	{}
+		/** The modal window containing this dialog. */
+		SLATE_ARGUMENT(TWeakPtr<SWindow>, ParentWindow)
+		/** The list of objects to delete. */
+		SLATE_ARGUMENT(TArrayView<UObject*>, ObjectsToDelete)
+		/** Whether deleted objects should also get `TRASH_` prepended to their name. */
+		SLATE_ATTRIBUTE(bool, RenameObjectsAsTrash)
+		/** Custom callback invoked when objects are actually deleted. */
+		SLATE_EVENT(FOnDeletedObject, OnDeletedObject)
+	SLATE_END_ARGS()
+
+	void Construct(const FArguments& InArgs);
+
+	bool ShouldPerformDelete() const;
+	void PerformReferenceReplacement();
+
+public:
+
+	static void RenameObjectAsTrash(FString& InOutName);
+	static FString MakeTrashName(const FString& InName);
+
+protected:
+
+	// SCompoundWidget interface.
+	virtual void Tick(const FGeometry& AllottedGeometry, const double InCurrentTime, const float InDeltaTime) override;
+	virtual FReply OnKeyDown(const FGeometry& MyGeometry, const FKeyEvent& InKeyEvent) override;
+
+private:
+
+	void StartScanning();
+	void ScanNextObject(IAssetRegistry& AssetRegistry);
+	void ScanNextReferencingPackage();
+	void FinishScanning();
+
+	TOptional<float> GetProgressBarPercent() const;
+
+	TSharedRef<SWidget> BuildReferencerAssetPicker();
+	bool OnShouldFilterReferencerAsset(const FAssetData& InAssetData) const;
+
+	EVisibility GetNoReferencesVisibility() const;
+	EVisibility GetReferencesVisiblity() const;
+	EVisibility GetProgressBarVisibility() const;
+
+	void OnAssetsActivated(const TArray<FAssetData>& ActivatedAssets, EAssetTypeActivationMethod::Type ActivationMethod);
+	bool IsDeleteEnabled() const;
+	FReply OnDeleteClicked();
+	FReply OnCancelClicked();
+
+	void CloseWindow();
+
+private:
+
+	enum class EState { Waiting, StartScanning, ScanNextObject, ScanNextReferencingPackage, FinishScanning };
+
+	TWeakPtr<SWindow> WeakParentWindow;
+
+	TArray<UObject*> ObjectsToDelete;
+
+	EState State = EState::Waiting;
+
+	TSet<FName> ReferencingPackages;
+	TSet<FName> PossiblyReferencingPackages;
+	bool bIsAnyReferencedInMemoryByNonUndo = false;
+	bool bIsAnyReferencedInMemoryByUndo = false;
+
+	TArray<FName> ReferencingPackagesToScan;
+	int32 NextObjectToScan = INDEX_NONE;
+	int32 NextPackageToScan = INDEX_NONE;
+
+	bool bPerformDelete = false;
+	bool bRenameObjectsAsTrash = false;
+	FOnDeletedObject OnDeletedObject;
+};
+
+}  // namespace UE::Cameras
+
