Cameras: deleting multiple nodes now only creates one undo transaction

diff --git a/Engine/Plugins/Cameras/GameplayCameras/Source/GameplayCamerasEditor/Private/Editors/ObjectTreeGraphSchema.cpp b/Engine/Plugins/Cameras/GameplayCameras/Source/GameplayCamerasEditor/Private/Editors/ObjectTreeGraphSchema.cpp
--- a/Engine/Plugins/Cameras/GameplayCameras/Source/GameplayCamerasEditor/Private/Editors/ObjectTreeGraphSchema.cpp
+++ b/Engine/Plugins/Cameras/GameplayCameras/Source/GameplayCamerasEditor/Private/Editors/ObjectTreeGraphSchema.cpp
@@ -823,13 +823,10 @@ bool UObjectTreeGraphSchema::SafeDeleteNodeFromGraph(UEdGraph* Graph, UEdGraphNo
 		return false;
 	}
 	
-	const FScopedTransaction Transaction(LOCTEXT("DeleteNode", "Delete Node"));
-
 	BreakNodeLinks(*Node);
 
 	UObjectTreeGraph* ObjectTreeGraph = CastChecked<UObjectTreeGraph>(Graph);
 	OnDeleteNodeFromGraph(ObjectTreeGraph, Node);
-	Node->DestroyNode();
 
 	return true;
 }
diff --git a/Engine/Plugins/Cameras/GameplayCameras/Source/GameplayCamerasEditor/Private/Editors/SObjectTreeGraphEditor.cpp b/Engine/Plugins/Cameras/GameplayCameras/Source/GameplayCamerasEditor/Private/Editors/SObjectTreeGraphEditor.cpp
--- a/Engine/Plugins/Cameras/GameplayCameras/Source/GameplayCamerasEditor/Private/Editors/SObjectTreeGraphEditor.cpp
+++ b/Engine/Plugins/Cameras/GameplayCameras/Source/GameplayCamerasEditor/Private/Editors/SObjectTreeGraphEditor.cpp
@@ -389,6 +389,8 @@ void SObjectTreeGraphEditor::DeleteNodes(TArrayView<UObjectTreeGraphNode*> Nodes
 	UEdGraph* CurrentGraph = GraphEditor->GetCurrentGraph();
 	const UEdGraphSchema* Schema = CurrentGraph->GetSchema();
 
+	const FScopedTransaction Transaction(LOCTEXT("DeleteNode", "Delete Node(s)"));
+
 	for (UObjectTreeGraphNode* Node : NodesToDelete)
 	{
 		if (Node)
