Cameras: use the new object delete dialog when deleting camera rigs

This correctly fixes up any references to the deleted camera rig.
Also, refresh the camera rig list on undo/redo.

diff --git a/Engine/Plugins/Cameras/GameplayCameras/Source/GameplayCamerasEditor/Private/Editors/SCameraRigList.cpp b/Engine/Plugins/Cameras/GameplayCameras/Source/GameplayCamerasEditor/Private/Editors/SCameraRigList.cpp
--- a/Engine/Plugins/Cameras/GameplayCameras/Source/GameplayCamerasEditor/Private/Editors/SCameraRigList.cpp
+++ b/Engine/Plugins/Cameras/GameplayCameras/Source/GameplayCamerasEditor/Private/Editors/SCameraRigList.cpp
@@ -4,6 +4,7 @@
 
 #include "Commands/CameraAssetEditorCommands.h"
 #include "Core/CameraAsset.h"
+#include "Editor.h"
 #include "Framework/Commands/UICommandList.h"
 #include "ScopedTransaction.h"
 #include "Styles/GameplayCamerasEditorStyle.h"
@@ -12,6 +13,8 @@
 #include "Widgets/Input/SSearchBox.h"
 #include "Widgets/Layout/SBox.h"
 #include "Widgets/SBoxPanel.h"
+#include "Widgets/SDeleteCameraObjectDialog.h"
+#include "Widgets/SWindow.h"
 #include "Widgets/Text/SInlineEditableTextBlock.h"
 #include "Widgets/Views/SListView.h"
 
@@ -202,6 +205,11 @@ SCameraRigList::~SCameraRigList()
 {
 }
 
+void SCameraRigList::RequestListRefresh()
+{
+	bUpdateItemSource = true;
+}
+
 void SCameraRigList::Tick(const FGeometry& AllottedGeometry, const double InCurrentTime, const float InDeltaTime)
 {
 	if (bUpdateItemSource)
@@ -312,29 +320,53 @@ bool SCameraRigList::CanRenameCameraRig()
 
 void SCameraRigList::OnDeleteCameraRig()
 {
+	// Check if we have anything to delete.
 	TArray<TSharedPtr<FCameraRigListItem>> SelectedItems;
 	ListView->GetSelectedItems(SelectedItems);
-	if (SelectedItems.Num() > 0)
+	if (SelectedItems.IsEmpty())
 	{
-		FScopedTransaction Transaction(LOCTEXT("DeleteCameraRigs", "Delete Camera Rigs"));
+		return;
+	}
+
+	// Display a dialog showing any referencing assets that would need to be modified.
+	TSharedRef<SWindow> DeleteCameraRigWindow = SNew(SWindow)
+		.Title(LOCTEXT("DeleteCameraRigWindowTitle", "Delete Camera Rig(s)"))
+		.ClientSize(FVector2D(600, 700));
+
+	TArray<UObject*> ObjectsToDelete;
+	for (TSharedPtr<FCameraRigListItem> SelectedItem : SelectedItems)
+	{
+		ObjectsToDelete.Add(SelectedItem->CameraRigAsset);
+	}
+
+	TSharedRef<SDeleteCameraObjectDialog> DeleteCameraRigDialog = SNew(SDeleteCameraObjectDialog)
+		.ParentWindow(DeleteCameraRigWindow)
+		.ObjectsToDelete(ObjectsToDelete)
+		.OnDeletedObject_Lambda([](UObject* Obj)
+					{
+						if (UCameraRigAsset* TrashCameraRig = Cast<UCameraRigAsset>(Obj))
+						{
+							SDeleteCameraObjectDialog::RenameObjectAsTrash(TrashCameraRig->Interface.DisplayName);
+						}
+					});
+	DeleteCameraRigWindow->SetContent(DeleteCameraRigDialog);
+
+	GEditor->EditorAddModalWindow(DeleteCameraRigWindow);
+
+	// Remove the camera rigs from the camera asset and perform the reference replacement.
+	const bool bPerformDelete = DeleteCameraRigDialog->ShouldPerformDelete();
+	if (bPerformDelete)
+	{
+		FScopedTransaction Transaction(LOCTEXT("DeleteCameraRigs", "Delete Camera Rig(s)"));
 
 		CameraAsset->Modify();
 
-		TStringBuilder<256> StringBuilder;
 		TArray<UCameraRigAsset*> DeletedCameraRigs;
-
 		for (TSharedPtr<FCameraRigListItem> Item : SelectedItems)
 		{
 			UCameraRigAsset* CameraRigAsset = Item->CameraRigAsset;
 			if (CameraRigAsset)
 			{
-				CameraRigAsset->Modify();
-				
-				StringBuilder.Reset();
-				StringBuilder.Append(TEXT("DELETED_"));
-				StringBuilder.Append(CameraRigAsset->GetName());
-				CameraRigAsset->Rename(StringBuilder.ToString());
-
 				const int32 NumRemoved = CameraAsset->RemoveCameraRig(CameraRigAsset);
 				ensure(NumRemoved == 1);
 
@@ -342,6 +374,8 @@ void SCameraRigList::OnDeleteCameraRig()
 			}
 		}
 
+		DeleteCameraRigDialog->PerformReferenceReplacement();
+
 		bUpdateItemSource = true;
 
 		OnCameraRigDeleted.ExecuteIfBound(DeletedCameraRigs);
diff --git a/Engine/Plugins/Cameras/GameplayCameras/Source/GameplayCamerasEditor/Private/Toolkits/CameraRigsAssetEditorMode.cpp b/Engine/Plugins/Cameras/GameplayCameras/Source/GameplayCamerasEditor/Private/Toolkits/CameraRigsAssetEditorMode.cpp
--- a/Engine/Plugins/Cameras/GameplayCameras/Source/GameplayCamerasEditor/Private/Toolkits/CameraRigsAssetEditorMode.cpp
+++ b/Engine/Plugins/Cameras/GameplayCameras/Source/GameplayCamerasEditor/Private/Toolkits/CameraRigsAssetEditorMode.cpp
@@ -4,6 +4,7 @@
 
 #include "Core/CameraAsset.h"
 #include "Core/CameraRigAsset.h"
+#include "Editor.h"
 #include "Editors/CameraNodeGraphSchema.h"
 #include "Editors/CameraRigTransitionGraphSchema.h"
 #include "Editors/SCameraRigAssetEditor.h"
@@ -76,6 +77,11 @@ void FCameraRigsAssetEditorMode::OnActivateMode(const FAssetEditorModeActivatePa
 	Impl->BindCommands(InParams.CommandList.ToSharedRef());
 
 	Impl->OnCameraRigBuildStatusDirtied().AddSP(this, &FCameraRigsAssetEditorMode::OnCameraRigBuildStatusDirtied);
+
+	if (GEditor)
+	{
+		GEditor->RegisterForUndo(this);
+	}
 }
 
 TSharedRef<SDockTab> FCameraRigsAssetEditorMode::SpawnTab_CameraRigs(const FSpawnTabArgs& Args)
@@ -91,6 +97,11 @@ TSharedRef<SDockTab> FCameraRigsAssetEditorMode::SpawnTab_CameraRigs(const FSpaw
 
 void FCameraRigsAssetEditorMode::OnDeactivateMode(const FAssetEditorModeDeactivateParams& InParams)
 {
+	if (GEditor)
+	{
+		GEditor->UnregisterForUndo(this);
+	}
+
 	Impl->OnCameraRigBuildStatusDirtied().RemoveAll(this);
 
 	Impl->UnregisterTabSpawners(InParams.TabManager.ToSharedRef());
@@ -100,6 +111,22 @@ void FCameraRigsAssetEditorMode::OnDeactivateMode(const FAssetEditorModeDeactiva
 	UToolMenus::UnregisterOwner(this);
 }
 
+void FCameraRigsAssetEditorMode::PostUndo(bool bSuccess)
+{
+	if (CameraRigsListWidget)
+	{
+		CameraRigsListWidget->RequestListRefresh();
+	}
+}
+
+void FCameraRigsAssetEditorMode::PostRedo(bool bSuccess)
+{
+	if (CameraRigsListWidget)
+	{
+		CameraRigsListWidget->RequestListRefresh();
+	}
+}
+
 void FCameraRigsAssetEditorMode::OnCameraRigListChanged(TArrayView<UCameraRigAsset* const> InCameraRigs)
 {
 }
diff --git a/Engine/Plugins/Cameras/GameplayCameras/Source/GameplayCamerasEditor/Public/Editors/SCameraRigList.h b/Engine/Plugins/Cameras/GameplayCameras/Source/GameplayCamerasEditor/Public/Editors/SCameraRigList.h
--- a/Engine/Plugins/Cameras/GameplayCameras/Source/GameplayCamerasEditor/Public/Editors/SCameraRigList.h
+++ b/Engine/Plugins/Cameras/GameplayCameras/Source/GameplayCamerasEditor/Public/Editors/SCameraRigList.h
@@ -83,6 +83,8 @@ public:
 	void Construct(const FArguments& InArgs);
 	~SCameraRigList();
 
+	void RequestListRefresh();
+
 protected:
 
 	// SWidget interface
diff --git a/Engine/Plugins/Cameras/GameplayCameras/Source/GameplayCamerasEditor/Public/Toolkits/CameraRigsAssetEditorMode.h b/Engine/Plugins/Cameras/GameplayCameras/Source/GameplayCamerasEditor/Public/Toolkits/CameraRigsAssetEditorMode.h
--- a/Engine/Plugins/Cameras/GameplayCameras/Source/GameplayCamerasEditor/Public/Toolkits/CameraRigsAssetEditorMode.h
+++ b/Engine/Plugins/Cameras/GameplayCameras/Source/GameplayCamerasEditor/Public/Toolkits/CameraRigsAssetEditorMode.h
@@ -5,6 +5,7 @@
 #include "CoreTypes.h"
 #include "Editors/ObjectTreeGraphConfig.h"
 #include "Editors/SCameraRigList.h"
+#include "EditorUndoClient.h"
 #include "Toolkits/AssetEditorMode.h"
 
 class FDocumentTabFactory;
@@ -21,6 +22,7 @@ class FCameraRigAssetEditorToolkitBase;
 
 class FCameraRigsAssetEditorMode
 	: public FAssetEditorMode
+	, public FEditorUndoClient
 {
 public:
 
@@ -33,9 +35,14 @@ public:
 
 protected:
 
+	// FAssetEditorMode interface.
 	virtual void OnActivateMode(const FAssetEditorModeActivateParams& InParams) override;
 	virtual void OnDeactivateMode(const FAssetEditorModeDeactivateParams& InParams) override;
 
+	// FEditorUndoClient interface.
+	virtual void PostUndo(bool bSuccess) override;
+	virtual void PostRedo(bool bSuccess) override;
+
 private:
 
 	TSharedRef<SDockTab> SpawnTab_CameraRigs(const FSpawnTabArgs& Args);
