Cameras: don't make editor graphs root GC objects anymore, make them standalone objects instead

diff --git a/Engine/Plugins/Cameras/GameplayCameras/Source/GameplayCamerasEditor/Private/Editors/SCameraRigAssetEditor.cpp b/Engine/Plugins/Cameras/GameplayCameras/Source/GameplayCamerasEditor/Private/Editors/SCameraRigAssetEditor.cpp
--- a/Engine/Plugins/Cameras/GameplayCameras/Source/GameplayCamerasEditor/Private/Editors/SCameraRigAssetEditor.cpp
+++ b/Engine/Plugins/Cameras/GameplayCameras/Source/GameplayCamerasEditor/Private/Editors/SCameraRigAssetEditor.cpp
@@ -73,9 +73,8 @@ void SCameraRigAssetEditor::CreateNodeGraphEditor()
 	UCameraNodeGraphSchema* DefaultSchemaObject = Cast<UCameraNodeGraphSchema>(SchemaClass->GetDefaultObject());
 	FObjectTreeGraphConfig GraphConfig = DefaultSchemaObject->BuildGraphConfig();
 
-	NodeGraph = NewObject<UObjectTreeGraph>(GetTransientPackage(), NAME_None, RF_Transactional);
+	NodeGraph = NewObject<UObjectTreeGraph>(GetTransientPackage(), NAME_None, RF_Transactional | RF_Standalone);
 	NodeGraph->Schema = SchemaClass;
-	NodeGraph->AddToRoot();
 	NodeGraph->Reset(CameraRigAsset, GraphConfig, EObjectTreeGraphBuildSource::RootObjectPackage);
 
 	NodeGraphChangedHandle = NodeGraph->AddOnGraphChangedHandler(
@@ -99,9 +98,8 @@ void SCameraRigAssetEditor::CreateTransitionGraphEditor()
 	UCameraRigTransitionGraphSchema* DefaultSchemaObject = Cast<UCameraRigTransitionGraphSchema>(SchemaClass->GetDefaultObject());
 	FObjectTreeGraphConfig GraphConfig = DefaultSchemaObject->BuildGraphConfig();
 
-	TransitionGraph = NewObject<UObjectTreeGraph>(GetTransientPackage(), NAME_None, RF_Transactional);
+	TransitionGraph = NewObject<UObjectTreeGraph>(GetTransientPackage(), NAME_None, RF_Transactional | RF_Standalone);
 	TransitionGraph->Schema = SchemaClass;
-	TransitionGraph->AddToRoot();
 	TransitionGraph->Reset(CameraRigAsset, GraphConfig, EObjectTreeGraphBuildSource::RootObjectPackage);
 
 	TransitionGraphChangedHandle = TransitionGraph->AddOnGraphChangedHandler(
diff --git a/Engine/Plugins/Cameras/GameplayCameras/Source/GameplayCamerasEditor/Private/Editors/SCameraRigTransitionEditor.cpp b/Engine/Plugins/Cameras/GameplayCameras/Source/GameplayCamerasEditor/Private/Editors/SCameraRigTransitionEditor.cpp
--- a/Engine/Plugins/Cameras/GameplayCameras/Source/GameplayCamerasEditor/Private/Editors/SCameraRigTransitionEditor.cpp
+++ b/Engine/Plugins/Cameras/GameplayCameras/Source/GameplayCamerasEditor/Private/Editors/SCameraRigTransitionEditor.cpp
@@ -67,9 +67,8 @@ void SCameraRigTransitionEditor::CreateTransitionGraphEditor()
 	UCameraRigTransitionGraphSchemaBase* DefaultSchemaObject = Cast<UCameraRigTransitionGraphSchemaBase>(TransitionGraphSchemaClass->GetDefaultObject());
 	FObjectTreeGraphConfig GraphConfig = DefaultSchemaObject->BuildGraphConfig();
 
-	TransitionGraph = NewObject<UObjectTreeGraph>(GetTransientPackage(), NAME_None, RF_Transactional);
+	TransitionGraph = NewObject<UObjectTreeGraph>(GetTransientPackage(), NAME_None, RF_Transactional | RF_Standalone);
 	TransitionGraph->Schema = TransitionGraphSchemaClass;
-	TransitionGraph->AddToRoot();
 	TransitionGraph->Reset(TransitionOwner, GraphConfig, EObjectTreeGraphBuildSource::RootObjectPackage);
 
 	TransitionGraphChangedHandle = TransitionGraph->AddOnGraphChangedHandler(
