Cameras: refactor how ObjectTreeGraph manipulation works

The problem was that auto-disconnecting pins before reconnecting them did not trigger the code paths for clearing the disconnected properties.
The new code re-implements disconnecting and reconnecting so that everything gets applied to the underlying object model.

diff --git a/Engine/Plugins/Cameras/GameplayCameras/Source/GameplayCamerasEditor/Private/Editors/CameraNodeGraphSchema.cpp b/Engine/Plugins/Cameras/GameplayCameras/Source/GameplayCamerasEditor/Private/Editors/CameraNodeGraphSchema.cpp
--- a/Engine/Plugins/Cameras/GameplayCameras/Source/GameplayCamerasEditor/Private/Editors/CameraNodeGraphSchema.cpp
+++ b/Engine/Plugins/Cameras/GameplayCameras/Source/GameplayCamerasEditor/Private/Editors/CameraNodeGraphSchema.cpp
@@ -161,7 +161,7 @@ const FPinConnectionResponse UCameraNodeGraphSchema::CanCreateConnection(const U
 	return Super::CanCreateConnection(A, B);
 }
 
-bool UCameraNodeGraphSchema::OnCreateConnection(UEdGraphPin* A, UEdGraphPin* B) const
+bool UCameraNodeGraphSchema::OnApplyConnection(UEdGraphPin* A, UEdGraphPin* B, FDelayedPinActions& Actions) const
 {
 	// Try to make a connection between a camera node's parameter pin and a camera rig interface parameter.
 	// First, figure out which is which.
@@ -213,8 +213,6 @@ bool UCameraNodeGraphSchema::OnCreateConnection(UEdGraphPin* A, UEdGraphPin* B)
 	}
 
 	// Make the connection.
-	const FScopedTransaction Transaction(LOCTEXT("ExposeCameraRigParameter", "Expose Camera Rig Parameter"));
-
 	RigParameter->Modify();
 
 	RigParameter->Target = CameraNode;
@@ -224,39 +222,35 @@ bool UCameraNodeGraphSchema::OnCreateConnection(UEdGraphPin* A, UEdGraphPin* B)
 		RigParameter->InterfaceParameterName = StructProperty->GetName();
 	}
 
-	UEdGraph* Graph = CameraNodeNode->GetGraph();
-	Graph->NotifyGraphChanged();
-
 	return true;
 }
 
-bool UCameraNodeGraphSchema::OnBreakPinLinks(UEdGraphPin& TargetPin, bool bSendsNodeNotification) const
+bool UCameraNodeGraphSchema::OnApplyDisconnection(UEdGraphPin* TargetPin, FDelayedPinActions& Actions, bool bIsReconnecting) const
 {
 	// See if we have a rig parameter connection to break.
-	if (TargetPin.PinType.PinCategory == PC_Self || TargetPin.PinType.PinCategory == PC_CameraParameter)
+	if (TargetPin->PinType.PinCategory == PC_Self || TargetPin->PinType.PinCategory == PC_CameraParameter)
 	{
-		UEdGraphPin* RigParameterSelfPin = &TargetPin;
-		if (TargetPin.PinType.PinCategory == PC_CameraParameter)
+		UEdGraphPin* RigParameterSelfPin = TargetPin;
+		if (TargetPin->PinType.PinCategory == PC_CameraParameter)
 		{
-			RigParameterSelfPin = TargetPin.LinkedTo[0];
+			RigParameterSelfPin = TargetPin->LinkedTo[0];
 		}
 
 		UObjectTreeGraphNode* RigParameterNode = Cast<UObjectTreeGraphNode>(RigParameterSelfPin->GetOwningNode());
-		if (RigParameterNode && RigParameterNode->IsObjectA<UCameraRigInterfaceParameter>())
+		if (!RigParameterNode)
 		{
-			const FScopedTransaction Transaction(LOCTEXT("BreakPinLinks", "Break Pin Links"));
-
-			UCameraRigInterfaceParameter* RigParameter = RigParameterNode->CastObject<UCameraRigInterfaceParameter>();
+			return false;
+		}
 
+		UCameraRigInterfaceParameter* RigParameter = RigParameterNode->CastObject<UCameraRigInterfaceParameter>();
+		if (RigParameter)
+		{
 			RigParameter->Modify();
 
 			RigParameter->Target = nullptr;
 			RigParameter->TargetPropertyName = NAME_None;
 			RigParameter->PrivateVariable = nullptr;
 
-			UEdGraph* Graph = RigParameterNode->GetGraph();
-			Graph->NotifyGraphChanged();
-
 			return true;
 		}
 	}
@@ -264,49 +258,19 @@ bool UCameraNodeGraphSchema::OnBreakPinLinks(UEdGraphPin& TargetPin, bool bSends
 	return false;
 }
 
-bool UCameraNodeGraphSchema::OnBreakSinglePinLink(UEdGraphPin* SourcePin, UEdGraphPin* TargetPin) const
+bool UCameraNodeGraphSchema::OnApplyDisconnection(UEdGraphPin* SourcePin, UEdGraphPin* TargetPin, FDelayedPinActions& Actions) const
 {
 	UObjectTreeGraphNode* RigParameterNode = nullptr;
 
 	if (SourcePin->PinType.PinCategory == PC_Self && TargetPin->PinType.PinCategory == PC_CameraParameter)
 	{
-		UObjectTreeGraphNode* SourceNode = Cast<UObjectTreeGraphNode>(SourcePin->GetOwningNode());
-		if (SourceNode)
-		{
-			RigParameterNode = SourceNode;
-		}
+		return OnApplyDisconnection(SourcePin, Actions, false);
 	}
 	else if (SourcePin->PinType.PinCategory == PC_CameraParameter && TargetPin->PinType.PinCategory == PC_Self)
 	{
-		UObjectTreeGraphNode* TargetNode = Cast<UObjectTreeGraphNode>(TargetPin->GetOwningNode());
-		if (TargetNode)
-		{ 
-			RigParameterNode = TargetNode;
-		}
-	}
-
-	UCameraRigInterfaceParameter* RigParameter = nullptr;
-	if (RigParameterNode)
-	{
-		RigParameter = RigParameterNode->CastObject<UCameraRigInterfaceParameter>();
+		return OnApplyDisconnection(TargetPin, Actions, false);
 	}
 
-	if (RigParameter)
-	{
-		const FScopedTransaction Transaction(LOCTEXT("BreakSinglePinLink", "Break Pin Link"));
-
-		RigParameter->Modify();
-
-		RigParameter->Target = nullptr;
-		RigParameter->TargetPropertyName = NAME_None;
-		RigParameter->PrivateVariable = nullptr;
-
-		UEdGraph* Graph = RigParameterNode->GetGraph();
-		Graph->NotifyGraphChanged();
-
-		return true;
-	}
-	
 	return false;
 }
 
@@ -339,9 +303,8 @@ UEdGraphNode* FCameraNodeGraphSchemaAction_NewInterfaceParameterNode::PerformAct
 	const UObjectTreeGraphSchema* Schema = CastChecked<UObjectTreeGraphSchema>(ParentGraph->GetSchema());
 
 	UCameraRigInterfaceParameter* NewInterfaceParameter = NewObject<UCameraRigInterfaceParameter>(CameraRig, NAME_None, RF_Transactional);
-	NewInterfaceParameter->Target = Target;
-	NewInterfaceParameter->TargetPropertyName = TargetPropertyName;
-	NewInterfaceParameter->InterfaceParameterName = TargetPropertyName.ToString();
+	// The interface parameter's properties will be set correctly inside AutowireNewNode by virtue
+	// of getting connected to the dragged camera node pin.
 
 	UObjectTreeGraphNode* NewGraphNode = Schema->CreateObjectNode(ObjectTreeGraph, NewInterfaceParameter);
 
diff --git a/Engine/Plugins/Cameras/GameplayCameras/Source/GameplayCamerasEditor/Private/Editors/ObjectTreeGraphNode.cpp b/Engine/Plugins/Cameras/GameplayCameras/Source/GameplayCamerasEditor/Private/Editors/ObjectTreeGraphNode.cpp
--- a/Engine/Plugins/Cameras/GameplayCameras/Source/GameplayCamerasEditor/Private/Editors/ObjectTreeGraphNode.cpp
+++ b/Engine/Plugins/Cameras/GameplayCameras/Source/GameplayCamerasEditor/Private/Editors/ObjectTreeGraphNode.cpp
@@ -24,10 +24,10 @@ void UObjectTreeGraphNode::Initialize(UObject* InObject)
 {
 	ensure(InObject);
 
-	Object = InObject;
+	WeakObject = InObject;
 
 	const FNodeContext NodeContext = GetNodeContext();
-	IObjectTreeGraphObject* GraphObject = Cast<IObjectTreeGraphObject>(Object);
+	IObjectTreeGraphObject* GraphObject = Cast<IObjectTreeGraphObject>(InObject);
 	if (GraphObject && GraphObject->HasSupportFlags(NodeContext.GraphConfig.GraphName, EObjectTreeGraphObjectSupportFlags::CommentText))
 	{
 		NodeComment = GraphObject->GetGraphNodeCommentText(NodeContext.GraphConfig.GraphName);
@@ -36,6 +36,7 @@ void UObjectTreeGraphNode::Initialize(UObject* InObject)
 
 FText UObjectTreeGraphNode::GetNodeTitle(ENodeTitleType::Type TitleType) const
 {
+	UObject* Object = WeakObject.Get();
 	if (Object)
 	{
 		const FNodeContext NodeContext = GetNodeContext();
@@ -63,6 +64,7 @@ FLinearColor UObjectTreeGraphNode::GetNodeBodyTintColor() const
 
 FText UObjectTreeGraphNode::GetTooltipText() const
 {
+	UObject* Object = WeakObject.Get();
 	if (Object)
 	{
 		return Object->GetClass()->GetToolTipText();
@@ -72,6 +74,7 @@ FText UObjectTreeGraphNode::GetTooltipText() const
 
 void UObjectTreeGraphNode::AllocateDefaultPins()
 {
+	UObject* Object = WeakObject.Get();
 	if (!ensure(Object))
 	{
 		return;
@@ -212,7 +215,7 @@ void UObjectTreeGraphNode::RemoveItemPin(UEdGraphPin* InItemPin)
 				InItemPin->PinType.PinSubCategory == UObjectTreeGraphSchema::PSC_ArrayPropertyItem))
 	{
 		// Don't call RemovePin() because that also removes the parent pin.
-		// We just want to tremove the child pin.
+		// We just want to remove the child pin.
 		const int32 NumPinRemoved = Pins.Remove(InItemPin);
 		ensure(NumPinRemoved == 1);
 		const int32 NumSubPinRemoved = InItemPin->ParentPin->SubPins.Remove(InItemPin);
@@ -244,6 +247,7 @@ void UObjectTreeGraphNode::RefreshArrayPropertyPinNames()
 
 void UObjectTreeGraphNode::GetAllConnectableProperties(TArray<FProperty*>& OutProperties) const
 {
+	UObject* Object = WeakObject.Get();
 	if (!ensure(Object))
 	{
 		return;
@@ -343,6 +347,7 @@ UEdGraphPin* UObjectTreeGraphNode::GetPinForPropertyNewItem(FArrayProperty* InPr
 
 FProperty* UObjectTreeGraphNode::GetPropertyForPin(const UEdGraphPin* InPin) const
 {
+	UObject* Object = WeakObject.Get();
 	if (!ensure(Object))
 	{
 		return nullptr;
@@ -374,6 +379,7 @@ FProperty* UObjectTreeGraphNode::GetPropertyForPin(const UEdGraphPin* InPin) con
 
 UClass* UObjectTreeGraphNode::GetConnectedObjectClassForPin(const UEdGraphPin* InPin) const
 {
+	UObject* Object = WeakObject.Get();
 	if (!ensure(Object))
 	{
 		return nullptr;
@@ -498,6 +504,7 @@ void UObjectTreeGraphNode::GetNodeContextMenuActions(class UToolMenu* Menu, clas
 
 bool UObjectTreeGraphNode::GetCanRenameNode() const
 {
+	UObject* Object = WeakObject.Get();
 	const FNodeContext NodeContext = GetNodeContext();
 	IObjectTreeGraphObject* GraphObject = Cast<IObjectTreeGraphObject>(Object);
 	return GraphObject && GraphObject->HasSupportFlags(NodeContext.GraphConfig.GraphName, EObjectTreeGraphObjectSupportFlags::CustomRename);
@@ -507,6 +514,7 @@ void UObjectTreeGraphNode::OnRenameNode(const FString& NewName)
 {
 	Super::OnRenameNode(NewName);
 
+	UObject* Object = WeakObject.Get();
 	if (IObjectTreeGraphObject* GraphObject = Cast<IObjectTreeGraphObject>(Object))
 	{
 		const FScopedTransaction Transaction(LOCTEXT("RenameNode", "Rename Node"));
@@ -531,6 +539,7 @@ bool UObjectTreeGraphNode::CanUserDeleteNode() const
 
 bool UObjectTreeGraphNode::SupportsCommentBubble() const
 {
+	UObject* Object = WeakObject.Get();
 	const FNodeContext NodeContext = GetNodeContext();
 	IObjectTreeGraphObject* GraphObject = Cast<IObjectTreeGraphObject>(Object);
 	return GraphObject && GraphObject->HasSupportFlags(NodeContext.GraphConfig.GraphName, EObjectTreeGraphObjectSupportFlags::CommentText);
@@ -540,6 +549,7 @@ void UObjectTreeGraphNode::OnUpdateCommentText(const FString& NewComment)
 {
 	Super::OnUpdateCommentText(NewComment);
 
+	UObject* Object = WeakObject.Get();
 	IObjectTreeGraphObject* GraphObject = Cast<IObjectTreeGraphObject>(Object);
 	if (GraphObject)
 	{
@@ -554,6 +564,7 @@ void UObjectTreeGraphNode::OnUpdateCommentText(const FString& NewComment)
 
 void UObjectTreeGraphNode::OnGraphNodeMoved(bool bMarkDirty)
 {
+	UObject* Object = WeakObject.Get();
 	IObjectTreeGraphObject* GraphObject = Cast<IObjectTreeGraphObject>(Object);
 	if (GraphObject)
 	{
@@ -569,6 +580,7 @@ UObjectTreeGraphNode::FNodeContext UObjectTreeGraphNode::GetNodeContext() const
 	UObjectTreeGraph* OuterGraph = CastChecked<UObjectTreeGraph>(GetGraph());
 	const FObjectTreeGraphConfig& OuterGraphConfig = OuterGraph->GetConfig();
 
+	UObject* Object = WeakObject.Get();
 	UClass* ObjectClass = Object->GetClass();
 	const FObjectTreeGraphClassConfig& ObjectClassConfig = OuterGraphConfig.GetObjectClassConfig(ObjectClass);
 
diff --git a/Engine/Plugins/Cameras/GameplayCameras/Source/GameplayCamerasEditor/Private/Editors/ObjectTreeGraphSchema.cpp b/Engine/Plugins/Cameras/GameplayCameras/Source/GameplayCamerasEditor/Private/Editors/ObjectTreeGraphSchema.cpp
--- a/Engine/Plugins/Cameras/GameplayCameras/Source/GameplayCamerasEditor/Private/Editors/ObjectTreeGraphSchema.cpp
+++ b/Engine/Plugins/Cameras/GameplayCameras/Source/GameplayCamerasEditor/Private/Editors/ObjectTreeGraphSchema.cpp
@@ -8,11 +8,9 @@
 #include "Editors/ObjectTreeGraphNode.h"
 #include "Exporters/Exporter.h"
 #include "Factories.h"
-#include "GameplayCameras.h"
 #include "IGameplayCamerasEditorModule.h"
 #include "ScopedTransaction.h"
 #include "Serialization/ArchiveUObject.h"
-#include "UObject/FastReferenceCollector.h"
 #include "UObject/UObjectIterator.h"
 #include "UObject/UnrealType.h"
 #include "UnrealExporter.h"
@@ -593,220 +591,389 @@ const FPinConnectionResponse UObjectTreeGraphSchema::CanCreateConnection(const U
 
 bool UObjectTreeGraphSchema::TryCreateConnection(UEdGraphPin* A, UEdGraphPin* B) const
 {
-	const bool bModified = UEdGraphSchema::TryCreateConnection(A, B);
+	const FPinConnectionResponse Response = CanCreateConnection(A, B);
+
+	FScopedTransaction Transaction(LOCTEXT("CreateConnection", "Create Connection"));
+
+	bool bModified = true;
+	FDelayedPinActions Actions;
+
+	// We need to reimplement this completely and skip the up-call to the base class.
+	// Yes, it makes me sad too. But this is because we need to interleave graph manipulation (like
+	// making and breaking pin links) with underlying model manipulation (setting and unsetting
+	// references between objects).
+	//
+	// Note that we do a few delicate things here:
+	//
+	// 1. We pass `true` to ApplyDisconnection` to indicate that disconnected pins should not be
+	//    removed (which happens when disconnecting an array property pin). We want those pins kept
+	//    around because they're only disconnected to be immediately reconnected.
+	//
+	// 2. However, in at least one situation we want the pin removed immediately. It's when there's
+	//    an object currently at index X in an array, and the user reconnects it to index Y. In that
+	//    case, we want to get rid of pin X, removing that item in the array, so that pin Y is at
+	//    the correct index. This is why we apply delayed pin actions as usual after graph 
+	//    manipulation is done.
+	//
+	switch (Response.Response)
+	{
+	case CONNECT_RESPONSE_MAKE:
+		A->MakeLinkTo(B);
+		ApplyConnection(A, B, Actions);
+		break;
+
+	case CONNECT_RESPONSE_BREAK_OTHERS_A:
+		ApplyDisconnection(A, Actions, true);
+		A->BreakAllPinLinks(true);
+		A->MakeLinkTo(B);
+		Actions.Apply();
+		ApplyConnection(A, B, Actions);
+		break;
+
+	case CONNECT_RESPONSE_BREAK_OTHERS_B:
+		ApplyDisconnection(B, Actions, true);
+		B->BreakAllPinLinks(true);
+		A->MakeLinkTo(B);
+		Actions.Apply();
+		ApplyConnection(A, B, Actions);
+		break;
+
+	case CONNECT_RESPONSE_BREAK_OTHERS_AB:
+		ApplyDisconnection(A, Actions, true);
+		ApplyDisconnection(B, Actions, true);
+		A->BreakAllPinLinks(true);
+		B->BreakAllPinLinks(true);
+		A->MakeLinkTo(B);
+		Actions.Apply();
+		ApplyConnection(A, B, Actions);
+		break;
+
+	case CONNECT_RESPONSE_MAKE_WITH_CONVERSION_NODE:
+		bModified = CreateAutomaticConversionNodeAndConnections(A, B);
+		break;
+
+	case CONNECT_RESPONSE_MAKE_WITH_PROMOTION:
+		bModified = CreatePromotedConnection(A, B);
+		break;
+
+	case CONNECT_RESPONSE_DISALLOW:
+	default:
+		bModified = false;
+		break;
+	}
+
 	if (!bModified)
 	{
+		ensure(Actions.IsEmpty());
+		Transaction.Cancel();
 		return false;
 	}
 
-	const bool bHandled = OnCreateConnection(A, B);
-	if (bHandled)
-	{
-		return true;
-	}
+#if WITH_EDITOR
+	A->GetOwningNode()->PinConnectionListChanged(A);
+	B->GetOwningNode()->PinConnectionListChanged(B);
+#endif
 
-	const FScopedTransaction Transaction(LOCTEXT("CreateConnection", "Create Connection"));
+	Actions.Apply();
 
-	UObjectTreeGraphNode* NodeA = Cast<UObjectTreeGraphNode>(A->GetOwningNode());
-	UObjectTreeGraphNode* NodeB = Cast<UObjectTreeGraphNode>(B->GetOwningNode());
+	return true;
+}
 
-	// Try to always reason back to A being the property pin, and B being the self pin of the
-	// object we want to set on the property.
-	if (A->PinType.PinCategory == PC_Self)
+void UObjectTreeGraphSchema::BreakPinLinks(UEdGraphPin& TargetPin, bool bSendsNodeNotification) const
+{
+	if (TargetPin.LinkedTo.IsEmpty())
 	{
-		Swap(A, B);
-		Swap(NodeA, NodeB);
+		Super::BreakPinLinks(TargetPin, bSendsNodeNotification);
+		return;
 	}
-	// We know we are in the right configuration now because UEdGraphSchema::TryCreateConnection
-	// already called CanCreateConnection, which we implemented above as checking that A and B
-	// are a property/self pin pair, one way or the other.
-
-	UObject* ObjectA = NodeA->GetObject();
-	UObject* ObjectB = NodeB->GetObject();
-
-	FProperty* PropertyA = NodeA->GetPropertyForPin(A);
 
-	if (FObjectProperty* ObjectPropertyA = CastField<FObjectProperty>(PropertyA))
-	{
-		ObjectA->PreEditChange(PropertyA);
-
-		ObjectA->Modify();
+	const FScopedTransaction Transaction(LOCTEXT("BreakPinLinks", "Break Pin Links"));
 
-		ObjectPropertyA->SetValue_InContainer(ObjectA, TObjectPtr<UObject>(ObjectB));
+	FDelayedPinActions Actions;
+	ApplyDisconnection(&TargetPin, Actions, false);
 
-		FPropertyChangedEvent PropertyChangedEvent(PropertyA, EPropertyChangeType::ValueSet);
-		ObjectA->PostEditChangeProperty(PropertyChangedEvent);
-	}
-	else if (FArrayProperty* ArrayPropertyA = CastField<FArrayProperty>(PropertyA))
-	{
-		ObjectA->PreEditChange(PropertyA);
+	Super::BreakPinLinks(TargetPin, bSendsNodeNotification);
 
-		ObjectA->Modify();
+	Actions.Apply();
+}
 
-		const int32 Index = NodeA->GetIndexOfArrayPin(A);
-		ensure(Index != INDEX_NONE);
+void UObjectTreeGraphSchema::BreakSinglePinLink(UEdGraphPin* SourcePin, UEdGraphPin* TargetPin) const
+{
+	const FScopedTransaction Transaction(LOCTEXT("BreakSinglePinLink", "Break Pin Link"));
 
-		FScriptArrayHelper ArrayHelper(ArrayPropertyA, ArrayPropertyA->ContainerPtrToValuePtr<void>(ObjectA));
-		const bool bAddNewItemPin = ArrayHelper.ExpandForIndex(Index);
+	FDelayedPinActions Actions;
+	ApplyDisconnection(SourcePin, TargetPin, Actions);
 
-		FObjectProperty* InnerProperty = CastFieldChecked<FObjectProperty>(ArrayPropertyA->Inner);
-		InnerProperty->SetObjectPropertyValue(ArrayHelper.GetRawPtr(Index), ObjectB);
+	Super::BreakSinglePinLink(SourcePin, TargetPin);
 
-		if (bAddNewItemPin)
-		{
-			NodeA->CreateNewItemPin(*ArrayPropertyA);
-			NodeA->GetGraph()->NotifyNodeChanged(NodeA);
-		}
+	Actions.Apply();
+}
 
-		FPropertyChangedEvent PropertyChangedEvent(PropertyA);
-		PropertyChangedEvent.ChangeType = bAddNewItemPin ? EPropertyChangeType::ArrayAdd : EPropertyChangeType::ValueSet;
-		ObjectA->PostEditChangeProperty(PropertyChangedEvent);
+void UObjectTreeGraphSchema::FDelayedPinActions::CreateNewItemPin(UObjectTreeGraphNode* Node, FArrayProperty* ArrayProperty)
+{
+	if (ensure(Node && ArrayProperty))
+	{
+		ItemPinsToCreate.Add({ Node, ArrayProperty });
 	}
+}
 
-	UEdGraph* Graph = NodeA->GetGraph();
-	Graph->NotifyGraphChanged();
-
-	return true;
+void UObjectTreeGraphSchema::FDelayedPinActions::RemoveItemPin(UEdGraphPin* Pin)
+{
+	if (ensure(Pin))
+	{
+		ItemPinsToRemove.Add(Pin);
+	}
 }
 
-bool UObjectTreeGraphSchema::OnCreateConnection(UEdGraphPin* A, UEdGraphPin* B) const
+bool UObjectTreeGraphSchema::FDelayedPinActions::IsEmpty() const
 {
-	return false;
+	return ItemPinsToCreate.IsEmpty() && ItemPinsToRemove.IsEmpty();
 }
 
-void UObjectTreeGraphSchema::BreakPinLinks(UEdGraphPin& TargetPin, bool bSendsNodeNotification) const
+void UObjectTreeGraphSchema::FDelayedPinActions::Apply()
 {
-	if (TargetPin.LinkedTo.IsEmpty())
+	TSet<UEdGraphNode*> NodesToNotify;
+	for (UEdGraphPin* Pin : ItemPinsToRemove)
 	{
-		Super::BreakPinLinks(TargetPin, bSendsNodeNotification);
-		return;
+		UObjectTreeGraphNode* OwningNode = Cast<UObjectTreeGraphNode>(Pin->GetOwningNode());
+		check(OwningNode);
+		OwningNode->RemoveItemPin(Pin);
+
+		NodesToNotify.Add(OwningNode);
+	}
+	for (TTuple<UObjectTreeGraphNode*, FArrayProperty*>& Pair : ItemPinsToCreate)
+	{
+		Pair.Key->CreateNewItemPin(*Pair.Value);
+
+		NodesToNotify.Add(Pair.Key);
+	}
+	for (UEdGraphNode* Node : NodesToNotify)
+	{
+		Node->GetGraph()->NotifyNodeChanged(Node);
 	}
+	ItemPinsToRemove.Reset();
+}
+
+void UObjectTreeGraphSchema::ApplyConnection(UEdGraphPin* A, UEdGraphPin* B, FDelayedPinActions& Actions) const
+{
+	// Input must have been validated prior to calling this method:
+	//
+	// - no null objects
+	// - pins belong to ObjectTreeGraph nodes
+	// - these nodes have valid objects
+	// - we should have a transaction active
+	//
+#if WITH_EDITOR
+	ensureMsgf(GUndo || !GEditor, TEXT("Setting property values on objects should be called inside a transaction"));
+#endif
+	
+	check(A && B);
 
-	const bool bHandled = OnBreakPinLinks(TargetPin, bSendsNodeNotification);
-	if (bHandled)
+	// See if a sub-class is handling this situation.
+	if (OnApplyConnection(A, B, Actions))
 	{
-		Super::BreakPinLinks(TargetPin, bSendsNodeNotification);
 		return;
 	}
 
-	const FScopedTransaction Transaction(LOCTEXT("BreakPinLinks", "Break Pin Links"));
+	// We handle situations where a property pin or array property pin is connected to the "self" pin
+	// of an object node. Lets see which pin is which.
+	UEdGraphPin* PropertyPin = nullptr;
+	UEdGraphPin* ValuePin = nullptr;
 
-	// TargetPin could be a self pin or a property pin, we need to handle both cases and directions.
-	UEdGraphPin* PropertyPin = &TargetPin;
-	UObjectTreeGraphNode* PropertyOwningNode = Cast<UObjectTreeGraphNode>(PropertyPin->GetOwningNode());
-	if (PropertyPin->PinType.PinCategory == PC_Self)
+	if (A->PinType.PinCategory == PC_Self && B->PinType.PinCategory == PC_Property)
 	{
-		PropertyPin = TargetPin.LinkedTo[0];
-		PropertyOwningNode = Cast<UObjectTreeGraphNode>(PropertyPin->GetOwningNode());
+		PropertyPin = B;
+		ValuePin = A;
 	}
+	else if (A->PinType.PinCategory == PC_Property && B->PinType.PinCategory == PC_Self)
+	{
+		PropertyPin = A;
+		ValuePin = B;
+	}
+
+	checkf(PropertyPin && ValuePin, TEXT("Invalid pins passed for setting property values."));
+
+	UObjectTreeGraphNode* PropertyNode = CastChecked<UObjectTreeGraphNode>(PropertyPin->GetOwningNode());
+	UObjectTreeGraphNode* ValueNode = CastChecked<UObjectTreeGraphNode>(ValuePin->GetOwningNode());
 
-	bool bRemovePropertyPin = false;
-	UObject* OwningObject = PropertyOwningNode->GetObject();
-	FProperty* Property = PropertyOwningNode->GetPropertyForPin(PropertyPin);
+	UObject* PropertyObject = PropertyNode->GetObject();
+	UObject* ValueObject = ValueNode->GetObject();
+	check(PropertyObject && ValueObject);
+
+	// If it is a property pin, set the value of the underlying property.
+	// If it is an array property pin, add a new item at the pin's index, and optionally add a new pin.
+	FProperty* Property = PropertyNode->GetPropertyForPin(PropertyPin);
 
 	if (FObjectProperty* ObjectProperty = CastField<FObjectProperty>(Property))
 	{
-		OwningObject->PreEditChange(Property);
+		PropertyObject->PreEditChange(Property);
 
-		OwningObject->Modify();
+		PropertyObject->Modify();
 
-		ObjectProperty->ClearValue_InContainer(OwningObject);
+		ObjectProperty->SetValue_InContainer(PropertyObject, TObjectPtr<UObject>(ValueObject));
 
 		FPropertyChangedEvent PropertyChangedEvent(Property, EPropertyChangeType::ValueSet);
-		OwningObject->PostEditChangeProperty(PropertyChangedEvent);
+		PropertyObject->PostEditChangeProperty(PropertyChangedEvent);
 	}
 	else if (FArrayProperty* ArrayProperty = CastField<FArrayProperty>(Property))
 	{
-		OwningObject->PreEditChange(Property);
+		PropertyObject->PreEditChange(Property);
 
-		OwningObject->Modify();
+		PropertyObject->Modify();
 
-		int32 Index = PropertyOwningNode->GetIndexOfArrayPin(PropertyPin);
+		const int32 Index = PropertyNode->GetIndexOfArrayPin(PropertyPin);
 		ensure(Index != INDEX_NONE);
 
-		FScriptArrayHelper ArrayHelper(ArrayProperty, ArrayProperty->ContainerPtrToValuePtr<void>(OwningObject));
-		ArrayHelper.RemoveValues(Index);
-
-		bRemovePropertyPin = true;
+		FScriptArrayHelper ArrayHelper(ArrayProperty, ArrayProperty->ContainerPtrToValuePtr<void>(PropertyObject));
+		const int32 PreviousNum = ArrayHelper.Num();
+		const bool bExpandedArray = ArrayHelper.ExpandForIndex(Index);
 
-		FPropertyChangedEvent PropertyChangedEvent(Property, EPropertyChangeType::ArrayRemove);
-		OwningObject->PostEditChangeProperty(PropertyChangedEvent);
-	}
+		FObjectProperty* InnerProperty = CastFieldChecked<FObjectProperty>(ArrayProperty->Inner);
+		InnerProperty->SetObjectPropertyValue(ArrayHelper.GetRawPtr(Index), ValueObject);
 
-	Super::BreakPinLinks(TargetPin, bSendsNodeNotification);
+		if (bExpandedArray)
+		{
+			ensure(Index == PreviousNum);  // We only support adding a single new element at the end, not "inserting"
+										   // into arbitrary indices requiring expanding the array by more than one
+										   // element.
+			Actions.CreateNewItemPin(PropertyNode, ArrayProperty);
+		}
 
-	if (bRemovePropertyPin)
-	{
-		PropertyOwningNode->RemoveItemPin(PropertyPin);
-		PropertyOwningNode->GetGraph()->NotifyNodeChanged(PropertyOwningNode);
+		FPropertyChangedEvent PropertyChangedEvent(Property);
+		PropertyChangedEvent.ChangeType = bExpandedArray ? EPropertyChangeType::ArrayAdd : EPropertyChangeType::ValueSet;
+		PropertyObject->PostEditChangeProperty(PropertyChangedEvent);
 	}
-
-	UEdGraph* Graph = PropertyOwningNode->GetGraph();
-	Graph->NotifyGraphChanged();
 }
 
-bool UObjectTreeGraphSchema::OnBreakPinLinks(UEdGraphPin& TargetPin, bool bSendsNodeNotification) const
+bool UObjectTreeGraphSchema::OnApplyConnection(UEdGraphPin* A, UEdGraphPin* B, FDelayedPinActions& Actions) const
 {
 	return false;
 }
 
-void UObjectTreeGraphSchema::BreakSinglePinLink(UEdGraphPin* SourcePin, UEdGraphPin* TargetPin) const
+void UObjectTreeGraphSchema::ApplyDisconnection(UEdGraphPin* TargetPin, FDelayedPinActions& Actions, bool bIsReconnecting) const
 {
-	const bool bHandled = OnBreakSinglePinLink(SourcePin, TargetPin);
-	if (bHandled)
+	// Input must have been validated prior to calling this method:
+	//
+	// - no null objects
+	// - the pin is the property pin to reset, or the self pin connected to a propert pin
+	// - the pin belongs to a ObjectTreeGraph node
+	// - this node has a valid object
+	// - we should have a transaction active
+	//
+#if WITH_EDITOR
+	ensureMsgf(GUndo || !GEditor, TEXT("Resetting property values on objects should be called inside a transaction"));
+#endif
+
+	check(TargetPin);
+
+	// See if we actually have anything to disconnect.
+	if (TargetPin->LinkedTo.IsEmpty())
 	{
-		Super::BreakSinglePinLink(SourcePin, TargetPin);
 		return;
 	}
 
-	const FScopedTransaction Transaction(LOCTEXT("BreakSinglePinLink", "Break Pin Link"));
+	// See if a sub-class is handling this situation.
+	if (OnApplyDisconnection(TargetPin, Actions, bIsReconnecting))
+	{
+		return;
+	}
 
-	// SourcePin could be the self-pin, and TargetPin the property pin, if the directions
-	// are reversed for that type of object/graph.
-	UEdGraphPin* PropertyPin = SourcePin;
-	UObjectTreeGraphNode* PropertyOwningNode = Cast<UObjectTreeGraphNode>(SourcePin->GetOwningNode());
-	if (SourcePin->PinType.PinCategory == PC_Self)
+	// We may either disconnect a self pin, or a property or array property pin. Let's see
+	// what sort of pin we were given: we want the property side of things.
+	bool bRemoveArrayItem = !bIsReconnecting;
+	if (TargetPin->PinType.PinCategory == PC_Self)
 	{
-		PropertyPin = SourcePin->LinkedTo[0];
-		PropertyOwningNode = Cast<UObjectTreeGraphNode>(PropertyPin->GetOwningNode());
+		TargetPin = TargetPin->LinkedTo[0];
+		// If `bIsReconnecting` is true, it means we need to leave `TargetPin` alone because
+		// we want to reconnect it right away. However, if the property pin is on the other
+		// side the of the link, we are free to remove it if it's an array property pin.
+		bRemoveArrayItem = true;
 	}
+	check(TargetPin->PinType.PinCategory == PC_Property);
 
-	bool bRemovePropertyPin = false;
-	UObject* OwningObject = PropertyOwningNode->GetObject();
-	FProperty* Property = PropertyOwningNode->GetPropertyForPin(PropertyPin);
+	UObjectTreeGraphNode* PropertyNode = Cast<UObjectTreeGraphNode>(TargetPin->GetOwningNode());
+	check(PropertyNode);
+
+	UObject* PropertyObject = PropertyNode->GetObject();
+	check(PropertyObject);
+
+	// If it is a property pin, clear the value of the underlying property.
+	// If it is an array property pin, remove the value at the given index in the underlying array,
+	// or just clear it if we want to reconnect that pin immediately.
+	FProperty* Property = PropertyNode->GetPropertyForPin(TargetPin);
 
 	if (FObjectProperty* ObjectProperty = CastField<FObjectProperty>(Property))
 	{
-		OwningObject->Modify();
+		PropertyObject->PreEditChange(Property);
+
+		PropertyObject->Modify();
+
+		ObjectProperty->ClearValue_InContainer(PropertyObject);
 
-		ObjectProperty->ClearValue_InContainer(OwningObject);
+		FPropertyChangedEvent PropertyChangedEvent(Property, EPropertyChangeType::ValueSet);
+		PropertyObject->PostEditChangeProperty(PropertyChangedEvent);
 	}
 	else if (FArrayProperty* ArrayProperty = CastField<FArrayProperty>(Property))
 	{
-		OwningObject->Modify();
+		PropertyObject->PreEditChange(Property);
+
+		PropertyObject->Modify();
 
-		int32 Index = PropertyOwningNode->GetIndexOfArrayPin(PropertyPin);
+		int32 Index = PropertyNode->GetIndexOfArrayPin(TargetPin);
 		ensure(Index != INDEX_NONE);
 
-		FScriptArrayHelper ArrayHelper(ArrayProperty, ArrayProperty->ContainerPtrToValuePtr<void>(OwningObject));
-		ArrayHelper.RemoveValues(Index);
+		FScriptArrayHelper ArrayHelper(ArrayProperty, ArrayProperty->ContainerPtrToValuePtr<void>(PropertyObject));
+		if (bRemoveArrayItem)
+		{
+			ArrayHelper.RemoveValues(Index);
+			Actions.RemoveItemPin(TargetPin);
+		}
+		else
+		{
+			FObjectProperty* InnerProperty = CastFieldChecked<FObjectProperty>(ArrayProperty->Inner);
+			InnerProperty->SetObjectPropertyValue(ArrayHelper.GetRawPtr(Index), nullptr);
+		}
 
-		bRemovePropertyPin = true;
+		FPropertyChangedEvent PropertyChangedEvent(Property, EPropertyChangeType::ArrayRemove);
+		PropertyObject->PostEditChangeProperty(PropertyChangedEvent);
 	}
+}
 
-	Super::BreakSinglePinLink(SourcePin, TargetPin);
+bool UObjectTreeGraphSchema::OnApplyDisconnection(UEdGraphPin* TargetPin, FDelayedPinActions& Actions, bool bIsReconnecting) const
+{
+	return false;
+}
 
-	if (bRemovePropertyPin)
+void UObjectTreeGraphSchema::ApplyDisconnection(UEdGraphPin* SourcePin, UEdGraphPin* TargetPin, FDelayedPinActions& Actions) const
+{
+	// Input must have been validated prior to calling this method:
+	//
+	// - no null objects
+	// - the pins are the property and self pin of a specific link
+	//
+	check(SourcePin && TargetPin);
+
+	// See if a sub-class is handling this situation.
+	if (OnApplyDisconnection(SourcePin, TargetPin, Actions))
 	{
-		PropertyOwningNode->RemoveItemPin(PropertyPin);
-		PropertyOwningNode->GetGraph()->NotifyNodeChanged(PropertyOwningNode);
+		return;
 	}
 
-	UEdGraph* Graph = PropertyOwningNode->GetGraph();
-	Graph->NotifyGraphChanged();
+	if (SourcePin->PinType.PinCategory == PC_Self && TargetPin->PinType.PinCategory == PC_Property)
+	{
+		return ApplyDisconnection(TargetPin, Actions, false);
+	}
+	else if (SourcePin->PinType.PinCategory == PC_Property && TargetPin->PinType.PinCategory == PC_Self)
+	{
+		return ApplyDisconnection(SourcePin, Actions, false);
+	}
+	else
+	{
+		checkf(false, TEXT("Invalid pins passed for setting property values."));
+	}
 }
 
-bool UObjectTreeGraphSchema::OnBreakSinglePinLink(UEdGraphPin* SourcePin, UEdGraphPin* TargetPin) const
+bool UObjectTreeGraphSchema::OnApplyDisconnection(UEdGraphPin* SourcePin, UEdGraphPin* TargetPin, FDelayedPinActions& Actions) const
 {
 	return false;
 }
diff --git a/Engine/Plugins/Cameras/GameplayCameras/Source/GameplayCamerasEditor/Public/Editors/CameraNodeGraphSchema.h b/Engine/Plugins/Cameras/GameplayCameras/Source/GameplayCamerasEditor/Public/Editors/CameraNodeGraphSchema.h
--- a/Engine/Plugins/Cameras/GameplayCameras/Source/GameplayCamerasEditor/Public/Editors/CameraNodeGraphSchema.h
+++ b/Engine/Plugins/Cameras/GameplayCameras/Source/GameplayCamerasEditor/Public/Editors/CameraNodeGraphSchema.h
@@ -33,9 +33,9 @@ protected:
 	virtual void OnCreateAllNodes(UObjectTreeGraph* InGraph, const FCreatedNodes& InCreatedNodes) const override;
 	virtual void OnAddConnectableObject(UObjectTreeGraph* InGraph, UObjectTreeGraphNode* InNewNode) const override;
 	virtual void OnRemoveConnectableObject(UObjectTreeGraph* InGraph, UObjectTreeGraphNode* InRemovedNode) const override;
-	virtual bool OnCreateConnection(UEdGraphPin* A, UEdGraphPin* B) const override;
-	virtual bool OnBreakPinLinks(UEdGraphPin& TargetPin, bool bSendsNodeNotification) const override;
-	virtual bool OnBreakSinglePinLink(UEdGraphPin* SourcePin, UEdGraphPin* TargetPin) const override;
+	virtual bool OnApplyConnection(UEdGraphPin* A, UEdGraphPin* B, FDelayedPinActions& Actions) const override;
+	virtual bool OnApplyDisconnection(UEdGraphPin* TargetPin, FDelayedPinActions& Actions, bool bIsReconnecting) const override;
+	virtual bool OnApplyDisconnection(UEdGraphPin* SourcePin, UEdGraphPin* TargetPin, FDelayedPinActions& Actions) const override;
 };
 
 /**
diff --git a/Engine/Plugins/Cameras/GameplayCameras/Source/GameplayCamerasEditor/Public/Editors/ObjectTreeGraphSchema.h b/Engine/Plugins/Cameras/GameplayCameras/Source/GameplayCamerasEditor/Public/Editors/ObjectTreeGraphSchema.h
--- a/Engine/Plugins/Cameras/GameplayCameras/Source/GameplayCamerasEditor/Public/Editors/ObjectTreeGraphSchema.h
+++ b/Engine/Plugins/Cameras/GameplayCameras/Source/GameplayCamerasEditor/Public/Editors/ObjectTreeGraphSchema.h
@@ -91,15 +91,28 @@ protected:
 		TMap<UObject*, UObjectTreeGraphNode*> CreatedNodes;
 	};
 
+	struct FDelayedPinActions
+	{
+		void CreateNewItemPin(UObjectTreeGraphNode* Node, FArrayProperty* ArrayProperty);
+		void RemoveItemPin(UEdGraphPin* Pin);
+
+		bool IsEmpty() const;
+		void Apply();
+
+	private:
+		TArray<TTuple<UObjectTreeGraphNode*, FArrayProperty*>> ItemPinsToCreate;
+		TArray<UEdGraphPin*> ItemPinsToRemove;
+	};
+
 	// UObjectTreeGraphSchema interface.
 	virtual void OnCreateAllNodes(UObjectTreeGraph* InGraph, const FCreatedNodes& InCreatedNodes) const;
 	virtual UObjectTreeGraphNode* OnCreateObjectNode(UObjectTreeGraph* InGraph, UObject* InObject) const;
 	virtual void OnAddConnectableObject(UObjectTreeGraph* InGraph, UObjectTreeGraphNode* InNewNode) const;
 	virtual void OnRemoveConnectableObject(UObjectTreeGraph* InGraph, UObjectTreeGraphNode* InRemovedNode) const;
 	virtual void CopyNonObjectNodes(TArrayView<UObject*> InObjects, FStringOutputDevice& OutDevice) const;
-	virtual bool OnCreateConnection(UEdGraphPin* A, UEdGraphPin* B) const;
-	virtual bool OnBreakPinLinks(UEdGraphPin& TargetPin, bool bSendsNodeNotification) const;
-	virtual bool OnBreakSinglePinLink(UEdGraphPin* SourcePin, UEdGraphPin* TargetPin) const;
+	virtual bool OnApplyConnection(UEdGraphPin* A, UEdGraphPin* B, FDelayedPinActions& Actions) const;
+	virtual bool OnApplyDisconnection(UEdGraphPin* TargetPin, FDelayedPinActions& Actions, bool bIsReconnecting) const;
+	virtual bool OnApplyDisconnection(UEdGraphPin* SourcePin, UEdGraphPin* TargetPin, FDelayedPinActions& Actions) const;
 	virtual void OnDeleteNodeFromGraph(UObjectTreeGraph* Graph, UEdGraphNode* Node) const;
 	virtual void FilterGraphContextPlaceableClasses(TArray<UClass*>& InOutClasses) const;
 
@@ -108,6 +121,12 @@ protected:
 	const FObjectTreeGraphClassConfig& GetObjectClassConfig(const UObjectTreeGraphNode* InNode) const;
 	const FObjectTreeGraphClassConfig& GetObjectClassConfig(const UObjectTreeGraph* InGraph, UClass* InObjectClass) const;
 
+protected:
+
+	void ApplyConnection(UEdGraphPin* A, UEdGraphPin* B, FDelayedPinActions& Actions) const;
+	void ApplyDisconnection(UEdGraphPin* TargetPin, FDelayedPinActions& Actions, bool bIsReconnecting) const;
+	void ApplyDisconnection(UEdGraphPin* SourcePin, UEdGraphPin* TargetPin, FDelayedPinActions& Actions) const;
+
 private:
 
 	void RemoveAllNodes(UObjectTreeGraph* InGraph) const;
