Missing file from previous changelist

diff --git a/Engine/Plugins/Cameras/GameplayCameras/Source/GameplayCamerasEditor/Public/Editors/ObjectTreeGraphNode.h b/Engine/Plugins/Cameras/GameplayCameras/Source/GameplayCamerasEditor/Public/Editors/ObjectTreeGraphNode.h
--- a/Engine/Plugins/Cameras/GameplayCameras/Source/GameplayCamerasEditor/Public/Editors/ObjectTreeGraphNode.h
+++ b/Engine/Plugins/Cameras/GameplayCameras/Source/GameplayCamerasEditor/Public/Editors/ObjectTreeGraphNode.h
@@ -31,7 +31,7 @@ public:
 	void Initialize(UObject* InObject);
 
 	/** Gets the underlying object represented by this graph node. */
-	UObject* GetObject() const { return Object.Get(); }
+	UObject* GetObject() const { return WeakObject.Get(); }
 	/** Gets whether we have a valid underlying object, and that it's a type of ObjectClass. */
 	template<typename ObjectClass> bool IsObjectA() const;
 	/** Gets the underlying object as a point to the given sub-class. */
@@ -107,7 +107,7 @@ protected:
 private:
 
 	UPROPERTY()
-	TObjectPtr<UObject> Object;
+	TWeakObjectPtr<UObject> WeakObject;
 
 	UPROPERTY()
 	TEnumAsByte<EEdGraphPinDirection> SelfPinDirectionOverride;
@@ -119,7 +119,7 @@ private:
 template<typename ObjectClass>
 bool UObjectTreeGraphNode::IsObjectA() const
 {
-	if (Object)
+	if (UObject* Object = WeakObject.Get())
 	{
 		return Object->IsA<ObjectClass>();
 	}
@@ -129,7 +129,7 @@ bool UObjectTreeGraphNode::IsObjectA() const
 template<typename ObjectClass> 
 ObjectClass* UObjectTreeGraphNode::CastObject() const
 {
-	if (Object)
+	if (UObject* Object = WeakObject.Get())
 	{
 		return Cast<ObjectClass>(Object);
 	}
