Cameras: handle edge-cases regarding camera rig interface parameters during the build process

The build process now handles dangling interface parameters, and interface parameters connected to dangling camera nodes.

diff --git a/Engine/Plugins/Cameras/GameplayCameras/Source/GameplayCameras/Private/Core/CameraRigAssetBuilder.cpp b/Engine/Plugins/Cameras/GameplayCameras/Source/GameplayCameras/Private/Core/CameraRigAssetBuilder.cpp
--- a/Engine/Plugins/Cameras/GameplayCameras/Source/GameplayCameras/Private/Core/CameraRigAssetBuilder.cpp
+++ b/Engine/Plugins/Cameras/GameplayCameras/Source/GameplayCameras/Private/Core/CameraRigAssetBuilder.cpp
@@ -8,6 +8,7 @@
 #include "Core/CameraRigBuildContext.h"
 #include "Core/CameraRigAsset.h"
 #include "Core/CameraVariableAssets.h"
+#include "Logging/TokenizedMessage.h"
 
 #define LOCTEXT_NAMESPACE "CameraRigAssetBuilder"
 
@@ -91,6 +92,23 @@ struct FPrivateVariableBuilder
 		return nullptr;
 	}
 
+	bool ReuseInterfaceParameter(UCameraRigInterfaceParameter* InterfaceParameter, UCameraVariableAsset* IntendedPrivateVariable)
+	{
+		using FReusableInterfaceParameterInfo = FCameraRigAssetBuilder::FReusableInterfaceParameterInfo;
+		FReusableInterfaceParameterInfo* FoundItem = Owner.OldInterfaceParameters.Find(InterfaceParameter);
+		if (ensure(FoundItem))
+		{
+			// This is an interface parameter that existed before. Flag things as modified if the 
+			// private variable is changing.
+			ensure(!FoundItem->Value);
+			FoundItem->Value = true;  // This one has now been re-used.
+			return FoundItem->Key != IntendedPrivateVariable;
+		}
+		// We should have had this interface parameter in our map, since we built it just a second ago!
+		// Something's wrong... oh well, flag things as modified.
+		return true;
+	}
+
 private:
 
 	UCameraVariableAsset* FindReusablePrivateVariable(FStructProperty* ForParameterProperty, UCameraNode* ForCameraNode)
@@ -187,7 +205,8 @@ void SetupPrivateVariable(
 
 	// Set the variable on both the interface parameter and the camera node. Flag them as modified
 	// if we actually changed anything.
-	if (InterfaceParameter->PrivateVariable != PrivateVariable)
+	const bool bShouldModifyInterfaceParameter = Builder.ReuseInterfaceParameter(InterfaceParameter, PrivateVariable);
+	if (bShouldModifyInterfaceParameter)
 	{
 		InterfaceParameter->Modify();
 	}
@@ -268,6 +287,25 @@ void FCameraRigAssetBuilder::FlattenCameraNodeHierarchy()
 			}
 		}
 	}
+
+#if WITH_EDITORONLY_DATA
+	// Check that all the camera nodes that are in the tree are also inside the camera 
+	// rig's AllNodeTreeObjects. This shouldn't happen unless someone added camera nodes
+	// directly via C++, or if there's a bug in the camera rig editor code, so emit a
+	// warning if that happens.
+	TSet<UObject*> FlattenedNodesSet(MakeArrayView((UObject**)FlattenedNodes.GetData(), FlattenedNodes.Num()));
+	TSet<UObject*> AllNodeTreeObjectsSet(ObjectPtrDecay(CameraRig->AllNodeTreeObjects));
+	TSet<UObject*> MissingNodeTreeObjects = FlattenedNodesSet.Difference(AllNodeTreeObjectsSet);
+	if (!MissingNodeTreeObjects.IsEmpty())
+	{
+		BuildLog.AddMessage(EMessageSeverity::Warning, 
+				FText::Format(
+					LOCTEXT("AllNodeTreeObjectsMismatch", 
+						"Found {0} nodes missing from the internal list. Please re-save the asset."),
+					MissingNodeTreeObjects.Num()));
+		CameraRig->AllNodeTreeObjects.Append(MissingNodeTreeObjects.Array());
+	}
+#endif  // WITH_EDITORONLY_DATA
 }
 
 void FCameraRigAssetBuilder::GatherOldDrivenParameters()
@@ -282,9 +320,23 @@ void FCameraRigAssetBuilder::GatherOldDrivenParameters()
 	// Note that parameters driven by user-defined variables are left alone.
 
 	OldDrivenParameters.Reset();
-
-	for (UCameraNode* CameraNode : FlattenedNodes)
+	TSet<UCameraVariableAsset*> GatheredVariables;
+
+	// Get our camera nodes from AllNodeTreeObjects if possible, so we also gather old 
+	// driven parameters from disconnected camera nodes.
+#if WITH_EDITORONLY_DATA
+	TArray<UObject*> ObjectsToGather(ObjectPtrDecay(CameraRig->AllNodeTreeObjects));
+#else
+	TArray<UObject*> ObjectsToGather(FlattenedNodes);
+#endif
+	for (UObject* Object : ObjectsToGather)
 	{
+		UCameraNode* CameraNode = Cast<UCameraNode>(Object);
+		if (!CameraNode)
+		{
+			continue;
+		}
+
 		UClass* CameraNodeClass = CameraNode->GetClass();
 		
 		for (TFieldIterator<FProperty> It(CameraNodeClass); It; ++It)
@@ -307,6 +359,7 @@ void FCameraRigAssetBuilder::GatherOldDrivenParameters()
 						OldDrivenParameters.Add(\
 								FDrivenParameterKey{ StructProperty, CameraNode },\
 								CameraParameterPtr->Variable);\
+						GatheredVariables.Add(CameraParameterPtr->Variable);\
 						CameraParameterPtr->Variable = nullptr;\
 					}\
 				}\
@@ -320,16 +373,46 @@ UE_CAMERA_VARIABLE_FOR_ALL_TYPES()
 		}
 	}
 
+	// Look for interface parameters that were disconnected in the graph editor. They may still 
+	// have a private variable inside them. We will discard those, but we want to rename them so
+	// that a new one can take their place (such as when another interface parameter gets connected
+	// to a similar camera node property and has the same name).
+	OldInterfaceParameters.Reset();
+
+#if WITH_EDITORONLY_DATA
+	// Leave ObjectsToGather as AllNodeTreeObjects.
+#else
+	ObjectsToGather.Reset();
+	for (UCameraRigInterfaceParameter* InterfaceParameter : CameraRig->Interface.InterfaceParameters)
+	{
+		ObjectsToGather.Add(InterfaceParameter);
+	}
+#endif
+	for (UObject* Object : ObjectsToGather)
+	{
+		UCameraRigInterfaceParameter* InterfaceParameter = Cast<UCameraRigInterfaceParameter>(Object);
+		if (!InterfaceParameter)
+		{
+			continue;
+		}
+
+		OldInterfaceParameters.Add(InterfaceParameter, FReusableInterfaceParameterInfo(InterfaceParameter->PrivateVariable, false));
+		if (InterfaceParameter->PrivateVariable)
+		{
+			GatheredVariables.Add(InterfaceParameter->PrivateVariable);
+		}
+	}
+
 	// Temporarily rename all old camera variables, so their names are available to the new
 	// driven parameters.
-	for (TPair<FDrivenParameterKey, UCameraVariableAsset*> Pair : OldDrivenParameters)
+	for (UCameraVariableAsset* GatheredVariable : GatheredVariables)
 	{
 		TStringBuilder<256> StringBuilder;
 		StringBuilder.Append("REUSABLE_");
-		StringBuilder.Append(Pair.Value->GetName());
-		// Rename non-transactionally because if nothing has change, we will rename it back
+		StringBuilder.Append(GatheredVariable->GetName());
+		// Rename non-transactionally because if nothing has changed, we will rename it back
 		// later and we don't want to dirty the package for nothing.
-		Pair.Value->Rename(StringBuilder.ToString(), nullptr, REN_NonTransactional);
+		GatheredVariable->Rename(StringBuilder.ToString(), nullptr, REN_NonTransactional);
 	}
 }
 
@@ -339,6 +422,9 @@ void FCameraRigAssetBuilder::BuildNewDrivenParameters()
 
 	TSet<FString> UsedInterfaceParameterNames;
 
+	using FBuiltDrivenParameter = TTuple<UCameraNode*, FName>;
+	TSet<FBuiltDrivenParameter> BuiltDrivenParameters;
+
 	const FString CameraRigName = CameraRig->GetName();
 	const FString CameraRigPathName = CameraRig->GetPathName();
 
@@ -396,6 +482,21 @@ void FCameraRigAssetBuilder::BuildNewDrivenParameters()
 		}
 		UsedInterfaceParameterNames.Add(InterfaceParameter->InterfaceParameterName);
 
+		// Check duplicate targets.
+		FBuiltDrivenParameter BuiltDrivenParameter(InterfaceParameter->Target, InterfaceParameter->TargetPropertyName);
+		if (BuiltDrivenParameters.Contains(BuiltDrivenParameter))
+		{
+			BuildLog.AddMessage(EMessageSeverity::Error,
+					InterfaceParameter,
+					FText::Format(LOCTEXT(
+						"InterfaceParameterTargetCollision",
+						"Multiple interface parameters targeting property '{0}' on camera node '{1}'. Ignoring duplicates."),
+						FText::FromName(InterfaceParameter->Target->GetFName()),
+						FText::FromName(InterfaceParameter->TargetPropertyName)));
+			continue;
+		}
+		BuiltDrivenParameters.Add(BuiltDrivenParameter);
+
 		// Get the target camera node property and check that it is a camera parameter struct.
 		UCameraNode* Target = InterfaceParameter->Target;
 		UClass* TargetClass = Target->GetClass();
diff --git a/Engine/Plugins/Cameras/GameplayCameras/Source/GameplayCameras/Public/Core/CameraRigAssetBuilder.h b/Engine/Plugins/Cameras/GameplayCameras/Source/GameplayCameras/Public/Core/CameraRigAssetBuilder.h
--- a/Engine/Plugins/Cameras/GameplayCameras/Source/GameplayCameras/Public/Core/CameraRigAssetBuilder.h
+++ b/Engine/Plugins/Cameras/GameplayCameras/Source/GameplayCameras/Public/Core/CameraRigAssetBuilder.h
@@ -71,6 +71,9 @@ private:
 	using FDrivenParameterKey = TTuple<FStructProperty*, UCameraNode*>;
 	TMap<FDrivenParameterKey, UCameraVariableAsset*> OldDrivenParameters;
 
+	using FReusableInterfaceParameterInfo = TTuple<UCameraVariableAsset*, bool>;
+	TMap<UCameraRigInterfaceParameter*, FReusableInterfaceParameterInfo> OldInterfaceParameters;
+
 	FCameraRigAllocationInfo AllocationInfo;
 
 	friend struct Internal::FPrivateVariableBuilder;
