Cameras: fix directors incorrectly being created with the RF_Standalone flag

diff --git a/Engine/Plugins/Cameras/GameplayCameras/Source/GameplayCameras/Private/Core/CameraAsset.cpp b/Engine/Plugins/Cameras/GameplayCameras/Source/GameplayCameras/Private/Core/CameraAsset.cpp
--- a/Engine/Plugins/Cameras/GameplayCameras/Source/GameplayCameras/Private/Core/CameraAsset.cpp
+++ b/Engine/Plugins/Cameras/GameplayCameras/Source/GameplayCameras/Private/Core/CameraAsset.cpp
@@ -4,6 +4,7 @@
 
 #include "Core/CameraAssetBuilder.h"
 #include "Core/CameraBuildLog.h"
+#include "Core/CameraDirector.h"
 
 #include UE_INLINE_GENERATED_CPP_BY_NAME(CameraAsset)
 
@@ -104,6 +105,24 @@ int32 UCameraAsset::RemoveExitTransition(UCameraRigTransition* InTransition)
 	return NumRemoved;
 }
 
+void UCameraAsset::PostLoad()
+{
+	Super::PostLoad();
+
+	if (CameraDirector)
+	{
+		EObjectFlags Flags = CameraDirector->GetFlags();
+		if (EnumHasAnyFlags(Flags, (RF_Public | RF_Standalone)))
+		{
+			UE_LOG(LogCameraSystem, Warning, 
+					TEXT("Removing incorrect object flags from camera director inside '%s', please re-save the asset."),
+					*GetPathNameSafe(this));
+			CameraDirector->Modify();
+			CameraDirector->ClearFlags(RF_Public | RF_Standalone);
+		}
+	}
+}
+
 #if WITH_EDITOR
 
 void UCameraAsset::PostEditChangeProperty(FPropertyChangedEvent& PropertyChangedEvent)
diff --git a/Engine/Plugins/Cameras/GameplayCameras/Source/GameplayCameras/Public/Core/CameraAsset.h b/Engine/Plugins/Cameras/GameplayCameras/Source/GameplayCameras/Public/Core/CameraAsset.h
--- a/Engine/Plugins/Cameras/GameplayCameras/Source/GameplayCameras/Public/Core/CameraAsset.h
+++ b/Engine/Plugins/Cameras/GameplayCameras/Source/GameplayCameras/Public/Core/CameraAsset.h
@@ -124,6 +124,7 @@ protected:
 #endif
 
 	// UObject interface.
+	virtual void PostLoad() override;
 #if WITH_EDITOR
 	virtual void PostEditChangeProperty(FPropertyChangedEvent& PropertyChangedEvent) override;
 #endif
diff --git a/Engine/Plugins/Cameras/GameplayCameras/Source/GameplayCamerasEditor/Private/Factories/CameraAssetFactory.cpp b/Engine/Plugins/Cameras/GameplayCameras/Source/GameplayCamerasEditor/Private/Factories/CameraAssetFactory.cpp
--- a/Engine/Plugins/Cameras/GameplayCameras/Source/GameplayCamerasEditor/Private/Factories/CameraAssetFactory.cpp
+++ b/Engine/Plugins/Cameras/GameplayCameras/Source/GameplayCamerasEditor/Private/Factories/CameraAssetFactory.cpp
@@ -65,7 +65,7 @@ UObject* UCameraAssetFactory::FactoryCreateNew(UClass* Class, UObject* Parent, F
 
 	if (CameraDirectorClass)
 	{
-		UCameraDirector* NewCameraDirector = NewObject<UCameraDirector>(NewCameraAsset, CameraDirectorClass, NAME_None, Flags | RF_Transactional);
+		UCameraDirector* NewCameraDirector = NewObject<UCameraDirector>(NewCameraAsset, CameraDirectorClass, NAME_None, RF_Transactional);
 		NewCameraAsset->SetCameraDirector(NewCameraDirector);
 
 		// Let the camera director do some scaffolding.
