Cameras: creating a variable should be done in an undo/redo transaction

diff --git a/Engine/Plugins/Cameras/GameplayCameras/Source/GameplayCamerasEditor/Private/Toolkits/CameraVariableCollectionEditorToolkit.cpp b/Engine/Plugins/Cameras/GameplayCameras/Source/GameplayCamerasEditor/Private/Toolkits/CameraVariableCollectionEditorToolkit.cpp
--- a/Engine/Plugins/Cameras/GameplayCameras/Source/GameplayCamerasEditor/Private/Toolkits/CameraVariableCollectionEditorToolkit.cpp
+++ b/Engine/Plugins/Cameras/GameplayCameras/Source/GameplayCamerasEditor/Private/Toolkits/CameraVariableCollectionEditorToolkit.cpp
@@ -253,6 +253,8 @@ FLinearColor FCameraVariableCollectionEditorToolkit::GetWorldCentricTabColorScal
 
 void FCameraVariableCollectionEditorToolkit::OnCreateVariable(TSubclassOf<UCameraVariableAsset> InVariableClass)
 {
+	FScopedTransaction CreateTransaction(LOCTEXT("CreateVariable", "Create camera variable"));
+
 	UCameraVariableAsset* NewVariable = NewObject<UCameraVariableAsset>(
 			VariableCollection, 
 			InVariableClass.Get(),
