Cameras: use a proper editable textbox for camera variable names in the variable collection editor

Also add a context menu to the list view for renaming the variable there too.

diff --git a/Engine/Plugins/Cameras/GameplayCameras/Source/GameplayCameras/Public/Core/CameraVariableAssets.h b/Engine/Plugins/Cameras/GameplayCameras/Source/GameplayCameras/Public/Core/CameraVariableAssets.h
--- a/Engine/Plugins/Cameras/GameplayCameras/Source/GameplayCameras/Public/Core/CameraVariableAssets.h
+++ b/Engine/Plugins/Cameras/GameplayCameras/Source/GameplayCameras/Public/Core/CameraVariableAssets.h
@@ -48,7 +48,7 @@ public:
 
 #if WITH_EDITORONLY_DATA
 	/** The name of the variable. */
-	UPROPERTY(EditAnywhere, Category=Camera)
+	UPROPERTY()
 	FString DisplayName;
 #endif  // WITH_EDITORONLY_DATA
 
diff --git a/Engine/Plugins/Cameras/GameplayCameras/Source/GameplayCamerasEditor/Private/Editors/SCameraVariableCollectionEditor.cpp b/Engine/Plugins/Cameras/GameplayCameras/Source/GameplayCamerasEditor/Private/Editors/SCameraVariableCollectionEditor.cpp
--- a/Engine/Plugins/Cameras/GameplayCameras/Source/GameplayCamerasEditor/Private/Editors/SCameraVariableCollectionEditor.cpp
+++ b/Engine/Plugins/Cameras/GameplayCameras/Source/GameplayCamerasEditor/Private/Editors/SCameraVariableCollectionEditor.cpp
@@ -5,15 +5,14 @@
 #include "Commands/CameraVariableCollectionEditorCommands.h"
 #include "Core/CameraVariableAssets.h"
 #include "Core/CameraVariableCollection.h"
-#include "Framework/Views/TableViewMetadata.h"
 #include "IDetailsView.h"
+#include "ScopedTransaction.h"
 #include "Styles/GameplayCamerasEditorStyle.h"
 #include "ToolMenus.h"
-#include "UObject/UObjectIterator.h"
 #include "Widgets/Input/SSearchBox.h"
 #include "Widgets/Layout/SBorder.h"
-#include "Widgets/Layout/SSpacer.h"
 #include "Widgets/SBoxPanel.h"
+#include "Widgets/Text/SInlineEditableTextBlock.h"
 #include "Widgets/Views/SListView.h"
 
 #define LOCTEXT_NAMESPACE "SCameraVariableCollectionEditor"
@@ -25,6 +24,8 @@ class SCameraVariableCollectionListRow : public SMultiColumnTableRow<UCameraVari
 {
 public:
 
+	using FSuperRowType = SMultiColumnTableRow<UCameraVariableAsset*>;
+
 	SLATE_BEGIN_ARGS(SCameraVariableCollectionListRow)
 		: _CameraVariable(nullptr)
 	{}
@@ -36,6 +37,8 @@ public:
 
 	void Construct(const FArguments& InArgs, const TSharedRef<STableViewBase>& OwnerTableView);
 
+	void EnterNameEditingMode();
+
 protected:
 
 	virtual TSharedRef<SWidget> GenerateWidgetForColumn(const FName& ColumnName) override;
@@ -46,10 +49,15 @@ private:
 	FText GetVariableType() const;
 	FText GetDefaultValue() const;
 
+	bool OnVerifyVariableNameChanged(const FText& Text, FText& OutErrorMessage);
+	void OnVariableNameCommitted(const FText& Text, ETextCommit::Type CommitType);
+
 private:
 
 	UCameraVariableAsset* CameraVariable = nullptr;
 
+	TSharedPtr<SInlineEditableTextBlock> EditableTextBlock;
+
 	TAttribute<FText> HighlightText;
 };
 
@@ -63,6 +71,11 @@ void SCameraVariableCollectionListRow::Construct(const FArguments& InArgs, const
 		OwnerTableView);
 }
 
+void SCameraVariableCollectionListRow::EnterNameEditingMode()
+{
+	EditableTextBlock->EnterEditingMode();
+}
+
 TSharedRef<SWidget> SCameraVariableCollectionListRow::GenerateWidgetForColumn(const FName& ColumnName)
 {
 	TSharedRef<FGameplayCamerasEditorStyle> CamerasEditorStyle = FGameplayCamerasEditorStyle::Get();
@@ -75,10 +88,13 @@ TSharedRef<SWidget> SCameraVariableCollectionListRow::GenerateWidgetForColumn(co
 			.Padding(8.f)
 			.VAlign(VAlign_Center)
 			[
-				SNew(STextBlock)
-				.TextStyle(CamerasEditorStyle, "CameraVariableCollectionEditor.Entry.Name")
+				SAssignNew(EditableTextBlock, SInlineEditableTextBlock)
+				.Style(CamerasEditorStyle, "CameraVariableCollectionEditor.Entry.Name")
 				.Text(this, &SCameraVariableCollectionListRow::GetVariableName)
+				.OnTextCommitted(this, &SCameraVariableCollectionListRow::OnVariableNameCommitted)
+				.OnVerifyTextChanged(this, &SCameraVariableCollectionListRow::OnVerifyVariableNameChanged)
 				.HighlightText(HighlightText)
+				.IsSelected(this, &FSuperRowType::IsSelectedExclusively)
 			];
 	}
 	else if (ColumnName == TEXT("VariableType"))
@@ -130,12 +146,40 @@ FText SCameraVariableCollectionListRow::GetDefaultValue() const
 	return FText::FromString(CameraVariable->FormatDefaultValue());
 }
 
+bool SCameraVariableCollectionListRow::OnVerifyVariableNameChanged(const FText& Text, FText& OutErrorMessage)
+{
+	return true;
+}
+
+void SCameraVariableCollectionListRow::OnVariableNameCommitted(const FText& Text, ETextCommit::Type CommitType)
+{
+	if (CameraVariable)
+	{
+		FScopedTransaction Transaction(LOCTEXT("RenameCameraVariable", "Rename camera variable"));
+
+		const FString NewDisplayName = Text.ToString();
+		CameraVariable->Modify();
+		CameraVariable->DisplayName = NewDisplayName;
+
+		// We just set the display name, but also rename the object itself, which helps for debugging.
+		FName NewObjectName = MakeObjectNameFromDisplayLabel(NewDisplayName, CameraVariable->GetFName());
+		NewObjectName = MakeUniqueObjectName(CameraVariable->GetOuter(), CameraVariable->GetClass(), NewObjectName);
+		CameraVariable->Rename(*NewObjectName.ToString());
+	}
+}
+
 void SCameraVariableCollectionEditor::Construct(const FArguments& InArgs)
 {
 	VariableCollection = InArgs._VariableCollection;
 
 	WeakDetailsView = InArgs._DetailsView;
 
+	CommandList = MakeShared<FUICommandList>();
+	if (InArgs._AdditionalCommands)
+	{
+		CommandList->Append(InArgs._AdditionalCommands.ToSharedRef());
+	}
+
 	SearchTextFilter = MakeShareable(new FEntryTextFilter(
 				FEntryTextFilter::FItemToStringArray::CreateSP(this, &SCameraVariableCollectionEditor::GetEntryStrings)));
 
@@ -164,6 +208,8 @@ void SCameraVariableCollectionEditor::Construct(const FArguments& InArgs)
 			.ListItemsSource(&FilteredItemSource)
 			.OnGenerateRow(this, &SCameraVariableCollectionEditor::OnListGenerateRow)
 			.OnSelectionChanged(this, &SCameraVariableCollectionEditor::OnListSectionChanged)
+			.OnItemScrolledIntoView(this, &SCameraVariableCollectionEditor::OnListItemScrolledIntoView)
+			.OnContextMenuOpening(this, &SCameraVariableCollectionEditor::OnListContextMenuOpening)
 			.HeaderRow
 			(
 				SNew(SHeaderRow)
@@ -197,6 +243,18 @@ void SCameraVariableCollectionEditor::GetSelectedVariables(TArray<UCameraVariabl
 	ListView->GetSelectedItems(OutSelection);
 }
 
+void SCameraVariableCollectionEditor::RequestRenameSelectedVariable()
+{
+	TArray<UCameraVariableAsset*> SelectedVariables = ListView->GetSelectedItems();
+	if (SelectedVariables.IsEmpty())
+	{
+		return;
+	}
+
+	bDeferredRequestRenameItem = true;
+	ListView->RequestScrollIntoView(SelectedVariables[0]);
+}
+
 void SCameraVariableCollectionEditor::RequestListRefresh()
 {
 	bUpdateFilteredItemSource = true;
@@ -254,6 +312,50 @@ void SCameraVariableCollectionEditor::OnListSectionChanged(UCameraVariableAsset*
 	SetDetailsViewObject(Item);
 }
 
+void SCameraVariableCollectionEditor::OnListItemScrolledIntoView(UCameraVariableAsset* Item, const TSharedPtr<ITableRow>& ItemWidget)
+{
+	if (bDeferredRequestRenameItem)
+	{
+		bDeferredRequestRenameItem = false;
+
+		TSharedPtr<ITableRow> RowWidget = ListView->WidgetFromItem(Item);
+		if (!ensure(RowWidget))
+		{
+			return;
+		}
+
+		TSharedPtr<SCameraVariableCollectionListRow> TypedRowWidget = StaticCastSharedPtr<SCameraVariableCollectionListRow>(RowWidget);
+		if (!ensure(TypedRowWidget))
+		{
+			return;
+		}
+
+		TypedRowWidget->EnterNameEditingMode();
+	}
+}
+
+TSharedPtr<SWidget> SCameraVariableCollectionEditor::OnListContextMenuOpening()
+{
+	static const FName ContextMenuName("CameraVariableList.ContextMenu");
+
+	UToolMenus* ToolMenus = UToolMenus::Get();
+
+	if (!ToolMenus->IsMenuRegistered(ContextMenuName))
+	{
+		const FCameraVariableCollectionEditorCommands& Commands = FCameraVariableCollectionEditorCommands::Get();
+
+		UToolMenu* ContextMenu = ToolMenus->RegisterMenu(ContextMenuName, NAME_None, EMultiBoxType::Menu);
+
+		FToolMenuSection& Section = ContextMenu->AddSection("Actions");
+		Section.AddEntry(FToolMenuEntry::InitMenuEntry(Commands.RenameVariable));
+		Section.AddEntry(FToolMenuEntry::InitMenuEntry(Commands.DeleteVariable));
+	}
+
+	FToolMenuContext MenuContext;
+	MenuContext.AppendCommandList(CommandList);
+	return ToolMenus->GenerateWidget(ContextMenuName, MenuContext);
+}
+
 void SCameraVariableCollectionEditor::GetEntryStrings(const UCameraVariableAsset* InItem, TArray<FString>& OutStrings)
 {
 	OutStrings.Add(InItem->GetName());
diff --git a/Engine/Plugins/Cameras/GameplayCameras/Source/GameplayCamerasEditor/Private/Styles/GameplayCamerasEditorStyle.cpp b/Engine/Plugins/Cameras/GameplayCameras/Source/GameplayCamerasEditor/Private/Styles/GameplayCamerasEditorStyle.cpp
--- a/Engine/Plugins/Cameras/GameplayCameras/Source/GameplayCamerasEditor/Private/Styles/GameplayCamerasEditorStyle.cpp
+++ b/Engine/Plugins/Cameras/GameplayCameras/Source/GameplayCamerasEditor/Private/Styles/GameplayCamerasEditorStyle.cpp
@@ -33,6 +33,7 @@ FGameplayCamerasEditorStyle::FGameplayCamerasEditorStyle()
 
 	const ISlateStyle& CoreStyle = FAppStyle::Get();
 	const FTextBlockStyle& NormalText = CoreStyle.GetWidgetStyle<FTextBlockStyle>("NormalText");
+	const FInlineEditableTextBlockStyle& NormalInlineEditableText = CoreStyle.GetWidgetStyle<FInlineEditableTextBlockStyle>("InlineEditableTextBlockStyle");
 	const FTableRowStyle AlternatingTableRowStyle = CoreStyle.GetWidgetStyle<FTableRowStyle>("TableView.AlternatingRow");
 
 	// Camera assets.
@@ -92,9 +93,11 @@ FGameplayCamerasEditorStyle::FGameplayCamerasEditorStyle()
 	Set("CameraVariableCollectionEditor.RenameVariable", new CORE_IMAGE_BRUSH_SVG("Starship/Common/Rename", Icon16x16));
 	Set("CameraVariableCollectionEditor.DeleteVariable", new CORE_IMAGE_BRUSH_SVG("Starship/Common/minus", Icon16x16));
 
-	Set("CameraVariableCollectionEditor.Entry.Name", FTextBlockStyle(NormalText) .SetFont(DEFAULT_FONT("Bold", 12)));
-	Set("CameraVariableCollectionEditor.Entry.Type", FTextBlockStyle(NormalText) .SetFont(DEFAULT_FONT("Italic", 10)));
-	Set("CameraVariableCollectionEditor.Entry.Value", FTextBlockStyle(NormalText) .SetFont(DEFAULT_FONT("Regular", 10)));
+	FInlineEditableTextBlockStyle CameraVariableCollectionEntryNameStyle(NormalInlineEditableText);
+	CameraVariableCollectionEntryNameStyle.TextStyle.SetFont(DEFAULT_FONT("Bold", 12));
+	Set("CameraVariableCollectionEditor.Entry.Name", CameraVariableCollectionEntryNameStyle);
+	Set("CameraVariableCollectionEditor.Entry.Type", FTextBlockStyle(NormalText).SetFont(DEFAULT_FONT("Italic", 10)));
+	Set("CameraVariableCollectionEditor.Entry.Value", FTextBlockStyle(NormalText).SetFont(DEFAULT_FONT("Regular", 10)));
 
 	// Debugger tool icons.
 	Set("Debugger.DebugInfoEnabled.Icon", new IMAGE_BRUSH_SVG("Icons/DebugInfo-ToggleCheck", Icon16x16, FStyleColors::AccentGreen));
diff --git a/Engine/Plugins/Cameras/GameplayCameras/Source/GameplayCamerasEditor/Private/Toolkits/CameraVariableCollectionEditorToolkit.cpp b/Engine/Plugins/Cameras/GameplayCameras/Source/GameplayCamerasEditor/Private/Toolkits/CameraVariableCollectionEditorToolkit.cpp
--- a/Engine/Plugins/Cameras/GameplayCameras/Source/GameplayCamerasEditor/Private/Toolkits/CameraVariableCollectionEditorToolkit.cpp
+++ b/Engine/Plugins/Cameras/GameplayCameras/Source/GameplayCamerasEditor/Private/Toolkits/CameraVariableCollectionEditorToolkit.cpp
@@ -36,7 +36,6 @@ const FName FCameraVariableCollectionEditorToolkit::DetailsViewTabId(TEXT("Camer
 FCameraVariableCollectionEditorToolkit::FCameraVariableCollectionEditorToolkit(UCameraVariableCollectionEditor* InOwningAssetEditor)
 	: FBaseAssetToolkit(InOwningAssetEditor)
 	, VariableCollection(InOwningAssetEditor->GetVariableCollection())
-	, CommandBindings(new FUICommandList())
 {
 	// Override base class default layout.
 	StandaloneDefaultLayout = FTabManager::NewLayout("CameraVariableCollectionEditor_Layout_v1")
@@ -129,7 +128,8 @@ void FCameraVariableCollectionEditorToolkit::CreateWidgets()
 	// Create the variable collection editor.
 	VariableCollectionEditorWidget = SNew(SCameraVariableCollectionEditor)
 		.DetailsView(DetailsView)
-		.VariableCollection(VariableCollection);
+		.VariableCollection(VariableCollection)
+		.AdditionalCommands(ToolkitCommands);
 }
 
 void FCameraVariableCollectionEditorToolkit::RegisterToolbar()
@@ -268,12 +268,7 @@ void FCameraVariableCollectionEditorToolkit::OnCreateVariable(TSubclassOf<UCamer
 
 void FCameraVariableCollectionEditorToolkit::OnRenameVariable()
 {
-	UClass* VariableAssetClass = UCameraVariableAsset::StaticClass();
-	FProperty* DisplayNameProperty = VariableAssetClass->FindPropertyByName(GET_MEMBER_NAME_CHECKED(UCameraVariableAsset, DisplayName));
-
-	FPropertyPath PropertyPath;
-	PropertyPath.AddProperty(FPropertyInfo(DisplayNameProperty));
-	DetailsView->HighlightProperty(PropertyPath);
+	VariableCollectionEditorWidget->RequestRenameSelectedVariable();
 }
 
 bool FCameraVariableCollectionEditorToolkit::CanRenameVariable()
diff --git a/Engine/Plugins/Cameras/GameplayCameras/Source/GameplayCamerasEditor/Public/Editors/SCameraVariableCollectionEditor.h b/Engine/Plugins/Cameras/GameplayCameras/Source/GameplayCamerasEditor/Public/Editors/SCameraVariableCollectionEditor.h
--- a/Engine/Plugins/Cameras/GameplayCameras/Source/GameplayCamerasEditor/Public/Editors/SCameraVariableCollectionEditor.h
+++ b/Engine/Plugins/Cameras/GameplayCameras/Source/GameplayCamerasEditor/Public/Editors/SCameraVariableCollectionEditor.h
@@ -2,14 +2,13 @@
 
 #pragma once
 
-#include "CoreTypes.h"
 #include "Misc/TextFilter.h"
 #include "Templates/SharedPointerFwd.h"
-#include "Templates/SubclassOf.h"
 #include "UObject/ObjectPtr.h"
 #include "Widgets/SCompoundWidget.h"
 
 class FAssetEditorToolkit;
+class FUICommandList;
 class IDetailsView;
 class ITableRow;
 class SSearchBox;
@@ -37,6 +36,8 @@ public:
 		SLATE_ARGUMENT(TWeakPtr<IDetailsView>, DetailsView)
 		/** The toolkit inside which this editor lives, if any. */
 		SLATE_ARGUMENT(TWeakPtr<FAssetEditorToolkit>, AssetEditorToolkit)
+		/** Command bindings for manipulating camera variables. */
+		SLATE_ARGUMENT(TSharedPtr<FUICommandList>, AdditionalCommands)
 	SLATE_END_ARGS()
 
 	void Construct(const FArguments& InArgs);
@@ -47,6 +48,9 @@ public:
 	/** Gets the selected variables in the list view. */
 	void GetSelectedVariables(TArray<UCameraVariableAsset*>& OutSelection) const;
 
+	/** Enter editing mode for the first currently selected variable's name. */
+	void RequestRenameSelectedVariable();
+
 	/** Requests that the list view be refreshed by the next tick. */
 	void RequestListRefresh();
 
@@ -62,6 +66,8 @@ private:
 
 	TSharedRef<ITableRow> OnListGenerateRow(UCameraVariableAsset* Item, const TSharedRef<STableViewBase>& OwnerTable);
 	void OnListSectionChanged(UCameraVariableAsset* Item, ESelectInfo::Type SelectInfo) const;
+	void OnListItemScrolledIntoView(UCameraVariableAsset* Item, const TSharedPtr<ITableRow>& ItemWidget);
+	TSharedPtr<SWidget> OnListContextMenuOpening();
 
 	void GetEntryStrings(const UCameraVariableAsset* InItem, TArray<FString>& OutStrings);
 	void OnSearchTextChanged(const FText& InFilterText);
@@ -74,6 +80,8 @@ private:
 
 	TWeakPtr<IDetailsView> WeakDetailsView;
 
+	TSharedPtr<FUICommandList> CommandList;
+
 	TSharedPtr<SListView<UCameraVariableAsset*>> ListView;
 
 	TArray<UCameraVariableAsset*> FilteredItemSource;
@@ -83,6 +91,7 @@ private:
 	TSharedPtr<SSearchBox> SearchBox;
 
 	bool bUpdateFilteredItemSource = false;
+	bool bDeferredRequestRenameItem = false;
 };
 
 }  // namespace UE::Cameras
diff --git a/Engine/Plugins/Cameras/GameplayCameras/Source/GameplayCamerasEditor/Public/Toolkits/CameraVariableCollectionEditorToolkit.h b/Engine/Plugins/Cameras/GameplayCameras/Source/GameplayCamerasEditor/Public/Toolkits/CameraVariableCollectionEditorToolkit.h
--- a/Engine/Plugins/Cameras/GameplayCameras/Source/GameplayCamerasEditor/Public/Toolkits/CameraVariableCollectionEditorToolkit.h
+++ b/Engine/Plugins/Cameras/GameplayCameras/Source/GameplayCamerasEditor/Public/Toolkits/CameraVariableCollectionEditorToolkit.h
@@ -77,9 +77,6 @@ private:
 	/** The asset being edited */
 	TObjectPtr<UCameraVariableCollection> VariableCollection;
 
-	/** Command bindings */
-	TSharedRef<FUICommandList> CommandBindings;
-
 	/** Camera variable collection editor widget */
 	TSharedPtr<SCameraVariableCollectionEditor> VariableCollectionEditorWidget;
 };
