Cameras: cleanup dead code

diff --git a/Engine/Plugins/Cameras/GameplayCameras/Source/GameplayCamerasEditor/Private/Toolkits/CameraRigAssetEditorToolkitBase.cpp b/Engine/Plugins/Cameras/GameplayCameras/Source/GameplayCamerasEditor/Private/Toolkits/CameraRigAssetEditorToolkitBase.cpp
--- a/Engine/Plugins/Cameras/GameplayCameras/Source/GameplayCamerasEditor/Private/Toolkits/CameraRigAssetEditorToolkitBase.cpp
+++ b/Engine/Plugins/Cameras/GameplayCameras/Source/GameplayCamerasEditor/Private/Toolkits/CameraRigAssetEditorToolkitBase.cpp
@@ -26,7 +26,6 @@ const FName FCameraRigAssetEditorToolkitBase::DetailsViewTabId(TEXT("CameraRigAs
 
 FCameraRigAssetEditorToolkitBase::FCameraRigAssetEditorToolkitBase(FName InLayoutName)
 	: StandardLayout(new FStandardToolkitLayout(InLayoutName))
-	, CommandBindings(new FUICommandList())
 {
 	StandardLayout->AddLeftTab(ToolboxTabId);
 	StandardLayout->AddCenterTab(CameraRigEditorTabId);
diff --git a/Engine/Plugins/Cameras/GameplayCameras/Source/GameplayCamerasEditor/Private/Toolkits/CameraRigTransitionEditorToolkitBase.cpp b/Engine/Plugins/Cameras/GameplayCameras/Source/GameplayCamerasEditor/Private/Toolkits/CameraRigTransitionEditorToolkitBase.cpp
--- a/Engine/Plugins/Cameras/GameplayCameras/Source/GameplayCamerasEditor/Private/Toolkits/CameraRigTransitionEditorToolkitBase.cpp
+++ b/Engine/Plugins/Cameras/GameplayCameras/Source/GameplayCamerasEditor/Private/Toolkits/CameraRigTransitionEditorToolkitBase.cpp
@@ -28,7 +28,6 @@ const FName FCameraRigTransitionEditorToolkitBase::DetailsViewTabId(TEXT("Camera
 
 FCameraRigTransitionEditorToolkitBase::FCameraRigTransitionEditorToolkitBase(FName InLayoutName)
 	: StandardLayout(new FStandardToolkitLayout(InLayoutName))
-	, CommandBindings(new FUICommandList())
 {
 	StandardLayout->AddLeftTab(ToolboxTabId);
 	StandardLayout->AddCenterTab(TransitionEditorTabId);
diff --git a/Engine/Plugins/Cameras/GameplayCameras/Source/GameplayCamerasEditor/Public/Toolkits/CameraRigAssetEditorToolkitBase.h b/Engine/Plugins/Cameras/GameplayCameras/Source/GameplayCamerasEditor/Public/Toolkits/CameraRigAssetEditorToolkitBase.h
--- a/Engine/Plugins/Cameras/GameplayCameras/Source/GameplayCamerasEditor/Public/Toolkits/CameraRigAssetEditorToolkitBase.h
+++ b/Engine/Plugins/Cameras/GameplayCameras/Source/GameplayCamerasEditor/Public/Toolkits/CameraRigAssetEditorToolkitBase.h
@@ -2,7 +2,6 @@
 
 #pragma once
 
-#include "CoreTypes.h"
 #include "Misc/NotifyHook.h"
 #include "UObject/GCObject.h"
 
@@ -86,9 +85,6 @@ private:
 	/** The layout for this editor toolkit */
 	TSharedPtr<FStandardToolkitLayout> StandardLayout;
 
-	/** Command bindings */
-	TSharedPtr<FUICommandList> CommandBindings;
-
 	/** The details view */
 	TSharedPtr<IDetailsView> DetailsView;
 
diff --git a/Engine/Plugins/Cameras/GameplayCameras/Source/GameplayCamerasEditor/Public/Toolkits/CameraRigTransitionEditorToolkitBase.h b/Engine/Plugins/Cameras/GameplayCameras/Source/GameplayCamerasEditor/Public/Toolkits/CameraRigTransitionEditorToolkitBase.h
--- a/Engine/Plugins/Cameras/GameplayCameras/Source/GameplayCamerasEditor/Public/Toolkits/CameraRigTransitionEditorToolkitBase.h
+++ b/Engine/Plugins/Cameras/GameplayCameras/Source/GameplayCamerasEditor/Public/Toolkits/CameraRigTransitionEditorToolkitBase.h
@@ -2,15 +2,12 @@
 
 #pragma once
 
-#include "CoreTypes.h"
-#include "Editors/CameraRigTransitionGraphSchema.h"
 #include "Misc/NotifyHook.h"
+#include "Templates/SubclassOf.h"
 #include "UObject/GCObject.h"
-#include "UObject/ScriptInterface.h"
 
 class FSpawnTabArgs;
 class FTabManager;
-class FUICommandList;
 class FWorkspaceItem;
 class IDetailsView;
 class SDockTab;
@@ -18,6 +15,7 @@ class SObjectTreeGraphToolbox;
 class UCameraRigTransitionGraphSchemaBase;
 class UToolMenu;
 struct FEdGraphEditAction;
+struct FGraphAppearanceInfo;
 
 namespace UE::Cameras
 {
@@ -79,9 +77,7 @@ private:
 	/** The layout for this editor */
 	TSharedPtr<FStandardToolkitLayout> StandardLayout;
 
-	/** Command bindings */
-	TSharedPtr<FUICommandList> CommandBindings;
-
+	/** The details view for this editor */
 	TSharedPtr<IDetailsView> DetailsView;
 
 	/** Camera transition editor widget */
