Cameras: fix undo/redo bugs caused by missing calls to Modify()

New tentative rules:

- The graph nodes and schema should call Modify() on graph nodes and graph pins when they are being modified.
- The graph schema handles modification of the underlying object model and should call Modify() on that when needed.
- Calls to the IObjectTreeGraphObject interface should not by default call Modify(), as the implementers are free to indeed modify, or not, the object.

diff --git a/Engine/Plugins/Cameras/GameplayCameras/Source/GameplayCameras/Private/Core/CameraAsset.cpp b/Engine/Plugins/Cameras/GameplayCameras/Source/GameplayCameras/Private/Core/CameraAsset.cpp
--- a/Engine/Plugins/Cameras/GameplayCameras/Source/GameplayCameras/Private/Core/CameraAsset.cpp
+++ b/Engine/Plugins/Cameras/GameplayCameras/Source/GameplayCameras/Private/Core/CameraAsset.cpp
@@ -202,6 +202,8 @@ const FString& UCameraAsset::GetGraphNodeCommentText(FName InGraphName) const
 
 void UCameraAsset::OnUpdateGraphNodeCommentText(FName InGraphName, const FString& NewComment)
 {
+	Modify();
+
 	TransitionGraphNodeComment = NewComment;
 }
 
diff --git a/Engine/Plugins/Cameras/GameplayCameras/Source/GameplayCameras/Private/Core/CameraNode.cpp b/Engine/Plugins/Cameras/GameplayCameras/Source/GameplayCameras/Private/Core/CameraNode.cpp
--- a/Engine/Plugins/Cameras/GameplayCameras/Source/GameplayCameras/Private/Core/CameraNode.cpp
+++ b/Engine/Plugins/Cameras/GameplayCameras/Source/GameplayCameras/Private/Core/CameraNode.cpp
@@ -51,6 +51,8 @@ void UCameraNode::GetGraphNodePosition(FName InGraphName, int32& NodePosX, int32
 
 void UCameraNode::OnGraphNodeMoved(FName InGraphName, int32 NodePosX, int32 NodePosY, bool bMarkDirty)
 {
+	Modify(bMarkDirty);
+
 	GraphNodePos.X = NodePosX;
 	GraphNodePos.Y = NodePosY;
 }
@@ -62,6 +64,8 @@ const FString& UCameraNode::GetGraphNodeCommentText(FName InGraphName) const
 
 void UCameraNode::OnUpdateGraphNodeCommentText(FName InGraphName, const FString& NewComment)
 {
+	Modify();
+
 	GraphNodeComment = NewComment;
 }
 
diff --git a/Engine/Plugins/Cameras/GameplayCameras/Source/GameplayCameras/Private/Core/CameraRigAsset.cpp b/Engine/Plugins/Cameras/GameplayCameras/Source/GameplayCameras/Private/Core/CameraRigAsset.cpp
--- a/Engine/Plugins/Cameras/GameplayCameras/Source/GameplayCameras/Private/Core/CameraRigAsset.cpp
+++ b/Engine/Plugins/Cameras/GameplayCameras/Source/GameplayCameras/Private/Core/CameraRigAsset.cpp
@@ -40,6 +40,8 @@ void UCameraRigInterfaceParameter::GetGraphNodePosition(FName InGraphName, int32
 
 void UCameraRigInterfaceParameter::OnGraphNodeMoved(FName InGraphName, int32 NodePosX, int32 NodePosY, bool bMarkDirty)
 {
+	Modify(bMarkDirty);
+
 	GraphNodePos.X = NodePosX;
 	GraphNodePos.Y = NodePosY;
 }
@@ -274,6 +276,8 @@ const FString& UCameraRigAsset::GetGraphNodeCommentText(FName InGraphName) const
 
 void UCameraRigAsset::OnUpdateGraphNodeCommentText(FName InGraphName, const FString& NewComment)
 {
+	Modify();
+
 	if (InGraphName == NodeTreeGraphName)
 	{
 		NodeGraphNodeComment = NewComment;
diff --git a/Engine/Plugins/Cameras/GameplayCameras/Source/GameplayCameras/Private/Core/CameraRigTransition.cpp b/Engine/Plugins/Cameras/GameplayCameras/Source/GameplayCameras/Private/Core/CameraRigTransition.cpp
--- a/Engine/Plugins/Cameras/GameplayCameras/Source/GameplayCameras/Private/Core/CameraRigTransition.cpp
+++ b/Engine/Plugins/Cameras/GameplayCameras/Source/GameplayCameras/Private/Core/CameraRigTransition.cpp
@@ -49,6 +49,8 @@ const FString& UCameraRigTransitionCondition::GetGraphNodeCommentText(FName InGr
 
 void UCameraRigTransitionCondition::OnUpdateGraphNodeCommentText(FName InGraphName, const FString& NewComment)
 {
+	Modify();
+
 	GraphNodeComment = NewComment;
 }
 
@@ -94,6 +96,8 @@ const FString& UCameraRigTransition::GetGraphNodeCommentText(FName InGraphName)
 
 void UCameraRigTransition::OnUpdateGraphNodeCommentText(FName InGraphName, const FString& NewComment)
 {
+	Modify();
+
 	GraphNodeComment = NewComment;
 }
 
diff --git a/Engine/Plugins/Cameras/GameplayCameras/Source/GameplayCamerasEditor/Private/Editors/CameraNodeGraphSchema.cpp b/Engine/Plugins/Cameras/GameplayCameras/Source/GameplayCamerasEditor/Private/Editors/CameraNodeGraphSchema.cpp
--- a/Engine/Plugins/Cameras/GameplayCameras/Source/GameplayCamerasEditor/Private/Editors/CameraNodeGraphSchema.cpp
+++ b/Engine/Plugins/Cameras/GameplayCameras/Source/GameplayCamerasEditor/Private/Editors/CameraNodeGraphSchema.cpp
@@ -299,13 +299,14 @@ UEdGraphNode* FCameraNodeGraphSchemaAction_NewInterfaceParameterNode::PerformAct
 
 	const FScopedTransaction Transaction(LOCTEXT("CreateNewNodeAction", "Create New Node"));
 
-	const FName GraphName = ObjectTreeGraph->GetConfig().GraphName;
 	const UObjectTreeGraphSchema* Schema = CastChecked<UObjectTreeGraphSchema>(ParentGraph->GetSchema());
 
 	UCameraRigInterfaceParameter* NewInterfaceParameter = NewObject<UCameraRigInterfaceParameter>(CameraRig, NAME_None, RF_Transactional);
 	// The interface parameter's properties will be set correctly inside AutowireNewNode by virtue
 	// of getting connected to the dragged camera node pin.
 
+	ObjectTreeGraph->Modify();
+
 	UObjectTreeGraphNode* NewGraphNode = Schema->CreateObjectNode(ObjectTreeGraph, NewInterfaceParameter);
 
 	Schema->AddConnectableObject(ObjectTreeGraph, NewGraphNode);
@@ -313,6 +314,7 @@ UEdGraphNode* FCameraNodeGraphSchemaAction_NewInterfaceParameterNode::PerformAct
 	NewGraphNode->NodePosX = Location.X;
 	NewGraphNode->NodePosY = Location.Y;
 	NewGraphNode->OnGraphNodeMoved(false);
+
 	NewGraphNode->AutowireNewNode(FromPin);
 
 	return NewGraphNode;
diff --git a/Engine/Plugins/Cameras/GameplayCameras/Source/GameplayCamerasEditor/Private/Editors/ObjectTreeGraphNode.cpp b/Engine/Plugins/Cameras/GameplayCameras/Source/GameplayCamerasEditor/Private/Editors/ObjectTreeGraphNode.cpp
--- a/Engine/Plugins/Cameras/GameplayCameras/Source/GameplayCamerasEditor/Private/Editors/ObjectTreeGraphNode.cpp
+++ b/Engine/Plugins/Cameras/GameplayCameras/Source/GameplayCamerasEditor/Private/Editors/ObjectTreeGraphNode.cpp
@@ -190,6 +190,8 @@ void UObjectTreeGraphNode::CreateNewItemPin(UEdGraphPin* InParentArrayPin)
 
 	const EEdGraphPinDirection PinDirection = NodeContext.GraphConfig.GetPropertyPinDirection(NodeContext.ObjectClass, PropertyName);
 
+	InParentArrayPin->Modify();
+
 	UEdGraphPin* ChildPin = CreatePin(PinDirection, ChildPinType, PropertyName);
 	ChildPin->PinFriendlyName = FText::Format(LOCTEXT("ArrayPinFriendlyNameFmt", "{0} {1}"), FText::FromName(PropertyName), NewIndex);
 
@@ -214,6 +216,8 @@ void UObjectTreeGraphNode::RemoveItemPin(UEdGraphPin* InItemPin)
 				InItemPin->PinType.PinCategory == UObjectTreeGraphSchema::PC_Property &&
 				InItemPin->PinType.PinSubCategory == UObjectTreeGraphSchema::PSC_ArrayPropertyItem))
 	{
+		InItemPin->ParentPin->Modify();
+
 		// Don't call RemovePin() because that also removes the parent pin.
 		// We just want to remove the child pin.
 		const int32 NumPinRemoved = Pins.Remove(InItemPin);
@@ -282,8 +286,11 @@ UEdGraphPin* UObjectTreeGraphNode::GetSelfPin() const
 
 void UObjectTreeGraphNode::OverrideSelfPinDirection(EEdGraphPinDirection Direction)
 {
+	Modify();
+
 	bOverrideSelfPinDirection = true;
 	SelfPinDirectionOverride = Direction;
+
 	if (UEdGraphPin* SelfPin = GetSelfPin())
 	{
 		SelfPin->Direction = Direction;
@@ -519,8 +526,6 @@ void UObjectTreeGraphNode::OnRenameNode(const FString& NewName)
 	{
 		const FScopedTransaction Transaction(LOCTEXT("RenameNode", "Rename Node"));
 
-		Object->Modify();
-
 		const FNodeContext NodeContext = GetNodeContext();
 		GraphObject->OnRenameGraphNode(NodeContext.GraphConfig.GraphName, NewName);
 	}
@@ -555,8 +560,6 @@ void UObjectTreeGraphNode::OnUpdateCommentText(const FString& NewComment)
 	{
 		const FScopedTransaction Transaction(LOCTEXT("UpdateNodeComment", "Update Node Comment"));
 
-		Object->Modify();
-
 		const FNodeContext NodeContext = GetNodeContext();
 		GraphObject->OnUpdateGraphNodeCommentText(NodeContext.GraphConfig.GraphName, NewComment);
 	}
@@ -568,8 +571,6 @@ void UObjectTreeGraphNode::OnGraphNodeMoved(bool bMarkDirty)
 	IObjectTreeGraphObject* GraphObject = Cast<IObjectTreeGraphObject>(Object);
 	if (GraphObject)
 	{
-		Object->Modify(bMarkDirty);
-
 		const FNodeContext NodeContext = GetNodeContext();
 		GraphObject->OnGraphNodeMoved(NodeContext.GraphConfig.GraphName, NodePosX, NodePosY, bMarkDirty);
 	}
@@ -580,11 +581,18 @@ UObjectTreeGraphNode::FNodeContext UObjectTreeGraphNode::GetNodeContext() const
 	UObjectTreeGraph* OuterGraph = CastChecked<UObjectTreeGraph>(GetGraph());
 	const FObjectTreeGraphConfig& OuterGraphConfig = OuterGraph->GetConfig();
 
-	UObject* Object = WeakObject.Get();
-	UClass* ObjectClass = Object->GetClass();
-	const FObjectTreeGraphClassConfig& ObjectClassConfig = OuterGraphConfig.GetObjectClassConfig(ObjectClass);
+	if (UObject* Object = WeakObject.Get())
+	{
+		UClass* ObjectClass = Object->GetClass();
+		const FObjectTreeGraphClassConfig& ObjectClassConfig = OuterGraphConfig.GetObjectClassConfig(ObjectClass);
 
-	return FNodeContext{ ObjectClass, OuterGraph, OuterGraphConfig, ObjectClassConfig };
+		return FNodeContext{ ObjectClass, OuterGraph, OuterGraphConfig, ObjectClassConfig };
+	}
+	else
+	{
+		const FObjectTreeGraphClassConfig& ObjectClassConfig = OuterGraphConfig.GetObjectClassConfig(nullptr);
+		return FNodeContext{ nullptr, OuterGraph, OuterGraphConfig, ObjectClassConfig };
+	}
 }
 
 const FObjectTreeGraphClassConfig& UObjectTreeGraphNode::GetObjectClassConfig() const
diff --git a/Engine/Plugins/Cameras/GameplayCameras/Source/GameplayCamerasEditor/Private/Editors/ObjectTreeGraphSchema.cpp b/Engine/Plugins/Cameras/GameplayCameras/Source/GameplayCamerasEditor/Private/Editors/ObjectTreeGraphSchema.cpp
--- a/Engine/Plugins/Cameras/GameplayCameras/Source/GameplayCamerasEditor/Private/Editors/ObjectTreeGraphSchema.cpp
+++ b/Engine/Plugins/Cameras/GameplayCameras/Source/GameplayCamerasEditor/Private/Editors/ObjectTreeGraphSchema.cpp
@@ -314,6 +314,8 @@ UObjectTreeGraphNode* UObjectTreeGraphSchema::OnCreateObjectNode(UObjectTreeGrap
 		GraphNodeClass = Config.DefaultGraphNodeClass;
 	}
 
+	InGraph->Modify();
+
 	FGraphNodeCreator<UObjectTreeGraphNode> GraphNodeCreator(*InGraph);
 	UObjectTreeGraphNode* NewNode = GraphNodeCreator.CreateNode(false, GraphNodeClass);
 	NewNode->Initialize(InObject);
@@ -1110,6 +1112,8 @@ void UObjectTreeGraphSchema::ImportNodesFromText(UObjectTreeGraph* InGraph, cons
 
 	TArray<UObject*> ImportedObjects;
 
+	InGraph->Modify();
+
 	// Import the given text as new objects.
 	UPackage* TempPackage = NewObject<UPackage>(nullptr, TEXT("/Engine/GameplayCamerasEditor/Transient"), RF_Transient);
 	TempPackage->AddToRoot();
@@ -1222,6 +1226,8 @@ UEdGraphNode* FObjectGraphSchemaAction_NewNode::PerformAction(UEdGraph* ParentGr
 
 	if (NewObject)
 	{
+		ObjectTreeGraph->Modify();
+
 		UObjectTreeGraphNode* NewGraphNode = Schema->CreateObjectNode(ObjectTreeGraph, NewObject);
 
 		Schema->AddConnectableObject(ObjectTreeGraph, NewGraphNode);
