Cameras: fix and improve camera variable creation workflow

- Fix missing Modify() call when creating a new camera variable in a collection
- Auto-edit the new variable's name

diff --git a/Engine/Plugins/Cameras/GameplayCameras/Source/GameplayCamerasEditor/Private/Editors/SCameraVariableCollectionEditor.cpp b/Engine/Plugins/Cameras/GameplayCameras/Source/GameplayCamerasEditor/Private/Editors/SCameraVariableCollectionEditor.cpp
--- a/Engine/Plugins/Cameras/GameplayCameras/Source/GameplayCamerasEditor/Private/Editors/SCameraVariableCollectionEditor.cpp
+++ b/Engine/Plugins/Cameras/GameplayCameras/Source/GameplayCamerasEditor/Private/Editors/SCameraVariableCollectionEditor.cpp
@@ -37,7 +37,7 @@ public:
 
 	void Construct(const FArguments& InArgs, const TSharedRef<STableViewBase>& OwnerTableView);
 
-	void EnterNameEditingMode();
+	void EnterNameEditingMode(FSimpleDelegate InOnTextCommitted);
 
 protected:
 
@@ -57,6 +57,7 @@ private:
 	UCameraVariableAsset* CameraVariable = nullptr;
 
 	TSharedPtr<SInlineEditableTextBlock> EditableTextBlock;
+	FSimpleDelegate OnTextComitted;
 
 	TAttribute<FText> HighlightText;
 };
@@ -71,8 +72,9 @@ void SCameraVariableCollectionListRow::Construct(const FArguments& InArgs, const
 		OwnerTableView);
 }
 
-void SCameraVariableCollectionListRow::EnterNameEditingMode()
+void SCameraVariableCollectionListRow::EnterNameEditingMode(FSimpleDelegate InOnTextCommitted)
 {
+	OnTextComitted = InOnTextCommitted;
 	EditableTextBlock->EnterEditingMode();
 }
 
@@ -166,6 +168,12 @@ void SCameraVariableCollectionListRow::OnVariableNameCommitted(const FText& Text
 		NewObjectName = MakeUniqueObjectName(CameraVariable->GetOuter(), CameraVariable->GetClass(), NewObjectName);
 		CameraVariable->Rename(*NewObjectName.ToString());
 	}
+
+	if (OnTextComitted.IsBound())
+	{
+		OnTextComitted.Execute();
+		OnTextComitted.Unbind();
+	}
 }
 
 void SCameraVariableCollectionEditor::Construct(const FArguments& InArgs)
@@ -243,6 +251,13 @@ void SCameraVariableCollectionEditor::GetSelectedVariables(TArray<UCameraVariabl
 	ListView->GetSelectedItems(OutSelection);
 }
 
+void SCameraVariableCollectionEditor::RequestRenameVariable(UCameraVariableAsset* InItem, FSimpleDelegate InOnRenamedItem)
+{
+	bDeferredRequestRenameItem = true;
+	OnDeferredRenamedItem = InOnRenamedItem;
+	ListView->RequestScrollIntoView(InItem);
+}
+
 void SCameraVariableCollectionEditor::RequestRenameSelectedVariable()
 {
 	TArray<UCameraVariableAsset*> SelectedVariables = ListView->GetSelectedItems();
@@ -330,7 +345,10 @@ void SCameraVariableCollectionEditor::OnListItemScrolledIntoView(UCameraVariable
 			return;
 		}
 
-		TypedRowWidget->EnterNameEditingMode();
+		TypedRowWidget->EnterNameEditingMode(OnDeferredRenamedItem);
+
+		// The delegate was copied into the row widget so we can unbind it here.
+		OnDeferredRenamedItem.Unbind();
 	}
 }
 
diff --git a/Engine/Plugins/Cameras/GameplayCameras/Source/GameplayCamerasEditor/Private/Toolkits/CameraVariableCollectionEditorToolkit.cpp b/Engine/Plugins/Cameras/GameplayCameras/Source/GameplayCamerasEditor/Private/Toolkits/CameraVariableCollectionEditorToolkit.cpp
--- a/Engine/Plugins/Cameras/GameplayCameras/Source/GameplayCamerasEditor/Private/Toolkits/CameraVariableCollectionEditorToolkit.cpp
+++ b/Engine/Plugins/Cameras/GameplayCameras/Source/GameplayCamerasEditor/Private/Toolkits/CameraVariableCollectionEditorToolkit.cpp
@@ -11,11 +11,9 @@
 #include "Framework/Docking/LayoutExtender.h"
 #include "Framework/Docking/TabManager.h"
 #include "IContentBrowserSingleton.h"
-#include "Misc/FeedbackContext.h"
 #include "Modules/ModuleManager.h"
 #include "ObjectTools.h"
 #include "PropertyEditorModule.h"
-#include "PropertyPath.h"
 #include "ScopedTransaction.h"
 #include "Styles/GameplayCamerasEditorStyle.h"
 #include "ToolMenus.h"
@@ -253,7 +251,9 @@ FLinearColor FCameraVariableCollectionEditorToolkit::GetWorldCentricTabColorScal
 
 void FCameraVariableCollectionEditorToolkit::OnCreateVariable(TSubclassOf<UCameraVariableAsset> InVariableClass)
 {
-	FScopedTransaction CreateTransaction(LOCTEXT("CreateVariable", "Create camera variable"));
+	GEditor->BeginTransaction(LOCTEXT("CreateVariable", "Create camera variable"));
+	
+	VariableCollection->Modify();
 
 	UCameraVariableAsset* NewVariable = NewObject<UCameraVariableAsset>(
 			VariableCollection, 
@@ -264,6 +264,12 @@ void FCameraVariableCollectionEditorToolkit::OnCreateVariable(TSubclassOf<UCamer
 	VariableCollection->Variables.Add(NewVariable);
 
 	VariableCollectionEditorWidget->RequestListRefresh();
+	VariableCollectionEditorWidget->RequestRenameVariable(NewVariable, FSimpleDelegate::CreateLambda([]()
+				{
+					// End the transaction when the user exits the editing mode on the
+					// editable text block for the new variable's name.
+					GEditor->EndTransaction();
+				}));
 }
 
 void FCameraVariableCollectionEditorToolkit::OnRenameVariable()
@@ -311,6 +317,7 @@ void FCameraVariableCollectionEditorToolkit::OnDeleteVariable()
 		FScopedTransaction DeleteTransaction(LOCTEXT("DeleteVariable", "Delete camera variable"));
 
 		VariableCollection->Modify();
+
 		for (UCameraVariableAsset* VariableToDelete : Selection)
 		{
 			VariableToDelete->Modify();
diff --git a/Engine/Plugins/Cameras/GameplayCameras/Source/GameplayCamerasEditor/Public/Editors/SCameraVariableCollectionEditor.h b/Engine/Plugins/Cameras/GameplayCameras/Source/GameplayCamerasEditor/Public/Editors/SCameraVariableCollectionEditor.h
--- a/Engine/Plugins/Cameras/GameplayCameras/Source/GameplayCamerasEditor/Public/Editors/SCameraVariableCollectionEditor.h
+++ b/Engine/Plugins/Cameras/GameplayCameras/Source/GameplayCamerasEditor/Public/Editors/SCameraVariableCollectionEditor.h
@@ -48,6 +48,9 @@ public:
 	/** Gets the selected variables in the list view. */
 	void GetSelectedVariables(TArray<UCameraVariableAsset*>& OutSelection) const;
 
+	/** Enter editing mode for the given variable's name. */
+	void RequestRenameVariable(UCameraVariableAsset* InItem, FSimpleDelegate InOnRenamedItem);
+
 	/** Enter editing mode for the first currently selected variable's name. */
 	void RequestRenameSelectedVariable();
 
@@ -92,6 +95,7 @@ private:
 
 	bool bUpdateFilteredItemSource = false;
 	bool bDeferredRequestRenameItem = false;
+	FSimpleDelegate OnDeferredRenamedItem;
 };
 
 }  // namespace UE::Cameras
