Cameras: hide the list of variables from the Details View in the collection asset

diff --git a/Engine/Plugins/Cameras/GameplayCameras/Source/GameplayCameras/Public/Core/CameraVariableCollection.h b/Engine/Plugins/Cameras/GameplayCameras/Source/GameplayCameras/Public/Core/CameraVariableCollection.h
--- a/Engine/Plugins/Cameras/GameplayCameras/Source/GameplayCameras/Public/Core/CameraVariableCollection.h
+++ b/Engine/Plugins/Cameras/GameplayCameras/Source/GameplayCameras/Public/Core/CameraVariableCollection.h
@@ -23,7 +23,7 @@ public:
 public:
 
 	/** The variables in this collection. */
-	UPROPERTY(EditAnywhere, Category="Camera")
+	UPROPERTY()
 	TArray<TObjectPtr<UCameraVariableAsset>> Variables;
 };
 
