Cameras: don't duplicate nodes that can't be created in the first place

diff --git a/Engine/Plugins/Cameras/GameplayCameras/Source/GameplayCamerasEditor/Private/Editors/ObjectTreeGraphNode.cpp b/Engine/Plugins/Cameras/GameplayCameras/Source/GameplayCamerasEditor/Private/Editors/ObjectTreeGraphNode.cpp
--- a/Engine/Plugins/Cameras/GameplayCameras/Source/GameplayCamerasEditor/Private/Editors/ObjectTreeGraphNode.cpp
+++ b/Engine/Plugins/Cameras/GameplayCameras/Source/GameplayCamerasEditor/Private/Editors/ObjectTreeGraphNode.cpp
@@ -531,6 +531,17 @@ void UObjectTreeGraphNode::OnRenameNode(const FString& NewName)
 	}
 }
 
+bool UObjectTreeGraphNode::CanDuplicateNode() const
+{
+	const FObjectTreeGraphClassConfig& ObjectClassConfig = GetObjectClassConfig();
+	if (!ObjectClassConfig.CanCreateNew())  // If it can't be created, it shouldn't be worked around by copy/pasting
+	{
+		return false;
+	}
+
+	return Super::CanUserDeleteNode();
+}
+
 bool UObjectTreeGraphNode::CanUserDeleteNode() const
 {
 	const FObjectTreeGraphClassConfig& ObjectClassConfig = GetObjectClassConfig();
diff --git a/Engine/Plugins/Cameras/GameplayCameras/Source/GameplayCamerasEditor/Public/Editors/ObjectTreeGraphNode.h b/Engine/Plugins/Cameras/GameplayCameras/Source/GameplayCamerasEditor/Public/Editors/ObjectTreeGraphNode.h
--- a/Engine/Plugins/Cameras/GameplayCameras/Source/GameplayCamerasEditor/Public/Editors/ObjectTreeGraphNode.h
+++ b/Engine/Plugins/Cameras/GameplayCameras/Source/GameplayCamerasEditor/Public/Editors/ObjectTreeGraphNode.h
@@ -75,6 +75,7 @@ public:
 	virtual void GetNodeContextMenuActions(class UToolMenu* Menu, class UGraphNodeContextMenuContext* Context) const override;
 	virtual bool GetCanRenameNode() const override;
 	virtual void OnRenameNode(const FString& NewName) override;
+	virtual bool CanDuplicateNode() const override;
 	virtual bool CanUserDeleteNode() const override;
 	virtual bool SupportsCommentBubble() const override;
 	virtual void OnUpdateCommentText(const FString& NewComment) override;
