Cameras: fixup pasted objects by severing connections to other objects

diff --git a/Engine/Plugins/Cameras/GameplayCameras/Source/GameplayCamerasEditor/Private/Editors/ObjectTreeGraphSchema.cpp b/Engine/Plugins/Cameras/GameplayCameras/Source/GameplayCamerasEditor/Private/Editors/ObjectTreeGraphSchema.cpp
--- a/Engine/Plugins/Cameras/GameplayCameras/Source/GameplayCamerasEditor/Private/Editors/ObjectTreeGraphSchema.cpp
+++ b/Engine/Plugins/Cameras/GameplayCameras/Source/GameplayCamerasEditor/Private/Editors/ObjectTreeGraphSchema.cpp
@@ -1124,6 +1124,51 @@ void UObjectTreeGraphSchema::ImportNodesFromText(UObjectTreeGraph* InGraph, cons
 	}
 	TempPackage->RemoveFromRoot();
 
+	// Sever references to objects outside of the set being copy/pasted.
+	TSet<UObject*> ImportedObjectSet(ImportedObjects);
+	const FObjectTreeGraphConfig& GraphConfig = InGraph->GetConfig();
+	for (UObject* Object : ImportedObjects)
+	{
+		UClass* ObjectClass = Object->GetClass();
+		for (TFieldIterator<FProperty> PropertyIt(ObjectClass); PropertyIt; ++PropertyIt)
+		{
+			if (FObjectProperty* ObjectProperty = CastField<FObjectProperty>(*PropertyIt))
+			{
+				if (!GraphConfig.IsConnectable(ObjectProperty))
+				{
+					continue;
+				}
+
+				TObjectPtr<UObject> OutConnectedObject;
+				ObjectProperty->GetValue_InContainer(Object, &OutConnectedObject);
+				if (OutConnectedObject && !ImportedObjectSet.Contains(OutConnectedObject))
+				{
+					ObjectProperty->SetValue_InContainer(Object, nullptr);
+				}
+			}
+			else if (FArrayProperty* ArrayProperty = CastField<FArrayProperty>(*PropertyIt))
+			{
+				if (!GraphConfig.IsConnectable(ArrayProperty))
+				{
+					continue;
+				}
+
+				FObjectProperty* InnerProperty = CastFieldChecked<FObjectProperty>(ArrayProperty->Inner);
+				FScriptArrayHelper ArrayHelper(ArrayProperty, ArrayProperty->ContainerPtrToValuePtr<void>(Object));
+
+				const int32 ArrayNum = ArrayHelper.Num();
+				for (int32 Index = ArrayNum - 1; Index >= 0; --Index)
+				{
+					UObject* OutConnectedObject = InnerProperty->GetObjectPropertyValue(ArrayHelper.GetRawPtr(Index));
+					if (OutConnectedObject && !ImportedObjectSet.Contains(OutConnectedObject))
+					{
+						ArrayHelper.RemoveValues(Index);
+					}
+				}
+			}
+		}
+	}
+
 	// Finish setting up the new objects: clear the transient flag from the transient package we used above,
 	// and move the objects under the our graph root.
 	UObject* GraphRootObject = InGraph->GetRootObject();
