Cameras: build camera assets on save

diff --git a/Engine/Plugins/Cameras/GameplayCameras/Source/GameplayCameras/Private/Core/CameraAsset.cpp b/Engine/Plugins/Cameras/GameplayCameras/Source/GameplayCameras/Private/Core/CameraAsset.cpp
--- a/Engine/Plugins/Cameras/GameplayCameras/Source/GameplayCameras/Private/Core/CameraAsset.cpp
+++ b/Engine/Plugins/Cameras/GameplayCameras/Source/GameplayCameras/Private/Core/CameraAsset.cpp
@@ -5,6 +5,7 @@
 #include "Core/CameraAssetBuilder.h"
 #include "Core/CameraBuildLog.h"
 #include "Core/CameraDirector.h"
+#include "UObject/ObjectSaveContext.h"
 
 #include UE_INLINE_GENERATED_CPP_BY_NAME(CameraAsset)
 
@@ -179,6 +180,21 @@ void UCameraAsset::DirtyBuildStatus()
 	BuildStatus = ECameraBuildStatus::Dirty;
 }
 
+void UCameraAsset::PreSave(FObjectPreSaveContext ObjectSaveContext)
+{
+#if WITH_EDITOR
+
+	if (!HasAnyFlags(RF_ClassDefaultObject | RF_ArchetypeObject))
+	{
+		// Build on save.
+		BuildCamera();
+	}
+
+#endif
+
+	Super::PreSave(ObjectSaveContext);
+}
+
 #if WITH_EDITOR
 
 void UCameraAsset::GetGraphNodePosition(FName InGraphName, int32& NodePosX, int32& NodePosY) const
diff --git a/Engine/Plugins/Cameras/GameplayCameras/Source/GameplayCameras/Public/Core/CameraAsset.h b/Engine/Plugins/Cameras/GameplayCameras/Source/GameplayCameras/Public/Core/CameraAsset.h
--- a/Engine/Plugins/Cameras/GameplayCameras/Source/GameplayCameras/Public/Core/CameraAsset.h
+++ b/Engine/Plugins/Cameras/GameplayCameras/Source/GameplayCameras/Public/Core/CameraAsset.h
@@ -124,6 +124,7 @@ protected:
 #endif
 
 	// UObject interface.
+	virtual void PreSave(FObjectPreSaveContext ObjectSaveContext) override;
 	virtual void PostLoad() override;
 #if WITH_EDITOR
 	virtual void PostEditChangeProperty(FPropertyChangedEvent& PropertyChangedEvent) override;
