Cameras: fix old camera variable assets possibly missing RF_Public

diff --git a/Engine/Plugins/Cameras/GameplayCameras/Source/GameplayCameras/Private/Core/CameraVariableCollection.cpp b/Engine/Plugins/Cameras/GameplayCameras/Source/GameplayCameras/Private/Core/CameraVariableCollection.cpp
--- a/Engine/Plugins/Cameras/GameplayCameras/Source/GameplayCameras/Private/Core/CameraVariableCollection.cpp
+++ b/Engine/Plugins/Cameras/GameplayCameras/Source/GameplayCameras/Private/Core/CameraVariableCollection.cpp
@@ -2,6 +2,9 @@
 
 #include "Core/CameraVariableCollection.h"
 
+#include "Core/CameraVariableAssets.h"
+#include "GameplayCameras.h"
+
 #include UE_INLINE_GENERATED_CPP_BY_NAME(CameraVariableCollection)
 
 UCameraVariableCollection::UCameraVariableCollection(const FObjectInitializer& ObjectInit)
@@ -9,3 +12,19 @@ UCameraVariableCollection::UCameraVariableCollection(const FObjectInitializer& O
 {
 }
 
+void UCameraVariableCollection::PostLoad()
+{
+	Super::PostLoad();
+
+#if WITH_EDITOR
+	for (UCameraVariableAsset* Variable : Variables)
+	{
+		if (!Variable->HasAnyFlags(RF_Public))
+		{
+			UE_LOG(LogCameraSystem, Warning, TEXT("Adding missing RF_Public flag on variable '%s'."), *GetPathNameSafe(Variable));
+			Variable->SetFlags(RF_Public);
+		}
+	}
+#endif  // WITH_EDITOR
+}
+
diff --git a/Engine/Plugins/Cameras/GameplayCameras/Source/GameplayCameras/Public/Core/CameraVariableCollection.h b/Engine/Plugins/Cameras/GameplayCameras/Source/GameplayCameras/Public/Core/CameraVariableCollection.h
--- a/Engine/Plugins/Cameras/GameplayCameras/Source/GameplayCameras/Public/Core/CameraVariableCollection.h
+++ b/Engine/Plugins/Cameras/GameplayCameras/Source/GameplayCameras/Public/Core/CameraVariableCollection.h
@@ -20,6 +20,11 @@ public:
 
 	UCameraVariableCollection(const FObjectInitializer& ObjectInit);
 
+protected:
+
+	// UObject interface
+	virtual void PostLoad() override;
+
 public:
 
 	/** The variables in this collection. */
