Cameras: alphabetically sort variables in the debug display

diff --git a/Engine/Plugins/Cameras/GameplayCameras/Source/GameplayCameras/Private/Debug/VariableTableDebugBlock.cpp b/Engine/Plugins/Cameras/GameplayCameras/Source/GameplayCameras/Private/Debug/VariableTableDebugBlock.cpp
--- a/Engine/Plugins/Cameras/GameplayCameras/Source/GameplayCameras/Private/Debug/VariableTableDebugBlock.cpp
+++ b/Engine/Plugins/Cameras/GameplayCameras/Source/GameplayCameras/Private/Debug/VariableTableDebugBlock.cpp
@@ -53,6 +53,15 @@ void FVariableTableDebugBlock::Initialize(const FCameraVariableTable& InVariable
 		EntryDebugInfo.bWrittenThisFrame = EnumHasAnyFlags(Entry.Flags, FCameraVariableTable::EEntryFlags::WrittenThisFrame);
 		Entries.Add(EntryDebugInfo);
 	}
+
+	Entries.StableSort([](const FEntryDebugInfo& A, const FEntryDebugInfo& B) -> bool
+			{
+				if (!A.Name.IsEmpty() || !B.Name.IsEmpty())
+				{
+					return A.Name.Compare(B.Name) < 0;
+				}
+				return A.ID < B.ID;
+			});
 }
 
 void FVariableTableDebugBlock::OnDebugDraw(const FCameraDebugBlockDrawParams& Params, FCameraDebugRenderer& Renderer)
