Cameras: add more utility functions to the Blueprint director and camera node

diff --git a/Engine/Plugins/Cameras/GameplayCameras/Source/GameplayCameras/Private/Directors/BlueprintCameraDirector.cpp b/Engine/Plugins/Cameras/GameplayCameras/Source/GameplayCameras/Private/Directors/BlueprintCameraDirector.cpp
--- a/Engine/Plugins/Cameras/GameplayCameras/Source/GameplayCameras/Private/Directors/BlueprintCameraDirector.cpp
+++ b/Engine/Plugins/Cameras/GameplayCameras/Source/GameplayCameras/Private/Directors/BlueprintCameraDirector.cpp
@@ -2,12 +2,14 @@
 
 #include "Directors/BlueprintCameraDirector.h"
 
+#include "Components/ActorComponent.h"
 #include "Core/CameraAsset.h"
 #include "Core/CameraBuildLog.h"
+#include "Core/CameraEvaluationContext.h"
 #include "Core/CameraRigAsset.h"
 #include "Core/CameraRigProxyAsset.h"
 #include "Core/CameraRigProxyTable.h"
-#include "Core/CameraEvaluationContext.h"
+#include "GameFramework/Actor.h"
 #include "GameplayCameras.h"
 
 #include UE_INLINE_GENERATED_CPP_BY_NAME(BlueprintCameraDirector)
@@ -69,6 +71,7 @@ void FBlueprintCameraDirectorEvaluator::OnRun(const FCameraDirectorEvaluationPar
 	{
 		FBlueprintCameraDirectorEvaluationParams BlueprintParams;
 		BlueprintParams.DeltaTime = Params.DeltaTime;
+		BlueprintParams.EvaluationContext = Params.OwnerContext;
 		if (Params.OwnerContext)
 		{
 			BlueprintParams.EvaluationContextOwner = Params.OwnerContext->GetOwner();
@@ -165,14 +168,86 @@ void UBlueprintCameraDirectorEvaluator::ActivateCameraRigViaProxy(UCameraRigProx
 	CurrentResult.ActiveCameraRigProxies.Add(CameraRigProxy);
 }
 
+AActor* UBlueprintCameraDirectorEvaluator::FindEvaluationContextOwnerActor(TSubclassOf<AActor> ActorClass) const
+{
+	if (CurrentContext)
+	{
+		if (UActorComponent* ContextOwnerAsComponent = Cast<UActorComponent>(CurrentContext->GetOwner()))
+		{
+			return ContextOwnerAsComponent->GetOwner();
+		}
+		else if (AActor* ContextOwnerAsActor = Cast<AActor>(CurrentContext->GetOwner()))
+		{
+			return ContextOwnerAsActor;
+		}
+		else
+		{
+			return nullptr;
+		}
+	}
+	else
+	{
+		FFrame::KismetExecutionMessage(
+				TEXT("Can't access evaluation context outside of RunCameraDirector"), 
+				ELogVerbosity::Error);
+		return nullptr;
+	}
+}
+
+FBlueprintCameraPose UBlueprintCameraDirectorEvaluator::GetInitialContextCameraPose() const
+{
+	if (CurrentContext)
+	{
+		return FBlueprintCameraPose::FromCameraPose(CurrentContext->GetInitialResult().CameraPose);
+	}
+	else
+	{
+		FFrame::KismetExecutionMessage(
+				TEXT("Can't access evaluation context's initial result outside of RunCameraDirector"), 
+				ELogVerbosity::Error);
+		return FBlueprintCameraPose();
+	}
+}
+
+FBlueprintCameraVariableTable UBlueprintCameraDirectorEvaluator::GetInitialContextVariableTable() const
+{
+	if (CurrentContext)
+	{
+		return FBlueprintCameraVariableTable(&CurrentContext->GetInitialResult().VariableTable);
+	}
+	else
+	{
+		FFrame::KismetExecutionMessage(
+				TEXT("Can't access evaluation context's initial result outside of RunCameraDirector"), 
+				ELogVerbosity::Error);
+		return FBlueprintCameraVariableTable();
+	}
+}
+
+void UBlueprintCameraDirectorEvaluator::SetInitialContextCameraPose(const FBlueprintCameraPose& InCameraPose)
+{
+	if (CurrentContext)
+	{
+		InCameraPose.ApplyTo(CurrentContext->GetInitialResult().CameraPose);
+	}
+	else
+	{
+		FFrame::KismetExecutionMessage(
+				TEXT("Can't access evaluation context's initial result outside of RunCameraDirector"), 
+				ELogVerbosity::Error);
+	}
+}
+
 void UBlueprintCameraDirectorEvaluator::NativeRunCameraDirector(const FBlueprintCameraDirectorEvaluationParams& Params, FBlueprintCameraDirectorEvaluationResult& OutResult)
 {
+	CurrentContext = Params.EvaluationContext;
 	CurrentResult = OutResult;
 	{
 		// Run the Blueprint logic.
 		RunCameraDirector(Params);
 	}
 	OutResult = CurrentResult;
+	CurrentContext.Reset();
 }
 
 FCameraDirectorEvaluatorPtr UBlueprintCameraDirector::OnBuildEvaluator(FCameraDirectorEvaluatorBuilder& Builder) const
diff --git a/Engine/Plugins/Cameras/GameplayCameras/Source/GameplayCameras/Private/Nodes/Utility/BlueprintCameraNode.cpp b/Engine/Plugins/Cameras/GameplayCameras/Source/GameplayCameras/Private/Nodes/Utility/BlueprintCameraNode.cpp
--- a/Engine/Plugins/Cameras/GameplayCameras/Source/GameplayCameras/Private/Nodes/Utility/BlueprintCameraNode.cpp
+++ b/Engine/Plugins/Cameras/GameplayCameras/Source/GameplayCameras/Private/Nodes/Utility/BlueprintCameraNode.cpp
@@ -2,11 +2,13 @@
 
 #include "Nodes/Utility/BlueprintCameraNode.h"
 
+#include "Components/ActorComponent.h"
 #include "Core/CameraBuildLog.h"
 #include "Core/CameraEvaluationContext.h"
 #include "Core/CameraNodeEvaluator.h"
 #include "Core/CameraRigBuildContext.h"
 #include "Core/CameraVariableTable.h"
+#include "GameFramework/Actor.h"
 #include "GameplayCameras.h"
 #include "Templates/UnrealTemplate.h"
 
@@ -71,10 +73,15 @@ void FBlueprintCameraNodeEvaluator::OnAddReferencedObjects(FReferenceCollector&
 
 void UBlueprintCameraNodeEvaluator::NativeRunCameraNode(const FCameraNodeEvaluationParams& Params, FCameraNodeEvaluationResult& OutResult)
 {
+	using namespace UE::Cameras;
+
 	bIsFirstFrame = Params.bIsFirstFrame;
 	EvaluationContextOwner = Params.EvaluationContext->GetOwner();
 	CameraPose = FBlueprintCameraPose::FromCameraPose(OutResult.CameraPose);
 
+	ensure(!CurrentContext.IsValid());
+	TGuardValue<TSharedPtr<const FCameraEvaluationContext>> CurrentContextGuard(CurrentContext, Params.EvaluationContext);
+
 	ensure(CurrentResult == nullptr);
 	TGuardValue<FCameraNodeEvaluationResult*> CurrentResultGuard(CurrentResult, &OutResult);
 
@@ -83,6 +90,32 @@ void UBlueprintCameraNodeEvaluator::NativeRunCameraNode(const FCameraNodeEvaluat
 	CameraPose.ApplyTo(OutResult.CameraPose);
 }
 
+AActor* UBlueprintCameraNodeEvaluator::FindEvaluationContextOwnerActor(TSubclassOf<AActor> ActorClass) const
+{
+	if (CurrentContext)
+	{
+		if (UActorComponent* ContextOwnerAsComponent = Cast<UActorComponent>(CurrentContext->GetOwner()))
+		{
+			return ContextOwnerAsComponent->GetOwner();
+		}
+		else if (AActor* ContextOwnerAsActor = Cast<AActor>(CurrentContext->GetOwner()))
+		{
+			return ContextOwnerAsActor;
+		}
+		else
+		{
+			return nullptr;
+		}
+	}
+	else
+	{
+		FFrame::KismetExecutionMessage(
+				TEXT("Can't access evaluation context outside of RunCameraDirector"), 
+				ELogVerbosity::Error);
+		return nullptr;
+	}
+}
+
 void UBlueprintCameraNode::OnBuild(FCameraRigBuildContext& BuildContext)
 {
 	if (!CameraNodeEvaluatorClass)
diff --git a/Engine/Plugins/Cameras/GameplayCameras/Source/GameplayCameras/Public/Core/CameraDirectorEvaluator.h b/Engine/Plugins/Cameras/GameplayCameras/Source/GameplayCameras/Public/Core/CameraDirectorEvaluator.h
--- a/Engine/Plugins/Cameras/GameplayCameras/Source/GameplayCameras/Public/Core/CameraDirectorEvaluator.h
+++ b/Engine/Plugins/Cameras/GameplayCameras/Source/GameplayCameras/Public/Core/CameraDirectorEvaluator.h
@@ -54,7 +54,7 @@ struct FCameraDirectorEvaluationParams
 	float DeltaTime = 0.f;
 
 	/** The context in which this director runs. */
-	TSharedPtr<const FCameraEvaluationContext> OwnerContext;
+	TSharedPtr<FCameraEvaluationContext> OwnerContext;
 };
 
 struct FActiveCameraRigInfo
diff --git a/Engine/Plugins/Cameras/GameplayCameras/Source/GameplayCameras/Public/Directors/BlueprintCameraDirector.h b/Engine/Plugins/Cameras/GameplayCameras/Source/GameplayCameras/Public/Directors/BlueprintCameraDirector.h
--- a/Engine/Plugins/Cameras/GameplayCameras/Source/GameplayCameras/Public/Directors/BlueprintCameraDirector.h
+++ b/Engine/Plugins/Cameras/GameplayCameras/Source/GameplayCameras/Public/Directors/BlueprintCameraDirector.h
@@ -4,6 +4,8 @@
 
 #include "Core/CameraDirector.h"
 #include "Core/CameraDirectorEvaluator.h"
+#include "GameFramework/BlueprintCameraPose.h"
+#include "GameFramework/BlueprintCameraVariableTable.h"
 #include "Templates/SubclassOf.h"
 
 #include "BlueprintCameraDirector.generated.h"
@@ -26,6 +28,9 @@ struct FBlueprintCameraDirectorEvaluationParams
 	/** The owner (if any) of the evaluation context we are running inside of. */
 	UPROPERTY(BlueprintReadWrite, Category="Evaluation")
 	TObjectPtr<UObject> EvaluationContextOwner;
+
+	/** The evaluation context we are running inside of. */
+	TSharedPtr<UE::Cameras::FCameraEvaluationContext> EvaluationContext;
 };
 
 /**
@@ -75,6 +80,38 @@ public:
 	UFUNCTION(BlueprintCallable, Category="Evaluation")
 	void ActivateCameraRigViaProxy(UCameraRigProxyAsset* CameraRigProxy);
 
+public:
+
+	/**
+	 * A utility function that tries to find if an actor owns the evaluation context.
+	 * Handles the situation where the evaluation context is an actor component (like a
+	 * UGameplayCameraComponent) or an actor itself.
+	 */
+	UFUNCTION(BlueprintPure, Category="Evaluation", meta=(DeterminesOutputType="ActorClass"))
+	AActor* FindEvaluationContextOwnerActor(TSubclassOf<AActor> ActorClass) const;
+
+	/**
+	 * Gets the initial evaluation context camera pose.
+	 */
+	UFUNCTION(BlueprintPure, Category="Evaluation")
+	FBlueprintCameraPose GetInitialContextCameraPose() const;
+
+	/**
+	 * Sets the initial evaluation context camera pose.
+	 * WARNING: this will change the initial pose of ALL running camera rigs!
+	 */
+	UFUNCTION(BlueprintCallable, Category="Evaluation")
+	void SetInitialContextCameraPose(const FBlueprintCameraPose& InCameraPose);
+
+	/**
+	 * Gets the initial evaluation context camera variable table.
+	 * WARNING: setting variables here will affect ALL running camera rigs!
+	 */
+	UFUNCTION(BlueprintPure, Category="Evaluation")
+	FBlueprintCameraVariableTable GetInitialContextVariableTable() const;
+
+public:
+
 	/** Native wrapper for RunCameraDirector. */
 	void NativeRunCameraDirector(
 			const FBlueprintCameraDirectorEvaluationParams& Params,
@@ -85,6 +122,10 @@ protected:
 	/** The current camera director evaluation result. */
 	UPROPERTY(BlueprintReadWrite, Category="Evaluation")
 	FBlueprintCameraDirectorEvaluationResult CurrentResult;
+
+private:
+
+	TSharedPtr<UE::Cameras::FCameraEvaluationContext> CurrentContext;
 };
 
 /**
diff --git a/Engine/Plugins/Cameras/GameplayCameras/Source/GameplayCameras/Public/Nodes/Utility/BlueprintCameraNode.h b/Engine/Plugins/Cameras/GameplayCameras/Source/GameplayCameras/Public/Nodes/Utility/BlueprintCameraNode.h
--- a/Engine/Plugins/Cameras/GameplayCameras/Source/GameplayCameras/Public/Nodes/Utility/BlueprintCameraNode.h
+++ b/Engine/Plugins/Cameras/GameplayCameras/Source/GameplayCameras/Public/Nodes/Utility/BlueprintCameraNode.h
@@ -11,10 +11,11 @@
 
 #include "BlueprintCameraNode.generated.h"
 
+struct FBlueprintCameraPose;
+
 namespace UE::Cameras
 {
 	class FCameraVariableTable;
-	struct FBlueprintCameraPose;
 	struct FCameraNodeEvaluationParams;
 	struct FCameraNodeEvaluationResult;
 };
@@ -41,6 +42,16 @@ public:
 	/** Runs this camera node. */
 	void NativeRunCameraNode(const FCameraNodeEvaluationParams& Params, FCameraNodeEvaluationResult& OutResult);
 
+public:
+
+	/**
+	 * A utility function that tries to find if an actor owns the evaluation context.
+	 * Handles the situation where the evaluation context is an actor component (like a
+	 * UGameplayCameraComponent) or an actor itself.
+	 */
+	UFUNCTION(BlueprintPure, Category="Evaluation", meta=(DeterminesOutputType="ActorClass"))
+	AActor* FindEvaluationContextOwnerActor(TSubclassOf<AActor> ActorClass) const;
+
 protected:
 
 	/** Whether this is the first frame of this camera node's lifetime. */
@@ -61,6 +72,8 @@ protected:
 
 private:
 
+	TSharedPtr<const UE::Cameras::FCameraEvaluationContext> CurrentContext;
+
 	FCameraNodeEvaluationResult* CurrentResult = nullptr;
 };
 
