Cameras: re-enable simple test for computing camera evaluator allocation info

diff --git a/Engine/Plugins/Cameras/GameplayCameras/Source/GameplayCameras/Private/Nodes/Common/ArrayCameraNode.cpp b/Engine/Plugins/Cameras/GameplayCameras/Source/GameplayCameras/Private/Nodes/Common/ArrayCameraNode.cpp
--- a/Engine/Plugins/Cameras/GameplayCameras/Source/GameplayCameras/Private/Nodes/Common/ArrayCameraNode.cpp
+++ b/Engine/Plugins/Cameras/GameplayCameras/Source/GameplayCameras/Private/Nodes/Common/ArrayCameraNode.cpp
@@ -26,12 +26,6 @@ private:
 
 UE_DEFINE_CAMERA_NODE_EVALUATOR(FArrayCameraNodeEvaluator)
 
-#if !(UE_BUILD_SHIPPING || UE_BUILD_TEST)
-// For unit tests.
-int32 GArrayCameraNodeEvaluatorSizeof = sizeof(FArrayCameraNodeEvaluator);
-int32 GArrayCameraNodeEvaluatorAlignof = alignof(FArrayCameraNodeEvaluator);
-#endif
-
 FCameraNodeEvaluatorChildrenView FArrayCameraNodeEvaluator::OnGetChildren()
 {
 	return FCameraNodeEvaluatorChildrenView(Children);
@@ -74,3 +68,9 @@ FCameraNodeEvaluatorPtr UArrayCameraNode::OnBuildEvaluator(FCameraNodeEvaluatorB
 	return Builder.BuildEvaluator<FArrayCameraNodeEvaluator>();
 }
 
+TTuple<int32, int32> UArrayCameraNode::GetEvaluatorAllocationInfo()
+{
+	using namespace UE::Cameras;
+	return { sizeof(FArrayCameraNodeEvaluator), alignof(FArrayCameraNodeEvaluator) };
+}
+
diff --git a/Engine/Plugins/Cameras/GameplayCameras/Source/GameplayCameras/Private/Nodes/Common/OffsetCameraNode.cpp b/Engine/Plugins/Cameras/GameplayCameras/Source/GameplayCameras/Private/Nodes/Common/OffsetCameraNode.cpp
--- a/Engine/Plugins/Cameras/GameplayCameras/Source/GameplayCameras/Private/Nodes/Common/OffsetCameraNode.cpp
+++ b/Engine/Plugins/Cameras/GameplayCameras/Source/GameplayCameras/Private/Nodes/Common/OffsetCameraNode.cpp
@@ -27,12 +27,6 @@ private:
 	TCameraParameterReader<FRotator3d> RotationReader;
 };
 
-#if !(UE_BUILD_SHIPPING || UE_BUILD_TEST)
-// For unit tests.
-int32 GOffsetCameraNodeEvaluatorSizeof = sizeof(FOffsetCameraNodeEvaluator);
-int32 GOffsetCameraNodeEvaluatorAlignof = alignof(FOffsetCameraNodeEvaluator);
-#endif
-
 UE_DEFINE_CAMERA_NODE_EVALUATOR(FOffsetCameraNodeEvaluator)
 
 void FOffsetCameraNodeEvaluator::OnInitialize(const FCameraNodeEvaluatorInitializeParams& Params, FCameraNodeEvaluationResult& OutResult)
@@ -111,3 +105,9 @@ FCameraNodeEvaluatorPtr UOffsetCameraNode::OnBuildEvaluator(FCameraNodeEvaluator
 	return Builder.BuildEvaluator<FOffsetCameraNodeEvaluator>();
 }
 
+TTuple<int32, int32> UOffsetCameraNode::GetEvaluatorAllocationInfo()
+{
+	using namespace UE::Cameras;
+	return { sizeof(FOffsetCameraNodeEvaluator), alignof(FOffsetCameraNodeEvaluator) };
+}
+
diff --git a/Engine/Plugins/Cameras/GameplayCameras/Source/GameplayCameras/Private/Tests/CameraRigAssetBuilderTests.cpp b/Engine/Plugins/Cameras/GameplayCameras/Source/GameplayCameras/Private/Tests/CameraRigAssetBuilderTests.cpp
--- a/Engine/Plugins/Cameras/GameplayCameras/Source/GameplayCameras/Private/Tests/CameraRigAssetBuilderTests.cpp
+++ b/Engine/Plugins/Cameras/GameplayCameras/Source/GameplayCameras/Private/Tests/CameraRigAssetBuilderTests.cpp
@@ -31,14 +31,6 @@ bool FCameraRigAssetBuilderNullTest::RunTest(const FString& Parameters)
 	return true;
 }
 
-namespace UE::Cameras
-{
-	//extern int32 GArrayCameraNodeEvaluatorSizeof;
-	//extern int32 GArrayCameraNodeEvaluatorAlignof;
-	//extern int32 GOffsetCameraNodeEvaluatorSizeof;
-	//extern int32 GOffsetCameraNodeEvaluatorAlignof;
-}
-
 IMPLEMENT_SIMPLE_AUTOMATION_TEST(FCameraRigAssetBuilderSimpleAllocationTest, "System.Engine.GameplayCameras.CameraRigAssetBuilder.SimpleAllocation", EAutomationTestFlags::EditorContext | EAutomationTestFlags::EngineFilter)
 bool FCameraRigAssetBuilderSimpleAllocationTest::RunTest(const FString& Parameters)
 {
@@ -54,9 +46,10 @@ bool FCameraRigAssetBuilderSimpleAllocationTest::RunTest(const FString& Paramete
 	UTEST_EQUAL("No evaluator allocation info", CameraRig->AllocationInfo.EvaluatorInfo.TotalSizeof, 0);
 	CameraRig->BuildCameraRig();
 
-	//int32 TotalSizeof = GArrayCameraNodeEvaluatorSizeof;
-	//TotalSizeof = Align(TotalSizeof, GOffsetCameraNodeEvaluatorAlignof) + GOffsetCameraNodeEvaluatorSizeof;
-	//UTEST_EQUAL("Evaluator allocation info", CameraRig->AllocationInfo.EvaluatorInfo.TotalSizeof, TotalSizeof);
+	int32 TotalSizeof = UArrayCameraNode::GetEvaluatorAllocationInfo().Key;
+	const TTuple<int32, int32> OffsetEvaluatorInfo = UOffsetCameraNode::GetEvaluatorAllocationInfo();
+	TotalSizeof = Align(TotalSizeof, OffsetEvaluatorInfo.Value) + OffsetEvaluatorInfo.Key;
+	UTEST_EQUAL("Evaluator allocation info", CameraRig->AllocationInfo.EvaluatorInfo.TotalSizeof, TotalSizeof);
 
 	return true;
 }
diff --git a/Engine/Plugins/Cameras/GameplayCameras/Source/GameplayCameras/Public/Nodes/Common/ArrayCameraNode.h b/Engine/Plugins/Cameras/GameplayCameras/Source/GameplayCameras/Public/Nodes/Common/ArrayCameraNode.h
--- a/Engine/Plugins/Cameras/GameplayCameras/Source/GameplayCameras/Public/Nodes/Common/ArrayCameraNode.h
+++ b/Engine/Plugins/Cameras/GameplayCameras/Source/GameplayCameras/Public/Nodes/Common/ArrayCameraNode.h
@@ -26,5 +26,10 @@ public:
 	/** The camera nodes to run. */
 	UPROPERTY()
 	TArray<TObjectPtr<UCameraNode>> Children;
+
+public:
+
+	// For unit tests.
+	static GAMEPLAYCAMERAS_API TTuple<int32, int32> GetEvaluatorAllocationInfo();
 };
 
diff --git a/Engine/Plugins/Cameras/GameplayCameras/Source/GameplayCameras/Public/Nodes/Common/OffsetCameraNode.h b/Engine/Plugins/Cameras/GameplayCameras/Source/GameplayCameras/Public/Nodes/Common/OffsetCameraNode.h
--- a/Engine/Plugins/Cameras/GameplayCameras/Source/GameplayCameras/Public/Nodes/Common/OffsetCameraNode.h
+++ b/Engine/Plugins/Cameras/GameplayCameras/Source/GameplayCameras/Public/Nodes/Common/OffsetCameraNode.h
@@ -34,5 +34,10 @@ public:
 	/** The space in which to apply the offset. */
 	UPROPERTY(EditAnywhere, Category=Common)
 	ECameraNodeSpace OffsetSpace = ECameraNodeSpace::CameraPose;
+
+public:
+
+	// For unit tests.
+	static GAMEPLAYCAMERAS_API TTuple<int32, int32> GetEvaluatorAllocationInfo();
 };
 
