Cameras: slightly change how the camera system actor is auto-spawned

- Move the code from the host class to the actor itself.
- Add a utility function for ensuring that there is a camera system actor spawned and that it is the view target.
- Add a utility function for checking if a camera system actor is active.

diff --git a/Engine/Plugins/Cameras/GameplayCameras/Source/GameplayCameras/Private/GameFramework/GameplayCameraComponent.cpp b/Engine/Plugins/Cameras/GameplayCameras/Source/GameplayCameras/Private/GameFramework/GameplayCameraComponent.cpp
--- a/Engine/Plugins/Cameras/GameplayCameras/Source/GameplayCameras/Private/GameFramework/GameplayCameraComponent.cpp
+++ b/Engine/Plugins/Cameras/GameplayCameras/Source/GameplayCameras/Private/GameFramework/GameplayCameraComponent.cpp
@@ -9,6 +9,7 @@
 #include "Engine/EngineTypes.h"
 #include "Engine/StaticMesh.h"
 #include "Engine/World.h"
+#include "GameFramework/GameplayCameraSystemActor.h"
 #include "GameFramework/GameplayCameraSystemHost.h"
 #include "GameplayCameras.h"
 #include "Kismet/GameplayStatics.h"
@@ -114,6 +115,8 @@ void UGameplayCameraComponent::ActivateCameraEvaluationContext(APlayerController
 		return;
 	}
 
+	AGameplayCameraSystemActor::AutoManageActiveViewTarget(PlayerController);
+
 	if (!EvaluationContext.IsValid())
 	{
 		EvaluationContext = MakeShared<FGameplayCameraComponentEvaluationContext>();
diff --git a/Engine/Plugins/Cameras/GameplayCameras/Source/GameplayCameras/Private/GameFramework/GameplayCameraSystemActor.cpp b/Engine/Plugins/Cameras/GameplayCameras/Source/GameplayCameras/Private/GameFramework/GameplayCameraSystemActor.cpp
--- a/Engine/Plugins/Cameras/GameplayCameras/Source/GameplayCameras/Private/GameFramework/GameplayCameraSystemActor.cpp
+++ b/Engine/Plugins/Cameras/GameplayCameras/Source/GameplayCameras/Private/GameFramework/GameplayCameraSystemActor.cpp
@@ -5,7 +5,9 @@
 #include "Engine/EngineTypes.h"
 #include "Engine/World.h"
 #include "GameFramework/GameplayCameraSystemComponent.h"
+#include "GameFramework/GameplayCameraSystemHost.h"
 #include "GameFramework/PlayerController.h"
+#include "GameplayCamerasSettings.h"
 
 #include UE_INLINE_GENERATED_CPP_BY_NAME(GameplayCameraSystemActor)
 
@@ -37,5 +39,37 @@ void AGameplayCameraSystemActor::EndViewTarget(APlayerController* PC)
 	Super::EndViewTarget(PC);
 }
 
+void AGameplayCameraSystemActor::AutoManageActiveViewTarget(APlayerController* PlayerController)
+{
+	static const TCHAR* AutoSpawnedActorName = TEXT("AutoSpawnedGameplayCameraSystemActor");
+
+	const UGameplayCamerasSettings* Settings = GetDefault<UGameplayCamerasSettings>();
+	if (!Settings->bAutoSpawnCameraSystemActor)
+	{
+		return;
+	}
+
+	UGameplayCameraSystemHost* Host = UGameplayCameraSystemHost::FindHost(PlayerController);
+	if (!Host)
+	{
+		UE_LOG(LogCameraSystem, Error, TEXT("Can't auto-manage active view target: no camera system host found!"));
+		return;
+	}
+	
+	AGameplayCameraSystemActor* SpawnedActor = FindObject<AGameplayCameraSystemActor>(PlayerController, AutoSpawnedActorName);
+	if (!SpawnedActor)
+	{
+		FActorSpawnParameters SpawnParams;
+		SpawnParams.Name = AutoSpawnedActorName;
+
+		UWorld* World = PlayerController->GetWorld();
+		SpawnedActor = World->SpawnActor<AGameplayCameraSystemActor>(SpawnParams);
+
+		SpawnedActor->Rename(nullptr, PlayerController);
+	}
+	check(SpawnedActor);
+	SpawnedActor->GetCameraSystemComponent()->ActivateCameraSystemForPlayerController(PlayerController);
+}
+
 #undef LOCTEXT_NAMESPACE
 
diff --git a/Engine/Plugins/Cameras/GameplayCameras/Source/GameplayCameras/Private/GameFramework/GameplayCameraSystemComponent.cpp b/Engine/Plugins/Cameras/GameplayCameras/Source/GameplayCameras/Private/GameFramework/GameplayCameraSystemComponent.cpp
--- a/Engine/Plugins/Cameras/GameplayCameras/Source/GameplayCameras/Private/GameFramework/GameplayCameraSystemComponent.cpp
+++ b/Engine/Plugins/Cameras/GameplayCameras/Source/GameplayCameras/Private/GameFramework/GameplayCameraSystemComponent.cpp
@@ -110,9 +110,12 @@ void UGameplayCameraSystemComponent::ActivateCameraSystemForPlayerIndex(int32 Pl
 
 void UGameplayCameraSystemComponent::ActivateCameraSystemForPlayerController(APlayerController* PlayerController)
 {
-	if (WeakPlayerController.IsValid())
+	if (APlayerController* ActivePlayerController = WeakPlayerController.Get())
 	{
-		DeactivateCameraSystem();
+		if (ActivePlayerController != PlayerController)
+		{
+			DeactivateCameraSystem();
+		}
 	}
 
 	AActor* OwningActor = GetOwner();
@@ -122,17 +125,47 @@ void UGameplayCameraSystemComponent::ActivateCameraSystemForPlayerController(APl
 		return;
 	}
 
-	CameraSystemHost = UGameplayCameraSystemHost::FindOrCreateHost(PlayerController);
 	if (!CameraSystemHost)
 	{
-		UE_LOG(LogCameraSystem, Error, TEXT("can't create camera system host!"));
-		return;
+		CameraSystemHost = UGameplayCameraSystemHost::FindOrCreateHost(PlayerController);
+		if (!CameraSystemHost)
+		{
+			UE_LOG(LogCameraSystem, Error, TEXT("can't create camera system host!"));
+			return;
+		}
 	}
 
 	PlayerController->SetViewTarget(OwningActor);
 	WeakPlayerController = PlayerController;
 }
 
+bool UGameplayCameraSystemComponent::IsCameraSystemActiveForPlayController(APlayerController* PlayerController) const
+{
+	APlayerController* ActivatedPlayerController = WeakPlayerController.Get();
+	if (!ActivatedPlayerController || ActivatedPlayerController  != PlayerController)
+	{
+		return false;
+	}
+
+	AActor* OwningActor = GetOwner();
+	if (!OwningActor)
+	{
+		return false;
+	}
+	
+	if (!CameraSystemHost)
+	{
+		return false;
+	}
+
+	if (!ActivatedPlayerController->PlayerCameraManager)
+	{
+		return false;
+	}
+
+	return ActivatedPlayerController->PlayerCameraManager->GetViewTarget() == OwningActor;
+}
+
 void UGameplayCameraSystemComponent::DeactivateCameraSystem(AActor* NextViewTarget)
 {
 	APlayerController* PlayerController = WeakPlayerController.Get();
diff --git a/Engine/Plugins/Cameras/GameplayCameras/Source/GameplayCameras/Private/GameFramework/GameplayCameraSystemHost.cpp b/Engine/Plugins/Cameras/GameplayCameras/Source/GameplayCameras/Private/GameFramework/GameplayCameraSystemHost.cpp
--- a/Engine/Plugins/Cameras/GameplayCameras/Source/GameplayCameras/Private/GameFramework/GameplayCameraSystemHost.cpp
+++ b/Engine/Plugins/Cameras/GameplayCameras/Source/GameplayCameras/Private/GameFramework/GameplayCameraSystemHost.cpp
@@ -7,7 +7,6 @@
 #include "GameFramework/GameplayCameraSystemActor.h"
 #include "GameFramework/GameplayCameraSystemComponent.h"
 #include "GameFramework/PlayerController.h"
-#include "GameplayCamerasSettings.h"
 #include "Templates/UnrealTemplate.h"
 
 const TCHAR* UGameplayCameraSystemHost::DefaultHostName = TEXT("GameplayCameraSystemHost");
@@ -40,16 +39,6 @@ UGameplayCameraSystemHost* UGameplayCameraSystemHost::FindOrCreateHost(APlayerCo
 
 	UGameplayCameraSystemHost* NewHost = NewObject<UGameplayCameraSystemHost>(PlayerController, HostName);
 
-	const UGameplayCamerasSettings* Settings = GetDefault<UGameplayCamerasSettings>();
-	if (Settings->bAutoSpawnCameraSystemActor)
-	{
-		UWorld* World = PlayerController->GetWorld();
-
-		FActorSpawnParameters SpawnParams;
-		AGameplayCameraSystemActor* CameraSystemActor = World->SpawnActor<AGameplayCameraSystemActor>(SpawnParams);
-		CameraSystemActor->GetCameraSystemComponent()->ActivateCameraSystemForPlayerController(PlayerController);
-	}
-
 	return NewHost;
 }
 
diff --git a/Engine/Plugins/Cameras/GameplayCameras/Source/GameplayCameras/Public/GameFramework/GameplayCameraSystemActor.h b/Engine/Plugins/Cameras/GameplayCameras/Source/GameplayCameras/Public/GameFramework/GameplayCameraSystemActor.h
--- a/Engine/Plugins/Cameras/GameplayCameras/Source/GameplayCameras/Public/GameFramework/GameplayCameraSystemActor.h
+++ b/Engine/Plugins/Cameras/GameplayCameras/Source/GameplayCameras/Public/GameFramework/GameplayCameraSystemActor.h
@@ -32,6 +32,16 @@ public:
 	UFUNCTION(BlueprintGetter, Category=Camera)
 	UGameplayCameraSystemComponent* GetCameraSystemComponent() const { return CameraSystemComponent; }
 
+public:
+
+	/**
+	 * Automatically sets a camera system actor as the view-target, spawning a unique
+	 * instance if needed. Doesn't do anything if bAutoSpawnCameraSystemActor is disabled
+	 * in the GameplayCameras settings.
+	 */
+	UFUNCTION(BlueprintCallable, Category=Camera)
+	static void AutoManageActiveViewTarget(APlayerController* PlayerController);
+
 public:
 
 	// AActor interface.
diff --git a/Engine/Plugins/Cameras/GameplayCameras/Source/GameplayCameras/Public/GameFramework/GameplayCameraSystemComponent.h b/Engine/Plugins/Cameras/GameplayCameras/Source/GameplayCameras/Public/GameFramework/GameplayCameraSystemComponent.h
--- a/Engine/Plugins/Cameras/GameplayCameras/Source/GameplayCameras/Public/GameFramework/GameplayCameraSystemComponent.h
+++ b/Engine/Plugins/Cameras/GameplayCameras/Source/GameplayCameras/Public/GameFramework/GameplayCameraSystemComponent.h
@@ -47,6 +47,10 @@ public:
 	UFUNCTION(BlueprintCallable, Category=Camera)
 	void ActivateCameraSystemForPlayerController(APlayerController* PlayerController);
 
+	/** Returns whether this component's actor is set as the view target for the given player. */
+	UFUNCTION(BlueprintCallable, Category=Camera)
+	bool IsCameraSystemActiveForPlayController(APlayerController* PlayerController) const;
+
 	/** Removes this component's actor from being the view target. */
 	UFUNCTION(BlueprintCallable, Category=Camera)
 	void DeactivateCameraSystem(AActor* NextViewTarget = nullptr);
@@ -91,6 +95,7 @@ private:
 	UPROPERTY()
 	TObjectPtr<UGameplayCameraSystemHost> CameraSystemHost;
 
+	UPROPERTY()
 	TWeakObjectPtr<APlayerController> WeakPlayerController;
 
 #if UE_GAMEPLAY_CAMERAS_DEBUG
diff --git a/Engine/Plugins/Cameras/GameplayCameras/Source/GameplayCameras/Public/GameFramework/GameplayCameraSystemHost.h b/Engine/Plugins/Cameras/GameplayCameras/Source/GameplayCameras/Public/GameFramework/GameplayCameraSystemHost.h
--- a/Engine/Plugins/Cameras/GameplayCameras/Source/GameplayCameras/Public/GameFramework/GameplayCameraSystemHost.h
+++ b/Engine/Plugins/Cameras/GameplayCameras/Source/GameplayCameras/Public/GameFramework/GameplayCameraSystemHost.h
@@ -15,7 +15,7 @@ namespace UE::Cameras
 }
 
 /**
- * A class that host a camera system evaluator so that it can be accessed in a game world.
+ * A class that hosts a camera system evaluator so that it can be accessed in a game world.
  *
  * The host doesn't stay alive very long if nothing references it. Gameplay camera components and actors
  * are meant to hold a reference to it while they use it. When nobody uses it, the host is meant to be
