Cameras: tweak display names for camera variable tables' BP functions

diff --git a/Engine/Plugins/Cameras/GameplayCameras/Source/GameplayCameras/Public/GameFramework/BlueprintCameraVariableTable.h b/Engine/Plugins/Cameras/GameplayCameras/Source/GameplayCameras/Public/GameFramework/BlueprintCameraVariableTable.h
--- a/Engine/Plugins/Cameras/GameplayCameras/Source/GameplayCameras/Public/GameFramework/BlueprintCameraVariableTable.h
+++ b/Engine/Plugins/Cameras/GameplayCameras/Source/GameplayCameras/Public/GameFramework/BlueprintCameraVariableTable.h
@@ -61,77 +61,77 @@ class UBlueprintCameraVariableTableFunctionLibrary : public UBlueprintFunctionLi
 public:
 
 	/** Gets a camera variable's value from the given table. */
-	UFUNCTION(BlueprintCallable, Category=Camera, meta=(CompactNodeTitle="Get Bool"))
+	UFUNCTION(BlueprintCallable, DisplayName="Get Boolean Variable", Category=Camera, meta=(CompactNodeTitle="Get Boolean Variable"))
 	static bool GetBooleanCameraVariable(const FBlueprintCameraVariableTable& VariableTable, UBooleanCameraVariable* Variable);
 
 	/** Gets a camera variable's value from the given table. */
-	UFUNCTION(BlueprintCallable, Category=Camera, meta=(CompactNodeTitle="Get Int"))
+	UFUNCTION(BlueprintCallable, DisplayName="Get Integer Variable", Category=Camera, meta=(CompactNodeTitle="Get Integer Variable"))
 	static int32 GetInteger32CameraVariable(const FBlueprintCameraVariableTable& VariableTable, UInteger32CameraVariable* Variable);
 
 	/** Gets a camera variable's value from the given table. */
-	UFUNCTION(BlueprintCallable, Category=Camera, meta=(CompactNodeTitle="Get Float"))
+	UFUNCTION(BlueprintCallable, DisplayName="Get Float Variable", Category=Camera, meta=(CompactNodeTitle="Get Float Variable"))
 	static float GetFloatCameraVariable(const FBlueprintCameraVariableTable& VariableTable, UFloatCameraVariable* Variable);
 
 	/** Gets a camera variable's value from the given table. */
-	UFUNCTION(BlueprintCallable, Category=Camera, meta=(CompactNodeTitle="Get Double"))
+	UFUNCTION(BlueprintCallable, DisplayName="Get Double Variable", Category=Camera, meta=(CompactNodeTitle="Get Double Variable"))
 	static double GetDoubleCameraVariable(const FBlueprintCameraVariableTable& VariableTable, UDoubleCameraVariable* Variable);
 
 	/** Gets a camera variable's value from the given table. */
-	UFUNCTION(BlueprintCallable, Category=Camera, meta=(CompactNodeTitle="Get Vec2"))
+	UFUNCTION(BlueprintCallable, DisplayName="Get Vector 2D Variable", Category=Camera, meta=(CompactNodeTitle="Get Vector 2D Variable"))
 	static FVector2D GetVector2CameraVariable(const FBlueprintCameraVariableTable& VariableTable, UVector2dCameraVariable* Variable);
 
 	/** Gets a camera variable's value from the given table. */
-	UFUNCTION(BlueprintCallable, Category=Camera, meta=(CompactNodeTitle="Get Vec3"))
+	UFUNCTION(BlueprintCallable, DisplayName="Get Vector Variable", Category=Camera, meta=(CompactNodeTitle="Get Vector Variable"))
 	static FVector GetVector3CameraVariable(const FBlueprintCameraVariableTable& VariableTable, UVector3dCameraVariable* Variable);
 
 	/** Gets a camera variable's value from the given table. */
-	UFUNCTION(BlueprintCallable, Category=Camera, meta=(CompactNodeTitle="Get Vec4"))
+	UFUNCTION(BlueprintCallable, DisplayName="Get Vector 4 Variable", Category=Camera, meta=(CompactNodeTitle="Get Vector 4 Variable"))
 	static FVector4 GetVector4CameraVariable(const FBlueprintCameraVariableTable& VariableTable, UVector4dCameraVariable* Variable);
 
 	/** Gets a camera variable's value from the given table. */
-	UFUNCTION(BlueprintCallable, Category=Camera, meta=(CompactNodeTitle="Get Rotator"))
+	UFUNCTION(BlueprintCallable, DisplayName="Get Rotator Variable", Category=Camera, meta=(CompactNodeTitle="Get Rotator Variable"))
 	static FRotator GetRotatorCameraVariable(const FBlueprintCameraVariableTable& VariableTable, URotator3dCameraVariable* Variable);
 
 	/** Gets a camera variable's value from the given table. */
-	UFUNCTION(BlueprintCallable, Category=Camera, meta=(CompactNodeTitle="Get Transform"))
+	UFUNCTION(BlueprintCallable, DisplayName="Get Transform Variable", Category=Camera, meta=(CompactNodeTitle="Get Transform Variable"))
 	static FTransform GetTransformCameraVariable(const FBlueprintCameraVariableTable& VariableTable, UTransform3dCameraVariable* Variable);
 
 public:
 
 	/** Sets a camera variable's value in the given table. */
-	UFUNCTION(BlueprintCallable, Category=Camera)
+	UFUNCTION(BlueprintCallable, DisplayName="Set Boolean Variable", Category=Camera)
 	static void SetBooleanCameraVariable(const FBlueprintCameraVariableTable& VariableTable, UBooleanCameraVariable* Variable, bool Value);
 
 	/** Sets a camera variable's value in the given table. */
-	UFUNCTION(BlueprintCallable, Category=Camera)
+	UFUNCTION(BlueprintCallable, DisplayName="Set Integer Variable", Category=Camera)
 	static void SetInteger32CameraVariable(const FBlueprintCameraVariableTable& VariableTable, UInteger32CameraVariable* Variable, int32 Value);
 
 	/** Sets a camera variable's value in the given table. */
-	UFUNCTION(BlueprintCallable, Category=Camera)
+	UFUNCTION(BlueprintCallable, DisplayName="Set Float Variable", Category=Camera)
 	static void SetFloatCameraVariable(const FBlueprintCameraVariableTable& VariableTable, UFloatCameraVariable* Variable, float Value);
 
 	/** Sets a camera variable's value in the given table. */
-	UFUNCTION(BlueprintCallable, Category=Camera)
+	UFUNCTION(BlueprintCallable, DisplayName="Set Double Variable", Category=Camera)
 	static void SetDoubleCameraVariable(const FBlueprintCameraVariableTable& VariableTable, UDoubleCameraVariable* Variable, double Value);
 
 	/** Sets a camera variable's value in the given table. */
-	UFUNCTION(BlueprintCallable, Category=Camera)
+	UFUNCTION(BlueprintCallable, DisplayName="Set Vector 2D Variable", Category=Camera)
 	static void SetVector2CameraVariable(const FBlueprintCameraVariableTable& VariableTable, UVector2dCameraVariable* Variable, const FVector2D& Value);
 
 	/** Sets a camera variable's value in the given table. */
-	UFUNCTION(BlueprintCallable, Category=Camera)
+	UFUNCTION(BlueprintCallable, DisplayName="Set Vector Variable", Category=Camera)
 	static void SetVector3CameraVariable(const FBlueprintCameraVariableTable& VariableTable, UVector3dCameraVariable* Variable, const FVector& Value);
 
 	/** Sets a camera variable's value in the given table. */
-	UFUNCTION(BlueprintCallable, Category=Camera)
+	UFUNCTION(BlueprintCallable, DisplayName="Set Vector 4 Variable", Category=Camera)
 	static void SetVector4CameraVariable(const FBlueprintCameraVariableTable& VariableTable, UVector4dCameraVariable* Variable, const FVector4& Value);
 
 	/** Sets a camera variable's value in the given table. */
-	UFUNCTION(BlueprintCallable, Category=Camera)
+	UFUNCTION(BlueprintCallable, DisplayName="Set Rotator Variable", Category=Camera)
 	static void SetRotatorCameraVariable(const FBlueprintCameraVariableTable& VariableTable, URotator3dCameraVariable* Variable, const FRotator& Value);
 
 	/** Sets a camera variable's value in the given table. */
-	UFUNCTION(BlueprintCallable, Category=Camera)
+	UFUNCTION(BlueprintCallable, DisplayName="Set Transform Variable", Category=Camera)
 	static void SetTransformCameraVariable(const FBlueprintCameraVariableTable& VariableTable, UTransform3dCameraVariable* Variable, const FTransform& Value);
 };
 
