Cameras: allow exposing camera rig interface parameters over multiple nesting levels

This is a big change since it affects everything from runtime parameter overrides to camera rig build process to UI.

Runtime:
- Add a "PreBuild" step to camera nodes, so that the camera rig node ("prefab node") can build the inner rig before we start gathering exposed parameters.
- Add a "CameraRigReference" structure, similar to the "StateTreeReference", which groups a reference to a camera rig asset along with any desired parameter overrides under the same umbrella.
- The camera rig builder now handles the situation where an interface parameter targets a prefab node's inner camera rig interface parameter. This is handled by creating an override on the prefab node, and driving that override with the outer interface parameter's private variable.

Editor:
- Add a specific graph node class for prefab nodes. They show their inner camera rig's interface parameters as input pins that can be themselves exposed as interface parameters on the outer rig.
- Remove the detail customization for the prefab node, and make a property type customization instead for the new CameraRigReference type. This is better since the old customization would have needed tweaks if we added new properties on the prefab node class.
- The CameraRigReference detail customization better reacts to changes in the inner camera rig, or re-assigning the referenced camera rig, thanks to some new global delegate firing when a camera rig has been built.
- Various code changes in the graph schema to remove assumptions that exposed pins are necessarily for Struct properties on a camera node (now they could be for parameter overrides, which aren't "real" properties in the UObject RTTI sense).
- Make the camera rig pointer detail customization handle the case of the object-ref property being inside a CameraRigReference type. In that case, it needs to look up metadata on that parent property.

Tests:
- Add automated tests for nested exposed parameters.
- Add a "named object registry" to unit test framework to make it easier to generate data and access that data later for value checks.

diff --git a/Engine/Plugins/Cameras/GameplayCameras/Source/GameplayCameras/Private/Core/CameraNode.cpp b/Engine/Plugins/Cameras/GameplayCameras/Source/GameplayCameras/Private/Core/CameraNode.cpp
--- a/Engine/Plugins/Cameras/GameplayCameras/Source/GameplayCameras/Private/Core/CameraNode.cpp
+++ b/Engine/Plugins/Cameras/GameplayCameras/Source/GameplayCameras/Private/Core/CameraNode.cpp
@@ -28,6 +28,11 @@ FCameraNodeChildrenView UCameraNode::GetChildren()
 	return OnGetChildren();
 }
 
+void UCameraNode::PreBuild(FCameraBuildLog& BuildLog)
+{
+	OnPreBuild(BuildLog);
+}
+
 void UCameraNode::Build(FCameraRigBuildContext& BuildContext)
 {
 	OnBuild(BuildContext);
diff --git a/Engine/Plugins/Cameras/GameplayCameras/Source/GameplayCameras/Private/Core/CameraRigAsset.cpp b/Engine/Plugins/Cameras/GameplayCameras/Source/GameplayCameras/Private/Core/CameraRigAsset.cpp
--- a/Engine/Plugins/Cameras/GameplayCameras/Source/GameplayCameras/Private/Core/CameraRigAsset.cpp
+++ b/Engine/Plugins/Cameras/GameplayCameras/Source/GameplayCameras/Private/Core/CameraRigAsset.cpp
@@ -5,10 +5,7 @@
 #include "Core/CameraAsset.h"
 #include "Core/CameraBuildLog.h"
 #include "Core/CameraNode.h"
-#include "Core/CameraParameters.h"
 #include "Core/CameraRigAssetBuilder.h"
-#include "Core/CameraVariableAssets.h"
-#include "Templates/UnrealTypeTraits.h"
 #include "UObject/ObjectSaveContext.h"
 
 #include UE_INLINE_GENERATED_CPP_BY_NAME(CameraRigAsset)
@@ -48,7 +45,6 @@ void UCameraRigInterfaceParameter::OnGraphNodeMoved(FName InGraphName, int32 Nod
 
 #endif
 
-
 void UCameraRigInterfaceParameter::PostLoad()
 {
 	if (!Guid.IsValid())
@@ -82,13 +78,24 @@ void UCameraRigInterfaceParameter::PostDuplicate(EDuplicateMode::Type DuplicateM
 
 UCameraRigInterfaceParameter* FCameraRigInterface::FindInterfaceParameterByName(const FString& ParameterName) const
 {
-	const TObjectPtr<UCameraRigInterfaceParameter>* FoundItem = InterfaceParameters.FindByPredicate([&ParameterName](UCameraRigInterfaceParameter* Item)
+	const TObjectPtr<UCameraRigInterfaceParameter>* FoundItem = InterfaceParameters.FindByPredicate(
+			[&ParameterName](UCameraRigInterfaceParameter* Item)
 			{
 				return Item->InterfaceParameterName == ParameterName;
 			});
 	return FoundItem ? *FoundItem : nullptr;
 }
 
+UCameraRigInterfaceParameter* FCameraRigInterface::FindInterfaceParameterByGuid(const FGuid& ParameterGuid) const
+{
+	const TObjectPtr<UCameraRigInterfaceParameter>* FoundItem = InterfaceParameters.FindByPredicate(
+			[&ParameterGuid](UCameraRigInterfaceParameter* Item)
+			{
+				return Item->Guid == ParameterGuid;
+			});
+	return FoundItem ? *FoundItem : nullptr;
+}
+
 bool FCameraRigInterface::HasInterfaceParameter(const FString& ParameterName) const
 {
 	return FindInterfaceParameterByName(ParameterName) != nullptr;
diff --git a/Engine/Plugins/Cameras/GameplayCameras/Source/GameplayCameras/Private/Core/CameraRigAssetBuilder.cpp b/Engine/Plugins/Cameras/GameplayCameras/Source/GameplayCameras/Private/Core/CameraRigAssetBuilder.cpp
--- a/Engine/Plugins/Cameras/GameplayCameras/Source/GameplayCameras/Private/Core/CameraRigAssetBuilder.cpp
+++ b/Engine/Plugins/Cameras/GameplayCameras/Source/GameplayCameras/Private/Core/CameraRigAssetBuilder.cpp
@@ -8,7 +8,9 @@
 #include "Core/CameraRigBuildContext.h"
 #include "Core/CameraRigAsset.h"
 #include "Core/CameraVariableAssets.h"
+#include "GameplayCamerasDelegates.h"
 #include "Logging/TokenizedMessage.h"
+#include "Nodes/Common/CameraRigCameraNode.h"
 
 #define LOCTEXT_NAMESPACE "CameraRigAssetBuilder"
 
@@ -84,11 +86,41 @@ struct FPrivateVariableBuilder
 	template<typename ExpectedVariableAssetType>
 	ExpectedVariableAssetType* FindReusablePrivateVariable(FStructProperty* ForParameterProperty, UCameraNode* ForCameraNode)
 	{
-		UCameraVariableAsset* ReusedVariable = FindReusablePrivateVariable(ForParameterProperty, ForCameraNode);
-		if (ReusedVariable)
+		using FDrivenParameterKey = FCameraRigAssetBuilder::FDrivenParameterKey;
+
+		FDrivenParameterKey ParameterKey{ ForParameterProperty, ForCameraNode };
+		UCameraVariableAsset** FoundItem = Owner.OldDrivenParameters.Find(ParameterKey);
+		if (FoundItem)
 		{
+			// We found an existing variable that was driving this camera node's property.
+			// Re-use it and remove it from the re-use pool.
+			UCameraVariableAsset* ReusedVariable = (*FoundItem);
+			Owner.OldDrivenParameters.Remove(ParameterKey);
 			return CastChecked<ExpectedVariableAssetType>(ReusedVariable);
 		}
+
+		return nullptr;
+	}
+
+	template<
+		typename ParameterOverrideType, 
+		typename ExpectedVariableAssetType = typename ParameterOverrideType::CameraParameterType::VariableAssetType>
+	ExpectedVariableAssetType* FindReusablePrivateVariable(ParameterOverrideType* ForParameterOverride, UCameraRigCameraNode* ForCameraRigNode)
+	{
+		using FDrivenOverrideKey = FCameraRigAssetBuilder::FDrivenOverrideKey;
+
+		FDrivenOverrideKey OverrideKey{ ForParameterOverride->InterfaceParameterGuid, ForCameraRigNode };
+		UCameraVariableAsset** FoundItem = Owner.OldDrivenOverrides.Find(OverrideKey);
+		if (FoundItem)
+		{
+			UCameraVariableAsset* ReusedVariable = (*FoundItem);
+			Owner.OldDrivenOverrides.Remove(OverrideKey);
+			// Don't do a checked cast here because interface parameters can change type if they are
+			// reconnected to a different type of parameter. If the cast fail, we will retur null and
+			// the variable won't get reused, which is what we want.
+			return Cast<ExpectedVariableAssetType>(ReusedVariable);
+		}
+
 		return nullptr;
 	}
 
@@ -109,41 +141,17 @@ struct FPrivateVariableBuilder
 		return true;
 	}
 
-private:
-
-	UCameraVariableAsset* FindReusablePrivateVariable(FStructProperty* ForParameterProperty, UCameraNode* ForCameraNode)
-	{
-		using FDrivenParameterKey = FCameraRigAssetBuilder::FDrivenParameterKey;
-
-		FDrivenParameterKey ParameterKey{ ForParameterProperty, ForCameraNode };
-		UCameraVariableAsset** FoundItem = Owner.OldDrivenParameters.Find(ParameterKey);
-		if (FoundItem)
-		{
-			// We found an existing variable that was driving this camera node's property.
-			// Re-use it and remove it from the re-use pool.
-			UCameraVariableAsset* ReusedVariable = (*FoundItem);
-			Owner.OldDrivenParameters.Remove(ParameterKey);
-			return ReusedVariable;
-		}
-
-		return nullptr;
-	}
-
 private:
 
 	FCameraRigAssetBuilder& Owner;
 };
 
 template<typename CameraParameterType>
-void SetupPrivateVariable(
+void CheckTwiceDrivenParameter(
 		FPrivateVariableBuilder& Builder, 
 		UCameraRigInterfaceParameter* InterfaceParameter, 
-		FStructProperty* ParameterTargetProperty,
 		CameraParameterType* CameraParameter)
 {
-	using ValueType = typename CameraParameterType::ValueType;
-	using VariableAssetType = typename CameraParameterType::VariableAssetType;
-
 	if (CameraParameter->Variable != nullptr)
 	{
 		// We should have cleared all exposed parameters in GatherOldDrivenParameters, so the only variables
@@ -165,19 +173,21 @@ void SetupPrivateVariable(
 						FText::FromName(InterfaceParameter->TargetPropertyName)));
 		}
 	}
+}
 
-	// Start by re-using the camera variable that was already driving this parameter on this node.
-	VariableAssetType* PrivateVariable = Builder.FindReusablePrivateVariable<VariableAssetType>(
-			ParameterTargetProperty, InterfaceParameter->Target);
-	const bool bIsReusedVariable = (PrivateVariable != nullptr);
-
+template<typename VariableAssetType>
+VariableAssetType* MakeOrRenamePrivateVariable(
+		FPrivateVariableBuilder& Builder,
+		const FString& InterfaceParameterName,
+		VariableAssetType* PrivateVariable)
+{
 	const FString VariableName = FString::Format(
 			TEXT("Override_{0}_{1}"), 
-			{ Builder.CameraRig->GetName(), InterfaceParameter->InterfaceParameterName });
+			{ Builder.CameraRig->GetName(), InterfaceParameterName });
 
 	if (PrivateVariable)
 	{
-		// Found a pre-existing variable! Make sure it's got the right name, in case the exposed rig parameter
+		// We have a pre-existing variable! Make sure it's got the right name, in case the exposed rig parameter
 		// was renamed. Keeping a good name is mostly to help with debugging.
 		FString OriginalName = PrivateVariable->GetName();
 		OriginalName.RemoveFromStart("REUSABLE_", ESearchCase::CaseSensitive);
@@ -195,11 +205,30 @@ void SetupPrivateVariable(
 		PrivateVariable = NewObject<VariableAssetType>(Builder.CameraRig, FName(*VariableName), RF_Transactional);
 	}
 
-	ensure(PrivateVariable->GetOuter() == Builder.CameraRig);
-
+	// Make sure it's a private variable.
 	PrivateVariable->bIsPrivate = true;
 	PrivateVariable->bAutoReset = false;
 
+	return PrivateVariable;
+}
+
+template<typename CameraParameterType, typename VariableAssetType = typename CameraParameterType::VariableAssetType>
+void DoSetupPrivateVariable(
+		FPrivateVariableBuilder& Builder,
+		UCameraRigInterfaceParameter* InterfaceParameter,
+		CameraParameterType* CameraParameter,
+		VariableAssetType* ReusedVariable
+		)
+{
+	using ValueType = typename CameraParameterType::ValueType;
+
+	const bool bIsReusedVariable = (ReusedVariable != nullptr);
+
+	// Either rename the camera variable we are re-using, or make a new one with the right name.
+	VariableAssetType* PrivateVariable = MakeOrRenamePrivateVariable(
+			Builder, InterfaceParameter->InterfaceParameterName, ReusedVariable);
+	ensure(PrivateVariable->GetOuter() == Builder.CameraRig);
+
 	// Set the default value of the variable to be the value in the camera parameter.
 	SetPrivateVariableDefaultValue<VariableAssetType, ValueType>(PrivateVariable, CameraParameter->Value);
 
@@ -218,6 +247,43 @@ void SetupPrivateVariable(
 	CameraParameter->Variable = PrivateVariable;
 }
 
+template<typename CameraParameterType>
+void SetupPrivateVariable(
+		FPrivateVariableBuilder& Builder, 
+		UCameraRigInterfaceParameter* InterfaceParameter, 
+		FStructProperty* ParameterTargetProperty,
+		CameraParameterType* CameraParameter)
+{
+	using ValueType = typename CameraParameterType::ValueType;
+	using VariableAssetType = typename CameraParameterType::VariableAssetType;
+
+	CheckTwiceDrivenParameter(Builder, InterfaceParameter, CameraParameter);
+
+	VariableAssetType* ReusedVariable = Builder.FindReusablePrivateVariable<VariableAssetType>(
+			ParameterTargetProperty, InterfaceParameter->Target);
+
+	DoSetupPrivateVariable(Builder, InterfaceParameter, CameraParameter, ReusedVariable);
+}
+
+template<typename ParameterOverrideType>
+void SetupPrivateVariable(
+		FPrivateVariableBuilder& Builder, 
+		UCameraRigInterfaceParameter* InterfaceParameter, 
+		ParameterOverrideType* ParameterOverride)
+{
+	using CameraParameterType = typename ParameterOverrideType::CameraParameterType;
+	using ValueType = typename CameraParameterType::ValueType;
+	using VariableAssetType = typename CameraParameterType::VariableAssetType;
+
+	CheckTwiceDrivenParameter(Builder, InterfaceParameter, &ParameterOverride->Value);
+
+	UCameraRigCameraNode* CameraRigNode = CastChecked<UCameraRigCameraNode>(InterfaceParameter->Target);
+	VariableAssetType* ReusedVariable = Builder.FindReusablePrivateVariable<ParameterOverrideType>(
+			ParameterOverride, CameraRigNode);
+
+	DoSetupPrivateVariable(Builder, InterfaceParameter, &ParameterOverride->Value, ReusedVariable);
+}
+
 }  // namespace Internal
 
 FCameraRigAssetBuilder::FCameraRigAssetBuilder(FCameraBuildLog& InBuildLog)
@@ -238,11 +304,15 @@ void FCameraRigAssetBuilder::BuildCameraRig(UCameraRigAsset* InCameraRig, FCusto
 	}
 
 	CameraRig = InCameraRig;
-	BuildLog.SetLoggingPrefix(InCameraRig->GetPathName() + TEXT(": "));
+	BuildLog.SetLoggingPrefix(CameraRig->GetPathName() + TEXT(": "));
 	{
 		BuildCameraRigImpl();
 
 		InCustomBuildStep.ExecuteIfBound(CameraRig, BuildLog);
+
+		CameraRig->EventHandlers.Notify(&ICameraRigAssetEventHandler::OnCameraRigBuilt, CameraRig);
+
+		FGameplayCamerasDelegates::OnCameraRigAssetBuilt().Broadcast(CameraRig, BuildLog);
 	}
 	BuildLog.SetLoggingPrefix(FString());
 	UpdateBuildStatus();
@@ -252,12 +322,16 @@ void FCameraRigAssetBuilder::BuildCameraRigImpl()
 {
 	if (!CameraRig->RootNode)
 	{
-		BuildLog.AddMessage(EMessageSeverity::Error, CameraRig, LOCTEXT("MissingRootNode", "Camera rig has no root node set."));
+		BuildLog.AddMessage(EMessageSeverity::Error, CameraRig, 
+				FText::Format(LOCTEXT("MissingRootNode", "Camera rig '{0}' has no root node."), 
+					FText::FromString(GetPathNameSafe(CameraRig))));
 		return;
 	}
 
 	FlattenCameraNodeHierarchy();
 
+	CallPreBuild();
+
 	GatherOldDrivenParameters();
 	BuildNewDrivenParameters();
 	DiscardUnusedPrivateVariables();
@@ -308,6 +382,14 @@ void FCameraRigAssetBuilder::FlattenCameraNodeHierarchy()
 #endif  // WITH_EDITORONLY_DATA
 }
 
+void FCameraRigAssetBuilder::CallPreBuild()
+{
+	for (UCameraNode* CameraNode : FlattenedNodes)
+	{
+		CameraNode->PreBuild(BuildLog);
+	}
+}
+
 void FCameraRigAssetBuilder::GatherOldDrivenParameters()
 {
 	// Keep track of what camera parameters were previously driven by private variables,
@@ -319,24 +401,56 @@ void FCameraRigAssetBuilder::GatherOldDrivenParameters()
 	//
 	// Note that parameters driven by user-defined variables are left alone.
 
-	OldDrivenParameters.Reset();
 	TSet<UCameraVariableAsset*> GatheredVariables;
+	TSet<UCameraNode*> CameraNodesToGather(FlattenedNodes);
 
-	// Get our camera nodes from AllNodeTreeObjects if possible, so we also gather old 
-	// driven parameters from disconnected camera nodes.
-#if WITH_EDITORONLY_DATA
-	TArray<UObject*> ObjectsToGather(ObjectPtrDecay(CameraRig->AllNodeTreeObjects));
-#else
-	TArray<UObject*> ObjectsToGather(FlattenedNodes);
-#endif
-	for (UObject* Object : ObjectsToGather)
-	{
-		UCameraNode* CameraNode = Cast<UCameraNode>(Object);
-		if (!CameraNode)
+	// Start by going through all interface parameters, remembering what private variable 
+	// they were associated with originally. Also collect that private variable to be
+	// renamed and put in the re-use pool.
+
+	OldInterfaceParameters.Reset();
+
+	for (UCameraRigInterfaceParameter* InterfaceParameter : CameraRig->Interface.InterfaceParameters)
+	{
+		OldInterfaceParameters.Add(InterfaceParameter, FReusableInterfaceParameterInfo(InterfaceParameter->PrivateVariable, false));
+		if (InterfaceParameter->PrivateVariable)
 		{
-			continue;
+			GatheredVariables.Add(InterfaceParameter->PrivateVariable);
 		}
 
+		ensureMsgf(
+				InterfaceParameter->Target == nullptr || InterfaceParameter->IsInOuter(CameraRig),
+				TEXT("Interface parameter '%s' points to camera node '%s' which isn't outer'ed to camera rig '%s'."),
+				*InterfaceParameter->InterfaceParameterName,
+				*GetNameSafe(InterfaceParameter->Target),
+				*GetPathNameSafe(CameraRig));
+	}
+
+	// Next go through all the camera nodes we know of. Nodes in FlattenedNodes are the ones
+	// connected to the camera rig's root node, so we are missing nodes that were disconnected
+	// since the last build. We could use AllNodeTreeObjects for that, but it only exists in 
+	// editor builds, and we don't want to rely on unit tests or runtime data manipulation to 
+	// have correctly populated it, so we'll try to gather any stray nodes by looking at
+	// objects outer'ed to the camera rig.
+
+	OldDrivenParameters.Reset();
+
+	ForEachObjectWithOuter(CameraRig, [&CameraNodesToGather](UObject* Obj)
+			{
+				if (UCameraNode* CameraNode = Cast<UCameraNode>(Obj))
+				{
+					CameraNodesToGather.Add(CameraNode);
+				}
+			});
+	const int32 NumStrayCameraNodes = (CameraNodesToGather.Num() - FlattenedNodes.Num());
+	if (NumStrayCameraNodes > 0)
+	{
+		UE_LOG(LogCameraSystem, Verbose, TEXT("Collected %d stray camera nodes while building camera rig '%s'."),
+				NumStrayCameraNodes, *GetPathNameSafe(CameraRig));
+	}
+
+	for (UCameraNode* CameraNode : CameraNodesToGather)
+	{
 		UClass* CameraNodeClass = CameraNode->GetClass();
 		
 		for (TFieldIterator<FProperty> It(CameraNodeClass); It; ++It)
@@ -371,35 +485,29 @@ UE_CAMERA_VARIABLE_FOR_ALL_TYPES()
 				// Some other struct property.
 			}
 		}
-	}
 
-	// Look for interface parameters that were disconnected in the graph editor. They may still 
-	// have a private variable inside them. We will discard those, but we want to rename them so
-	// that a new one can take their place (such as when another interface parameter gets connected
-	// to a similar camera node property and has the same name).
-	OldInterfaceParameters.Reset();
-
-#if WITH_EDITORONLY_DATA
-	// Leave ObjectsToGather as AllNodeTreeObjects.
-#else
-	ObjectsToGather.Reset();
-	for (UCameraRigInterfaceParameter* InterfaceParameter : CameraRig->Interface.InterfaceParameters)
-	{
-		ObjectsToGather.Add(InterfaceParameter);
-	}
-#endif
-	for (UObject* Object : ObjectsToGather)
-	{
-		UCameraRigInterfaceParameter* InterfaceParameter = Cast<UCameraRigInterfaceParameter>(Object);
-		if (!InterfaceParameter)
+		if (UCameraRigCameraNode* CameraRigNode = Cast<UCameraRigCameraNode>(CameraNode))
 		{
-			continue;
-		}
+			FCameraRigParameterOverrides& ParameterOverrides = CameraRigNode->CameraRigReference.GetParameterOverrides();
 
-		OldInterfaceParameters.Add(InterfaceParameter, FReusableInterfaceParameterInfo(InterfaceParameter->PrivateVariable, false));
-		if (InterfaceParameter->PrivateVariable)
-		{
-			GatheredVariables.Add(InterfaceParameter->PrivateVariable);
+#define UE_CAMERA_VARIABLE_FOR_TYPE(ValueType, ValueName)\
+			for (F##ValueName##CameraRigParameterOverride& ParameterOverride : ParameterOverrides.Get##ValueName##Overrides())\
+			{\
+				if (ParameterOverride.Value.Variable)\
+				{\
+					UObject* VariableOuter = ParameterOverride.Value.Variable->GetOuter();\
+					if (VariableOuter == CameraRig)\
+					{\
+						OldDrivenOverrides.Add(\
+								FDrivenOverrideKey{ ParameterOverride.InterfaceParameterGuid, CameraRigNode },\
+								ParameterOverride.Value.Variable);\
+						GatheredVariables.Add(ParameterOverride.Value.Variable);\
+						ParameterOverride.Value.Variable = nullptr;\
+					}\
+				}\
+			}
+UE_CAMERA_VARIABLE_FOR_ALL_TYPES()
+#undef UE_CAMERA_VARIABLE_FOR_TYPE
 		}
 	}
 
@@ -418,8 +526,6 @@ UE_CAMERA_VARIABLE_FOR_ALL_TYPES()
 
 void FCameraRigAssetBuilder::BuildNewDrivenParameters()
 {
-	using namespace Internal;
-
 	TSet<FString> UsedInterfaceParameterNames;
 
 	using FBuiltDrivenParameter = TTuple<UCameraNode*, FName>;
@@ -431,6 +537,11 @@ void FCameraRigAssetBuilder::BuildNewDrivenParameters()
 	// Look at the new interface parameters and setup the driven camera node parameters with
 	// private camera variables. We have gathered the old ones previously so we can re-use them,
 	// instead of creating new variable assets each time.
+	//
+	// Additionally, we need to handle camera rig nodes with special code, for the case of an
+	// interface parameter driving a camera rig override (which in turn drives the inner rig's
+	// interface parameter, and so on). This is basically for multi-level interface parameters
+	// overrides.
 	for (UCameraRigInterfaceParameter* InterfaceParameter : CameraRig->Interface.InterfaceParameters)
 	{
 		// Do some basic validation.
@@ -497,12 +608,19 @@ void FCameraRigAssetBuilder::BuildNewDrivenParameters()
 		}
 		BuiltDrivenParameters.Add(BuiltDrivenParameter);
 
-		// Get the target camera node property and check that it is a camera parameter struct.
-		UCameraNode* Target = InterfaceParameter->Target;
-		UClass* TargetClass = Target->GetClass();
-		FProperty* TargetProperty = TargetClass->FindPropertyByName(InterfaceParameter->TargetPropertyName);
-		if (!TargetProperty)
+		// See if this interface parameter is overriding a camera node parameter.
+		// Otherwise, maybe it's targeting a camera rig node's override for an inner rig interface parameter.
+		if (SetupCameraParameterOverride(InterfaceParameter))
+		{
+			// Implicit continue.
+		}
+		else if (SetupInnerCameraRigParameterOverride(InterfaceParameter))
+		{
+			// Implicit continue.
+		}
+		else
 		{
+			UCameraNode* Target = InterfaceParameter->Target;
 			BuildLog.AddMessage(EMessageSeverity::Error,
 					Target,
 					FText::Format(LOCTEXT(
@@ -511,53 +629,153 @@ void FCameraRigAssetBuilder::BuildNewDrivenParameters()
 						FText::FromString(InterfaceParameter->InterfaceParameterName), 
 						FText::FromName(InterfaceParameter->TargetPropertyName),
 						FText::FromName(Target->GetFName())));
-			continue;
 		}
+	}
+}
 
-		FStructProperty* TargetStructProperty = CastField<FStructProperty>(TargetProperty);
-		if (!TargetStructProperty)
-		{
-			BuildLog.AddMessage(EMessageSeverity::Error,
-					Target,
-					FText::Format(LOCTEXT(
+bool FCameraRigAssetBuilder::SetupCameraParameterOverride(UCameraRigInterfaceParameter* InterfaceParameter)
+{
+	using namespace Internal;
+
+	// Here we hook up interface parameters connected to a camera node property. This property is supposed
+	// to be of one of the camera parameter types (FBooleanCameraParameter, FInteger32CameraParameter, etc.)
+	// so they have both a fixed value (bool, int32, etc.) and a "private variable" which is a reference to 
+	// a corresponding camera variable asset (UBooleanCameraVariable, UInteger32CameraVariable, etc.) which
+	// has been set to "private".
+	//
+	// So the goal of this method is to create a private variable and set it on both the interface parameter
+	// and the camera node property. This way, if someone wants to override the value of that interface
+	// parameter, they set the value of the variable defined on it. It will then drive the value of the
+	// corresponding camera node property.
+
+	UCameraNode* Target = InterfaceParameter->Target;
+	UClass* TargetClass = Target->GetClass();
+	FProperty* TargetProperty = TargetClass->FindPropertyByName(InterfaceParameter->TargetPropertyName);
+	if (!TargetProperty)
+	{
+		// No match, try something else.
+		return false;
+	}
+
+	FStructProperty* TargetStructProperty = CastField<FStructProperty>(TargetProperty);
+	if (!TargetStructProperty)
+	{
+		BuildLog.AddMessage(EMessageSeverity::Error,
+				Target,
+				FText::Format(LOCTEXT(
 						"InvalidCameraNodeProperty",
 						"Invalid interface parameter '{0}', driving property '{1}' on '{2}', but it's not a camera parameter."),
-						FText::FromString(InterfaceParameter->InterfaceParameterName), 
-						FText::FromName(InterfaceParameter->TargetPropertyName),
-						FText::FromName(Target->GetFName())));
-			continue;
-		}
+					FText::FromString(InterfaceParameter->InterfaceParameterName), 
+					FText::FromName(InterfaceParameter->TargetPropertyName),
+					FText::FromName(Target->GetFName())));
+		return true;
+	}
 
-		// Get the type of the camera parameter by matching the struct against all the types we support,
-		// and create a private camera variable asset to drive its value.
-		FPrivateVariableBuilder PrivateVariableBuilder(*this);
+	// Get the type of the camera parameter by matching the struct against all the types we support,
+	// and create a private camera variable asset to drive its value.
+	FPrivateVariableBuilder PrivateVariableBuilder(*this);
 #define UE_CAMERA_VARIABLE_FOR_TYPE(ValueType, ValueName)\
-		if (TargetStructProperty->Struct == F##ValueName##CameraParameter::StaticStruct())\
-		{\
-			auto* CameraParameterPtr = TargetStructProperty->ContainerPtrToValuePtr<F##ValueName##CameraParameter>(Target);\
-			SetupPrivateVariable(PrivateVariableBuilder, InterfaceParameter, TargetStructProperty, CameraParameterPtr);\
-		}\
-		else
-UE_CAMERA_VARIABLE_FOR_ALL_TYPES()
+	if (TargetStructProperty->Struct == F##ValueName##CameraParameter::StaticStruct())\
+	{\
+		auto* CameraParameterPtr = TargetStructProperty->ContainerPtrToValuePtr<F##ValueName##CameraParameter>(Target);\
+		SetupPrivateVariable(PrivateVariableBuilder, InterfaceParameter, TargetStructProperty, CameraParameterPtr);\
+	}\
+	else
+	UE_CAMERA_VARIABLE_FOR_ALL_TYPES()
 #undef UE_CAMERA_VARIABLE_FOR_TYPE
-		{
-			BuildLog.AddMessage(EMessageSeverity::Error,
-					InterfaceParameter,
-					FText::Format(LOCTEXT(
+	{
+		BuildLog.AddMessage(EMessageSeverity::Error,
+				InterfaceParameter,
+				FText::Format(LOCTEXT(
 						"InvalidCameraNodeProperty",
 						"Invalid interface parameter '{0}', driving property '{1}' on '{2}', but it's not a camera parameter."),
-						FText::FromString(InterfaceParameter->InterfaceParameterName), 
-						FText::FromName(InterfaceParameter->TargetPropertyName),
-						FText::FromName(Target->GetFName())));
-			continue;
-		}
+					FText::FromString(InterfaceParameter->InterfaceParameterName), 
+					FText::FromName(InterfaceParameter->TargetPropertyName),
+					FText::FromName(Target->GetFName())));
 	}
+
+	return true;
+}
+
+bool FCameraRigAssetBuilder::SetupInnerCameraRigParameterOverride(UCameraRigInterfaceParameter* InterfaceParameter)
+{
+	using namespace Internal;
+
+	// Here we hook up interface parameters connected specifically to a camera rig node (aka "prefab node").
+	// Unlike other camera nodes, the camera rig node doesn't have "actual" properties on it, in the sense 
+	// of UClass and FProperty (although we could in the future generate a UClass, like Blueprints do).
+	// Camera rig nodes expose the same properties as their inner camera rig, i.e. they "forward expose" the
+	// interface parameters defined on their inner camera rig.
+	//
+	// So the goal of this method is to handle multi-level exposed parameters. That is: we are exposing the
+	// interface parameter of an inner camera rig as one of our own interface parameter. Just like in the
+	// previous method (see above) we create a private camera variable to set on the interface parameter,
+	// but instead of also setting it on a camera node property, here we set it on an override entry on
+	// the camera rig node's list of overrides.
+	//
+	// Note that this camera rig node may or may not have an existing override. If the user forwards the
+	// parameter without changing its default value, there would not be an existing override and we have
+	// to create our own. If there is an existing override, we set the private variable on it and it
+	// will use the user-defined new override value when the variable isn't set.
+
+	UCameraRigCameraNode* Target = Cast<UCameraRigCameraNode>(InterfaceParameter->Target);
+	if (!Target)
+	{
+		// No match, try something else.
+		return false;
+	}
+
+	// Look for an interface parameter matching the target name.
+	UCameraRigAsset* InnerCameraRig = Target->CameraRigReference.GetCameraRig();
+	if (!InnerCameraRig)
+	{
+		return false;
+	}
+	UCameraRigInterfaceParameter* InnerInterfaceParameter = InnerCameraRig->Interface.FindInterfaceParameterByName(
+			InterfaceParameter->TargetPropertyName.ToString());
+	if (!InnerInterfaceParameter)
+	{
+		return false;
+	}
+
+	// Found it! Check that the inner camera rig was built.
+	if (InnerInterfaceParameter->PrivateVariable == nullptr)
+	{
+		BuildLog.AddMessage(EMessageSeverity::Error, Target,
+				FText::Format(
+					LOCTEXT("UnbuiltInnerCameraRig",
+						"Can't expose inner camera rig parameter '{0}': the inner camera rig '{1}' failed to build."),
+					FText::FromName(InterfaceParameter->TargetPropertyName),
+					FText::FromString(GetPathNameSafe(InnerCameraRig))));
+		return true;
+	}
+
+	// Look for an override that matches the given interface parameter. Create one if we don't find any.
+	FCameraRigParameterOverrides& ParameterOverrides = Target->CameraRigReference.GetParameterOverrides();
+	switch (InnerInterfaceParameter->PrivateVariable->GetVariableType())
+	{
+#define UE_CAMERA_VARIABLE_FOR_TYPE(ValueType, ValueName)\
+		case ECameraVariableType::ValueName:\
+			{\
+				F##ValueName##CameraRigParameterOverride& ValueName##Override = \
+						ParameterOverrides.FindOrAddParameterOverride<F##ValueName##CameraRigParameterOverride>(InnerInterfaceParameter);\
+				FPrivateVariableBuilder PrivateVariableBuilder(*this);\
+				SetupPrivateVariable(PrivateVariableBuilder, InterfaceParameter, &ValueName##Override);\
+			}\
+			break;
+	UE_CAMERA_VARIABLE_FOR_ALL_TYPES()
+#undef UE_CAMERA_VARIABLE_FOR_TYPE
+	}
+
+	return true;
 }
 
 void FCameraRigAssetBuilder::DiscardUnusedPrivateVariables()
 {
 	// Now that we've rebuilt all exposed parameters, anything left from the old list 
 	// must be discarded.
+	TSet<UCameraVariableAsset*> VariablesToTrash;
+
 	for (TPair<FDrivenParameterKey, UCameraVariableAsset*> Pair : OldDrivenParameters)
 	{
 		// We null'ed the driving variable in GatherOldDrivenParameters. Now it's time
@@ -565,15 +783,29 @@ void FCameraRigAssetBuilder::DiscardUnusedPrivateVariables()
 		UCameraNode* Target = Pair.Key.Value;
 		Target->Modify();
 
-		// Trash the old camera variable. This helps with debugging.
-		UCameraVariableAsset* VariableToDiscard = Pair.Value;
+		VariablesToTrash.Add(Pair.Value);
+	}
+	OldDrivenParameters.Reset();
+
+	for (TPair<FDrivenOverrideKey, UCameraVariableAsset*> Pair : OldDrivenOverrides)
+	{
+		// We null'ed the override variable in GatherOldDrivenParameters. Flag the 
+		// camera rig node that owns this override as modified.
+		UCameraRigCameraNode* CameraRigNode = Pair.Key.Value;
+		CameraRigNode->Modify();
+
+		VariablesToTrash.Add(Pair.Value);
+	}
+	OldDrivenOverrides.Reset();
+	
+	// Trash the old camera variable. This helps with debugging.
+	for (UCameraVariableAsset* VariableToTrash : VariablesToTrash)
+	{
 		TStringBuilder<256> StringBuilder;
 		StringBuilder.Append("TRASH_");
-		StringBuilder.Append(VariableToDiscard->GetName());
-		VariableToDiscard->Rename(StringBuilder.ToString());
+		StringBuilder.Append(VariableToTrash->GetName());
+		VariableToTrash->Rename(StringBuilder.ToString());
 	}
-
-	OldDrivenParameters.Reset();
 }
 
 void FCameraRigAssetBuilder::BuildAllocationInfo()
diff --git a/Engine/Plugins/Cameras/GameplayCameras/Source/GameplayCameras/Private/Core/CameraRigAssetReference.cpp b/Engine/Plugins/Cameras/GameplayCameras/Source/GameplayCameras/Private/Core/CameraRigAssetReference.cpp
new file mode 100644
--- /dev/null
+++ b/Engine/Plugins/Cameras/GameplayCameras/Source/GameplayCameras/Private/Core/CameraRigAssetReference.cpp
@@ -0,0 +1,75 @@
+// Copyright Epic Games, Inc. All Rights Reserved.
+
+#include "Core/CameraRigAssetReference.h"
+
+#include UE_INLINE_GENERATED_CPP_BY_NAME(CameraRigAssetReference)
+
+void FCameraRigParameterOverrides::Reset()
+{
+#define UE_CAMERA_VARIABLE_FOR_TYPE(ValueType, ValueName)\
+	ValueName##Overrides.Reset();
+UE_CAMERA_VARIABLE_FOR_ALL_TYPES()
+#undef UE_CAMERA_VARIABLE_FOR_TYPE
+}
+
+bool FCameraRigAssetReference::UpdateParameterOverrides()
+{
+	if (!CameraRig)
+	{
+		bool bHasAnyOverride = false;
+#define UE_CAMERA_VARIABLE_FOR_TYPE(ValueType, ValueName)\
+		for (F##ValueName##CameraRigParameterOverride& ParameterOverride : ParameterOverrides.ValueName##Overrides)\
+		{\
+			ParameterOverride.bInvalid = true;\
+			bHasAnyOverride = true;\
+		}
+		UE_CAMERA_VARIABLE_FOR_ALL_TYPES()
+#undef UE_CAMERA_VARIABLE_FOR_TYPE
+		return bHasAnyOverride;
+	}
+
+	bool bAnyModified = false;
+	FCameraRigInterface& CameraRigInterface = CameraRig->Interface;
+#define UE_CAMERA_VARIABLE_FOR_TYPE(ValueType, ValueName)\
+	for (F##ValueName##CameraRigParameterOverride& ParameterOverride : ParameterOverrides.ValueName##Overrides)\
+	{\
+		UCameraRigInterfaceParameter* InterfaceParameter = CameraRigInterface.FindInterfaceParameterByGuid(\
+				ParameterOverride.InterfaceParameterGuid);\
+		\
+		const bool bWasInvalid = ParameterOverride.bInvalid;\
+		ParameterOverride.bInvalid = (InterfaceParameter == nullptr);\
+		bAnyModified |= (bWasInvalid != ParameterOverride.bInvalid);\
+		\
+		if (InterfaceParameter)\
+		{\
+			if (ParameterOverride.InterfaceParameterName != InterfaceParameter->InterfaceParameterName)\
+			{\
+				ParameterOverride.InterfaceParameterName = InterfaceParameter->InterfaceParameterName;\
+				bAnyModified = true;\
+			}\
+			UCameraVariableAsset* InterfaceParameterVariable = InterfaceParameter->PrivateVariable;\
+			const FGuid NewPrivateVariableGuid = InterfaceParameterVariable ? InterfaceParameterVariable->GetGuid() : FGuid();\
+			if (ParameterOverride.PrivateVariableGuid != NewPrivateVariableGuid)\
+			{\
+				ParameterOverride.PrivateVariableGuid = NewPrivateVariableGuid;\
+				bAnyModified = true;\
+			}\
+		}\
+	}
+	UE_CAMERA_VARIABLE_FOR_ALL_TYPES()
+#undef UE_CAMERA_VARIABLE_FOR_TYPE
+	return bAnyModified;
+}
+
+bool FCameraRigAssetReference::SerializeFromMismatchedTag(struct FPropertyTag const& Tag, FStructuredArchive::FSlot Slot)
+{
+	if (Tag.Type == NAME_SoftObjectProperty)
+	{
+		FSoftObjectPtr CameraRigPath;
+		Slot << CameraRigPath;
+		CameraRig = Cast<UCameraRigAsset>(CameraRigPath.Get());
+		return true;
+	}
+	return false;
+}
+
diff --git a/Engine/Plugins/Cameras/GameplayCameras/Source/GameplayCameras/Private/GameplayCamerasDelegates.cpp b/Engine/Plugins/Cameras/GameplayCameras/Source/GameplayCameras/Private/GameplayCamerasDelegates.cpp
new file mode 100644
--- /dev/null
+++ b/Engine/Plugins/Cameras/GameplayCameras/Source/GameplayCameras/Private/GameplayCamerasDelegates.cpp
@@ -0,0 +1,11 @@
+// Copyright Epic Games, Inc. All Rights Reserved.
+
+#include "GameplayCamerasDelegates.h"
+
+namespace UE::Cameras
+{
+
+FOnCameraRigAssetBuilt FGameplayCamerasDelegates::OnCameraRigAssetBuiltDelegates;
+
+}  // namespace UE::Cameras
+
diff --git a/Engine/Plugins/Cameras/GameplayCameras/Source/GameplayCameras/Private/Nodes/Common/CameraRigCameraNode.cpp b/Engine/Plugins/Cameras/GameplayCameras/Source/GameplayCameras/Private/Nodes/Common/CameraRigCameraNode.cpp
--- a/Engine/Plugins/Cameras/GameplayCameras/Source/GameplayCameras/Private/Nodes/Common/CameraRigCameraNode.cpp
+++ b/Engine/Plugins/Cameras/GameplayCameras/Source/GameplayCameras/Private/Nodes/Common/CameraRigCameraNode.cpp
@@ -28,6 +28,21 @@ void ApplyParameterOverrides(
 		using ParameterType = decltype(ParameterOverrideType::Value);
 		using ValueType = typename ParameterType::ValueType;
 
+		if (!ParameterOverride.PrivateVariableGuid.IsValid())
+		{
+#if WITH_EDITOR
+			// Ignore un-built parameter overrides in the editor since the user could have just added
+			// an override while PIE is running. They need to hit the Build button for the override
+			// to apply.
+			continue;
+#else
+			UE_LOG(LogCameraSystem, Error, 
+					TEXT("Invalid parameter override '%s' in camera rig '%s'. Was it built/cooked?"),
+					*ParameterOverride.InterfaceParameterName,
+					*GetPathNameSafe(CameraRig));
+#endif
+		}
+
 		FCameraVariableID InterfaceParameterID(FCameraVariableID::FromHashValue(GetTypeHash(ParameterOverride.PrivateVariableGuid)));
 
 		if (ParameterOverride.Value.Variable != nullptr)
@@ -36,9 +51,23 @@ void ApplyParameterOverrides(
 			// prefab's variable. Basically, we forward the value from one variable to the next.
 			FCameraVariableDefinition OverrideDefinition(ParameterOverride.Value.Variable->GetVariableDefinition());
 
-			OutVariableTable.SetValue<ValueType>(
-					InterfaceParameterID,
-					OutVariableTable.GetValue<typename ParameterType::ValueType>(OverrideDefinition.VariableID));
+			const ValueType* OverrideValuePtr = OutVariableTable.FindValue<ValueType>(OverrideDefinition.VariableID);
+			if (OverrideValuePtr)
+			{
+				OutVariableTable.SetValue<ValueType>(InterfaceParameterID, *OverrideValuePtr);
+			}
+			else
+			{
+				// Once again, ignore un-built data. Only emit errors when running outside of the editor.
+#if !WITH_EDITOR
+				UE_LOG(LogCameraSystem, Error, 
+						TEXT("Camera variable '%s' for parameter override '%s' in camera rig '%s' isn't in the variable "
+							"table. Was it built/cooked?"),
+						*GetNameSafe(ParameterOverride.Value.Variable),
+						*ParameterOverride.InterfaceParameterName,
+						*GetPathNameSafe(CameraRig));
+#endif
+			}
 		}
 		else
 		{
@@ -82,11 +111,11 @@ FCameraNodeEvaluatorChildrenView FCameraRigCameraNodeEvaluator::OnGetChildren()
 void FCameraRigCameraNodeEvaluator::OnBuild(const FCameraNodeEvaluatorBuildParams& Params)
 {
 	const UCameraRigCameraNode* CameraRigNode = GetCameraNodeAs<UCameraRigCameraNode>();
-	if (CameraRigNode->CameraRig)
+	if (const UCameraRigAsset* CameraRig = CameraRigNode->CameraRigReference.GetCameraRig())
 	{
-		if (CameraRigNode->CameraRig->RootNode)
+		if (CameraRig->RootNode)
 		{
-			CameraRigRootEvaluator = Params.BuildEvaluator(CameraRigNode->CameraRig->RootNode);
+			CameraRigRootEvaluator = Params.BuildEvaluator(CameraRig->RootNode);
 		}
 	}
 }
@@ -113,10 +142,13 @@ void FCameraRigCameraNodeEvaluator::ApplyParameterOverrides(FCameraVariableTable
 {
 	const UCameraRigCameraNode* PrefabNode = GetCameraNodeAs<UCameraRigCameraNode>();
 
+	const UCameraRigAsset* CameraRig = PrefabNode->CameraRigReference.GetCameraRig();
+	const FCameraRigParameterOverrides& ParameterOverrides = PrefabNode->CameraRigReference.GetParameterOverrides();
+
 #define UE_CAMERA_VARIABLE_FOR_TYPE(ValueType, ValueName)\
 	Internal::ApplyParameterOverrides(\
-			PrefabNode->CameraRig,\
-			TArrayView<const F##ValueName##CameraRigParameterOverride>(PrefabNode->ValueName##Overrides),\
+			CameraRig,\
+			ParameterOverrides.Get##ValueName##Overrides(),\
 			OutVariableTable);
 UE_CAMERA_VARIABLE_FOR_ALL_TYPES()
 #undef UE_CAMERA_VARIABLE_FOR_TYPE
@@ -135,15 +167,22 @@ struct FCameraRigCameraNodeBuilder
 		, BuildContext(InBuildContext)
 	{}
 
-	void Setup(const UCameraRigAsset* InCameraRig)
+	void Setup()
 	{
+		// Build a map matching each of our inner camera rig's interface parameter to its Guid.
 		ParametersByGuid.Reset();
-		for (TObjectPtr<UCameraRigInterfaceParameter> InterfaceParameter : InCameraRig->Interface.InterfaceParameters)
+		const UCameraRigAsset* CameraRig = CameraNode->CameraRigReference.GetCameraRig();
+		for (TObjectPtr<UCameraRigInterfaceParameter> InterfaceParameter : CameraRig->Interface.InterfaceParameters)
 		{
 			ParametersByGuid.Add(InterfaceParameter->Guid, InterfaceParameter);
 		}
 	}
 
+	template<typename ParameterOverrideType>
+	void BuildCameraRigParameterOverride(ParameterOverrideType& ParameterOverride);
+
+private:
+
 	const UCameraRigInterfaceParameter* FindInterfaceParameter(const FGuid& InterfaceParameterGuid)
 	{
 		if (UCameraRigInterfaceParameter** FoundItem = ParametersByGuid.Find(InterfaceParameterGuid))
@@ -153,9 +192,6 @@ struct FCameraRigCameraNodeBuilder
 		return nullptr;
 	}
 
-	template<typename ParameterOverrideType>
-	void BuildCameraRigParameterOverride(ParameterOverrideType& ParameterOverride);
-
 private:
 
 	TMap<FGuid, UCameraRigInterfaceParameter*> ParametersByGuid;
@@ -164,6 +200,9 @@ private:
 template<typename ParameterOverrideType>
 void FCameraRigCameraNodeBuilder::BuildCameraRigParameterOverride(ParameterOverrideType& ParameterOverride)
 {
+	const UCameraRigAsset* CameraRig = CameraNode->CameraRigReference.GetCameraRig();
+
+	// Each parameter override should point to a valid interface parameter on the inner rig, via its Guid.
 	const UCameraRigInterfaceParameter* InterfaceParameter = FindInterfaceParameter(ParameterOverride.InterfaceParameterGuid);
 	if (!InterfaceParameter)
 	{
@@ -171,58 +210,84 @@ void FCameraRigCameraNodeBuilder::BuildCameraRigParameterOverride(ParameterOverr
 				FText::Format(
 					LOCTEXT("MissingInterfaceParameter", "No camera rig interface parameter named '{0}' exists on '{1}'."),
 					FText::FromString(ParameterOverride.InterfaceParameterName),
-					FText::FromString(GetNameSafe(CameraNode->CameraRig))));
+					FText::FromString(GetNameSafe(CameraRig))));
 		return;
 	}
 
+	// The inner rig's interface parameter should have been built, i.e. it should have a private camera variable
+	// assigned for driving its value.
 	if (!InterfaceParameter->PrivateVariable)
 	{
 		BuildContext.BuildLog.AddMessage(EMessageSeverity::Error, CameraNode,
 				FText::Format(
 					LOCTEXT("UnbuiltInterfaceParameter", "Camera rig interface parameter '{0}' was not built correctly on '{1}'."),
 					FText::FromString(ParameterOverride.InterfaceParameterName),
-					FText::FromString(GetNameSafe(CameraNode->CameraRig))));
+					FText::FromString(GetNameSafe(CameraRig))));
 		return;
 	}
 
+	// The inner rig's interface parameter is driven by this private variable. Let's remember its Guid so we
+	// can override its value in the variable table at runtime.
 	ParameterOverride.PrivateVariableGuid = InterfaceParameter->PrivateVariable->GetGuid();
+	// Update the last known name for this interface parameter.
+	ParameterOverride.InterfaceParameterName = InterfaceParameter->InterfaceParameterName;
+
+	// The build process automatically gathers variables that drive amera parameters on a camera node, but
+	// nothing else for now. We therefore need to help it out by manually reporting the variables that
+	// drive our parameter overrides.
+	FCameraVariableTableAllocationInfo& VariableTableAllocationInfo = BuildContext.AllocationInfo.VariableTableInfo;
+	if (ParameterOverride.Value.Variable)
+	{
+		VariableTableAllocationInfo.VariableDefinitions.Add(ParameterOverride.Value.Variable->GetVariableDefinition());
+	}
 }
 
 }  // namespace Internal
 
 }  // namespace UE::Cameras
 
+void UCameraRigCameraNode::OnPreBuild(FCameraBuildLog& BuildLog)
+{
+	// Build the inner camera rig. Silently skip it if it's not set... but we will
+	// report an error in OnBuild about it.
+	if (UCameraRigAsset* CameraRig = CameraRigReference.GetCameraRig())
+	{
+		CameraRig->BuildCameraRig(BuildLog);
+	}
+}
+
 void UCameraRigCameraNode::OnBuild(FCameraRigBuildContext& BuildContext)
 {
 	using namespace UE::Cameras;
 	using namespace UE::Cameras::Internal;
 
+	UCameraRigAsset* CameraRig = CameraRigReference.GetCameraRig();
 	if (!CameraRig)
 	{
-		BuildContext.BuildLog.AddMessage(EMessageSeverity::Error, this, LOCTEXT("MissingCameraRig", "No camera rig specified on camera rig node."));
+		BuildContext.BuildLog.AddMessage(EMessageSeverity::Error, this, 
+				LOCTEXT("MissingCameraRig", "No camera rig specified on camera rig node."));
 		return;
 	}
 
-	// Build the inner camera rig. Whatever allocations it needs for its evaluators and
+	// Whatever allocations our inner camera rig needs for its evaluators and
 	// their camera variables, we add that to our camera rig's allocation info.
-	CameraRig->BuildCameraRig(BuildContext.BuildLog);
 	BuildContext.AllocationInfo.Append(CameraRig->AllocationInfo);
 
 	// Next, we set things up for the runtime. Mostly, we want to get the camera variable 
 	// Guids that we need to write the override values to.
 	FCameraRigCameraNodeBuilder InternalBuilder(this, BuildContext);
-	InternalBuilder.Setup(CameraRig);
+	InternalBuilder.Setup();
 
+	FCameraRigParameterOverrides& ParameterOverrides = CameraRigReference.GetParameterOverrides();
 #define UE_CAMERA_VARIABLE_FOR_TYPE(ValueType, ValueName)\
 	{\
-		for (F##ValueName##CameraRigParameterOverride& ParameterOverride : ValueName##Overrides)\
+		for (F##ValueName##CameraRigParameterOverride& ParameterOverride : ParameterOverrides.Get##ValueName##Overrides())\
 		{\
 			InternalBuilder.BuildCameraRigParameterOverride(ParameterOverride);\
 		}\
 	}
 UE_CAMERA_VARIABLE_FOR_ALL_TYPES()
 #undef UE_CAMERA_VARIABLE_FOR_TYPE
-
 }
 
 FCameraNodeEvaluatorPtr UCameraRigCameraNode::OnBuildEvaluator(FCameraNodeEvaluatorBuilder& Builder) const
@@ -231,5 +296,26 @@ FCameraNodeEvaluatorPtr UCameraRigCameraNode::OnBuildEvaluator(FCameraNodeEvalua
 	return Builder.BuildEvaluator<FCameraRigCameraNodeEvaluator>();
 }
 
+void UCameraRigCameraNode::PostLoad()
+{
+	Super::PostLoad();
+
+	if (CameraRig_DEPRECATED)
+	{
+		CameraRigReference.SetCameraRig(CameraRig_DEPRECATED);
+		CameraRig_DEPRECATED = nullptr;
+	}
+
+	FCameraRigParameterOverrides& ParameterOverrides = CameraRigReference.GetParameterOverrides();
+#define UE_CAMERA_VARIABLE_FOR_TYPE(ValueType, ValueName)\
+	if (ValueName##Overrides_DEPRECATED.Num() > 0)\
+	{\
+		ParameterOverrides.AppendParameterOverrides<F##ValueName##CameraRigParameterOverride>(ValueName##Overrides_DEPRECATED);\
+		ValueName##Overrides_DEPRECATED.Reset();\
+	}
+UE_CAMERA_VARIABLE_FOR_ALL_TYPES()
+#undef UE_CAMERA_VARIABLE_FOR_TYPE
+}
+
 #undef LOCTEXT_NAMESPACE
 
diff --git a/Engine/Plugins/Cameras/GameplayCameras/Source/GameplayCameras/Private/Tests/CameraRigAssetBuilderTests.cpp b/Engine/Plugins/Cameras/GameplayCameras/Source/GameplayCameras/Private/Tests/CameraRigAssetBuilderTests.cpp
--- a/Engine/Plugins/Cameras/GameplayCameras/Source/GameplayCameras/Private/Tests/CameraRigAssetBuilderTests.cpp
+++ b/Engine/Plugins/Cameras/GameplayCameras/Source/GameplayCameras/Private/Tests/CameraRigAssetBuilderTests.cpp
@@ -3,8 +3,9 @@
 #include "Core/CameraRigAsset.h"
 #include "Misc/AutomationTest.h"
 #include "Nodes/Blends/SmoothBlendCameraNode.h"
-#include "Nodes/Common/OffsetCameraNode.h"
+#include "Nodes/Common/CameraRigCameraNode.h"
 #include "Nodes/Common/LensParametersCameraNode.h"
+#include "Nodes/Common/OffsetCameraNode.h"
 #include "Tests/GameplayCamerasTestBuilder.h"
 
 #define LOCTEXT_NAMESPACE "CameraRigAssetBuilderTests"
@@ -22,7 +23,7 @@ bool FCameraRigAssetBuilderNullTest::RunTest(const FString& Parameters)
 	AddExpectedMessage(
 			FString::Format(TEXT("Camera rig '{0}' has no root node."), ErrorArgs),
 			ELogVerbosity::Error,
-			EAutomationExpectedMessageFlags::Exact,
+			EAutomationExpectedMessageFlags::Contains,
 			1,
 			false);
 	CameraRig->BuildCameraRig();
@@ -147,5 +148,140 @@ bool FCameraRigAssetBuilderReassignParameterTest::RunTest(const FString& Paramet
 	return true;
 }
 
+IMPLEMENT_SIMPLE_AUTOMATION_TEST(FCameraRigAssetBuilderDrivenOverridesTest, "System.Engine.GameplayCameras.CameraRigAssetBuilder.DrivenOverrides", EAutomationTestFlags::EditorContext | EAutomationTestFlags::EngineFilter)
+bool FCameraRigAssetBuilderDrivenOverridesTest::RunTest(const FString& Parameters)
+{
+	using namespace UE::Cameras::Test;
+
+	TSharedRef<FNamedObjectRegistry> Registry = MakeShared<FNamedObjectRegistry>();
+
+	// Make a camera rig with an offset node (10, 20, 30) and a focal length node (20). Expose both parameters
+	// as interface parameters.
+	UCameraRigAsset* InnerCameraRig = FCameraRigAssetTestBuilder(Registry, TEXT("InnerCameraRig"))
+		.MakeArrayRootNode()
+			.AddArrayChild<UOffsetCameraNode>().Named(TEXT("Offset"))
+				.SetParameter(&UOffsetCameraNode::TranslationOffset, FVector3d(10, 20, 30))
+				.Done()
+			.AddArrayChild<ULensParametersCameraNode>().Named(TEXT("Lens"))
+				.SetParameter(&ULensParametersCameraNode::FocalLength, 20.f)
+				.Done()
+			.Done()
+		.ExposeParameter(TEXT("OffsetParam"), TEXT("Offset"), GET_MEMBER_NAME_CHECKED(UOffsetCameraNode, TranslationOffset))
+		.ExposeParameter(TEXT("FocalLengthParam"), TEXT("Lens"), GET_MEMBER_NAME_CHECKED(ULensParametersCameraNode, FocalLength))
+		.Get();
+
+	// Make a camera rig that uses the previous one, with overrides on both the offset (now 15, 25, 35)
+	// and the focal length (now 25). Expose the offset further up as an interface parameter.
+	UCameraRigCameraNode* MiddlePrefabNode = nullptr;
+
+	UCameraRigAsset* MiddleCameraRig = FCameraRigAssetTestBuilder(Registry, TEXT("MiddleCameraRig"))
+		.MakeRootNode<UCameraRigCameraNode>()
+			.Pin(MiddlePrefabNode)
+			.Setup([InnerCameraRig](UCameraRigCameraNode* Node, FNamedObjectRegistry* Registry)
+					{
+						Node->CameraRigReference.SetCameraRig(InnerCameraRig);
+
+						FCameraRigParameterOverrides& ParameterOverrides = Node->CameraRigReference.GetParameterOverrides();
+
+						auto* OffsetParam = Registry->Get<UCameraRigInterfaceParameter>(TEXT("OffsetParam"));
+						auto& OffsetParamOverride = ParameterOverrides.FindOrAddParameterOverride<FVector3dCameraRigParameterOverride>(OffsetParam);
+						OffsetParamOverride.Value = FVector3d(15, 25, 35);
+
+						auto* FocalLengthParam = Registry->Get<UCameraRigInterfaceParameter>(TEXT("FocalLengthParam"));
+						auto& FocalLengthParamOverride = ParameterOverrides.FindOrAddParameterOverride<FFloatCameraRigParameterOverride>(FocalLengthParam);
+						FocalLengthParamOverride.Value = 25.f;
+					})
+			.Done()
+			.ExposeParameter(TEXT("MiddleOffsetParam"), MiddlePrefabNode, TEXT("OffsetParam"))
+		.Get();
+
+	// Make another camera rig that uses the previous one, which makes a total of 3 nesting levels of camera rigs.
+	// This level overrides the offset parameter some more (now 20, 50, 70).
+	UCameraRigCameraNode* OuterPrefabNode = nullptr;
+
+	UCameraRigAsset* OuterCameraRig = FCameraRigAssetTestBuilder(Registry, TEXT("OuterCameraRig"))
+		.MakeRootNode<UCameraRigCameraNode>()
+			.Pin(OuterPrefabNode)
+			.Setup([MiddleCameraRig](UCameraRigCameraNode* Node, FNamedObjectRegistry* Registry)
+					{
+						Node->CameraRigReference.SetCameraRig(MiddleCameraRig);
+
+						FCameraRigParameterOverrides& ParameterOverrides = Node->CameraRigReference.GetParameterOverrides();
+						auto* MiddleOffsetParam = Registry->Get<UCameraRigInterfaceParameter>(TEXT("MiddleOffsetParam"));
+						auto& MiddleOffsetParamOverride = ParameterOverrides.FindOrAddParameterOverride<FVector3dCameraRigParameterOverride>(MiddleOffsetParam);
+						MiddleOffsetParamOverride.Value = FVector3d(20, 50, 70);
+					})
+			.Done()
+		.Get();
+
+	OuterCameraRig->BuildCameraRig();
+
+	UCameraRigInterfaceParameter* OffsetParam = InnerCameraRig->Interface.InterfaceParameters[0];
+	UCameraRigInterfaceParameter* FocalLengthParam = InnerCameraRig->Interface.InterfaceParameters[1];
+
+	UTEST_EQUAL_EXPR(OffsetParam->PrivateVariable->GetName(), "Override_InnerCameraRig_OffsetParam");
+	UTEST_EQUAL_EXPR(FocalLengthParam->PrivateVariable->GetName(), "Override_InnerCameraRig_FocalLengthParam");
+
+	// Test that the inner nodes are driven by the interface parameters.
+	{
+		UOffsetCameraNode* OffsetNode = Registry->Get<UOffsetCameraNode>("Offset");
+		UTEST_EQUAL_EXPR((UCameraVariableAsset*)OffsetNode->TranslationOffset.Variable.Get(), OffsetParam->PrivateVariable.Get());
+		UTEST_EQUAL_EXPR(OffsetNode->TranslationOffset.Variable->DefaultValue, OffsetNode->TranslationOffset.Value);
+		UTEST_EQUAL_EXPR(OffsetNode->TranslationOffset.Variable->DefaultValue, FVector3d(10, 20, 30));
+
+		ULensParametersCameraNode* LensNode = Registry->Get<ULensParametersCameraNode>("Lens");
+		UTEST_EQUAL_EXPR((UCameraVariableAsset*)LensNode->FocalLength.Variable.Get(), FocalLengthParam->PrivateVariable.Get());
+		UTEST_EQUAL_EXPR(LensNode->FocalLength.Variable->DefaultValue, LensNode->FocalLength.Value);
+		UTEST_EQUAL_EXPR(LensNode->FocalLength.Variable->DefaultValue, 20.f);
+	}
+
+	// Test that the middle prefab node is driving the inner interface parameters, and that one of those
+	// overrides is in turn driven by the middle camera rig's interface parameter.
+	{
+		FCameraRigParameterOverrides& ParameterOverrides = MiddlePrefabNode->CameraRigReference.GetParameterOverrides();
+
+		FVector3dCameraRigParameterOverride* OffsetParamOverride = 
+			ParameterOverrides.FindParameterOverride<FVector3dCameraRigParameterOverride>(OffsetParam->Guid);
+		UTEST_NOT_NULL("OffsetParamOverride", OffsetParamOverride);
+
+		UTEST_EQUAL_EXPR(OffsetParamOverride->InterfaceParameterName, "OffsetParam");
+		UTEST_EQUAL_EXPR(OffsetParamOverride->PrivateVariableGuid, OffsetParam->PrivateVariable->GetGuid());
+		UTEST_EQUAL_EXPR(OffsetParamOverride->Value.Value, FVector3d(15, 25, 35));
+
+		FFloatCameraRigParameterOverride* FocalLengthParamOverride =
+			ParameterOverrides.FindParameterOverride<FFloatCameraRigParameterOverride>(FocalLengthParam->Guid);
+		UTEST_NOT_NULL("FocalLengthParamOverride", FocalLengthParamOverride);
+
+		UTEST_EQUAL_EXPR(FocalLengthParamOverride->InterfaceParameterName, "FocalLengthParam");
+		UTEST_EQUAL_EXPR(FocalLengthParamOverride->PrivateVariableGuid, FocalLengthParam->PrivateVariable->GetGuid());
+		UTEST_EQUAL_EXPR(FocalLengthParamOverride->Value.Value, 25.f);
+	}
+
+	UCameraRigInterfaceParameter* MiddleOffsetParam = MiddleCameraRig->Interface.InterfaceParameters[0];
+	{
+		FCameraRigParameterOverrides& ParameterOverrides = MiddlePrefabNode->CameraRigReference.GetParameterOverrides();
+
+		FVector3dCameraRigParameterOverride* OffsetParamOverride = 
+			ParameterOverrides.FindParameterOverride<FVector3dCameraRigParameterOverride>(OffsetParam->Guid);
+		UTEST_EQUAL_EXPR((UCameraVariableAsset*)OffsetParamOverride->Value.Variable.Get(), MiddleOffsetParam->PrivateVariable.Get());
+		UTEST_EQUAL_EXPR(OffsetParamOverride->Value.Variable->DefaultValue, FVector3d(15, 25, 35));
+	}
+
+	// Test that the outer prefab node is driving the middle interface parameters.
+	{
+		FCameraRigParameterOverrides& ParameterOverrides = OuterPrefabNode->CameraRigReference.GetParameterOverrides();
+
+		FVector3dCameraRigParameterOverride* OffsetParamOverride =
+			ParameterOverrides.FindParameterOverride<FVector3dCameraRigParameterOverride>(MiddleOffsetParam->Guid);
+		UTEST_NOT_NULL("OffsetParamOverride", OffsetParamOverride);
+
+		UTEST_EQUAL_EXPR(OffsetParamOverride->InterfaceParameterName, "MiddleOffsetParam");
+		UTEST_EQUAL_EXPR(OffsetParamOverride->PrivateVariableGuid, MiddleOffsetParam->PrivateVariable->GetGuid());
+		UTEST_EQUAL_EXPR(OffsetParamOverride->Value.Value, FVector3d(20, 50, 70));
+	}
+
+	return true;
+}
+
 #undef LOCTEXT_NAMESPACE
 
diff --git a/Engine/Plugins/Cameras/GameplayCameras/Source/GameplayCameras/Private/Tests/GameplayCamerasTestBuilder.cpp b/Engine/Plugins/Cameras/GameplayCameras/Source/GameplayCameras/Private/Tests/GameplayCamerasTestBuilder.cpp
--- a/Engine/Plugins/Cameras/GameplayCameras/Source/GameplayCameras/Private/Tests/GameplayCamerasTestBuilder.cpp
+++ b/Engine/Plugins/Cameras/GameplayCameras/Source/GameplayCameras/Private/Tests/GameplayCamerasTestBuilder.cpp
@@ -6,6 +6,16 @@ namespace UE::Cameras::Test
 {
 
 FCameraRigAssetTestBuilder::FCameraRigAssetTestBuilder(FName Name, UObject* Outer)
+{
+	Initialize(nullptr, Name, Outer);
+}
+
+FCameraRigAssetTestBuilder::FCameraRigAssetTestBuilder(TSharedRef<FNamedObjectRegistry> InNamedObjectRegistry, FName Name, UObject* Outer)
+{
+	Initialize(InNamedObjectRegistry, Name, Outer);
+}
+
+void FCameraRigAssetTestBuilder::Initialize(TSharedPtr<FNamedObjectRegistry> InNamedObjectRegistry, FName Name, UObject* Outer)
 {
 	if (Outer == nullptr)
 	{
@@ -14,6 +24,12 @@ FCameraRigAssetTestBuilder::FCameraRigAssetTestBuilder(FName Name, UObject* Oute
 
 	CameraRig = NewObject<UCameraRigAsset>(Outer, Name);
 	TCameraObjectInitializer<UCameraRigAsset>::SetObject(CameraRig);
+
+	NamedObjectRegistry = InNamedObjectRegistry;
+	if (!NamedObjectRegistry)
+	{
+		NamedObjectRegistry = MakeShared<FNamedObjectRegistry>();
+	}
 }
 
 TCameraRigTransitionTestBuilder<FCameraRigAssetTestBuilder> FCameraRigAssetTestBuilder::AddEnterTransition()
@@ -36,9 +52,17 @@ FCameraRigAssetTestBuilder& FCameraRigAssetTestBuilder::ExposeParameter(const FS
 	InterfaceParameter->InterfaceParameterName = ParameterName;
 	InterfaceParameter->Target = Target;
 	InterfaceParameter->TargetPropertyName = TargetPropertyName;
+	NamedObjectRegistry->Register(InterfaceParameter, ParameterName);
 	CameraRig->Interface.InterfaceParameters.Add(InterfaceParameter);
 	return *this;
 }
 
+FCameraRigAssetTestBuilder& FCameraRigAssetTestBuilder::ExposeParameter(const FString& ParameterName, const FString& TargetName, FName TargetPropertyName)
+{
+	UCameraNode* Target = NamedObjectRegistry->Get<UCameraNode>(TargetName);
+	ensure(Target);
+	return ExposeParameter(ParameterName, Target, TargetPropertyName);
+}
+
 }  // namespace UE::Cameras::Test
 
diff --git a/Engine/Plugins/Cameras/GameplayCameras/Source/GameplayCameras/Private/Tests/GameplayCamerasTestBuilder.h b/Engine/Plugins/Cameras/GameplayCameras/Source/GameplayCameras/Private/Tests/GameplayCamerasTestBuilder.h
--- a/Engine/Plugins/Cameras/GameplayCameras/Source/GameplayCameras/Private/Tests/GameplayCamerasTestBuilder.h
+++ b/Engine/Plugins/Cameras/GameplayCameras/Source/GameplayCameras/Private/Tests/GameplayCamerasTestBuilder.h
@@ -73,6 +73,52 @@ private:
 	ObjectType* Object = nullptr;
 };
 
+/**
+ * A simple repository matching UObject instances to names.
+ */
+class FNamedObjectRegistry
+{
+public:
+
+	/** Adds an object to the repository. */
+	void Register(UObject* InObject, const FString& InName)
+	{
+		ensure(InObject && !InName.IsEmpty());
+		NamedObjects.Add(InName, InObject);
+	}
+
+	/** Gets an object from the repository. */
+	UObject* Get(const FString& InName) const
+	{
+		if (UObject* const* Found = NamedObjects.Find(InName))
+		{
+			return *Found;
+		}
+		return nullptr;
+	}
+
+	/** Gets an object from the repository with a call to CastChecked. */
+	template<typename ObjectClass>
+	ObjectClass* Get(const FString& InName) const
+	{
+		return CastChecked<ObjectClass>(Get(InName));
+	}
+
+private:
+
+	TMap<FString, UObject*> NamedObjects;
+};
+
+/**
+ * Interface for something that has access to a named object repository.
+ */
+struct IHasNamedObjectRegistry
+{
+	virtual ~IHasNamedObjectRegistry() {}
+
+	virtual FNamedObjectRegistry* GetNamedObjectRegistry() = 0;
+};
+
 /**
  * A builder class for camera nodes.
  */
@@ -84,6 +130,7 @@ template<
 class TCameraNodeTestBuilder 
 	: public TScopedConstruction<ParentType>
 	, public TCameraObjectInitializer<NodeType>
+	, public IHasNamedObjectRegistry
 {
 public:
 
@@ -107,6 +154,18 @@ public:
 	/** Pins the built camera node to a given pointer, for being able to later refer to it. */
 	ThisType& Pin(NodeType*& OutPtr) { OutPtr = CameraNode; return *this; }
 
+	/** Give a name to the built camera node, to be recalled later. */
+	template<typename VV = std::enable_if_t<TPointerIsConvertibleFromTo<ParentType, IHasNamedObjectRegistry>::Value>>
+	ThisType& Named(const TCHAR* InName)
+	{
+		FNamedObjectRegistry* NamedObjectRegistry = GetNamedObjectRegistry();
+		if (ensure(NamedObjectRegistry))
+		{
+			NamedObjectRegistry->Register(CameraNode, InName);
+		}
+		return *this;
+	}
+
 	/** Sets the value of a camera parameter field on the camera node. */
 	template<typename ParameterType>
 	ThisType& SetParameter(
@@ -118,6 +177,25 @@ public:
 		return *this;
 	}
 
+	/**
+	 * Runs a custom setup callback on the camera node.
+	 */
+	ThisType& Setup(TFunction<void(NodeType*)> SetupCallback)
+	{
+		SetupCallback(CameraNode);
+		return *this;
+	}
+
+	/**
+	 * Runs a custom setup callback on the camera node with the named object registry provided.
+	 */
+	ThisType& Setup(TFunction<void(NodeType*, FNamedObjectRegistry*)> SetupCallback)
+	{
+		FNamedObjectRegistry* NamedObjectRegistry = GetNamedObjectRegistry();
+		SetupCallback(CameraNode, NamedObjectRegistry);
+		return *this;
+	}
+
 	/**
 	 * Adds a child camera node via a public array member field on the camera node.
 	 * Returns a builder for the child. You can go back to the current builder by
@@ -137,6 +215,23 @@ public:
 		return ChildBuilder;
 	}
 
+	/**
+	 * Convenience implementation of AddChild specifically for array nodes.
+	 */
+	template<
+		typename ChildNodeType,
+		typename = std::enable_if_t<
+			TPointerIsConvertibleFromTo<NodeType, UArrayCameraNode>::Value &&
+			TPointerIsConvertibleFromTo<ChildNodeType, UCameraNode>::Value>
+		>
+	TCameraNodeTestBuilder<ThisType, ChildNodeType>
+	AddArrayChild()
+	{
+		TCameraNodeTestBuilder<ThisType, ChildNodeType> ChildBuilder(*this, CameraNode->GetOuter());
+		CastChecked<UArrayCameraNode>(CameraNode)->Children.Add(ChildBuilder.Get());
+		return ChildBuilder;
+	}
+
 	/** 
 	 * Casting operator that returns a builder for the same camera node, but typed
 	 * around a parent class of the camera node's class. Mostly useful for implicit casting
@@ -154,6 +249,19 @@ public:
 				CameraNode);
 	}
 
+	/** Gets the named object registry from the parent. */
+	virtual FNamedObjectRegistry* GetNamedObjectRegistry() override
+	{
+		if constexpr (TPointerIsConvertibleFromTo<ParentType, IHasNamedObjectRegistry>::Value)
+		{
+			return TScopedConstruction<ParentType>::Parent.GetNamedObjectRegistry();
+		}
+		else
+		{
+			return nullptr;
+		}
+	}
+
 private:
 
 	enum class EForceReuseCameraNode { Yes };
@@ -177,6 +285,7 @@ template<typename ParentType>
 class TCameraRigTransitionTestBuilder 
 	: public TScopedConstruction<ParentType>
 	, public TCameraObjectInitializer<UCameraRigTransition>
+	, public IHasNamedObjectRegistry
 {
 public:
 
@@ -198,6 +307,21 @@ public:
 	/** Gets the built transition object. */
 	UCameraRigTransition* Get() const { return Transition; }
 
+	/** Pins the built transition to a given pointer, for being able to later refer to it. */
+	ThisType& Pin(UCameraRigTransition*& OutPtr) { OutPtr = Transition; return *this; }
+
+	/** Give a name to the built transition, to be recalled later. */
+	template<typename VV = std::enable_if_t<TPointerIsConvertibleFromTo<ParentType, IHasNamedObjectRegistry>::Value>>
+	ThisType& Named(const TCHAR* InName)
+	{
+		FNamedObjectRegistry* NamedObjectRegistry = GetNamedObjectRegistry();
+		if (ensure(NamedObjectRegistry))
+		{
+			NamedObjectRegistry->Register(Transition, InName);
+		}
+		return *this;
+	}
+
 	/** 
 	 * Creates a blend node of the given type, and returns a builder for it.
 	 * You can go back to this transition builder by calling Done() on the blend builder.
@@ -238,6 +362,19 @@ public:
 		return *this;
 	}
 
+	/** Gets the named object registry from the parent. */
+	virtual FNamedObjectRegistry* GetNamedObjectRegistry() override
+	{
+		if constexpr (TPointerIsConvertibleFromTo<ParentType, IHasNamedObjectRegistry>::Value)
+		{
+			return TScopedConstruction<ParentType>::Parent.GetNamedObjectRegistry();
+		}
+		else
+		{
+			return nullptr;
+		}
+	}
+
 private:
 
 	UCameraRigTransition* Transition;
@@ -263,13 +400,16 @@ private:
  *				.Done()
  *			.Get();
  */
-class FCameraRigAssetTestBuilder : public TCameraObjectInitializer<UCameraRigAsset>
+class FCameraRigAssetTestBuilder 
+	: public TCameraObjectInitializer<UCameraRigAsset>
+	, public IHasNamedObjectRegistry
 {
 public:
 
 	using ThisType = FCameraRigAssetTestBuilder;
 
 	FCameraRigAssetTestBuilder(FName Name = NAME_None, UObject* Outer = nullptr);
+	FCameraRigAssetTestBuilder(TSharedRef<FNamedObjectRegistry> InNamedObjectRegistry, FName Name = NAME_None, UObject* Outer = nullptr);
 
 	/** Gets the built camera rig. */
 	UCameraRigAsset* Get() { return CameraRig; }
@@ -287,6 +427,14 @@ public:
 		return NodeBuilder;
 	}
 
+	/**
+	 * A convenience method that calls MakeRootNode with a UArrayCameraNode.
+	 */
+	TCameraNodeTestBuilder<ThisType, UArrayCameraNode> MakeArrayRootNode()
+	{
+		return MakeRootNode<UArrayCameraNode>();
+	}
+
 	/**
 	 * Adds a new enter transition and returns a builder for it. You can come back to the
 	 * rig builder by calling Done() on the transition builder.
@@ -302,12 +450,33 @@ public:
 	 * Creates a new exposed rig parameter and hooks it up to the given camera node's property.
 	 * When building the node hierarchy, you can use the Pin() method on the node builders to
 	 * save a pointer to nodes you need for ExposeParameter().
+	 *
+	 * The created parameter is automatically stored in the named object registry under its name.
 	 */
 	FCameraRigAssetTestBuilder& ExposeParameter(const FString& ParameterName, UCameraNode* Target, FName TargetPropertyName);
 
+	/**
+	 * A variant of ExposeParameter that retrieves the target node from the named registry.
+	 *
+	 * The created parameter is automatically stored in the named object registry under its name.
+	 */
+	FCameraRigAssetTestBuilder& ExposeParameter(const FString& ParameterName, const FString& TargetName, FName TargetPropertyName);
+
+	/** Gets the named object registry. */
+	virtual FNamedObjectRegistry* GetNamedObjectRegistry() override
+	{
+		return NamedObjectRegistry.Get();
+	}
+
+private:
+
+	void Initialize(TSharedPtr<FNamedObjectRegistry> InNamedObjectRegistry, FName Name, UObject* Outer);
+
 private:
 
 	UCameraRigAsset* CameraRig;
+
+	TSharedPtr<FNamedObjectRegistry> NamedObjectRegistry;
 };
 
 }  // namespace UE::Cameras::Test
diff --git a/Engine/Plugins/Cameras/GameplayCameras/Source/GameplayCameras/Public/Core/CameraNode.h b/Engine/Plugins/Cameras/GameplayCameras/Source/GameplayCameras/Public/Core/CameraNode.h
--- a/Engine/Plugins/Cameras/GameplayCameras/Source/GameplayCameras/Public/Core/CameraNode.h
+++ b/Engine/Plugins/Cameras/GameplayCameras/Source/GameplayCameras/Public/Core/CameraNode.h
@@ -12,6 +12,7 @@
 
 namespace UE::Cameras
 {
+	class FCameraBuildLog;
 	struct FCameraRigBuildContext;
 }
 
@@ -30,12 +31,16 @@ class GAMEPLAYCAMERAS_API UCameraNode
 
 public:
 	
+	using FCameraBuildLog = UE::Cameras::FCameraBuildLog;
 	using FCameraRigBuildContext = UE::Cameras::FCameraRigBuildContext;
 	using FCameraNodeEvaluatorBuilder = UE::Cameras::FCameraNodeEvaluatorBuilder;
 
 	/** Get the list of children under this node. */
 	FCameraNodeChildrenView GetChildren();
 
+	/** Optional build step executed at the beginning of the build process. */
+	void PreBuild(FCameraBuildLog& BuildLog);
+
 	/** Gets optional info about this node's required allocations at runtime. */
 	void Build(FCameraRigBuildContext& BuildContext);
 
@@ -47,6 +52,9 @@ protected:
 	/** Get the list of children under this node. */
 	virtual FCameraNodeChildrenView OnGetChildren() { return FCameraNodeChildrenView(); }
 
+	/** Optional build step executed at the beginning of the build process. */
+	virtual void OnPreBuild(FCameraBuildLog& BuildLog) {}
+
 	/** Gets optional info about this node's required allocations at runtime. */
 	virtual void OnBuild(FCameraRigBuildContext& BuildContext) {}
 
diff --git a/Engine/Plugins/Cameras/GameplayCameras/Source/GameplayCameras/Public/Core/CameraRigAsset.h b/Engine/Plugins/Cameras/GameplayCameras/Source/GameplayCameras/Public/Core/CameraRigAsset.h
--- a/Engine/Plugins/Cameras/GameplayCameras/Source/GameplayCameras/Public/Core/CameraRigAsset.h
+++ b/Engine/Plugins/Cameras/GameplayCameras/Source/GameplayCameras/Public/Core/CameraRigAsset.h
@@ -3,6 +3,7 @@
 #pragma once
 
 #include "Core/CameraBuildStatus.h"
+#include "Core/CameraEventHandler.h"
 #include "Core/CameraNodeEvaluatorFwd.h"
 #include "Core/CameraRigTransition.h"
 #include "Core/CameraVariableTableFwd.h"
@@ -16,12 +17,25 @@
 #include "CameraRigAsset.generated.h"
 
 class UCameraNode;
+class UCameraRigAsset;
 class UCameraVariableAsset;
 
 namespace UE::Cameras
 {
 	class FCameraBuildLog;
 	class FCameraRigAssetBuilder;
+
+	/**
+	 * Interface for listening to changes on a camera rig asset.
+	 */
+	class ICameraRigAssetEventHandler
+	{
+	public:
+		virtual ~ICameraRigAssetEventHandler() {}
+
+		/** Called when the camera rig asset has been built. */
+		virtual void OnCameraRigBuilt(const UCameraRigAsset* CameraRigAsset) {}
+	};
 }
 
 /**
@@ -141,6 +155,9 @@ public:
 	/** Finds an exposed parameter by name. */
 	GAMEPLAYCAMERAS_API UCameraRigInterfaceParameter* FindInterfaceParameterByName(const FString& ParameterName) const;
 
+	/** Finds an exposed parameter by Guid. */
+	GAMEPLAYCAMERAS_API UCameraRigInterfaceParameter* FindInterfaceParameterByGuid(const FGuid& ParameterGuid) const;
+
 	/** Returns whether an exposed parameter with the given name exists. */
 	GAMEPLAYCAMERAS_API bool HasInterfaceParameter(const FString& ParameterName) const;
 };
@@ -243,6 +260,9 @@ public:
 	GAMEPLAYCAMERAS_API static const FName NodeTreeGraphName;
 	GAMEPLAYCAMERAS_API static const FName TransitionsGraphName;
 
+	/** Event handlers to be notified of data changes. */
+	UE::Cameras::TCameraEventHandlerContainer<UE::Cameras::ICameraRigAssetEventHandler> EventHandlers;
+
 protected:
 
 	// IObjectTreeGraphObject interface.
@@ -303,7 +323,6 @@ private:
 	UPROPERTY(Instanced, meta=(ObjectTreeGraphHidden=true))
 	TArray<TObjectPtr<UObject>> AllTransitionsObjects;
 
-
 	// Deprecated properties.
 
 	UPROPERTY()
diff --git a/Engine/Plugins/Cameras/GameplayCameras/Source/GameplayCameras/Public/Core/CameraRigAssetBuilder.h b/Engine/Plugins/Cameras/GameplayCameras/Source/GameplayCameras/Public/Core/CameraRigAssetBuilder.h
--- a/Engine/Plugins/Cameras/GameplayCameras/Source/GameplayCameras/Public/Core/CameraRigAssetBuilder.h
+++ b/Engine/Plugins/Cameras/GameplayCameras/Source/GameplayCameras/Public/Core/CameraRigAssetBuilder.h
@@ -13,6 +13,7 @@
 class FStructProperty;
 class UCameraNode;
 class UCameraRigAsset;
+class UCameraRigCameraNode;
 class UCameraVariableAsset;
 struct FCameraVariableTableAllocationInfo;
 
@@ -51,6 +52,8 @@ private:
 
 	void FlattenCameraNodeHierarchy();
 
+	void CallPreBuild();
+
 	void GatherOldDrivenParameters();
 	void BuildNewDrivenParameters();
 	void DiscardUnusedPrivateVariables();
@@ -60,6 +63,11 @@ private:
 
 	void UpdateBuildStatus();
 
+private:
+
+	bool SetupCameraParameterOverride(UCameraRigInterfaceParameter* InterfaceParameter);
+	bool SetupInnerCameraRigParameterOverride(UCameraRigInterfaceParameter* InterfaceParameter);
+
 private:
 
 	FCameraBuildLog& BuildLog;
@@ -70,6 +78,8 @@ private:
 
 	using FDrivenParameterKey = TTuple<FStructProperty*, UCameraNode*>;
 	TMap<FDrivenParameterKey, UCameraVariableAsset*> OldDrivenParameters;
+	using FDrivenOverrideKey = TTuple<FGuid, UCameraRigCameraNode*>;
+	TMap<FDrivenOverrideKey, UCameraVariableAsset*> OldDrivenOverrides;
 
 	using FReusableInterfaceParameterInfo = TTuple<UCameraVariableAsset*, bool>;
 	TMap<UCameraRigInterfaceParameter*, FReusableInterfaceParameterInfo> OldInterfaceParameters;
diff --git a/Engine/Plugins/Cameras/GameplayCameras/Source/GameplayCameras/Public/Core/CameraRigAssetReference.h b/Engine/Plugins/Cameras/GameplayCameras/Source/GameplayCameras/Public/Core/CameraRigAssetReference.h
new file mode 100644
--- /dev/null
+++ b/Engine/Plugins/Cameras/GameplayCameras/Source/GameplayCameras/Public/Core/CameraRigAssetReference.h
@@ -0,0 +1,472 @@
+// Copyright Epic Games, Inc. All Rights Reserved.
+
+#pragma once
+
+#include "Core/CameraParameters.h"
+#include "Core/CameraRigAsset.h"
+
+#include "CameraRigAssetReference.generated.h"
+
+struct FCameraRigAssetReference;
+
+namespace UE::Cameras
+{
+	class FCameraRigAssetBuilder;
+	class FCameraRigAssetReferenceDetailsCustomization;
+}
+
+/** Base struct for camera rig parameter overrides. */
+USTRUCT()
+struct FCameraRigParameterOverrideBase
+{
+	GENERATED_BODY()
+
+	/**
+	 * The Guid of the overriden interface parameter in the inner camera rig.
+	 */
+	UPROPERTY()
+	FGuid InterfaceParameterGuid;
+
+	/**
+	 * The Guid of the overriden interface parameter's private variable in the
+	 * inner camera rig.
+	 *
+	 * This can be derived from InterfaceParameterGuid, but we cache this during
+	 * the build process to avoid searching for interface parameters.
+	 */
+	UPROPERTY()
+	FGuid PrivateVariableGuid;
+
+	/**
+	 * The name of the overriden interface parameter in the inner camera rig.
+	 *
+	 * This can be derived from InterfaceParameterGuid, but we cache this during
+	 * the build process to avoid searching for interface parameters.
+	 */
+	UPROPERTY()
+	FString InterfaceParameterName;
+
+	/**
+	 *
+	 */
+	UPROPERTY()
+	bool bInvalid = false;
+};
+
+USTRUCT()
+struct FBooleanCameraRigParameterOverride : public FCameraRigParameterOverrideBase
+{
+	GENERATED_BODY()
+
+	using CameraParameterType = FBooleanCameraParameter;
+
+	UPROPERTY()
+	FBooleanCameraParameter Value;
+};
+
+USTRUCT()
+struct FInteger32CameraRigParameterOverride : public FCameraRigParameterOverrideBase
+{
+	GENERATED_BODY()
+
+	using CameraParameterType = FInteger32CameraParameter;
+
+	UPROPERTY()
+	FInteger32CameraParameter Value;
+};
+
+USTRUCT()
+struct FFloatCameraRigParameterOverride : public FCameraRigParameterOverrideBase
+{
+	GENERATED_BODY()
+
+	using CameraParameterType = FFloatCameraParameter;
+
+	UPROPERTY()
+	FFloatCameraParameter Value;
+};
+
+USTRUCT()
+struct FDoubleCameraRigParameterOverride : public FCameraRigParameterOverrideBase
+{
+	GENERATED_BODY()
+
+	using CameraParameterType = FDoubleCameraParameter;
+
+	UPROPERTY()
+	FDoubleCameraParameter Value;
+};
+
+USTRUCT()
+struct FVector2fCameraRigParameterOverride : public FCameraRigParameterOverrideBase
+{
+	GENERATED_BODY()
+
+	using CameraParameterType = FVector2fCameraParameter;
+
+	UPROPERTY()
+	FVector2fCameraParameter Value;
+};
+
+USTRUCT()
+struct FVector2dCameraRigParameterOverride : public FCameraRigParameterOverrideBase
+{
+	GENERATED_BODY()
+
+	using CameraParameterType = FVector2dCameraParameter;
+
+	UPROPERTY()
+	FVector2dCameraParameter Value;
+};
+
+USTRUCT()
+struct FVector3fCameraRigParameterOverride : public FCameraRigParameterOverrideBase
+{
+	GENERATED_BODY()
+
+	using CameraParameterType = FVector3fCameraParameter;
+
+	UPROPERTY()
+	FVector3fCameraParameter Value;
+};
+
+USTRUCT()
+struct FVector3dCameraRigParameterOverride : public FCameraRigParameterOverrideBase
+{
+	GENERATED_BODY()
+
+	using CameraParameterType = FVector3dCameraParameter;
+
+	UPROPERTY()
+	FVector3dCameraParameter Value;
+};
+
+USTRUCT()
+struct FVector4fCameraRigParameterOverride : public FCameraRigParameterOverrideBase
+{
+	GENERATED_BODY()
+
+	using CameraParameterType = FVector4fCameraParameter;
+
+	UPROPERTY()
+	FVector4fCameraParameter Value;
+};
+
+USTRUCT()
+struct FVector4dCameraRigParameterOverride : public FCameraRigParameterOverrideBase
+{
+	GENERATED_BODY()
+
+	using CameraParameterType = FVector4dCameraParameter;
+
+	UPROPERTY()
+	FVector4dCameraParameter Value;
+};
+
+USTRUCT()
+struct FRotator3fCameraRigParameterOverride : public FCameraRigParameterOverrideBase
+{
+	GENERATED_BODY()
+
+	using CameraParameterType = FRotator3fCameraParameter;
+
+	UPROPERTY()
+	FRotator3fCameraParameter Value;
+};
+
+USTRUCT()
+struct FRotator3dCameraRigParameterOverride : public FCameraRigParameterOverrideBase
+{
+	GENERATED_BODY()
+
+	using CameraParameterType = FRotator3dCameraParameter;
+
+	UPROPERTY()
+	FRotator3dCameraParameter Value;
+};
+
+USTRUCT()
+struct FTransform3fCameraRigParameterOverride : public FCameraRigParameterOverrideBase
+{
+	GENERATED_BODY()
+
+	using CameraParameterType = FTransform3fCameraParameter;
+
+	UPROPERTY()
+	FTransform3fCameraParameter Value;
+};
+
+USTRUCT()
+struct FTransform3dCameraRigParameterOverride : public FCameraRigParameterOverrideBase
+{
+	GENERATED_BODY()
+
+	using CameraParameterType = FTransform3dCameraParameter;
+
+	UPROPERTY()
+	FTransform3dCameraParameter Value;
+};
+
+/**
+ * A structure that holds lists of camera rig interface parameter overrides, one list
+ * per parameter type.
+ */
+USTRUCT()
+struct GAMEPLAYCAMERAS_API FCameraRigParameterOverrides
+{
+	GENERATED_BODY()
+
+public:
+
+	/** Whether the given camera rig interface parameter is currently overriden. */
+	template<typename ParameterOverrideType>
+	bool IsParameterOverriden(const FGuid& CameraRigParameterGuid) const;
+
+	/** Find a parameter override for the given inner camera rig interface parameter. */
+	template<typename ParameterOverrideType>
+	ParameterOverrideType* FindParameterOverride(const FGuid& CameraRigParameterGuid);
+
+	/** Find or create a parameter override for the given inner camera rig interface parameter. */
+	template<typename ParameterOverrideType>
+	ParameterOverrideType& FindOrAddParameterOverride(const UCameraRigInterfaceParameter* CameraRigParameter);
+
+	/** Remove any parameter override for the given inner camera rig interface parameter. */
+	template<typename ParameterOverrideType>
+	void RemoveParameterOverride(const FGuid& CameraRigParameterGuid);
+
+	/** Remove all parameter overrides. */
+	void Reset();
+
+public:
+
+#define UE_CAMERA_VARIABLE_FOR_TYPE(ValueType, ValueName)\
+	inline TArrayView<const F##ValueName##CameraRigParameterOverride> Get##ValueName##Overrides() const\
+	{\
+		return ValueName##Overrides;\
+	}\
+	inline TArrayView<F##ValueName##CameraRigParameterOverride> Get##ValueName##Overrides()\
+	{\
+		return ValueName##Overrides;\
+	}
+UE_CAMERA_VARIABLE_FOR_ALL_TYPES()
+#undef UE_CAMERA_VARIABLE_FOR_TYPE
+
+
+private:
+
+	template<typename ParameterOverrideType>
+	bool IsParameterOverriden(const TArray<ParameterOverrideType>& OverridesArray, const FGuid& CameraRigParameterGuid) const;
+
+	template<typename ParameterOverrideType>
+	ParameterOverrideType* FindParameterOverride(TArray<ParameterOverrideType>& OverridesArray, const FGuid& CameraRigParameterGuid);
+
+	template<typename ParameterOverrideType>
+	ParameterOverrideType& FindOrAddParameterOverride(TArray<ParameterOverrideType>& OverridesArray, const UCameraRigInterfaceParameter* CameraRigParameter);
+
+	template<typename ParameterOverrideType>
+	void RemoveParameterOverride(TArray<ParameterOverrideType>& OverridesArray, const FGuid& CameraRigParameterGuid);
+
+
+public:
+
+	// Internal API.
+
+	template<typename ParameterOverrideType>
+	inline void AppendParameterOverrides(TArrayView<ParameterOverrideType> NewOverrides);
+
+private:
+
+	// Interface parameter overrides
+
+	UPROPERTY()
+	TArray<FBooleanCameraRigParameterOverride> BooleanOverrides;
+	UPROPERTY()
+	TArray<FInteger32CameraRigParameterOverride> Integer32Overrides;
+	UPROPERTY()
+	TArray<FFloatCameraRigParameterOverride> FloatOverrides;
+	UPROPERTY()
+	TArray<FDoubleCameraRigParameterOverride> DoubleOverrides;
+	UPROPERTY()
+	TArray<FVector2fCameraRigParameterOverride> Vector2fOverrides;
+	UPROPERTY()
+	TArray<FVector2dCameraRigParameterOverride> Vector2dOverrides;
+	UPROPERTY()
+	TArray<FVector3fCameraRigParameterOverride> Vector3fOverrides;
+	UPROPERTY()
+	TArray<FVector3dCameraRigParameterOverride> Vector3dOverrides;
+	UPROPERTY()
+	TArray<FVector4fCameraRigParameterOverride> Vector4fOverrides;
+	UPROPERTY()
+	TArray<FVector4dCameraRigParameterOverride> Vector4dOverrides;
+	UPROPERTY()
+	TArray<FRotator3fCameraRigParameterOverride> Rotator3fOverrides;
+	UPROPERTY()
+	TArray<FRotator3dCameraRigParameterOverride> Rotator3dOverrides;
+	UPROPERTY()
+	TArray<FTransform3fCameraRigParameterOverride> Transform3fOverrides;
+	UPROPERTY()
+	TArray<FTransform3dCameraRigParameterOverride> Transform3dOverrides;
+
+	friend struct FCameraRigAssetReference;
+
+	friend class UE::Cameras::FCameraRigAssetBuilder;
+};
+
+template<typename ParameterOverrideType>
+bool FCameraRigParameterOverrides::IsParameterOverriden(const TArray<ParameterOverrideType>& OverridesArray, const FGuid& CameraRigParameterGuid) const
+{
+	return OverridesArray.ContainsByPredicate(
+			[CameraRigParameterGuid](ParameterOverrideType& Item)
+			{
+				return (Item.InterfaceParameterGuid == CameraRigParameterGuid);
+			});
+}
+
+template<typename ParameterOverrideType>
+ParameterOverrideType* FCameraRigParameterOverrides::FindParameterOverride(TArray<ParameterOverrideType>& OverridesArray, const FGuid& CameraRigParameterGuid)
+{
+	ParameterOverrideType* FoundItem = OverridesArray.FindByPredicate(
+			[CameraRigParameterGuid](ParameterOverrideType& Item)
+			{
+				return (Item.InterfaceParameterGuid == CameraRigParameterGuid);
+			});
+	return FoundItem;
+}
+
+template<typename ParameterOverrideType>
+ParameterOverrideType& FCameraRigParameterOverrides::FindOrAddParameterOverride(TArray<ParameterOverrideType>& OverridesArray, const UCameraRigInterfaceParameter* CameraRigParameter)
+{
+	ParameterOverrideType* Existing = FindParameterOverride<ParameterOverrideType>(OverridesArray, CameraRigParameter->Guid);
+	if (Existing)
+	{
+		return *Existing;
+	}
+	else
+	{
+		ParameterOverrideType& NewOverride = OverridesArray.Emplace_GetRef();
+		NewOverride.InterfaceParameterGuid = CameraRigParameter->Guid;
+		NewOverride.InterfaceParameterName = CameraRigParameter->InterfaceParameterName;
+		return NewOverride;
+	}
+}
+
+template<typename ParameterOverrideType>
+void FCameraRigParameterOverrides::RemoveParameterOverride(TArray<ParameterOverrideType>& OverridesArray, const FGuid& CameraRigParameterGuid)
+{
+	OverridesArray.RemoveAll(
+			[CameraRigParameterGuid](ParameterOverrideType& Item)
+			{
+				return (Item.InterfaceParameterGuid == CameraRigParameterGuid);
+			});
+}
+
+#define UE_CAMERA_VARIABLE_FOR_TYPE(ValueType, ValueName)\
+template<>\
+inline bool FCameraRigParameterOverrides::IsParameterOverriden<F##ValueName##CameraRigParameterOverride>(const FGuid& CameraRigParameterGuid) const\
+{\
+	return this->IsParameterOverriden<F##ValueName##CameraRigParameterOverride>(this->ValueName##Overrides, CameraRigParameterGuid);\
+}\
+template<>\
+inline F##ValueName##CameraRigParameterOverride* FCameraRigParameterOverrides::FindParameterOverride<F##ValueName##CameraRigParameterOverride>(const FGuid& CameraRigParameterGuid)\
+{\
+	return this->FindParameterOverride<F##ValueName##CameraRigParameterOverride>(this->ValueName##Overrides, CameraRigParameterGuid);\
+}\
+template<>\
+inline F##ValueName##CameraRigParameterOverride& FCameraRigParameterOverrides::FindOrAddParameterOverride<F##ValueName##CameraRigParameterOverride>(const UCameraRigInterfaceParameter* CameraRigParameter)\
+{\
+	return this->FindOrAddParameterOverride<F##ValueName##CameraRigParameterOverride>(this->ValueName##Overrides, CameraRigParameter);\
+}\
+template<>\
+inline void FCameraRigParameterOverrides::RemoveParameterOverride<F##ValueName##CameraRigParameterOverride>(const FGuid& CameraRigParameterGuid)\
+{\
+	this->RemoveParameterOverride<F##ValueName##CameraRigParameterOverride>(this->ValueName##Overrides, CameraRigParameterGuid);\
+}\
+template<>\
+inline void FCameraRigParameterOverrides::AppendParameterOverrides<F##ValueName##CameraRigParameterOverride>(TArrayView<F##ValueName##CameraRigParameterOverride> NewOverrides)\
+{\
+	this->ValueName##Overrides.Append(NewOverrides);\
+}
+UE_CAMERA_VARIABLE_FOR_ALL_TYPES()
+#undef UE_CAMERA_VARIABLE_FOR_TYPE
+
+/**
+ * A structure holding a reference to a camera rig asset, along with the interface parameter
+ * override values.
+ */
+USTRUCT(BlueprintType)
+struct GAMEPLAYCAMERAS_API FCameraRigAssetReference
+{
+	GENERATED_BODY()
+
+	/** Returns whether this reference points to a valid camera rig. */
+	bool IsValid() const
+	{
+		return CameraRig != nullptr;
+	}
+
+	/** Gets the referenced camera rig. */
+	UCameraRigAsset* GetCameraRig()
+	{
+		return CameraRig;
+	}
+
+	/** Gets the referenced camera rig. */
+	const UCameraRigAsset* GetCameraRig() const
+	{
+		return CameraRig;
+	}
+
+	/**
+	 * Sets the referenced camerar rig.
+	 * This will check any existing parameter overrides, tagging them as invalid as needed.
+	 */
+	void SetCameraRig(UCameraRigAsset* InCameraRig)
+	{
+		if (CameraRig != InCameraRig)
+		{
+			CameraRig = InCameraRig;
+			UpdateParameterOverrides();
+		}
+	}
+
+	/** Gets the parameter overrides. */
+	const FCameraRigParameterOverrides& GetParameterOverrides() const
+	{
+		return ParameterOverrides;
+	}
+
+	/** Gets the parameter overrides. */
+	FCameraRigParameterOverrides& GetParameterOverrides()
+	{
+		return ParameterOverrides;
+	}
+
+public:
+
+	// Internal API.
+
+	bool SerializeFromMismatchedTag(struct FPropertyTag const& Tag, FStructuredArchive::FSlot Slot);
+	bool UpdateParameterOverrides();
+
+private:
+
+	/** The referenced camera rig. */
+	UPROPERTY(EditAnywhere, Category="")
+	TObjectPtr<UCameraRigAsset> CameraRig;
+
+	/** The parameter overrides. */
+	UPROPERTY(EditAnywhere, Category="", meta=(FixedLayout))
+	FCameraRigParameterOverrides ParameterOverrides;
+
+	friend class UE::Cameras::FCameraRigAssetReferenceDetailsCustomization;
+};
+
+template<>
+struct TStructOpsTypeTraits<FCameraRigAssetReference> : public TStructOpsTypeTraitsBase2<FCameraRigAssetReference>
+{
+	enum
+	{
+		WithStructuredSerializeFromMismatchedTag = true
+	};
+};
+
diff --git a/Engine/Plugins/Cameras/GameplayCameras/Source/GameplayCameras/Public/Core/CameraVariableAssets.h b/Engine/Plugins/Cameras/GameplayCameras/Source/GameplayCameras/Public/Core/CameraVariableAssets.h
--- a/Engine/Plugins/Cameras/GameplayCameras/Source/GameplayCameras/Public/Core/CameraVariableAssets.h
+++ b/Engine/Plugins/Cameras/GameplayCameras/Source/GameplayCameras/Public/Core/CameraVariableAssets.h
@@ -76,6 +76,8 @@ public:
 
 	using ValueType = bool;
 
+	bool GetDefaultValue() const { return bDefaultValue; }
+
 	virtual ECameraVariableType GetVariableType() const override { return ECameraVariableType::Boolean; }
 	virtual const uint8* GetDefaultValuePtr() const override { return reinterpret_cast<const uint8*>(&bDefaultValue); }
 
@@ -100,6 +102,8 @@ public:
 
 	using ValueType = int32;
 
+	int32 GetDefaultValue() const { return DefaultValue; }
+
 	virtual ECameraVariableType GetVariableType() const override { return ECameraVariableType::Integer32; }
 	virtual const uint8* GetDefaultValuePtr() const override { return reinterpret_cast<const uint8*>(&DefaultValue); }
 
@@ -124,6 +128,8 @@ public:
 
 	using ValueType = float;
 
+	float GetDefaultValue() const { return DefaultValue; }
+
 	virtual ECameraVariableType GetVariableType() const override { return ECameraVariableType::Float; }
 	virtual const uint8* GetDefaultValuePtr() const override { return reinterpret_cast<const uint8*>(&DefaultValue); }
 
@@ -148,6 +154,8 @@ public:
 
 	using ValueType = double;
 
+	double GetDefaultValue() const { return DefaultValue; }
+
 	virtual ECameraVariableType GetVariableType() const override { return ECameraVariableType::Double; }
 	virtual const uint8* GetDefaultValuePtr() const override { return reinterpret_cast<const uint8*>(&DefaultValue); }
 
@@ -172,6 +180,8 @@ public:
 
 	using ValueType = FVector2f;
 
+	const ValueType& GetDefaultValue() const { return DefaultValue; }
+
 	virtual ECameraVariableType GetVariableType() const override { return ECameraVariableType::Vector2f; }
 	virtual const uint8* GetDefaultValuePtr() const override { return reinterpret_cast<const uint8*>(&DefaultValue); }
 
@@ -196,6 +206,8 @@ public:
 
 	using ValueType = FVector2d;
 
+	const ValueType& GetDefaultValue() const { return DefaultValue; }
+
 	virtual ECameraVariableType GetVariableType() const override { return ECameraVariableType::Vector2d; }
 	virtual const uint8* GetDefaultValuePtr() const override { return reinterpret_cast<const uint8*>(&DefaultValue); }
 
@@ -220,6 +232,8 @@ public:
 
 	using ValueType = FVector3f;
 
+	const ValueType& GetDefaultValue() const { return DefaultValue; }
+
 	virtual ECameraVariableType GetVariableType() const override { return ECameraVariableType::Vector3f; }
 	virtual const uint8* GetDefaultValuePtr() const override { return reinterpret_cast<const uint8*>(&DefaultValue); }
 
@@ -244,6 +258,8 @@ public:
 
 	using ValueType = FVector3d;
 
+	const ValueType& GetDefaultValue() const { return DefaultValue; }
+
 	virtual ECameraVariableType GetVariableType() const override { return ECameraVariableType::Vector3d; }
 	virtual const uint8* GetDefaultValuePtr() const override { return reinterpret_cast<const uint8*>(&DefaultValue); }
 
@@ -268,6 +284,8 @@ public:
 
 	using ValueType = FVector4f;
 
+	const ValueType& GetDefaultValue() const { return DefaultValue; }
+
 	virtual ECameraVariableType GetVariableType() const override { return ECameraVariableType::Vector4f; }
 	virtual const uint8* GetDefaultValuePtr() const override { return reinterpret_cast<const uint8*>(&DefaultValue); }
 
@@ -292,6 +310,8 @@ public:
 
 	using ValueType = FVector4d;
 
+	const ValueType& GetDefaultValue() const { return DefaultValue; }
+
 	virtual ECameraVariableType GetVariableType() const override { return ECameraVariableType::Vector4d; }
 	virtual const uint8* GetDefaultValuePtr() const override { return reinterpret_cast<const uint8*>(&DefaultValue); }
 
@@ -316,6 +336,8 @@ public:
 
 	using ValueType = FRotator3f;
 
+	const ValueType& GetDefaultValue() const { return DefaultValue; }
+
 	virtual ECameraVariableType GetVariableType() const override { return ECameraVariableType::Rotator3f; }
 	virtual const uint8* GetDefaultValuePtr() const override { return reinterpret_cast<const uint8*>(&DefaultValue); }
 
@@ -340,6 +362,8 @@ public:
 
 	using ValueType = FRotator3d;
 
+	const ValueType& GetDefaultValue() const { return DefaultValue; }
+
 	virtual ECameraVariableType GetVariableType() const override { return ECameraVariableType::Rotator3d; }
 	virtual const uint8* GetDefaultValuePtr() const override { return reinterpret_cast<const uint8*>(&DefaultValue); }
 
@@ -364,6 +388,8 @@ public:
 
 	using ValueType = FTransform3f;
 
+	const ValueType& GetDefaultValue() const { return DefaultValue; }
+
 	virtual ECameraVariableType GetVariableType() const override { return ECameraVariableType::Transform3f; }
 	virtual const uint8* GetDefaultValuePtr() const override { return reinterpret_cast<const uint8*>(&DefaultValue); }
 
@@ -388,6 +414,8 @@ public:
 
 	using ValueType = FTransform3d;
 
+	const ValueType& GetDefaultValue() const { return DefaultValue; }
+
 	virtual ECameraVariableType GetVariableType() const override { return ECameraVariableType::Transform3d; }
 	virtual const uint8* GetDefaultValuePtr() const override { return reinterpret_cast<const uint8*>(&DefaultValue); }
 
diff --git a/Engine/Plugins/Cameras/GameplayCameras/Source/GameplayCameras/Public/GameplayCamerasDelegates.h b/Engine/Plugins/Cameras/GameplayCameras/Source/GameplayCameras/Public/GameplayCamerasDelegates.h
new file mode 100644
--- /dev/null
+++ b/Engine/Plugins/Cameras/GameplayCameras/Source/GameplayCameras/Public/GameplayCamerasDelegates.h
@@ -0,0 +1,36 @@
+// Copyright Epic Games, Inc. All Rights Reserved.
+
+#pragma once
+
+#include "CoreTypes.h"
+#include "Delegates/Delegate.h"
+
+class UCameraRigAsset;
+
+namespace UE::Cameras
+{
+
+class FCameraBuildLog;
+
+DECLARE_MULTICAST_DELEGATE_TwoParams(FOnCameraRigAssetBuilt, UCameraRigAsset*, FCameraBuildLog&);
+
+/**
+ * Global delegates for the GameplayCameras module.
+ */
+class GAMEPLAYCAMERAS_API FGameplayCamerasDelegates
+{
+public:
+
+	/** Broadcast for when a camera rig has been built. */
+	static inline FOnCameraRigAssetBuilt& OnCameraRigAssetBuilt()
+	{
+		return OnCameraRigAssetBuiltDelegates;
+	}
+
+private:
+
+	static FOnCameraRigAssetBuilt OnCameraRigAssetBuiltDelegates;
+};
+
+}  // namespace UE::Cameras
+
diff --git a/Engine/Plugins/Cameras/GameplayCameras/Source/GameplayCameras/Public/Nodes/Common/CameraRigCameraNode.h b/Engine/Plugins/Cameras/GameplayCameras/Source/GameplayCameras/Public/Nodes/Common/CameraRigCameraNode.h
--- a/Engine/Plugins/Cameras/GameplayCameras/Source/GameplayCameras/Public/Nodes/Common/CameraRigCameraNode.h
+++ b/Engine/Plugins/Cameras/GameplayCameras/Source/GameplayCameras/Public/Nodes/Common/CameraRigCameraNode.h
@@ -3,8 +3,8 @@
 #pragma once
 
 #include "Core/CameraNode.h"
-#include "Core/CameraParameters.h"
 #include "Core/CameraRigAsset.h"
+#include "Core/CameraRigAssetReference.h"
 
 #include "CameraRigCameraNode.generated.h"
 
@@ -13,286 +13,6 @@ namespace UE::Cameras
 	class FCameraRigCameraNodeEvaluator;
 }
 
-USTRUCT()
-struct FBooleanCameraRigParameterOverride
-{
-	GENERATED_BODY()
-
-	using CameraParameterType = FBooleanCameraParameter;
-
-	UPROPERTY()
-	FBooleanCameraParameter Value;
-
-	UPROPERTY()
-	FGuid PrivateVariableGuid;
-
-	UPROPERTY()
-	FGuid InterfaceParameterGuid;
-
-	UPROPERTY()
-	FString InterfaceParameterName;
-};
-
-USTRUCT()
-struct FInteger32CameraRigParameterOverride
-{
-	GENERATED_BODY()
-
-	using CameraParameterType = FInteger32CameraParameter;
-
-	UPROPERTY()
-	FInteger32CameraParameter Value;
-
-	UPROPERTY()
-	FGuid PrivateVariableGuid;
-
-	UPROPERTY()
-	FGuid InterfaceParameterGuid;
-
-	UPROPERTY()
-	FString InterfaceParameterName;
-};
-
-USTRUCT()
-struct FFloatCameraRigParameterOverride
-{
-	GENERATED_BODY()
-
-	using CameraParameterType = FFloatCameraParameter;
-
-	UPROPERTY()
-	FFloatCameraParameter Value;
-
-	UPROPERTY()
-	FGuid PrivateVariableGuid;
-
-	UPROPERTY()
-	FGuid InterfaceParameterGuid;
-
-	UPROPERTY()
-	FString InterfaceParameterName;
-};
-
-USTRUCT()
-struct FDoubleCameraRigParameterOverride
-{
-	GENERATED_BODY()
-
-	using CameraParameterType = FDoubleCameraParameter;
-
-	UPROPERTY()
-	FDoubleCameraParameter Value;
-
-	UPROPERTY()
-	FGuid PrivateVariableGuid;
-
-	UPROPERTY()
-	FGuid InterfaceParameterGuid;
-
-	UPROPERTY()
-	FString InterfaceParameterName;
-};
-
-USTRUCT()
-struct FVector2fCameraRigParameterOverride
-{
-	GENERATED_BODY()
-
-	using CameraParameterType = FVector2fCameraParameter;
-
-	UPROPERTY()
-	FVector2fCameraParameter Value;
-
-	UPROPERTY()
-	FGuid PrivateVariableGuid;
-
-	UPROPERTY()
-	FGuid InterfaceParameterGuid;
-
-	UPROPERTY()
-	FString InterfaceParameterName;
-};
-
-USTRUCT()
-struct FVector2dCameraRigParameterOverride
-{
-	GENERATED_BODY()
-
-	using CameraParameterType = FVector2dCameraParameter;
-
-	UPROPERTY()
-	FVector2dCameraParameter Value;
-
-	UPROPERTY()
-	FGuid PrivateVariableGuid;
-
-	UPROPERTY()
-	FGuid InterfaceParameterGuid;
-
-	UPROPERTY()
-	FString InterfaceParameterName;
-};
-
-USTRUCT()
-struct FVector3fCameraRigParameterOverride
-{
-	GENERATED_BODY()
-
-	using CameraParameterType = FVector3fCameraParameter;
-
-	UPROPERTY()
-	FVector3fCameraParameter Value;
-
-	UPROPERTY()
-	FGuid PrivateVariableGuid;
-
-	UPROPERTY()
-	FGuid InterfaceParameterGuid;
-
-	UPROPERTY()
-	FString InterfaceParameterName;
-};
-
-USTRUCT()
-struct FVector3dCameraRigParameterOverride
-{
-	GENERATED_BODY()
-
-	using CameraParameterType = FVector3dCameraParameter;
-
-	UPROPERTY()
-	FVector3dCameraParameter Value;
-
-	UPROPERTY()
-	FGuid PrivateVariableGuid;
-
-	UPROPERTY()
-	FGuid InterfaceParameterGuid;
-
-	UPROPERTY()
-	FString InterfaceParameterName;
-};
-
-USTRUCT()
-struct FVector4fCameraRigParameterOverride
-{
-	GENERATED_BODY()
-
-	using CameraParameterType = FVector4fCameraParameter;
-
-	UPROPERTY()
-	FVector4fCameraParameter Value;
-
-	UPROPERTY()
-	FGuid PrivateVariableGuid;
-
-	UPROPERTY()
-	FGuid InterfaceParameterGuid;
-
-	UPROPERTY()
-	FString InterfaceParameterName;
-};
-
-USTRUCT()
-struct FVector4dCameraRigParameterOverride
-{
-	GENERATED_BODY()
-
-	using CameraParameterType = FVector4dCameraParameter;
-
-	UPROPERTY()
-	FVector4dCameraParameter Value;
-
-	UPROPERTY()
-	FGuid PrivateVariableGuid;
-
-	UPROPERTY()
-	FGuid InterfaceParameterGuid;
-
-	UPROPERTY()
-	FString InterfaceParameterName;
-};
-
-USTRUCT()
-struct FRotator3fCameraRigParameterOverride
-{
-	GENERATED_BODY()
-
-	using CameraParameterType = FRotator3fCameraParameter;
-
-	UPROPERTY()
-	FRotator3fCameraParameter Value;
-
-	UPROPERTY()
-	FGuid PrivateVariableGuid;
-
-	UPROPERTY()
-	FGuid InterfaceParameterGuid;
-
-	UPROPERTY()
-	FString InterfaceParameterName;
-};
-
-USTRUCT()
-struct FRotator3dCameraRigParameterOverride
-{
-	GENERATED_BODY()
-
-	using CameraParameterType = FRotator3dCameraParameter;
-
-	UPROPERTY()
-	FRotator3dCameraParameter Value;
-
-	UPROPERTY()
-	FGuid PrivateVariableGuid;
-
-	UPROPERTY()
-	FGuid InterfaceParameterGuid;
-
-	UPROPERTY()
-	FString InterfaceParameterName;
-};
-
-USTRUCT()
-struct FTransform3fCameraRigParameterOverride
-{
-	GENERATED_BODY()
-
-	using CameraParameterType = FTransform3fCameraParameter;
-
-	UPROPERTY()
-	FTransform3fCameraParameter Value;
-
-	UPROPERTY()
-	FGuid PrivateVariableGuid;
-
-	UPROPERTY()
-	FGuid InterfaceParameterGuid;
-
-	UPROPERTY()
-	FString InterfaceParameterName;
-};
-
-USTRUCT()
-struct FTransform3dCameraRigParameterOverride
-{
-	GENERATED_BODY()
-
-	using CameraParameterType = FTransform3dCameraParameter;
-
-	UPROPERTY()
-	FTransform3dCameraParameter Value;
-
-	UPROPERTY()
-	FGuid PrivateVariableGuid;
-
-	UPROPERTY()
-	FGuid InterfaceParameterGuid;
-
-	UPROPERTY()
-	FString InterfaceParameterName;
-};
-
 /**
  * A camera node that runs a camera rig's own node tree.
  */
@@ -301,128 +21,56 @@ class UCameraRigCameraNode : public UCameraNode
 {
 	GENERATED_BODY()
 
-public:
-
-	template<typename ParameterOverrideType>
-	ParameterOverrideType* FindParameterOverride(const FGuid& CameraRigParameterGuid);
-	
-	template<typename ParameterOverrideType>
-	ParameterOverrideType& GetOrAddParameterOverride(const UCameraRigInterfaceParameter* CameraRigParameter);
-
-	template<typename ParameterOverrideType>
-	void RemoveParameterOverride(const FGuid& CameraRigParameterGuid);
-
-private:
-
-	template<typename ParameterOverrideType>
-	ParameterOverrideType* FindParameterOverride(TArray<ParameterOverrideType>& OverridesArray, const FGuid& CameraRigParameterGuid);
-	
-	template<typename ParameterOverrideType>
-	ParameterOverrideType& GetOrAddParameterOverride(TArray<ParameterOverrideType>& OverridesArray, const UCameraRigInterfaceParameter* CameraRigParameter);
-
-	template<typename ParameterOverrideType>
-	void RemoveParameterOverride(TArray<ParameterOverrideType>& OverridesArray, const FGuid& CameraRigParameterGuid);
-
 protected:
 
 	// UCameraNode interface.
+	virtual void OnPreBuild(FCameraBuildLog& BuildLog) override;
 	virtual void OnBuild(FCameraRigBuildContext& BuildContext) override;
 	virtual FCameraNodeEvaluatorPtr OnBuildEvaluator(FCameraNodeEvaluatorBuilder& Builder) const override;
 
+	// UObject interface.
+	void PostLoad() override;
+
 public:
 
 	/** The camera rig to run. */
 	UPROPERTY(EditAnywhere, Category=Common, meta=(ObjectTreeGraphHidden=true))
-	TObjectPtr<UCameraRigAsset> CameraRig;
+	FCameraRigAssetReference CameraRigReference;
 
 private:
 
+	// Deprecated properties, predating FCameraRigAssetReference
+
 	UPROPERTY()
-	TArray<FBooleanCameraRigParameterOverride> BooleanOverrides;
+	TObjectPtr<UCameraRigAsset> CameraRig_DEPRECATED;
+
 	UPROPERTY()
-	TArray<FInteger32CameraRigParameterOverride> Integer32Overrides;
+	TArray<FBooleanCameraRigParameterOverride> BooleanOverrides_DEPRECATED;
 	UPROPERTY()
-	TArray<FFloatCameraRigParameterOverride> FloatOverrides;
+	TArray<FInteger32CameraRigParameterOverride> Integer32Overrides_DEPRECATED;
 	UPROPERTY()
-	TArray<FDoubleCameraRigParameterOverride> DoubleOverrides;
+	TArray<FFloatCameraRigParameterOverride> FloatOverrides_DEPRECATED;
 	UPROPERTY()
-	TArray<FVector2fCameraRigParameterOverride> Vector2fOverrides;
+	TArray<FDoubleCameraRigParameterOverride> DoubleOverrides_DEPRECATED;
 	UPROPERTY()
-	TArray<FVector2dCameraRigParameterOverride> Vector2dOverrides;
+	TArray<FVector2fCameraRigParameterOverride> Vector2fOverrides_DEPRECATED;
 	UPROPERTY()
-	TArray<FVector3fCameraRigParameterOverride> Vector3fOverrides;
+	TArray<FVector2dCameraRigParameterOverride> Vector2dOverrides_DEPRECATED;
 	UPROPERTY()
-	TArray<FVector3dCameraRigParameterOverride> Vector3dOverrides;
+	TArray<FVector3fCameraRigParameterOverride> Vector3fOverrides_DEPRECATED;
 	UPROPERTY()
-	TArray<FVector4fCameraRigParameterOverride> Vector4fOverrides;
+	TArray<FVector3dCameraRigParameterOverride> Vector3dOverrides_DEPRECATED;
 	UPROPERTY()
-	TArray<FVector4dCameraRigParameterOverride> Vector4dOverrides;
+	TArray<FVector4fCameraRigParameterOverride> Vector4fOverrides_DEPRECATED;
 	UPROPERTY()
-	TArray<FRotator3fCameraRigParameterOverride> Rotator3fOverrides;
+	TArray<FVector4dCameraRigParameterOverride> Vector4dOverrides_DEPRECATED;
 	UPROPERTY()
-	TArray<FRotator3dCameraRigParameterOverride> Rotator3dOverrides;
+	TArray<FRotator3fCameraRigParameterOverride> Rotator3fOverrides_DEPRECATED;
 	UPROPERTY()
-	TArray<FTransform3fCameraRigParameterOverride> Transform3fOverrides;
+	TArray<FRotator3dCameraRigParameterOverride> Rotator3dOverrides_DEPRECATED;
 	UPROPERTY()
-	TArray<FTransform3dCameraRigParameterOverride> Transform3dOverrides;
-
-	friend class UE::Cameras::FCameraRigCameraNodeEvaluator;
+	TArray<FTransform3fCameraRigParameterOverride> Transform3fOverrides_DEPRECATED;
+	UPROPERTY()
+	TArray<FTransform3dCameraRigParameterOverride> Transform3dOverrides_DEPRECATED;
 };
 
-template<typename ParameterOverrideType>
-ParameterOverrideType* UCameraRigCameraNode::FindParameterOverride(TArray<ParameterOverrideType>& OverridesArray, const FGuid& CameraRigParameterGuid)
-{
-	ParameterOverrideType* FoundItem = OverridesArray.FindByPredicate(
-			[CameraRigParameterGuid](ParameterOverrideType& Item)
-			{
-				return (Item.InterfaceParameterGuid == CameraRigParameterGuid);
-			});
-	return FoundItem;
-}
-
-template<typename ParameterOverrideType>
-ParameterOverrideType& UCameraRigCameraNode::GetOrAddParameterOverride(TArray<ParameterOverrideType>& OverridesArray, const UCameraRigInterfaceParameter* CameraRigParameter)
-{
-	ParameterOverrideType* Existing = FindParameterOverride<ParameterOverrideType>(CameraRigParameter->Guid);
-	if (Existing)
-	{
-		return *Existing;
-	}
-	else
-	{
-		ParameterOverrideType& NewOverride = OverridesArray.Emplace_GetRef();
-		NewOverride.InterfaceParameterGuid = CameraRigParameter->Guid;
-		NewOverride.InterfaceParameterName = CameraRigParameter->InterfaceParameterName;
-		return NewOverride;
-	}
-}
-
-template<typename ParameterOverrideType>
-void UCameraRigCameraNode::RemoveParameterOverride(TArray<ParameterOverrideType>& OverridesArray, const FGuid& CameraRigParameterGuid)
-{
-	OverridesArray.RemoveAll(
-			[CameraRigParameterGuid](ParameterOverrideType& Item)
-			{
-				return (Item.InterfaceParameterGuid == CameraRigParameterGuid);
-			});
-}
-
-#define UE_CAMERA_VARIABLE_FOR_TYPE(ValueType, ValueName)\
-template<>\
-inline F##ValueName##CameraRigParameterOverride* UCameraRigCameraNode::FindParameterOverride<F##ValueName##CameraRigParameterOverride>(const FGuid& CameraRigParameterGuid)\
-{\
-	return FindParameterOverride<F##ValueName##CameraRigParameterOverride>(ValueName##Overrides, CameraRigParameterGuid);\
-}\
-template<>\
-inline F##ValueName##CameraRigParameterOverride& UCameraRigCameraNode::GetOrAddParameterOverride<F##ValueName##CameraRigParameterOverride>(const UCameraRigInterfaceParameter* CameraRigParameter)\
-{\
-	return GetOrAddParameterOverride<F##ValueName##CameraRigParameterOverride>(ValueName##Overrides, CameraRigParameter);\
-}\
-template<>\
-inline void UCameraRigCameraNode::RemoveParameterOverride<F##ValueName##CameraRigParameterOverride>(const FGuid& CameraRigParameterGuid)\
-{\
-	RemoveParameterOverride<F##ValueName##CameraRigParameterOverride>(ValueName##Overrides, CameraRigParameterGuid);\
-}
-UE_CAMERA_VARIABLE_FOR_ALL_TYPES()
-#undef UE_CAMERA_VARIABLE_FOR_TYPE
-
diff --git a/Engine/Plugins/Cameras/GameplayCameras/Source/GameplayCamerasEditor/Private/Customizations/CameraParameterDetailsCustomizations.cpp b/Engine/Plugins/Cameras/GameplayCameras/Source/GameplayCamerasEditor/Private/Customizations/CameraParameterDetailsCustomizations.cpp
--- a/Engine/Plugins/Cameras/GameplayCameras/Source/GameplayCamerasEditor/Private/Customizations/CameraParameterDetailsCustomizations.cpp
+++ b/Engine/Plugins/Cameras/GameplayCameras/Source/GameplayCamerasEditor/Private/Customizations/CameraParameterDetailsCustomizations.cpp
@@ -96,6 +96,7 @@ void FCameraParameterDetailsCustomization::CustomizeHeader(TSharedRef<IPropertyH
 		StructProperty->CreatePropertyNameWidget()
 	]
 	.ValueContent()
+	.MinDesiredWidth(100.f)
 	.MaxDesiredWidth(MaxValueWidgetDesiredWidth)
 	.HAlign(HAlign_Fill)
 	[
@@ -127,7 +128,6 @@ void FCameraParameterDetailsCustomization::CustomizeHeader(TSharedRef<IPropertyH
 			.ContentPadding(1.f)
 			.ButtonStyle(FAppStyle::Get(), "SimpleButton")
 			.ToolTipText(LOCTEXT("SetVariable_ToolTip", "Selects a camera variable to drive this parameter"))
-			.IsEnabled(VariableClass != nullptr && !bIsExposedParameterVariable)
 			.ButtonContent()
 			[
 				SNew(SHorizontalBox)
@@ -151,6 +151,8 @@ void FCameraParameterDetailsCustomization::CustomizeHeader(TSharedRef<IPropertyH
 					[
 						SNew(STextBlock)
 						.Text(VariableInfoText)
+						.MinDesiredWidth(20.f)
+						.OverflowPolicy(ETextOverflowPolicy::Clip)
 					]
 				]
 				+SHorizontalBox::Slot()
diff --git a/Engine/Plugins/Cameras/GameplayCameras/Source/GameplayCamerasEditor/Private/Customizations/CameraRigAssetReferenceDetailsCustomization.cpp b/Engine/Plugins/Cameras/GameplayCameras/Source/GameplayCamerasEditor/Private/Customizations/CameraRigAssetReferenceDetailsCustomization.cpp
new file mode 100644
--- /dev/null
+++ b/Engine/Plugins/Cameras/GameplayCameras/Source/GameplayCamerasEditor/Private/Customizations/CameraRigAssetReferenceDetailsCustomization.cpp
@@ -0,0 +1,533 @@
+// Copyright Epic Games, Inc. All Rights Reserved.
+
+#include "Customizations/CameraRigAssetReferenceDetailsCustomization.h"
+
+#include "Core/CameraParameters.h"
+#include "Core/CameraRigAsset.h"
+#include "Core/CameraRigAssetReference.h"
+#include "GameplayCamerasDelegates.h"
+#include "IDetailChildrenBuilder.h"
+#include "IDetailGroup.h"
+#include "IPropertyTypeCustomization.h"
+#include "IPropertyUtilities.h"
+#include "IStructureDataProvider.h"
+#include "PropertyCustomizationHelpers.h"
+#include "Widgets/Images/SImage.h"
+#include "Widgets/Text/STextBlock.h"
+
+#define LOCTEXT_NAMESPACE "CameraRigAssetReferenceDetailsCustomization"
+
+namespace UE::Cameras
+{
+
+/**
+ * A class that holds information regarding a given camera rig parameter override, and how it should
+ * be shown in a details view panel.
+ */
+class FCameraRigParameterOverrideDetailRow 
+	: public IStructureDataProvider 
+	, public TSharedFromThis<FCameraRigParameterOverrideDetailRow>
+{
+public:
+
+	FCameraRigParameterOverrideDetailRow(TSharedPtr<IPropertyHandle> InCameraRigReferenceProperty, TSharedPtr<IPropertyUtilities> InPropertyUtilities);
+
+	void Initialize(UCameraRigInterfaceParameter* InInterfaceParameter);
+	void InitializeAsInvalid(const FString& InInterfaceParameterName, const FGuid& InInterfaceParameterGuid);
+
+	void BuildDetailPropertyRow(IDetailChildrenBuilder& StructBuilder);
+
+	const FString& GetDisplayName() const;
+	
+private:
+
+	template<typename ParameterOverrideType>
+	void InitializeValues();
+
+	template<typename ParameterOverrideType>
+	void BuildDetailPropertyRowImpl(IDetailChildrenBuilder& StructBuilder);
+
+	void AccessCameraRigReferences(TArray<FCameraRigAssetReference*>& OutCameraRigReferences) const;
+	void ModifyOuterObjects() const;
+
+	template<typename ParameterOverrideType, typename CameraParameterType = typename ParameterOverrideType::CameraParameterType>
+	void OnPropertyValueChanged();
+	template<typename ParameterOverrideType>
+	bool OnIsResetToDefaultVisible(TSharedPtr<IPropertyHandle> PropertyHandle);
+	template<typename ParameterOverrideType>
+	void OnResetToDefault(TSharedPtr<IPropertyHandle> PropertyHandle);
+
+	void OnRemoveInvalidOverride();
+
+public:
+
+	// IStructureDataProvider interface.
+	virtual bool IsValid() const override
+	{
+		return true; 
+	}
+	
+	virtual const UStruct* GetBaseStructure() const override
+	{
+		return ParameterType;
+	}
+	
+	virtual void GetInstances(TArray<TSharedPtr<FStructOnScope>>& OutInstances, const UStruct* ExpectedBaseStructure) const override
+	{
+		OutInstances = ParameterOverrideStructs;
+	}
+
+private:
+
+	TSharedPtr<IPropertyHandle> CameraRigReferenceProperty;
+	TSharedPtr<IPropertyUtilities> PropertyUtilities;
+
+	UCameraRigInterfaceParameter* InterfaceParameter = nullptr;
+	FString InvalidInterfaceParameterName;
+	FGuid InvalidInterfaceParameterGuid;
+
+	const UStruct* ParameterType = nullptr;
+	TArray<TSharedPtr<FStructOnScope>> ParameterOverrideStructs;
+	TSharedPtr<FStructOnScope> DefaultValue;
+};
+
+FCameraRigParameterOverrideDetailRow::FCameraRigParameterOverrideDetailRow(TSharedPtr<IPropertyHandle> InCameraRigReferenceProperty, TSharedPtr<IPropertyUtilities> InPropertyUtilities)
+	: CameraRigReferenceProperty(InCameraRigReferenceProperty)
+	, PropertyUtilities(InPropertyUtilities)
+{
+}
+
+void FCameraRigParameterOverrideDetailRow::Initialize(UCameraRigInterfaceParameter* InInterfaceParameter)
+{
+	InterfaceParameter = InInterfaceParameter;
+
+	if (InterfaceParameter && InterfaceParameter->PrivateVariable != nullptr)
+	{
+		switch (InterfaceParameter->PrivateVariable->GetVariableType())
+		{
+#define UE_CAMERA_VARIABLE_FOR_TYPE(ValueType, ValueName)\
+		case ECameraVariableType::ValueName:\
+			InitializeValues<F##ValueName##CameraRigParameterOverride>();\
+			break;
+			UE_CAMERA_VARIABLE_FOR_ALL_TYPES()
+#undef UE_CAMERA_VARIABLE_FOR_TYPE
+		}
+	}
+}
+
+void FCameraRigParameterOverrideDetailRow::InitializeAsInvalid(const FString& InInterfaceParameterName, const FGuid& InInterfaceParameterGuid)
+{
+	InvalidInterfaceParameterName = InInterfaceParameterName;
+	InvalidInterfaceParameterGuid = InInterfaceParameterGuid;
+}
+
+template<typename ParameterOverrideType>
+void FCameraRigParameterOverrideDetailRow::InitializeValues()
+{
+	using CameraParameterType = typename ParameterOverrideType::CameraParameterType;
+	using VariableAssetType = typename CameraParameterType::VariableAssetType;
+
+	// Get the camera parameter type.
+	ParameterType = CameraParameterType::StaticStruct();
+
+	// Get the default value of the camera parameter for this override.
+	const VariableAssetType* TypedPrivateVariable = CastChecked<VariableAssetType>(InterfaceParameter->PrivateVariable);
+	const CameraParameterType DefaultValueParameter(TypedPrivateVariable->GetDefaultValue());
+	DefaultValue = MakeShared<TStructOnScope<CameraParameterType>>(DefaultValueParameter);
+
+	// Add all the actual parameter override values, if any, or the default value, if not.
+	TArray<FCameraRigAssetReference*> CameraRigReferences;
+	AccessCameraRigReferences(CameraRigReferences);
+	for (FCameraRigAssetReference* CameraRigReference : CameraRigReferences)
+	{
+		check(CameraRigReference);
+
+		FCameraRigParameterOverrides& ParameterOverrides = CameraRigReference->GetParameterOverrides();
+		ParameterOverrideType* FoundOverride = ParameterOverrides.FindParameterOverride<ParameterOverrideType>(
+				InterfaceParameter->Guid);
+		if (FoundOverride)
+		{
+			ParameterOverrideStructs.Add(MakeShared<TStructOnScope<CameraParameterType>>(FoundOverride->Value));
+		}
+		else
+		{
+			ParameterOverrideStructs.Add(MakeShared<TStructOnScope<CameraParameterType>>(DefaultValueParameter));
+		}
+	}
+}
+
+void FCameraRigParameterOverrideDetailRow::BuildDetailPropertyRow(IDetailChildrenBuilder& StructBuilder)
+{
+	if (InterfaceParameter == nullptr)
+	{
+		// This is an old parameter override that isn't valid anymore (e.g. the camera rig's interface
+		// parameter was deleted).
+		FResetToDefaultOverride ResetToDefault = FResetToDefaultOverride::Create(
+				TAttribute<bool>(true),
+				FSimpleDelegate::CreateSP(this, &FCameraRigParameterOverrideDetailRow::OnRemoveInvalidOverride));
+
+		StructBuilder.AddCustomRow(FText::FromString(InvalidInterfaceParameterName))
+			.NameContent()
+			[
+				SNew(STextBlock)
+				.Font(IPropertyTypeCustomizationUtils::GetRegularFont())
+				.Text(FText::FromString(InvalidInterfaceParameterName))
+			]
+			.ValueContent()
+			[
+				SNew(SHorizontalBox)
+				+SHorizontalBox::Slot()
+				.AutoWidth()
+				[
+					SNew(SImage)
+					.Image(FAppStyle::GetBrush("Icons.Warning"))
+				]
+				+SHorizontalBox::Slot()
+				.FillWidth(1.f)
+				[
+					SNew(STextBlock)
+					.Font(IPropertyTypeCustomizationUtils::GetRegularFont())
+					.Text(LOCTEXT("InvalidParameterOverrideWarning", "No such parameter found."))
+				]
+			]
+			.OverrideResetToDefault(ResetToDefault);
+	}
+	else if (InterfaceParameter->PrivateVariable == nullptr)
+	{
+		// The camera rig wasn't built, show a message.
+		StructBuilder.AddCustomRow(FText::FromString(InterfaceParameter->InterfaceParameterName))
+			.NameContent()
+			[
+				SNew(STextBlock)
+				.Font(IPropertyTypeCustomizationUtils::GetRegularFont())
+				.Text(FText::FromString(InterfaceParameter->InterfaceParameterName))
+			]
+			.ValueContent()
+			[
+				SNew(STextBlock)
+				.Font(IPropertyTypeCustomizationUtils::GetRegularFont())
+				.Text(LOCTEXT("UnbuiltInnerCameraRigWarning", "Please build the child camera rig."))
+			];
+	}
+	else
+	{
+		switch (InterfaceParameter->PrivateVariable->GetVariableType())
+		{
+#define UE_CAMERA_VARIABLE_FOR_TYPE(ValueType, ValueName)\
+			case ECameraVariableType::ValueName:\
+				BuildDetailPropertyRowImpl<F##ValueName##CameraRigParameterOverride>(StructBuilder);\
+				break;
+			UE_CAMERA_VARIABLE_FOR_ALL_TYPES()
+#undef UE_CAMERA_VARIABLE_FOR_TYPE
+		}
+	}
+}
+
+const FString& FCameraRigParameterOverrideDetailRow::GetDisplayName() const
+{
+	if (InterfaceParameter)
+	{
+		return InterfaceParameter->InterfaceParameterName;
+	}
+	else
+	{
+		return InvalidInterfaceParameterName;
+	}
+}
+
+void FCameraRigParameterOverrideDetailRow::AccessCameraRigReferences(TArray<FCameraRigAssetReference*>& OutCameraRigReferences) const
+{
+	TArray<void*> RawData;
+	CameraRigReferenceProperty->AccessRawData(RawData);
+	for (int32 Index = 0; Index < RawData.Num(); ++Index)
+	{
+		OutCameraRigReferences.Add(static_cast<FCameraRigAssetReference*>(RawData[Index]));
+	}
+}
+
+void FCameraRigParameterOverrideDetailRow::ModifyOuterObjects() const
+{
+	TArray<UObject*> OuterObjects;
+	CameraRigReferenceProperty->GetOuterObjects(OuterObjects);
+	for (UObject* OuterObject : OuterObjects)
+	{
+		OuterObject->Modify();
+	}
+}
+
+template<typename ParameterOverrideType>
+void FCameraRigParameterOverrideDetailRow::BuildDetailPropertyRowImpl(IDetailChildrenBuilder& StructBuilder)
+{
+	using CameraParameterType = typename ParameterOverrideType::CameraParameterType;
+
+	// Add a new property row to the detail view. The row shows our copy of a camera parameter that
+	// corresponds to the sort of camera parameter needed to override a given camera rig parameter.
+	IDetailPropertyRow* ParameterOverrideRow = StructBuilder.AddExternalStructure(SharedThis(this));
+
+	ParameterOverrideRow->DisplayName(FText::FromString(InterfaceParameter->InterfaceParameterName));
+
+	// Set up callbacks so that when the user edits the camera parameter, we replicate the edit onto
+	// the parameter overrides list (by adding/removing the override, or updating its value).
+	ParameterOverrideRow->GetPropertyHandle()->SetOnPropertyValueChanged(FSimpleDelegate::CreateSP(
+				this, &FCameraRigParameterOverrideDetailRow::OnPropertyValueChanged<ParameterOverrideType>));
+	ParameterOverrideRow->GetPropertyHandle()->SetOnChildPropertyValueChanged(FSimpleDelegate::CreateSP(
+				this, &FCameraRigParameterOverrideDetailRow::OnPropertyValueChanged<ParameterOverrideType>));
+
+	FResetToDefaultOverride ResetToDefault = FResetToDefaultOverride::Create(
+			FIsResetToDefaultVisible::CreateSP(
+				this, &FCameraRigParameterOverrideDetailRow::OnIsResetToDefaultVisible<ParameterOverrideType>),
+			FResetToDefaultHandler::CreateSP(
+				this, &FCameraRigParameterOverrideDetailRow::OnResetToDefault<ParameterOverrideType>));
+	ParameterOverrideRow->OverrideResetToDefault(ResetToDefault);
+}
+
+template<typename ParameterOverrideType, typename CameraParameterType>
+void FCameraRigParameterOverrideDetailRow::OnPropertyValueChanged()
+{
+	const CameraParameterType* TypedDefaultValue = reinterpret_cast<CameraParameterType*>(DefaultValue->GetStructMemory());
+
+	TArray<FCameraRigAssetReference*> CameraRigReferences;
+	AccessCameraRigReferences(CameraRigReferences);
+	check(CameraRigReferences.Num() == ParameterOverrideStructs.Num());
+
+	ModifyOuterObjects();
+
+	for (int32 Index = 0; Index < CameraRigReferences.Num(); ++Index)
+	{
+		FCameraRigAssetReference* CameraRigReference(CameraRigReferences[Index]);
+		TSharedPtr<FStructOnScope> ParameterOverrideStruct(ParameterOverrideStructs[Index]);
+
+		const CameraParameterType* TypedOverrideParameter = reinterpret_cast<CameraParameterType*>(
+				ParameterOverrideStruct->GetStructMemory());
+
+		FCameraRigParameterOverrides& ParameterOverrides = CameraRigReference->GetParameterOverrides();
+		const bool bEqualValues = CameraParameterValueEquals<typename CameraParameterType::ValueType>(
+				TypedOverrideParameter->Value, TypedDefaultValue->Value);
+		if (bEqualValues && TypedOverrideParameter->Variable == nullptr)
+		{
+			ParameterOverrides.RemoveParameterOverride<ParameterOverrideType>(InterfaceParameter->Guid);
+		}
+		else
+		{
+			ParameterOverrideType& Override = ParameterOverrides.FindOrAddParameterOverride<ParameterOverrideType>(InterfaceParameter);
+			Override.Value = *TypedOverrideParameter;
+		}
+	}
+}
+
+template<typename ParameterOverrideType>
+bool FCameraRigParameterOverrideDetailRow::OnIsResetToDefaultVisible(TSharedPtr<IPropertyHandle> PropertyHandle)
+{
+	using CameraParameterType = typename ParameterOverrideType::CameraParameterType;
+
+	const CameraParameterType* TypedDefaultValue = reinterpret_cast<CameraParameterType*>(DefaultValue->GetStructMemory());
+
+	for (TSharedPtr<FStructOnScope> ParameterOverrideStruct : ParameterOverrideStructs)
+	{
+		const CameraParameterType* TypedOverrideParameter = reinterpret_cast<CameraParameterType*>(ParameterOverrideStruct->GetStructMemory());
+		if (TypedOverrideParameter->Variable != nullptr ||
+				!CameraParameterValueEquals<typename CameraParameterType::ValueType>(
+					TypedOverrideParameter->Value, TypedDefaultValue->Value))
+		{
+			return true;
+		}
+	}
+
+	return false;
+}
+
+template<typename ParameterOverrideType>
+void FCameraRigParameterOverrideDetailRow::OnResetToDefault(TSharedPtr<IPropertyHandle> PropertyHandle)
+{
+	using CameraParameterType = typename ParameterOverrideType::CameraParameterType;
+
+	const CameraParameterType* TypedDefaultValue = reinterpret_cast<CameraParameterType*>(DefaultValue->GetStructMemory());
+
+	ModifyOuterObjects();
+
+	for (TSharedPtr<FStructOnScope> ParameterOverrideStruct : ParameterOverrideStructs)
+	{
+		CameraParameterType* TypedOverrideParameter = reinterpret_cast<CameraParameterType*>(ParameterOverrideStruct->GetStructMemory());
+
+		TypedOverrideParameter->Value = TypedDefaultValue->Value;
+		TypedOverrideParameter->Variable = nullptr;
+	}
+}
+
+void FCameraRigParameterOverrideDetailRow::OnRemoveInvalidOverride()
+{
+	if (InvalidInterfaceParameterGuid.IsValid())
+	{
+		ModifyOuterObjects();
+
+		TArray<FCameraRigAssetReference*> CameraRigReferences;
+		AccessCameraRigReferences(CameraRigReferences);
+
+		for (FCameraRigAssetReference* CameraRigReference : CameraRigReferences)
+		{
+			FCameraRigParameterOverrides& ParameterOverrides = CameraRigReference->GetParameterOverrides();
+#define UE_CAMERA_VARIABLE_FOR_TYPE(ValueType, ValueName)\
+			ParameterOverrides.RemoveParameterOverride<F##ValueName##CameraRigParameterOverride>(InvalidInterfaceParameterGuid);
+			UE_CAMERA_VARIABLE_FOR_ALL_TYPES()
+#undef UE_CAMERA_VARIABLE_FOR_TYPE
+		}
+
+		if (PropertyUtilities)
+		{
+			PropertyUtilities->RequestForceRefresh();
+		}
+	}
+}
+
+TSharedRef<IPropertyTypeCustomization> FCameraRigAssetReferenceDetailsCustomization::MakeInstance()
+{
+	return MakeShared<FCameraRigAssetReferenceDetailsCustomization>();
+}
+
+void FCameraRigAssetReferenceDetailsCustomization::CustomizeHeader(TSharedRef<IPropertyHandle> StructPropertyHandle, FDetailWidgetRow& HeaderRow, IPropertyTypeCustomizationUtils& StructCustomizationUtils)
+{
+	// Remember a few things.
+	CameraRigReferenceProperty = StructPropertyHandle;
+	PropertyUtilities = StructCustomizationUtils.GetPropertyUtilities();
+
+	// Display the camera rig picker widget.
+	TSharedPtr<IPropertyHandle> CameraRigProperty = StructPropertyHandle->GetChildHandle(
+			GET_MEMBER_NAME_CHECKED(FCameraRigAssetReference, CameraRig));
+
+	HeaderRow
+	.NameContent()
+	[
+		CameraRigProperty->CreatePropertyNameWidget()
+	]
+	.ValueContent()
+	[
+		CameraRigProperty->CreatePropertyValueWidgetWithCustomization(nullptr)
+	]
+	.ShouldAutoExpand(true);
+	
+	// Setup a callback so that when a different camera rig is picked, we refresh the list
+	// of parameter overrides.
+	CameraRigProperty->SetOnPropertyValueChanged(
+			FSimpleDelegate::CreateSP(this, &FCameraRigAssetReferenceDetailsCustomization::OnCameraRigChanged));
+
+	// Setup a similar callback for when the selected camera rig get (re)built.
+	FGameplayCamerasDelegates::OnCameraRigAssetBuilt().AddSP(this, &FCameraRigAssetReferenceDetailsCustomization::OnCameraRigBuilt);
+}
+
+void FCameraRigAssetReferenceDetailsCustomization::CustomizeChildren(TSharedRef<IPropertyHandle> StructPropertyHandle, IDetailChildrenBuilder& StructBuilder, IPropertyTypeCustomizationUtils& StructCustomizationUtils)
+{
+	// Update our list of parameter overrides and add each of them as a detail view row
+	// below the camera rig picker.
+	UpdateParameterOverrides(nullptr, false);
+	BuildParameterOverrideRows(StructBuilder);
+}
+
+void FCameraRigAssetReferenceDetailsCustomization::OnCameraRigChanged()
+{
+	UpdateParameterOverrides(nullptr, true);
+}
+
+void FCameraRigAssetReferenceDetailsCustomization::OnCameraRigBuilt(UCameraRigAsset* CameraRig, FCameraBuildLog& BuildLog)
+{
+	UpdateParameterOverrides(CameraRig, true);
+}
+
+void FCameraRigAssetReferenceDetailsCustomization::UpdateParameterOverrides(const UCameraRigAsset* CameraRigToUpdate, bool bRequestRefresh)
+{
+	// Get all the camera rig references we are editing.
+	TArray<void*> RawData;
+	CameraRigReferenceProperty->AccessRawData(RawData);
+
+	TArray<FCameraRigAssetReference*> CameraRigReferences;
+	for (int32 Index = 0; Index < RawData.Num(); Index++)
+	{
+		FCameraRigAssetReference* CameraRigReference = static_cast<FCameraRigAssetReference*>(RawData[Index]);
+		check(CameraRigReference);
+		CameraRigReferences.Add(CameraRigReference);
+	}
+
+	// Update the parameter overrides for the given camera rig, or for all the ones we're editing if none
+	// was specified. As we go, determine if we're editing parameter overrides for the same camera rig
+	// across all edited references, or if we have a mix.
+	bool bIsEditingNullCameraRig = false;
+	TSet<const UCameraRigAsset*> UsedCameraRigs;
+	for (FCameraRigAssetReference* CameraRigReference : CameraRigReferences)
+	{
+		const UCameraRigAsset* CameraRig = CameraRigReference->GetCameraRig();
+		if (CameraRigToUpdate == nullptr || CameraRigToUpdate == CameraRig)
+		{
+			CameraRigReference->UpdateParameterOverrides();
+		}
+
+		if (CameraRig)
+		{
+			UsedCameraRigs.Add(CameraRig);
+		}
+		else
+		{
+			bIsEditingNullCameraRig = true;
+		}
+	}
+
+	// Re-create our list of parameter overrides.
+	ParameterOverrideRows.Reset();
+	TSet<FGuid> UsedInterfaceParameterGuids;
+	if (UsedCameraRigs.Num() == 1 && !bIsEditingNullCameraRig)
+	{
+		// Add rows for all of the camera rig's parameters.
+		const UCameraRigAsset* CameraRig = UsedCameraRigs.Array()[0];
+		for (UCameraRigInterfaceParameter* InterfaceParameter : CameraRig->Interface.InterfaceParameters)
+		{
+			TSharedPtr<FCameraRigParameterOverrideDetailRow> Row = MakeShared<FCameraRigParameterOverrideDetailRow>(
+					CameraRigReferenceProperty, PropertyUtilities);
+			Row->Initialize(InterfaceParameter);
+			ParameterOverrideRows.Add(Row);
+			UsedInterfaceParameterGuids.Add(InterfaceParameter->Guid);
+		}
+
+		// Add any left-over/invalid parameter overrides.
+		for (FCameraRigAssetReference* CameraRigReference : CameraRigReferences)
+		{
+			FCameraRigParameterOverrides& ParameterOverrides = CameraRigReference->GetParameterOverrides();
+#define UE_CAMERA_VARIABLE_FOR_TYPE(ValueType, ValueName)\
+			for (F##ValueName##CameraRigParameterOverride& ParameterOverride : ParameterOverrides.Get##ValueName##Overrides())\
+			{\
+				if (!UsedInterfaceParameterGuids.Contains(ParameterOverride.InterfaceParameterGuid))\
+				{\
+					TSharedPtr<FCameraRigParameterOverrideDetailRow> InvalidRow = MakeShared<FCameraRigParameterOverrideDetailRow>(\
+							CameraRigReferenceProperty, PropertyUtilities);\
+					InvalidRow->InitializeAsInvalid(ParameterOverride.InterfaceParameterName, ParameterOverride.InterfaceParameterGuid);\
+					ParameterOverrideRows.Add(InvalidRow);\
+				}\
+			}
+			UE_CAMERA_VARIABLE_FOR_ALL_TYPES()
+#undef UE_CAMERA_VARIABLE_FOR_TYPE
+		}
+
+		// Sort rows by name.
+		ParameterOverrideRows.StableSort(
+				[](TSharedPtr<FCameraRigParameterOverrideDetailRow> A,
+					TSharedPtr<FCameraRigParameterOverrideDetailRow> B)
+				{
+					return A->GetDisplayName().Compare(B->GetDisplayName()) < 0;
+				});
+	}
+
+	if (bRequestRefresh && PropertyUtilities)
+	{
+		PropertyUtilities->RequestForceRefresh();
+	}
+}
+
+void FCameraRigAssetReferenceDetailsCustomization::BuildParameterOverrideRows(IDetailChildrenBuilder& StructBuilder)
+{
+	for (TSharedPtr<FCameraRigParameterOverrideDetailRow> Row : ParameterOverrideRows)
+	{
+		Row->BuildDetailPropertyRow(StructBuilder);
+	}
+}
+
+}  // namespace UE::Cameras
+
+#undef LOCTEXT_NAMESPACE
+
diff --git a/Engine/Plugins/Cameras/GameplayCameras/Source/GameplayCamerasEditor/Private/Customizations/CameraRigCameraNodeDetailsCustomization.cpp b/Engine/Plugins/Cameras/GameplayCameras/Source/GameplayCamerasEditor/Private/Customizations/CameraRigCameraNodeDetailsCustomization.cpp
deleted file mode 100644
diff --git a/Engine/Plugins/Cameras/GameplayCameras/Source/GameplayCamerasEditor/Private/Customizations/CameraRigPtrDetailsCustomization.cpp b/Engine/Plugins/Cameras/GameplayCameras/Source/GameplayCamerasEditor/Private/Customizations/CameraRigPtrDetailsCustomization.cpp
--- a/Engine/Plugins/Cameras/GameplayCameras/Source/GameplayCamerasEditor/Private/Customizations/CameraRigPtrDetailsCustomization.cpp
+++ b/Engine/Plugins/Cameras/GameplayCameras/Source/GameplayCamerasEditor/Private/Customizations/CameraRigPtrDetailsCustomization.cpp
@@ -6,6 +6,7 @@
 #include "Core/CameraAsset.h"
 #include "Core/CameraDirector.h"
 #include "Core/CameraRigAsset.h"
+#include "Core/CameraRigAssetReference.h"
 #include "DetailCategoryBuilder.h"
 #include "DetailWidgetRow.h"
 #include "Helpers/CameraAssetReferenceGatherer.h"
@@ -29,19 +30,17 @@ void FCameraRigPtrDetailsCustomization::CustomizeHeader(TSharedRef<IPropertyHand
 {
 	CameraRigPropertyHandle = StructPropertyHandle;
 
-	FProperty* StructProperty = StructPropertyHandle->GetProperty();
-	const bool bUseCameraRigPicker = StructProperty->GetBoolMetaData("UseCameraRigPicker");
-	const bool bUseCameraDirectorRigPicker = StructProperty->GetBoolMetaData("UseCameraDirectorRigPicker");
+	EPickerMode PickerMode = DeterminePickerMode();
 
 	TSharedPtr<SWidget> ValueContentWidget;
-	if (bUseCameraRigPicker || bUseCameraDirectorRigPicker)
+	if (PickerMode == EPickerMode::CameraRigPicker || PickerMode == EPickerMode::CameraDirectorRigPicker)
 	{
 		FOnGetContent OnGetComboMenuContent;
-		if (bUseCameraRigPicker)
+		if (PickerMode == EPickerMode::CameraRigPicker)
 		{
 			OnGetComboMenuContent = FOnGetContent::CreateSP(this, &FCameraRigPtrDetailsCustomization::OnBuildCameraRigNamePicker);
 		}
-		else if (bUseCameraDirectorRigPicker)
+		else if (PickerMode == EPickerMode::CameraDirectorRigPicker)
 		{
 			OnGetComboMenuContent = FOnGetContent::CreateSP(this, &FCameraRigPtrDetailsCustomization::OnBuildCameraDirectorRigNamePicker);
 		}
@@ -81,6 +80,28 @@ void FCameraRigPtrDetailsCustomization::CustomizeChildren(TSharedRef<IPropertyHa
 {
 }
 
+FCameraRigPtrDetailsCustomization::EPickerMode FCameraRigPtrDetailsCustomization::DeterminePickerMode()
+{
+	FProperty* CameraRigProperty = CameraRigPropertyHandle->GetProperty();
+	FField* MetaDataField = CameraRigProperty;
+	if (CameraRigProperty->GetOwnerStruct() == FCameraRigAssetReference::StaticStruct())
+	{
+		TSharedPtr<IPropertyHandle> CameraRigReferencePropertyHandle = CameraRigPropertyHandle->GetParentHandle();
+		ensure(CameraRigReferencePropertyHandle);
+		MetaDataField = CameraRigReferencePropertyHandle->GetProperty();
+	}
+
+	if (MetaDataField->GetBoolMetaData("UseCameraDirectorRigPicker"))
+	{
+		return EPickerMode::CameraDirectorRigPicker;
+	}
+	if (MetaDataField->GetBoolMetaData("UseCameraRigPicker"))
+	{
+		return EPickerMode::CameraRigPicker;
+	}
+	return EPickerMode::StandardPicker;
+}
+
 FText FCameraRigPtrDetailsCustomization::OnGetComboText() const
 {
 	UObject* Value;
diff --git a/Engine/Plugins/Cameras/GameplayCameras/Source/GameplayCamerasEditor/Private/Editors/CameraNodeGraphNode.cpp b/Engine/Plugins/Cameras/GameplayCameras/Source/GameplayCamerasEditor/Private/Editors/CameraNodeGraphNode.cpp
--- a/Engine/Plugins/Cameras/GameplayCameras/Source/GameplayCamerasEditor/Private/Editors/CameraNodeGraphNode.cpp
+++ b/Engine/Plugins/Cameras/GameplayCameras/Source/GameplayCamerasEditor/Private/Editors/CameraNodeGraphNode.cpp
@@ -5,7 +5,6 @@
 #include "Core/CameraParameters.h"
 #include "EdGraph/EdGraphPin.h"
 #include "Editors/CameraNodeGraphSchema.h"
-#include "ScopedTransaction.h"
 #include "ToolMenus.h"
 
 UCameraNodeGraphNode::UCameraNodeGraphNode(const FObjectInitializer& ObjInit)
@@ -15,7 +14,7 @@ UCameraNodeGraphNode::UCameraNodeGraphNode(const FObjectInitializer& ObjInit)
 
 void UCameraNodeGraphNode::AllocateDefaultPins()
 {
-	UObjectTreeGraphNode::AllocateDefaultPins();
+	Super::AllocateDefaultPins();
 
 	// Add extra input pins for any camera parameter.
 	UClass* CameraNodeClass = GetObject()->GetClass();
@@ -59,19 +58,12 @@ UEdGraphPin* UCameraNodeGraphNode::GetPinForCameraParameterProperty(const FName&
 	return nullptr;
 }
 
-FStructProperty* UCameraNodeGraphNode::GetCameraParameterPropertyForPin(const UEdGraphPin* InPin) const
+FName UCameraNodeGraphNode::GetCameraParameterPropertyForPin(const UEdGraphPin* InPin) const
 {
-	UObject* CameraNode = GetObject();
-	if (!ensure(CameraNode))
+	if (InPin->PinType.PinCategory == UCameraNodeGraphSchema::PC_CameraParameter)
 	{
-		return nullptr;
+		return InPin->GetFName();
 	}
-	if (InPin->PinType.PinCategory != UCameraNodeGraphSchema::PC_CameraParameter)
-	{
-		return nullptr;
-	}
-
-	UClass* CameraNodeClass = CameraNode->GetClass();
-	return CastField<FStructProperty>(CameraNodeClass->FindPropertyByName(InPin->GetFName()));
+	return NAME_None;
 }
 
diff --git a/Engine/Plugins/Cameras/GameplayCameras/Source/GameplayCamerasEditor/Private/Editors/CameraNodeGraphSchema.cpp b/Engine/Plugins/Cameras/GameplayCameras/Source/GameplayCamerasEditor/Private/Editors/CameraNodeGraphSchema.cpp
--- a/Engine/Plugins/Cameras/GameplayCameras/Source/GameplayCamerasEditor/Private/Editors/CameraNodeGraphSchema.cpp
+++ b/Engine/Plugins/Cameras/GameplayCameras/Source/GameplayCamerasEditor/Private/Editors/CameraNodeGraphSchema.cpp
@@ -8,10 +8,12 @@
 #include "EdGraph/EdGraphPin.h"
 #include "Editors/CameraNodeGraphNode.h"
 #include "Editors/CameraRigInterfaceParameterGraphNode.h"
+#include "Editors/CameraRigNodeGraphNode.h"
 #include "Editors/ObjectTreeGraph.h"
 #include "Editors/ObjectTreeGraphConfig.h"
 #include "Editors/ObjectTreeGraphNode.h"
 #include "GameplayCamerasEditorSettings.h"
+#include "Nodes/Common/CameraRigCameraNode.h"
 
 #include "ScopedTransaction.h"
 
@@ -42,6 +44,8 @@ FObjectTreeGraphConfig UCameraNodeGraphSchema::BuildGraphConfig() const
 		.StripDisplayNameSuffix(TEXT("Camera Node"))
 		.CreateCategoryMetaData(TEXT("CameraNodeCategories"))
 		.GraphNodeClass(UCameraNodeGraphNode::StaticClass());
+	GraphConfig.ObjectClassConfigs.Emplace(UCameraRigCameraNode::StaticClass())
+		.GraphNodeClass(UCameraRigNodeGraphNode::StaticClass());
 	GraphConfig.ObjectClassConfigs.Emplace(UCameraRigInterfaceParameter::StaticClass())
 		.SelfPinName(NAME_None)  // No self pin name, we just want the title
 		.CanCreateNew(false)
@@ -121,7 +125,8 @@ void UCameraNodeGraphSchema::GetGraphContextActions(FGraphContextMenuBuilder& Co
 		if (DraggedPin->PinType.PinCategory == PC_CameraParameter)
 		{
 			UCameraNodeGraphNode* CameraNodeNode = Cast<UCameraNodeGraphNode>(DraggedPin->GetOwningNode());
-			FStructProperty* StructProperty = CameraNodeNode->GetCameraParameterPropertyForPin(DraggedPin);
+			const FName PropertyName = CameraNodeNode->GetCameraParameterPropertyForPin(DraggedPin);
+			ensure(PropertyName != NAME_None);
 
 			TSharedRef<FCameraNodeGraphSchemaAction_NewInterfaceParameterNode> Action = 
 				MakeShared<FCameraNodeGraphSchemaAction_NewInterfaceParameterNode>(
@@ -129,7 +134,7 @@ void UCameraNodeGraphSchema::GetGraphContextActions(FGraphContextMenuBuilder& Co
 						LOCTEXT("NewInterfaceParameterAction", "Camera Rig Parameter"),
 						LOCTEXT("NewInterfaceParameterActionToolTip", "Exposes this parameter on the camera rig"));
 			Action->Target = Cast<UCameraNode>(CameraNodeNode->GetObject());
-			Action->TargetPropertyName = StructProperty->GetFName();
+			Action->TargetPropertyName = PropertyName;
 			ContextMenuBuilder.AddAction(StaticCastSharedPtr<FEdGraphSchemaAction>(Action.ToSharedPtr()));
 
 			return;
@@ -206,8 +211,8 @@ bool UCameraNodeGraphSchema::OnApplyConnection(UEdGraphPin* A, UEdGraphPin* B, F
 	{
 		return false;
 	}
-	FStructProperty* StructProperty = CameraNodeNode->GetCameraParameterPropertyForPin(CameraParameterPin);
-	if (!StructProperty)
+	const FName PropertyName = CameraNodeNode->GetCameraParameterPropertyForPin(CameraParameterPin);
+	if (PropertyName.IsNone())
 	{
 		return false;
 	}
@@ -216,10 +221,10 @@ bool UCameraNodeGraphSchema::OnApplyConnection(UEdGraphPin* A, UEdGraphPin* B, F
 	RigParameter->Modify();
 
 	RigParameter->Target = CameraNode;
-	RigParameter->TargetPropertyName = StructProperty->GetFName();
+	RigParameter->TargetPropertyName = PropertyName;
 	if (RigParameter->InterfaceParameterName.IsEmpty())
 	{
-		RigParameter->InterfaceParameterName = StructProperty->GetName();
+		RigParameter->InterfaceParameterName = PropertyName.ToString();
 	}
 
 	return true;
diff --git a/Engine/Plugins/Cameras/GameplayCameras/Source/GameplayCamerasEditor/Private/Editors/CameraRigNodeGraphNode.cpp b/Engine/Plugins/Cameras/GameplayCameras/Source/GameplayCamerasEditor/Private/Editors/CameraRigNodeGraphNode.cpp
new file mode 100644
--- /dev/null
+++ b/Engine/Plugins/Cameras/GameplayCameras/Source/GameplayCamerasEditor/Private/Editors/CameraRigNodeGraphNode.cpp
@@ -0,0 +1,31 @@
+// Copyright Epic Games, Inc. All Rights Reserved.
+
+#include "Editors/CameraRigNodeGraphNode.h"
+
+#include "Core/CameraRigAsset.h"
+#include "Editors/CameraNodeGraphSchema.h"
+#include "Nodes/Common/CameraRigCameraNode.h"
+
+UCameraRigNodeGraphNode::UCameraRigNodeGraphNode(const FObjectInitializer& ObjInit)
+	: Super(ObjInit)
+{
+}
+
+void UCameraRigNodeGraphNode::AllocateDefaultPins()
+{
+	Super::AllocateDefaultPins();
+
+	// Add extra input pins for the inner camera rig's exposed parameters.
+	UCameraRigCameraNode* CameraRigNode = Cast<UCameraRigCameraNode>(GetObject());
+	if (ensure(CameraRigNode) && CameraRigNode->CameraRigReference.IsValid())
+	{
+		FCameraRigInterface& CameraRigInterface = CameraRigNode->CameraRigReference.GetCameraRig()->Interface;
+		for (UCameraRigInterfaceParameter* InterfaceParameter : CameraRigInterface.InterfaceParameters)
+		{
+			FEdGraphPinType PinType;
+			PinType.PinCategory = UCameraNodeGraphSchema::PC_CameraParameter;
+			CreatePin(EGPD_Input, PinType, *InterfaceParameter->InterfaceParameterName);
+		}
+	}
+}
+
diff --git a/Engine/Plugins/Cameras/GameplayCameras/Source/GameplayCamerasEditor/Private/GameplayCamerasEditorModule.cpp b/Engine/Plugins/Cameras/GameplayCameras/Source/GameplayCamerasEditor/Private/GameplayCamerasEditorModule.cpp
--- a/Engine/Plugins/Cameras/GameplayCameras/Source/GameplayCamerasEditor/Private/GameplayCamerasEditorModule.cpp
+++ b/Engine/Plugins/Cameras/GameplayCameras/Source/GameplayCamerasEditor/Private/GameplayCamerasEditorModule.cpp
@@ -13,7 +13,7 @@
 #include "Commands/CameraVariableCollectionEditorCommands.h"
 #include "Commands/GameplayCamerasDebuggerCommands.h"
 #include "Customizations/CameraParameterDetailsCustomizations.h"
-#include "Customizations/CameraRigCameraNodeDetailsCustomization.h"
+#include "Customizations/CameraRigAssetReferenceDetailsCustomization.h"
 #include "Customizations/CameraRigPtrDetailsCustomization.h"
 #include "Customizations/CameraProxyTableDetailsCustomization.h"
 #include "Customizations/SingleCameraDirectorDetailsCustomization.h"
@@ -472,15 +472,12 @@ private:
 
 		FCameraParameterDetailsCustomization::Register(PropertyEditorModule);
 
-		PropertyEditorModule.RegisterCustomClassLayout("CameraRigCameraNode", FOnGetDetailCustomizationInstance::CreateStatic(
-					&FCameraRigCameraNodeDetailsCustomization::MakeInstance));
-		PropertyEditorModule.RegisterCustomClassLayout("SingleCameraDirector", FOnGetDetailCustomizationInstance::CreateStatic(
-					&FSingleCameraDirectorDetailsCustomization::MakeInstance));
-
 		PropertyEditorModule.RegisterCustomPropertyTypeLayout("CameraRigProxyTableEntry", FOnGetPropertyTypeCustomizationInstance::CreateStatic(
 					&FCameraProxyTableEntryDetailsCustomization::MakeInstance));
 		PropertyEditorModule.RegisterCustomPropertyTypeLayout("CameraRigAsset", FOnGetPropertyTypeCustomizationInstance::CreateStatic(
 					&FCameraRigPtrDetailsCustomization::MakeInstance));
+		PropertyEditorModule.RegisterCustomPropertyTypeLayout("CameraRigAssetReference", FOnGetPropertyTypeCustomizationInstance::CreateStatic(
+					&FCameraRigAssetReferenceDetailsCustomization::MakeInstance));
 	}
 
 	void UnregisterDetailsCustomizations()
@@ -492,8 +489,10 @@ private:
 		if (PropertyEditorModule)
 		{
 			FCameraParameterDetailsCustomization::Unregister(*PropertyEditorModule);
+
 			PropertyEditorModule->UnregisterCustomPropertyTypeLayout("CameraRigProxyTableEntry");
-			PropertyEditorModule->UnregisterCustomClassLayout("SingleCameraDirector");
+			PropertyEditorModule->UnregisterCustomPropertyTypeLayout("CameraRigAsset");
+			PropertyEditorModule->UnregisterCustomPropertyTypeLayout("CameraRigAssetReference");
 		}
 	}
 
diff --git a/Engine/Plugins/Cameras/GameplayCameras/Source/GameplayCamerasEditor/Public/Customizations/CameraRigAssetReferenceDetailsCustomization.h b/Engine/Plugins/Cameras/GameplayCameras/Source/GameplayCamerasEditor/Public/Customizations/CameraRigAssetReferenceDetailsCustomization.h
new file mode 100644
--- /dev/null
+++ b/Engine/Plugins/Cameras/GameplayCameras/Source/GameplayCamerasEditor/Public/Customizations/CameraRigAssetReferenceDetailsCustomization.h
@@ -0,0 +1,48 @@
+// Copyright Epic Games, Inc. All Rights Reserved.
+
+#pragma once
+
+#include "IPropertyTypeCustomization.h"
+
+class FStructOnScope;
+class IStructureDataProvider;
+class UCameraRigAsset;
+class UCameraRigInterfaceParameter;
+struct FCameraRigAssetReference;
+
+namespace UE::Cameras
+{
+
+class FCameraBuildLog;
+class FCameraRigParameterOverrideDetailRow;
+
+class FCameraRigAssetReferenceDetailsCustomization : public IPropertyTypeCustomization
+{
+public:
+
+	static TSharedRef<IPropertyTypeCustomization> MakeInstance();
+
+public:
+
+	// IPropertyTypeCustomization interface.
+	virtual void CustomizeHeader(TSharedRef<IPropertyHandle> StructPropertyHandle, FDetailWidgetRow& HeaderRow, IPropertyTypeCustomizationUtils& StructCustomizationUtils) override;
+	virtual void CustomizeChildren(TSharedRef<IPropertyHandle> StructPropertyHandle, IDetailChildrenBuilder& StructBuilder, IPropertyTypeCustomizationUtils& StructCustomizationUtils) override;
+
+private:
+
+	void OnCameraRigChanged();
+	void OnCameraRigBuilt(UCameraRigAsset* CameraRig, FCameraBuildLog& BuildLog);
+
+	void UpdateParameterOverrides(const UCameraRigAsset* CameraRigToUpdate, bool bRequestRefresh);
+
+	void BuildParameterOverrideRows(IDetailChildrenBuilder& StructBuilder);
+
+private:
+
+	TSharedPtr<IPropertyHandle> CameraRigReferenceProperty;
+	TSharedPtr<IPropertyUtilities> PropertyUtilities;
+
+	TArray<TSharedPtr<FCameraRigParameterOverrideDetailRow>> ParameterOverrideRows;
+};
+
+}  // namespace UE::Cameras
diff --git a/Engine/Plugins/Cameras/GameplayCameras/Source/GameplayCamerasEditor/Public/Customizations/CameraRigCameraNodeDetailsCustomization.h b/Engine/Plugins/Cameras/GameplayCameras/Source/GameplayCamerasEditor/Public/Customizations/CameraRigCameraNodeDetailsCustomization.h
deleted file mode 100644
diff --git a/Engine/Plugins/Cameras/GameplayCameras/Source/GameplayCamerasEditor/Public/Customizations/CameraRigPtrDetailsCustomization.h b/Engine/Plugins/Cameras/GameplayCameras/Source/GameplayCamerasEditor/Public/Customizations/CameraRigPtrDetailsCustomization.h
--- a/Engine/Plugins/Cameras/GameplayCameras/Source/GameplayCamerasEditor/Public/Customizations/CameraRigPtrDetailsCustomization.h
+++ b/Engine/Plugins/Cameras/GameplayCameras/Source/GameplayCamerasEditor/Public/Customizations/CameraRigPtrDetailsCustomization.h
@@ -31,6 +31,10 @@ public:
 
 private:
 
+	enum class EPickerMode { StandardPicker, CameraRigPicker, CameraDirectorRigPicker };
+
+	EPickerMode DeterminePickerMode();
+
 	FText OnGetComboText() const;
 	FText OnGetComboToolTipText() const;
 	TSharedRef<SWidget> OnBuildCameraRigNamePicker();
diff --git a/Engine/Plugins/Cameras/GameplayCameras/Source/GameplayCamerasEditor/Public/Editors/CameraNodeGraphNode.h b/Engine/Plugins/Cameras/GameplayCameras/Source/GameplayCamerasEditor/Public/Editors/CameraNodeGraphNode.h
--- a/Engine/Plugins/Cameras/GameplayCameras/Source/GameplayCamerasEditor/Public/Editors/CameraNodeGraphNode.h
+++ b/Engine/Plugins/Cameras/GameplayCameras/Source/GameplayCamerasEditor/Public/Editors/CameraNodeGraphNode.h
@@ -22,8 +22,8 @@ public:
 
 	/** Gets input pin for given camera parameter property on the underlying camera node. */
 	UEdGraphPin* GetPinForCameraParameterProperty(const FName& InPropertyName) const;
-	/** Gets camera parameter property on the underlying camera node for the given input pin. */
-	FStructProperty* GetCameraParameterPropertyForPin(const UEdGraphPin* InPin) const;
+	/** Gets camera parameter property name on the underlying camera node for the given input pin. */
+	FName GetCameraParameterPropertyForPin(const UEdGraphPin* InPin) const;
 
 public:
 
diff --git a/Engine/Plugins/Cameras/GameplayCameras/Source/GameplayCamerasEditor/Public/Editors/CameraRigNodeGraphNode.h b/Engine/Plugins/Cameras/GameplayCameras/Source/GameplayCamerasEditor/Public/Editors/CameraRigNodeGraphNode.h
new file mode 100644
--- /dev/null
+++ b/Engine/Plugins/Cameras/GameplayCameras/Source/GameplayCamerasEditor/Public/Editors/CameraRigNodeGraphNode.h
@@ -0,0 +1,23 @@
+// Copyright Epic Games, Inc. All Rights Reserved.
+
+#pragma once
+
+#include "Editors/CameraNodeGraphNode.h"
+
+#include "CameraRigNodeGraphNode.generated.h"
+
+UCLASS()
+class UCameraRigNodeGraphNode : public UCameraNodeGraphNode
+{
+	GENERATED_BODY()
+
+public:
+
+	UCameraRigNodeGraphNode(const FObjectInitializer& ObjInit);
+
+public:
+
+	// UEdGraphNode interface
+	virtual void AllocateDefaultPins() override;
+};
+
