Cameras: ignore non-blocking hits in the collision-push node

diff --git a/Engine/Plugins/Cameras/GameplayCameras/Source/GameplayCameras/Private/Nodes/Collision/CollisionPushCameraNode.cpp b/Engine/Plugins/Cameras/GameplayCameras/Source/GameplayCameras/Private/Nodes/Collision/CollisionPushCameraNode.cpp
--- a/Engine/Plugins/Cameras/GameplayCameras/Source/GameplayCameras/Private/Nodes/Collision/CollisionPushCameraNode.cpp
+++ b/Engine/Plugins/Cameras/GameplayCameras/Source/GameplayCameras/Private/Nodes/Collision/CollisionPushCameraNode.cpp
@@ -263,7 +263,7 @@ void FCollisionPushCameraNodeEvaluator::RunCollisionTrace(UWorld* World, APlayer
 	if (CollisionPushNode->bRunAsyncCollision)
 	{
 		CollisionTraceHandle = World->AsyncSweepByChannel(
-			EAsyncTraceType::Multi,
+			EAsyncTraceType::Single,
 			TraceStart, TraceEnd, FQuat::Identity,
 			CollisionChannel,
 			SweepShape,
@@ -319,6 +319,11 @@ void FCollisionPushCameraNodeEvaluator::HandleCollisionTraceResult(UWorld* World
 
 	for (const FHitResult& Hit : HitResults)
 	{
+		if (!Hit.bBlockingHit)
+		{
+			continue;
+		}
+
 		double TraceLength = FVector3d::Distance(Hit.TraceStart, Hit.TraceEnd);
 		double DistanceToHit = FVector3d::Distance(Hit.TraceEnd, Hit.Location);
 		if (ensure(TraceLength > 0))
