Cameras: fix ensure triggered when deleting a camera rig with no referencers

diff --git a/Engine/Plugins/Cameras/GameplayCameras/Source/GameplayCamerasEditor/Private/Widgets/SDeleteCameraObjectDialog.cpp b/Engine/Plugins/Cameras/GameplayCameras/Source/GameplayCamerasEditor/Private/Widgets/SDeleteCameraObjectDialog.cpp
--- a/Engine/Plugins/Cameras/GameplayCameras/Source/GameplayCamerasEditor/Private/Widgets/SDeleteCameraObjectDialog.cpp
+++ b/Engine/Plugins/Cameras/GameplayCameras/Source/GameplayCamerasEditor/Private/Widgets/SDeleteCameraObjectDialog.cpp
@@ -355,9 +355,20 @@ void SDeleteCameraObjectDialog::ScanNextObject(IAssetRegistry& AssetRegistry)
 
 void SDeleteCameraObjectDialog::ScanNextReferencingPackage()
 {
-	if (ReferencingPackagesToScan.IsEmpty() && !PossiblyReferencingPackages.IsEmpty())
+	// If this is the first call, initialize the list of packages to scan.
+	// But if this is the first call and there's nothing to do, we can immediately
+	// skip to the next step.
+	if (ReferencingPackagesToScan.IsEmpty())
 	{
-		ReferencingPackagesToScan = PossiblyReferencingPackages.Array();
+		if (!PossiblyReferencingPackages.IsEmpty())
+		{
+			ReferencingPackagesToScan = PossiblyReferencingPackages.Array();
+		}
+		else
+		{
+			State = EState::FinishScanning;
+			return;
+		}
 	}
 
 	if (!ensure(ReferencingPackagesToScan.IsValidIndex(NextPackageToScan)))
