Cameras: remove ensure about null camera rig root nodes

The error was already logged during build and initialization, no need to throw an assert on evaluation.

diff --git a/Engine/Plugins/Cameras/GameplayCameras/Source/GameplayCameras/Private/Core/BlendStackCameraNode.cpp b/Engine/Plugins/Cameras/GameplayCameras/Source/GameplayCameras/Private/Core/BlendStackCameraNode.cpp
--- a/Engine/Plugins/Cameras/GameplayCameras/Source/GameplayCameras/Private/Core/BlendStackCameraNode.cpp
+++ b/Engine/Plugins/Cameras/GameplayCameras/Source/GameplayCameras/Private/Core/BlendStackCameraNode.cpp
@@ -432,7 +432,7 @@ void FBlendStackCameraNodeEvaluator::OnRun(const FCameraNodeEvaluationParams& Pa
 
 		// Run the camera rig's root node.
 		FCameraNodeEvaluator* RootEvaluator = Entry.RootEvaluator->GetRootEvaluator();
-		if (ensure(RootEvaluator))
+		if (RootEvaluator)
 		{
 			RootEvaluator->Run(CurParams, CurResult);
 		}
