Cameras: don't crash when trying to access the initial pose or variable table of a non-activated camera component

diff --git a/Engine/Plugins/Cameras/GameplayCameras/Source/GameplayCameras/Private/GameFramework/GameplayCameraComponent.cpp b/Engine/Plugins/Cameras/GameplayCameras/Source/GameplayCameras/Private/GameFramework/GameplayCameraComponent.cpp
--- a/Engine/Plugins/Cameras/GameplayCameras/Source/GameplayCameras/Private/GameFramework/GameplayCameraComponent.cpp
+++ b/Engine/Plugins/Cameras/GameplayCameras/Source/GameplayCameras/Private/GameFramework/GameplayCameraComponent.cpp
@@ -134,18 +134,53 @@ void UGameplayCameraComponent::ActivateCameraEvaluationContext(APlayerController
 
 FBlueprintCameraPose UGameplayCameraComponent::GetInitialPose() const
 {
-	return FBlueprintCameraPose::FromCameraPose(EvaluationContext->GetInitialResult().CameraPose);
+	if (EvaluationContext)
+	{
+		return FBlueprintCameraPose::FromCameraPose(EvaluationContext->GetInitialResult().CameraPose);
+	}
+	else
+	{
+		FFrame::KismetExecutionMessage(
+				*FString::Format(
+					TEXT("Can't get initial camera pose on Gameplay Camera component '{0}': it isn't active."),
+					{ *GetNameSafe(this) }),
+				ELogVerbosity::Error);
+		return FBlueprintCameraPose();
+	}
 }
 
 void UGameplayCameraComponent::SetInitialPose(const FBlueprintCameraPose& CameraPose)
 {
-	FCameraPose InitialPose = EvaluationContext->GetInitialResult().CameraPose;
-	CameraPose.ApplyTo(InitialPose);
+	if (EvaluationContext)
+	{
+		FCameraPose InitialPose = EvaluationContext->GetInitialResult().CameraPose;
+		CameraPose.ApplyTo(InitialPose);
+	}
+	else
+	{
+		FFrame::KismetExecutionMessage(
+				*FString::Format(
+					TEXT("Can't set initial camera pose on Gameplay Camera component '{0}': it isn't active."),
+					{ GetNameSafe(this) }),
+				ELogVerbosity::Error);
+	}
 }
 
 FBlueprintCameraVariableTable UGameplayCameraComponent::GetInitialVariableTable() const
 {
-	return FBlueprintCameraVariableTable(&EvaluationContext->GetInitialResult().VariableTable);
+	if (EvaluationContext)
+	{
+		return FBlueprintCameraVariableTable(&EvaluationContext->GetInitialResult().VariableTable);
+	}
+	else
+	{
+		FFrame::KismetExecutionMessage(
+				*FString::Format(
+					TEXT("Can't get initial camera variable table on Gameplay Camera component '{0}': it isn't active."),
+					{ GetNameSafe(this) }),
+				ELogVerbosity::Error);
+		return FBlueprintCameraVariableTable();
+	}
 }
 
 void UGameplayCameraComponent::OnRegister()
