Cameras: clear selection on undo/redo to avoid executing subsequent commands on potentially deleted nodes

diff --git a/Engine/Plugins/Cameras/GameplayCameras/Source/GameplayCamerasEditor/Private/Editors/SObjectTreeGraphEditor.cpp b/Engine/Plugins/Cameras/GameplayCameras/Source/GameplayCamerasEditor/Private/Editors/SObjectTreeGraphEditor.cpp
--- a/Engine/Plugins/Cameras/GameplayCameras/Source/GameplayCamerasEditor/Private/Editors/SObjectTreeGraphEditor.cpp
+++ b/Engine/Plugins/Cameras/GameplayCameras/Source/GameplayCamerasEditor/Private/Editors/SObjectTreeGraphEditor.cpp
@@ -11,6 +11,7 @@
 #include "Editors/ObjectTreeGraphSchema.h"
 #include "Editors/SObjectTreeGraphTitleBar.h"
 #include "Editors/SObjectTreeGraphToolbox.h"
+#include "Framework/Application/SlateApplication.h"
 #include "Framework/Commands/GenericCommands.h"
 #include "GraphEditorActions.h"
 #include "HAL/PlatformApplicationMisc.h"
@@ -251,7 +252,11 @@ void SObjectTreeGraphEditor::PostUndo(bool bSuccess)
 {
 	if (bSuccess)
 	{
+		GraphEditor->ClearSelectionSet();
+
 		GraphEditor->NotifyGraphChanged();
+
+		FSlateApplication::Get().DismissAllMenus();
 	}
 }
 
