Cameras: add player controller and view target information in the debug display

diff --git a/Engine/Plugins/Cameras/GameplayCameras/Source/GameplayCameras/Private/Debug/PlayerControllersDebugBlock.cpp b/Engine/Plugins/Cameras/GameplayCameras/Source/GameplayCameras/Private/Debug/PlayerControllersDebugBlock.cpp
new file mode 100644
--- /dev/null
+++ b/Engine/Plugins/Cameras/GameplayCameras/Source/GameplayCameras/Private/Debug/PlayerControllersDebugBlock.cpp
@@ -0,0 +1,98 @@
+// Copyright Epic Games, Inc. All Rights Reserved.
+
+#include "Debug/PlayerControllersDebugBlock.h"
+
+#include "Camera/PlayerCameraManager.h"
+#include "Debug/CameraDebugBlockBuilder.h"
+#include "Debug/CameraDebugRenderer.h"
+#include "Engine/World.h"
+#include "GameFramework/PlayerController.h"
+
+#if UE_GAMEPLAY_CAMERAS_DEBUG
+
+namespace UE::Cameras
+{
+
+UE_DEFINE_CAMERA_DEBUG_BLOCK(FPlayerControllersDebugBlock)
+
+FPlayerControllersDebugBlock::FPlayerControllersDebugBlock()
+{
+}
+
+void FPlayerControllersDebugBlock::Initialize(UWorld* World)
+{
+	if (!World)
+	{
+		return;
+	}
+
+	bHadValidWorld = true;
+
+	for (auto It = World->GetPlayerControllerIterator(); It; ++It)
+	{
+		APlayerController* PlayerController = It->Get();
+		if (!PlayerController)
+		{
+			continue;
+		}
+		if (!PlayerController->IsLocalPlayerController())
+		{
+			continue;
+		}
+
+		APlayerCameraManager* CameraManager = PlayerController->PlayerCameraManager;
+		AActor* ActiveViewTarget = CameraManager ? CameraManager->GetViewTarget() : nullptr;
+
+		FPlayerControllerDebugInfo PlayerControllerInfo;
+		PlayerControllerInfo.PlayerControllerName = *GetNameSafe(PlayerController);
+		PlayerControllerInfo.CameraManagerName = *GetNameSafe(CameraManager);
+		PlayerControllerInfo.ActiveViewTargetName = *GetNameSafe(ActiveViewTarget);
+		PlayerControllers.Add(PlayerControllerInfo);
+	}
+}
+
+void FPlayerControllersDebugBlock::OnDebugDraw(const FCameraDebugBlockDrawParams& Params, FCameraDebugRenderer& Renderer)
+{
+	Renderer.AddText("{cam_title}Player Controllers:{cam_default}");
+	Renderer.AddIndent();
+	{
+		Renderer.AddText(TEXT("%d active local player controller(s)\n"), PlayerControllers.Num());
+		if (bHadValidWorld)
+		{
+			for (FPlayerControllerDebugInfo& PlayerController : PlayerControllers)
+			{
+				Renderer.AddText(TEXT("- {cam_notice}%s{cam_default}"), *PlayerController.PlayerControllerName);
+				Renderer.AddIndent();
+				{
+					Renderer.AddText(TEXT("Camera manager: {cam_notice}%s{cam_default}\n"), *PlayerController.CameraManagerName);
+					Renderer.AddText(TEXT("View target: {cam_notice}%s{cam_default}"), *PlayerController.ActiveViewTargetName);
+				}
+				Renderer.RemoveIndent();
+			}
+		}
+		else
+		{
+			Renderer.AddText(TEXT("<invalid world>"));
+		}
+	}
+	Renderer.RemoveIndent();
+}
+
+void FPlayerControllersDebugBlock::OnSerialize(FArchive& Ar)
+{
+	Ar << PlayerControllers;
+	Ar << bHadValidWorld;
+}
+
+FArchive& operator<< (FArchive& Ar, FPlayerControllersDebugBlock::FPlayerControllerDebugInfo& PlayerControllerInfo)
+{
+	Ar << PlayerControllerInfo.PlayerControllerName;
+	Ar << PlayerControllerInfo.CameraManagerName;
+	Ar << PlayerControllerInfo.ActiveViewTargetName;
+	return Ar;
+}
+
+}  // namespace UE::Cameras
+
+#endif  // UE_GAMEPLAY_CAMERAS_DEBUG
+
diff --git a/Engine/Plugins/Cameras/GameplayCameras/Source/GameplayCameras/Private/Debug/RootCameraDebugBlock.cpp b/Engine/Plugins/Cameras/GameplayCameras/Source/GameplayCameras/Private/Debug/RootCameraDebugBlock.cpp
--- a/Engine/Plugins/Cameras/GameplayCameras/Source/GameplayCameras/Private/Debug/RootCameraDebugBlock.cpp
+++ b/Engine/Plugins/Cameras/GameplayCameras/Source/GameplayCameras/Private/Debug/RootCameraDebugBlock.cpp
@@ -16,6 +16,7 @@
 #include "Debug/CameraNodeEvaluatorDebugBlock.h"
 #include "Debug/CameraPoseDebugBlock.h"
 #include "Debug/CategoryTitleDebugBlock.h"
+#include "Debug/PlayerControllersDebugBlock.h"
 #include "Debug/VariableTableDebugBlock.h"
 #include "Debug/ViewfinderDebugBlock.h"
 #include "HAL/IConsoleManager.h"
@@ -126,6 +127,12 @@ void FRootCameraDebugBlock::BuildDebugBlocks(const FCameraSystemEvaluator& Camer
 			ResultDebugBlock.GetCameraPoseDebugBlock()->WithShowUnchangedCVar(TEXT("GameplayCameras.Debug.PoseStats.ShowUnchanged"));
 			ResultDebugBlock.GetVariableTableDebugBlock()->WithShowVariableIDsCVar(TEXT("GameplayCameras.Debug.PoseStats.ShowUnchanged"));
 		}
+		FPlayerControllersDebugBlock& PlayerControllersDebugBlock = Builder.BuildDebugBlock<FPlayerControllersDebugBlock>();
+		PoseStatsCategory.AddChild(&PlayerControllersDebugBlock);
+		{
+			UObject* CameraSystemOwner = CameraSystem.GetOwner();
+			PlayerControllersDebugBlock.Initialize(CameraSystemOwner ? CameraSystemOwner->GetWorld() : nullptr);
+		}
 	}
 	Builder.EndChildDebugBlock();
 	
diff --git a/Engine/Plugins/Cameras/GameplayCameras/Source/GameplayCameras/Public/Debug/PlayerControllersDebugBlock.h b/Engine/Plugins/Cameras/GameplayCameras/Source/GameplayCameras/Public/Debug/PlayerControllersDebugBlock.h
new file mode 100644
--- /dev/null
+++ b/Engine/Plugins/Cameras/GameplayCameras/Source/GameplayCameras/Public/Debug/PlayerControllersDebugBlock.h
@@ -0,0 +1,49 @@
+// Copyright Epic Games, Inc. All Rights Reserved.
+
+#pragma once
+
+#include "Debug/CameraDebugBlock.h"
+
+#if UE_GAMEPLAY_CAMERAS_DEBUG
+
+namespace UE::Cameras
+{
+
+struct FCameraDebugBlockBuilder;
+
+/**
+ * A debug block for showing the list of local player controllers and their view targets.
+ */
+class FPlayerControllersDebugBlock : public FCameraDebugBlock
+{
+	UE_DECLARE_CAMERA_DEBUG_BLOCK(GAMEPLAYCAMERAS_API, FPlayerControllersDebugBlock)
+
+public:
+
+	FPlayerControllersDebugBlock();
+
+	void Initialize(UWorld* World);
+
+protected:
+
+	virtual void OnDebugDraw(const FCameraDebugBlockDrawParams& Params, FCameraDebugRenderer& Renderer) override;
+	virtual void OnSerialize(FArchive& Ar) override;
+
+private:
+
+	struct FPlayerControllerDebugInfo
+	{
+		FString PlayerControllerName;
+		FString CameraManagerName;
+		FString ActiveViewTargetName;
+	};
+	TArray<FPlayerControllerDebugInfo> PlayerControllers;
+	bool bHadValidWorld = false;
+
+	friend FArchive& operator<< (FArchive&, FPlayerControllerDebugInfo&);
+};
+
+}  // namespace UE::Cameras
+
+#endif  // UE_GAMEPLAY_CAMERAS_DEBUG
+
