Cameras: add more information to the director tree debug display mode

diff --git a/Engine/Plugins/Cameras/GameplayCameras/Source/GameplayCameras/Private/Debug/CameraDirectorTreeDebugBlock.cpp b/Engine/Plugins/Cameras/GameplayCameras/Source/GameplayCameras/Private/Debug/CameraDirectorTreeDebugBlock.cpp
--- a/Engine/Plugins/Cameras/GameplayCameras/Source/GameplayCameras/Private/Debug/CameraDirectorTreeDebugBlock.cpp
+++ b/Engine/Plugins/Cameras/GameplayCameras/Source/GameplayCameras/Private/Debug/CameraDirectorTreeDebugBlock.cpp
@@ -37,30 +37,54 @@ void FCameraDirectorTreeDebugBlock::Initialize(const FCameraEvaluationContextSta
 		TSharedPtr<FCameraEvaluationContext> Context = Entry.WeakContext.Pin();
 
 		FDirectorDebugInfo EntryDebugInfo;
-		if (Context)
-		{
-			EntryDebugInfo.CameraAssetName = GetNameSafe(Context->GetCameraAsset());
-			EntryDebugInfo.InitialContextTransform = Context->GetInitialResult().CameraPose.GetTransform();
-			EntryDebugInfo.bIsValid = true;
-		}
-		else
-		{
-			EntryDebugInfo.bIsValid = false;
-		}
+		InitializeEntry(Context, EntryDebugInfo, Builder);
 		CameraDirectors.Add(EntryDebugInfo);
+	}
+}
 
-		if (Context)
-		{
-			const FCameraNodeEvaluationResult& InitialResult = Context->GetInitialResult();
-			AddChild(&Builder.BuildDebugBlock<FCameraPoseDebugBlock>(InitialResult.CameraPose)
-					.WithShowUnchangedCVar(TEXT("GameplayCameras.Debug.ContextInitialResult.ShowUnchanged")));
-		}
-		else
+void FCameraDirectorTreeDebugBlock::Initialize(TArrayView<const TSharedPtr<FCameraEvaluationContext>> Contexts, FCameraDebugBlockBuilder& Builder)
+{
+	for (TSharedPtr<FCameraEvaluationContext> Context : Contexts)
+	{
+		FDirectorDebugInfo EntryDebugInfo;
+		InitializeEntry(Context, EntryDebugInfo, Builder);
+		CameraDirectors.Add(EntryDebugInfo);
+	}
+}
+
+void FCameraDirectorTreeDebugBlock::InitializeEntry(TSharedPtr<FCameraEvaluationContext> Context, FDirectorDebugInfo& EntryDebugInfo, FCameraDebugBlockBuilder& Builder)
+{
+	if (Context)
+	{
+		const UObject* ContextOwner = Context->GetOwner();
+
+		EntryDebugInfo.CameraAssetName = GetNameSafe(Context->GetCameraAsset());
+		EntryDebugInfo.OwnerName = *GetNameSafe(ContextOwner);
+		EntryDebugInfo.OwnerClassName = *GetNameSafe(ContextOwner ? ContextOwner->GetClass() : nullptr);
+		EntryDebugInfo.InitialContextTransform = Context->GetInitialResult().CameraPose.GetTransform();
+		EntryDebugInfo.bIsValid = true;
+
+		const FCameraNodeEvaluationResult& InitialResult = Context->GetInitialResult();
+		AddChild(&Builder.BuildDebugBlock<FCameraPoseDebugBlock>(InitialResult.CameraPose)
+				.WithShowUnchangedCVar(TEXT("GameplayCameras.Debug.ContextInitialResult.ShowUnchanged")));
+
+		TArrayView<const TSharedPtr<FCameraEvaluationContext>> ChildrenContexts = Context->GetChildrenContexts();
+		if (ChildrenContexts.Num() > 0)
 		{
-			// Dummy debug block.
-			AddChild(&Builder.BuildDebugBlock<FCameraDebugBlock>());
+			FCameraDirectorTreeDebugBlock& ChildBlock = Builder.StartChildDebugBlock<FCameraDirectorTreeDebugBlock>();
+			{
+				ChildBlock.Initialize(ChildrenContexts, Builder);
+			}
+			Builder.EndChildDebugBlock();
 		}
 	}
+	else
+	{
+		EntryDebugInfo.bIsValid = false;
+
+		// Dummy debug block.
+		AddChild(&Builder.BuildDebugBlock<FCameraDebugBlock>());
+	}
 }
 
 void FCameraDirectorTreeDebugBlock::OnDebugDraw(const FCameraDebugBlockDrawParams& Params, FCameraDebugRenderer& Renderer)
@@ -121,6 +145,10 @@ void FCameraDirectorTreeDebugBlock::OnSerialize(FArchive& Ar)
 FArchive& operator<< (FArchive& Ar, FCameraDirectorTreeDebugBlock::FDirectorDebugInfo& DirectorDebugInfo)
 {
 	Ar << DirectorDebugInfo.CameraAssetName;
+	Ar << DirectorDebugInfo.OwnerClassName;
+	Ar << DirectorDebugInfo.OwnerName;
+	Ar << DirectorDebugInfo.InitialContextTransform;
+	Ar << DirectorDebugInfo.bIsValid;
 	return Ar;
 }
 
diff --git a/Engine/Plugins/Cameras/GameplayCameras/Source/GameplayCameras/Public/Debug/CameraDirectorTreeDebugBlock.h b/Engine/Plugins/Cameras/GameplayCameras/Source/GameplayCameras/Public/Debug/CameraDirectorTreeDebugBlock.h
--- a/Engine/Plugins/Cameras/GameplayCameras/Source/GameplayCameras/Public/Debug/CameraDirectorTreeDebugBlock.h
+++ b/Engine/Plugins/Cameras/GameplayCameras/Source/GameplayCameras/Public/Debug/CameraDirectorTreeDebugBlock.h
@@ -10,6 +10,7 @@
 namespace UE::Cameras
 {
 
+class FCameraEvaluationContext;
 struct FCameraDebugBlockBuilder;
 struct FCameraEvaluationContextStack;
 
@@ -25,17 +26,26 @@ public:
 	FCameraDirectorTreeDebugBlock();
 
 	void Initialize(const FCameraEvaluationContextStack& ContextStack, FCameraDebugBlockBuilder& Builder);
+	void Initialize(TArrayView<const TSharedPtr<FCameraEvaluationContext>> Contexts, FCameraDebugBlockBuilder& Builder);
 
 protected:
 
 	virtual void OnDebugDraw(const FCameraDebugBlockDrawParams& Params, FCameraDebugRenderer& Renderer) override;
 	virtual void OnSerialize(FArchive& Ar) override;
 
+private:
+
+	struct FDirectorDebugInfo;
+
+	void InitializeEntry(TSharedPtr<FCameraEvaluationContext> Context, FDirectorDebugInfo& EntryDebugInfo, FCameraDebugBlockBuilder& Builder);
+
 private:
 
 	struct FDirectorDebugInfo
 	{
 		FString CameraAssetName;
+		FString OwnerClassName;
+		FString OwnerName;
 		FTransform3d InitialContextTransform = FTransform3d::Identity;
 		bool bIsValid = false;
 	};
