Cameras: add static Blueprint function to get the auto-spawned camera system actor

diff --git a/Engine/Plugins/Cameras/GameplayCameras/Source/GameplayCameras/Private/GameFramework/GameplayCameraSystemActor.cpp b/Engine/Plugins/Cameras/GameplayCameras/Source/GameplayCameras/Private/GameFramework/GameplayCameraSystemActor.cpp
--- a/Engine/Plugins/Cameras/GameplayCameras/Source/GameplayCameras/Private/GameFramework/GameplayCameraSystemActor.cpp
+++ b/Engine/Plugins/Cameras/GameplayCameras/Source/GameplayCameras/Private/GameFramework/GameplayCameraSystemActor.cpp
@@ -39,36 +39,59 @@ void AGameplayCameraSystemActor::EndViewTarget(APlayerController* PC)
 	Super::EndViewTarget(PC);
 }
 
-void AGameplayCameraSystemActor::AutoManageActiveViewTarget(APlayerController* PlayerController)
+AGameplayCameraSystemActor* AGameplayCameraSystemActor::GetAutoSpawnedCameraSystemActor(APlayerController* PlayerController, bool bForceSpawn)
 {
 	static const TCHAR* AutoSpawnedActorName = TEXT("AutoSpawnedGameplayCameraSystemActor");
 
 	const UGameplayCamerasSettings* Settings = GetDefault<UGameplayCamerasSettings>();
 	if (!Settings->bAutoSpawnCameraSystemActor)
 	{
-		return;
+		return nullptr;
 	}
 
 	UGameplayCameraSystemHost* Host = UGameplayCameraSystemHost::FindHost(PlayerController);
 	if (!Host)
 	{
-		UE_LOG(LogCameraSystem, Error, TEXT("Can't auto-manage active view target: no camera system host found!"));
-		return;
+		if (bForceSpawn)
+		{
+			Host = UGameplayCameraSystemHost::FindOrCreateHost(PlayerController);
+		}
+		else
+		{
+			UE_LOG(LogCameraSystem, Error, TEXT("Can't auto-manage active view target: no camera system host found!"));
+			return nullptr;
+		}
 	}
-	
+
 	AGameplayCameraSystemActor* SpawnedActor = FindObject<AGameplayCameraSystemActor>(PlayerController, AutoSpawnedActorName);
 	if (!SpawnedActor)
 	{
-		FActorSpawnParameters SpawnParams;
-		SpawnParams.Name = AutoSpawnedActorName;
+		if (bForceSpawn)
+		{
+			FActorSpawnParameters SpawnParams;
+			SpawnParams.Name = AutoSpawnedActorName;
+
+			UWorld* World = PlayerController->GetWorld();
+			SpawnedActor = World->SpawnActor<AGameplayCameraSystemActor>(SpawnParams);
+
+			SpawnedActor->Rename(nullptr, PlayerController);
+		}
+		else
+		{
+			return nullptr;
+		}
+	}
 
-		UWorld* World = PlayerController->GetWorld();
-		SpawnedActor = World->SpawnActor<AGameplayCameraSystemActor>(SpawnParams);
+	return SpawnedActor;
+}
 
-		SpawnedActor->Rename(nullptr, PlayerController);
+void AGameplayCameraSystemActor::AutoManageActiveViewTarget(APlayerController* PlayerController)
+{
+	AGameplayCameraSystemActor* SpawnedActor = GetAutoSpawnedCameraSystemActor(PlayerController, true);
+	if (SpawnedActor)
+	{
+		SpawnedActor->GetCameraSystemComponent()->ActivateCameraSystemForPlayerController(PlayerController);
 	}
-	check(SpawnedActor);
-	SpawnedActor->GetCameraSystemComponent()->ActivateCameraSystemForPlayerController(PlayerController);
 }
 
 #undef LOCTEXT_NAMESPACE
diff --git a/Engine/Plugins/Cameras/GameplayCameras/Source/GameplayCameras/Public/GameFramework/GameplayCameraSystemActor.h b/Engine/Plugins/Cameras/GameplayCameras/Source/GameplayCameras/Public/GameFramework/GameplayCameraSystemActor.h
--- a/Engine/Plugins/Cameras/GameplayCameras/Source/GameplayCameras/Public/GameFramework/GameplayCameraSystemActor.h
+++ b/Engine/Plugins/Cameras/GameplayCameras/Source/GameplayCameras/Public/GameFramework/GameplayCameraSystemActor.h
@@ -34,6 +34,14 @@ public:
 
 public:
 
+	/**
+	 * Get the automatically spawned camera system actor, if it exists.
+	 * Returns null if bAutoSpawnCameraSystemActor is disabled in the GameplayCameras 
+	 * settings, or if no actor has been spawned yet.
+	 */
+	UFUNCTION(BlueprintCallable, Category=Camera)
+	static AGameplayCameraSystemActor* GetAutoSpawnedCameraSystemActor(APlayerController* PlayerController, bool bForceSpawn = false);
+
 	/**
 	 * Automatically sets a camera system actor as the view-target, spawning a unique
 	 * instance if needed. Doesn't do anything if bAutoSpawnCameraSystemActor is disabled
