Cameras: fix bug blending "flipped" properties in the camera pose

diff --git a/Engine/Plugins/Cameras/GameplayCameras/Source/GameplayCameras/Private/Core/CameraPose.cpp b/Engine/Plugins/Cameras/GameplayCameras/Source/GameplayCameras/Private/Core/CameraPose.cpp
--- a/Engine/Plugins/Cameras/GameplayCameras/Source/GameplayCameras/Private/Core/CameraPose.cpp
+++ b/Engine/Plugins/Cameras/GameplayCameras/Source/GameplayCameras/Private/Core/CameraPose.cpp
@@ -344,7 +344,7 @@ void FCameraPose::InternalLerpChanged(const FCameraPose& ToPose, float Factor, c
 #define UE_CAMERA_POSE_FOR_PROPERTY(PropType, PropName)\
 		if ((!bInvertMask && InMask.PropName) || (bInvertMask && !InMask.PropName))\
 		{\
-			if (!bChangedOnly || ToPoseChangedFlags.PropName && Factor >= 0.5f)\
+			if ((!bChangedOnly || ToPoseChangedFlags.PropName) && Factor >= 0.5f)\
 			{\
 				ensureMsgf(ChangedFlags.PropName, TEXT("Interpolating " #PropName " from default value!"));\
 				Set##PropName(ToPose.Get##PropName());\
