Cameras: add clipping planes camera node

diff --git a/Engine/Plugins/Cameras/GameplayCameras/Source/GameplayCameras/Private/Nodes/Common/ClippingPlanesCameraNode.cpp b/Engine/Plugins/Cameras/GameplayCameras/Source/GameplayCameras/Private/Nodes/Common/ClippingPlanesCameraNode.cpp
new file mode 100644
--- /dev/null
+++ b/Engine/Plugins/Cameras/GameplayCameras/Source/GameplayCameras/Private/Nodes/Common/ClippingPlanesCameraNode.cpp
@@ -0,0 +1,61 @@
+// Copyright Epic Games, Inc. All Rights Reserved.
+
+#include "Nodes/Common/ClippingPlanesCameraNode.h"
+
+#include "Core/CameraEvaluationContext.h"
+#include "Core/CameraParameterReader.h"
+#include "GameplayCameras.h"
+
+#include UE_INLINE_GENERATED_CPP_BY_NAME(ClippingPlanesCameraNode)
+
+namespace UE::Cameras
+{
+
+class FClippingPlanesCameraNodeEvaluator : public FCameraNodeEvaluator
+{
+	UE_DECLARE_CAMERA_NODE_EVALUATOR(GAMEPLAYCAMERAS_API, FClippingPlanesCameraNodeEvaluator)
+
+protected:
+
+	virtual void OnInitialize(const FCameraNodeEvaluatorInitializeParams& Params, FCameraNodeEvaluationResult& OutResult) override;
+	virtual void OnRun(const FCameraNodeEvaluationParams& Params, FCameraNodeEvaluationResult& OutResult) override;
+
+private:
+
+	TCameraParameterReader<double> NearPlaneReader;
+	TCameraParameterReader<double> FarPlaneReader;
+};
+
+UE_DEFINE_CAMERA_NODE_EVALUATOR(FClippingPlanesCameraNodeEvaluator)
+
+void FClippingPlanesCameraNodeEvaluator::OnInitialize(const FCameraNodeEvaluatorInitializeParams& Params, FCameraNodeEvaluationResult& OutResult)
+{
+	const UClippingPlanesCameraNode* ClippingPlanesNode = GetCameraNodeAs<UClippingPlanesCameraNode>();
+	NearPlaneReader.Initialize(ClippingPlanesNode->NearPlane);
+	FarPlaneReader.Initialize(ClippingPlanesNode->FarPlane);
+}
+
+void FClippingPlanesCameraNodeEvaluator::OnRun(const FCameraNodeEvaluationParams& Params, FCameraNodeEvaluationResult& OutResult)
+{
+	FCameraPose& OutPose = OutResult.CameraPose;
+
+	float NearPlane = NearPlaneReader.Get(OutResult.VariableTable);
+	if (NearPlane > 0)
+	{
+		OutPose.SetNearClippingPlane(NearPlane);
+	}
+	float FarPlane = FarPlaneReader.Get(OutResult.VariableTable);
+	if (FarPlane > 0)
+	{
+		OutPose.SetFarClippingPlane(FarPlane);
+	}
+}
+
+}  // namespace UE::Cameras
+
+FCameraNodeEvaluatorPtr UClippingPlanesCameraNode::OnBuildEvaluator(FCameraNodeEvaluatorBuilder& Builder) const
+{
+	using namespace UE::Cameras;
+	return Builder.BuildEvaluator<FClippingPlanesCameraNodeEvaluator>();
+}
+
diff --git a/Engine/Plugins/Cameras/GameplayCameras/Source/GameplayCameras/Public/Nodes/Common/ClippingPlanesCameraNode.h b/Engine/Plugins/Cameras/GameplayCameras/Source/GameplayCameras/Public/Nodes/Common/ClippingPlanesCameraNode.h
new file mode 100644
--- /dev/null
+++ b/Engine/Plugins/Cameras/GameplayCameras/Source/GameplayCameras/Public/Nodes/Common/ClippingPlanesCameraNode.h
@@ -0,0 +1,31 @@
+// Copyright Epic Games, Inc. All Rights Reserved.
+
+#pragma once
+
+#include "Core/CameraNode.h"
+#include "Core/CameraParameters.h"
+
+#include "ClippingPlanesCameraNode.generated.h"
+
+/**
+ * A camera node that sets the clipping planes for the evaluated camera.
+ */
+UCLASS(MinimalAPI, meta=(CameraNodeCategories="Common,Rendering"))
+class UClippingPlanesCameraNode : public UCameraNode
+{
+	GENERATED_BODY()
+
+protected:
+
+	// UCameraNode interface.
+	virtual FCameraNodeEvaluatorPtr OnBuildEvaluator(FCameraNodeEvaluatorBuilder& Builder) const override;
+
+public:
+
+	UPROPERTY(EditAnywhere, Category="Clipping", meta=(ClampMin="0.001"))
+	FDoubleCameraParameter NearPlane;
+
+	UPROPERTY(EditAnywhere, Category="Clipping", meta=(ClampMin="0.001"))
+	FDoubleCameraParameter FarPlane;
+};
+
