Cameras: add suport for post-process settings and more aspect-ratio-related options in the evaluation result

- Add post-process settings to the result structure.
- Add aspect ratio constraint setting to the camera pose and related types.
- Properly harvest these new settings when passing the evaluated result to the core engine.
- Consolidate the logic for overriding/lerp'ing evaluation results in the evaluation result class. This removes some code duplication in blend nodes.
- Fix some missing serialization logic on the evaluation result and related structures.
- Add post-process setting camera node and filmback camera node.

Note that the sensor size parameters have been removed from the lens paramters camera node, as they better belong to the filmback camera node.

diff --git a/Engine/Plugins/Cameras/GameplayCameras/Source/GameplayCameras/Private/Core/BlendStackCameraNode.cpp b/Engine/Plugins/Cameras/GameplayCameras/Source/GameplayCameras/Private/Core/BlendStackCameraNode.cpp
--- a/Engine/Plugins/Cameras/GameplayCameras/Source/GameplayCameras/Private/Core/BlendStackCameraNode.cpp
+++ b/Engine/Plugins/Cameras/GameplayCameras/Source/GameplayCameras/Private/Core/BlendStackCameraNode.cpp
@@ -421,6 +421,7 @@ void FBlendStackCameraNodeEvaluator::OnRun(const FCameraNodeEvaluationParams& Pa
 		// Start with the input given to us.
 		CurResult.CameraPose = OutResult.CameraPose;
 		CurResult.CameraRigJoints.OverrideAll(OutResult.CameraRigJoints);
+		CurResult.PostProcessSettings.OverrideAll(OutResult.PostProcessSettings);
 
 		// Override it with whatever the evaluation context has set on its result.
 		const FCameraNodeEvaluationResult& ContextResult(ResolvedEntry.Context->GetInitialResult());
diff --git a/Engine/Plugins/Cameras/GameplayCameras/Source/GameplayCameras/Private/Core/CameraNodeEvaluator.cpp b/Engine/Plugins/Cameras/GameplayCameras/Source/GameplayCameras/Private/Core/CameraNodeEvaluator.cpp
--- a/Engine/Plugins/Cameras/GameplayCameras/Source/GameplayCameras/Private/Core/CameraNodeEvaluator.cpp
+++ b/Engine/Plugins/Cameras/GameplayCameras/Source/GameplayCameras/Private/Core/CameraNodeEvaluator.cpp
@@ -16,6 +16,7 @@ void FCameraNodeEvaluationResult::Reset(bool bResetVariableTable)
 {
 	CameraPose.Reset();
 	CameraRigJoints.Reset();
+	PostProcessSettings.Reset();
 
 	if (bResetVariableTable)
 	{
@@ -26,10 +27,45 @@ void FCameraNodeEvaluationResult::Reset(bool bResetVariableTable)
 	bIsValid = false;
 }
 
+void FCameraNodeEvaluationResult::OverrideAll(const FCameraNodeEvaluationResult& OtherResult)
+{
+	CameraPose.OverrideAll(OtherResult.CameraPose);
+	VariableTable.OverrideAll(OtherResult.VariableTable);
+	CameraRigJoints.OverrideAll(OtherResult.CameraRigJoints);
+	PostProcessSettings.OverrideAll(OtherResult.PostProcessSettings);
+	bIsCameraCut = OtherResult.bIsCameraCut;
+	bIsValid = OtherResult.bIsValid;
+}
+
+void FCameraNodeEvaluationResult::LerpAll(const FCameraNodeEvaluationResult& ToResult, float BlendFactor)
+{
+	// Blend all properties.
+	CameraPose.LerpAll(ToResult.CameraPose, BlendFactor);
+	VariableTable.LerpAll(ToResult.VariableTable, BlendFactor);
+
+	// Merge/blend the joints.
+	CameraRigJoints.LerpAll(ToResult.CameraRigJoints, BlendFactor);
+
+	// Merge/blend the post-process settings.
+	PostProcessSettings.LerpAll(ToResult.PostProcessSettings, BlendFactor);
+
+	// If we have even a fraction of a camera cut, we need to make the
+	// whole result into a camera cut.
+	if (BlendFactor > 0.f && ToResult.bIsCameraCut)
+	{
+		bIsCameraCut = true;
+	}
+
+	// The blended result is valid if both input results are valid.
+	bIsValid = (bIsValid && ToResult.bIsValid);
+}
+
 void FCameraNodeEvaluationResult::Serialize(FArchive& Ar)
 {
 	CameraPose.SerializeWithFlags(Ar);
 	VariableTable.Serialize(Ar);
+	CameraRigJoints.Serialize(Ar);
+	PostProcessSettings.Serialize(Ar);
 	Ar << bIsCameraCut;
 	Ar << bIsValid;
 }
diff --git a/Engine/Plugins/Cameras/GameplayCameras/Source/GameplayCameras/Private/Core/CameraPose.cpp b/Engine/Plugins/Cameras/GameplayCameras/Source/GameplayCameras/Private/Core/CameraPose.cpp
--- a/Engine/Plugins/Cameras/GameplayCameras/Source/GameplayCameras/Private/Core/CameraPose.cpp
+++ b/Engine/Plugins/Cameras/GameplayCameras/Source/GameplayCameras/Private/Core/CameraPose.cpp
@@ -165,6 +165,12 @@ double FCameraPose::GetSensorAspectRatio(float SensorWidth, float SensorHeight)
 	return (SensorHeight > 0.f) ? (SensorWidth / SensorHeight) : 0.0;
 }
 
+void FCameraPose::GetDefaultSensorSize(float& OutSensorWidth, float& OutSensorHeight)
+{
+	OutSensorWidth = 24.89f;
+	OutSensorHeight = 18.67f;
+}
+
 FRay3d FCameraPose::GetAimRay() const
 {
 	const bool bDirectionIsNormalized = false;
diff --git a/Engine/Plugins/Cameras/GameplayCameras/Source/GameplayCameras/Private/Core/CameraRigJoints.cpp b/Engine/Plugins/Cameras/GameplayCameras/Source/GameplayCameras/Private/Core/CameraRigJoints.cpp
--- a/Engine/Plugins/Cameras/GameplayCameras/Source/GameplayCameras/Private/Core/CameraRigJoints.cpp
+++ b/Engine/Plugins/Cameras/GameplayCameras/Source/GameplayCameras/Private/Core/CameraRigJoints.cpp
@@ -7,6 +7,13 @@
 namespace UE::Cameras
 {
 
+FArchive& operator<< (FArchive& Ar, FCameraRigJoint& RigJoint)
+{
+	Ar << RigJoint.VariableID;
+	Ar << RigJoint.Transform;
+	return Ar;
+}
+
 void FCameraRigJoints::AddJoint(const FCameraRigJoint& InJoint)
 {
 	Joints.Add(InJoint);
@@ -27,6 +34,11 @@ void FCameraRigJoints::Reset()
 	Joints.Reset();
 }
 
+void FCameraRigJoints::Serialize(FArchive& Ar)
+{
+	Ar << Joints;
+}
+
 void FCameraRigJoints::OverrideAll(const FCameraRigJoints& OtherJoints)
 {
 	Joints = OtherJoints.Joints;
diff --git a/Engine/Plugins/Cameras/GameplayCameras/Source/GameplayCameras/Private/Core/CameraSystemEvaluator.cpp b/Engine/Plugins/Cameras/GameplayCameras/Source/GameplayCameras/Private/Core/CameraSystemEvaluator.cpp
--- a/Engine/Plugins/Cameras/GameplayCameras/Source/GameplayCameras/Private/Core/CameraSystemEvaluator.cpp
+++ b/Engine/Plugins/Cameras/GameplayCameras/Source/GameplayCameras/Private/Core/CameraSystemEvaluator.cpp
@@ -251,10 +251,26 @@ void FCameraSystemEvaluator::PostUpdateServices(float DeltaTime, ECameraEvaluati
 
 void FCameraSystemEvaluator::GetEvaluatedCameraView(FMinimalViewInfo& DesiredView)
 {
-	const FCameraPose& CameraPose = Result.CameraPose;
+	const FCameraPose& CameraPose = RootNodeResult.CameraPose;
 	DesiredView.Location = CameraPose.GetLocation();
 	DesiredView.Rotation = CameraPose.GetRotation();
 	DesiredView.FOV = CameraPose.GetEffectiveFieldOfView();
+
+	DesiredView.AspectRatio = CameraPose.GetSensorAspectRatio();
+	DesiredView.bConstrainAspectRatio = CameraPose.GetConstrainAspectRatio();
+	DesiredView.AspectRatioAxisConstraint = CameraPose.GetOverrideAspectRatioAxisConstraint() ?
+		CameraPose.GetAspectRatioAxisConstraint() : TOptional<EAspectRatioAxisConstraint>();
+
+	// TODO: add support for ortho cameras.
+	DesiredView.PerspectiveNearClipPlane = CameraPose.GetNearClippingPlane();
+
+	for (const FPostProcessSettingsCollectionEntry& PostProcessEntry : RootNodeResult.PostProcessSettings.GetEntries())
+	{
+		DesiredView.PostProcessSettings = PostProcessEntry.PostProcessSettings;
+		DesiredView.PostProcessBlendWeight = PostProcessEntry.PostProcessBlendWeight;
+		// TODO: support multiple post-process settings by getting the PlayerCameraManager
+		break;
+	}
 }
 
 #if UE_GAMEPLAY_CAMERAS_DEBUG
diff --git a/Engine/Plugins/Cameras/GameplayCameras/Source/GameplayCameras/Private/Core/PostProcessSettingsCollection.cpp b/Engine/Plugins/Cameras/GameplayCameras/Source/GameplayCameras/Private/Core/PostProcessSettingsCollection.cpp
new file mode 100644
--- /dev/null
+++ b/Engine/Plugins/Cameras/GameplayCameras/Source/GameplayCameras/Private/Core/PostProcessSettingsCollection.cpp
@@ -0,0 +1,62 @@
+// Copyright Epic Games, Inc. All Rights Reserved.
+
+#include "Core/PostProcessSettingsCollection.h"
+
+namespace UE::Cameras
+{
+
+FArchive& operator<< (FArchive& Ar, FPostProcessSettingsCollectionEntry& Entry)
+{
+	static FPostProcessSettings DefaultPostProcessSettings;
+
+	UScriptStruct* PostProcessSettingsStruct = FPostProcessSettings::StaticStruct();
+	PostProcessSettingsStruct->SerializeItem(Ar, &Entry.PostProcessSettings, &DefaultPostProcessSettings);
+
+	Ar << Entry.PostProcessBlendWeight;
+
+	return Ar;
+}
+
+void FPostProcessSettingsCollection::Add(const FPostProcessSettings& InPostProcess, float InBlendWeight)
+{
+	Entries.Add(FEntry{ InPostProcess, InBlendWeight });
+}
+
+void FPostProcessSettingsCollection::Reset()
+{
+	Entries.Reset();
+}
+
+void FPostProcessSettingsCollection::OverrideAll(const FPostProcessSettingsCollection& OtherCollection)
+{
+	Entries = OtherCollection.Entries;
+}
+
+void FPostProcessSettingsCollection::LerpAll(const FPostProcessSettingsCollection& ToCollection, float BlendFactor)
+{
+	// Blend down our entries. Remove any entries with zero weight.
+	for (auto It = Entries.CreateIterator(); It; ++It)
+	{
+		FEntry& ThisEntry = *It;
+		ThisEntry.PostProcessBlendWeight *= (1.f - BlendFactor);
+		
+		if (ThisEntry.PostProcessBlendWeight < UE_SMALL_NUMBER)
+		{
+			It.RemoveCurrent();
+		}
+	}
+
+	// Blend up the other collection's entries.
+	for (const FEntry& OtherEntry : ToCollection.Entries)
+	{
+		Entries.Add(FEntry{ OtherEntry.PostProcessSettings, OtherEntry.PostProcessBlendWeight * BlendFactor });
+	}
+}
+
+void FPostProcessSettingsCollection::Serialize(FArchive& Ar)
+{
+	Ar << Entries;
+}
+
+}  // namespace UE::Cameras
+
diff --git a/Engine/Plugins/Cameras/GameplayCameras/Source/GameplayCameras/Private/GameFramework/BlueprintCameraPose.cpp b/Engine/Plugins/Cameras/GameplayCameras/Source/GameplayCameras/Private/GameFramework/BlueprintCameraPose.cpp
--- a/Engine/Plugins/Cameras/GameplayCameras/Source/GameplayCameras/Private/GameFramework/BlueprintCameraPose.cpp
+++ b/Engine/Plugins/Cameras/GameplayCameras/Source/GameplayCameras/Private/GameFramework/BlueprintCameraPose.cpp
@@ -130,6 +130,7 @@ FBlueprintCameraPose UBlueprintCameraPoseFunctionLibrary::MakeCameraPoseFromCame
 		Result.Rotation = (AdditiveTransform.GetRotation() * CameraComponentTransform.GetRotation()).Rotator();
 		Result.FieldOfView = CameraComponent->FieldOfView + AdditiveFOV;
 		Result.ConstrainAspectRatio = CameraComponent->bConstrainAspectRatio;
+		Result.OverrideAspectRatioAxisConstraint = CameraComponent->bOverrideAspectRatioAxisConstraint;
 		Result.AspectRatioAxisConstraint = CameraComponent->AspectRatioAxisConstraint;
 		Result.SensorWidth = CameraComponent->AspectRatio * Result.SensorHeight;
 	}
diff --git a/Engine/Plugins/Cameras/GameplayCameras/Source/GameplayCameras/Private/Nodes/Blends/OrbitBlendCameraNode.cpp b/Engine/Plugins/Cameras/GameplayCameras/Source/GameplayCameras/Private/Nodes/Blends/OrbitBlendCameraNode.cpp
--- a/Engine/Plugins/Cameras/GameplayCameras/Source/GameplayCameras/Private/Nodes/Blends/OrbitBlendCameraNode.cpp
+++ b/Engine/Plugins/Cameras/GameplayCameras/Source/GameplayCameras/Private/Nodes/Blends/OrbitBlendCameraNode.cpp
@@ -110,8 +110,8 @@ void FOrbitBlendCameraNodeEvaluator::OnBlendResults(const FCameraNodeBlendParams
 	// If we don't have a driving blend, just cut to the new camera rig.
 	if (!DrivingBlendEvaluator)
 	{
-		BlendedResult.CameraPose.OverrideAll(ChildResult.CameraPose);
-		BlendedResult.VariableTable.OverrideAll(ChildResult.VariableTable);
+		BlendedResult.OverrideAll(ChildResult);
+
 		BlendedResult.bIsCameraCut = true;
 
 		OutResult.bIsBlendFinished = true;
diff --git a/Engine/Plugins/Cameras/GameplayCameras/Source/GameplayCameras/Private/Nodes/Blends/PopBlendCameraNode.cpp b/Engine/Plugins/Cameras/GameplayCameras/Source/GameplayCameras/Private/Nodes/Blends/PopBlendCameraNode.cpp
--- a/Engine/Plugins/Cameras/GameplayCameras/Source/GameplayCameras/Private/Nodes/Blends/PopBlendCameraNode.cpp
+++ b/Engine/Plugins/Cameras/GameplayCameras/Source/GameplayCameras/Private/Nodes/Blends/PopBlendCameraNode.cpp
@@ -33,11 +33,9 @@ void FPopBlendCameraNodeEvaluator::OnBlendResults(const FCameraNodeBlendParams&
 {
 	const FCameraNodeEvaluationResult& ChildResult(Params.ChildResult);
 	FCameraNodeEvaluationResult& BlendedResult(OutResult.BlendedResult);
-	
-	BlendedResult.CameraPose.OverrideAll(ChildResult.CameraPose);
-	BlendedResult.VariableTable.OverrideAll(ChildResult.VariableTable);
-	BlendedResult.CameraRigJoints.OverrideAll(ChildResult.CameraRigJoints);
 
+	BlendedResult.OverrideAll(ChildResult);
+	
 	if (ChildResult.bIsCameraCut || Params.ChildParams.bIsFirstFrame)
 	{
 		BlendedResult.bIsCameraCut = true;
diff --git a/Engine/Plugins/Cameras/GameplayCameras/Source/GameplayCameras/Private/Nodes/Blends/SimpleBlendCameraNode.cpp b/Engine/Plugins/Cameras/GameplayCameras/Source/GameplayCameras/Private/Nodes/Blends/SimpleBlendCameraNode.cpp
--- a/Engine/Plugins/Cameras/GameplayCameras/Source/GameplayCameras/Private/Nodes/Blends/SimpleBlendCameraNode.cpp
+++ b/Engine/Plugins/Cameras/GameplayCameras/Source/GameplayCameras/Private/Nodes/Blends/SimpleBlendCameraNode.cpp
@@ -37,19 +37,7 @@ void FSimpleBlendCameraNodeEvaluator::OnBlendResults(const FCameraNodeBlendParam
 	const FCameraNodeEvaluationResult& ChildResult(Params.ChildResult);
 	FCameraNodeEvaluationResult& BlendedResult(OutResult.BlendedResult);
 
-	// Blend all properties.
-	BlendedResult.CameraPose.LerpAll(ChildResult.CameraPose, BlendFactor);
-	BlendedResult.VariableTable.LerpAll(ChildResult.VariableTable, BlendFactor);
-
-	// Merge/blend the joints.
-	BlendedResult.CameraRigJoints.LerpAll(ChildResult.CameraRigJoints, BlendFactor);
-
-	// If we have even a fraction of a camera cut, we need to make the
-	// whole result into a camera cut.
-	if (BlendFactor > 0.f && ChildResult.bIsCameraCut)
-	{
-		BlendedResult.bIsCameraCut = true;
-	}
+	BlendedResult.LerpAll(ChildResult, BlendFactor);
 
 	OutResult.bIsBlendFull = BlendFactor >= 1.f;
 	OutResult.bIsBlendFinished = bIsBlendFinished;
diff --git a/Engine/Plugins/Cameras/GameplayCameras/Source/GameplayCameras/Private/Nodes/Common/FilmbackCameraNode.cpp b/Engine/Plugins/Cameras/GameplayCameras/Source/GameplayCameras/Private/Nodes/Common/FilmbackCameraNode.cpp
new file mode 100644
--- /dev/null
+++ b/Engine/Plugins/Cameras/GameplayCameras/Source/GameplayCameras/Private/Nodes/Common/FilmbackCameraNode.cpp
@@ -0,0 +1,76 @@
+// Copyright Epic Games, Inc. All Rights Reserved.
+
+#include "Nodes/Common/FilmbackCameraNode.h"
+
+#include "Core/CameraParameterReader.h"
+#include "Core/CameraPose.h"
+
+#include UE_INLINE_GENERATED_CPP_BY_NAME(FilmbackCameraNode)
+
+namespace UE::Cameras
+{
+
+class FFilmbackCameraNodeEvaluator : public FCameraNodeEvaluator
+{
+	UE_DECLARE_CAMERA_NODE_EVALUATOR(GAMEPLAYCAMERAS_API, FFilmbackCameraNodeEvaluator)
+
+protected:
+
+	virtual void OnInitialize(const FCameraNodeEvaluatorInitializeParams& Params, FCameraNodeEvaluationResult& OutResult) override;
+	virtual void OnRun(const FCameraNodeEvaluationParams& Params, FCameraNodeEvaluationResult& OutResult) override;
+
+private:
+
+	TCameraParameterReader<float> SensorWidthReader;
+	TCameraParameterReader<float> SensorHeightReader;
+	TCameraParameterReader<bool> ConstrainAspectRatioReader;
+	TCameraParameterReader<bool> OverrideAspectRatioAxisConstraintReader;
+};
+
+UE_DEFINE_CAMERA_NODE_EVALUATOR(FFilmbackCameraNodeEvaluator)
+
+void FFilmbackCameraNodeEvaluator::OnInitialize(const FCameraNodeEvaluatorInitializeParams& Params, FCameraNodeEvaluationResult& OutResult)
+{
+	const UFilmbackCameraNode* FilmbackNode = GetCameraNodeAs<UFilmbackCameraNode>();
+	SensorWidthReader.Initialize(FilmbackNode->SensorWidth);
+	SensorHeightReader.Initialize(FilmbackNode->SensorHeight);
+	ConstrainAspectRatioReader.Initialize(FilmbackNode->ConstrainAspectRatio);
+	OverrideAspectRatioAxisConstraintReader.Initialize(FilmbackNode->OverrideAspectRatioAxisConstraint);
+}
+
+void FFilmbackCameraNodeEvaluator::OnRun(const FCameraNodeEvaluationParams& Params, FCameraNodeEvaluationResult& OutResult)
+{
+	FCameraPose& OutPose = OutResult.CameraPose;
+
+	float SensorWidth = SensorWidthReader.Get(OutResult.VariableTable);
+	if (SensorWidth > 0)
+	{
+		OutPose.SetSensorWidth(SensorWidth);
+	}
+	float SensorHeight = SensorHeightReader.Get(OutResult.VariableTable);
+	if (SensorHeight > 0)
+	{
+		OutPose.SetSensorHeight(SensorHeight);
+	}
+
+	OutPose.SetConstrainAspectRatio(ConstrainAspectRatioReader.Get(OutResult.VariableTable));
+	OutPose.SetOverrideAspectRatioAxisConstraint(OverrideAspectRatioAxisConstraintReader.Get(OutResult.VariableTable));
+
+	// TODO: add support for enum camera parameters.
+	const UFilmbackCameraNode* FilmbackNode = GetCameraNodeAs<UFilmbackCameraNode>();
+	OutPose.SetAspectRatioAxisConstraint(FilmbackNode->AspectRatioAxisConstraint);
+}
+
+}  // namespace UE::Cameras
+
+UFilmbackCameraNode::UFilmbackCameraNode(const FObjectInitializer& ObjectInitializer)
+{
+	FCameraPose::GetDefaultSensorSize(SensorWidth.Value, SensorHeight.Value);
+}
+
+FCameraNodeEvaluatorPtr UFilmbackCameraNode::OnBuildEvaluator(FCameraNodeEvaluatorBuilder& Builder) const
+{
+	using namespace UE::Cameras;
+	return Builder.BuildEvaluator<FFilmbackCameraNodeEvaluator>();
+}
+
diff --git a/Engine/Plugins/Cameras/GameplayCameras/Source/GameplayCameras/Private/Nodes/Common/LensParametersCameraNode.cpp b/Engine/Plugins/Cameras/GameplayCameras/Source/GameplayCameras/Private/Nodes/Common/LensParametersCameraNode.cpp
--- a/Engine/Plugins/Cameras/GameplayCameras/Source/GameplayCameras/Private/Nodes/Common/LensParametersCameraNode.cpp
+++ b/Engine/Plugins/Cameras/GameplayCameras/Source/GameplayCameras/Private/Nodes/Common/LensParametersCameraNode.cpp
@@ -22,8 +22,6 @@ protected:
 
 private:
 
-	TCameraParameterReader<float> SensorWidthReader;
-	TCameraParameterReader<float> SensorHeightReader;
 	TCameraParameterReader<float> FocusDistanceReader;
 	TCameraParameterReader<float> FocalLengthReader;
 	TCameraParameterReader<float> ApertureReader;
@@ -34,8 +32,6 @@ UE_DEFINE_CAMERA_NODE_EVALUATOR(FLensParametersCameraNodeEvaluator)
 void FLensParametersCameraNodeEvaluator::OnInitialize(const FCameraNodeEvaluatorInitializeParams& Params, FCameraNodeEvaluationResult& OutResult)
 {
 	const ULensParametersCameraNode* LensParametersNode = GetCameraNodeAs<ULensParametersCameraNode>();
-	SensorWidthReader.Initialize(LensParametersNode->SensorWidth);
-	SensorHeightReader.Initialize(LensParametersNode->SensorHeight);
 	FocusDistanceReader.Initialize(LensParametersNode->FocusDistance);
 	FocalLengthReader.Initialize(LensParametersNode->FocalLength);
 	ApertureReader.Initialize(LensParametersNode->Aperture);
@@ -45,16 +41,6 @@ void FLensParametersCameraNodeEvaluator::OnRun(const FCameraNodeEvaluationParams
 {
 	FCameraPose& OutPose = OutResult.CameraPose;
 
-	float SensorWidth = SensorWidthReader.Get(OutResult.VariableTable);
-	if (SensorWidth > 0)
-	{
-		OutPose.SetSensorWidth(SensorWidth);
-	}
-	float SensorHeight = SensorHeightReader.Get(OutResult.VariableTable);
-	if (SensorHeight > 0)
-	{
-		OutPose.SetSensorHeight(SensorHeight);
-	}
 	float FocusDistance = FocusDistanceReader.Get(OutResult.VariableTable);
 	if (FocusDistance > 0)
 	{
diff --git a/Engine/Plugins/Cameras/GameplayCameras/Source/GameplayCameras/Private/Nodes/Common/PostProcessCameraNode.cpp b/Engine/Plugins/Cameras/GameplayCameras/Source/GameplayCameras/Private/Nodes/Common/PostProcessCameraNode.cpp
new file mode 100644
--- /dev/null
+++ b/Engine/Plugins/Cameras/GameplayCameras/Source/GameplayCameras/Private/Nodes/Common/PostProcessCameraNode.cpp
@@ -0,0 +1,54 @@
+// Copyright Epic Games, Inc. All Rights Reserved.
+
+#include "Nodes/Common/PostProcessCameraNode.h"
+
+#include "Core/CameraParameterReader.h"
+#include "Core/CameraPose.h"
+
+#include UE_INLINE_GENERATED_CPP_BY_NAME(PostProcessCameraNode)
+
+namespace UE::Cameras
+{
+
+class FPostProcessCameraNodeEvaluator : public FCameraNodeEvaluator
+{
+	UE_DECLARE_CAMERA_NODE_EVALUATOR(GAMEPLAYCAMERAS_API, FPostProcessCameraNodeEvaluator)
+
+protected:
+
+	virtual void OnInitialize(const FCameraNodeEvaluatorInitializeParams& Params, FCameraNodeEvaluationResult& OutResult) override;
+	virtual void OnRun(const FCameraNodeEvaluationParams& Params, FCameraNodeEvaluationResult& OutResult) override;
+
+private:
+
+	TCameraParameterReader<float> PostProcessBlendWeightReader;
+};
+
+UE_DEFINE_CAMERA_NODE_EVALUATOR(FPostProcessCameraNodeEvaluator)
+
+void FPostProcessCameraNodeEvaluator::OnInitialize(const FCameraNodeEvaluatorInitializeParams& Params, FCameraNodeEvaluationResult& OutResult)
+{
+	const UPostProcessCameraNode* PostProcessNode = GetCameraNodeAs<UPostProcessCameraNode>();
+	PostProcessBlendWeightReader.Initialize(PostProcessNode->PostProcessBlendWeight);
+}
+
+void FPostProcessCameraNodeEvaluator::OnRun(const FCameraNodeEvaluationParams& Params, FCameraNodeEvaluationResult& OutResult)
+{
+	FCameraPose& OutPose = OutResult.CameraPose;
+
+	float PostProcessBlendWeight = PostProcessBlendWeightReader.Get(OutResult.VariableTable);
+	if (PostProcessBlendWeight > 0)
+	{
+		const UPostProcessCameraNode* PostProcessNode = GetCameraNodeAs<UPostProcessCameraNode>();
+		OutResult.PostProcessSettings.Add(PostProcessNode->PostProcessSettings, PostProcessBlendWeight);
+	}
+}
+
+}  // namespace UE::Cameras
+
+FCameraNodeEvaluatorPtr UPostProcessCameraNode::OnBuildEvaluator(FCameraNodeEvaluatorBuilder& Builder) const
+{
+	using namespace UE::Cameras;
+	return Builder.BuildEvaluator<FPostProcessCameraNodeEvaluator>();
+}
+
diff --git a/Engine/Plugins/Cameras/GameplayCameras/Source/GameplayCameras/Public/Core/CameraNodeEvaluator.h b/Engine/Plugins/Cameras/GameplayCameras/Source/GameplayCameras/Public/Core/CameraNodeEvaluator.h
--- a/Engine/Plugins/Cameras/GameplayCameras/Source/GameplayCameras/Public/Core/CameraNodeEvaluator.h
+++ b/Engine/Plugins/Cameras/GameplayCameras/Source/GameplayCameras/Public/Core/CameraNodeEvaluator.h
@@ -2,13 +2,13 @@
 
 #pragma once
 
-#include "Core/CameraRigEvaluationInfo.h"
 #include "Core/CameraObjectRtti.h"
 #include "Core/CameraPose.h"
+#include "Core/CameraRigEvaluationInfo.h"
 #include "Core/CameraRigJoints.h"
 #include "Core/CameraVariableTable.h"
 #include "Core/ObjectChildrenView.h"
-#include "CoreTypes.h"
+#include "Core/PostProcessSettingsCollection.h"
 #include "Debug/CameraDebugBlockFwd.h"
 #include "Debug/RootCameraDebugBlock.h"
 #include "GameplayCameras.h"
@@ -152,8 +152,12 @@ struct GAMEPLAYCAMERAS_API FCameraNodeEvaluationResult
 	/** The list of joints in the current camera rig. */
 	FCameraRigJoints CameraRigJoints;
 
+	/** Post-process settings for the camera. */
+	FPostProcessSettingsCollection PostProcessSettings;
+
 	/** Whether the current frame is a camera cut. */
 	bool bIsCameraCut = false;
+
 	/** Whether this result is valid. */
 	bool bIsValid = false;
 
@@ -162,6 +166,12 @@ public:
 	/** Reset this result to its default (non-valid) state.  */
 	void Reset(bool bResetVariableTable);
 
+	/** Override this result with the given other result. */
+	void OverrideAll(const FCameraNodeEvaluationResult& OtherResult);
+
+	/** Interpolate this result towards the other given result. */
+	void LerpAll(const FCameraNodeEvaluationResult& ToResult, float BlendFactor);
+
 	/** Serializes this result to the given archive. */
 	void Serialize(FArchive& Ar);
 };
diff --git a/Engine/Plugins/Cameras/GameplayCameras/Source/GameplayCameras/Public/Core/CameraPose.h b/Engine/Plugins/Cameras/GameplayCameras/Source/GameplayCameras/Public/Core/CameraPose.h
--- a/Engine/Plugins/Cameras/GameplayCameras/Source/GameplayCameras/Public/Core/CameraPose.h
+++ b/Engine/Plugins/Cameras/GameplayCameras/Source/GameplayCameras/Public/Core/CameraPose.h
@@ -32,6 +32,7 @@ enum EAspectRatioAxisConstraint : int;
 
 #define UE_CAMERA_POSE_FOR_FLIPPING_PROPERTIES()\
 	UE_CAMERA_POSE_FOR_PROPERTY(bool, ConstrainAspectRatio)\
+	UE_CAMERA_POSE_FOR_PROPERTY(bool, OverrideAspectRatioAxisConstraint)\
 	UE_CAMERA_POSE_FOR_PROPERTY(EAspectRatioAxisConstraint, AspectRatioAxisConstraint)
 
 #define UE_CAMERA_POSE_FOR_ALL_PROPERTIES()\
@@ -169,6 +170,9 @@ public:
 	/** Computes the aspect ratio of a camera sensor. */
 	static double GetSensorAspectRatio(float SensorWidth, float SensorHeight);
 
+	/** Gets the default sensor size. */
+	static void GetDefaultSensorSize(float& OutSensorWidth, float& OutSensorHeight);
+
 public:
 
 	// Interpolation
@@ -257,7 +261,11 @@ private:
 	UPROPERTY()
 	bool ConstrainAspectRatio = false;
 
-	/** If ConstrainAspectRatio is false, how we should compute FieldOfView */
+	/** Whether to override the default aspect ratio axis constraint defined on the player controller */
+	UPROPERTY()
+	bool OverrideAspectRatioAxisConstraint = false;
+
+	/** If ConstrainAspectRatio is false and OverrideAspectRatioAxisConstraint is true, how we should compute FieldOfView */
 	UPROPERTY()
 	TEnumAsByte<EAspectRatioAxisConstraint> AspectRatioAxisConstraint = EAspectRatioAxisConstraint::AspectRatio_MaintainYFOV;
 
diff --git a/Engine/Plugins/Cameras/GameplayCameras/Source/GameplayCameras/Public/Core/CameraRigJoints.h b/Engine/Plugins/Cameras/GameplayCameras/Source/GameplayCameras/Public/Core/CameraRigJoints.h
--- a/Engine/Plugins/Cameras/GameplayCameras/Source/GameplayCameras/Public/Core/CameraRigJoints.h
+++ b/Engine/Plugins/Cameras/GameplayCameras/Source/GameplayCameras/Public/Core/CameraRigJoints.h
@@ -11,7 +11,7 @@ namespace UE::Cameras
 /**
  * A structure describing a joint in a camera rig.
  */
-struct FCameraRigJoint
+struct GAMEPLAYCAMERAS_API FCameraRigJoint
 {
 	/** The variable driving the rotation of this joint. */
 	FCameraVariableID VariableID;
@@ -19,12 +19,14 @@ struct FCameraRigJoint
 	FTransform3d Transform;
 };
 
+FArchive& operator<< (FArchive& Ar, FCameraRigJoint& RigJoint);
+
 /**
  * A structure describing the joints of a camera rig.
  * These joints allow for "manipulating" the rig, e.g. to make it point
  * towards a desired target or direction.
  */
-class FCameraRigJoints
+class GAMEPLAYCAMERAS_API FCameraRigJoints
 {
 public:
 
@@ -41,6 +43,8 @@ public:
 	/** Removes all previously added joints. */
 	void Reset();
 
+	void Serialize(FArchive& Ar);
+
 public:
 
 	/** Override the joints with another set of joints. */
diff --git a/Engine/Plugins/Cameras/GameplayCameras/Source/GameplayCameras/Public/Core/CameraVariableTable.h b/Engine/Plugins/Cameras/GameplayCameras/Source/GameplayCameras/Public/Core/CameraVariableTable.h
--- a/Engine/Plugins/Cameras/GameplayCameras/Source/GameplayCameras/Public/Core/CameraVariableTable.h
+++ b/Engine/Plugins/Cameras/GameplayCameras/Source/GameplayCameras/Public/Core/CameraVariableTable.h
@@ -66,14 +66,14 @@ struct FCameraVariableTableFlags
  * a map of metadata keyed by variable ID. A variable ID can be anything, but will
  * generally be the hash of the variable name.
  */
-class FCameraVariableTable
+class GAMEPLAYCAMERAS_API FCameraVariableTable
 {
 public:
 
 	FCameraVariableTable();
 	FCameraVariableTable(FCameraVariableTable&& Other);
 	FCameraVariableTable& operator=(FCameraVariableTable&& Other);
-	GAMEPLAYCAMERAS_API ~FCameraVariableTable();
+	~FCameraVariableTable();
 
 	FCameraVariableTable(const FCameraVariableTable&) = delete;
 	FCameraVariableTable& operator=(const FCameraVariableTable&) = delete;
@@ -86,7 +86,7 @@ public:
 	 * This may re-allocate the internal memory buffer. It's recommended to pre-compute
 	 * the allocation information needed for a table, and initialize it once.
 	 */
-	GAMEPLAYCAMERAS_API void AddVariable(const FCameraVariableDefinition& VariableDefinition);
+	void AddVariable(const FCameraVariableDefinition& VariableDefinition);
 
 public:
 
@@ -169,8 +169,8 @@ private:
 
 	void ReallocateBuffer(uint32 MinRequired = 0);
 
-	GAMEPLAYCAMERAS_API FEntry* FindEntry(FCameraVariableID VariableID);
-	GAMEPLAYCAMERAS_API const FEntry* FindEntry(FCameraVariableID VariableID) const;
+	FEntry* FindEntry(FCameraVariableID VariableID);
+	const FEntry* FindEntry(FCameraVariableID VariableID) const;
 
 	void InternalOverride(const FCameraVariableTable& OtherTable, ECameraVariableTableFilter Filter, const FCameraVariableTableFlags* InMask, bool bInvertMask, FCameraVariableTableFlags* OutMask);
 	void InternalLerp(const FCameraVariableTable& ToTable, ECameraVariableTableFilter Filter, float Factor, const FCameraVariableTableFlags* InMask, bool bInvertMask, FCameraVariableTableFlags* OutMask);
diff --git a/Engine/Plugins/Cameras/GameplayCameras/Source/GameplayCameras/Public/Core/CameraVariableTableFwd.h b/Engine/Plugins/Cameras/GameplayCameras/Source/GameplayCameras/Public/Core/CameraVariableTableFwd.h
--- a/Engine/Plugins/Cameras/GameplayCameras/Source/GameplayCameras/Public/Core/CameraVariableTableFwd.h
+++ b/Engine/Plugins/Cameras/GameplayCameras/Source/GameplayCameras/Public/Core/CameraVariableTableFwd.h
@@ -96,6 +96,12 @@ public:
 		return In.Value;
 	}
 
+	friend FArchive& operator<< (FArchive& Ar, FCameraVariableID& In)
+	{
+		Ar << In.Value;
+		return Ar;
+	}
+
 private:
 
 	FCameraVariableID(uint32 InValue) : Value(InValue) {}
diff --git a/Engine/Plugins/Cameras/GameplayCameras/Source/GameplayCameras/Public/Core/PostProcessSettingsCollection.h b/Engine/Plugins/Cameras/GameplayCameras/Source/GameplayCameras/Public/Core/PostProcessSettingsCollection.h
new file mode 100644
--- /dev/null
+++ b/Engine/Plugins/Cameras/GameplayCameras/Source/GameplayCameras/Public/Core/PostProcessSettingsCollection.h
@@ -0,0 +1,76 @@
+// Copyright Epic Games, Inc. All Rights Reserved.
+
+#pragma once
+
+#include "Engine/Scene.h"
+
+namespace UE::Cameras
+{
+
+/** An entry into an FPostProcessSettingsCollection. */
+struct FPostProcessSettingsCollectionEntry
+{
+	/** The post-process settings. */
+	FPostProcessSettings PostProcessSettings;
+	/** The blend factor for the post-process settings. */
+	float PostProcessBlendWeight = 1.f;
+};
+
+FArchive& operator<< (FArchive& Ar, FPostProcessSettingsCollectionEntry& Entry);
+
+/**
+ * A collection of post-process settings, with associated blend weights.
+ * TODO: the camera system currently supports only one post-process setting.
+ */
+struct GAMEPLAYCAMERAS_API FPostProcessSettingsCollection
+{
+	/** Returns whether this collection has any entries. */
+	bool IsEmpty() const
+	{
+		return Entries.IsEmpty();
+	}
+
+	/** Gets the entries in this collection. */
+	TArrayView<const FPostProcessSettingsCollectionEntry> GetEntries() const
+	{
+		return Entries;
+	}
+
+	/** Gets the entries in this collection. */
+	TArrayView<FPostProcessSettingsCollectionEntry> GetEntries() 
+	{
+		return Entries;
+	}
+
+	/** Adds a new entry in the collection. */
+	void Add(const FPostProcessSettings& InPostProcess, float InBlendWeight);
+
+	/** Removes all entries in the collection. */
+	void Reset();
+
+	/** 
+	 * Removes all entries in the collection, replacing them with the entries in the
+	 * given other collection.
+	 */
+	void OverrideAll(const FPostProcessSettingsCollection& OtherCollection);
+
+	/**
+	 * Interpolates the entries in this collection towards the entries in the given other
+	 * collection. This effectively merges the given collection into this collection, ramping
+	 * down the blend weight of existing entries, and ramping up the blend weight of the
+	 * new entries.
+	 */
+	void LerpAll(const FPostProcessSettingsCollection& ToCollection, float BlendFactor);
+
+	/** Serializes this collection into the given archive. */
+	void Serialize(FArchive& Ar);
+
+private:
+
+	using FEntry = FPostProcessSettingsCollectionEntry;
+	using FCollection = TArray<FEntry, TInlineAllocator<1>>;
+	FCollection Entries;
+};
+
+}  // namespace UE::Cameras
+
diff --git a/Engine/Plugins/Cameras/GameplayCameras/Source/GameplayCameras/Public/GameFramework/BlueprintCameraPose.h b/Engine/Plugins/Cameras/GameplayCameras/Source/GameplayCameras/Public/GameFramework/BlueprintCameraPose.h
--- a/Engine/Plugins/Cameras/GameplayCameras/Source/GameplayCameras/Public/GameFramework/BlueprintCameraPose.h
+++ b/Engine/Plugins/Cameras/GameplayCameras/Source/GameplayCameras/Public/GameFramework/BlueprintCameraPose.h
@@ -80,6 +80,10 @@ public:
 	UPROPERTY(BlueprintReadWrite, EditAnywhere, Category="Camera")
 	bool ConstrainAspectRatio = false;
 
+	/** Whether to override the default aspect ratio axis constraint defined on the player controller */
+	UPROPERTY()
+	bool OverrideAspectRatioAxisConstraint = false;
+
 	/** The aspect ratio axis constraint to use if aspect ratio is constrained. */
 	UPROPERTY(BlueprintReadWrite, EditAnywhere, Category="Camera")
 	TEnumAsByte<EAspectRatioAxisConstraint> AspectRatioAxisConstraint = EAspectRatioAxisConstraint::AspectRatio_MaintainYFOV;
diff --git a/Engine/Plugins/Cameras/GameplayCameras/Source/GameplayCameras/Public/Nodes/Common/FilmbackCameraNode.h b/Engine/Plugins/Cameras/GameplayCameras/Source/GameplayCameras/Public/Nodes/Common/FilmbackCameraNode.h
new file mode 100644
--- /dev/null
+++ b/Engine/Plugins/Cameras/GameplayCameras/Source/GameplayCameras/Public/Nodes/Common/FilmbackCameraNode.h
@@ -0,0 +1,54 @@
+// Copyright Epic Games, Inc. All Rights Reserved.
+
+#pragma once
+
+#include "Core/CameraNode.h"
+#include "Core/CameraParameters.h"
+
+#include "FilmbackCameraNode.generated.h"
+
+/**
+ * A camera node that configures the camera filmback.
+ */
+UCLASS(MinimalAPI, meta=(CameraNodeCategories="Common,Body"))
+class UFilmbackCameraNode : public UCameraNode
+{
+	GENERATED_BODY()
+
+public:
+
+	UFilmbackCameraNode(const FObjectInitializer& ObjectInitializer);
+
+protected:
+
+	// UCameraNode interface.
+	virtual FCameraNodeEvaluatorPtr OnBuildEvaluator(FCameraNodeEvaluatorBuilder& Builder) const override;
+
+public:
+
+	/** Horizontal size of filmback or digital sensor, in mm. */
+	UPROPERTY(EditAnywhere, Category="Filmback", meta=(ClampMin="0.001", ForceUnits=mm))
+	FFloatCameraParameter SensorWidth;
+
+	/** Vertical size of filmback or digital sensor, in mm. */
+	UPROPERTY(EditAnywhere, Category="Filmback", meta=(ClampMin="0.001", ForceUnits=mm))
+	FFloatCameraParameter SensorHeight;
+
+	/** Whether to constrain the aspect ratio of the evaluated camera. */
+	UPROPERTY(EditAnywhere, Category="Filmback")
+	FBooleanCameraParameter ConstrainAspectRatio = false;
+
+	/**
+	 * Whether to override the default aspect ratio axis constraint defined on the player controller.
+	 */
+	UPROPERTY(EditAnywhere, Category="Filmback")
+	FBooleanCameraParameter OverrideAspectRatioAxisConstraint = false;
+
+	/** 
+	 * Defines the axis along which to constrain the aspect ratio of the evaluated camera.
+	 * Only used when ConstrainAspectRatio is false and OverrideAspectRatioAxisConstraint is true.
+	 */
+	UPROPERTY(EditAnywhere, Category="Filmback")
+	TEnumAsByte<EAspectRatioAxisConstraint> AspectRatioAxisConstraint = EAspectRatioAxisConstraint::AspectRatio_MaintainYFOV;
+};
+
diff --git a/Engine/Plugins/Cameras/GameplayCameras/Source/GameplayCameras/Public/Nodes/Common/LensParametersCameraNode.h b/Engine/Plugins/Cameras/GameplayCameras/Source/GameplayCameras/Public/Nodes/Common/LensParametersCameraNode.h
--- a/Engine/Plugins/Cameras/GameplayCameras/Source/GameplayCameras/Public/Nodes/Common/LensParametersCameraNode.h
+++ b/Engine/Plugins/Cameras/GameplayCameras/Source/GameplayCameras/Public/Nodes/Common/LensParametersCameraNode.h
@@ -23,14 +23,6 @@ protected:
 
 public:
 
-	/** Horizontal size of filmback or digital sensor, in mm. */
-	UPROPERTY(EditAnywhere, Category="Filmback", meta=(ClampMin="0.001", ForceUnits=mm))
-	FFloatCameraParameter SensorWidth;
-
-	/** Vertical size of filmback or digital sensor, in mm. */
-	UPROPERTY(EditAnywhere, Category="Filmback", meta=(ClampMin="0.001", ForceUnits=mm))
-	FFloatCameraParameter SensorHeight;
-
 	/** Manually-controlled focus distance (manual focus mode only) */
 	UPROPERTY(EditAnywhere, Category="Lens Parameters", meta=(Units=cm))
 	FFloatCameraParameter FocusDistance;
@@ -44,4 +36,3 @@ public:
 	FFloatCameraParameter Aperture;
 };
 
-
diff --git a/Engine/Plugins/Cameras/GameplayCameras/Source/GameplayCameras/Public/Nodes/Common/PostProcessCameraNode.h b/Engine/Plugins/Cameras/GameplayCameras/Source/GameplayCameras/Public/Nodes/Common/PostProcessCameraNode.h
new file mode 100644
--- /dev/null
+++ b/Engine/Plugins/Cameras/GameplayCameras/Source/GameplayCameras/Public/Nodes/Common/PostProcessCameraNode.h
@@ -0,0 +1,32 @@
+// Copyright Epic Games, Inc. All Rights Reserved.
+
+#pragma once
+
+#include "Core/CameraNode.h"
+#include "Core/CameraParameters.h"
+#include "Engine/Scene.h"
+
+#include "PostProcessCameraNode.generated.h"
+
+/**
+ * A camera node that adds post-process settings on the evaluated camera.
+ */
+UCLASS(MinimalAPI, meta=(CameraNodeCategories="Common,Rendering"))
+class UPostProcessCameraNode : public UCameraNode
+{
+	GENERATED_BODY()
+
+protected:
+
+	// UCameraNode interface.
+	virtual FCameraNodeEvaluatorPtr OnBuildEvaluator(FCameraNodeEvaluatorBuilder& Builder) const override;
+
+public:
+
+	UPROPERTY(EditAnywhere, Category="Rendering")
+	FPostProcessSettings PostProcessSettings;
+
+	UPROPERTY(EditAnywhere, Category="Rendering")
+	FFloatCameraParameter PostProcessBlendWeight = 1.f;
+};
+
