Cameras: hide the scratch evaluation layer from the UI

Also remove the user layers. This should be replaced with a user-defined enum later.

diff --git a/Engine/Plugins/Cameras/GameplayCameras/Source/GameplayCameras/Public/Core/RootCameraNode.h b/Engine/Plugins/Cameras/GameplayCameras/Source/GameplayCameras/Public/Core/RootCameraNode.h
--- a/Engine/Plugins/Cameras/GameplayCameras/Source/GameplayCameras/Public/Core/RootCameraNode.h
+++ b/Engine/Plugins/Cameras/GameplayCameras/Source/GameplayCameras/Public/Core/RootCameraNode.h
@@ -20,10 +20,7 @@ enum class ECameraRigLayer : uint8
 	Main UMETA(DisplayName="Main Layer"),
 	Global UMETA(DisplayName="Global Layer"),
 	Visual UMETA(DisplayName="Visual Layer"),
-	ScratchMain UMETA(DisplayName="Scratch Main Layer"),
-	User0,
-	User1,
-	User2
+	ScratchMain UMETA(Hidden)
 };
 ENUM_CLASS_FLAGS(ECameraRigLayer)
 
