Cameras: add more information to camera pose debug display

diff --git a/Engine/Plugins/Cameras/GameplayCameras/Source/GameplayCameras/Private/Debug/CameraPoseDebugBlock.cpp b/Engine/Plugins/Cameras/GameplayCameras/Source/GameplayCameras/Private/Debug/CameraPoseDebugBlock.cpp
--- a/Engine/Plugins/Cameras/GameplayCameras/Source/GameplayCameras/Private/Debug/CameraPoseDebugBlock.cpp
+++ b/Engine/Plugins/Cameras/GameplayCameras/Source/GameplayCameras/Private/Debug/CameraPoseDebugBlock.cpp
@@ -63,6 +63,10 @@ void FCameraPoseDebugBlock::OnDebugDraw(const FCameraDebugBlockDrawParams& Param
 	}
 	UE_CAMERA_POSE_FOR_ALL_PROPERTIES()
 #undef UE_CAMERA_POSE_FOR_PROPERTY
+
+	Renderer.SetTextColor(Colors.Default);
+	Renderer.AddText(TEXT("Effective FOV  : %f\n"), CameraPose.GetEffectiveFieldOfView());
+	Renderer.AddText(TEXT("Effective Aspect Ratio  : %f\n"), CameraPose.GetSensorAspectRatio());
 }
 
 void FCameraPoseDebugBlock::OnSerialize(FArchive& Ar)
