Cameras: consolidate gameplay tag transition conditions into one

diff --git a/Engine/Plugins/Cameras/GameplayCameras/Source/GameplayCameras/Private/Transitions/GameplayTagTransitionConditions.cpp b/Engine/Plugins/Cameras/GameplayCameras/Source/GameplayCameras/Private/Transitions/GameplayTagTransitionConditions.cpp
--- a/Engine/Plugins/Cameras/GameplayCameras/Source/GameplayCameras/Private/Transitions/GameplayTagTransitionConditions.cpp
+++ b/Engine/Plugins/Cameras/GameplayCameras/Source/GameplayCameras/Private/Transitions/GameplayTagTransitionConditions.cpp
@@ -6,31 +6,46 @@
 
 #include UE_INLINE_GENERATED_CPP_BY_NAME(GameplayTagTransitionConditions)
 
-bool UPreviousGameplayTagTransitionCondition::OnTransitionMatches(const FCameraRigTransitionConditionMatchParams& Params) const
+bool UGameplayTagTransitionCondition::OnTransitionMatches(const FCameraRigTransitionConditionMatchParams& Params) const
 {
-	if (Params.FromCameraRig)
+	bool bPreviousMatches = true;
+
+	if (!PreviousGameplayTagQuery.IsEmpty())
 	{
-		FGameplayTagContainer TagContainer;
-		Params.FromCameraRig->GetOwnedGameplayTags(TagContainer);
-		if (TagContainer.MatchesQuery(GameplayTagQuery))
+		if (Params.FromCameraRig)
+		{
+			FGameplayTagContainer TagContainer;
+			Params.FromCameraRig->GetOwnedGameplayTags(TagContainer);
+			if (TagContainer.MatchesQuery(PreviousGameplayTagQuery))
+			{
+				bPreviousMatches = true;
+			}
+		}
+		else
 		{
-			return true;
+			bPreviousMatches = false;
 		}
 	}
-	return false;
-}
 
-bool UNextGameplayTagTransitionCondition::OnTransitionMatches(const FCameraRigTransitionConditionMatchParams& Params) const
-{
-	if (Params.ToCameraRig)
+	bool bNextMatches = true;
+
+	if (!NextGameplayTagQuery.IsEmpty())
 	{
-		FGameplayTagContainer TagContainer;
-		Params.ToCameraRig->GetOwnedGameplayTags(TagContainer);
-		if (TagContainer.MatchesQuery(GameplayTagQuery))
+		if (Params.ToCameraRig)
+		{
+			FGameplayTagContainer TagContainer;
+			Params.ToCameraRig->GetOwnedGameplayTags(TagContainer);
+			if (TagContainer.MatchesQuery(NextGameplayTagQuery))
+			{
+				bNextMatches = true;
+			}
+		}
+		else
 		{
-			return true;
+			bNextMatches = false;
 		}
 	}
-	return false;
+
+	return bPreviousMatches && bNextMatches;
 }
 
diff --git a/Engine/Plugins/Cameras/GameplayCameras/Source/GameplayCameras/Public/Transitions/GameplayTagTransitionConditions.h b/Engine/Plugins/Cameras/GameplayCameras/Source/GameplayCameras/Public/Transitions/GameplayTagTransitionConditions.h
--- a/Engine/Plugins/Cameras/GameplayCameras/Source/GameplayCameras/Public/Transitions/GameplayTagTransitionConditions.h
+++ b/Engine/Plugins/Cameras/GameplayCameras/Source/GameplayCameras/Public/Transitions/GameplayTagTransitionConditions.h
@@ -8,11 +8,11 @@
 #include "GameplayTagTransitionConditions.generated.h"
 
 /**
- * A transition condition that matches the gameplay tags on the previous
- * camera rig and asset.
+ * A transition condition that matches the gameplay tags on the previous and next
+ * camera rigs and assets. Both queries need to pass. Empty queries pass by default.
  */
 UCLASS(MinimalAPI)
-class UPreviousGameplayTagTransitionCondition 
+class UGameplayTagTransitionCondition
 	: public UCameraRigTransitionCondition
 {
 	GENERATED_BODY()
@@ -21,29 +21,11 @@ public:
 
 	/** The gameplay tags to look for on the previous camera rig/asset. */
 	UPROPERTY(EditAnywhere, Category=Transition)
-	FGameplayTagQuery GameplayTagQuery;
-
-protected:
-
-	// UCameraRigTransitionCondition interface.
-	virtual bool OnTransitionMatches(const FCameraRigTransitionConditionMatchParams& Params) const override;
-};
-
-/**
- * A transition condition that matches the gameplay tags on the next
- * camera rig and asset.
- */
-UCLASS(MinimalAPI)
-class UNextGameplayTagTransitionCondition 
-	: public UCameraRigTransitionCondition
-{
-	GENERATED_BODY()
-
-public:
+	FGameplayTagQuery PreviousGameplayTagQuery;
 
 	/** The gameplay tags to look for on the next camera rig/asset. */
 	UPROPERTY(EditAnywhere, Category=Transition)
-	FGameplayTagQuery GameplayTagQuery;
+	FGameplayTagQuery NextGameplayTagQuery;
 
 protected:
 
