Cameras: add parameter blending information to know if the blend has finished

diff --git a/Engine/Plugins/Cameras/GameplayCameras/Source/GameplayCameras/Private/Nodes/Blends/OrbitBlendCameraNode.cpp b/Engine/Plugins/Cameras/GameplayCameras/Source/GameplayCameras/Private/Nodes/Blends/OrbitBlendCameraNode.cpp
--- a/Engine/Plugins/Cameras/GameplayCameras/Source/GameplayCameras/Private/Nodes/Blends/OrbitBlendCameraNode.cpp
+++ b/Engine/Plugins/Cameras/GameplayCameras/Source/GameplayCameras/Private/Nodes/Blends/OrbitBlendCameraNode.cpp
@@ -91,6 +91,9 @@ void FOrbitBlendCameraNodeEvaluator::OnBlendParameters(const FCameraNodePreBlend
 	{
 		const FCameraVariableTable& ChildVariableTable(Params.ChildVariableTable);
 		OutResult.VariableTable.Override(ChildVariableTable, ECameraVariableTableFilter::Input);
+
+		OutResult.bIsBlendFinished = true;
+		OutResult.bIsBlendFull = true;
 	}
 }
 
diff --git a/Engine/Plugins/Cameras/GameplayCameras/Source/GameplayCameras/Private/Nodes/Blends/PopBlendCameraNode.cpp b/Engine/Plugins/Cameras/GameplayCameras/Source/GameplayCameras/Private/Nodes/Blends/PopBlendCameraNode.cpp
--- a/Engine/Plugins/Cameras/GameplayCameras/Source/GameplayCameras/Private/Nodes/Blends/PopBlendCameraNode.cpp
+++ b/Engine/Plugins/Cameras/GameplayCameras/Source/GameplayCameras/Private/Nodes/Blends/PopBlendCameraNode.cpp
@@ -27,6 +27,9 @@ void FPopBlendCameraNodeEvaluator::OnBlendParameters(const FCameraNodePreBlendPa
 {
 	const FCameraVariableTable& ChildVariableTable(Params.ChildVariableTable);
 	OutResult.VariableTable.Override(ChildVariableTable, ECameraVariableTableFilter::Input);
+
+	OutResult.bIsBlendFull = true;
+	OutResult.bIsBlendFinished = true;
 }
 
 void FPopBlendCameraNodeEvaluator::OnBlendResults(const FCameraNodeBlendParams& Params, FCameraNodeBlendResult& OutResult)
diff --git a/Engine/Plugins/Cameras/GameplayCameras/Source/GameplayCameras/Private/Nodes/Blends/SimpleBlendCameraNode.cpp b/Engine/Plugins/Cameras/GameplayCameras/Source/GameplayCameras/Private/Nodes/Blends/SimpleBlendCameraNode.cpp
--- a/Engine/Plugins/Cameras/GameplayCameras/Source/GameplayCameras/Private/Nodes/Blends/SimpleBlendCameraNode.cpp
+++ b/Engine/Plugins/Cameras/GameplayCameras/Source/GameplayCameras/Private/Nodes/Blends/SimpleBlendCameraNode.cpp
@@ -29,7 +29,9 @@ void FSimpleBlendCameraNodeEvaluator::OnRun(const FCameraNodeEvaluationParams& P
 void FSimpleBlendCameraNodeEvaluator::OnBlendParameters(const FCameraNodePreBlendParams& Params, FCameraNodePreBlendResult& OutResult)
 {
 	const FCameraVariableTable& ChildVariableTable(Params.ChildVariableTable);
-	OutResult.VariableTable.Lerp(ChildVariableTable, ECameraVariableTableFilter::Input, BlendFactor);
+
+	OutResult.bIsBlendFull = BlendFactor >= 1.f;
+	OutResult.bIsBlendFinished = bIsBlendFinished;
 }
 
 void FSimpleBlendCameraNodeEvaluator::OnBlendResults(const FCameraNodeBlendParams& Params, FCameraNodeBlendResult& OutResult)
diff --git a/Engine/Plugins/Cameras/GameplayCameras/Source/GameplayCameras/Public/Core/BlendCameraNode.h b/Engine/Plugins/Cameras/GameplayCameras/Source/GameplayCameras/Public/Core/BlendCameraNode.h
--- a/Engine/Plugins/Cameras/GameplayCameras/Source/GameplayCameras/Public/Core/BlendCameraNode.h
+++ b/Engine/Plugins/Cameras/GameplayCameras/Source/GameplayCameras/Public/Core/BlendCameraNode.h
@@ -52,6 +52,12 @@ struct FCameraNodePreBlendResult
 
 	/** The variable table to received blended parameters. */
 	FCameraVariableTable& VariableTable;
+
+	/** Whether the blend has reached 100%. */
+	bool bIsBlendFull = false;
+
+	/** Whether the blend is finished. */
+	bool bIsBlendFinished = false;
 };
 
 /**
