Cameras: make rig interface parameters' private variables into input variables

This will let us pre-blend them, which is important if they are driving input nodes.

diff --git a/Engine/Plugins/Cameras/GameplayCameras/Source/GameplayCameras/Private/Core/CameraRigAssetBuilder.cpp b/Engine/Plugins/Cameras/GameplayCameras/Source/GameplayCameras/Private/Core/CameraRigAssetBuilder.cpp
--- a/Engine/Plugins/Cameras/GameplayCameras/Source/GameplayCameras/Private/Core/CameraRigAssetBuilder.cpp
+++ b/Engine/Plugins/Cameras/GameplayCameras/Source/GameplayCameras/Private/Core/CameraRigAssetBuilder.cpp
@@ -205,7 +205,8 @@ VariableAssetType* MakeOrRenamePrivateVariable(
 		PrivateVariable = NewObject<VariableAssetType>(Builder.CameraRig, FName(*VariableName), RF_Transactional);
 	}
 
-	// Make sure it's a private variable.
+	// Make sure it's a private input variable.
+	PrivateVariable->bIsInput = true;
 	PrivateVariable->bIsPrivate = true;
 	PrivateVariable->bAutoReset = false;
 
diff --git a/Engine/Plugins/Cameras/GameplayCameras/Source/GameplayCameras/Private/Core/CameraVariableAssets.cpp b/Engine/Plugins/Cameras/GameplayCameras/Source/GameplayCameras/Private/Core/CameraVariableAssets.cpp
--- a/Engine/Plugins/Cameras/GameplayCameras/Source/GameplayCameras/Private/Core/CameraVariableAssets.cpp
+++ b/Engine/Plugins/Cameras/GameplayCameras/Source/GameplayCameras/Private/Core/CameraVariableAssets.cpp
@@ -21,6 +21,7 @@ FCameraVariableDefinition UCameraVariableAsset::GetVariableDefinition() const
 	VariableDefinition.VariableID = GetVariableID();
 	VariableDefinition.VariableType = GetVariableType();
 	VariableDefinition.bIsPrivate = bIsPrivate;
+	VariableDefinition.bIsInput = bIsInput;
 #if WITH_EDITORONLY_DATA
 	VariableDefinition.VariableName = GetName();
 #endif
diff --git a/Engine/Plugins/Cameras/GameplayCameras/Source/GameplayCameras/Public/Core/CameraVariableAssets.h b/Engine/Plugins/Cameras/GameplayCameras/Source/GameplayCameras/Public/Core/CameraVariableAssets.h
--- a/Engine/Plugins/Cameras/GameplayCameras/Source/GameplayCameras/Public/Core/CameraVariableAssets.h
+++ b/Engine/Plugins/Cameras/GameplayCameras/Source/GameplayCameras/Public/Core/CameraVariableAssets.h
@@ -60,6 +60,10 @@ public:
 	UPROPERTY()
 	bool bIsPrivate = false;
 
+	/** Whether this variable is an input variable that gets blended before node evaluators are run. */
+	UPROPERTY()
+	bool bIsInput = false;
+
 private:
 
 	UPROPERTY()
