Cameras: new utility class that builds/caches a flattened version of a camera node tree, use it in the camera rig builder

diff --git a/Engine/Plugins/Cameras/GameplayCameras/Source/GameplayCameras/Private/Core/CameraNodeHierarchy.cpp b/Engine/Plugins/Cameras/GameplayCameras/Source/GameplayCameras/Private/Core/CameraNodeHierarchy.cpp
new file mode 100644
--- /dev/null
+++ b/Engine/Plugins/Cameras/GameplayCameras/Source/GameplayCameras/Private/Core/CameraNodeHierarchy.cpp
@@ -0,0 +1,84 @@
+// Copyright Epic Games, Inc. All Rights Reserved.
+
+#include "Core/CameraNodeHierarchy.h"
+
+#include "Core/CameraRigAsset.h"
+
+namespace UE::Cameras
+{
+
+FCameraNodeHierarchy::FCameraNodeHierarchy()
+{
+}
+
+FCameraNodeHierarchy::FCameraNodeHierarchy(UCameraRigAsset* InCameraRig)
+{
+	Build(InCameraRig);
+}
+
+TArrayView<UCameraNode* const> FCameraNodeHierarchy::GetFlattenedHierarchy() const
+{
+	return FlattenedHierarchy;
+}
+
+int32 FCameraNodeHierarchy::Num() const
+{
+	return FlattenedHierarchy.Num();
+}
+
+void FCameraNodeHierarchy::Build(UCameraRigAsset* InCameraRig)
+{
+	Build(InCameraRig ? InCameraRig->RootNode : nullptr);
+}
+
+void FCameraNodeHierarchy::Build(UCameraNode* InRootCameraNode)
+{
+	Reset();
+
+	if (!InRootCameraNode)
+	{
+		return;
+	}
+
+	TArray<UCameraNode*> NodeStack;
+	NodeStack.Add(InRootCameraNode);
+	while (!NodeStack.IsEmpty())
+	{
+		UCameraNode* CurrentNode = NodeStack.Pop();
+		FlattenedHierarchy.Add(CurrentNode);
+
+		FCameraNodeChildrenView CurrentChildren = CurrentNode->GetChildren();
+		for (UCameraNode* ChildNode : ReverseIterate(CurrentChildren))
+		{
+			if (ChildNode)
+			{
+				NodeStack.Add(ChildNode);
+			}
+		}
+	}
+}
+
+void FCameraNodeHierarchy::Reset()
+{
+	FlattenedHierarchy.Reset();
+}
+
+#if WITH_EDITORONLY_DATA
+
+bool FCameraNodeHierarchy::FindMissingConnectableObjects(TArrayView<UObject* const> ConnectableObjects, TSet<UObject*>& OutMissingObjects)
+{
+	TSet<UObject*> ConnectableObjectsSet(ConnectableObjects);
+	return FindMissingConnectableObjects(ConnectableObjectsSet, OutMissingObjects);
+}
+
+bool FCameraNodeHierarchy::FindMissingConnectableObjects(const TSet<UObject*> ConnectableObjectsSet, TSet<UObject*>& OutMissingObjects)
+{
+	TSet<UObject*> FlattenedHierarchySet(MakeArrayView((UObject**)FlattenedHierarchy.GetData(), FlattenedHierarchy.Num()));
+	OutMissingObjects = FlattenedHierarchySet.Difference(ConnectableObjectsSet);
+	return !OutMissingObjects.IsEmpty();
+}
+
+#endif  // WITH_EDITORONLY_DATA
+
+}  // namespace UE::Cameras
+
diff --git a/Engine/Plugins/Cameras/GameplayCameras/Source/GameplayCameras/Private/Core/CameraRigAssetBuilder.cpp b/Engine/Plugins/Cameras/GameplayCameras/Source/GameplayCameras/Private/Core/CameraRigAssetBuilder.cpp
--- a/Engine/Plugins/Cameras/GameplayCameras/Source/GameplayCameras/Private/Core/CameraRigAssetBuilder.cpp
+++ b/Engine/Plugins/Cameras/GameplayCameras/Source/GameplayCameras/Private/Core/CameraRigAssetBuilder.cpp
@@ -329,7 +329,7 @@ void FCameraRigAssetBuilder::BuildCameraRigImpl()
 		return;
 	}
 
-	FlattenCameraNodeHierarchy();
+	BuildCameraNodeHierarchy();
 
 	CallPreBuild();
 
@@ -340,38 +340,19 @@ void FCameraRigAssetBuilder::BuildCameraRigImpl()
 	BuildAllocationInfo();
 }
 
-void FCameraRigAssetBuilder::FlattenCameraNodeHierarchy()
+void FCameraRigAssetBuilder::BuildCameraNodeHierarchy()
 {
 	// Build a flat list of the camera rig's node hierarchy. It's easier to iterate during
 	// our build process.
-	FlattenedNodes.Reset();
-
-	TArray<UCameraNode*> NodeStack;
-	NodeStack.Add(CameraRig->RootNode);
-	while (!NodeStack.IsEmpty())
-	{
-		UCameraNode* CurrentNode = NodeStack.Pop();
-		FlattenedNodes.Add(CurrentNode);
-
-		FCameraNodeChildrenView CurrentChildren = CurrentNode->GetChildren();
-		for (UCameraNode* Child : ReverseIterate(CurrentChildren))
-		{
-			if (Child)
-			{
-				NodeStack.Add(Child);
-			}
-		}
-	}
+	CameraNodeHierarchy.Build(CameraRig);
 
 #if WITH_EDITORONLY_DATA
 	// Check that all the camera nodes that are in the tree are also inside the camera 
 	// rig's AllNodeTreeObjects. This shouldn't happen unless someone added camera nodes
 	// directly via C++, or if there's a bug in the camera rig editor code, so emit a
 	// warning if that happens.
-	TSet<UObject*> FlattenedNodesSet(MakeArrayView((UObject**)FlattenedNodes.GetData(), FlattenedNodes.Num()));
-	TSet<UObject*> AllNodeTreeObjectsSet(ObjectPtrDecay(CameraRig->AllNodeTreeObjects));
-	TSet<UObject*> MissingNodeTreeObjects = FlattenedNodesSet.Difference(AllNodeTreeObjectsSet);
-	if (!MissingNodeTreeObjects.IsEmpty())
+	TSet<UObject*> MissingNodeTreeObjects;
+	if (CameraNodeHierarchy.FindMissingConnectableObjects(ObjectPtrDecay(CameraRig->AllNodeTreeObjects), MissingNodeTreeObjects))
 	{
 		BuildLog.AddMessage(EMessageSeverity::Warning, 
 				FText::Format(
@@ -385,7 +366,7 @@ void FCameraRigAssetBuilder::FlattenCameraNodeHierarchy()
 
 void FCameraRigAssetBuilder::CallPreBuild()
 {
-	for (UCameraNode* CameraNode : FlattenedNodes)
+	for (UCameraNode* CameraNode : CameraNodeHierarchy.GetFlattenedHierarchy())
 	{
 		CameraNode->PreBuild(BuildLog);
 	}
@@ -403,7 +384,7 @@ void FCameraRigAssetBuilder::GatherOldDrivenParameters()
 	// Note that parameters driven by user-defined variables are left alone.
 
 	TSet<UCameraVariableAsset*> GatheredVariables;
-	TSet<UCameraNode*> CameraNodesToGather(FlattenedNodes);
+	TSet<UCameraNode*> CameraNodesToGather(CameraNodeHierarchy.GetFlattenedHierarchy());
 
 	// Start by going through all interface parameters, remembering what private variable 
 	// they were associated with originally. Also collect that private variable to be
@@ -427,7 +408,7 @@ void FCameraRigAssetBuilder::GatherOldDrivenParameters()
 				*GetPathNameSafe(CameraRig));
 	}
 
-	// Next go through all the camera nodes we know of. Nodes in FlattenedNodes are the ones
+	// Next go through all the camera nodes we know of. Nodes in CameraNodeHierarchy are the ones
 	// connected to the camera rig's root node, so we are missing nodes that were disconnected
 	// since the last build. We could use AllNodeTreeObjects for that, but it only exists in 
 	// editor builds, and we don't want to rely on unit tests or runtime data manipulation to 
@@ -443,7 +424,7 @@ void FCameraRigAssetBuilder::GatherOldDrivenParameters()
 					CameraNodesToGather.Add(CameraNode);
 				}
 			});
-	const int32 NumStrayCameraNodes = (CameraNodesToGather.Num() - FlattenedNodes.Num());
+	const int32 NumStrayCameraNodes = (CameraNodesToGather.Num() - CameraNodeHierarchy.Num());
 	if (NumStrayCameraNodes > 0)
 	{
 		UE_LOG(LogCameraSystem, Verbose, TEXT("Collected %d stray camera nodes while building camera rig '%s'."),
@@ -823,7 +804,7 @@ void FCameraRigAssetBuilder::BuildAllocationInfo()
 	Storage.GetAllocationInfo(AllocationInfo.EvaluatorInfo);
 
 	// Compute the allocation info for camera variables.
-	for (UCameraNode* CameraNode : FlattenedNodes)
+	for (UCameraNode* CameraNode : CameraNodeHierarchy.GetFlattenedHierarchy())
 	{
 		BuildAllocationInfo(CameraNode);
 	}
diff --git a/Engine/Plugins/Cameras/GameplayCameras/Source/GameplayCameras/Public/Core/CameraNodeHierarchy.h b/Engine/Plugins/Cameras/GameplayCameras/Source/GameplayCameras/Public/Core/CameraNodeHierarchy.h
new file mode 100644
--- /dev/null
+++ b/Engine/Plugins/Cameras/GameplayCameras/Source/GameplayCameras/Public/Core/CameraNodeHierarchy.h
@@ -0,0 +1,70 @@
+// Copyright Epic Games, Inc. All Rights Reserved.
+
+#pragma once
+
+#include "Containers/Array.h"
+#include "Core/CameraNode.h"
+
+class UCameraRigAsset;
+
+namespace UE::Cameras
+{
+
+/**
+ * A utility class that stores a flattened hierarchy of camera nodes.
+ * Unconnected camera nodes aren't included, of course.
+ */
+class GAMEPLAYCAMERAS_API FCameraNodeHierarchy
+{
+public:
+
+	/** Build an empty hierarchy. */
+	FCameraNodeHierarchy();
+	/** Build a hierarchy starting from the given camera rig's root node. */
+	FCameraNodeHierarchy(UCameraRigAsset* InCameraRig);
+	/** Build a hierarchy starting from the given root node. */
+	FCameraNodeHierarchy(UCameraNode* InRootCameraNode);
+
+	/** Get the list of camera nodes in depth-first order. */
+	TArrayView<UCameraNode* const> GetFlattenedHierarchy() const;
+
+	/** Returns the number of camera nodes in this hierarchy. */
+	int32 Num() const;
+
+public:
+
+	/** Build a hierarchy starting from the given camera rig's root node. */
+	void Build(UCameraRigAsset* InCameraRig);
+	/** Build a hierarchy starting from the given root node. */
+	void Build(UCameraNode* InRootCameraNode);
+	/** Resets this object to an empty hierarchy. */
+	void Reset();
+
+public:
+
+	/** Executes the given predicate on each camera node in depth-first order. */
+	template<typename PredicateClass>
+	void ForEach(PredicateClass&& Predicate)
+	{
+		for (UCameraNode* Node : FlattenedHierarchy)
+		{
+			Predicate(Node);
+		}
+	}
+
+public:
+
+	// Internal API.
+
+#if WITH_EDITORONLY_DATA
+	bool FindMissingConnectableObjects(TArrayView<UObject* const> ConnectableObjects, TSet<UObject*>& OutMissingObjects);
+	bool FindMissingConnectableObjects(const TSet<UObject*> ConnectableObjectsSet, TSet<UObject*>& OutMissingObjects);
+#endif  // WITH_EDITORONLY_DATA
+
+private:
+
+	TArray<UCameraNode*> FlattenedHierarchy;
+};
+
+}  // namespace UE::Cameras
+
diff --git a/Engine/Plugins/Cameras/GameplayCameras/Source/GameplayCameras/Public/Core/CameraRigAssetBuilder.h b/Engine/Plugins/Cameras/GameplayCameras/Source/GameplayCameras/Public/Core/CameraRigAssetBuilder.h
--- a/Engine/Plugins/Cameras/GameplayCameras/Source/GameplayCameras/Public/Core/CameraRigAssetBuilder.h
+++ b/Engine/Plugins/Cameras/GameplayCameras/Source/GameplayCameras/Public/Core/CameraRigAssetBuilder.h
@@ -4,6 +4,7 @@
 
 #include "Core/CameraBuildLog.h"
 #include "Core/CameraNodeEvaluatorFwd.h"
+#include "Core/CameraNodeHierarchy.h"
 #include "Core/CameraRigAsset.h"
 #include "Core/CameraVariableTableFwd.h"
 #include "CoreTypes.h"
@@ -50,7 +51,7 @@ private:
 
 	void BuildCameraRigImpl();
 
-	void FlattenCameraNodeHierarchy();
+	void BuildCameraNodeHierarchy();
 
 	void CallPreBuild();
 
@@ -74,7 +75,7 @@ private:
 
 	UCameraRigAsset* CameraRig = nullptr;
 
-	TArray<UCameraNode*> FlattenedNodes;
+	FCameraNodeHierarchy CameraNodeHierarchy;
 
 	using FDrivenParameterKey = TTuple<FStructProperty*, UCameraNode*>;
 	TMap<FDrivenParameterKey, UCameraVariableAsset*> OldDrivenParameters;
