Cameras: refactor and optimize which camera node evaluator APIs are recursive or not, and how they're called

- Add a node evaluator hierarchy class that builds/caches a flattened version of a node evaluator tree.
- Auto-build this flattened hierarchy during node evaluation initialization.
- Only Initialize and AddReferencedObjects are recursive now.
- Use this utility class to call functions "recursively" on the tree. This should be more efficient since the hierarchy is already built, which skips many calls to virtual functions like GetChildren().
- Use these call wrappers in the blend stack and the IK aiming classes.
- In order to get the hierarchy of the entire camera system's evaluator tree for IK purposes, add an API on the root node evaluator.

diff --git a/Engine/Plugins/Cameras/GameplayCameras/Source/GameplayCameras/Private/Core/BlendStackCameraNode.cpp b/Engine/Plugins/Cameras/GameplayCameras/Source/GameplayCameras/Private/Core/BlendStackCameraNode.cpp
--- a/Engine/Plugins/Cameras/GameplayCameras/Source/GameplayCameras/Private/Core/BlendStackCameraNode.cpp
+++ b/Engine/Plugins/Cameras/GameplayCameras/Source/GameplayCameras/Private/Core/BlendStackCameraNode.cpp
@@ -152,16 +152,12 @@ bool FBlendStackCameraNodeEvaluator::InitializeEntry(
 	NewEntry.Result.VariableTable.Initialize(CameraRig->AllocationInfo.VariableTableInfo);
 
 	// Initialize the node evaluators.
-	FCameraNodeEvaluatorInitializeParams InitParams;
+	FCameraNodeEvaluatorInitializeParams InitParams(&NewEntry.EvaluatorHierarchy);
 	InitParams.Evaluator = Evaluator;
 	InitParams.EvaluationContext = EvaluationContext;
 	InitParams.LastActiveCameraRigInfo = GetActiveCameraRigEvaluationInfo();
 	RootEvaluator->Initialize(InitParams, NewEntry.Result);
 
-	// Gather blended parameter evaluators.
-	NewEntry.ParameterEvaluators.Reset();
-	GatherEntryParameterEvaluators(RootEvaluator, NewEntry.ParameterEvaluators);
-
 	// Wrap up!
 	NewEntry.EvaluationContext = EvaluationContext;
 	NewEntry.CameraRig = CameraRig;
@@ -177,7 +173,7 @@ void FBlendStackCameraNodeEvaluator::FreezeEntry(FCameraRigEntry& Entry)
 	// Deallocate our node evaluators and clear any pointers we kept to them.
 	Entry.EvaluatorStorage.DestroyEvaluatorTree(true);
 	Entry.RootEvaluator = nullptr;
-	Entry.ParameterEvaluators.Reset();
+	Entry.EvaluatorHierarchy.Reset();
 
 	Entry.EvaluationContext.Reset();
 	
@@ -350,10 +346,11 @@ void FBlendStackCameraNodeEvaluator::OnRun(const FCameraNodeEvaluationParams& Pa
 			FCameraBlendedParameterUpdateParams InputParams(CurParams, CurResult.CameraPose);
 			FCameraBlendedParameterUpdateResult InputResult(CurResult.VariableTable);
 
-			for (FCameraNodeEvaluator* ParameterEvaluator : Entry.ParameterEvaluators)
-			{
-				ParameterEvaluator->UpdateParameters(InputParams, InputResult);
-			}
+			Entry.EvaluatorHierarchy.ForEachEvaluator(ECameraNodeEvaluatorFlags::NeedsParameterUpdate,
+					[&InputParams, &InputResult](FCameraNodeEvaluator* ParameterEvaluator)
+					{
+						ParameterEvaluator->UpdateParameters(InputParams, InputResult);
+					});
 
 			Entry.bInputRunThisFrame = true;
 		}
diff --git a/Engine/Plugins/Cameras/GameplayCameras/Source/GameplayCameras/Private/Core/CameraIKAim.cpp b/Engine/Plugins/Cameras/GameplayCameras/Source/GameplayCameras/Private/Core/CameraIKAim.cpp
--- a/Engine/Plugins/Cameras/GameplayCameras/Source/GameplayCameras/Private/Core/CameraIKAim.cpp
+++ b/Engine/Plugins/Cameras/GameplayCameras/Source/GameplayCameras/Private/Core/CameraIKAim.cpp
@@ -63,6 +63,12 @@ bool FCameraIKAim::DoRun(const FCameraIKAimParams& Params, const FCameraRigEvalu
 	ScratchResult.VariableTable.Initialize(CameraRigInfo.CameraRig->AllocationInfo.VariableTableInfo);
 	ScratchResult.VariableTable.OverrideAll(CameraSystemEvaluator->GetEvaluatedResult().VariableTable);
 
+	// Initialize our hierarchy caches.
+	FRootCameraNodeEvaluator* CameraSystemRootEvaluator = CameraSystemEvaluator->GetRootNodeEvaluator();
+	FSingleCameraRigHierarchyBuildParams HierarchyParams;
+	HierarchyParams.CameraRigInfo = CameraRigInfo;
+	CameraSystemRootEvaluator->BuildSingleCameraRigHierarchy(HierarchyParams, CameraSystemHierarchy);
+
 	// Iterate on the solution.
 	FAimIterationInfo IterationInfo;
 	double LastErrorAngle = TNumericLimits<double>::Max();
@@ -133,13 +139,12 @@ void FCameraIKAim::RunRootCameraNode(const FCameraIKAimParams& Params, const FCa
 	FRootCameraNodeEvaluator* RootEvaluator = Params.Evaluator->GetRootNodeEvaluator();
 
 	// Save the initial state of the camera rig.
-	// TODO: we could optimize this a bit by caching the list of nodes that want serialization.
 	{
 		EvaluatorSnapshot.Reset();
 
 		FCameraNodeEvaluatorSerializeParams SerializeParams;
 		FMemoryWriter Writer(EvaluatorSnapshot);
-		RootEvaluator->Serialize(SerializeParams, Writer);
+		CameraSystemHierarchy.CallSerialize(SerializeParams, Writer);
 	}
 
 	// Run the system.
@@ -158,7 +163,7 @@ void FCameraIKAim::RunRootCameraNode(const FCameraIKAimParams& Params, const FCa
 	{
 		FCameraNodeEvaluatorSerializeParams SerializeParams;
 		FMemoryReader Reader(EvaluatorSnapshot);
-		RootEvaluator->Serialize(SerializeParams, Reader);
+		CameraSystemHierarchy.CallSerialize(SerializeParams, Reader);
 	}
 }
 
@@ -244,7 +249,11 @@ void FCameraIKAim::AimTwoBonesCameraRig(const FCameraIKAimParams& Params, const
 	Operation.Yaw = FConsumableDouble::Delta(Correction.Yaw);
 	Operation.Pitch = FConsumableDouble::Delta(Correction.Pitch);
 
-	CameraRigInfo.RootEvaluator->ExecuteOperation(OperationParams, Operation);
+	CameraSystemHierarchy.ForEachEvaluator(TEXT("ActiveCameraRig"), ECameraNodeEvaluatorFlags::SupportsOperations,
+			[&OperationParams, &Operation](FCameraNodeEvaluator* Evaluator)
+			{
+				Evaluator->ExecuteOperation(OperationParams, Operation);
+			});
 
 	if (Operation.Yaw.HasValue() || Operation.Pitch.HasValue())
 	{
diff --git a/Engine/Plugins/Cameras/GameplayCameras/Source/GameplayCameras/Private/Core/CameraNodeEvaluator.cpp b/Engine/Plugins/Cameras/GameplayCameras/Source/GameplayCameras/Private/Core/CameraNodeEvaluator.cpp
--- a/Engine/Plugins/Cameras/GameplayCameras/Source/GameplayCameras/Private/Core/CameraNodeEvaluator.cpp
+++ b/Engine/Plugins/Cameras/GameplayCameras/Source/GameplayCameras/Private/Core/CameraNodeEvaluator.cpp
@@ -3,6 +3,7 @@
 #include "Core/CameraNodeEvaluator.h"
 
 #include "Core/CameraNode.h"
+#include "Core/CameraNodeEvaluatorHierarchy.h"
 #include "Debug/CameraDebugBlockBuilder.h"
 #include "Debug/CameraNodeEvaluatorDebugBlock.h"
 #include "UObject/UObjectGlobals.h"
@@ -12,6 +13,11 @@ namespace UE::Cameras
 
 UE_GAMEPLAY_CAMERAS_DEFINE_RTTI(FCameraNodeEvaluator)
 
+FCameraNodeEvaluatorInitializeParams::FCameraNodeEvaluatorInitializeParams(FCameraNodeEvaluatorHierarchy* InHierarchy)
+	: Hierarchy(InHierarchy)
+{
+}
+
 void FCameraNodeEvaluationResult::Reset(bool bResetVariableTable)
 {
 	CameraPose.Reset();
@@ -81,13 +87,14 @@ FCameraNodeEvaluator* FCameraNodeEvaluatorBuildParams::BuildEvaluator(const UCam
 	return nullptr;
 }
 
-FCameraNodeEvaluator::FCameraNodeEvaluator()
+void FCameraNodeEvaluator::SetPrivateCameraNode(TObjectPtr<const UCameraNode> InCameraNode)
 {
+	PrivateCameraNode = InCameraNode;
 }
 
-void FCameraNodeEvaluator::SetPrivateCameraNode(TObjectPtr<const UCameraNode> InCameraNode)
+void FCameraNodeEvaluator::AddNodeEvaluatorFlags(ECameraNodeEvaluatorFlags InFlags)
 {
-	PrivateCameraNode = InCameraNode;
+	PrivateFlags |= InFlags;
 }
 
 void FCameraNodeEvaluator::SetNodeEvaluatorFlags(ECameraNodeEvaluatorFlags InFlags)
@@ -95,28 +102,25 @@ void FCameraNodeEvaluator::SetNodeEvaluatorFlags(ECameraNodeEvaluatorFlags InFla
 	PrivateFlags = InFlags;
 }
 
-FCameraNodeEvaluatorChildrenView FCameraNodeEvaluator::GetChildren()
+void FCameraNodeEvaluator::Build(const FCameraNodeEvaluatorBuildParams& Params)
 {
-	return OnGetChildren();
+	OnBuild(Params);
 }
 
-void FCameraNodeEvaluator::ExecuteOperation(const FCameraOperationParams& Params, FCameraOperation& Operation)
+void FCameraNodeEvaluator::Initialize(const FCameraNodeEvaluatorInitializeParams& Params, FCameraNodeEvaluationResult& OutResult)
 {
-	if (!PrivateCameraNode || PrivateCameraNode->bIsEnabled)
+	if (Params.Hierarchy)
 	{
-		if (EnumHasAnyFlags(PrivateFlags, ECameraNodeEvaluatorFlags::SupportsOperations))
-		{
-			OnExecuteOperation(Params, Operation);
-		}
-		else
+		Params.Hierarchy->AddEvaluator(this);
+	}
+
+	OnInitialize(Params, OutResult);
+
+	for (FCameraNodeEvaluator* Child : GetChildren())
+	{
+		if (Child)
 		{
-			for (FCameraNodeEvaluator* Child : GetChildren())
-			{
-				if (Child)
-				{
-					Child->ExecuteOperation(Params, Operation);
-				}
-			}
+			Child->Initialize(Params, OutResult);
 		}
 	}
 }
@@ -139,79 +143,40 @@ void FCameraNodeEvaluator::AddReferencedObjects(FReferenceCollector& Collector)
 	}
 }
 
-void FCameraNodeEvaluator::Serialize(const FCameraNodeEvaluatorSerializeParams& Params, FArchive& Ar)
-{
-	OnSerialize(Params, Ar);
-
-	for (FCameraNodeEvaluator* Child : GetChildren())
-	{
-		if (Child)
-		{
-			Child->Serialize(Params, Ar);
-		}
-	}
-}
-
-void FCameraNodeEvaluator::Build(const FCameraNodeEvaluatorBuildParams& Params)
+FCameraNodeEvaluatorChildrenView FCameraNodeEvaluator::GetChildren()
 {
-	OnBuild(Params);
+	return OnGetChildren();
 }
 
-void FCameraNodeEvaluator::Initialize(const FCameraNodeEvaluatorInitializeParams& Params, FCameraNodeEvaluationResult& OutResult)
+void FCameraNodeEvaluator::UpdateParameters(const FCameraBlendedParameterUpdateParams& Params, FCameraBlendedParameterUpdateResult& OutResult)
 {
-	OnInitialize(Params, OutResult);
-
-	for (FCameraNodeEvaluator* Child : GetChildren())
+	if (!PrivateCameraNode || PrivateCameraNode->bIsEnabled)
 	{
-		if (Child)
-		{
-			Child->Initialize(Params, OutResult);
-		}
+		OnUpdateParameters(Params, OutResult);
 	}
 }
 
-void FCameraNodeEvaluator::UpdateParameters(const FCameraBlendedParameterUpdateParams& Params, FCameraBlendedParameterUpdateResult& OutResult)
+void FCameraNodeEvaluator::Run(const FCameraNodeEvaluationParams& Params, FCameraNodeEvaluationResult& OutResult)
 {
 	if (!PrivateCameraNode || PrivateCameraNode->bIsEnabled)
 	{
-		if (EnumHasAnyFlags(PrivateFlags, ECameraNodeEvaluatorFlags::NeedsParameterUpdate))
-		{
-			OnUpdateParameters(Params, OutResult);
-		}
-		else
-		{
-			for (FCameraNodeEvaluator* Child : GetChildren())
-			{
-				if (Child)
-				{
-					Child->UpdateParameters(Params, OutResult);
-				}
-			}
-		}
+		OnRun(Params, OutResult);
 	}
 }
 
-void FCameraNodeEvaluator::Run(const FCameraNodeEvaluationParams& Params, FCameraNodeEvaluationResult& OutResult)
+void FCameraNodeEvaluator::ExecuteOperation(const FCameraOperationParams& Params, FCameraOperation& Operation)
 {
 	if (!PrivateCameraNode || PrivateCameraNode->bIsEnabled)
 	{
-		if (EnumHasAnyFlags(PrivateFlags, ECameraNodeEvaluatorFlags::NeedsEvaluationUpdate))
-		{
-			OnRun(Params, OutResult);
-		}
-		else
-		{
-			for (FCameraNodeEvaluator* Child : GetChildren())
-			{
-				if (Child)
-				{
-					Child->Run(Params, OutResult);
-				}
-			}
-		}
+		OnExecuteOperation(Params, Operation);
 	}
 }
 
+void FCameraNodeEvaluator::Serialize(const FCameraNodeEvaluatorSerializeParams& Params, FArchive& Ar)
+{
+	OnSerialize(Params, Ar);
+}
+
 #if UE_GAMEPLAY_CAMERAS_DEBUG
 
 void FCameraNodeEvaluator::BuildDebugBlocks(const FCameraDebugBlockBuildParams& Params, FCameraDebugBlockBuilder& Builder)
diff --git a/Engine/Plugins/Cameras/GameplayCameras/Source/GameplayCameras/Private/Core/CameraNodeEvaluatorHierarchy.cpp b/Engine/Plugins/Cameras/GameplayCameras/Source/GameplayCameras/Private/Core/CameraNodeEvaluatorHierarchy.cpp
new file mode 100644
--- /dev/null
+++ b/Engine/Plugins/Cameras/GameplayCameras/Source/GameplayCameras/Private/Core/CameraNodeEvaluatorHierarchy.cpp
@@ -0,0 +1,114 @@
+// Copyright Epic Games, Inc. All Rights Reserved.
+
+#include "Core/CameraNodeEvaluatorHierarchy.h"
+
+namespace UE::Cameras
+{
+
+FCameraNodeEvaluatorHierarchy::FCameraNodeEvaluatorHierarchy()
+{
+}
+
+FCameraNodeEvaluatorHierarchy::FCameraNodeEvaluatorHierarchy(FCameraNodeEvaluator* InRootEvaluator)
+{
+	Build(InRootEvaluator);
+}
+
+TArrayView<FCameraNodeEvaluator* const> FCameraNodeEvaluatorHierarchy::GetFlattenedHierarchy() const
+{
+	return FlattenedHierarchy;
+}
+
+void FCameraNodeEvaluatorHierarchy::GetFlattenedHierarchy(ECameraNodeEvaluatorFlags FilterFlags, TArray<FCameraNodeEvaluator*> OutEvaluators) const
+{
+	for (FCameraNodeEvaluator* Evaluator : FlattenedHierarchy)
+	{
+		if (EnumHasAllFlags(Evaluator->GetNodeEvaluatorFlags(), FilterFlags))
+		{
+			OutEvaluators.Add(Evaluator);
+		}
+	}
+}
+
+void FCameraNodeEvaluatorHierarchy::Build(FCameraNodeEvaluator* InRootEvaluator)
+{
+	Reset();
+
+	Append(InRootEvaluator);
+}
+
+void FCameraNodeEvaluatorHierarchy::Append(FCameraNodeEvaluator* InRootEvaluator)
+{
+	if (!InRootEvaluator)
+	{
+		return;
+	}
+
+	TArray<FCameraNodeEvaluator*> EvaluatorStack;
+	EvaluatorStack.Add(InRootEvaluator);
+	while (!EvaluatorStack.IsEmpty())
+	{
+		FCameraNodeEvaluator* TopEvaluator = EvaluatorStack.Pop();
+		FlattenedHierarchy.Add(TopEvaluator);
+
+		FCameraNodeEvaluatorChildrenView TopEvaluatorChildren(TopEvaluator->GetChildren());
+		for (FCameraNodeEvaluator* ChildEvaluator : ReverseIterate(TopEvaluatorChildren))
+		{
+			if (ChildEvaluator)
+			{
+				EvaluatorStack.Add(ChildEvaluator);
+			}
+		}
+	}
+}
+
+void FCameraNodeEvaluatorHierarchy::AppendTagged(const FName TaggedRangeName, FCameraNodeEvaluator* InRootEvaluator)
+{
+	FTaggedRange& NewRange = TaggedRanges.Add(TaggedRangeName);
+	NewRange.StartIndex = FlattenedHierarchy.Num();
+	Append(InRootEvaluator);
+	NewRange.EndIndex = FlattenedHierarchy.Num();
+}
+
+void FCameraNodeEvaluatorHierarchy::AddEvaluator(FCameraNodeEvaluator* Evaluator)
+{
+	if (Evaluator)
+	{
+		FlattenedHierarchy.Add(Evaluator);
+	}
+}
+
+void FCameraNodeEvaluatorHierarchy::Reset()
+{
+	FlattenedHierarchy.Reset();
+}
+
+void FCameraNodeEvaluatorHierarchy::CallUpdateParameters(const FCameraBlendedParameterUpdateParams& Params, FCameraBlendedParameterUpdateResult& OutResult)
+{
+	ForEachEvaluator(ECameraNodeEvaluatorFlags::NeedsParameterUpdate, 
+			[&Params, &OutResult](FCameraNodeEvaluator* Evaluator)
+			{
+				Evaluator->UpdateParameters(Params, OutResult);
+			});
+}
+
+void FCameraNodeEvaluatorHierarchy::CallExecuteOperation(const FCameraOperationParams& Params, FCameraOperation& Operation)
+{
+	ForEachEvaluator(ECameraNodeEvaluatorFlags::SupportsOperations, 
+			[&Params, &Operation](FCameraNodeEvaluator* Evaluator)
+			{
+				Evaluator->ExecuteOperation(Params, Operation);
+			});
+}
+
+void FCameraNodeEvaluatorHierarchy::CallSerialize(const FCameraNodeEvaluatorSerializeParams& Params, FArchive& Ar)
+{
+	ForEachEvaluator(ECameraNodeEvaluatorFlags::NeedsSerialize,
+			[&Params, &Ar](FCameraNodeEvaluator* Evaluator)
+			{
+				Evaluator->Serialize(Params, Ar);
+			});
+}
+
+}  // namespace UE::Cameras
+
diff --git a/Engine/Plugins/Cameras/GameplayCameras/Source/GameplayCameras/Private/Core/DefaultRootCameraNode.cpp b/Engine/Plugins/Cameras/GameplayCameras/Source/GameplayCameras/Private/Core/DefaultRootCameraNode.cpp
--- a/Engine/Plugins/Cameras/GameplayCameras/Source/GameplayCameras/Private/Core/DefaultRootCameraNode.cpp
+++ b/Engine/Plugins/Cameras/GameplayCameras/Source/GameplayCameras/Private/Core/DefaultRootCameraNode.cpp
@@ -4,6 +4,7 @@
 
 #include "Core/BlendStackCameraNode.h"
 #include "Core/CameraEvaluationContext.h"
+#include "Core/CameraNodeEvaluatorHierarchy.h"
 #include "Core/RootCameraNodeCameraRigEvent.h"
 #include "Debug/BlendStacksCameraDebugBlock.h"
 #include "Debug/CameraDebugBlockBuilder.h"
@@ -91,6 +92,15 @@ void FDefaultRootCameraNodeEvaluator::OnActivateCameraRig(const FActivateCameraR
 	}
 }
 
+void FDefaultRootCameraNodeEvaluator::OnBuildSingleCameraRigHierarchy(const FSingleCameraRigHierarchyBuildParams& Params, FCameraNodeEvaluatorHierarchy& OutHierarchy)
+{
+	OutHierarchy.Build(BaseLayer);
+	{
+		OutHierarchy.AppendTagged(Params.CameraRigRangeName, Params.CameraRigInfo.RootEvaluator);
+	}
+	OutHierarchy.Append(GlobalLayer);
+}
+
 void FDefaultRootCameraNodeEvaluator::OnRunSingleCameraRig(const FSingleCameraRigEvaluationParams& Params, FCameraNodeEvaluationResult& OutResult)
 {
 	BaseLayer->Run(Params.EvaluationParams, OutResult);
@@ -101,7 +111,9 @@ void FDefaultRootCameraNodeEvaluator::OnRunSingleCameraRig(const FSingleCameraRi
 		const FCameraNodeEvaluationResult* CameraRigResult = Params.CameraRigInfo.LastResult;
 		FCameraBlendedParameterUpdateParams InputParams(Params.EvaluationParams, CameraRigResult->CameraPose);
 		FCameraBlendedParameterUpdateResult InputResult(OutResult.VariableTable);
-		RootEvaluator->UpdateParameters(InputParams, InputResult);
+
+		FCameraNodeEvaluatorHierarchy Hierarchy(RootEvaluator);
+		Hierarchy.CallUpdateParameters(InputParams, InputResult);
 	}
 
 	{
diff --git a/Engine/Plugins/Cameras/GameplayCameras/Source/GameplayCameras/Private/Core/RootCameraNode.cpp b/Engine/Plugins/Cameras/GameplayCameras/Source/GameplayCameras/Private/Core/RootCameraNode.cpp
--- a/Engine/Plugins/Cameras/GameplayCameras/Source/GameplayCameras/Private/Core/RootCameraNode.cpp
+++ b/Engine/Plugins/Cameras/GameplayCameras/Source/GameplayCameras/Private/Core/RootCameraNode.cpp
@@ -22,6 +22,11 @@ void FRootCameraNodeEvaluator::ActivateCameraRig(const FActivateCameraRigParams&
 	OnActivateCameraRig(Params);
 }
 
+void FRootCameraNodeEvaluator::BuildSingleCameraRigHierarchy(const FSingleCameraRigHierarchyBuildParams& Params, FCameraNodeEvaluatorHierarchy& OutHierarchy)
+{
+	OnBuildSingleCameraRigHierarchy(Params, OutHierarchy);
+}
+
 void FRootCameraNodeEvaluator::RunSingleCameraRig(const FSingleCameraRigEvaluationParams& Params, FCameraNodeEvaluationResult& OutResult)
 {
 	// Before we do the actual evaluation we need to ask the system to auto-reset
diff --git a/Engine/Plugins/Cameras/GameplayCameras/Source/GameplayCameras/Private/Nodes/Common/BoomArmCameraNode.cpp b/Engine/Plugins/Cameras/GameplayCameras/Source/GameplayCameras/Private/Nodes/Common/BoomArmCameraNode.cpp
--- a/Engine/Plugins/Cameras/GameplayCameras/Source/GameplayCameras/Private/Nodes/Common/BoomArmCameraNode.cpp
+++ b/Engine/Plugins/Cameras/GameplayCameras/Source/GameplayCameras/Private/Nodes/Common/BoomArmCameraNode.cpp
@@ -79,9 +79,7 @@ void FBoomArmCameraNodeEvaluator::OnBuild(const FCameraNodeEvaluatorBuildParams&
 
 void FBoomArmCameraNodeEvaluator::OnInitialize(const FCameraNodeEvaluatorInitializeParams& Params, FCameraNodeEvaluationResult& OutResult)
 {
-	SetNodeEvaluatorFlags(
-			ECameraNodeEvaluatorFlags::NeedsEvaluationUpdate |
-			ECameraNodeEvaluatorFlags::SupportsOperations);
+	AddNodeEvaluatorFlags(ECameraNodeEvaluatorFlags::SupportsOperations);
 
 	const UBoomArmCameraNode* BoomArmNode = GetCameraNodeAs<UBoomArmCameraNode>();
 	BoomOffsetReader.Initialize(BoomArmNode->BoomOffset);
diff --git a/Engine/Plugins/Cameras/GameplayCameras/Source/GameplayCameras/Private/Nodes/Input/CameraRigInput1DSlot.cpp b/Engine/Plugins/Cameras/GameplayCameras/Source/GameplayCameras/Private/Nodes/Input/CameraRigInput1DSlot.cpp
--- a/Engine/Plugins/Cameras/GameplayCameras/Source/GameplayCameras/Private/Nodes/Input/CameraRigInput1DSlot.cpp
+++ b/Engine/Plugins/Cameras/GameplayCameras/Source/GameplayCameras/Private/Nodes/Input/CameraRigInput1DSlot.cpp
@@ -19,10 +19,7 @@ UE_DEFINE_CAMERA_NODE_EVALUATOR(FCameraRigInput1DSlotEvaluator)
 
 FCameraRigInput1DSlotEvaluator::FCameraRigInput1DSlotEvaluator()
 {
-	SetNodeEvaluatorFlags(
-			ECameraNodeEvaluatorFlags::NeedsParameterUpdate | 
-			ECameraNodeEvaluatorFlags::NeedsEvaluationUpdate |
-			ECameraNodeEvaluatorFlags::SupportsOperations);
+	AddNodeEvaluatorFlags(ECameraNodeEvaluatorFlags::NeedsParameterUpdate | ECameraNodeEvaluatorFlags::SupportsOperations);
 }
 
 void FCameraRigInput1DSlotEvaluator::OnInitialize(const FCameraNodeEvaluatorInitializeParams& Params, FCameraNodeEvaluationResult& OutResult)
diff --git a/Engine/Plugins/Cameras/GameplayCameras/Source/GameplayCameras/Private/Nodes/Input/CameraRigInput2DSlot.cpp b/Engine/Plugins/Cameras/GameplayCameras/Source/GameplayCameras/Private/Nodes/Input/CameraRigInput2DSlot.cpp
--- a/Engine/Plugins/Cameras/GameplayCameras/Source/GameplayCameras/Private/Nodes/Input/CameraRigInput2DSlot.cpp
+++ b/Engine/Plugins/Cameras/GameplayCameras/Source/GameplayCameras/Private/Nodes/Input/CameraRigInput2DSlot.cpp
@@ -19,10 +19,7 @@ UE_DEFINE_CAMERA_NODE_EVALUATOR(FCameraRigInput2DSlotEvaluator)
 
 FCameraRigInput2DSlotEvaluator::FCameraRigInput2DSlotEvaluator()
 {
-	SetNodeEvaluatorFlags(
-			ECameraNodeEvaluatorFlags::NeedsParameterUpdate | 
-			ECameraNodeEvaluatorFlags::NeedsEvaluationUpdate |
-			ECameraNodeEvaluatorFlags::SupportsOperations);
+	AddNodeEvaluatorFlags(ECameraNodeEvaluatorFlags::NeedsParameterUpdate | ECameraNodeEvaluatorFlags::SupportsOperations);
 }
 
 void FCameraRigInput2DSlotEvaluator::OnInitialize(const FCameraNodeEvaluatorInitializeParams& Params, FCameraNodeEvaluationResult& OutResult)
diff --git a/Engine/Plugins/Cameras/GameplayCameras/Source/GameplayCameras/Public/Core/BlendStackCameraNode.h b/Engine/Plugins/Cameras/GameplayCameras/Source/GameplayCameras/Public/Core/BlendStackCameraNode.h
--- a/Engine/Plugins/Cameras/GameplayCameras/Source/GameplayCameras/Public/Core/BlendStackCameraNode.h
+++ b/Engine/Plugins/Cameras/GameplayCameras/Source/GameplayCameras/Public/Core/BlendStackCameraNode.h
@@ -4,6 +4,7 @@
 
 #include "Core/CameraNode.h"
 #include "Core/CameraNodeEvaluator.h"
+#include "Core/CameraNodeEvaluatorHierarchy.h"
 #include "Core/CameraNodeEvaluatorStorage.h"
 #include "Core/CameraRigAsset.h"
 #include "Core/CameraRigEvaluationInfo.h"
@@ -178,8 +179,8 @@ protected:
 		FCameraNodeEvaluatorStorage EvaluatorStorage;
 		/** Root evaluator. */
 		FBlendStackRootCameraNodeEvaluator* RootEvaluator = nullptr;
-		/** Evaluators needing parameter update. */
-		TArray<FCameraNodeEvaluator*> ParameterEvaluators;
+		/** The evaluator tree. */
+		FCameraNodeEvaluatorHierarchy EvaluatorHierarchy;
 		/** Result for this node tree. */
 		FCameraNodeEvaluationResult Result;
 		/** Whether this is the first frame this entry runs. */
diff --git a/Engine/Plugins/Cameras/GameplayCameras/Source/GameplayCameras/Public/Core/CameraIKAim.h b/Engine/Plugins/Cameras/GameplayCameras/Source/GameplayCameras/Public/Core/CameraIKAim.h
--- a/Engine/Plugins/Cameras/GameplayCameras/Source/GameplayCameras/Public/Core/CameraIKAim.h
+++ b/Engine/Plugins/Cameras/GameplayCameras/Source/GameplayCameras/Public/Core/CameraIKAim.h
@@ -3,6 +3,7 @@
 #pragma once
 
 #include "Core/CameraNodeEvaluator.h"
+#include "Core/CameraNodeEvaluatorHierarchy.h"
 #include "Core/CameraNodeEvaluatorStorage.h"
 #include "Math/MathFwd.h"
 
@@ -82,6 +83,7 @@ private:
 
 	FCameraNodeEvaluationResult ScratchResult;
 	
+	FCameraNodeEvaluatorHierarchy CameraSystemHierarchy;
 	TArray<uint8> EvaluatorSnapshot;
 };
 
diff --git a/Engine/Plugins/Cameras/GameplayCameras/Source/GameplayCameras/Public/Core/CameraNodeEvaluator.h b/Engine/Plugins/Cameras/GameplayCameras/Source/GameplayCameras/Public/Core/CameraNodeEvaluator.h
--- a/Engine/Plugins/Cameras/GameplayCameras/Source/GameplayCameras/Public/Core/CameraNodeEvaluator.h
+++ b/Engine/Plugins/Cameras/GameplayCameras/Source/GameplayCameras/Public/Core/CameraNodeEvaluator.h
@@ -23,6 +23,7 @@ namespace UE::Cameras
 
 class FCameraEvaluationContext;
 class FCameraNodeEvaluator;
+class FCameraNodeEvaluatorHierarchy;
 class FCameraSystemEvaluator;
 struct FCameraOperation;
 struct FCameraNodeEvaluationParams;
@@ -36,8 +37,10 @@ enum class ECameraNodeEvaluatorFlags
 {
 	None = 0,
 	NeedsParameterUpdate = 1 << 0,
-	NeedsEvaluationUpdate = 1 << 1,
-	SupportsOperations = 1 << 2
+	NeedsSerialize = 1 << 1,
+	SupportsOperations = 1 << 2,
+
+	Default = NeedsSerialize
 };
 ENUM_CLASS_FLAGS(ECameraNodeEvaluatorFlags)
 
@@ -81,6 +84,18 @@ struct FCameraNodeEvaluatorInitializeParams
 	 * is being pushed on top of a non-empty blend stack.
 	 */
 	FCameraRigEvaluationInfo LastActiveCameraRigInfo;
+
+public:
+
+	FCameraNodeEvaluatorInitializeParams() = default;
+	FCameraNodeEvaluatorInitializeParams(FCameraNodeEvaluatorHierarchy* InHierarchy);
+
+private:
+
+	/** An optional hierarchy to populate while initialize the evaluator tree. */
+	FCameraNodeEvaluatorHierarchy* Hierarchy = nullptr;
+
+	friend class FCameraNodeEvaluator;
 };
 
 /**
@@ -203,15 +218,18 @@ class FCameraNodeEvaluator
 
 public:
 
-	GAMEPLAYCAMERAS_API FCameraNodeEvaluator();
-	GAMEPLAYCAMERAS_API virtual ~FCameraNodeEvaluator() {}
+	GAMEPLAYCAMERAS_API FCameraNodeEvaluator() = default;
+	GAMEPLAYCAMERAS_API virtual ~FCameraNodeEvaluator() = default;
 
 	/** Called to build any children evaluators. */
 	void Build(const FCameraNodeEvaluatorBuildParams& Params);
 
-	/** Initialize this evaluator. */
+	/** Initialize this evaluator and all its descendants. */
 	void Initialize(const FCameraNodeEvaluatorInitializeParams& Params, FCameraNodeEvaluationResult& OutResult);
 
+	/** Collect referenced UObjects for this node and all its descendants. */
+	void AddReferencedObjects(FReferenceCollector& Collector);
+
 	/** Get the list of children under this evaluator. */
 	FCameraNodeEvaluatorChildrenView GetChildren();
 
@@ -224,9 +242,6 @@ public:
 	/** Execute an IK operation. */
 	void ExecuteOperation(const FCameraOperationParams& Params, FCameraOperation& Operation);
 
-	/** Collect referenced UObjects. */
-	void AddReferencedObjects(FReferenceCollector& Collector);
-
 	/** Serializes the state of this evaluator. */
 	void Serialize(const FCameraNodeEvaluatorSerializeParams& Params, FArchive& Ar);
 
@@ -255,6 +270,9 @@ public:
 
 protected:
 
+	/** Adds flags for this evaluator. */
+	void AddNodeEvaluatorFlags(ECameraNodeEvaluatorFlags InFlags);
+
 	/** Sets the flags for this evaluator. */
 	void SetNodeEvaluatorFlags(ECameraNodeEvaluatorFlags InFlags);
 
@@ -263,25 +281,36 @@ protected:
 	/** Called to build any children evaluators. */
 	GAMEPLAYCAMERAS_API virtual void OnBuild(const FCameraNodeEvaluatorBuildParams& Params) {}
 
-	/** Initialize this evaluator. */
+	/** Initialize this evaluator. Children and descendants will be automatically initialized too. */
 	GAMEPLAYCAMERAS_API virtual void OnInitialize(const FCameraNodeEvaluatorInitializeParams& Params, FCameraNodeEvaluationResult& OutResult) {}
 
+	/** Collect referenced UObjects for this node. */
+	GAMEPLAYCAMERAS_API virtual void OnAddReferencedObjects(FReferenceCollector& Collector) {}
+
 	/** Get the list of children under this evaluator. */
 	GAMEPLAYCAMERAS_API virtual FCameraNodeEvaluatorChildrenView OnGetChildren() { return FCameraNodeEvaluatorChildrenView(); }
 
-	/** Called to update and store the blended parameters for this node. */
+	/**
+	 * Called to update and store the blended parameters for this node.
+	 * Requires setting the ECameraNodeEvaluatorFlags::NeedsParameterUpdate flag.
+	 */
 	GAMEPLAYCAMERAS_API virtual void OnUpdateParameters(const FCameraBlendedParameterUpdateParams& Params, FCameraBlendedParameterUpdateResult& OutResult) {}
 
-	/** Run this evaluator. */
+	/**
+	 * Run this evaluator. This node evaluator is responsible for calling Run() on its children as appropriate.
+	 */
 	GAMEPLAYCAMERAS_API virtual void OnRun(const FCameraNodeEvaluationParams& Params, FCameraNodeEvaluationResult& OutResult) {}
 
-	/** Execute an IK operation. */
+	/** 
+	 * Execute an IK operation.
+	 * Requires setting the ECameraNodeEvaluatorFlags::SupportsOperations flag.
+	 */
 	GAMEPLAYCAMERAS_API virtual void OnExecuteOperation(const FCameraOperationParams& Params, FCameraOperation& Operation) {}
 
-	/** Collect referenced UObjects. */
-	GAMEPLAYCAMERAS_API virtual void OnAddReferencedObjects(FReferenceCollector& Collector) {}
-
-	/** Serializes the state of this evaluator. */
+	/**
+	 * Serializes the state of this evaluator.
+	 * Requires setting the ECameraNodeEvaluatorFlags::NeedsSerialize flag, which is set by default.
+	 */
 	GAMEPLAYCAMERAS_API virtual void OnSerialize(const FCameraNodeEvaluatorSerializeParams& Params, FArchive& Ar) {}
 
 #if UE_GAMEPLAY_CAMERAS_DEBUG
@@ -295,7 +324,7 @@ private:
 	TObjectPtr<const UCameraNode> PrivateCameraNode;
 
 	/** The flags for this evaluator. */
-	ECameraNodeEvaluatorFlags PrivateFlags = ECameraNodeEvaluatorFlags::NeedsEvaluationUpdate;
+	ECameraNodeEvaluatorFlags PrivateFlags = ECameraNodeEvaluatorFlags::Default;
 };
 
 /** Utility base class for camera node evaluators of a specific camera node type. */
diff --git a/Engine/Plugins/Cameras/GameplayCameras/Source/GameplayCameras/Public/Core/CameraNodeEvaluatorHierarchy.h b/Engine/Plugins/Cameras/GameplayCameras/Source/GameplayCameras/Public/Core/CameraNodeEvaluatorHierarchy.h
new file mode 100644
--- /dev/null
+++ b/Engine/Plugins/Cameras/GameplayCameras/Source/GameplayCameras/Public/Core/CameraNodeEvaluatorHierarchy.h
@@ -0,0 +1,122 @@
+// Copyright Epic Games, Inc. All Rights Reserved.
+
+#pragma once
+
+#include "Containers/Array.h"
+#include "Core/CameraNodeEvaluator.h"
+
+namespace UE::Cameras
+{
+
+/**
+ * A utility class that stores a flattened hierarchy of camera node evaluators.
+ */
+class GAMEPLAYCAMERAS_API FCameraNodeEvaluatorHierarchy
+{
+public:
+
+	/** Builds an empty hierarchy. */
+	FCameraNodeEvaluatorHierarchy();
+	/** Builds a hierarchy starting from the given root evaluator. */
+	FCameraNodeEvaluatorHierarchy(FCameraNodeEvaluator* InRootEvaluator);
+
+	/** Get the list of evaluators in depth-first order. */
+	TArrayView<FCameraNodeEvaluator* const> GetFlattenedHierarchy() const;
+
+	/** Get the list of evaluators matching the given evaluator flags, in depth-first order. */
+	void GetFlattenedHierarchy(ECameraNodeEvaluatorFlags FilterFlags, TArray<FCameraNodeEvaluator*> OutEvaluators) const;
+
+public:
+
+	/** Builds a hierarchy starting from the given root evaluator. */
+	void Build(FCameraNodeEvaluator* InRootEvaluator);
+	/** Append another hierarchy to the existing hierarchy, starting from the given root evaluator. */
+	void Append(FCameraNodeEvaluator* InRootEvaluator);
+	/** Append and tag another hierarchy to the existing hierarchy, starting from the given root evaluator. */
+	void AppendTagged(const FName TaggedRangeName, FCameraNodeEvaluator* InRootEvaluator);
+	/** Add an evaluator to the existing hierarchy. */
+	void AddEvaluator(FCameraNodeEvaluator* Evaluator);
+	/** Resets this object to an emtpy hierarchy. */
+	void Reset();
+
+public:
+
+	/** Executes the given predicate on each evaluator in depth-first order. */
+	template<typename PredicateClass>
+	void ForEachEvaluator(PredicateClass&& Predicate)
+	{
+		for (FCameraNodeEvaluator* Evaluator : FlattenedHierarchy)
+		{
+			Predicate(Evaluator);
+		}
+	}
+
+	/** Executes the given predicate on each evaluator in the specified range in depth-first order. */
+	template<typename PredicateClass>
+	void ForEachEvaluator(const FName TaggedRangeName, PredicateClass&& Predicate)
+	{
+		const FTaggedRange& TaggedRange = TaggedRanges.FindChecked(TaggedRangeName);
+		for (int32 Index = TaggedRange.StartIndex; Index < TaggedRange.EndIndex; ++Index)
+		{
+			FCameraNodeEvaluator* Evaluator(FlattenedHierarchy[Index]);
+			Predicate(Evaluator);
+		}
+	}
+
+	/**
+	 * Executes the given predicate on each evaluator matching the specified evaluator flags in 
+	 * depth-first order.
+	 */
+	template<typename PredicateClass>
+	void ForEachEvaluator(ECameraNodeEvaluatorFlags FilterFlags, PredicateClass&& Predicate)
+	{
+		for (FCameraNodeEvaluator* Evaluator : FlattenedHierarchy)
+		{
+			if (EnumHasAllFlags(Evaluator->GetNodeEvaluatorFlags(), FilterFlags))
+			{
+				Predicate(Evaluator);
+			}
+		}
+	}
+
+	/** 
+	 * Executes the given predicate on each evaluator in the specified range matching the specified 
+	 * evaluator flags in depth-first order.
+	 */
+	template<typename PredicateClass>
+	void ForEachEvaluator(const FName TaggedRangeName, ECameraNodeEvaluatorFlags FilterFlags, PredicateClass&& Predicate)
+	{
+		const FTaggedRange& TaggedRange = TaggedRanges.FindChecked(TaggedRangeName);
+		for (int32 Index = TaggedRange.StartIndex; Index < TaggedRange.EndIndex; ++Index)
+		{
+			FCameraNodeEvaluator* Evaluator(FlattenedHierarchy[Index]);
+			if (EnumHasAllFlags(Evaluator->GetNodeEvaluatorFlags(), FilterFlags))
+			{
+				Predicate(Evaluator);
+			}
+		}
+	}
+
+public:
+
+	/** Helper method to call UpdateParameters on the appropriate nodes in the hierarchy. */
+	void CallUpdateParameters(const FCameraBlendedParameterUpdateParams& Params, FCameraBlendedParameterUpdateResult& OutResult);
+	/** Helper method to call ExecuteOperation on the appropriate nodes in the hierarchy. */
+	void CallExecuteOperation(const FCameraOperationParams& Params, FCameraOperation& Operation);
+	/** Helper method to call Serialize on the appropriate nodes in the hierarchy. */
+	void CallSerialize(const FCameraNodeEvaluatorSerializeParams& Params, FArchive& Ar);
+
+private:
+
+	TArray<FCameraNodeEvaluator*> FlattenedHierarchy;
+
+	struct FTaggedRange
+	{
+		int32 StartIndex;
+		int32 EndIndex;
+	};
+	TMap<FName, FTaggedRange> TaggedRanges;
+};
+
+}  // namespace UE::Cameras
+
diff --git a/Engine/Plugins/Cameras/GameplayCameras/Source/GameplayCameras/Public/Core/DefaultRootCameraNode.h b/Engine/Plugins/Cameras/GameplayCameras/Source/GameplayCameras/Public/Core/DefaultRootCameraNode.h
--- a/Engine/Plugins/Cameras/GameplayCameras/Source/GameplayCameras/Public/Core/DefaultRootCameraNode.h
+++ b/Engine/Plugins/Cameras/GameplayCameras/Source/GameplayCameras/Public/Core/DefaultRootCameraNode.h
@@ -65,6 +65,7 @@ protected:
 
 	// FRootCameraNodeEvaluator interface.
 	virtual void OnActivateCameraRig(const FActivateCameraRigParams& Params) override;
+	virtual void OnBuildSingleCameraRigHierarchy(const FSingleCameraRigHierarchyBuildParams& Params, FCameraNodeEvaluatorHierarchy& OutHierarchy) override;
 	virtual void OnRunSingleCameraRig(const FSingleCameraRigEvaluationParams& Params, FCameraNodeEvaluationResult& OutResult) override;
 
 private:
diff --git a/Engine/Plugins/Cameras/GameplayCameras/Source/GameplayCameras/Public/Core/RootCameraNode.h b/Engine/Plugins/Cameras/GameplayCameras/Source/GameplayCameras/Public/Core/RootCameraNode.h
--- a/Engine/Plugins/Cameras/GameplayCameras/Source/GameplayCameras/Public/Core/RootCameraNode.h
+++ b/Engine/Plugins/Cameras/GameplayCameras/Source/GameplayCameras/Public/Core/RootCameraNode.h
@@ -59,6 +59,18 @@ struct FActivateCameraRigParams
 	ECameraRigLayer Layer = ECameraRigLayer::Main;
 };
 
+/**
+ * Parameter structure for building a single camera rig hierarchy.
+ */
+struct FSingleCameraRigHierarchyBuildParams
+{
+	/** The camera rig to build the hierachy for. */
+	FCameraRigEvaluationInfo CameraRigInfo;
+
+	/** The name of the range to tag for the camera rig's nodes. */
+	FName CameraRigRangeName = TEXT("ActiveCameraRig");
+};
+
 /**
  * Parameter structure for evaluating a single camera rig.
  */
@@ -83,6 +95,14 @@ public:
 	/** Activates a camera rig. */
 	void ActivateCameraRig(const FActivateCameraRigParams& Params);
 
+	/**
+	 * Builds the hierarchy of the system for a given single camera rig.
+	 * This is expected to return the nodes of all the layers, except for the main layer which
+	 * should only have the nodes of the given camera rig (i.e. it shouldn't have nodes of
+	 * other currently active camera rigs).
+	 */
+	void BuildSingleCameraRigHierarchy(const FSingleCameraRigHierarchyBuildParams& Params, FCameraNodeEvaluatorHierarchy& OutHierarchy);
+
 	/**
 	 * Evaluates a single camera rig.
 	 * This is expected to run all layers as usual, except for the main layer which should
@@ -103,6 +123,9 @@ protected:
 	/** Activates a camera rig. */
 	virtual void OnActivateCameraRig(const FActivateCameraRigParams& Params) {}
 
+	/* Builds the hierarchy of the system for a given single camera rig. */
+	virtual void OnBuildSingleCameraRigHierarchy(const FSingleCameraRigHierarchyBuildParams& Params, FCameraNodeEvaluatorHierarchy& OutHierarchy) {}
+
 	/** Evaluates a single camera rig. See comments on RunSingleCameraRig. */
 	virtual void OnRunSingleCameraRig(const FSingleCameraRigEvaluationParams& Params, FCameraNodeEvaluationResult& OutResult) {}
 
