Cameras: add ability to optionally blend private variables

Also add Visual Studio debugger visualizers for camera variable tables.

diff --git a/Engine/Plugins/Cameras/GameplayCameras/Source/GameplayCameras/Private/Core/CameraVariableTable.cpp b/Engine/Plugins/Cameras/GameplayCameras/Source/GameplayCameras/Private/Core/CameraVariableTable.cpp
--- a/Engine/Plugins/Cameras/GameplayCameras/Source/GameplayCameras/Private/Core/CameraVariableTable.cpp
+++ b/Engine/Plugins/Cameras/GameplayCameras/Source/GameplayCameras/Private/Core/CameraVariableTable.cpp
@@ -332,7 +332,7 @@ void FCameraVariableTable::Serialize(FArchive& Ar)
 
 void FCameraVariableTable::OverrideAll(const FCameraVariableTable& OtherTable)
 {
-	const ECameraVariableTableFilter Filter = ECameraVariableTableFilter::All;
+	const ECameraVariableTableFilter Filter = ECameraVariableTableFilter::AllPublic;
 	InternalOverride(OtherTable, Filter, nullptr, false, nullptr);
 }
 
@@ -353,6 +353,7 @@ void FCameraVariableTable::InternalOverride(const FCameraVariableTable& OtherTab
 	const bool bChangedOnly = EnumHasAnyFlags(Filter, ECameraVariableTableFilter::ChangedOnly);
 	const bool bInputs = EnumHasAnyFlags(Filter, ECameraVariableTableFilter::Input);
 	const bool bOutputs = EnumHasAnyFlags(Filter, ECameraVariableTableFilter::Output);
+	const bool bPrivates = EnumHasAnyFlags(Filter, ECameraVariableTableFilter::Private);
 
 	for (const FEntry& OtherEntry : OtherTable.Entries)
 	{
@@ -362,7 +363,7 @@ void FCameraVariableTable::InternalOverride(const FCameraVariableTable& OtherTab
 		if (EnumHasAnyFlags(OtherFlags, EEntryFlags::Written)
 				&& (!bChangedOnly || EnumHasAnyFlags(OtherFlags, EEntryFlags::WrittenThisFrame))
 				&& ((bInputs && bOtherEntryIsInput) || (bOutputs && !bOtherEntryIsInput))
-				&& !EnumHasAnyFlags(OtherFlags, EEntryFlags::Private)
+				&& (bPrivates || !EnumHasAnyFlags(OtherFlags, EEntryFlags::Private))
 				&& IsVariableInMask(OtherEntry.ID, InMask, bInvertMask))
 		{
 			// See if we know this variable.
@@ -424,7 +425,7 @@ void FCameraVariableTable::InternalOverride(const FCameraVariableTable& OtherTab
 
 void FCameraVariableTable::LerpAll(const FCameraVariableTable& ToTable, float Factor)
 {
-	const ECameraVariableTableFilter Filter = ECameraVariableTableFilter::All;
+	const ECameraVariableTableFilter Filter = ECameraVariableTableFilter::AllPublic;
 	InternalLerp(ToTable, Filter, Factor, nullptr, false, nullptr);
 }
 
@@ -445,6 +446,7 @@ void FCameraVariableTable::InternalLerp(const FCameraVariableTable& ToTable, ECa
 	const bool bChangedOnly = EnumHasAnyFlags(Filter, ECameraVariableTableFilter::ChangedOnly);
 	const bool bInputs = EnumHasAnyFlags(Filter, ECameraVariableTableFilter::Input);
 	const bool bOutputs = EnumHasAnyFlags(Filter, ECameraVariableTableFilter::Output);
+	const bool bPrivates = EnumHasAnyFlags(Filter, ECameraVariableTableFilter::Private);
 
 	for (const FEntry& ToEntry : ToTable.Entries)
 	{
@@ -454,7 +456,7 @@ void FCameraVariableTable::InternalLerp(const FCameraVariableTable& ToTable, ECa
 		if (EnumHasAnyFlags(ToFlags, EEntryFlags::Written)
 				&& (!bChangedOnly || EnumHasAnyFlags(ToFlags, EEntryFlags::WrittenThisFrame))
 				&& ((bInputs && bToEntryIsInput) || (bOutputs && !bToEntryIsInput))
-				&& !EnumHasAnyFlags(ToFlags, EEntryFlags::Private)
+				&& (bPrivates || !EnumHasAnyFlags(ToFlags, EEntryFlags::Private))
 				&& IsVariableInMask(ToEntry.ID, InMask, bInvertMask))
 		{
 			// See if we know this variable.
diff --git a/Engine/Plugins/Cameras/GameplayCameras/Source/GameplayCameras/Public/Core/CameraVariableTable.h b/Engine/Plugins/Cameras/GameplayCameras/Source/GameplayCameras/Public/Core/CameraVariableTable.h
--- a/Engine/Plugins/Cameras/GameplayCameras/Source/GameplayCameras/Public/Core/CameraVariableTable.h
+++ b/Engine/Plugins/Cameras/GameplayCameras/Source/GameplayCameras/Public/Core/CameraVariableTable.h
@@ -26,17 +26,19 @@ struct TCameraVariableInterpolation;
 enum class ECameraVariableTableFilter
 {
 	None = 0,
-	/** Only include input variables. */
+	/** Include input variables. */
 	Input = 1 << 0,
-	/** Only include output variables (i.e. anything not an input). */
+	/** Include output variables (i.e. anything not an input). */
 	Output = 1 << 1,
-	/** Only include changed variables. */
+	/** Include private variables. */
+	Private = 1 << 3,
+	/** Include changed variables. */
 	ChangedOnly = 1 << 2,
 
 	/** All variables. */
-	All = Input | Output,
+	AllPublic = Input | Output,
 	/** All changed variables. */
-	AllChanged = Input | Output | ChangedOnly
+	AllPublicChanged = Input | Output | ChangedOnly
 };
 ENUM_CLASS_FLAGS(ECameraVariableTableFilter)
 
