Cameras: add build process API for rig transitions and transition conditions, move transition matching code to the transition class

diff --git a/Engine/Plugins/Cameras/GameplayCameras/Source/GameplayCameras/Private/Core/BlendStackCameraNode.cpp b/Engine/Plugins/Cameras/GameplayCameras/Source/GameplayCameras/Private/Core/BlendStackCameraNode.cpp
--- a/Engine/Plugins/Cameras/GameplayCameras/Source/GameplayCameras/Private/Core/BlendStackCameraNode.cpp
+++ b/Engine/Plugins/Cameras/GameplayCameras/Source/GameplayCameras/Private/Core/BlendStackCameraNode.cpp
@@ -652,16 +652,7 @@ const UCameraRigTransition* FBlendStackCameraNodeEvaluator::FindTransition(
 	// The transition should be used if all its conditions pass.
 	for (TObjectPtr<const UCameraRigTransition> Transition : Transitions)
 	{
-		bool bConditionsPass = true;
-		for (const UCameraRigTransitionCondition* Condition : Transition->Conditions)
-		{
-			if (!Condition->TransitionMatches(MatchParams))
-			{
-				bConditionsPass = false;
-				break;
-			}
-		}
-
+		const bool bConditionsPass = Transition->AllConditionsMatch(MatchParams);
 		if (bConditionsPass)
 		{
 			return Transition;
diff --git a/Engine/Plugins/Cameras/GameplayCameras/Source/GameplayCameras/Private/Core/CameraRigTransition.cpp b/Engine/Plugins/Cameras/GameplayCameras/Source/GameplayCameras/Private/Core/CameraRigTransition.cpp
--- a/Engine/Plugins/Cameras/GameplayCameras/Source/GameplayCameras/Private/Core/CameraRigTransition.cpp
+++ b/Engine/Plugins/Cameras/GameplayCameras/Source/GameplayCameras/Private/Core/CameraRigTransition.cpp
@@ -2,8 +2,16 @@
 
 #include "Core/CameraRigTransition.h"
 
+#include "Core/BlendCameraNode.h"
+#include "Core/CameraBuildLog.h"
+#include "Core/CameraRigBuildContext.h"
+#include "Logging/TokenizedMessage.h"
+#include "Nodes/Blends/SimpleBlendCameraNode.h"
+
 #include UE_INLINE_GENERATED_CPP_BY_NAME(CameraRigTransition)
 
+#define LOCTEXT_NAMESPACE "CameraRigTransition"
+
 void UCameraRigTransitionCondition::PostLoad()
 {
 #if WITH_EDITORONLY_DATA
@@ -26,6 +34,11 @@ bool UCameraRigTransitionCondition::TransitionMatches(const FCameraRigTransition
 	return OnTransitionMatches(Params);
 }
 
+void UCameraRigTransitionCondition::Build(FCameraRigBuildContext& BuildContext)
+{
+	OnBuild(BuildContext);
+}
+
 #if WITH_EDITOR
 
 void UCameraRigTransitionCondition::GetGraphNodePosition(FName InGraphName, int32& NodePosX, int32& NodePosY) const
@@ -73,6 +86,47 @@ void UCameraRigTransition::PostLoad()
 	Super::PostLoad();
 }
 
+bool UCameraRigTransition::AllConditionsMatch(const FCameraRigTransitionConditionMatchParams& Params) const
+{
+	for (const UCameraRigTransitionCondition* Condition : Conditions)
+	{
+		if (Condition && !Condition->TransitionMatches(Params))
+		{
+			return false;
+		}
+	}
+
+	return true;
+}
+
+void UCameraRigTransition::Build(FCameraRigBuildContext& BuildContext)
+{
+	if (Blend)
+	{
+		Blend->Build(BuildContext);
+	}
+	else
+	{
+		BuildContext.BuildLog.AddMessage(
+				EMessageSeverity::Error, this,
+				LOCTEXT("NullBlendError", "No blend defined on transition. To make a straight-cut transition, use the Pop blend."));
+	}
+
+	for (UCameraRigTransitionCondition* Condition : Conditions)
+	{
+		if (Condition)
+		{
+			Condition->Build(BuildContext);
+		}
+		else
+		{
+			BuildContext.BuildLog.AddMessage(
+					EMessageSeverity::Error, this,
+					LOCTEXT("NullConditionError", "Found an invalid transition condition."));
+		}
+	}
+}
+
 #if WITH_EDITOR
 
 void UCameraRigTransition::GetGraphNodePosition(FName InGraphName, int32& NodePosX, int32& NodePosY) const
@@ -103,3 +157,5 @@ void UCameraRigTransition::OnUpdateGraphNodeCommentText(FName InGraphName, const
 
 #endif  // WITH_EDITOR
 
+#undef LOCTEXT_NAMESPACE
+
diff --git a/Engine/Plugins/Cameras/GameplayCameras/Source/GameplayCameras/Public/Core/CameraRigTransition.h b/Engine/Plugins/Cameras/GameplayCameras/Source/GameplayCameras/Public/Core/CameraRigTransition.h
--- a/Engine/Plugins/Cameras/GameplayCameras/Source/GameplayCameras/Public/Core/CameraRigTransition.h
+++ b/Engine/Plugins/Cameras/GameplayCameras/Source/GameplayCameras/Public/Core/CameraRigTransition.h
@@ -13,6 +13,11 @@ class UBlendCameraNode;
 class UCameraAsset;
 class UCameraRigAsset;
 
+namespace UE::Cameras
+{
+	struct FCameraRigBuildContext;
+}
+
 /**
  * Parameter structure for camera transitions.
  */
@@ -44,14 +49,22 @@ class UCameraRigTransitionCondition
 
 public:
 
+	using FCameraRigBuildContext = UE::Cameras::FCameraRigBuildContext;
+
 	/** Evaluates whether this transition should be used. */
 	bool TransitionMatches(const FCameraRigTransitionConditionMatchParams& Params) const;
 
+	/** Build process callback for this transition. */
+	void Build(FCameraRigBuildContext& BuildContext);
+
 protected:
 
 	/** Evaluates whether this transition should be used. */
 	virtual bool OnTransitionMatches(const FCameraRigTransitionConditionMatchParams& Params) const { return false; }
 
+	/** Build process callback for this transition. */
+	virtual void OnBuild(FCameraRigBuildContext& BuildContext) {}
+
 protected:
 
 	// UObject interface.
@@ -142,6 +155,16 @@ public:
 	UPROPERTY(EditAnywhere, Category="Transition")
 	bool bOverrideInitialOrientation = false;
 
+public:
+
+	using FCameraRigBuildContext = UE::Cameras::FCameraRigBuildContext;
+
+	/** Returns whether all transition condition matches the given parameters. */
+	bool AllConditionsMatch(const FCameraRigTransitionConditionMatchParams& Params) const;
+
+	/** Build process callback for this transition. */
+	void Build(FCameraRigBuildContext& BuildContext);
+
 protected:
 
 	// UObject interface.
