Cameras: add support for camera rig merging

When a new camera rig instance gets pushed on top of a "similar" one on the blend stack, we optionally merge them together. Camera rigs are considered "similar" if they both boil down to the same camera rig prefab (probably with different interface parameter overrides). Merged camera rigs only run their underlying prefab rig once, using a blended version of their parameter overrides.

For example, consider that camera rig B is pushed on top of camera rig A. They both run prefab rig C, which exposes a parameter named "BoomLength". A sets BoomLength to 10, and B sets it to 20. When we merge rigs A and B together, we blend these overrides (e.g. 50% of A and 50% of B, giving us a BoomLength of 15) and then run prefab rig C with them (so C will run with BoomLength 15).

Camera rig merging is opt-in only, with an "advanced" option on transitions.

- The blend stack root node now can maintain a stack of parameter overrides. It runs it in the parameter blending phase.
- Make the prefab node (Camera Rig node) also apply overrides in the parameter blending phase.
- Optimize things a bit by only applying overrides on initialization if they are not dynamic (i.e. if they're not themselves driven by variables).
- Revamp how the blend stack determines the top active camera rig.
- Move the interface parameter overriding code into a utilty class, so both the prefab node and the blend stack root node can use it.
- Make it possible for blend nodes to blend private variables, which is needed to blend merged camera rigs' parameter overrides.
- Fix various other bugs with input variables not being correctly blended in all cases.

diff --git a/Engine/Plugins/Cameras/GameplayCameras/Source/GameplayCameras/Private/Core/BlendStackCameraNode.cpp b/Engine/Plugins/Cameras/GameplayCameras/Source/GameplayCameras/Private/Core/BlendStackCameraNode.cpp
--- a/Engine/Plugins/Cameras/GameplayCameras/Source/GameplayCameras/Private/Core/BlendStackCameraNode.cpp
+++ b/Engine/Plugins/Cameras/GameplayCameras/Source/GameplayCameras/Private/Core/BlendStackCameraNode.cpp
@@ -12,7 +12,6 @@
 #include "Debug/CameraDebugBlockBuilder.h"
 #include "Debug/CameraDebugRenderer.h"
 #include "Debug/CameraNodeEvaluationResultDebugBlock.h"
-#include "Debug/CameraNodeEvaluatorDebugBlock.h"
 #include "Debug/CameraPoseDebugBlock.h"
 #include "Debug/VariableTableDebugBlock.h"
 #include "HAL/IConsoleManager.h"
@@ -21,6 +20,7 @@
 #include "Math/ColorList.h"
 #include "Modules/ModuleManager.h"
 #include "Nodes/Blends/PopBlendCameraNode.h"
+#include "Nodes/Common/CameraRigCameraNode.h"
 
 #include UE_INLINE_GENERATED_CPP_BY_NAME(BlendStackCameraNode)
 
@@ -55,19 +55,59 @@ FBlendStackCameraNodeEvaluator::~FBlendStackCameraNodeEvaluator()
 
 void FBlendStackCameraNodeEvaluator::Push(const FBlendStackCameraPushParams& Params)
 {
+	bool bSearchedForTransition = false;
+	const UCameraRigTransition* Transition = nullptr;
+
 	if (!Entries.IsEmpty())
 	{
-		// Don't push anything if what is being requested is already the
-		// active camera rig.
-		const FCameraRigEntry& TopEntry(Entries.Top());
-		if (!TopEntry.bIsFrozen
-				&& TopEntry.CameraRig == Params.CameraRig
+		FCameraRigEntry& TopEntry(Entries.Top());
+		if (!TopEntry.bIsFrozen 
 				&& TopEntry.EvaluationContext == Params.EvaluationContext)
 		{
-			return;
+			// Don't push anything is what is being requested is already the active 
+			// camera rig.
+			if (TopEntry.CameraRig == Params.CameraRig)
+			{
+				return;
+			}
+
+			// See if we can merge the new camera rig onto the active camera rig.
+			const EBlendStackEntryComparison Comparison = TopEntry.RootEvaluator->Compare(Params.CameraRig);
+
+			if (Comparison == EBlendStackEntryComparison::Active)
+			{
+				// This camera rig is already the active one on the merged stack.
+				return;
+			}
+
+			if (Comparison == EBlendStackEntryComparison::EligibleForMerge)
+			{
+				// This camera rig can be merged with the one current running. However, we
+				// only do it if the transition explicitly allows it.
+				bSearchedForTransition = true;
+				Transition = FindTransition(Params);
+
+				if (Transition && Transition->bAllowCameraRigMerging)
+				{
+					PushVariantEntry(Params, Transition);
+					return;
+				}
+			}
 		}
 	}
 
+	// It's a legitimate new entry in the blend stack.
+	if (!bSearchedForTransition)
+	{
+		bSearchedForTransition = true;
+		Transition = FindTransition(Params);
+	}
+
+	PushNewEntry(Params, Transition);
+}
+
+void FBlendStackCameraNodeEvaluator::PushNewEntry(const FBlendStackCameraPushParams& Params, const UCameraRigTransition* Transition)
+{
 	// Create the new root node to wrap the new camera rig's root node, and the specific
 	// blend node for this transition.
 	// We need to const-cast here to be able to use our own blend stack node as the outer
@@ -81,7 +121,7 @@ void FBlendStackCameraNodeEvaluator::Push(const FBlendStackCameraPushParams& Par
 		// Find a transition and use its blend. If no transition is found,
 		// make a camera cut transition.
 		UBlendCameraNode* ModeBlend = nullptr;
-		if (const UCameraRigTransition* Transition = FindTransition(Params))
+		if (Transition)
 		{
 			ModeBlend = Transition->Blend;
 			UsedTransition = Transition;
@@ -107,18 +147,9 @@ void FBlendStackCameraNodeEvaluator::Push(const FBlendStackCameraPushParams& Par
 	}
 
 #if WITH_EDITOR
-	IGameplayCamerasModule& GameplayCamerasModule = FModuleManager::GetModuleChecked<IGameplayCamerasModule>("GameplayCameras");
-	TSharedPtr<IGameplayCamerasLiveEditManager> LiveEditManager = GameplayCamerasModule.GetLiveEditManager();
-	Params.CameraRig->GatherPackages(NewEntry.ListenedPackages);
-	for (const UPackage* ListenPackage : NewEntry.ListenedPackages)
-	{
-		int32& NumListens = AllListenedPackages.FindOrAdd(ListenPackage, 0);
-		if (NumListens == 0)
-		{
-			LiveEditManager->AddListener(ListenPackage, this);
-		}
-		++NumListens;
-	}
+	// Listen to changes to the packages inside which this camera rig is defined. We will hot-reload the
+	// camera node evaluators for this camera rig when we detect changes.
+	AddPackageListeners(NewEntry);
 #endif  // WITH_EDITOR
 
 	// Important: we need to move the new entry here because copying evaluator storage
@@ -131,6 +162,28 @@ void FBlendStackCameraNodeEvaluator::Push(const FBlendStackCameraPushParams& Par
 	}
 }
 
+void FBlendStackCameraNodeEvaluator::PushVariantEntry(const FBlendStackCameraPushParams& PushParams, const UCameraRigTransition* Transition)
+{
+	const UCameraRigCameraNode* PrefabNode = Cast<const UCameraRigCameraNode>(PushParams.CameraRig->RootNode);
+	const UBlendCameraNode* Blend = Transition ? Transition->Blend : nullptr;
+
+	FCameraRigEntry& TopEntry = Entries.Top();
+	FCameraNodeEvaluatorBuilder Builder(TopEntry.EvaluatorStorage);
+	FCameraNodeEvaluatorBuildParams BuildParams(Builder);
+	TopEntry.RootEvaluator->MergeCameraRig(BuildParams, PrefabNode, Blend);
+
+	// Swap out the camera rig registered as "active" for this entry.
+#if WITH_EDITOR
+	RemoveListenedPackages(TopEntry);
+#endif
+	{
+		TopEntry.CameraRig = PushParams.CameraRig;
+	}
+#if WITH_EDITOR
+	AddPackageListeners(TopEntry);
+#endif
+}
+
 bool FBlendStackCameraNodeEvaluator::InitializeEntry(
 		FCameraRigEntry& NewEntry, 
 		const UCameraRigAsset* CameraRig,
@@ -527,6 +580,27 @@ void FBlendStackCameraNodeEvaluator::PopEntries(int32 FirstIndexToKeep)
 
 #if WITH_EDITOR
 
+void FBlendStackCameraNodeEvaluator::AddPackageListeners(FCameraRigEntry& Entry)
+{
+	if (!ensure(Entry.CameraRig))
+	{
+		return;
+	}
+
+	IGameplayCamerasModule& GameplayCamerasModule = FModuleManager::GetModuleChecked<IGameplayCamerasModule>("GameplayCameras");
+	TSharedPtr<IGameplayCamerasLiveEditManager> LiveEditManager = GameplayCamerasModule.GetLiveEditManager();
+	Entry.CameraRig->GatherPackages(Entry.ListenedPackages);
+	for (const UPackage* ListenPackage : Entry.ListenedPackages)
+	{
+		int32& NumListens = AllListenedPackages.FindOrAdd(ListenPackage, 0);
+		if (NumListens == 0)
+		{
+			LiveEditManager->AddListener(ListenPackage, this);
+		}
+		++NumListens;
+	}
+}
+
 void FBlendStackCameraNodeEvaluator::RemoveListenedPackages(FCameraRigEntry& Entry)
 {
 	IGameplayCamerasModule& GameplayCamerasModule = FModuleManager::GetModuleChecked<IGameplayCamerasModule>("GameplayCameras");
diff --git a/Engine/Plugins/Cameras/GameplayCameras/Source/GameplayCameras/Private/Core/BlendStackRootCameraNode.cpp b/Engine/Plugins/Cameras/GameplayCameras/Source/GameplayCameras/Private/Core/BlendStackRootCameraNode.cpp
--- a/Engine/Plugins/Cameras/GameplayCameras/Source/GameplayCameras/Private/Core/BlendStackRootCameraNode.cpp
+++ b/Engine/Plugins/Cameras/GameplayCameras/Source/GameplayCameras/Private/Core/BlendStackRootCameraNode.cpp
@@ -2,11 +2,15 @@
 
 #include "Core/BlendStackRootCameraNode.h"
 
+#include "Algo/Transform.h"
 #include "Core/BlendCameraNode.h"
 #include "Core/CameraRigAsset.h"
+#include "Core/CameraRigAssetReference.h"
+#include "Core/CameraRigParameterOverrideEvaluator.h"
 #include "Debug/CameraDebugBlock.h"
 #include "Debug/CameraDebugBlockBuilder.h"
 #include "Debug/CameraDebugRenderer.h"
+#include "Nodes/Common/CameraRigCameraNode.h"
 
 #include UE_INLINE_GENERATED_CPP_BY_NAME(BlendStackRootCameraNode)
 
@@ -37,10 +41,16 @@ UE_DEFINE_CAMERA_NODE_EVALUATOR(FBlendStackRootCameraNodeEvaluator)
 
 UE_DECLARE_CAMERA_DEBUG_BLOCK_START(GAMEPLAYCAMERAS_API, FBlendStackRootCameraDebugBlock)
 	UE_DECLARE_CAMERA_DEBUG_BLOCK_FIELD(FString, CameraRigAssetName)
+	UE_DECLARE_CAMERA_DEBUG_BLOCK_FIELD(TArray<FString>, BlendedParameterOverridesEntries)
 UE_DECLARE_CAMERA_DEBUG_BLOCK_END()
 
 UE_DEFINE_CAMERA_DEBUG_BLOCK_WITH_FIELDS(FBlendStackRootCameraDebugBlock)
 
+FBlendStackRootCameraNodeEvaluator::FBlendStackRootCameraNodeEvaluator()
+{
+	AddNodeEvaluatorFlags(ECameraNodeEvaluatorFlags::NeedsParameterUpdate);
+}
+
 FCameraNodeEvaluatorChildrenView FBlendStackRootCameraNodeEvaluator::OnGetChildren()
 {
 	FCameraNodeEvaluatorChildrenView Children;
@@ -48,6 +58,10 @@ FCameraNodeEvaluatorChildrenView FBlendStackRootCameraNodeEvaluator::OnGetChildr
 	{
 		Children.Add(BlendEvaluator);
 	}
+	for (FBlendedParameterOverrides& BlendedParameterOverrides : BlendedParameterOverridesStack)
+	{
+		Children.Add(BlendedParameterOverrides.BlendEvaluator);
+	}
 	if (RootEvaluator)
 	{
 		Children.Add(RootEvaluator);
@@ -58,22 +72,145 @@ FCameraNodeEvaluatorChildrenView FBlendStackRootCameraNodeEvaluator::OnGetChildr
 void FBlendStackRootCameraNodeEvaluator::OnBuild(const FCameraNodeEvaluatorBuildParams& Params)
 {
 	const UBlendStackRootCameraNode* RootNode = GetCameraNodeAs<UBlendStackRootCameraNode>();
+
 	BlendEvaluator = Params.BuildEvaluatorAs<FBlendCameraNodeEvaluator>(RootNode->Blend);
 	RootEvaluator = Params.BuildEvaluator(RootNode->RootNode);
 }
 
 void FBlendStackRootCameraNodeEvaluator::OnInitialize(const FCameraNodeEvaluatorInitializeParams& Params, FCameraNodeEvaluationResult& OutResult)
 {
-#if UE_GAMEPLAY_CAMERAS_DEBUG
 	const UBlendStackRootCameraNode* RootNode = GetCameraNodeAs<UBlendStackRootCameraNode>();
 	if (RootNode->RootNode)
 	{
+		// See if the camera rig we are running is only made up of a prefab with some overrides.
+		InitialPrefabNode = Cast<const UCameraRigCameraNode>(RootNode->RootNode);
+		if (InitialPrefabNode)
+		{
+			BlendablePrefabCameraRig = InitialPrefabNode->CameraRigReference.GetCameraRig();
+		}
+
+#if UE_GAMEPLAY_CAMERAS_DEBUG
 		if (const UCameraRigAsset* CameraRig = RootNode->RootNode->GetTypedOuter<UCameraRigAsset>())
 		{
 			CameraRigAssetName = GetNameSafe(CameraRig);
 		}
-	}
 #endif  // UE_GAMEPLAY_CAMERAS_DEBUG
+	}
+}
+
+EBlendStackEntryComparison FBlendStackRootCameraNodeEvaluator::Compare(const UCameraRigAsset* CameraRig) const
+{
+	if (!BlendablePrefabCameraRig)
+	{
+		return EBlendStackEntryComparison::Different;
+	}
+	
+	const UCameraRigCameraNode* PrefabNode = Cast<const UCameraRigCameraNode>(CameraRig->RootNode);
+	if (!PrefabNode)
+	{
+		return EBlendStackEntryComparison::Different;
+	}
+
+	const UCameraRigAsset* Prefab = PrefabNode->CameraRigReference.GetCameraRig();
+	if (!Prefab)
+	{
+		return EBlendStackEntryComparison::Different;
+	}
+
+	if (Prefab != BlendablePrefabCameraRig)
+	{
+		return EBlendStackEntryComparison::Different;
+	}
+
+	if (BlendedParameterOverridesStack.IsEmpty())
+	{
+		return EBlendStackEntryComparison::EligibleForMerge;
+	}
+	
+	const FBlendedParameterOverrides& TopEntry = BlendedParameterOverridesStack.Top();
+	if (TopEntry.PrefabNodeAsset == CameraRig)
+	{
+		return EBlendStackEntryComparison::Active;
+	}
+
+	return EBlendStackEntryComparison::EligibleForMerge;
+}
+
+void FBlendStackRootCameraNodeEvaluator::MergeCameraRig(const FCameraNodeEvaluatorBuildParams& Params, const UCameraRigCameraNode* PrefabNode, const UBlendCameraNode* Blend)
+{
+	if (!ensureMsgf(PrefabNode, TEXT("No prefab node given.")))
+	{
+		return;
+	}
+
+	if (!ensureMsgf(BlendablePrefabCameraRig,
+				TEXT("Adding blended parameter overrides for a camera rig that doesn't support it.")))
+	{
+		return;
+	}
+
+	if (!ensureMsgf(PrefabNode->CameraRigReference.GetCameraRig() == BlendablePrefabCameraRig,
+				TEXT("Adding blended parameter overrides for a different camera rig.")))
+	{
+		return;
+	}
+
+	InitializeBlendedParameterOverridesStack();
+
+	FBlendedParameterOverrides BlendedParameterOverrides;
+	BlendedParameterOverrides.PrefabNodeAsset = PrefabNode->GetTypedOuter<UCameraRigAsset>();
+	BlendedParameterOverrides.PrefabNode = PrefabNode;
+	BlendedParameterOverrides.Blend = Blend;
+	BlendedParameterOverrides.Result.VariableTable.Initialize(BlendedParameterOverridesTableAllocationInfo);
+	if (Blend)
+	{
+		BlendedParameterOverrides.BlendEvaluator = Params.BuildEvaluatorAs<FBlendCameraNodeEvaluator>(Blend);
+	}
+
+	BlendedParameterOverridesStack.Add(MoveTemp(BlendedParameterOverrides));
+}
+
+void FBlendStackRootCameraNodeEvaluator::InitializeBlendedParameterOverridesStack()
+{
+	if (!ensureMsgf(BlendablePrefabCameraRig, TEXT("The blended parameter overiddes stack has already been initialized.")))
+	{
+		return;
+	}
+
+	if (!BlendedParameterOverridesStack.IsEmpty())
+	{
+		return;
+	}
+
+	// Build the allocation info for the variable tables we keep with each set of parameter overrides.
+	for (const UCameraRigInterfaceParameter* InterfaceParameter : BlendablePrefabCameraRig->Interface.InterfaceParameters)
+	{
+		if (!ensure(InterfaceParameter))
+		{
+			continue;
+		}
+		if (!InterfaceParameter->PrivateVariable)
+		{
+			continue;
+		}
+
+		BlendedParameterOverridesTableAllocationInfo.VariableDefinitions.Add(
+				InterfaceParameter->PrivateVariable->GetVariableDefinition());;
+	}
+
+	FCameraRigCameraNodeEvaluator* RootPrefabNodeEvaluator = RootEvaluator->CastThisChecked<FCameraRigCameraNodeEvaluator>();
+	RootPrefabNodeEvaluator->SetApplyParameterOverrides(false);
+
+	FBlendedParameterOverrides InitialParameterOverrides;
+	InitialParameterOverrides.PrefabNodeAsset = InitialPrefabNode->GetTypedOuter<UCameraRigAsset>();
+	InitialParameterOverrides.PrefabNode = InitialPrefabNode;
+	InitialParameterOverrides.Result.VariableTable.Initialize(BlendedParameterOverridesTableAllocationInfo);
+	BlendedParameterOverridesStack.Add(MoveTemp(InitialParameterOverrides));
+}
+
+void FBlendStackRootCameraNodeEvaluator::OnUpdateParameters(const FCameraBlendedParameterUpdateParams& Params, FCameraBlendedParameterUpdateResult& OutResult)
+{
+	RunBlendedParameterOverridesStack(Params, OutResult);
 }
 
 void FBlendStackRootCameraNodeEvaluator::OnRun(const FCameraNodeEvaluationParams& Params, FCameraNodeEvaluationResult& OutResult)
@@ -88,12 +225,115 @@ void FBlendStackRootCameraNodeEvaluator::OnRun(const FCameraNodeEvaluationParams
 	}
 }
 
+void FBlendStackRootCameraNodeEvaluator::SetDefaultInterfaceParameterValues(FCameraVariableTable& OutVariableTable)
+{
+	for (const UCameraRigInterfaceParameter* InterfaceParameter : BlendablePrefabCameraRig->Interface.InterfaceParameters)
+	{
+		if (!ensure(InterfaceParameter))
+		{
+			continue;
+		}
+
+		const UCameraVariableAsset* PrivateVariable = InterfaceParameter->PrivateVariable;
+		if (!PrivateVariable)
+		{
+			continue;
+		}
+
+		OutVariableTable.SetValue(
+				PrivateVariable->GetVariableID(), 
+				PrivateVariable->GetVariableType(), 
+				PrivateVariable->GetDefaultValuePtr());
+	}
+}
+
+void FBlendStackRootCameraNodeEvaluator::RunBlendedParameterOverridesStack(const FCameraBlendedParameterUpdateParams& Params, FCameraBlendedParameterUpdateResult& OutResult)
+{
+	if (BlendedParameterOverridesStack.IsEmpty())
+	{
+		return;
+	}
+
+	if (!ensure(BlendablePrefabCameraRig))
+	{
+		return;
+	}
+
+	int32 PopEntriesBelow = INDEX_NONE;
+	for (int32 EntryIndex = 0; EntryIndex < BlendedParameterOverridesStack.Num(); ++EntryIndex)
+	{
+		FBlendedParameterOverrides& BlendedParameterOverrides(BlendedParameterOverridesStack[EntryIndex]);
+		FCameraNodeEvaluationResult& CurResult(BlendedParameterOverrides.Result);
+
+		// Start by setting the default values of all parameters. If we don't do this, parameter overrides
+		// wouldn't have a base value to blend from.
+		SetDefaultInterfaceParameterValues(CurResult.VariableTable);
+
+		// Next, override the defaults with the specific values of this entry.
+		FCameraRigParameterOverrideEvaluator OverrideEvaluator(BlendedParameterOverrides.PrefabNode->CameraRigReference);
+		OverrideEvaluator.ApplyParameterOverrides(CurResult.VariableTable, false);
+
+		// Finally, update the parameter overrides' blend, and apply it.
+		if (BlendedParameterOverrides.BlendEvaluator)
+		{
+			BlendedParameterOverrides.BlendEvaluator->Run(Params.EvaluationParams, CurResult);
+
+			FCameraNodePreBlendParams BlendParams(Params.EvaluationParams, Params.LastCameraPose, CurResult.VariableTable);
+			BlendParams.ExtraVariableTableFilter = ECameraVariableTableFilter::Private;
+			FCameraNodePreBlendResult BlendResult(OutResult.VariableTable);
+			BlendedParameterOverrides.BlendEvaluator->BlendParameters(BlendParams, BlendResult);
+
+			if (BlendResult.bIsBlendFinished && BlendResult.bIsBlendFull)
+			{
+				PopEntriesBelow = EntryIndex;
+			}
+		}
+		else
+		{
+			OutResult.VariableTable.Override(
+					CurResult.VariableTable, ECameraVariableTableFilter::Input | ECameraVariableTableFilter::Private);
+
+			PopEntriesBelow = EntryIndex;
+		}
+	}
+	if (PopEntriesBelow >= 0)
+	{
+		BlendedParameterOverridesStack.RemoveAt(0, PopEntriesBelow);
+	}
+}
+
+void FBlendStackRootCameraNodeEvaluator::OnAddReferencedObjects(FReferenceCollector& Collector)
+{
+	Collector.AddReferencedObject(BlendablePrefabCameraRig);
+	Collector.AddReferencedObject(InitialPrefabNode);
+
+	for (FBlendedParameterOverrides& BlendedParameterOverrides: BlendedParameterOverridesStack)
+	{
+		Collector.AddReferencedObject(BlendedParameterOverrides.PrefabNodeAsset);
+		Collector.AddReferencedObject(BlendedParameterOverrides.PrefabNode);
+		Collector.AddReferencedObject(BlendedParameterOverrides.Blend);
+	}
+}
+
 #if UE_GAMEPLAY_CAMERAS_DEBUG
 
 void FBlendStackRootCameraNodeEvaluator::OnBuildDebugBlocks(const FCameraDebugBlockBuildParams& Params, FCameraDebugBlockBuilder& Builder)
 {
 	FBlendStackRootCameraDebugBlock& DebugBlock = Builder.StartChildDebugBlock<FBlendStackRootCameraDebugBlock>();
 	DebugBlock.CameraRigAssetName = CameraRigAssetName;
+	Algo::Transform(BlendedParameterOverridesStack, DebugBlock.BlendedParameterOverridesEntries,
+			[](const FBlendedParameterOverrides& Item)
+			{
+				if (Item.PrefabNode)
+				{
+					const UCameraRigAsset* OuterCameraRig = Item.PrefabNode->GetTypedOuter<const UCameraRigAsset>();
+					if (OuterCameraRig)
+					{
+						return OuterCameraRig->GetDisplayName();
+					}
+				}
+				return FString(TEXT("<invalid camera rig>"));
+			});
 
 	if (BlendEvaluator)
 	{
@@ -106,6 +346,22 @@ void FBlendStackRootCameraNodeEvaluator::OnBuildDebugBlocks(const FCameraDebugBl
 		Builder.EndChildDebugBlock();
 	}
 
+	Builder.StartChildDebugBlock<FCameraDebugBlock>();
+	for (const FBlendedParameterOverrides& BlendedParameterOverrides : BlendedParameterOverridesStack)
+	{
+		if (BlendedParameterOverrides.BlendEvaluator)
+		{
+			BlendedParameterOverrides.BlendEvaluator->BuildDebugBlocks(Params, Builder);
+		}
+		else
+		{
+			// Dummy block.
+			Builder.StartChildDebugBlock<FCameraDebugBlock>();
+			Builder.EndChildDebugBlock();
+		}
+	}
+	Builder.EndChildDebugBlock();
+
 	if (RootEvaluator)
 	{
 		RootEvaluator->BuildDebugBlocks(Params, Builder);
@@ -125,14 +381,30 @@ void FBlendStackRootCameraDebugBlock::OnDebugDraw(const FCameraDebugBlockDrawPar
 {
 	TArrayView<FCameraDebugBlock*> ChildrenView(GetChildren());
 
-	Renderer.AddText("{cam_passive}<Blend>\n");
+	Renderer.AddText(TEXT("{cam_passive}<Blend>\n"));
 	Renderer.AddIndent();
 	ChildrenView[0]->DebugDraw(Params, Renderer);
 	Renderer.RemoveIndent();
 
+	if (!BlendedParameterOverridesEntries.IsEmpty())
+	{
+		Renderer.AddText(TEXT("{cam_passive}<%d Merged Camera Rigs>\n"), BlendedParameterOverridesEntries.Num());
+		Renderer.AddIndent();
+		{
+			int32 ChildIndex = 0;
+			for (FCameraDebugBlock* ParameterOverridesDebugBlock : ChildrenView[1]->GetChildren())
+			{
+				const FString& ParameterOverridesCameraRigName = BlendedParameterOverridesEntries[ChildIndex++];
+				Renderer.AddText(ParameterOverridesCameraRigName);
+				ParameterOverridesDebugBlock->DebugDraw(Params, Renderer);
+			}
+		}
+		Renderer.RemoveIndent();
+	}
+
 	Renderer.AddText(TEXT("{cam_passive}<CameraRig %s>\n"), *CameraRigAssetName);
 	Renderer.AddIndent();
-	ChildrenView[1]->DebugDraw(Params, Renderer);
+	ChildrenView[2]->DebugDraw(Params, Renderer);
 	Renderer.RemoveIndent();
 
 	Renderer.SkipAllBlocks();
diff --git a/Engine/Plugins/Cameras/GameplayCameras/Source/GameplayCameras/Private/Core/BlendStackRootCameraNode.h b/Engine/Plugins/Cameras/GameplayCameras/Source/GameplayCameras/Private/Core/BlendStackRootCameraNode.h
--- a/Engine/Plugins/Cameras/GameplayCameras/Source/GameplayCameras/Private/Core/BlendStackRootCameraNode.h
+++ b/Engine/Plugins/Cameras/GameplayCameras/Source/GameplayCameras/Private/Core/BlendStackRootCameraNode.h
@@ -9,6 +9,7 @@
 
 class UBlendCameraNode;
 class UCameraRigAsset;
+class UCameraRigCameraNode;
 
 /**
  * Root camera node for running a camera rig in a blend stack.
@@ -42,6 +43,13 @@ namespace UE::Cameras
 
 class FBlendCameraNodeEvaluator;
 
+enum class EBlendStackEntryComparison
+{
+	Different,
+	EligibleForMerge,
+	Active
+};
+
 /**
  * Evaluator for the blend stack entry root node.
  */
@@ -51,26 +59,53 @@ class FBlendStackRootCameraNodeEvaluator : public FCameraNodeEvaluator
 
 public:
 
+	FBlendStackRootCameraNodeEvaluator();
+
 	FBlendCameraNodeEvaluator* GetBlendEvaluator() const { return BlendEvaluator; }
 	FCameraNodeEvaluator* GetRootEvaluator() const { return RootEvaluator; }
 
+	EBlendStackEntryComparison Compare(const UCameraRigAsset* CameraRig) const;
+	void MergeCameraRig(const FCameraNodeEvaluatorBuildParams& Params, const UCameraRigCameraNode* PrefabNode, const UBlendCameraNode* Blend);
+
 protected:
 
 	// FCameraNodeEvaluator interface.
 	virtual FCameraNodeEvaluatorChildrenView OnGetChildren() override;
 	virtual void OnBuild(const FCameraNodeEvaluatorBuildParams& Params) override;
 	virtual void OnInitialize(const FCameraNodeEvaluatorInitializeParams& Params, FCameraNodeEvaluationResult& OutResult) override;
+	virtual void OnUpdateParameters(const FCameraBlendedParameterUpdateParams& Params, FCameraBlendedParameterUpdateResult& OutResult) override;
 	virtual void OnRun(const FCameraNodeEvaluationParams& Params, FCameraNodeEvaluationResult& OutResult) override;
+	virtual void OnAddReferencedObjects(FReferenceCollector& Collector) override;
 
 #if UE_GAMEPLAY_CAMERAS_DEBUG
 	virtual void OnBuildDebugBlocks(const FCameraDebugBlockBuildParams& Params, FCameraDebugBlockBuilder& Builder) override;
 #endif  // UE_GAMEPLAY_CAMERAS_DEBUG
 
+private:
+
+	void InitializeBlendedParameterOverridesStack();
+	void SetDefaultInterfaceParameterValues(FCameraVariableTable& OutVariableTable);
+	void RunBlendedParameterOverridesStack(const FCameraBlendedParameterUpdateParams& Params, FCameraBlendedParameterUpdateResult& OutResult);
+
 private:
 
 	FBlendCameraNodeEvaluator* BlendEvaluator = nullptr;
 	FCameraNodeEvaluator* RootEvaluator = nullptr;
 
+	TObjectPtr<const UCameraRigAsset> BlendablePrefabCameraRig;
+	TObjectPtr<const UCameraRigCameraNode> InitialPrefabNode;
+	FCameraVariableTableAllocationInfo BlendedParameterOverridesTableAllocationInfo;
+
+	struct FBlendedParameterOverrides
+	{
+		TObjectPtr<const UCameraRigAsset> PrefabNodeAsset;
+		TObjectPtr<const UCameraRigCameraNode> PrefabNode;
+		TObjectPtr<const UBlendCameraNode> Blend;
+		FBlendCameraNodeEvaluator* BlendEvaluator = nullptr;
+		FCameraNodeEvaluationResult Result;
+	};
+	TArray<FBlendedParameterOverrides> BlendedParameterOverridesStack;
+
 #if UE_GAMEPLAY_CAMERAS_DEBUG
 	FString CameraRigAssetName;
 #endif  // UE_GAMEPLAY_CAMERAS_DEBUG
diff --git a/Engine/Plugins/Cameras/GameplayCameras/Source/GameplayCameras/Private/Core/CameraIKAim.cpp b/Engine/Plugins/Cameras/GameplayCameras/Source/GameplayCameras/Private/Core/CameraIKAim.cpp
--- a/Engine/Plugins/Cameras/GameplayCameras/Source/GameplayCameras/Private/Core/CameraIKAim.cpp
+++ b/Engine/Plugins/Cameras/GameplayCameras/Source/GameplayCameras/Private/Core/CameraIKAim.cpp
@@ -62,6 +62,7 @@ bool FCameraIKAim::DoRun(const FCameraIKAimParams& Params, const FCameraRigEvalu
 	// Initialize our scratch result.
 	ScratchResult.VariableTable.Initialize(CameraRigInfo.CameraRig->AllocationInfo.VariableTableInfo);
 	ScratchResult.VariableTable.OverrideAll(CameraSystemEvaluator->GetEvaluatedResult().VariableTable);
+	ScratchResult.VariableTable.Override(CameraRigInfo.LastResult->VariableTable, ECameraVariableTableFilter::AllPublic | ECameraVariableTableFilter::Private);
 
 	// Initialize our hierarchy caches.
 	FRootCameraNodeEvaluator* CameraSystemRootEvaluator = CameraSystemEvaluator->GetRootNodeEvaluator();
diff --git a/Engine/Plugins/Cameras/GameplayCameras/Source/GameplayCameras/Private/Core/CameraRigParameterOverrideEvaluator.cpp b/Engine/Plugins/Cameras/GameplayCameras/Source/GameplayCameras/Private/Core/CameraRigParameterOverrideEvaluator.cpp
new file mode 100644
--- /dev/null
+++ b/Engine/Plugins/Cameras/GameplayCameras/Source/GameplayCameras/Private/Core/CameraRigParameterOverrideEvaluator.cpp
@@ -0,0 +1,85 @@
+// Copyright Epic Games, Inc. All Rights Reserved.
+
+#include "Core/CameraRigParameterOverrideEvaluator.h"
+
+#include "Core/CameraRigAsset.h"
+#include "Core/CameraRigAssetReference.h"
+#include "Core/CameraVariableTable.h"
+
+namespace UE::Cameras
+{
+
+namespace Internal
+{
+
+template<typename ParameterOverrideType>
+void ApplyParameterOverrides(
+		const UCameraRigAsset* CameraRig, 
+		TArrayView<const ParameterOverrideType> ParameterOverrides, 
+		FCameraVariableTable& OutVariableTable,
+		bool bDrivenOverridesOnly)
+{
+	for (const ParameterOverrideType& ParameterOverride : ParameterOverrides)
+	{
+		using ParameterType = decltype(ParameterOverrideType::Value);
+		using ValueType = typename ParameterType::ValueType;
+
+		if (!ParameterOverride.PrivateVariableGuid.IsValid())
+		{
+			// Ignore un-built parameter overrides in the editor since the user could have just added
+			// an override while PIE is running. They need to hit the Build button for the override
+			// to apply.
+			// Outside of the editor, report this as an error.
+#if !WITH_EDITOR
+			UE_LOG(LogCameraSystem, Error, 
+					TEXT("Invalid parameter override '%s' in camera rig '%s'. Was it built/cooked?"),
+					*ParameterOverride.InterfaceParameterName,
+					*GetPathNameSafe(CameraRig));
+#endif
+			continue;
+		}
+
+		FCameraVariableID InterfaceParameterID(FCameraVariableID::FromHashValue(GetTypeHash(ParameterOverride.PrivateVariableGuid)));
+
+		if (ParameterOverride.Value.Variable != nullptr)
+		{
+			// The override is driven by a variable... read its value and set it as the value for the
+			// prefab's variable. Basically, we forward the value from one variable to the next.
+			FCameraVariableDefinition OverrideDefinition(ParameterOverride.Value.Variable->GetVariableDefinition());
+
+			const ValueType OverrideValue = OutVariableTable.GetValue<ValueType>(
+					OverrideDefinition.VariableID, ParameterOverride.Value.Variable->GetDefaultValue());
+			OutVariableTable.SetValue<ValueType>(InterfaceParameterID, OverrideValue);
+		}
+		else if (!bDrivenOverridesOnly)
+		{
+			// The override is a fixed value. Just set that on the prefab's variable.
+			OutVariableTable.SetValue<ValueType>(InterfaceParameterID, ParameterOverride.Value.Value);
+		}
+	}
+}
+
+}  // namespace Internal
+
+FCameraRigParameterOverrideEvaluator::FCameraRigParameterOverrideEvaluator(const FCameraRigAssetReference& InCameraRigReference)
+	: CameraRigReference(InCameraRigReference)
+{
+}
+
+void FCameraRigParameterOverrideEvaluator::ApplyParameterOverrides(FCameraVariableTable& OutVariableTable, bool bDrivenOverridesOnly)
+{
+	const UCameraRigAsset* CameraRig = CameraRigReference.GetCameraRig();
+	const FCameraRigParameterOverrides& ParameterOverrides = CameraRigReference.GetParameterOverrides();
+
+#define UE_CAMERA_VARIABLE_FOR_TYPE(ValueType, ValueName)\
+	Internal::ApplyParameterOverrides(\
+			CameraRig,\
+			ParameterOverrides.Get##ValueName##Overrides(),\
+			OutVariableTable,\
+			bDrivenOverridesOnly);
+UE_CAMERA_VARIABLE_FOR_ALL_TYPES()
+#undef UE_CAMERA_VARIABLE_FOR_TYPE
+}
+
+}  // namespace UE::Cameras
+
diff --git a/Engine/Plugins/Cameras/GameplayCameras/Source/GameplayCameras/Private/Core/CameraRigParameterOverrideEvaluator.h b/Engine/Plugins/Cameras/GameplayCameras/Source/GameplayCameras/Private/Core/CameraRigParameterOverrideEvaluator.h
new file mode 100644
--- /dev/null
+++ b/Engine/Plugins/Cameras/GameplayCameras/Source/GameplayCameras/Private/Core/CameraRigParameterOverrideEvaluator.h
@@ -0,0 +1,37 @@
+// Copyright Epic Games, Inc. All Rights Reserved.
+
+#pragma once
+
+struct FCameraRigAssetReference;
+
+namespace UE::Cameras
+{
+
+class FCameraVariableTable;
+
+/**
+ * Utility class for applying interface parameter overrides to a camera rig via a
+ * given variable table.
+ */
+class FCameraRigParameterOverrideEvaluator
+{
+public:
+
+	/** Creates a new parameter override evaluator. */
+	FCameraRigParameterOverrideEvaluator(const FCameraRigAssetReference& InCameraRigReference);
+
+	/** 
+	 * Applies override values to the given variable table.
+	 *
+	 * @param OutVariableTable  The variable table in which to set the override values.
+	 * @param bDrivenOverridesOnly  Whether only overrides driven by variables should be applied.
+	 */
+	void ApplyParameterOverrides(FCameraVariableTable& OutVariableTable, bool bDrivenOverridesOnly = false);
+
+private:
+
+	const FCameraRigAssetReference& CameraRigReference;
+};
+
+}  // namespace UE::Cameras
+
diff --git a/Engine/Plugins/Cameras/GameplayCameras/Source/GameplayCameras/Private/Nodes/Blends/SimpleBlendCameraNode.cpp b/Engine/Plugins/Cameras/GameplayCameras/Source/GameplayCameras/Private/Nodes/Blends/SimpleBlendCameraNode.cpp
--- a/Engine/Plugins/Cameras/GameplayCameras/Source/GameplayCameras/Private/Nodes/Blends/SimpleBlendCameraNode.cpp
+++ b/Engine/Plugins/Cameras/GameplayCameras/Source/GameplayCameras/Private/Nodes/Blends/SimpleBlendCameraNode.cpp
@@ -29,6 +29,7 @@ void FSimpleBlendCameraNodeEvaluator::OnRun(const FCameraNodeEvaluationParams& P
 void FSimpleBlendCameraNodeEvaluator::OnBlendParameters(const FCameraNodePreBlendParams& Params, FCameraNodePreBlendResult& OutResult)
 {
 	const FCameraVariableTable& ChildVariableTable(Params.ChildVariableTable);
+	OutResult.VariableTable.Lerp(ChildVariableTable, ECameraVariableTableFilter::Input | Params.ExtraVariableTableFilter, BlendFactor);
 
 	OutResult.bIsBlendFull = BlendFactor >= 1.f;
 	OutResult.bIsBlendFinished = bIsBlendFinished;
diff --git a/Engine/Plugins/Cameras/GameplayCameras/Source/GameplayCameras/Private/Nodes/Common/CameraRigCameraNode.cpp b/Engine/Plugins/Cameras/GameplayCameras/Source/GameplayCameras/Private/Nodes/Common/CameraRigCameraNode.cpp
--- a/Engine/Plugins/Cameras/GameplayCameras/Source/GameplayCameras/Private/Nodes/Common/CameraRigCameraNode.cpp
+++ b/Engine/Plugins/Cameras/GameplayCameras/Source/GameplayCameras/Private/Nodes/Common/CameraRigCameraNode.cpp
@@ -6,6 +6,7 @@
 #include "Core/CameraNodeEvaluator.h"
 #include "Core/CameraRigAsset.h"
 #include "Core/CameraRigBuildContext.h"
+#include "Core/CameraRigParameterOverrideEvaluator.h"
 
 #include UE_INLINE_GENERATED_CPP_BY_NAME(CameraRigCameraNode)
 
@@ -14,95 +15,13 @@
 namespace UE::Cameras
 {
 
-namespace Internal
-{
+UE_DEFINE_CAMERA_NODE_EVALUATOR(FCameraRigCameraNodeEvaluator)
 
-template<typename ParameterOverrideType>
-void ApplyParameterOverrides(
-		const UCameraRigAsset* CameraRig, 
-		TArrayView<const ParameterOverrideType> ParameterOverrides, 
-		FCameraVariableTable& OutVariableTable)
+FCameraRigCameraNodeEvaluator::FCameraRigCameraNodeEvaluator()
 {
-	for (const ParameterOverrideType& ParameterOverride : ParameterOverrides)
-	{
-		using ParameterType = decltype(ParameterOverrideType::Value);
-		using ValueType = typename ParameterType::ValueType;
-
-		if (!ParameterOverride.PrivateVariableGuid.IsValid())
-		{
-#if WITH_EDITOR
-			// Ignore un-built parameter overrides in the editor since the user could have just added
-			// an override while PIE is running. They need to hit the Build button for the override
-			// to apply.
-			continue;
-#else
-			UE_LOG(LogCameraSystem, Error, 
-					TEXT("Invalid parameter override '%s' in camera rig '%s'. Was it built/cooked?"),
-					*ParameterOverride.InterfaceParameterName,
-					*GetPathNameSafe(CameraRig));
-#endif
-		}
-
-		FCameraVariableID InterfaceParameterID(FCameraVariableID::FromHashValue(GetTypeHash(ParameterOverride.PrivateVariableGuid)));
-
-		if (ParameterOverride.Value.Variable != nullptr)
-		{
-			// The override is driven by a variable... read its value and set it as the value for the
-			// prefab's variable. Basically, we forward the value from one variable to the next.
-			FCameraVariableDefinition OverrideDefinition(ParameterOverride.Value.Variable->GetVariableDefinition());
-
-			const ValueType* OverrideValuePtr = OutVariableTable.FindValue<ValueType>(OverrideDefinition.VariableID);
-			if (OverrideValuePtr)
-			{
-				OutVariableTable.SetValue<ValueType>(InterfaceParameterID, *OverrideValuePtr);
-			}
-			else
-			{
-				// Once again, ignore un-built data. Only emit errors when running outside of the editor.
-#if !WITH_EDITOR
-				UE_LOG(LogCameraSystem, Error, 
-						TEXT("Camera variable '%s' for parameter override '%s' in camera rig '%s' isn't in the variable "
-							"table. Was it built/cooked?"),
-						*GetNameSafe(ParameterOverride.Value.Variable),
-						*ParameterOverride.InterfaceParameterName,
-						*GetPathNameSafe(CameraRig));
-#endif
-			}
-		}
-		else
-		{
-			// The override is a fixed value. Just set that on the prefab's variable.
-			OutVariableTable.SetValue<ValueType>(
-					InterfaceParameterID,
-					ParameterOverride.Value.Value);
-		}
-	}
+	AddNodeEvaluatorFlags(ECameraNodeEvaluatorFlags::NeedsParameterUpdate);
 }
 
-}  // namespace Internal
-
-class FCameraRigCameraNodeEvaluator : public FCameraNodeEvaluator
-{
-	UE_DECLARE_CAMERA_NODE_EVALUATOR(GAMEPLAYCAMERAS_API, FCameraRigCameraNodeEvaluator)
-
-protected:
-
-	virtual FCameraNodeEvaluatorChildrenView OnGetChildren() override;
-	virtual void OnInitialize(const FCameraNodeEvaluatorInitializeParams& Params, FCameraNodeEvaluationResult& OutResult) override;
-	virtual void OnBuild(const FCameraNodeEvaluatorBuildParams& Params) override;
-	virtual void OnRun(const FCameraNodeEvaluationParams& Params, FCameraNodeEvaluationResult& OutResult) override;
-
-private:
-
-	void ApplyParameterOverrides(FCameraVariableTable& OutVariableTable);
-
-private:
-
-	FCameraNodeEvaluator* CameraRigRootEvaluator = nullptr;
-};
-
-UE_DEFINE_CAMERA_NODE_EVALUATOR(FCameraRigCameraNodeEvaluator)
-
 FCameraNodeEvaluatorChildrenView FCameraRigCameraNodeEvaluator::OnGetChildren()
 {
 	return FCameraNodeEvaluatorChildrenView({ CameraRigRootEvaluator });
@@ -123,35 +42,43 @@ void FCameraRigCameraNodeEvaluator::OnBuild(const FCameraNodeEvaluatorBuildParam
 void FCameraRigCameraNodeEvaluator::OnInitialize(const FCameraNodeEvaluatorInitializeParams& Params, FCameraNodeEvaluationResult& OutResult)
 {
 	// Apply overrides right away.
-	ApplyParameterOverrides(OutResult.VariableTable);
+	ApplyParameterOverrides(OutResult.VariableTable, false);
 }
 
-void FCameraRigCameraNodeEvaluator::OnRun(const FCameraNodeEvaluationParams& Params, FCameraNodeEvaluationResult& OutResult)
+void FCameraRigCameraNodeEvaluator::OnUpdateParameters(const FCameraBlendedParameterUpdateParams& Params, FCameraBlendedParameterUpdateResult& OutResult)
 {
 	// Keep applying overrides in case they are driven by a variable.
-	// TODO: we could skip this step for constant overrides.
-	ApplyParameterOverrides(OutResult.VariableTable);
+	const bool bDrivenOverridesOnly = true; //(Params.EvaluationParams.EvaluationType == ECameraNodeEvaluationType::Standard);
+	ApplyParameterOverrides(OutResult.VariableTable, bDrivenOverridesOnly);
+}
 
+void FCameraRigCameraNodeEvaluator::OnRun(const FCameraNodeEvaluationParams& Params, FCameraNodeEvaluationResult& OutResult)
+{
 	if (CameraRigRootEvaluator)
 	{
 		CameraRigRootEvaluator->Run(Params, OutResult);
 	}
 }
 
-void FCameraRigCameraNodeEvaluator::ApplyParameterOverrides(FCameraVariableTable& OutVariableTable)
+void FCameraRigCameraNodeEvaluator::ApplyParameterOverrides(FCameraVariableTable& OutVariableTable, bool bDrivenOnly)
 {
-	const UCameraRigCameraNode* PrefabNode = GetCameraNodeAs<UCameraRigCameraNode>();
+	if (bApplyParameterOverrides)
+	{
+		const UCameraRigCameraNode* PrefabNode = GetCameraNodeAs<UCameraRigCameraNode>();
 
-	const UCameraRigAsset* CameraRig = PrefabNode->CameraRigReference.GetCameraRig();
-	const FCameraRigParameterOverrides& ParameterOverrides = PrefabNode->CameraRigReference.GetParameterOverrides();
+		FCameraRigParameterOverrideEvaluator OverrideEvaluator(PrefabNode->CameraRigReference);
+		OverrideEvaluator.ApplyParameterOverrides(OutVariableTable, bDrivenOnly);
+	}
+}
 
-#define UE_CAMERA_VARIABLE_FOR_TYPE(ValueType, ValueName)\
-	Internal::ApplyParameterOverrides(\
-			CameraRig,\
-			ParameterOverrides.Get##ValueName##Overrides(),\
-			OutVariableTable);
-UE_CAMERA_VARIABLE_FOR_ALL_TYPES()
-#undef UE_CAMERA_VARIABLE_FOR_TYPE
+bool FCameraRigCameraNodeEvaluator::IsApplyingParameterOverrides() const
+{
+	return bApplyParameterOverrides;
+}
+
+void FCameraRigCameraNodeEvaluator::SetApplyParameterOverrides(bool bShouldApply)
+{
+	bApplyParameterOverrides = bShouldApply;
 }
 
 namespace Internal
diff --git a/Engine/Plugins/Cameras/GameplayCameras/Source/GameplayCameras/Public/Core/BlendCameraNode.h b/Engine/Plugins/Cameras/GameplayCameras/Source/GameplayCameras/Public/Core/BlendCameraNode.h
--- a/Engine/Plugins/Cameras/GameplayCameras/Source/GameplayCameras/Public/Core/BlendCameraNode.h
+++ b/Engine/Plugins/Cameras/GameplayCameras/Source/GameplayCameras/Public/Core/BlendCameraNode.h
@@ -39,6 +39,8 @@ struct FCameraNodePreBlendParams
 	const FCameraPose& LastCameraPose;
 	/** The variable table of the node tree being blended. */
 	const FCameraVariableTable& ChildVariableTable;
+	/** Extra filter for variable table blending. */
+	ECameraVariableTableFilter ExtraVariableTableFilter = ECameraVariableTableFilter::None;
 };
 
 /**
diff --git a/Engine/Plugins/Cameras/GameplayCameras/Source/GameplayCameras/Public/Core/BlendStackCameraNode.h b/Engine/Plugins/Cameras/GameplayCameras/Source/GameplayCameras/Public/Core/BlendStackCameraNode.h
--- a/Engine/Plugins/Cameras/GameplayCameras/Source/GameplayCameras/Public/Core/BlendStackCameraNode.h
+++ b/Engine/Plugins/Cameras/GameplayCameras/Source/GameplayCameras/Public/Core/BlendStackCameraNode.h
@@ -154,6 +154,9 @@ protected:
 		TSharedPtr<const FCameraEvaluationContext> EvaluationContext,
 		UBlendStackRootCameraNode* EntryRootNode);
 
+	void PushVariantEntry(const FBlendStackCameraPushParams& Params, const UCameraRigTransition* Transition);
+	void PushNewEntry(const FBlendStackCameraPushParams& Params, const UCameraRigTransition* Transition);
+
 	void FreezeEntry(FCameraRigEntry& Entry);
 
 	void GatherEntryParameterEvaluators(FCameraNodeEvaluator* RootEvaluator, TArray<FCameraNodeEvaluator*>& OutParameterEvaluators);
@@ -161,6 +164,7 @@ protected:
 	void BroadcastCameraRigEvent(EBlendStackCameraRigEventType EventType, const FCameraRigEntry& Entry, const UCameraRigTransition* Transition = nullptr) const;
 
 #if WITH_EDITOR
+	void AddPackageListeners(FCameraRigEntry& Entry);
 	void RemoveListenedPackages(FCameraRigEntry& Entry);
 	void RemoveListenedPackages(TSharedPtr<IGameplayCamerasLiveEditManager> LiveEditManager, FCameraRigEntry& Entry);
 #endif
diff --git a/Engine/Plugins/Cameras/GameplayCameras/Source/GameplayCameras/Public/Core/CameraRigTransition.h b/Engine/Plugins/Cameras/GameplayCameras/Source/GameplayCameras/Public/Core/CameraRigTransition.h
--- a/Engine/Plugins/Cameras/GameplayCameras/Source/GameplayCameras/Public/Core/CameraRigTransition.h
+++ b/Engine/Plugins/Cameras/GameplayCameras/Source/GameplayCameras/Public/Core/CameraRigTransition.h
@@ -155,6 +155,16 @@ public:
 	UPROPERTY(EditAnywhere, Category="Transition")
 	bool bOverrideInitialOrientation = false;
 
+	/**
+	 * Whether this transition allows merging two similar camera rigs together.
+	 * Similar camera rigs run the same underlying camera rig prefab with different parameter
+	 * overrides. When merged, instead of pushing a new camera rig instance on the blend stack,
+	 * only the parameter overrides are kept. These parameter overrides are blended together
+	 * and the underlying camera rig prefab is run only once.
+	 */
+	UPROPERTY(EditAnywhere, Category="Advanced")
+	bool bAllowCameraRigMerging = false;
+
 public:
 
 	using FCameraRigBuildContext = UE::Cameras::FCameraRigBuildContext;
diff --git a/Engine/Plugins/Cameras/GameplayCameras/Source/GameplayCameras/Public/Nodes/Common/CameraRigCameraNode.h b/Engine/Plugins/Cameras/GameplayCameras/Source/GameplayCameras/Public/Nodes/Common/CameraRigCameraNode.h
--- a/Engine/Plugins/Cameras/GameplayCameras/Source/GameplayCameras/Public/Nodes/Common/CameraRigCameraNode.h
+++ b/Engine/Plugins/Cameras/GameplayCameras/Source/GameplayCameras/Public/Nodes/Common/CameraRigCameraNode.h
@@ -8,11 +8,6 @@
 
 #include "CameraRigCameraNode.generated.h"
 
-namespace UE::Cameras
-{
-	class FCameraRigCameraNodeEvaluator;
-}
-
 /**
  * A camera node that runs a camera rig's own node tree.
  */
@@ -41,7 +36,7 @@ private:
 
 	// Deprecated properties, predating FCameraRigAssetReference
 
-	UPROPERTY()
+	UPROPERTY(meta=(ObjectTreeGraphHidden=true))
 	TObjectPtr<UCameraRigAsset> CameraRig_DEPRECATED;
 
 	UPROPERTY()
@@ -74,3 +69,39 @@ private:
 	TArray<FTransform3dCameraRigParameterOverride> Transform3dOverrides_DEPRECATED;
 };
 
+namespace UE::Cameras
+{
+
+class FCameraRigCameraNodeEvaluator : public FCameraNodeEvaluator
+{
+	UE_DECLARE_CAMERA_NODE_EVALUATOR(GAMEPLAYCAMERAS_API, FCameraRigCameraNodeEvaluator)
+
+public:
+
+	FCameraRigCameraNodeEvaluator();
+
+	// Internal API.
+	bool IsApplyingParameterOverrides() const;
+	void SetApplyParameterOverrides(bool bShouldApply);
+
+protected:
+
+	// FCameraNodeEvaluator interface.
+	virtual FCameraNodeEvaluatorChildrenView OnGetChildren() override;
+	virtual void OnInitialize(const FCameraNodeEvaluatorInitializeParams& Params, FCameraNodeEvaluationResult& OutResult) override;
+	virtual void OnUpdateParameters(const FCameraBlendedParameterUpdateParams& Params, FCameraBlendedParameterUpdateResult& OutResult) override;
+	virtual void OnBuild(const FCameraNodeEvaluatorBuildParams& Params) override;
+	virtual void OnRun(const FCameraNodeEvaluationParams& Params, FCameraNodeEvaluationResult& OutResult) override;
+
+private:
+
+	void ApplyParameterOverrides(FCameraVariableTable& OutVariableTable, bool bDrivenOnly);
+
+private:
+
+	FCameraNodeEvaluator* CameraRigRootEvaluator = nullptr;
+	bool bApplyParameterOverrides = true;
+};
+
+}  // namesapce UE::Cameras
+
