Cameras: add camera rig picker mode that only shows camera rigs from inside the same package as the edited object

This is for picking a camera rig from inside a camera asset.

diff --git a/Engine/Plugins/Cameras/GameplayCameras/Source/GameplayCamerasEditor/Private/Customizations/CameraRigPtrDetailsCustomization.cpp b/Engine/Plugins/Cameras/GameplayCameras/Source/GameplayCamerasEditor/Private/Customizations/CameraRigPtrDetailsCustomization.cpp
--- a/Engine/Plugins/Cameras/GameplayCameras/Source/GameplayCamerasEditor/Private/Customizations/CameraRigPtrDetailsCustomization.cpp
+++ b/Engine/Plugins/Cameras/GameplayCameras/Source/GameplayCamerasEditor/Private/Customizations/CameraRigPtrDetailsCustomization.cpp
@@ -33,12 +33,16 @@ void FCameraRigPtrDetailsCustomization::CustomizeHeader(TSharedRef<IPropertyHand
 	EPickerMode PickerMode = DeterminePickerMode();
 
 	TSharedPtr<SWidget> ValueContentWidget;
-	if (PickerMode == EPickerMode::CameraRigPicker || PickerMode == EPickerMode::CameraDirectorRigPicker)
+	if (PickerMode != EPickerMode::PrefabCameraRigPicker)
 	{
 		FOnGetContent OnGetComboMenuContent;
-		if (PickerMode == EPickerMode::CameraRigPicker)
+		if (PickerMode == EPickerMode::AnyCameraRigPicker)
 		{
-			OnGetComboMenuContent = FOnGetContent::CreateSP(this, &FCameraRigPtrDetailsCustomization::OnBuildCameraRigNamePicker);
+			OnGetComboMenuContent = FOnGetContent::CreateSP(this, &FCameraRigPtrDetailsCustomization::OnBuildAnyCameraRigNamePicker);
+		}
+		else if (PickerMode == EPickerMode::SelfCameraRigPicker)
+		{
+			OnGetComboMenuContent = FOnGetContent::CreateSP(this, &FCameraRigPtrDetailsCustomization::OnBuildSelfCameraRigNamePicker);
 		}
 		else if (PickerMode == EPickerMode::CameraDirectorRigPicker)
 		{
@@ -95,11 +99,15 @@ FCameraRigPtrDetailsCustomization::EPickerMode FCameraRigPtrDetailsCustomization
 	{
 		return EPickerMode::CameraDirectorRigPicker;
 	}
+	if (MetaDataField->GetBoolMetaData("UseSelfCameraRigPicker"))
+	{
+		return EPickerMode::SelfCameraRigPicker;
+	}
 	if (MetaDataField->GetBoolMetaData("UseCameraRigPicker"))
 	{
-		return EPickerMode::CameraRigPicker;
+		return EPickerMode::AnyCameraRigPicker;
 	}
-	return EPickerMode::StandardPicker;
+	return EPickerMode::PrefabCameraRigPicker;
 }
 
 FText FCameraRigPtrDetailsCustomization::OnGetComboText() const
@@ -176,7 +184,34 @@ TSharedRef<SWidget> FCameraRigPtrDetailsCustomization::OnBuildCameraDirectorRigN
 	return BuildCameraRigNamePickerImpl(PickerConfig);
 }
 
-TSharedRef<SWidget> FCameraRigPtrDetailsCustomization::OnBuildCameraRigNamePicker()
+TSharedRef<SWidget> FCameraRigPtrDetailsCustomization::OnBuildSelfCameraRigNamePicker()
+{
+	FCameraRigPickerConfig PickerConfig;
+	PickerConfig.bCanSelectCameraAsset = false;
+
+	TArray<UObject*> OuterObjects;
+	CameraRigPropertyHandle->GetOuterObjects(OuterObjects);
+
+	TSet<UCameraAsset*> OuterCameraAssets;
+	for (UObject* OuterObject : OuterObjects)
+	{
+		UCameraAsset* OuterCameraAsset = OuterObject->GetTypedOuter<UCameraAsset>();
+		if (OuterCameraAsset)
+		{
+			OuterCameraAssets.Add(OuterCameraAsset);
+		}
+	}
+
+	for (UCameraAsset* OuterCameraAsset : OuterCameraAssets)
+	{
+		PickerConfig.InitialCameraAssetSelection = FAssetData(OuterCameraAsset);
+		break;
+	}
+
+	return BuildCameraRigNamePickerImpl(PickerConfig);
+}
+
+TSharedRef<SWidget> FCameraRigPtrDetailsCustomization::OnBuildAnyCameraRigNamePicker()
 {
 	FCameraRigPickerConfig PickerConfig;
 	PickerConfig.bCanSelectCameraAsset = true;
diff --git a/Engine/Plugins/Cameras/GameplayCameras/Source/GameplayCamerasEditor/Public/Customizations/CameraRigPtrDetailsCustomization.h b/Engine/Plugins/Cameras/GameplayCameras/Source/GameplayCamerasEditor/Public/Customizations/CameraRigPtrDetailsCustomization.h
--- a/Engine/Plugins/Cameras/GameplayCameras/Source/GameplayCamerasEditor/Public/Customizations/CameraRigPtrDetailsCustomization.h
+++ b/Engine/Plugins/Cameras/GameplayCameras/Source/GameplayCamerasEditor/Public/Customizations/CameraRigPtrDetailsCustomization.h
@@ -31,13 +31,14 @@ public:
 
 private:
 
-	enum class EPickerMode { StandardPicker, CameraRigPicker, CameraDirectorRigPicker };
+	enum class EPickerMode { PrefabCameraRigPicker, AnyCameraRigPicker, SelfCameraRigPicker, CameraDirectorRigPicker };
 
 	EPickerMode DeterminePickerMode();
 
 	FText OnGetComboText() const;
 	FText OnGetComboToolTipText() const;
-	TSharedRef<SWidget> OnBuildCameraRigNamePicker();
+	TSharedRef<SWidget> OnBuildAnyCameraRigNamePicker();
+	TSharedRef<SWidget> OnBuildSelfCameraRigNamePicker();
 	TSharedRef<SWidget> OnBuildCameraDirectorRigNamePicker();
 	TSharedRef<SWidget> BuildCameraRigNamePickerImpl(FCameraRigPickerConfig& PickerConfig);
 	void OnPickerAssetSelected(UCameraRigAsset* SelectedItem);
