Cameras: change how ObjectTreeGraph objects are collected

Auto-collecting objects from the package can easily result in the wrong set of objects when multiple graphs must be kept separate inside the same package, such as with camera assets. This change therefore brings more customization options for how objects are collected when the graph is built. For camera editor graphs, we can now directly rely on AllNodeTreeObjects, AllTransitionsObjects, etc.

diff --git a/Engine/Plugins/Cameras/GameplayCameras/Source/GameplayCamerasEditor/Private/Editors/CameraNodeGraphSchema.cpp b/Engine/Plugins/Cameras/GameplayCameras/Source/GameplayCamerasEditor/Private/Editors/CameraNodeGraphSchema.cpp
--- a/Engine/Plugins/Cameras/GameplayCameras/Source/GameplayCamerasEditor/Private/Editors/CameraNodeGraphSchema.cpp
+++ b/Engine/Plugins/Cameras/GameplayCameras/Source/GameplayCamerasEditor/Private/Editors/CameraNodeGraphSchema.cpp
@@ -4,6 +4,7 @@
 
 #include "Core/BlendCameraNode.h"
 #include "Core/CameraNode.h"
+#include "Core/CameraNodeHierarchy.h"
 #include "Core/CameraRigAsset.h"
 #include "EdGraph/EdGraphPin.h"
 #include "Editors/CameraNodeGraphNode.h"
@@ -12,8 +13,10 @@
 #include "Editors/ObjectTreeGraph.h"
 #include "Editors/ObjectTreeGraphConfig.h"
 #include "Editors/ObjectTreeGraphNode.h"
+#include "Framework/Notifications/NotificationManager.h"
 #include "GameplayCamerasEditorSettings.h"
 #include "Nodes/Common/CameraRigCameraNode.h"
+#include "Widgets/Notifications/SNotificationList.h"
 
 #include "ScopedTransaction.h"
 
@@ -62,27 +65,70 @@ FObjectTreeGraphConfig UCameraNodeGraphSchema::BuildGraphConfig() const
 	return GraphConfig;
 }
 
-void UCameraNodeGraphSchema::OnCreateAllNodes(UObjectTreeGraph* InGraph, const FCreatedNodes& InCreatedNodes) const
+void UCameraNodeGraphSchema::CollectAllObjects(UObjectTreeGraph* InGraph, TSet<UObject*>& OutAllObjects) const
 {
-	Super::OnCreateAllNodes(InGraph, InCreatedNodes);
+	using namespace UE::Cameras;
+
+	// Only get the graph objects from the root interface.
+	CollectAllConnectableObjectsFromRootInterface(InGraph, OutAllObjects, false);
 
+	// See if we are missing objects from AllNodeTreeObjects... if so, add them and notify the user.
 	UCameraRigAsset* CameraRig = Cast<UCameraRigAsset>(InGraph->GetRootObject());
-	if (ensure(CameraRig))
+	if (CameraRig)
 	{
-		for (UCameraRigInterfaceParameter* InterfaceParameter : CameraRig->Interface.InterfaceParameters)
+		FCameraNodeHierarchy Hierarchy(CameraRig);
+
+		TSet<UObject*> AllNodeTreeObjects;
+		((IObjectTreeGraphRootObject*)CameraRig)->GetConnectableObjects(UCameraRigAsset::NodeTreeGraphName, AllNodeTreeObjects);
+
+		TSet<UObject*> MissingNodeTreeObjects;
+		if (Hierarchy.FindMissingConnectableObjects(AllNodeTreeObjects, MissingNodeTreeObjects))
 		{
-			UObjectTreeGraphNode* const* InterfaceParameterNode = InCreatedNodes.CreatedNodes.Find(InterfaceParameter);
-			UObjectTreeGraphNode* const* CameraNodeNode = InCreatedNodes.CreatedNodes.Find(InterfaceParameter->Target);
-			if (InterfaceParameterNode && CameraNodeNode)
+			FNotificationInfo NotificationInfo(
+					FText::Format(
+						LOCTEXT("AllNodeTreeObjectsMismatch", 
+							"Found {0} nodes missing from the internal list. Please re-save the asset."),
+						MissingNodeTreeObjects.Num()));
+			NotificationInfo.ExpireDuration = 4.0f;
+			FSlateNotificationManager::Get().AddNotification(NotificationInfo);
+
+			for (UObject* MissingObject : MissingNodeTreeObjects)
 			{
-				UEdGraphPin* InterfaceParameterSelfPin = (*InterfaceParameterNode)->GetSelfPin();
-				UEdGraphPin* CameraParameterPin = Cast<UCameraNodeGraphNode>(*CameraNodeNode)->GetPinForCameraParameterProperty(InterfaceParameter->TargetPropertyName);
-				InterfaceParameterSelfPin->MakeLinkTo(CameraParameterPin);
+				((IObjectTreeGraphRootObject*)CameraRig)->AddConnectableObject(UCameraRigAsset::NodeTreeGraphName, MissingObject);
 			}
 		}
 	}
 }
 
+void UCameraNodeGraphSchema::OnCreateAllNodes(UObjectTreeGraph* InGraph, const FCreatedNodes& InCreatedNodes) const
+{
+	Super::OnCreateAllNodes(InGraph, InCreatedNodes);
+
+	UObject* RootObject = InGraph->GetRootObject();
+	if (!RootObject)
+	{
+		return;
+	}
+
+	UCameraRigAsset* CameraRig = Cast<UCameraRigAsset>(InGraph->GetRootObject());
+	if (!ensure(CameraRig))
+	{
+		return;
+	}
+	
+	for (UCameraRigInterfaceParameter* InterfaceParameter : CameraRig->Interface.InterfaceParameters)
+	{
+		UObjectTreeGraphNode* const* InterfaceParameterNode = InCreatedNodes.CreatedNodes.Find(InterfaceParameter);
+		UObjectTreeGraphNode* const* CameraNodeNode = InCreatedNodes.CreatedNodes.Find(InterfaceParameter->Target);
+		if (InterfaceParameterNode && CameraNodeNode)
+		{
+			UEdGraphPin* InterfaceParameterSelfPin = (*InterfaceParameterNode)->GetSelfPin();
+			UEdGraphPin* CameraParameterPin = Cast<UCameraNodeGraphNode>(*CameraNodeNode)->GetPinForCameraParameterProperty(InterfaceParameter->TargetPropertyName);
+			InterfaceParameterSelfPin->MakeLinkTo(CameraParameterPin);
+		}
+	}
+}
+
 void UCameraNodeGraphSchema::OnAddConnectableObject(UObjectTreeGraph* InGraph, UObjectTreeGraphNode* InNewNode) const
 {
 	Super::OnAddConnectableObject(InGraph, InNewNode);
diff --git a/Engine/Plugins/Cameras/GameplayCameras/Source/GameplayCamerasEditor/Private/Editors/CameraRigTransitionGraphSchemaBase.cpp b/Engine/Plugins/Cameras/GameplayCameras/Source/GameplayCamerasEditor/Private/Editors/CameraRigTransitionGraphSchemaBase.cpp
--- a/Engine/Plugins/Cameras/GameplayCameras/Source/GameplayCamerasEditor/Private/Editors/CameraRigTransitionGraphSchemaBase.cpp
+++ b/Engine/Plugins/Cameras/GameplayCameras/Source/GameplayCamerasEditor/Private/Editors/CameraRigTransitionGraphSchemaBase.cpp
@@ -37,6 +37,14 @@ FObjectTreeGraphConfig UCameraRigTransitionGraphSchemaBase::BuildGraphConfig() c
 	return GraphConfig;
 }
 
+void UCameraRigTransitionGraphSchemaBase::CollectAllObjects(UObjectTreeGraph* InGraph, TSet<UObject*>& OutAllObjects) const
+{
+	using namespace UE::Cameras;
+
+	// Only get the graph objects from the root interface.
+	CollectAllConnectableObjectsFromRootInterface(InGraph, OutAllObjects, false);
+}
+
 void UCameraRigTransitionGraphSchemaBase::GetGraphContextActions(FGraphContextMenuBuilder& ContextMenuBuilder) const
 {
 	ETransitionGraphContextActions PossibleActions = GetTransitionGraphContextActions(ContextMenuBuilder);
diff --git a/Engine/Plugins/Cameras/GameplayCameras/Source/GameplayCamerasEditor/Private/Editors/CameraSharedTransitionGraphSchema.cpp b/Engine/Plugins/Cameras/GameplayCameras/Source/GameplayCamerasEditor/Private/Editors/CameraSharedTransitionGraphSchema.cpp
--- a/Engine/Plugins/Cameras/GameplayCameras/Source/GameplayCamerasEditor/Private/Editors/CameraSharedTransitionGraphSchema.cpp
+++ b/Engine/Plugins/Cameras/GameplayCameras/Source/GameplayCamerasEditor/Private/Editors/CameraSharedTransitionGraphSchema.cpp
@@ -20,7 +20,6 @@ void UCameraSharedTransitionGraphSchema::OnBuildGraphConfig(FObjectTreeGraphConf
 	InOutGraphConfig.ConnectableObjectClasses.Add(UCameraAsset::StaticClass());
 	InOutGraphConfig.GraphDisplayInfo.PlainName = LOCTEXT("NodeGraphPlainName", "SharedTransitions");
 	InOutGraphConfig.GraphDisplayInfo.DisplayName = LOCTEXT("NodeGraphDisplayName", "Shared Transitions");
-	InOutGraphConfig.StopAutoCollectAtObjectClasses.Add(UCameraRigAsset::StaticClass());
 	InOutGraphConfig.ObjectClassConfigs.Emplace(UCameraAsset::StaticClass())
 		.HasSelfPin(false)
 		.OnlyAsRoot()
diff --git a/Engine/Plugins/Cameras/GameplayCameras/Source/GameplayCamerasEditor/Private/Editors/ObjectTreeGraph.cpp b/Engine/Plugins/Cameras/GameplayCameras/Source/GameplayCamerasEditor/Private/Editors/ObjectTreeGraph.cpp
--- a/Engine/Plugins/Cameras/GameplayCameras/Source/GameplayCamerasEditor/Private/Editors/ObjectTreeGraph.cpp
+++ b/Engine/Plugins/Cameras/GameplayCameras/Source/GameplayCamerasEditor/Private/Editors/ObjectTreeGraph.cpp
@@ -16,7 +16,7 @@ UObjectTreeGraph::UObjectTreeGraph(const FObjectInitializer& ObjInit)
 	Schema = UObjectTreeGraphSchema::StaticClass();
 }
 
-void UObjectTreeGraph::Reset(TObjectPtr<UObject> InRootObject, const FObjectTreeGraphConfig& InConfig, EObjectTreeGraphBuildSource InSource)
+void UObjectTreeGraph::Reset(TObjectPtr<UObject> InRootObject, const FObjectTreeGraphConfig& InConfig)
 {
 	WeakRootObject = InRootObject;
 
@@ -26,7 +26,7 @@ void UObjectTreeGraph::Reset(TObjectPtr<UObject> InRootObject, const FObjectTree
 		Config.DefaultGraphNodeClass = UObjectTreeGraphNode::StaticClass();
 	}
 
-	RebuildGraph(InSource);
+	RebuildGraph();
 }
 
 UObjectTreeGraphNode* UObjectTreeGraph::FindObjectNode(UObject* InObject) const
@@ -49,12 +49,12 @@ const FObjectTreeGraphConfig& UObjectTreeGraph::GetConfig() const
 	return Config;
 }
 
-void UObjectTreeGraph::RebuildGraph(EObjectTreeGraphBuildSource InSource)
+void UObjectTreeGraph::RebuildGraph()
 {
 	const UObjectTreeGraphSchema* GraphSchema = Cast<UObjectTreeGraphSchema>(GetSchema());
 	if (ensure(GraphSchema))
 	{
-		GraphSchema->RebuildGraph(this, InSource);
+		GraphSchema->RebuildGraph(this);
 	}
 }
 
diff --git a/Engine/Plugins/Cameras/GameplayCameras/Source/GameplayCamerasEditor/Private/Editors/ObjectTreeGraphSchema.cpp b/Engine/Plugins/Cameras/GameplayCameras/Source/GameplayCamerasEditor/Private/Editors/ObjectTreeGraphSchema.cpp
--- a/Engine/Plugins/Cameras/GameplayCameras/Source/GameplayCamerasEditor/Private/Editors/ObjectTreeGraphSchema.cpp
+++ b/Engine/Plugins/Cameras/GameplayCameras/Source/GameplayCamerasEditor/Private/Editors/ObjectTreeGraphSchema.cpp
@@ -134,10 +134,10 @@ UObjectTreeGraphSchema::UObjectTreeGraphSchema(const FObjectInitializer& ObjInit
 {
 }
 
-void UObjectTreeGraphSchema::RebuildGraph(UObjectTreeGraph* InGraph, EObjectTreeGraphBuildSource InSource) const
+void UObjectTreeGraphSchema::RebuildGraph(UObjectTreeGraph* InGraph) const
 {
 	RemoveAllNodes(InGraph);
-	CreateAllNodes(InGraph, InSource);
+	CreateAllNodes(InGraph);
 	InGraph->NotifyGraphChanged();
 }
 
@@ -150,38 +150,69 @@ void UObjectTreeGraphSchema::RemoveAllNodes(UObjectTreeGraph* InGraph) const
 	}
 }
 
-void UObjectTreeGraphSchema::CreateAllNodes(UObjectTreeGraph* InGraph, EObjectTreeGraphBuildSource InSource) const
+void UObjectTreeGraphSchema::CollectAllObjects(UObjectTreeGraph* InGraph, TSet<UObject*>& OutAllObjects) const
 {
+	// By default, collect all objects referenced directly or indirectly by the root object, within
+	// the same package, unless the root object implements the IObjectTreeGraphRootObject interface,
+	// in which case get the list of objects from it.
+	// Override this method to collect objects differently.
+	const bool bHasRootInterface = CollectAllConnectableObjectsFromRootInterface(InGraph, OutAllObjects, true);
+	if (!bHasRootInterface)
+	{
+		CollectAllReferencedObjects(InGraph, OutAllObjects);
+	}
+}
+
+void UObjectTreeGraphSchema::CollectAllReferencedObjects(UObjectTreeGraph* InGraph, TSet<UObject*>& OutAllObjects)
+{
+	using namespace UE::ObjectTreeGraph;
+
 	UObject* RootObject = InGraph->GetRootObject();
 	if (!RootObject)
 	{
 		return;
 	}
 
-	const FObjectTreeGraphConfig& GraphConfig = InGraph->GetConfig();
+	// Make sure the root object itself is in there.
+	OutAllObjects.Add(RootObject);
 
-	TSet<UObject*> AllObjects;
-	// Gather up all the objects we need for the graph. Start by all objects that are referenced (directly or
-	// indirectly) by the root object. Our custom reference collector will not collect references that go outside
-	// of the root object's package.
-	if (GraphConfig.bAutoCollectInitialObjects)
+	// Use a reference collector that doesn't go outside of the root object package.
+	TArray<UObject*> ReferencedObjects;
+	FPackageReferenceCollector Collector(RootObject, ReferencedObjects);
+	Collector.CollectReferences();
+	OutAllObjects.Append(ReferencedObjects);
+}
+
+bool UObjectTreeGraphSchema::CollectAllConnectableObjectsFromRootInterface(UObjectTreeGraph* InGraph, TSet<UObject*>& OutAllObjects, bool bAllowNoRootInterface)
+{
+	UObject* RootObject = InGraph->GetRootObject();
+	if (!RootObject)
 	{
-		using namespace UE::ObjectTreeGraph;
+		return true;
+	}
 
-		// Make sure the root object itself is in there.
-		AllObjects.Add(RootObject);
+	// Make sure the root object itself is in there.
+	OutAllObjects.Add(RootObject);
 
-		TArray<UObject*> ReferencedObjects;
-		FPackageReferenceCollector Collector(RootObject, ReferencedObjects);
-		Collector.StopAtObjectClasses(GraphConfig.StopAutoCollectAtObjectClasses);
-		Collector.CollectReferences();
-		AllObjects.Append(ReferencedObjects);
-	}
-	// Add any other custom objects the root object may want.
-	if (IObjectTreeGraphRootObject* RootObjectInterface = Cast<IObjectTreeGraphRootObject>(RootObject))
+	// Get all the objects we need from the dedicated interface for this.
+	IObjectTreeGraphRootObject* RootObjectInterface = Cast<IObjectTreeGraphRootObject>(RootObject);
+	ensureMsgf(RootObjectInterface || bAllowNoRootInterface,
+			TEXT("Root object '%s' was expected to implement IObjectTreeGraphRootObject, but doesn't."),
+			*GetNameSafe(RootObject));
+	if (RootObjectInterface)
 	{
-		RootObjectInterface->GetConnectableObjects(GraphConfig.GraphName, AllObjects);
+		const FObjectTreeGraphConfig& GraphConfig = InGraph->GetConfig();
+		RootObjectInterface->GetConnectableObjects(GraphConfig.GraphName, OutAllObjects);
+		return true;
 	}
+	return false;
+}
+
+void UObjectTreeGraphSchema::CreateAllNodes(UObjectTreeGraph* InGraph) const
+{
+	// Collect the objects.
+	TSet<UObject*> AllObjects;
+	CollectAllObjects(InGraph, AllObjects);
 	
 	// Create all the nodes.
 	FCreatedNodes CreatedNodes;
@@ -197,8 +228,11 @@ void UObjectTreeGraphSchema::CreateAllNodes(UObjectTreeGraph* InGraph, EObjectTr
 	InGraph->RootObjectNode = nullptr;
 	if (!AllObjects.IsEmpty())
 	{
+		UObject* RootObject = InGraph->GetRootObject();
 		UObjectTreeGraphNode** CreatedRootObjectNode = CreatedNodes.CreatedNodes.Find(RootObject);
-		if (ensure(CreatedRootObjectNode))
+		if (ensureMsgf(CreatedRootObjectNode, 
+					TEXT("Can't find root object '%s' in the list of created graph nodes!"),
+					*GetNameSafe(RootObject)))
 		{
 			InGraph->RootObjectNode = *CreatedRootObjectNode;
 		}
diff --git a/Engine/Plugins/Cameras/GameplayCameras/Source/GameplayCamerasEditor/Private/Editors/SCameraRigAssetEditor.cpp b/Engine/Plugins/Cameras/GameplayCameras/Source/GameplayCamerasEditor/Private/Editors/SCameraRigAssetEditor.cpp
--- a/Engine/Plugins/Cameras/GameplayCameras/Source/GameplayCamerasEditor/Private/Editors/SCameraRigAssetEditor.cpp
+++ b/Engine/Plugins/Cameras/GameplayCameras/Source/GameplayCamerasEditor/Private/Editors/SCameraRigAssetEditor.cpp
@@ -75,7 +75,7 @@ void SCameraRigAssetEditor::CreateNodeGraphEditor()
 
 	NodeGraph = NewObject<UObjectTreeGraph>(GetTransientPackage(), NAME_None, RF_Transactional | RF_Standalone);
 	NodeGraph->Schema = SchemaClass;
-	NodeGraph->Reset(CameraRigAsset, GraphConfig, EObjectTreeGraphBuildSource::RootObjectPackage);
+	NodeGraph->Reset(CameraRigAsset, GraphConfig);
 
 	NodeGraphChangedHandle = NodeGraph->AddOnGraphChangedHandler(
 			FOnGraphChanged::FDelegate::CreateSP(this, &SCameraRigAssetEditor::OnGraphChanged));
@@ -100,7 +100,7 @@ void SCameraRigAssetEditor::CreateTransitionGraphEditor()
 
 	TransitionGraph = NewObject<UObjectTreeGraph>(GetTransientPackage(), NAME_None, RF_Transactional | RF_Standalone);
 	TransitionGraph->Schema = SchemaClass;
-	TransitionGraph->Reset(CameraRigAsset, GraphConfig, EObjectTreeGraphBuildSource::RootObjectPackage);
+	TransitionGraph->Reset(CameraRigAsset, GraphConfig);
 
 	TransitionGraphChangedHandle = TransitionGraph->AddOnGraphChangedHandler(
 			FOnGraphChanged::FDelegate::CreateSP(this, &SCameraRigAssetEditor::OnGraphChanged));
diff --git a/Engine/Plugins/Cameras/GameplayCameras/Source/GameplayCamerasEditor/Private/Editors/SCameraRigTransitionEditor.cpp b/Engine/Plugins/Cameras/GameplayCameras/Source/GameplayCamerasEditor/Private/Editors/SCameraRigTransitionEditor.cpp
--- a/Engine/Plugins/Cameras/GameplayCameras/Source/GameplayCamerasEditor/Private/Editors/SCameraRigTransitionEditor.cpp
+++ b/Engine/Plugins/Cameras/GameplayCameras/Source/GameplayCamerasEditor/Private/Editors/SCameraRigTransitionEditor.cpp
@@ -69,7 +69,7 @@ void SCameraRigTransitionEditor::CreateTransitionGraphEditor()
 
 	TransitionGraph = NewObject<UObjectTreeGraph>(GetTransientPackage(), NAME_None, RF_Transactional | RF_Standalone);
 	TransitionGraph->Schema = TransitionGraphSchemaClass;
-	TransitionGraph->Reset(TransitionOwner, GraphConfig, EObjectTreeGraphBuildSource::RootObjectPackage);
+	TransitionGraph->Reset(TransitionOwner, GraphConfig);
 
 	TransitionGraphChangedHandle = TransitionGraph->AddOnGraphChangedHandler(
 			FOnGraphChanged::FDelegate::CreateSP(this, &SCameraRigTransitionEditor::OnGraphChanged));
diff --git a/Engine/Plugins/Cameras/GameplayCameras/Source/GameplayCamerasEditor/Public/Editors/CameraNodeGraphSchema.h b/Engine/Plugins/Cameras/GameplayCameras/Source/GameplayCamerasEditor/Public/Editors/CameraNodeGraphSchema.h
--- a/Engine/Plugins/Cameras/GameplayCameras/Source/GameplayCamerasEditor/Public/Editors/CameraNodeGraphSchema.h
+++ b/Engine/Plugins/Cameras/GameplayCameras/Source/GameplayCamerasEditor/Public/Editors/CameraNodeGraphSchema.h
@@ -30,6 +30,7 @@ protected:
 	virtual const FPinConnectionResponse CanCreateConnection(const UEdGraphPin* A, const UEdGraphPin* B) const override;
 
 	// UObjectTreeGraphSchema interface.
+	virtual void CollectAllObjects(UObjectTreeGraph* InGraph, TSet<UObject*>& OutAllObjects) const override;
 	virtual void OnCreateAllNodes(UObjectTreeGraph* InGraph, const FCreatedNodes& InCreatedNodes) const override;
 	virtual void OnAddConnectableObject(UObjectTreeGraph* InGraph, UObjectTreeGraphNode* InNewNode) const override;
 	virtual void OnRemoveConnectableObject(UObjectTreeGraph* InGraph, UObjectTreeGraphNode* InRemovedNode) const override;
diff --git a/Engine/Plugins/Cameras/GameplayCameras/Source/GameplayCamerasEditor/Public/Editors/CameraRigTransitionGraphSchemaBase.h b/Engine/Plugins/Cameras/GameplayCameras/Source/GameplayCamerasEditor/Public/Editors/CameraRigTransitionGraphSchemaBase.h
--- a/Engine/Plugins/Cameras/GameplayCameras/Source/GameplayCamerasEditor/Public/Editors/CameraRigTransitionGraphSchemaBase.h
+++ b/Engine/Plugins/Cameras/GameplayCameras/Source/GameplayCamerasEditor/Public/Editors/CameraRigTransitionGraphSchemaBase.h
@@ -37,6 +37,7 @@ protected:
 	virtual void GetGraphContextActions(FGraphContextMenuBuilder& ContextMenuBuilder) const override;
 
 	// UObjectTreeGraphSchema interface.
+	virtual void CollectAllObjects(UObjectTreeGraph* InGraph, TSet<UObject*>& OutAllObjects) const override;
 	virtual void FilterGraphContextPlaceableClasses(TArray<UClass*>& InOutClasses) const override;
 
 protected:
diff --git a/Engine/Plugins/Cameras/GameplayCameras/Source/GameplayCamerasEditor/Public/Editors/ObjectTreeGraph.h b/Engine/Plugins/Cameras/GameplayCameras/Source/GameplayCamerasEditor/Public/Editors/ObjectTreeGraph.h
--- a/Engine/Plugins/Cameras/GameplayCameras/Source/GameplayCamerasEditor/Public/Editors/ObjectTreeGraph.h
+++ b/Engine/Plugins/Cameras/GameplayCameras/Source/GameplayCamerasEditor/Public/Editors/ObjectTreeGraph.h
@@ -15,8 +15,6 @@
 class UObjectTreeGraphNode;
 class UObjectTreeGraphSchema;
 
-enum class EObjectTreeGraphBuildSource : uint8;
-
 /**
  * A node graph that represents a hierarchy of objects and their relationships.
  */
@@ -31,10 +29,7 @@ public:
 	UObjectTreeGraph(const FObjectInitializer& ObjInit);
 
 	/** Initializes the graph given a root object and a graph configuration. */
-	void Reset(
-			TObjectPtr<UObject> InRootObject, 
-			const FObjectTreeGraphConfig& InConfig, 
-			EObjectTreeGraphBuildSource InSource);
+	void Reset(TObjectPtr<UObject> InRootObject, const FObjectTreeGraphConfig& InConfig);
 
 	/** Gets the root object. */
 	UObject* GetRootObject() const { return WeakRootObject.Get(); }
@@ -49,7 +44,7 @@ public:
 
 public:
 
-	void RebuildGraph(EObjectTreeGraphBuildSource InSource);
+	void RebuildGraph();
 
 private:
 
diff --git a/Engine/Plugins/Cameras/GameplayCameras/Source/GameplayCamerasEditor/Public/Editors/ObjectTreeGraphConfig.h b/Engine/Plugins/Cameras/GameplayCameras/Source/GameplayCamerasEditor/Public/Editors/ObjectTreeGraphConfig.h
--- a/Engine/Plugins/Cameras/GameplayCameras/Source/GameplayCamerasEditor/Public/Editors/ObjectTreeGraphConfig.h
+++ b/Engine/Plugins/Cameras/GameplayCameras/Source/GameplayCamerasEditor/Public/Editors/ObjectTreeGraphConfig.h
@@ -178,15 +178,6 @@ public:
 	/** Advanced, optional bits of configuration for specific classes and sub-classes of objects. */
 	TMap<UClass*, FObjectTreeGraphClassConfig> ObjectClassConfigs;
 
-	/** Stop auto-collection of initial objects at any object that is of the given types. */
-	TArray<UClass*> StopAutoCollectAtObjectClasses;
-
-	/** 
-	 * Whether to collect all objets in the graph automatically from the network of referenced
-	 * objects from the root.
-	 */
-	bool bAutoCollectInitialObjects = true;
-
 public:
 
 	/** Creates a new graph config. */
diff --git a/Engine/Plugins/Cameras/GameplayCameras/Source/GameplayCamerasEditor/Public/Editors/ObjectTreeGraphSchema.h b/Engine/Plugins/Cameras/GameplayCameras/Source/GameplayCamerasEditor/Public/Editors/ObjectTreeGraphSchema.h
--- a/Engine/Plugins/Cameras/GameplayCameras/Source/GameplayCamerasEditor/Public/Editors/ObjectTreeGraphSchema.h
+++ b/Engine/Plugins/Cameras/GameplayCameras/Source/GameplayCamerasEditor/Public/Editors/ObjectTreeGraphSchema.h
@@ -14,13 +14,6 @@ class UObjectTreeGraph;
 class UObjectTreeGraphNode;
 struct FObjectTreeGraphClassConfig;
 
-/** Specifies where to find graph node objects. */
-enum class EObjectTreeGraphBuildSource : uint8
-{
-	/** Create graph nodes from all eligible objects found in the same package. */
-	RootObjectPackage
-};
-
 /**
  * Schema class for an object tree graph.
  */
@@ -46,7 +39,7 @@ public:
 	UObjectTreeGraphSchema(const FObjectInitializer& ObjInit);
 
 	/** Rebuilds the graph from scratch. */
-	void RebuildGraph(UObjectTreeGraph* InGraph, EObjectTreeGraphBuildSource InSource) const;
+	void RebuildGraph(UObjectTreeGraph* InGraph) const;
 
 	/** Creates an object graph node for the given object. */
 	UObjectTreeGraphNode* CreateObjectNode(UObjectTreeGraph* InGraph, UObject* InObject) const;
@@ -105,6 +98,7 @@ protected:
 	};
 
 	// UObjectTreeGraphSchema interface.
+	virtual void CollectAllObjects(UObjectTreeGraph* InGraph, TSet<UObject*>& OutAllObjects) const;
 	virtual void OnCreateAllNodes(UObjectTreeGraph* InGraph, const FCreatedNodes& InCreatedNodes) const;
 	virtual UObjectTreeGraphNode* OnCreateObjectNode(UObjectTreeGraph* InGraph, UObject* InObject) const;
 	virtual void OnAddConnectableObject(UObjectTreeGraph* InGraph, UObjectTreeGraphNode* InNewNode) const;
@@ -118,11 +112,12 @@ protected:
 
 protected:
 
+	static void CollectAllReferencedObjects(UObjectTreeGraph* InGraph, TSet<UObject*>& OutAllObjects);
+	static bool CollectAllConnectableObjectsFromRootInterface(UObjectTreeGraph* InGraph, TSet<UObject*>& OutAllObjects, bool bAllowNoRootInterface);
+
 	const FObjectTreeGraphClassConfig& GetObjectClassConfig(const UObjectTreeGraphNode* InNode) const;
 	const FObjectTreeGraphClassConfig& GetObjectClassConfig(const UObjectTreeGraph* InGraph, UClass* InObjectClass) const;
 
-protected:
-
 	void ApplyConnection(UEdGraphPin* A, UEdGraphPin* B, FDelayedPinActions& Actions) const;
 	void ApplyDisconnection(UEdGraphPin* TargetPin, FDelayedPinActions& Actions, bool bIsReconnecting) const;
 	void ApplyDisconnection(UEdGraphPin* SourcePin, UEdGraphPin* TargetPin, FDelayedPinActions& Actions) const;
@@ -130,7 +125,7 @@ protected:
 private:
 
 	void RemoveAllNodes(UObjectTreeGraph* InGraph) const;
-	void CreateAllNodes(UObjectTreeGraph* InGraph, EObjectTreeGraphBuildSource InSource) const;
+	void CreateAllNodes(UObjectTreeGraph* InGraph) const;
 	void CreateConnections(UObjectTreeGraphNode* InGraphNode, const FCreatedNodes& InCreatedNodes) const;
 };
 
