Cameras: add transition conditions that match other rigs in the same asset

diff --git a/Engine/Plugins/Cameras/GameplayCameras/Source/GameplayCameras/Private/Transitions/DefaultTransitionConditions.cpp b/Engine/Plugins/Cameras/GameplayCameras/Source/GameplayCameras/Private/Transitions/DefaultTransitionConditions.cpp
new file mode 100644
--- /dev/null
+++ b/Engine/Plugins/Cameras/GameplayCameras/Source/GameplayCameras/Private/Transitions/DefaultTransitionConditions.cpp
@@ -0,0 +1,13 @@
+// Copyright Epic Games, Inc. All Rights Reserved.
+
+#include "Transitions/DefaultTransitionConditions.h"
+
+#include UE_INLINE_GENERATED_CPP_BY_NAME(DefaultTransitionConditions)
+
+bool UIsCameraRigTransitionCondition::OnTransitionMatches(const FCameraRigTransitionConditionMatchParams& Params) const
+{
+	const bool bPreviousMatches = (!PreviousCameraRig || PreviousCameraRig == Params.FromCameraRig);
+	const bool bNextMatches = (!NextCameraRig || NextCameraRig == Params.ToCameraRig);
+	return bPreviousMatches && bNextMatches;
+}
+
diff --git a/Engine/Plugins/Cameras/GameplayCameras/Source/GameplayCameras/Public/Transitions/DefaultTransitionConditions.h b/Engine/Plugins/Cameras/GameplayCameras/Source/GameplayCameras/Public/Transitions/DefaultTransitionConditions.h
new file mode 100644
--- /dev/null
+++ b/Engine/Plugins/Cameras/GameplayCameras/Source/GameplayCameras/Public/Transitions/DefaultTransitionConditions.h
@@ -0,0 +1,34 @@
+// Copyright Epic Games, Inc. All Rights Reserved.
+
+#pragma once
+
+#include "Core/CameraRigTransition.h"
+
+#include "DefaultTransitionConditions.generated.h"
+
+/**
+ * A transition condition that matches the previous and/or next camera rigs against
+ * one of the owning camera asset's rigs.
+ */
+UCLASS(MinimalAPI)
+class UIsCameraRigTransitionCondition 
+	: public UCameraRigTransitionCondition
+{
+	GENERATED_BODY()
+
+public:
+
+	/** Passes if null, or equal to the previous camera rig. */
+	UPROPERTY(EditAnywhere, Category=Transition, meta=(UseSelfCameraRigPicker=true, ObjectTreeGraphHidden=true))
+	TObjectPtr<UCameraRigAsset> PreviousCameraRig;
+
+	/** Passes if null, or equal to the next camera rig. */
+	UPROPERTY(EditAnywhere, Category=Transition, meta=(UseSelfCameraRigPicker=true, ObjectTreeGraphHidden=true))
+	TObjectPtr<UCameraRigAsset> NextCameraRig;
+
+protected:
+
+	// UCameraRigTransitionCondition interface.
+	virtual bool OnTransitionMatches(const FCameraRigTransitionConditionMatchParams& Params) const override;
+};
+
