Cameras: add support for auto-combined camera rigs

When a director activates more than one camera rig in the same frame, we dynamically combine those camera rigs (in order of activation). They will run as if they had been combined "by hand" using an Array camera node.

diff --git a/Engine/Plugins/Cameras/GameplayCameras/Source/GameplayCameras/Private/Core/CameraRigAssetReference.cpp b/Engine/Plugins/Cameras/GameplayCameras/Source/GameplayCameras/Private/Core/CameraRigAssetReference.cpp
--- a/Engine/Plugins/Cameras/GameplayCameras/Source/GameplayCameras/Private/Core/CameraRigAssetReference.cpp
+++ b/Engine/Plugins/Cameras/GameplayCameras/Source/GameplayCameras/Private/Core/CameraRigAssetReference.cpp
@@ -12,6 +12,15 @@ UE_CAMERA_VARIABLE_FOR_ALL_TYPES()
 #undef UE_CAMERA_VARIABLE_FOR_TYPE
 }
 
+FCameraRigAssetReference::FCameraRigAssetReference()
+{
+}
+
+FCameraRigAssetReference::FCameraRigAssetReference(UCameraRigAsset* InCameraRig)
+	: CameraRig(InCameraRig)
+{
+}
+
 bool FCameraRigAssetReference::UpdateParameterOverrides()
 {
 	if (!CameraRig)
diff --git a/Engine/Plugins/Cameras/GameplayCameras/Source/GameplayCameras/Private/Core/CameraRigCombinationRegistry.cpp b/Engine/Plugins/Cameras/GameplayCameras/Source/GameplayCameras/Private/Core/CameraRigCombinationRegistry.cpp
new file mode 100644
--- /dev/null
+++ b/Engine/Plugins/Cameras/GameplayCameras/Source/GameplayCameras/Private/Core/CameraRigCombinationRegistry.cpp
@@ -0,0 +1,138 @@
+// Copyright Epic Games, Inc. All Rights Reserved.
+
+#include "Core/CameraRigCombinationRegistry.h"
+
+#include "Core/CameraRigAsset.h"
+#include "Core/CameraRigParameterOverrideEvaluator.h"
+
+namespace UE::Cameras
+{
+
+FCameraRigCombinationRegistry::FCameraRigCombinationKey::FCameraRigCombinationKey(TArrayView<const UCameraRigAsset*> InCombination)
+{
+	CachedHash = 0;
+	for (const UCameraRigAsset* CameraRig : InCombination)
+	{
+		ensure(CameraRig);
+		TWeakObjectPtr<const UCameraRigAsset> WeakCameraRig(CameraRig);
+		CachedHash = HashCombine(CachedHash, GetTypeHash(WeakCameraRig));
+		Combination.Add(WeakCameraRig);
+	}
+}
+
+const UCameraRigAsset* FCameraRigCombinationRegistry::FindOrCreateCombination(TArrayView<const UCameraRigAsset*> InCombination)
+{
+	FCameraRigCombinationKey Key(InCombination);
+	int32& Index = Combinations.FindOrAdd(Key, INDEX_NONE);
+	if (Index == INDEX_NONE)
+	{
+		Index = CombinedCameraRigs.Find(nullptr);
+		if (Index == INDEX_NONE)
+		{
+			Index = CombinedCameraRigs.Add(nullptr);
+		}
+		check(CombinedCameraRigs.IsValidIndex(Index));
+
+		UCameraRigAsset* NewCombinedCameraRig = NewObject<UCameraRigAsset>();
+		UCombinedCameraRigsCameraNode* PrefabNode = NewObject<UCombinedCameraRigsCameraNode>(NewCombinedCameraRig);
+		{
+			for (const UCameraRigAsset* IndividualCameraRig : InCombination)
+			{
+				PrefabNode->CameraRigReferences.Add(const_cast<UCameraRigAsset*>(IndividualCameraRig));
+
+				NewCombinedCameraRig->AllocationInfo.Append(IndividualCameraRig->AllocationInfo);
+			}
+		}
+		NewCombinedCameraRig->RootNode = PrefabNode;
+
+		CombinedCameraRigs[Index] = NewCombinedCameraRig;
+	}
+	return CombinedCameraRigs[Index];
+}
+
+void FCameraRigCombinationRegistry::AddReferencedObjects(FReferenceCollector& Collector)
+{
+	Collector.AddReferencedObjects(CombinedCameraRigs);
+}
+
+class FCombinedCameraRigsCameraNodeEvaluator : public FCameraNodeEvaluator
+{
+	UE_DECLARE_CAMERA_NODE_EVALUATOR(GAMEPLAYCAMERAS_API, FCombinedCameraRigsCameraNodeEvaluator)
+
+protected:
+
+	virtual FCameraNodeEvaluatorChildrenView OnGetChildren() override;
+	virtual void OnInitialize(const FCameraNodeEvaluatorInitializeParams& Params, FCameraNodeEvaluationResult& OutResult) override;
+	virtual void OnBuild(const FCameraNodeEvaluatorBuildParams& Params) override;
+	virtual void OnRun(const FCameraNodeEvaluationParams& Params, FCameraNodeEvaluationResult& OutResult) override;
+
+private:
+
+	void ApplyParameterOverrides(FCameraVariableTable& OutVariableTable, bool bDrivenOnly);
+
+private:
+
+	TArray<FCameraNodeEvaluator*> CameraRigRootEvaluators;
+};
+
+UE_DEFINE_CAMERA_NODE_EVALUATOR(FCombinedCameraRigsCameraNodeEvaluator)
+
+FCameraNodeEvaluatorChildrenView FCombinedCameraRigsCameraNodeEvaluator::OnGetChildren()
+{
+	return FCameraNodeEvaluatorChildrenView(CameraRigRootEvaluators);
+}
+
+void FCombinedCameraRigsCameraNodeEvaluator::OnBuild(const FCameraNodeEvaluatorBuildParams& Params)
+{
+	const UCombinedCameraRigsCameraNode* CombinedRigsNode = GetCameraNodeAs<UCombinedCameraRigsCameraNode>();
+	for (const FCameraRigAssetReference& IndividualCameraRigReference : CombinedRigsNode->CameraRigReferences)
+	{
+		if (const UCameraRigAsset* CameraRig = IndividualCameraRigReference.GetCameraRig())
+		{
+			if (CameraRig->RootNode)
+			{
+				FCameraNodeEvaluator* CameraRigRootEvaluator = Params.BuildEvaluator(CameraRig->RootNode);
+				if (CameraRigRootEvaluator)
+				{
+					CameraRigRootEvaluators.Add(CameraRigRootEvaluator);
+				}
+			}
+		}
+	}
+}
+
+void FCombinedCameraRigsCameraNodeEvaluator::OnInitialize(const FCameraNodeEvaluatorInitializeParams& Params, FCameraNodeEvaluationResult& OutResult)
+{
+	// Apply overrides right away.
+	ApplyParameterOverrides(OutResult.VariableTable, false);
+}
+
+void FCombinedCameraRigsCameraNodeEvaluator::OnRun(const FCameraNodeEvaluationParams& Params, FCameraNodeEvaluationResult& OutResult)
+{
+	// Keep applying overrides in case they are driven by a variable.
+	ApplyParameterOverrides(OutResult.VariableTable, true);
+
+	for (FCameraNodeEvaluator* CameraRigRootEvaluator : CameraRigRootEvaluators)
+	{
+		CameraRigRootEvaluator->Run(Params, OutResult);
+	}
+}
+
+void FCombinedCameraRigsCameraNodeEvaluator::ApplyParameterOverrides(FCameraVariableTable& OutVariableTable, bool bDrivenOnly)
+{
+	const UCombinedCameraRigsCameraNode* CombinedRigsNode = GetCameraNodeAs<UCombinedCameraRigsCameraNode>();
+	for (const FCameraRigAssetReference& IndividualCameraRigReference : CombinedRigsNode->CameraRigReferences)
+	{
+		FCameraRigParameterOverrideEvaluator OverrideEvaluator(IndividualCameraRigReference);
+		OverrideEvaluator.ApplyParameterOverrides(OutVariableTable, bDrivenOnly);
+	}
+}
+
+}  // namespace UE::Cameras
+
+FCameraNodeEvaluatorPtr UCombinedCameraRigsCameraNode::OnBuildEvaluator(FCameraNodeEvaluatorBuilder& Builder) const
+{
+	using namespace UE::Cameras;
+	return Builder.BuildEvaluator<FCombinedCameraRigsCameraNodeEvaluator>();
+}
+
diff --git a/Engine/Plugins/Cameras/GameplayCameras/Source/GameplayCameras/Private/Core/CameraRigCombinationRegistry.h b/Engine/Plugins/Cameras/GameplayCameras/Source/GameplayCameras/Private/Core/CameraRigCombinationRegistry.h
new file mode 100644
--- /dev/null
+++ b/Engine/Plugins/Cameras/GameplayCameras/Source/GameplayCameras/Private/Core/CameraRigCombinationRegistry.h
@@ -0,0 +1,68 @@
+// Copyright Epic Games, Inc. All Rights Reserved.
+
+#pragma once
+
+#include "CoreTypes.h"
+#include "Core/CameraNode.h"
+#include "Core/CameraRigAsset.h"
+#include "Core/CameraRigAssetReference.h"
+
+#include "CameraRigCombinationRegistry.generated.h"
+
+namespace UE::Cameras
+{
+
+class FCameraRigCombinationRegistry
+{
+public:
+
+	const UCameraRigAsset* FindOrCreateCombination(TArrayView<const UCameraRigAsset*> InCombination);
+
+	void AddReferencedObjects(FReferenceCollector& Collector);
+
+private:
+
+	using FCameraRigCombination = TArray<TWeakObjectPtr<const UCameraRigAsset>, TInlineAllocator<4>>;
+
+	struct FCameraRigCombinationKey
+	{
+		FCameraRigCombinationKey(TArrayView<const UCameraRigAsset*> InCombination);
+		FCameraRigCombination Combination;
+		uint32 CachedHash = 0;
+	};
+
+	friend inline uint32 GetTypeHash(const FCameraRigCombinationKey& Key)
+	{
+		return Key.CachedHash;
+	}
+
+	friend inline bool operator==(const FCameraRigCombinationKey& A, const FCameraRigCombinationKey& B)
+	{
+		return A.Combination == B.Combination;
+	}
+
+	using FCameraRigCombinationMap = TMap<FCameraRigCombinationKey, int32>;
+	FCameraRigCombinationMap Combinations;
+
+	TArray<TObjectPtr<const UCameraRigAsset>> CombinedCameraRigs;
+};
+
+}  // namespace UE::Cameras
+ 
+UCLASS(MinimalAPI, Hidden)
+class UCombinedCameraRigsCameraNode : public UCameraNode
+{
+	GENERATED_BODY()
+
+protected:
+
+	// UCameraNode interface.
+	virtual FCameraNodeEvaluatorPtr OnBuildEvaluator(FCameraNodeEvaluatorBuilder& Builder) const override;
+
+public:
+
+	/** The camera rigs to run. */
+	UPROPERTY()
+	TArray<FCameraRigAssetReference> CameraRigReferences;
+};
+
diff --git a/Engine/Plugins/Cameras/GameplayCameras/Source/GameplayCameras/Private/Core/CameraSystemEvaluator.cpp b/Engine/Plugins/Cameras/GameplayCameras/Source/GameplayCameras/Private/Core/CameraSystemEvaluator.cpp
--- a/Engine/Plugins/Cameras/GameplayCameras/Source/GameplayCameras/Private/Core/CameraSystemEvaluator.cpp
+++ b/Engine/Plugins/Cameras/GameplayCameras/Source/GameplayCameras/Private/Core/CameraSystemEvaluator.cpp
@@ -2,16 +2,19 @@
 
 #include "Core/CameraSystemEvaluator.h"
 
+#include "Algo/Transform.h"
 #include "Camera/CameraTypes.h"
 #include "Core/CameraDirectorEvaluator.h"
 #include "Core/CameraEvaluationContext.h"
 #include "Core/CameraEvaluationService.h"
+#include "Core/CameraRigCombinationRegistry.h"
 #include "Core/DefaultRootCameraNode.h"
 #include "Debug/CameraDebugBlock.h"
 #include "Debug/CameraDebugBlockBuilder.h"
 #include "Debug/CameraDebugRenderer.h"
 #include "Debug/CameraSystemTrace.h"
 #include "Debug/RootCameraDebugBlock.h"
+#include "GameplayCamerasSettings.h"
 #include "Services/AutoResetCameraVariableService.h"
 #include "Services/OrientationInitializationService.h"
 #include "UObject/Package.h"
@@ -69,6 +72,8 @@ void FCameraSystemEvaluator::Initialize(const FCameraSystemEvaluatorCreateParams
 	VariableAutoResetService = MakeShared<FAutoResetCameraVariableService>();
 	RegisterEvaluationService(VariableAutoResetService.ToSharedRef());
 	RegisterEvaluationService(MakeShared<FOrientationInitializationService>());
+
+	CameraRigCombinationRegistry = MakeShared<FCameraRigCombinationRegistry>();
 }
 
 FCameraSystemEvaluator::~FCameraSystemEvaluator()
@@ -98,6 +103,10 @@ void FCameraSystemEvaluator::AddReferencedObjects(FReferenceCollector& Collector
 	{
 		EvaluationService->AddReferencedObjects(Collector);
 	}
+	if (CameraRigCombinationRegistry)
+	{
+		CameraRigCombinationRegistry->AddReferencedObjects(Collector);
+	}
 }
 
 void FCameraSystemEvaluator::PushEvaluationContext(TSharedRef<FCameraEvaluationContext> EvaluationContext)
@@ -180,6 +189,7 @@ void FCameraSystemEvaluator::Update(const FCameraSystemEvaluationParams& Params)
 
 		if (DirectorResult.ActiveCameraRigs.Num() == 1)
 		{
+			// Only one camera rig to activate... let's do that.
 			const FActiveCameraRigInfo& ActiveCameraRig = DirectorResult.ActiveCameraRigs[0];
 
 			FActivateCameraRigParams CameraRigParams;
@@ -188,7 +198,38 @@ void FCameraSystemEvaluator::Update(const FCameraSystemEvaluationParams& Params)
 			CameraRigParams.CameraRig = ActiveCameraRig.CameraRig;
 			RootEvaluator->ActivateCameraRig(CameraRigParams);
 		}
-		// TODO: handle the case of composite camera rigs.
+		else if (DirectorResult.ActiveCameraRigs.Num() > 1)
+		{
+			// We have a combination of camera rigs to activate. Let's dynamically generate a new camera rig
+			// asset that combines them.
+#if WITH_EDITOR
+			const UGameplayCamerasSettings* Settings = GetDefault<UGameplayCamerasSettings>();
+			if (DirectorResult.ActiveCameraRigs.Num() > Settings->CombinedCameraRigNumThreshold)
+			{
+				UE_LOG(LogCameraSystem, Warning, 
+						TEXT("Activating %d camera rigs combined! Is the camera director doing this on purpose? "
+							"If so, raise the CombinedCameraRigNumThreshold setting to remove this warning."),
+						DirectorResult.ActiveCameraRigs.Num());
+			}
+#endif
+
+			// All combined camera rigs must belong to the same evaluation context.
+			TArray<const UCameraRigAsset*> Combination;
+			TSharedPtr<const FCameraEvaluationContext> CommonContext = DirectorResult.ActiveCameraRigs[0].EvaluationContext;
+			for (const FActiveCameraRigInfo& ActiveCameraRig : DirectorResult.ActiveCameraRigs)
+			{
+				Combination.Add(ActiveCameraRig.CameraRig);
+				ensureMsgf(ActiveCameraRig.EvaluationContext == CommonContext,
+						TEXT("All combined camera rigs must be activated from the same evaluation context."));
+			}
+			const UCameraRigAsset* CombinedCameraRig = CameraRigCombinationRegistry->FindOrCreateCombination(Combination);
+
+			FActivateCameraRigParams CameraRigParams;
+			CameraRigParams.Evaluator = this;
+			CameraRigParams.EvaluationContext = CommonContext;
+			CameraRigParams.CameraRig = CombinedCameraRig;
+			RootEvaluator->ActivateCameraRig(CameraRigParams);
+		}
 	}
 
 	{
diff --git a/Engine/Plugins/Cameras/GameplayCameras/Source/GameplayCameras/Public/Core/CameraRigAssetReference.h b/Engine/Plugins/Cameras/GameplayCameras/Source/GameplayCameras/Public/Core/CameraRigAssetReference.h
--- a/Engine/Plugins/Cameras/GameplayCameras/Source/GameplayCameras/Public/Core/CameraRigAssetReference.h
+++ b/Engine/Plugins/Cameras/GameplayCameras/Source/GameplayCameras/Public/Core/CameraRigAssetReference.h
@@ -398,6 +398,11 @@ struct GAMEPLAYCAMERAS_API FCameraRigAssetReference
 {
 	GENERATED_BODY()
 
+public:
+
+	FCameraRigAssetReference();
+	FCameraRigAssetReference(UCameraRigAsset* InCameraRig);
+
 	/** Returns whether this reference points to a valid camera rig. */
 	bool IsValid() const
 	{
diff --git a/Engine/Plugins/Cameras/GameplayCameras/Source/GameplayCameras/Public/Core/CameraSystemEvaluator.h b/Engine/Plugins/Cameras/GameplayCameras/Source/GameplayCameras/Public/Core/CameraSystemEvaluator.h
--- a/Engine/Plugins/Cameras/GameplayCameras/Source/GameplayCameras/Public/Core/CameraSystemEvaluator.h
+++ b/Engine/Plugins/Cameras/GameplayCameras/Source/GameplayCameras/Public/Core/CameraSystemEvaluator.h
@@ -27,6 +27,7 @@ namespace UE::Cameras
 class FAutoResetCameraVariableService;
 class FCameraEvaluationContext;
 class FCameraEvaluationService;
+class FCameraRigCombinationRegistry;
 class FRootCameraNodeEvaluator;
 enum class ECameraEvaluationServiceFlags;
 struct FRootCameraNodeCameraRigEvent;
@@ -172,6 +173,9 @@ private:
 	/** Quick access to the variable auto-reset service. */
 	TSharedPtr<FAutoResetCameraVariableService> VariableAutoResetService;
 
+	/** Registry for programmatically building combinations of camera rigs. */
+	TSharedPtr<FCameraRigCombinationRegistry> CameraRigCombinationRegistry;
+
 	/** Storage buffer for the root evaluator. */
 	FCameraNodeEvaluatorStorage RootEvaluatorStorage;
 
diff --git a/Engine/Plugins/Cameras/GameplayCameras/Source/GameplayCameras/Public/GameplayCamerasSettings.h b/Engine/Plugins/Cameras/GameplayCameras/Source/GameplayCameras/Public/GameplayCamerasSettings.h
--- a/Engine/Plugins/Cameras/GameplayCameras/Source/GameplayCameras/Public/GameplayCamerasSettings.h
+++ b/Engine/Plugins/Cameras/GameplayCameras/Source/GameplayCameras/Public/GameplayCamerasSettings.h
@@ -27,6 +27,12 @@ public:
 	UPROPERTY(EditAnywhere, Config, Category="General")
 	bool bAutoSpawnCameraSystemActor = true;
 
+	/**
+	 * The number of camera rigs combined in one frame past which the camera system emits a warning.
+	 */
+	UPROPERTY(EditAnywhere, Config, Category="General")
+	int32 CombinedCameraRigNumThreshold = 10;
+
 public:
 
 	/** The default angle tolerance to accept an aiming operation. */
