Cameras: adjust what parameter structs pass the system evaluator

diff --git a/Engine/Plugins/Cameras/GameplayCameras/Source/GameplayCameras/Private/Core/BlendStackCameraNode.cpp b/Engine/Plugins/Cameras/GameplayCameras/Source/GameplayCameras/Private/Core/BlendStackCameraNode.cpp
--- a/Engine/Plugins/Cameras/GameplayCameras/Source/GameplayCameras/Private/Core/BlendStackCameraNode.cpp
+++ b/Engine/Plugins/Cameras/GameplayCameras/Source/GameplayCameras/Private/Core/BlendStackCameraNode.cpp
@@ -206,7 +206,7 @@ bool FBlendStackCameraNodeEvaluator::InitializeEntry(
 
 	// Initialize the node evaluators.
 	FCameraNodeEvaluatorInitializeParams InitParams(&NewEntry.EvaluatorHierarchy);
-	InitParams.Evaluator = Evaluator;
+	InitParams.Evaluator = OwningEvaluator;
 	InitParams.EvaluationContext = EvaluationContext;
 	InitParams.LastActiveCameraRigInfo = GetActiveCameraRigEvaluationInfo();
 	RootEvaluator->Initialize(InitParams, NewEntry.Result);
diff --git a/Engine/Plugins/Cameras/GameplayCameras/Source/GameplayCameras/Private/Core/CameraEvaluationContext.cpp b/Engine/Plugins/Cameras/GameplayCameras/Source/GameplayCameras/Private/Core/CameraEvaluationContext.cpp
--- a/Engine/Plugins/Cameras/GameplayCameras/Source/GameplayCameras/Private/Core/CameraEvaluationContext.cpp
+++ b/Engine/Plugins/Cameras/GameplayCameras/Source/GameplayCameras/Private/Core/CameraEvaluationContext.cpp
@@ -107,6 +107,7 @@ void FCameraEvaluationContext::Activate(const FCameraEvaluationContextActivatePa
 	if (ensure(DirectorEvaluator))
 	{
 		FCameraDirectorActivateParams DirectorParams;
+		DirectorParams.Evaluator = Params.Evaluator;
 		DirectorParams.OwnerContext = SharedThis(this);
 		DirectorEvaluator->Activate(DirectorParams);
 	}
diff --git a/Engine/Plugins/Cameras/GameplayCameras/Source/GameplayCameras/Private/Core/CameraEvaluationContextStack.cpp b/Engine/Plugins/Cameras/GameplayCameras/Source/GameplayCameras/Private/Core/CameraEvaluationContextStack.cpp
--- a/Engine/Plugins/Cameras/GameplayCameras/Source/GameplayCameras/Private/Core/CameraEvaluationContextStack.cpp
+++ b/Engine/Plugins/Cameras/GameplayCameras/Source/GameplayCameras/Private/Core/CameraEvaluationContextStack.cpp
@@ -60,6 +60,7 @@ void FCameraEvaluationContextStack::PushContext(TSharedRef<FCameraEvaluationCont
 	FContextEntry NewEntry;
 
 	FCameraEvaluationContextActivateParams ActivateParams;
+	ActivateParams.Evaluator = Evaluator;
 	Context->Activate(ActivateParams);
 	
 	NewEntry.WeakContext = Context;
diff --git a/Engine/Plugins/Cameras/GameplayCameras/Source/GameplayCameras/Private/Core/CameraSystemEvaluator.cpp b/Engine/Plugins/Cameras/GameplayCameras/Source/GameplayCameras/Private/Core/CameraSystemEvaluator.cpp
--- a/Engine/Plugins/Cameras/GameplayCameras/Source/GameplayCameras/Private/Core/CameraSystemEvaluator.cpp
+++ b/Engine/Plugins/Cameras/GameplayCameras/Source/GameplayCameras/Private/Core/CameraSystemEvaluator.cpp
@@ -193,7 +193,6 @@ void FCameraSystemEvaluator::Update(const FCameraSystemEvaluationParams& Params)
 			const FActiveCameraRigInfo& ActiveCameraRig = DirectorResult.ActiveCameraRigs[0];
 
 			FActivateCameraRigParams CameraRigParams;
-			CameraRigParams.Evaluator = this;
 			CameraRigParams.EvaluationContext = ActiveCameraRig.EvaluationContext;
 			CameraRigParams.CameraRig = ActiveCameraRig.CameraRig;
 			RootEvaluator->ActivateCameraRig(CameraRigParams);
@@ -225,7 +224,6 @@ void FCameraSystemEvaluator::Update(const FCameraSystemEvaluationParams& Params)
 			const UCameraRigAsset* CombinedCameraRig = CameraRigCombinationRegistry->FindOrCreateCombination(Combination);
 
 			FActivateCameraRigParams CameraRigParams;
-			CameraRigParams.Evaluator = this;
 			CameraRigParams.EvaluationContext = CommonContext;
 			CameraRigParams.CameraRig = CombinedCameraRig;
 			RootEvaluator->ActivateCameraRig(CameraRigParams);
diff --git a/Engine/Plugins/Cameras/GameplayCameras/Source/GameplayCameras/Private/Core/DefaultRootCameraNode.cpp b/Engine/Plugins/Cameras/GameplayCameras/Source/GameplayCameras/Private/Core/DefaultRootCameraNode.cpp
--- a/Engine/Plugins/Cameras/GameplayCameras/Source/GameplayCameras/Private/Core/DefaultRootCameraNode.cpp
+++ b/Engine/Plugins/Cameras/GameplayCameras/Source/GameplayCameras/Private/Core/DefaultRootCameraNode.cpp
@@ -85,7 +85,6 @@ void FDefaultRootCameraNodeEvaluator::OnActivateCameraRig(const FActivateCameraR
 	if (ensure(TargetStack))
 	{
 		FBlendStackCameraPushParams PushParams;
-		PushParams.Evaluator = Params.Evaluator;
 		PushParams.EvaluationContext = Params.EvaluationContext;
 		PushParams.CameraRig = Params.CameraRig;
 		TargetStack->Push(PushParams);
diff --git a/Engine/Plugins/Cameras/GameplayCameras/Source/GameplayCameras/Private/GameFramework/ActivateCameraRigFunctions.cpp b/Engine/Plugins/Cameras/GameplayCameras/Source/GameplayCameras/Private/GameFramework/ActivateCameraRigFunctions.cpp
--- a/Engine/Plugins/Cameras/GameplayCameras/Source/GameplayCameras/Private/GameFramework/ActivateCameraRigFunctions.cpp
+++ b/Engine/Plugins/Cameras/GameplayCameras/Source/GameplayCameras/Private/GameFramework/ActivateCameraRigFunctions.cpp
@@ -116,7 +116,6 @@ void UControllerGameplayCameraEvaluationComponent::ActivateCameraRigs()
 			FActivateCameraRigParams Params;
 			Params.CameraRig = CameraRigInfo.CameraRig;
 			Params.EvaluationContext = EvaluationContext;
-			Params.Evaluator = SystemEvaluator.Get();
 			Params.Layer = CameraRigInfo.EvaluationLayer;
 			RootNodeEvaluator->ActivateCameraRig(Params);
 
diff --git a/Engine/Plugins/Cameras/GameplayCameras/Source/GameplayCameras/Public/Core/CameraDirectorEvaluator.h b/Engine/Plugins/Cameras/GameplayCameras/Source/GameplayCameras/Public/Core/CameraDirectorEvaluator.h
--- a/Engine/Plugins/Cameras/GameplayCameras/Source/GameplayCameras/Public/Core/CameraDirectorEvaluator.h
+++ b/Engine/Plugins/Cameras/GameplayCameras/Source/GameplayCameras/Public/Core/CameraDirectorEvaluator.h
@@ -17,6 +17,7 @@ namespace UE::Cameras
 class FCameraDirectorEvaluator;
 class FCameraDirectorEvaluatorStorage;
 class FCameraEvaluationContext;
+class FCameraSystemEvaluator;
 
 /**
  * Parameter structure for initializing a newly created camera director evaluator.
@@ -35,10 +36,16 @@ struct FCameraDirectorInitializeParams
  */
 struct FCameraDirectorActivateParams
 {
+	/** The camera system that will run the camera director. */
+	FCameraSystemEvaluator* Evaluator = nullptr;
+
 	/** The evaluation context that owns the camera director. */
 	TSharedPtr<FCameraEvaluationContext> OwnerContext;
 };
 
+/**
+ * Parameter structure for deactivating a camera director evaluator.
+ */
 struct FCameraDirectorDeactivateParams
 {
 	/** The evaluation context that owns the camera director. */
@@ -57,9 +64,14 @@ struct FCameraDirectorEvaluationParams
 	TSharedPtr<FCameraEvaluationContext> OwnerContext;
 };
 
+/**
+ * Structure describing a camera rig that should be active.
+ */
 struct FActiveCameraRigInfo
 {
+	/** The evaluation context to run the specified camera rig. */
 	TSharedPtr<const FCameraEvaluationContext> EvaluationContext;
+	/** The camera rig that should be running. */
 	TObjectPtr<const UCameraRigAsset> CameraRig;
 };
 
@@ -69,7 +81,8 @@ struct FActiveCameraRigInfo
 struct FCameraDirectorEvaluationResult
 {
 	using FActiveCameraRigInfos = TArray<FActiveCameraRigInfo, TInlineAllocator<2>>;
-	/** The camera rig(s) that the director says should be active this frame. */
+
+	/** The camera rig(s) that the director says should be active this frame in the main layer. */
 	FActiveCameraRigInfos ActiveCameraRigs;
 
 	void Add(TSharedPtr<const FCameraEvaluationContext> InContext, TObjectPtr<const UCameraRigAsset> InCameraRig)
diff --git a/Engine/Plugins/Cameras/GameplayCameras/Source/GameplayCameras/Public/Core/CameraEvaluationContext.h b/Engine/Plugins/Cameras/GameplayCameras/Source/GameplayCameras/Public/Core/CameraEvaluationContext.h
--- a/Engine/Plugins/Cameras/GameplayCameras/Source/GameplayCameras/Public/Core/CameraEvaluationContext.h
+++ b/Engine/Plugins/Cameras/GameplayCameras/Source/GameplayCameras/Public/Core/CameraEvaluationContext.h
@@ -16,6 +16,8 @@ class UCameraDirector;
 namespace UE::Cameras
 {
 
+class FCameraSystemEvaluator;
+
 /**
  * Parameter struct for initializing an evaluation context.
  */
@@ -35,6 +37,7 @@ struct FCameraEvaluationContextInitializeParams
  */
 struct FCameraEvaluationContextActivateParams
 {
+	FCameraSystemEvaluator* Evaluator = nullptr;
 };
 
 /**
diff --git a/Engine/Plugins/Cameras/GameplayCameras/Source/GameplayCameras/Public/Core/CameraSystemEvaluator.h b/Engine/Plugins/Cameras/GameplayCameras/Source/GameplayCameras/Public/Core/CameraSystemEvaluator.h
--- a/Engine/Plugins/Cameras/GameplayCameras/Source/GameplayCameras/Public/Core/CameraSystemEvaluator.h
+++ b/Engine/Plugins/Cameras/GameplayCameras/Source/GameplayCameras/Public/Core/CameraSystemEvaluator.h
@@ -18,7 +18,6 @@ class UCameraDirector;
 class UCameraRigAsset;
 class UCanvas;
 class URootCameraNode;
-enum class ECameraRigLayer : uint8;
 struct FMinimalViewInfo;
 
 namespace UE::Cameras
diff --git a/Engine/Plugins/Cameras/GameplayCameras/Source/GameplayCameras/Public/Core/RootCameraNode.h b/Engine/Plugins/Cameras/GameplayCameras/Source/GameplayCameras/Public/Core/RootCameraNode.h
--- a/Engine/Plugins/Cameras/GameplayCameras/Source/GameplayCameras/Public/Core/RootCameraNode.h
+++ b/Engine/Plugins/Cameras/GameplayCameras/Source/GameplayCameras/Public/Core/RootCameraNode.h
@@ -19,8 +19,7 @@ enum class ECameraRigLayer : uint8
 	Base UMETA(DisplayName="Base Layer"),
 	Main UMETA(DisplayName="Main Layer"),
 	Global UMETA(DisplayName="Global Layer"),
-	Visual UMETA(DisplayName="Visual Layer"),
-	ScratchMain UMETA(Hidden)
+	Visual UMETA(DisplayName="Visual Layer")
 };
 ENUM_CLASS_FLAGS(ECameraRigLayer)
 
@@ -46,9 +45,6 @@ struct FRootCameraNodeCameraRigEvent;
  */
 struct FActivateCameraRigParams
 {
-	/** The evaluator currently running.*/
-	FCameraSystemEvaluator* Evaluator = nullptr;
-
 	/** The evaluation context in which the camera rig runs. */
 	TSharedPtr<const FCameraEvaluationContext> EvaluationContext;
 
