Cameras: split blend stack evaluator in two, add ability to add/remove persistent camera rigs from the non-main layers

- Separate the blend stack evaluator into two implementations: transient and persistent blend stacks. Transient blend stacks are for the usual gameplay camera rigs, which blend each other out. Persistent blend stacks are for the base/global/visual stacks where rigs live for as long as they want, and the "top" rig isn't necessarily more special than the others.
- Add a DeactivateCameraRig API on the root camera node.

diff --git a/Engine/Plugins/Cameras/GameplayCameras/Source/GameplayCameras/Private/Core/BlendStackCameraNode.cpp b/Engine/Plugins/Cameras/GameplayCameras/Source/GameplayCameras/Private/Core/BlendStackCameraNode.cpp
--- a/Engine/Plugins/Cameras/GameplayCameras/Source/GameplayCameras/Private/Core/BlendStackCameraNode.cpp
+++ b/Engine/Plugins/Cameras/GameplayCameras/Source/GameplayCameras/Private/Core/BlendStackCameraNode.cpp
@@ -27,7 +27,17 @@
 FCameraNodeEvaluatorPtr UBlendStackCameraNode::OnBuildEvaluator(FCameraNodeEvaluatorBuilder& Builder) const
 {
 	using namespace UE::Cameras;
-	return Builder.BuildEvaluator<FBlendStackCameraNodeEvaluator>();
+
+	switch (BlendStackType)
+	{
+		case ECameraBlendStackType::Persistent:
+			return Builder.BuildEvaluator<FPersistentBlendStackCameraNodeEvaluator>();
+		case ECameraBlendStackType::Transient:
+			return Builder.BuildEvaluator<FTransientBlendStackCameraNodeEvaluator>();
+		default:
+			ensure(false);
+			return nullptr;
+	}
 }
 
 namespace UE::Cameras
@@ -46,6 +56,8 @@
 	TEXT(""));
 
 UE_DEFINE_CAMERA_NODE_EVALUATOR(FBlendStackCameraNodeEvaluator)
+UE_DEFINE_CAMERA_NODE_EVALUATOR(FTransientBlendStackCameraNodeEvaluator)
+UE_DEFINE_CAMERA_NODE_EVALUATOR(FPersistentBlendStackCameraNodeEvaluator)
 
 FBlendStackCameraNodeEvaluator::~FBlendStackCameraNodeEvaluator()
 {
@@ -53,144 +65,15 @@
 	PopEntries(Entries.Num());
 }
 
-void FBlendStackCameraNodeEvaluator::Push(const FBlendStackCameraPushParams& Params)
-{
-	bool bSearchedForTransition = false;
-	const UCameraRigTransition* Transition = nullptr;
-
-	if (!Entries.IsEmpty())
-	{
-		FCameraRigEntry& TopEntry(Entries.Top());
-		if (!TopEntry.bIsFrozen 
-				&& TopEntry.EvaluationContext == Params.EvaluationContext)
-		{
-			// Don't push anything is what is being requested is already the active 
-			// camera rig.
-			if (TopEntry.CameraRig == Params.CameraRig)
-			{
-				return;
-			}
-
-			// See if we can merge the new camera rig onto the active camera rig.
-			const EBlendStackEntryComparison Comparison = TopEntry.RootEvaluator->Compare(Params.CameraRig);
-
-			if (Comparison == EBlendStackEntryComparison::Active)
-			{
-				// This camera rig is already the active one on the merged stack.
-				return;
-			}
-
-			if (Comparison == EBlendStackEntryComparison::EligibleForMerge)
-			{
-				// This camera rig can be merged with the one current running. However, we
-				// only do it if the transition explicitly allows it.
-				bSearchedForTransition = true;
-				Transition = FindTransition(Params);
-
-				if (Transition && Transition->bAllowCameraRigMerging)
-				{
-					PushVariantEntry(Params, Transition);
-					return;
-				}
-			}
-		}
-	}
-
-	// It's a legitimate new entry in the blend stack.
-	if (!bSearchedForTransition)
-	{
-		bSearchedForTransition = true;
-		Transition = FindTransition(Params);
-	}
-
-	PushNewEntry(Params, Transition);
-}
-
-void FBlendStackCameraNodeEvaluator::PushNewEntry(const FBlendStackCameraPushParams& Params, const UCameraRigTransition* Transition)
-{
-	// Create the new root node to wrap the new camera rig's root node, and the specific
-	// blend node for this transition.
-	// We need to const-cast here to be able to use our own blend stack node as the outer
-	// of the new node.
-	const UCameraRigTransition* UsedTransition = nullptr;
-	UObject* Outer = const_cast<UObject*>((UObject*)GetCameraNode());
-	UBlendStackRootCameraNode* EntryRootNode = NewObject<UBlendStackRootCameraNode>(Outer, NAME_None);
-	{
-		EntryRootNode->RootNode = Params.CameraRig->RootNode;
-
-		// Find a transition and use its blend. If no transition is found,
-		// make a camera cut transition.
-		UBlendCameraNode* ModeBlend = nullptr;
-		if (Transition)
-		{
-			ModeBlend = Transition->Blend;
-			UsedTransition = Transition;
-		}
-		if (!ModeBlend)
-		{
-			ModeBlend = NewObject<UPopBlendCameraNode>(EntryRootNode, NAME_None);
-		}
-		EntryRootNode->Blend = ModeBlend;
-	}
-
-	// Make the new stack entry, and use its storage buffer to build the tree of evaluators.
-	FCameraRigEntry NewEntry;
-	const bool bInitialized = InitializeEntry(
-			NewEntry, 
-			Params.CameraRig,
-			Params.Evaluator,
-			Params.EvaluationContext,
-			EntryRootNode);
-	if (!bInitialized)
-	{
-		return;
-	}
-
-#if WITH_EDITOR
-	// Listen to changes to the packages inside which this camera rig is defined. We will hot-reload the
-	// camera node evaluators for this camera rig when we detect changes.
-	AddPackageListeners(NewEntry);
-#endif  // WITH_EDITOR
-
-	// Important: we need to move the new entry here because copying evaluator storage
-	// is disabled.
-	Entries.Add(MoveTemp(NewEntry));
-
-	if (OnCameraRigEventDelegate.IsBound())
-	{
-		BroadcastCameraRigEvent(EBlendStackCameraRigEventType::Pushed, Entries.Last(), UsedTransition);
-	}
-}
-
-void FBlendStackCameraNodeEvaluator::PushVariantEntry(const FBlendStackCameraPushParams& PushParams, const UCameraRigTransition* Transition)
-{
-	const UCameraRigCameraNode* PrefabNode = Cast<const UCameraRigCameraNode>(PushParams.CameraRig->RootNode);
-	const UBlendCameraNode* Blend = Transition ? Transition->Blend : nullptr;
-
-	FCameraRigEntry& TopEntry = Entries.Top();
-	FCameraNodeEvaluatorBuilder Builder(TopEntry.EvaluatorStorage);
-	FCameraNodeEvaluatorBuildParams BuildParams(Builder);
-	TopEntry.RootEvaluator->MergeCameraRig(BuildParams, PrefabNode, Blend);
-
-	// Swap out the camera rig registered as "active" for this entry.
-#if WITH_EDITOR
-	RemoveListenedPackages(TopEntry);
-#endif
-	{
-		TopEntry.CameraRig = PushParams.CameraRig;
-	}
-#if WITH_EDITOR
-	AddPackageListeners(TopEntry);
-#endif
-}
-
 bool FBlendStackCameraNodeEvaluator::InitializeEntry(
 		FCameraRigEntry& NewEntry, 
 		const UCameraRigAsset* CameraRig,
-		FCameraSystemEvaluator* Evaluator,
 		TSharedPtr<const FCameraEvaluationContext> EvaluationContext,
 		UBlendStackRootCameraNode* EntryRootNode)
 {
+	// Clear the evaluator hierarchy in case we are hot-reloading an entry.
+	NewEntry.EvaluatorHierarchy.Reset();
+
 	// Generate the hierarchy of node evaluators inside our storage buffer.
 	FCameraNodeEvaluatorTreeBuildParams BuildParams;
 	BuildParams.RootCameraNode = EntryRootNode;
@@ -228,8 +111,10 @@
 	Entry.RootEvaluator = nullptr;
 	Entry.EvaluatorHierarchy.Reset();
 
+	Entry.RootNode = nullptr;
+
 	Entry.EvaluationContext.Reset();
-	
+
 #if WITH_EDITOR
 	RemoveListenedPackages(Entry);
 #endif
@@ -237,32 +122,6 @@
 	Entry.bIsFrozen = true;
 }
 
-void FBlendStackCameraNodeEvaluator::GatherEntryParameterEvaluators(FCameraNodeEvaluator* RootEvaluator, TArray<FCameraNodeEvaluator*>& OutParameterEvaluators)
-{
-	TArray<FCameraNodeEvaluator*> EvaluatorStack;
-	EvaluatorStack.Add(RootEvaluator);
-	while (!EvaluatorStack.IsEmpty())
-	{
-		FCameraNodeEvaluator* CurEvaluator = EvaluatorStack.Pop();
-
-		if (EnumHasAnyFlags(CurEvaluator->GetNodeEvaluatorFlags(), ECameraNodeEvaluatorFlags::NeedsParameterUpdate))
-		{
-			OutParameterEvaluators.Add(CurEvaluator);
-		}
-		else
-		{
-			FCameraNodeEvaluatorChildrenView CurChildren(CurEvaluator->GetChildren());
-			for (FCameraNodeEvaluator* Child : ReverseIterate(CurChildren))
-			{
-				if (Child)
-				{
-					EvaluatorStack.Add(Child);
-				}
-			}
-		}
-	}
-}
-
 FCameraRigEvaluationInfo FBlendStackCameraNodeEvaluator::GetActiveCameraRigEvaluationInfo() const
 {
 	if (Entries.Num() > 0)
@@ -413,10 +272,10 @@
 		// Blends aren't expected to use the camera pose to do any logic until BlendResults().
 		if (!Entry.bBlendRunThisFrame)
 		{
-			FBlendCameraNodeEvaluator* BlendEvaluator = Entry.RootEvaluator->GetBlendEvaluator();
-			if (ensure(BlendEvaluator))
+			FBlendCameraNodeEvaluator* EntryBlendEvaluator = Entry.RootEvaluator->GetBlendEvaluator();
+			if (EntryBlendEvaluator)
 			{
-				BlendEvaluator->Run(CurParams, CurResult);
+				EntryBlendEvaluator->Run(CurParams, CurResult);
 			}
 
 			Entry.bBlendRunThisFrame = true;
@@ -439,10 +298,14 @@
 			FCameraNodePreBlendResult PreBlendResult(OutResult.VariableTable);
 
 			FBlendCameraNodeEvaluator* EntryBlendEvaluator = Entry.RootEvaluator->GetBlendEvaluator();
-			if (ensure(EntryBlendEvaluator))
+			if (EntryBlendEvaluator)
 			{
 				EntryBlendEvaluator->BlendParameters(PreBlendParams, PreBlendResult);
 			}
+			else
+			{
+				OutResult.VariableTable.Override(CurResult.VariableTable, ECameraVariableTableFilter::Input);
+			}
 		}
 		else
 		{
@@ -519,23 +382,21 @@
 			}
 			else
 			{
-				OutResult.CameraPose.OverrideAll(CurResult.CameraPose);
-				OutResult.VariableTable.OverrideAll(CurResult.VariableTable);
+				OutResult.OverrideAll(CurResult);
 
 				PopEntriesBelow = ResolvedEntry.EntryIndex;
 			}
 		}
 		else
 		{
-			OutResult.CameraPose.OverrideAll(CurResult.CameraPose);
-			OutResult.VariableTable.OverrideAll(CurResult.VariableTable);
+			OutResult.OverrideAll(CurResult);
 
 			PopEntriesBelow = ResolvedEntry.EntryIndex;
 		}
 	}
 
 	// Pop out camera rigs that have been blended out.
-	if (BlendStackNode->bAutoPop && PopEntriesBelow != INDEX_NONE)
+	if (BlendStackNode->BlendStackType == ECameraBlendStackType::Transient && PopEntriesBelow != INDEX_NONE)
 	{
 		PopEntries(PopEntriesBelow);
 	}
@@ -549,6 +410,26 @@
 	}
 }
 
+void FBlendStackCameraNodeEvaluator::PopEntry(int32 EntryIndex)
+{
+	if (!ensure(Entries.IsValidIndex(EntryIndex)))
+	{
+		return;
+	}
+
+	FCameraRigEntry& Entry = Entries[EntryIndex];
+#if WITH_EDITOR
+	RemoveListenedPackages(Entry);
+#endif  // WITH_EDITOR
+
+	if (OnCameraRigEventDelegate.IsBound())
+	{
+		BroadcastCameraRigEvent(EBlendStackCameraRigEventType::Popped, Entry);
+	}
+
+	Entries.RemoveAt(EntryIndex);
+}
+
 void FBlendStackCameraNodeEvaluator::PopEntries(int32 FirstIndexToKeep)
 {
 	if (UNLIKELY(Entries.IsEmpty()))
@@ -629,7 +510,249 @@
 
 #endif  // WITH_EDITOR
 
-const UCameraRigTransition* FBlendStackCameraNodeEvaluator::FindTransition(const FBlendStackCameraPushParams& Params) const
+void FBlendStackCameraNodeEvaluator::BroadcastCameraRigEvent(EBlendStackCameraRigEventType EventType, const FCameraRigEntry& Entry, const UCameraRigTransition* Transition) const
 {
+	FBlendStackCameraRigEvent Event;
+	Event.EventType = EventType;
+	Event.BlendStackEvaluator = this;
+	Event.CameraRigInfo = FCameraRigEvaluationInfo(
+			Entry.EvaluationContext.Pin(),
+			Entry.CameraRig,
+			&Entry.Result,
+			Entry.RootEvaluator);
+	Event.Transition = Transition;
+
+	OnCameraRigEventDelegate.Broadcast(Event);
+}
+
+void FBlendStackCameraNodeEvaluator::OnAddReferencedObjects(FReferenceCollector& Collector)
+{
+	for (FCameraRigEntry& Entry : Entries)
+	{
+		Collector.AddReferencedObject(Entry.CameraRig);
+		Collector.AddReferencedObject(Entry.RootNode);
+	}
+}
+
+void FBlendStackCameraNodeEvaluator::OnSerialize(const FCameraNodeEvaluatorSerializeParams& Params, FArchive& Ar)
+{
+	if (Ar.IsSaving())
+	{
+		int32 NumEntries = Entries.Num();
+		Ar << NumEntries;
+	}
+	else if (Ar.IsLoading())
+	{
+		int32 LoadedNumEntries = 0;
+		Ar << LoadedNumEntries;
+
+		ensure(LoadedNumEntries == Entries.Num());
+	}
+
+	for (FCameraRigEntry& Entry : Entries)
+	{
+		Entry.Result.Serialize(Ar);
+		Ar << Entry.bIsFirstFrame;
+		Ar << Entry.bInputRunThisFrame;
+		Ar << Entry.bBlendRunThisFrame;
+		Ar << Entry.bIsFrozen;
+#if UE_GAMEPLAY_CAMERAS_TRACE
+		Ar << Entry.bLogWarnings;
+#endif  // UE_GAMEPLAY_CAMERAS_TRACE
+	}
+}
+
+#if WITH_EDITOR
+
+void FBlendStackCameraNodeEvaluator::OnPostBuildAsset(const FGameplayCameraAssetBuildEvent& BuildEvent)
+{
+	for (FCameraRigEntry& Entry : Entries)
+	{
+		const bool bRebuildEntry = Entry.ListenedPackages.Contains(BuildEvent.AssetPackage);
+		if (bRebuildEntry)
+		{
+			Entry.EvaluatorStorage.DestroyEvaluatorTree();
+			Entry.EvaluatorHierarchy.Reset();
+
+			// Re-assign the root node in case the camera rig's root was changed.
+			Entry.RootNode->RootNode = Entry.CameraRig->RootNode;
+
+			// Remove the blend on the root node, since we don't want the reloaded camera rig to re-blend-in
+			// for no good reason.
+			Entry.RootNode->Blend = NewObject<UPopBlendCameraNode>(Entry.RootNode, NAME_None);
+
+			// Rebuild the evaluator tree.
+			const bool bInitialized = InitializeEntry(
+					Entry,
+					Entry.CameraRig,
+					Entry.EvaluationContext.Pin(),
+					Entry.RootNode);
+			if (!bInitialized)
+			{
+				Entry.bIsFrozen = true;
+				continue;
+			}
+		}
+	}
+}
+
+#endif  // WITH_EDITOR
+
+void FTransientBlendStackCameraNodeEvaluator::Push(const FBlendStackCameraPushParams& Params)
+{
+	bool bSearchedForTransition = false;
+	const UCameraRigTransition* Transition = nullptr;
+
+	if (!Entries.IsEmpty())
+	{
+		FCameraRigEntry& TopEntry(Entries.Top());
+		if (!TopEntry.bIsFrozen 
+				&& TopEntry.EvaluationContext == Params.EvaluationContext)
+		{
+			// Don't push anything is what is being requested is already the active 
+			// camera rig.
+			if (TopEntry.CameraRig == Params.CameraRig)
+			{
+				return;
+			}
+
+			// See if we can merge the new camera rig onto the active camera rig.
+			const EBlendStackEntryComparison Comparison = TopEntry.RootEvaluator->Compare(Params.CameraRig);
+
+			if (Comparison == EBlendStackEntryComparison::Active)
+			{
+				// This camera rig is already the active one on the merged stack.
+				return;
+			}
+
+			if (Comparison == EBlendStackEntryComparison::EligibleForMerge)
+			{
+				// This camera rig can be merged with the one current running. However, we
+				// only do it if the transition explicitly allows it.
+				bSearchedForTransition = true;
+				Transition = FindTransition(Params);
+
+				if (Transition && Transition->bAllowCameraRigMerging)
+				{
+					PushVariantEntry(Params, Transition);
+					return;
+				}
+			}
+		}
+	}
+
+	// It's a legitimate new entry in the blend stack.
+	if (!bSearchedForTransition)
+	{
+		bSearchedForTransition = true;
+		Transition = FindTransition(Params);
+	}
+
+	PushNewEntry(Params, Transition);
+}
+
+void FTransientBlendStackCameraNodeEvaluator::PushNewEntry(const FBlendStackCameraPushParams& Params, const UCameraRigTransition* Transition)
+{
+	// Create the new root node to wrap the new camera rig's root node, and the specific
+	// blend node for this transition.
+	// We need to const-cast here to be able to use our own blend stack node as the outer
+	// of the new node.
+	const UCameraRigTransition* UsedTransition = nullptr;
+	UObject* Outer = const_cast<UObject*>((UObject*)GetCameraNode());
+	UBlendStackRootCameraNode* EntryRootNode = NewObject<UBlendStackRootCameraNode>(Outer, NAME_None);
+	{
+		EntryRootNode->RootNode = Params.CameraRig->RootNode;
+
+		// Find a transition and use its blend. If no transition is found,
+		// make a camera cut transition.
+		UBlendCameraNode* ModeBlend = nullptr;
+		if (Transition)
+		{
+			ModeBlend = Transition->Blend;
+			UsedTransition = Transition;
+		}
+		if (!ModeBlend)
+		{
+			ModeBlend = NewObject<UPopBlendCameraNode>(EntryRootNode, NAME_None);
+		}
+		EntryRootNode->Blend = ModeBlend;
+	}
+
+	// Make the new stack entry, and use its storage buffer to build the tree of evaluators.
+	FCameraRigEntry NewEntry;
+	const bool bInitialized = InitializeEntry(
+			NewEntry, 
+			Params.CameraRig,
+			Params.EvaluationContext,
+			EntryRootNode);
+	if (!bInitialized)
+	{
+		return;
+	}
+
+#if WITH_EDITOR
+	// Listen to changes to the packages inside which this camera rig is defined. We will hot-reload the
+	// camera node evaluators for this camera rig when we detect changes.
+	AddPackageListeners(NewEntry);
+#endif  // WITH_EDITOR
+
+	// Important: we need to move the new entry here because copying evaluator storage
+	// is disabled.
+	Entries.Add(MoveTemp(NewEntry));
+
+	if (OnCameraRigEventDelegate.IsBound())
+	{
+		BroadcastCameraRigEvent(EBlendStackCameraRigEventType::Pushed, Entries.Last(), UsedTransition);
+	}
+}
+
+void FTransientBlendStackCameraNodeEvaluator::PushVariantEntry(const FBlendStackCameraPushParams& PushParams, const UCameraRigTransition* Transition)
+{
+	const UCameraRigCameraNode* PrefabNode = Cast<const UCameraRigCameraNode>(PushParams.CameraRig->RootNode);
+	const UBlendCameraNode* Blend = Transition ? Transition->Blend : nullptr;
+
+	FCameraRigEntry& TopEntry = Entries.Top();
+	FCameraNodeEvaluatorBuilder Builder(TopEntry.EvaluatorStorage);
+	FCameraNodeEvaluatorBuildParams BuildParams(Builder);
+	TopEntry.RootEvaluator->MergeCameraRig(BuildParams, PrefabNode, Blend);
+
+	// Swap out the camera rig registered as "active" for this entry.
+#if WITH_EDITOR
+	RemoveListenedPackages(TopEntry);
+#endif
+	{
+		TopEntry.CameraRig = PushParams.CameraRig;
+	}
+#if WITH_EDITOR
+	AddPackageListeners(TopEntry);
+#endif
+}
+
+void FTransientBlendStackCameraNodeEvaluator::Freeze(const FBlendStackCameraFreezeParams& Params)
+{
+	for (FCameraRigEntry& Entry : Entries)
+	{
+		if (!Entry.bIsFrozen && 
+				Entry.CameraRig == Params.CameraRig &&
+				Entry.EvaluationContext == Params.EvaluationContext)
+		{
+			FreezeEntry(Entry);
+		}
+	}
+}
+
+void FTransientBlendStackCameraNodeEvaluator::FreezeAll(TSharedPtr<FCameraEvaluationContext> EvaluationContext)
+{
+	for (FCameraRigEntry& Entry : Entries)
+	{
+		if (!Entry.bIsFrozen && Entry.EvaluationContext == EvaluationContext)
+		{
+			FreezeEntry(Entry);
+		}
+	}
+}
+
+const UCameraRigTransition* FTransientBlendStackCameraNodeEvaluator::FindTransition(const FBlendStackCameraPushParams& Params) const
+{
 	const UBlendStackCameraNode* BlendStackNode = GetCameraNodeAs<UBlendStackCameraNode>();
 
@@ -712,7 +835,7 @@
 	return nullptr;
 }
 
-const UCameraRigTransition* FBlendStackCameraNodeEvaluator::FindTransition(
+const UCameraRigTransition* FTransientBlendStackCameraNodeEvaluator::FindTransition(
 			TArrayView<const TObjectPtr<UCameraRigTransition>> Transitions, 
 			const UCameraRigAsset* FromCameraRig, const UCameraAsset* FromCameraAsset, bool bFromFrozen,
 			const UCameraRigAsset* ToCameraRig, const UCameraAsset* ToCameraAsset) const
@@ -736,94 +859,63 @@
 	return nullptr;
 }
 
-void FBlendStackCameraNodeEvaluator::BroadcastCameraRigEvent(EBlendStackCameraRigEventType EventType, const FCameraRigEntry& Entry, const UCameraRigTransition* Transition) const
+void FPersistentBlendStackCameraNodeEvaluator::Insert(const FBlendStackCameraInsertParams& Params)
 {
-	FBlendStackCameraRigEvent Event;
-	Event.EventType = EventType;
-	Event.BlendStackEvaluator = this;
-	Event.CameraRigInfo = FCameraRigEvaluationInfo(
-			Entry.EvaluationContext.Pin(),
-			Entry.CameraRig,
-			&Entry.Result,
-			Entry.RootEvaluator);
-	Event.Transition = Transition;
-
-	OnCameraRigEventDelegate.Broadcast(Event);
-}
-
-void FBlendStackCameraNodeEvaluator::OnAddReferencedObjects(FReferenceCollector& Collector)
-{
-	for (FCameraRigEntry& Entry : Entries)
+	// See if we already have this camera rig and evaluation context in the stack.
+	for (const FCameraRigEntry& Entry : Entries)
 	{
-		Collector.AddReferencedObject(Entry.CameraRig);
-		Collector.AddReferencedObject(Entry.RootNode);
-	}
-}
-
-void FBlendStackCameraNodeEvaluator::OnSerialize(const FCameraNodeEvaluatorSerializeParams& Params, FArchive& Ar)
-{
-	if (Ar.IsSaving())
-	{
-		int32 NumEntries = Entries.Num();
-		Ar << NumEntries;
-	}
-	else if (Ar.IsLoading())
-	{
-		int32 LoadedNumEntries = 0;
-		Ar << LoadedNumEntries;
-
-		ensure(LoadedNumEntries == Entries.Num());
+		if (!Entry.bIsFrozen &&
+				Entry.CameraRig == Params.CameraRig &&
+				Entry.EvaluationContext == Params.EvaluationContext)
+		{
+			return;
+		}
 	}
 
-	for (FCameraRigEntry& Entry : Entries)
+	// TODO: add support for slot indices or something, to allow callers to specify a place in the stack.
+	UObject* Outer = const_cast<UObject*>((UObject*)GetCameraNode());
+	UBlendStackRootCameraNode* EntryRootNode = NewObject<UBlendStackRootCameraNode>(Outer, NAME_None);
 	{
-		Entry.Result.Serialize(Ar);
-		Ar << Entry.bIsFirstFrame;
-		Ar << Entry.bInputRunThisFrame;
-		Ar << Entry.bBlendRunThisFrame;
-		Ar << Entry.bIsFrozen;
-#if UE_GAMEPLAY_CAMERAS_TRACE
-		Ar << Entry.bLogWarnings;
-#endif  // UE_GAMEPLAY_CAMERAS_TRACE
+		EntryRootNode->RootNode = Params.CameraRig->RootNode;
+		// TODO: add support for blending in and out.
 	}
-}
 
+	FCameraRigEntry NewEntry;
+	const bool bInitialized = InitializeEntry(
+			NewEntry, 
+			Params.CameraRig,
+			Params.EvaluationContext,
+			EntryRootNode);
+	if (!bInitialized)
+	{
+		return;
+	}
+
 #if WITH_EDITOR
-
-void FBlendStackCameraNodeEvaluator::OnPostBuildAsset(const FGameplayCameraAssetBuildEvent& BuildEvent)
-{
-	for (FCameraRigEntry& Entry : Entries)
-	{
-		const bool bRebuildEntry = Entry.ListenedPackages.Contains(BuildEvent.AssetPackage);
-		if (bRebuildEntry)
-		{
-			Entry.EvaluatorStorage.DestroyEvaluatorTree();
-
-			// Re-assign the root node in case the camera rig's root was changed.
-			Entry.RootNode->RootNode = Entry.CameraRig->RootNode;
-
-			// Remove the blend on the root node, since we don't want the reloaded camera rig to re-blend-in
-			// for no good reason.
-			Entry.RootNode->Blend = NewObject<UPopBlendCameraNode>(Entry.RootNode, NAME_None);
-
-			// Rebuild the evaluator tree.
-			const bool bInitialized = InitializeEntry(
-					Entry,
-					Entry.CameraRig,
-					OwningEvaluator,
-					Entry.EvaluationContext.Pin(),
-					Entry.RootNode);
-			if (!bInitialized)
-			{
-				Entry.bIsFrozen = true;
-				continue;
-			}
-		}
-	}
-}
-
+	AddPackageListeners(NewEntry);
 #endif  // WITH_EDITOR
 
+	Entries.Add(MoveTemp(NewEntry));
+
+	if (OnCameraRigEventDelegate.IsBound())
+	{
+		BroadcastCameraRigEvent(EBlendStackCameraRigEventType::Pushed, Entries.Last(), nullptr);
+	}
+}
+
+void FPersistentBlendStackCameraNodeEvaluator::Remove(const FBlendStackCameraRemoveParams& Params)
+{
+	for (int32 Index = Entries.Num() - 1; Index >= 0; --Index)
+	{
+		FCameraRigEntry& Entry(Entries[Index]);
+		if (Entry.CameraRig == Params.CameraRig &&
+				Entry.EvaluationContext == Params.EvaluationContext)
+		{
+			PopEntry(Index);
+		}
+	}
+}
+
 #if UE_GAMEPLAY_CAMERAS_DEBUG
 
 void FBlendStackCameraNodeEvaluator::OnBuildDebugBlocks(const FCameraDebugBlockBuildParams& Params, FCameraDebugBlockBuilder& Builder)
@@ -842,11 +934,11 @@
 		FCameraDebugBlock& EntryDebugBlock = Builder.BuildDebugBlock<FCameraDebugBlock>();
 		StackDebugBlock.AddChild(&EntryDebugBlock);
 		{
-			if (Entry.RootEvaluator)
+			FCameraNodeEvaluator* BlendEvaluator = Entry.RootEvaluator ? Entry.RootEvaluator->GetBlendEvaluator() : nullptr;
+			if (BlendEvaluator)
 			{
 				Builder.StartParentDebugBlockOverride(EntryDebugBlock);
 				{
-					FCameraNodeEvaluator* BlendEvaluator = Entry.RootEvaluator->GetBlendEvaluator();
 					BlendEvaluator->BuildDebugBlocks(Params, Builder);
 				}
 				Builder.EndParentDebugBlockOverride();
@@ -878,16 +970,18 @@
 FBlendStackSummaryCameraDebugBlock::FBlendStackSummaryCameraDebugBlock(const FBlendStackCameraNodeEvaluator& InEvaluator)
 {
 	NumEntries = InEvaluator.Entries.Num();
+	BlendStackType = InEvaluator.GetCameraNodeAs<UBlendStackCameraNode>()->BlendStackType;
 }
 
 void FBlendStackSummaryCameraDebugBlock::OnDebugDraw(const FCameraDebugBlockDrawParams& Params, FCameraDebugRenderer& Renderer)
 {
-	Renderer.AddText(TEXT("%d entries"), NumEntries);
+	Renderer.AddText(TEXT("%d entries, %s stack"), NumEntries, *UEnum::GetValueAsString(BlendStackType));
 }
 
 void FBlendStackSummaryCameraDebugBlock::OnSerialize(FArchive& Ar)
 {
 	Ar << NumEntries;
+	Ar << BlendStackType;
 }
 
 UE_DEFINE_CAMERA_DEBUG_BLOCK(FBlendStackCameraDebugBlock);
diff --git a/Engine/Plugins/Cameras/GameplayCameras/Source/GameplayCameras/Private/Core/DefaultRootCameraNode.cpp b/Engine/Plugins/Cameras/GameplayCameras/Source/GameplayCameras/Private/Core/DefaultRootCameraNode.cpp
--- a/Engine/Plugins/Cameras/GameplayCameras/Source/GameplayCameras/Private/Core/DefaultRootCameraNode.cpp
+++ b/Engine/Plugins/Cameras/GameplayCameras/Source/GameplayCameras/Private/Core/DefaultRootCameraNode.cpp
@@ -17,11 +17,11 @@ namespace UE::Cameras::Private
 
 TObjectPtr<UBlendStackCameraNode> CreateBlendStack(
 		UObject* This, const FObjectInitializer& ObjectInit,
-		const FName& Name, bool bAutoPop = true, bool bBlendFirstCameraRig = false)
+		const FName& Name, ECameraBlendStackType BlendStackType, bool bBlendFirstCameraRig)
 {
 	TObjectPtr<UBlendStackCameraNode> NewBlendStack = ObjectInit.CreateDefaultSubobject<UBlendStackCameraNode>(
 			This, Name);
-	NewBlendStack->bAutoPop = bAutoPop;
+	NewBlendStack->BlendStackType = BlendStackType;
 	NewBlendStack->bBlendFirstCameraRig = bBlendFirstCameraRig;
 	return NewBlendStack;
 }
@@ -33,10 +33,10 @@ UDefaultRootCameraNode::UDefaultRootCameraNode(const FObjectInitializer& ObjectI
 {
 	using namespace UE::Cameras::Private;
 
-	BaseLayer = CreateBlendStack(this, ObjectInit, TEXT("BaseLayer"), false, true);
-	MainLayer = CreateBlendStack(this, ObjectInit, TEXT("MainLayer"));
-	GlobalLayer = CreateBlendStack(this, ObjectInit, TEXT("GlobalLayer"), false, true);
-	VisualLayer = CreateBlendStack(this, ObjectInit, TEXT("VisualLayer"), false, true);
+	BaseLayer = CreateBlendStack(this, ObjectInit, TEXT("BaseLayer"), ECameraBlendStackType::Persistent, true);
+	MainLayer = CreateBlendStack(this, ObjectInit, TEXT("MainLayer"), ECameraBlendStackType::Transient, false);
+	GlobalLayer = CreateBlendStack(this, ObjectInit, TEXT("GlobalLayer"), ECameraBlendStackType::Persistent, true);
+	VisualLayer = CreateBlendStack(this, ObjectInit, TEXT("VisualLayer"), ECameraBlendStackType::Persistent, true);
 }
 
 FCameraNodeEvaluatorPtr UDefaultRootCameraNode::OnBuildEvaluator(FCameraNodeEvaluatorBuilder& Builder) const
@@ -53,15 +53,16 @@ UE_DEFINE_CAMERA_NODE_EVALUATOR(FDefaultRootCameraNodeEvaluator)
 void FDefaultRootCameraNodeEvaluator::OnBuild(const FCameraNodeEvaluatorBuildParams& Params)
 {
 	const UDefaultRootCameraNode* Data = GetCameraNodeAs<UDefaultRootCameraNode>();
-	BaseLayer = BuildBlendStackEvaluator(Params, Data->BaseLayer);
-	MainLayer = BuildBlendStackEvaluator(Params, Data->MainLayer);
-	GlobalLayer = BuildBlendStackEvaluator(Params, Data->GlobalLayer);
-	VisualLayer = BuildBlendStackEvaluator(Params, Data->VisualLayer);
+	BaseLayer = BuildBlendStackEvaluator<FPersistentBlendStackCameraNodeEvaluator>(Params, Data->BaseLayer);
+	MainLayer = BuildBlendStackEvaluator<FTransientBlendStackCameraNodeEvaluator>(Params, Data->MainLayer);
+	GlobalLayer = BuildBlendStackEvaluator<FPersistentBlendStackCameraNodeEvaluator>(Params, Data->GlobalLayer);
+	VisualLayer = BuildBlendStackEvaluator<FPersistentBlendStackCameraNodeEvaluator>(Params, Data->VisualLayer);
 }
 
-FBlendStackCameraNodeEvaluator* FDefaultRootCameraNodeEvaluator::BuildBlendStackEvaluator(const FCameraNodeEvaluatorBuildParams& Params, UBlendStackCameraNode* BlendStackNode)
+template<typename EvaluatorType>
+EvaluatorType* FDefaultRootCameraNodeEvaluator::BuildBlendStackEvaluator(const FCameraNodeEvaluatorBuildParams& Params, UBlendStackCameraNode* BlendStackNode)
 {
-	FBlendStackCameraNodeEvaluator* BlendStackEvaluator = Params.BuildEvaluatorAs<FBlendStackCameraNodeEvaluator>(BlendStackNode);
+	EvaluatorType* BlendStackEvaluator = Params.BuildEvaluatorAs<EvaluatorType>(BlendStackNode);
 	BlendStackEvaluator->OnCameraRigEvent().AddRaw(this, &FDefaultRootCameraNodeEvaluator::OnBlendStackEvent);
 	return BlendStackEvaluator;
 }
@@ -81,13 +82,69 @@ void FDefaultRootCameraNodeEvaluator::OnRun(const FCameraNodeEvaluationParams& P
 
 void FDefaultRootCameraNodeEvaluator::OnActivateCameraRig(const FActivateCameraRigParams& Params)
 {
-	FBlendStackCameraNodeEvaluator* TargetStack = GetBlendStackEvaluator(Params.Layer);
-	if (ensure(TargetStack))
+	if (Params.Layer == ECameraRigLayer::Main)
 	{
 		FBlendStackCameraPushParams PushParams;
 		PushParams.EvaluationContext = Params.EvaluationContext;
 		PushParams.CameraRig = Params.CameraRig;
-		TargetStack->Push(PushParams);
+		MainLayer->Push(PushParams);
+	}
+	else
+	{
+		FPersistentBlendStackCameraNodeEvaluator* TargetLayer = nullptr;
+		switch (Params.Layer)
+		{
+			case ECameraRigLayer::Base:
+				TargetLayer = BaseLayer;
+				break;
+			case ECameraRigLayer::Global:
+				TargetLayer = GlobalLayer;
+				break;
+			case ECameraRigLayer::Visual:
+				TargetLayer = VisualLayer;
+				break;
+		}
+		if (ensure(TargetLayer))
+		{
+			FBlendStackCameraInsertParams InsertParams;
+			InsertParams.EvaluationContext = Params.EvaluationContext;
+			InsertParams.CameraRig = Params.CameraRig;
+			TargetLayer->Insert(InsertParams);
+		}
+	}
+}
+
+void FDefaultRootCameraNodeEvaluator::OnDeactivateCameraRig(const FDeactivateCameraRigParams& Params)
+{
+	if (Params.Layer == ECameraRigLayer::Main)
+	{
+		FBlendStackCameraFreezeParams FreezeParams;
+		FreezeParams.CameraRig = Params.CameraRig;
+		FreezeParams.EvaluationContext = Params.EvaluationContext;
+		MainLayer->Freeze(FreezeParams);
+	}
+	else
+	{
+		FPersistentBlendStackCameraNodeEvaluator* TargetLayer = nullptr;
+		switch (Params.Layer)
+		{
+			case ECameraRigLayer::Base:
+				TargetLayer = BaseLayer;
+				break;
+			case ECameraRigLayer::Global:
+				TargetLayer = GlobalLayer;
+				break;
+			case ECameraRigLayer::Visual:
+				TargetLayer = VisualLayer;
+				break;
+		}
+		if (ensure(TargetLayer))
+		{
+			FBlendStackCameraRemoveParams RemoveParams;
+			RemoveParams.EvaluationContext = Params.EvaluationContext;
+			RemoveParams.CameraRig = Params.CameraRig;
+			TargetLayer->Remove(RemoveParams);
+		}
 	}
 }
 
@@ -132,23 +189,6 @@ void FDefaultRootCameraNodeEvaluator::OnRunSingleCameraRig(const FSingleCameraRi
 	// Don't run the visual layer.
 }
 
-FBlendStackCameraNodeEvaluator* FDefaultRootCameraNodeEvaluator::GetBlendStackEvaluator(ECameraRigLayer Layer) const
-{
-	switch (Layer)
-	{
-		case ECameraRigLayer::Base:
-			return BaseLayer;
-		case ECameraRigLayer::Main:
-			return MainLayer;
-		case ECameraRigLayer::Global:
-			return GlobalLayer;
-		case ECameraRigLayer::Visual:
-			return VisualLayer;
-		default:
-			return nullptr;
-	}
-}
-
 void FDefaultRootCameraNodeEvaluator::OnBlendStackEvent(const FBlendStackCameraRigEvent& InEvent)
 {
 	if (InEvent.EventType == EBlendStackCameraRigEventType::Pushed ||
diff --git a/Engine/Plugins/Cameras/GameplayCameras/Source/GameplayCameras/Private/Core/RootCameraNode.cpp b/Engine/Plugins/Cameras/GameplayCameras/Source/GameplayCameras/Private/Core/RootCameraNode.cpp
--- a/Engine/Plugins/Cameras/GameplayCameras/Source/GameplayCameras/Private/Core/RootCameraNode.cpp
+++ b/Engine/Plugins/Cameras/GameplayCameras/Source/GameplayCameras/Private/Core/RootCameraNode.cpp
@@ -22,6 +22,11 @@ void FRootCameraNodeEvaluator::ActivateCameraRig(const FActivateCameraRigParams&
 	OnActivateCameraRig(Params);
 }
 
+void FRootCameraNodeEvaluator::DeactivateCameraRig(const FDeactivateCameraRigParams& Params)
+{
+	OnDeactivateCameraRig(Params);
+}
+
 void FRootCameraNodeEvaluator::BuildSingleCameraRigHierarchy(const FSingleCameraRigHierarchyBuildParams& Params, FCameraNodeEvaluatorHierarchy& OutHierarchy)
 {
 	OnBuildSingleCameraRigHierarchy(Params, OutHierarchy);
diff --git a/Engine/Plugins/Cameras/GameplayCameras/Source/GameplayCameras/Public/Core/BlendStackCameraNode.h b/Engine/Plugins/Cameras/GameplayCameras/Source/GameplayCameras/Public/Core/BlendStackCameraNode.h
--- a/Engine/Plugins/Cameras/GameplayCameras/Source/GameplayCameras/Public/Core/BlendStackCameraNode.h
+++ b/Engine/Plugins/Cameras/GameplayCameras/Source/GameplayCameras/Public/Core/BlendStackCameraNode.h
@@ -37,6 +37,20 @@ class IGameplayCamerasLiveEditManager;
 
 }  // namespace UE::Cameras
 
+UENUM()
+enum class ECameraBlendStackType
+{
+	/**
+	 * Camera rigs in a transient blend stack get automatically popped out of the stack 
+	 * when another rig has reached 100% blend above them.
+	 */
+	Transient,
+	/**
+	 * Camera rigs in a persistent blend stack stay in the stack until explicitly removed.
+	 */
+	Persistent
+};
+
 /**
  * A blend stack implemented as a camera node.
  */
@@ -53,11 +67,10 @@ protected:
 public:
 
 	/** 
-	 * Whether to automatically pop camera rigs out of the stack when another rig
-	 * has reached 100% blend above them.
+	 * The type of blend stack this should run as.
 	 */
 	UPROPERTY()
-	bool bAutoPop = true;
+	ECameraBlendStackType BlendStackType = ECameraBlendStackType::Transient;
 
 	/**
 	 * Whether to blend-in the first camera rig when the stack is previously empty.
@@ -69,21 +82,6 @@ public:
 namespace UE::Cameras
 {
 
-/**
- * Parameter structure for pushing a camera rig onto a blend stack.
- */
-struct FBlendStackCameraPushParams
-{
-	/** The evaluator currently running.*/
-	FCameraSystemEvaluator* Evaluator = nullptr;
-
-	/** The evaluation context within which a camera rig's node tree should run. */
-	TSharedPtr<const FCameraEvaluationContext> EvaluationContext;
-
-	/** The source camera rig asset to instantiate and push on the blend stack. */
-	TObjectPtr<const UCameraRigAsset> CameraRig;
-};
-
 DECLARE_MULTICAST_DELEGATE_OneParam(FOnBlendStackCameraRigEvent, const FBlendStackCameraRigEvent&);
 
 /**
@@ -101,9 +99,6 @@ public:
 
 	~FBlendStackCameraNodeEvaluator();
 
-	/** Push a new camera rig onto the blend stack. */
-	void Push(const FBlendStackCameraPushParams& Params);
-
 	/** Returns information about the top (active) camera rig, if any. */
 	FCameraRigEvaluationInfo GetActiveCameraRigEvaluationInfo() const;
 
@@ -138,28 +133,16 @@ protected:
 
 	struct FCameraRigEntry;
 
-	// Utility functions for finding an appropriate transition.
-	const UCameraRigTransition* FindTransition(const FBlendStackCameraPushParams& Params) const;
-	const UCameraRigTransition* FindTransition(
-			TArrayView<const TObjectPtr<UCameraRigTransition>> Transitions, 
-			const UCameraRigAsset* FromCameraRig, const UCameraAsset* FromCameraAsset, bool bFromFrozen,
-			const UCameraRigAsset* ToCameraRig, const UCameraAsset* ToCameraAsset) const;
-
-	void PopEntries(int32 FirstIndexToKeep);
-
 	bool InitializeEntry(
 		FCameraRigEntry& NewEntry, 
 		const UCameraRigAsset* CameraRig,
-		FCameraSystemEvaluator* Evaluator,
 		TSharedPtr<const FCameraEvaluationContext> EvaluationContext,
 		UBlendStackRootCameraNode* EntryRootNode);
 
-	void PushVariantEntry(const FBlendStackCameraPushParams& Params, const UCameraRigTransition* Transition);
-	void PushNewEntry(const FBlendStackCameraPushParams& Params, const UCameraRigTransition* Transition);
-
 	void FreezeEntry(FCameraRigEntry& Entry);
 
-	void GatherEntryParameterEvaluators(FCameraNodeEvaluator* RootEvaluator, TArray<FCameraNodeEvaluator*>& OutParameterEvaluators);
+	void PopEntry(int32 EntryIndex);
+	void PopEntries(int32 FirstIndexToKeep);
 
 	void BroadcastCameraRigEvent(EBlendStackCameraRigEventType EventType, const FCameraRigEntry& Entry, const UCameraRigTransition* Transition = nullptr) const;
 
@@ -224,6 +207,110 @@ protected:
 #endif  // UE_GAMEPLAY_CAMERAS_DEBUG
 };
 
+/**
+ * Parameter structure for pushing a camera rig onto a transient blend stack.
+ */
+struct FBlendStackCameraPushParams
+{
+	/** The evaluation context within which a camera rig's node tree should run. */
+	TSharedPtr<const FCameraEvaluationContext> EvaluationContext;
+
+	/** The source camera rig asset to instantiate and push on the blend stack. */
+	TObjectPtr<const UCameraRigAsset> CameraRig;
+};
+
+/**
+ * Parameter structure for freezing a camera rig inside a transient blend stack.
+ */
+struct FBlendStackCameraFreezeParams
+{
+	/** The evaluation context within which a camera rig's node tree is running. */
+	TSharedPtr<const FCameraEvaluationContext> EvaluationContext;
+
+	/** The source camera rig asset that is running. */
+	TObjectPtr<const UCameraRigAsset> CameraRig;
+};
+
+/**
+ * Evaluator for a transient blend stack.
+ */
+class FTransientBlendStackCameraNodeEvaluator 
+	: public FBlendStackCameraNodeEvaluator
+{
+
+	UE_DECLARE_CAMERA_NODE_EVALUATOR_EX(GAMEPLAYCAMERAS_API, FTransientBlendStackCameraNodeEvaluator, FBlendStackCameraNodeEvaluator)
+
+public:
+
+	/** Push a new camera rig onto the blend stack. */
+	void Push(const FBlendStackCameraPushParams& Params);
+
+	/** Freeze a camera rig. */
+	void Freeze(const FBlendStackCameraFreezeParams& Params);
+
+	/** Freeze all camera rigs that belong to a given evaluation context. */
+	void FreezeAll(TSharedPtr<FCameraEvaluationContext> EvaluationContext);
+
+private:
+
+	// Utility functions for finding an appropriate transition.
+	const UCameraRigTransition* FindTransition(const FBlendStackCameraPushParams& Params) const;
+	const UCameraRigTransition* FindTransition(
+			TArrayView<const TObjectPtr<UCameraRigTransition>> Transitions, 
+			const UCameraRigAsset* FromCameraRig, const UCameraAsset* FromCameraAsset, bool bFromFrozen,
+			const UCameraRigAsset* ToCameraRig, const UCameraAsset* ToCameraAsset) const;
+
+	void PushVariantEntry(const FBlendStackCameraPushParams& Params, const UCameraRigTransition* Transition);
+	void PushNewEntry(const FBlendStackCameraPushParams& Params, const UCameraRigTransition* Transition);
+
+};
+
+/**
+ * Parameter structure for inserting a camera rig into a persistent blend stack.
+ */
+struct FBlendStackCameraInsertParams
+{
+	/** The evaluation context within which a camera rig's node tree should run. */
+	TSharedPtr<const FCameraEvaluationContext> EvaluationContext;
+
+	/** The source camera rig asset to instantiate and push on the blend stack. */
+	TObjectPtr<const UCameraRigAsset> CameraRig;
+};
+
+/**
+ * Parameter structure for removing a camera rig from a persistent blend stack.
+ */
+struct FBlendStackCameraRemoveParams
+{
+	/** The evaluation context within which a camera rig's node tree should run. */
+	TSharedPtr<const FCameraEvaluationContext> EvaluationContext;
+
+	/** The source camera rig asset to instantiate and push on the blend stack. */
+	TObjectPtr<const UCameraRigAsset> CameraRig;
+};
+
+/**
+ * Evaluator for a persistent blend stack.
+ */
+class FPersistentBlendStackCameraNodeEvaluator 
+	: public FBlendStackCameraNodeEvaluator
+{
+	UE_DECLARE_CAMERA_NODE_EVALUATOR_EX(GAMEPLAYCAMERAS_API, FPersistentBlendStackCameraNodeEvaluator, FBlendStackCameraNodeEvaluator)
+
+public:
+
+	/** Insert a new camera rig onto the blend stack. */
+	void Insert(const FBlendStackCameraInsertParams& Params);
+
+	/** Remove an existing camera rig from the blend stack. */
+	void Remove(const FBlendStackCameraRemoveParams& Params);
+
+private:
+
+	// Utility functions for finding an appropriate transition.
+	const UCameraRigTransition* FindTransition(const FBlendStackCameraPushParams& Params) const;
+};
+
 #if UE_GAMEPLAY_CAMERAS_DEBUG
 
 class FBlendStackSummaryCameraDebugBlock : public FCameraDebugBlock
@@ -243,6 +330,7 @@ protected:
 private:
 
 	int32 NumEntries;
+	ECameraBlendStackType BlendStackType;
 };
 
 class FBlendStackCameraDebugBlock : public FCameraDebugBlock
diff --git a/Engine/Plugins/Cameras/GameplayCameras/Source/GameplayCameras/Public/Core/DefaultRootCameraNode.h b/Engine/Plugins/Cameras/GameplayCameras/Source/GameplayCameras/Public/Core/DefaultRootCameraNode.h
--- a/Engine/Plugins/Cameras/GameplayCameras/Source/GameplayCameras/Public/Core/DefaultRootCameraNode.h
+++ b/Engine/Plugins/Cameras/GameplayCameras/Source/GameplayCameras/Public/Core/DefaultRootCameraNode.h
@@ -43,7 +43,8 @@ public:
 namespace UE::Cameras
 {
 
-class FBlendStackCameraNodeEvaluator;
+class FPersistentBlendStackCameraNodeEvaluator;
+class FTransientBlendStackCameraNodeEvaluator;
 
 /**
  * Evaluator for the default root camera node.
@@ -65,23 +66,23 @@ protected:
 
 	// FRootCameraNodeEvaluator interface.
 	virtual void OnActivateCameraRig(const FActivateCameraRigParams& Params) override;
+	virtual void OnDeactivateCameraRig(const FDeactivateCameraRigParams& Params) override;
 	virtual void OnBuildSingleCameraRigHierarchy(const FSingleCameraRigHierarchyBuildParams& Params, FCameraNodeEvaluatorHierarchy& OutHierarchy) override;
 	virtual void OnRunSingleCameraRig(const FSingleCameraRigEvaluationParams& Params, FCameraNodeEvaluationResult& OutResult) override;
 
 private:
 
-	FBlendStackCameraNodeEvaluator* BuildBlendStackEvaluator(const FCameraNodeEvaluatorBuildParams& Params, UBlendStackCameraNode* BlendStackNode);
-
-	FBlendStackCameraNodeEvaluator* GetBlendStackEvaluator(ECameraRigLayer Layer) const;
+	template<typename EvaluatorType>
+	EvaluatorType* BuildBlendStackEvaluator(const FCameraNodeEvaluatorBuildParams& Params, UBlendStackCameraNode* BlendStackNode);
 
 	void OnBlendStackEvent(const FBlendStackCameraRigEvent& InEvent);
 
 private:
 
-	FBlendStackCameraNodeEvaluator* BaseLayer;
-	FBlendStackCameraNodeEvaluator* MainLayer;
-	FBlendStackCameraNodeEvaluator* GlobalLayer;
-	FBlendStackCameraNodeEvaluator* VisualLayer;
+	FPersistentBlendStackCameraNodeEvaluator* BaseLayer;
+	FTransientBlendStackCameraNodeEvaluator* MainLayer;
+	FPersistentBlendStackCameraNodeEvaluator* GlobalLayer;
+	FPersistentBlendStackCameraNodeEvaluator* VisualLayer;
 };
 
 }  // namespace UE::Cameras
diff --git a/Engine/Plugins/Cameras/GameplayCameras/Source/GameplayCameras/Public/Core/RootCameraNode.h b/Engine/Plugins/Cameras/GameplayCameras/Source/GameplayCameras/Public/Core/RootCameraNode.h
--- a/Engine/Plugins/Cameras/GameplayCameras/Source/GameplayCameras/Public/Core/RootCameraNode.h
+++ b/Engine/Plugins/Cameras/GameplayCameras/Source/GameplayCameras/Public/Core/RootCameraNode.h
@@ -55,6 +55,21 @@ struct FActivateCameraRigParams
 	ECameraRigLayer Layer = ECameraRigLayer::Main;
 };
 
+/**
+ * Parameter structure for deaactivating a running camera rig.
+ */
+struct FDeactivateCameraRigParams
+{
+	/** The evaluation context in which the camera rig runs. */
+	TSharedPtr<const FCameraEvaluationContext> EvaluationContext;
+
+	/** The source camera rig asset that was instantiated. */
+	TObjectPtr<const UCameraRigAsset> CameraRig;
+
+	/** The evaluation layer on which the camera rig is running. */
+	ECameraRigLayer Layer = ECameraRigLayer::Main;
+};
+
 /**
  * Parameter structure for building a single camera rig hierarchy.
  */
@@ -91,6 +106,9 @@ public:
 	/** Activates a camera rig. */
 	void ActivateCameraRig(const FActivateCameraRigParams& Params);
 
+	/** Deactivates a camera rig. */
+	void DeactivateCameraRig(const FDeactivateCameraRigParams& Params);
+
 	/**
 	 * Builds the hierarchy of the system for a given single camera rig.
 	 * This is expected to return the nodes of all the layers, except for the main layer which
@@ -118,6 +136,9 @@ protected:
 
 	/** Activates a camera rig. */
 	virtual void OnActivateCameraRig(const FActivateCameraRigParams& Params) {}
+	
+	/** Deactivates a camera rig. */
+	virtual void OnDeactivateCameraRig(const FDeactivateCameraRigParams& Params) {}
 
 	/* Builds the hierarchy of the system for a given single camera rig. */
 	virtual void OnBuildSingleCameraRigHierarchy(const FSingleCameraRigHierarchyBuildParams& Params, FCameraNodeEvaluatorHierarchy& OutHierarchy) {}
