Cameras: move the player controller evaluation context component into its own files.

diff --git a/Engine/Plugins/Cameras/GameplayCameras/Source/GameplayCameras/Private/GameFramework/ActivateCameraRigFunctions.cpp b/Engine/Plugins/Cameras/GameplayCameras/Source/GameplayCameras/Private/GameFramework/ActivateCameraRigFunctions.cpp
--- a/Engine/Plugins/Cameras/GameplayCameras/Source/GameplayCameras/Private/GameFramework/ActivateCameraRigFunctions.cpp
+++ b/Engine/Plugins/Cameras/GameplayCameras/Source/GameplayCameras/Private/GameFramework/ActivateCameraRigFunctions.cpp
@@ -2,18 +2,9 @@
 
 #include "GameFramework/ActivateCameraRigFunctions.h"
 
-#include "Camera/PlayerCameraManager.h"
-#include "Core/CameraRigAsset.h"
 #include "Core/RootCameraNode.h"
-#include "Core/CameraSystemEvaluator.h"
-#include "Core/CameraEvaluationContext.h"
-#include "Core/CameraEvaluationContextStack.h"
-#include "GameFramework/GameplayCameraSystemActor.h"
-#include "GameFramework/GameplayCameraSystemComponent.h"
-#include "GameFramework/GameplayCameraSystemHost.h"
-#include "GameFramework/PlayerController.h"
+#include "GameFramework/ControllerGameplayCameraEvaluationComponent.h"
 #include "GameplayCameras.h"
-#include "Templates/SharedPointer.h"
 
 #include UE_INLINE_GENERATED_CPP_BY_NAME(ActivateCameraRigFunctions)
 
@@ -42,111 +33,10 @@ void UActivateCameraRigFunctions::ActivateCameraRig(UObject* WorldContextObject,
 		return;
 	}
 
-	// Register our evaluation component on the given player controller, if it's not there already.
-	UControllerGameplayCameraEvaluationComponent* CameraEvaluationComponent = PlayerController->FindComponentByClass<UControllerGameplayCameraEvaluationComponent>();
-	if (!CameraEvaluationComponent)
+	UControllerGameplayCameraEvaluationComponent* CameraEvaluationComponent = UControllerGameplayCameraEvaluationComponent::FindOrAddComponent(PlayerController);
+	if (ensure(CameraEvaluationComponent))
 	{
-		CameraEvaluationComponent = NewObject<UControllerGameplayCameraEvaluationComponent>(
-				PlayerController, TEXT("ControllerGameplayCameraEvaluationComponent"), RF_Transient);
-		CameraEvaluationComponent->RegisterComponent();
-	}
-
-	// Activate the camera rig.
-	CameraEvaluationComponent->ActivateCameraRig(CameraRig, EvaluationLayer);
-}
-
-UControllerGameplayCameraEvaluationComponent::UControllerGameplayCameraEvaluationComponent(const FObjectInitializer& ObjectInitializer)
-	: Super(ObjectInitializer)
-{
-	bAutoActivate = true;
-}
-
-void UControllerGameplayCameraEvaluationComponent::ActivateCameraRig(UCameraRigAsset* CameraRig, ECameraRigLayer EvaluationLayer)
-{
-	FCameraRigInfo NewCameraRigInfo;
-	NewCameraRigInfo.CameraRig = CameraRig;
-	NewCameraRigInfo.EvaluationLayer = EvaluationLayer;
-	NewCameraRigInfo.bActivated = false;
-	CameraRigInfos.Add(NewCameraRigInfo);
-
-	if (IsActive())
-	{
-		ActivateCameraRigs();
-	}
-}
-
-void UControllerGameplayCameraEvaluationComponent::BeginPlay()
-{
-	Super::BeginPlay();
-
-	ActivateCameraRigs();
-}
-
-void UControllerGameplayCameraEvaluationComponent::EndPlay(const EEndPlayReason::Type EndPlayReason)
-{
-	CameraRigInfos.Reset();
-	EvaluationContext.Reset();
-
-	Super::EndPlay(EndPlayReason);
-}
-
-void UControllerGameplayCameraEvaluationComponent::ActivateCameraRigs()
-{
-	using namespace UE::Cameras;
-
-	EnsureCameraSystemHost();
-	if (!CameraSystemHost)
-	{
-		return;
-	}
-	
-	EnsureEvaluationContext();
-	if (!EvaluationContext)
-	{
-		return;
-	}
-
-	TSharedPtr<FCameraSystemEvaluator> SystemEvaluator = CameraSystemHost->GetCameraSystemEvaluator();
-	FRootCameraNodeEvaluator* RootNodeEvaluator = SystemEvaluator->GetRootNodeEvaluator();
-
-	for (FCameraRigInfo& CameraRigInfo : CameraRigInfos)
-	{
-		if (!CameraRigInfo.bActivated)
-		{
-			FActivateCameraRigParams Params;
-			Params.CameraRig = CameraRigInfo.CameraRig;
-			Params.EvaluationContext = EvaluationContext;
-			Params.Layer = CameraRigInfo.EvaluationLayer;
-			RootNodeEvaluator->ActivateCameraRig(Params);
-
-			CameraRigInfo.bActivated = true;
-		}
-	}
-}
-
-void UControllerGameplayCameraEvaluationComponent::EnsureEvaluationContext()
-{
-	using namespace UE::Cameras;
-
-	if (!EvaluationContext.IsValid())
-	{
-		APlayerController* PlayerController = GetOwner<APlayerController>();
-
-		FCameraEvaluationContextInitializeParams InitParams;
-		InitParams.Owner = this;
-		InitParams.PlayerController = PlayerController;
-		EvaluationContext = MakeShared<FCameraEvaluationContext>(InitParams);
-		EvaluationContext->GetInitialResult().bIsValid = true;	
-	}
-}
-
-void UControllerGameplayCameraEvaluationComponent::EnsureCameraSystemHost()
-{
-	if (!CameraSystemHost)
-	{
-		APlayerController* PlayerController = GetOwner<APlayerController>();
-
-		CameraSystemHost = UGameplayCameraSystemHost::FindOrCreateHost(PlayerController);
+		CameraEvaluationComponent->ActivateCameraRig(CameraRig, EvaluationLayer);
 	}
 }
 
diff --git a/Engine/Plugins/Cameras/GameplayCameras/Source/GameplayCameras/Private/GameFramework/ControllerGameplayCameraEvaluationComponent.cpp b/Engine/Plugins/Cameras/GameplayCameras/Source/GameplayCameras/Private/GameFramework/ControllerGameplayCameraEvaluationComponent.cpp
new file mode 100644
--- /dev/null
+++ b/Engine/Plugins/Cameras/GameplayCameras/Source/GameplayCameras/Private/GameFramework/ControllerGameplayCameraEvaluationComponent.cpp
@@ -0,0 +1,143 @@
+// Copyright Epic Games, Inc. All Rights Reserved.
+
+#include "GameFramework/ControllerGameplayCameraEvaluationComponent.h"
+
+#include "Core/CameraEvaluationContext.h"
+#include "Core/CameraRigAsset.h"
+#include "Core/CameraSystemEvaluator.h"
+#include "Core/RootCameraNode.h"
+#include "GameFramework/GameplayCameraSystemHost.h"
+#include "GameFramework/PlayerController.h"
+
+#include UE_INLINE_GENERATED_CPP_BY_NAME(ControllerGameplayCameraEvaluationComponent)
+
+UControllerGameplayCameraEvaluationComponent::UControllerGameplayCameraEvaluationComponent(const FObjectInitializer& ObjectInitializer)
+	: Super(ObjectInitializer)
+{
+	bAutoActivate = true;
+}
+
+void UControllerGameplayCameraEvaluationComponent::ActivateCameraRig(UCameraRigAsset* CameraRig, ECameraRigLayer EvaluationLayer)
+{
+	FCameraRigInfo NewCameraRigInfo;
+	NewCameraRigInfo.CameraRig = CameraRig;
+	NewCameraRigInfo.EvaluationLayer = EvaluationLayer;
+	NewCameraRigInfo.bActivated = false;
+	CameraRigInfos.Add(NewCameraRigInfo);
+
+	if (IsActive())
+	{
+		ActivateCameraRigs();
+	}
+}
+
+void UControllerGameplayCameraEvaluationComponent::BeginPlay()
+{
+	Super::BeginPlay();
+
+	ActivateCameraRigs();
+}
+
+void UControllerGameplayCameraEvaluationComponent::EndPlay(const EEndPlayReason::Type EndPlayReason)
+{
+	CameraRigInfos.Reset();
+	EvaluationContext.Reset();
+
+	Super::EndPlay(EndPlayReason);
+}
+
+void UControllerGameplayCameraEvaluationComponent::ActivateCameraRigs()
+{
+	using namespace UE::Cameras;
+
+	EnsureCameraSystemHost();
+	if (!CameraSystemHost)
+	{
+		return;
+	}
+	
+	EnsureEvaluationContext();
+	if (!EvaluationContext)
+	{
+		return;
+	}
+
+	TSharedPtr<FCameraSystemEvaluator> SystemEvaluator = CameraSystemHost->GetCameraSystemEvaluator();
+	FRootCameraNodeEvaluator* RootNodeEvaluator = SystemEvaluator->GetRootNodeEvaluator();
+
+	for (FCameraRigInfo& CameraRigInfo : CameraRigInfos)
+	{
+		if (!CameraRigInfo.bActivated)
+		{
+			FActivateCameraRigParams Params;
+			Params.CameraRig = CameraRigInfo.CameraRig;
+			Params.EvaluationContext = EvaluationContext;
+			Params.Layer = CameraRigInfo.EvaluationLayer;
+			RootNodeEvaluator->ActivateCameraRig(Params);
+
+			CameraRigInfo.bActivated = true;
+		}
+	}
+}
+
+void UControllerGameplayCameraEvaluationComponent::EnsureEvaluationContext()
+{
+	using namespace UE::Cameras;
+
+	if (!EvaluationContext.IsValid())
+	{
+		APlayerController* PlayerController = GetOwner<APlayerController>();
+
+		FCameraEvaluationContextInitializeParams InitParams;
+		InitParams.Owner = this;
+		InitParams.PlayerController = PlayerController;
+		EvaluationContext = MakeShared<FCameraEvaluationContext>(InitParams);
+		EvaluationContext->GetInitialResult().bIsValid = true;	
+	}
+}
+
+void UControllerGameplayCameraEvaluationComponent::EnsureCameraSystemHost()
+{
+	if (!CameraSystemHost)
+	{
+		APlayerController* PlayerController = GetOwner<APlayerController>();
+
+		CameraSystemHost = UGameplayCameraSystemHost::FindOrCreateHost(PlayerController);
+	}
+}
+
+UControllerGameplayCameraEvaluationComponent* UControllerGameplayCameraEvaluationComponent::FindComponent(APlayerController* PlayerController)
+{
+	return PlayerController->FindComponentByClass<UControllerGameplayCameraEvaluationComponent>();
+}
+
+UControllerGameplayCameraEvaluationComponent* UControllerGameplayCameraEvaluationComponent::FindOrAddComponent(APlayerController* PlayerController)
+{
+	UControllerGameplayCameraEvaluationComponent* ControllerComponent = FindComponent(PlayerController);
+	if (!ControllerComponent)
+	{
+		ControllerComponent = NewObject<UControllerGameplayCameraEvaluationComponent>(
+				PlayerController, TEXT("ControllerGameplayCameraEvaluationComponent"), RF_Transient);
+		ControllerComponent->RegisterComponent();
+	}
+	return ControllerComponent;
+}
+
+TSharedPtr<UE::Cameras::FCameraEvaluationContext> UControllerGameplayCameraEvaluationComponent::FindEvaluationContext(APlayerController* PlayerController)
+{
+	if (UControllerGameplayCameraEvaluationComponent* ControllerComponent = FindComponent(PlayerController))
+	{
+		ControllerComponent->EnsureEvaluationContext();
+		return ControllerComponent->EvaluationContext;
+	}
+	return nullptr;
+}
+
+TSharedRef<UE::Cameras::FCameraEvaluationContext> UControllerGameplayCameraEvaluationComponent::FindOrAddEvaluationContext(APlayerController* PlayerController)
+{
+	UControllerGameplayCameraEvaluationComponent* ControllerComponent = FindOrAddComponent(PlayerController);
+	check(ControllerComponent);
+	ControllerComponent->EnsureEvaluationContext();
+	return ControllerComponent->EvaluationContext.ToSharedRef();
+}
+
diff --git a/Engine/Plugins/Cameras/GameplayCameras/Source/GameplayCameras/Public/GameFramework/ActivateCameraRigFunctions.h b/Engine/Plugins/Cameras/GameplayCameras/Source/GameplayCameras/Public/GameFramework/ActivateCameraRigFunctions.h
--- a/Engine/Plugins/Cameras/GameplayCameras/Source/GameplayCameras/Public/GameFramework/ActivateCameraRigFunctions.h
+++ b/Engine/Plugins/Cameras/GameplayCameras/Source/GameplayCameras/Public/GameFramework/ActivateCameraRigFunctions.h
@@ -2,26 +2,14 @@
 
 #pragma once
 
-#include "Components/ActorComponent.h"
-#include "Core/CameraEvaluationContext.h"
-#include "Core/RootCameraNode.h"
 #include "CoreTypes.h"
 #include "Kismet/BlueprintFunctionLibrary.h"
-#include "Templates/SharedPointerFwd.h"
-#include "UObject/ObjectMacros.h"
-#include "UObject/ObjectPtr.h"
 
 #include "ActivateCameraRigFunctions.generated.h"
 
 class APlayerController;
 class UCameraRigAsset;
-class UGameplayCameraSystemHost;
-
-namespace UE::Cameras
-{
-	class FCameraEvaluationContext;
-	class FCameraSystemEvaluator;
-}  // namespace UE::Cameras
+enum class ECameraRigLayer : uint8;
 
 /**
  * Blueprint functions for activating camera rigs in the base/global/visual layers.
@@ -37,65 +25,22 @@ class UActivateCameraRigFunctions : public UBlueprintFunctionLibrary
 
 public:
 
-	/** Activates the given camera rig in the base layer. */
+	/** Activates the given camera rig prefab in the base layer. */
 	UFUNCTION(BlueprintCallable, Category="Camera", meta=(WorldContext="WorldContextObject"))
 	static void ActivateBaseCameraRig(UObject* WorldContextObject, APlayerController* PlayerController, UCameraRigAsset* CameraRig);
 
-	/** Activates the given camera rig in the global layer. */
+	/** Activates the given camera rig prefab in the global layer. */
 	UFUNCTION(BlueprintCallable, Category="Camera", meta = (WorldContext = "WorldContextObject"))
 	static void ActivateGlobalCameraRig(UObject* WorldContextObject, APlayerController* PlayerController, UCameraRigAsset* CameraRig);
 
-	/** Activates the given camera rig in the visual layer. */
+	/** Activates the given camera rig prefab in the visual layer. */
 	UFUNCTION(BlueprintCallable, Category="Camera", meta = (WorldContext = "WorldContextObject"))
 	static void ActivateVisualCameraRig(UObject* WorldContextObject, APlayerController* PlayerController, UCameraRigAsset* CameraRig);
 
-	/** Activates the given camera rig in the given layer. */
-	UFUNCTION(BlueprintCallable, Category="Camera", meta = (WorldContext = "WorldContextObject"))
-	static void ActivateCameraRig(UObject* WorldContextObject, APlayerController* PlayerController, UCameraRigAsset* CameraRig, ECameraRigLayer EvaluationLayer);
-};
-
-/**
- * A component, attached to a player controller, that can run camera rigs activated from
- * a global place like the Blueprint functions inside UActivateCameraRigFunctions.
- */
-UCLASS(Hidden, MinimalAPI)
-class UControllerGameplayCameraEvaluationComponent : public UActorComponent
-{
-	GENERATED_BODY()
-
-public:
-
-	UControllerGameplayCameraEvaluationComponent(const FObjectInitializer& ObjectInitializer);
-
-	/** Activates a new camera rig. */
-	void ActivateCameraRig(UCameraRigAsset* CameraRig, ECameraRigLayer EvaluationLayer);
-
-public:
-
-	// UActorComponent interface.
-	virtual void BeginPlay() override;
-	virtual void EndPlay(const EEndPlayReason::Type EndPlayReason) override;
-
 private:
 
-	void ActivateCameraRigs();
-	void EnsureEvaluationContext();
-	void EnsureCameraSystemHost();
-
-private:
-
-	struct FCameraRigInfo
-	{
-		TObjectPtr<UCameraRigAsset> CameraRig;
-		ECameraRigLayer EvaluationLayer;
-		bool bActivated = false;
-	};
-
-	TArray<FCameraRigInfo> CameraRigInfos;
-
-	TSharedPtr<UE::Cameras::FCameraEvaluationContext> EvaluationContext;
-
-	UPROPERTY()
-	TObjectPtr<UGameplayCameraSystemHost> CameraSystemHost;
+	/** Activates the given camera rig in the given layer. Should not be used with Main layer. */
+	UFUNCTION(BlueprintCallable, Category="Camera", meta = (WorldContext = "WorldContextObject"))
+	static void ActivateCameraRig(UObject* WorldContextObject, APlayerController* PlayerController, UCameraRigAsset* CameraRig, ECameraRigLayer EvaluationLayer);
 };
 
diff --git a/Engine/Plugins/Cameras/GameplayCameras/Source/GameplayCameras/Public/GameFramework/ControllerGameplayCameraEvaluationComponent.h b/Engine/Plugins/Cameras/GameplayCameras/Source/GameplayCameras/Public/GameFramework/ControllerGameplayCameraEvaluationComponent.h
new file mode 100644
--- /dev/null
+++ b/Engine/Plugins/Cameras/GameplayCameras/Source/GameplayCameras/Public/GameFramework/ControllerGameplayCameraEvaluationComponent.h
@@ -0,0 +1,72 @@
+// Copyright Epic Games, Inc. All Rights Reserved.
+
+#pragma once
+
+#include "Components/ActorComponent.h"
+#include "Templates/SharedPointerFwd.h"
+#include "UObject/ObjectPtr.h"
+
+#include "ControllerGameplayCameraEvaluationComponent.generated.h"
+
+class UCameraRigAsset;
+class UGameplayCameraSystemHost;
+enum class ECameraRigLayer : uint8;
+
+namespace UE::Cameras
+{
+	class FCameraEvaluationContext;
+}
+
+/**
+ * A component, attached to a player controller, that can run camera rigs activated from
+ * a global place like the Blueprint functions inside UActivateCameraRigFunctions.
+ */
+UCLASS(Hidden, MinimalAPI)
+class UControllerGameplayCameraEvaluationComponent : public UActorComponent
+{
+	GENERATED_BODY()
+
+public:
+
+	UControllerGameplayCameraEvaluationComponent(const FObjectInitializer& ObjectInitializer);
+
+	/** Activates a new camera rig. */
+	void ActivateCameraRig(UCameraRigAsset* CameraRig, ECameraRigLayer EvaluationLayer);
+
+public:
+
+	static UControllerGameplayCameraEvaluationComponent* FindComponent(APlayerController* PlayerController);
+	static UControllerGameplayCameraEvaluationComponent* FindOrAddComponent(APlayerController* PlayerController);
+
+	static TSharedPtr<UE::Cameras::FCameraEvaluationContext> FindEvaluationContext(APlayerController* PlayerController);
+	static TSharedRef<UE::Cameras::FCameraEvaluationContext> FindOrAddEvaluationContext(APlayerController* PlayerController);
+
+public:
+
+	// UActorComponent interface.
+	virtual void BeginPlay() override;
+	virtual void EndPlay(const EEndPlayReason::Type EndPlayReason) override;
+
+private:
+
+	void ActivateCameraRigs();
+	void EnsureEvaluationContext();
+	void EnsureCameraSystemHost();
+
+private:
+
+	struct FCameraRigInfo
+	{
+		TObjectPtr<UCameraRigAsset> CameraRig;
+		ECameraRigLayer EvaluationLayer;
+		bool bActivated = false;
+	};
+
+	TArray<FCameraRigInfo> CameraRigInfos;
+
+	TSharedPtr<UE::Cameras::FCameraEvaluationContext> EvaluationContext;
+
+	UPROPERTY()
+	TObjectPtr<UGameplayCameraSystemHost> CameraSystemHost;
+};
+
