Cameras: Blueprint camera director now has endpoints for activation/deactivation, and can activate/deactivate base/global/visual camera rigs

diff --git a/Engine/Plugins/Cameras/GameplayCameras/Source/GameplayCameras/Private/Directors/BlueprintCameraDirector.cpp b/Engine/Plugins/Cameras/GameplayCameras/Source/GameplayCameras/Private/Directors/BlueprintCameraDirector.cpp
--- a/Engine/Plugins/Cameras/GameplayCameras/Source/GameplayCameras/Private/Directors/BlueprintCameraDirector.cpp
+++ b/Engine/Plugins/Cameras/GameplayCameras/Source/GameplayCameras/Private/Directors/BlueprintCameraDirector.cpp
@@ -5,11 +5,15 @@
 #include "Components/ActorComponent.h"
 #include "Core/CameraAsset.h"
 #include "Core/CameraBuildLog.h"
+#include "Core/CameraDirectorEvaluator.h"
 #include "Core/CameraEvaluationContext.h"
 #include "Core/CameraRigAsset.h"
 #include "Core/CameraRigProxyAsset.h"
 #include "Core/CameraRigProxyTable.h"
+#include "Core/CameraSystemEvaluator.h"
+#include "Core/RootCameraNode.h"
 #include "GameFramework/Actor.h"
+#include "GameFramework/ControllerGameplayCameraEvaluationComponent.h"
 #include "GameplayCameras.h"
 
 #include UE_INLINE_GENERATED_CPP_BY_NAME(BlueprintCameraDirector)
@@ -19,6 +23,12 @@
 namespace UE::Cameras
 {
 
+void FBlueprintCameraDirectorEvaluationResult::Reset()
+{
+	ActiveCameraRigProxies.Reset();
+	ActiveCameraRigs.Reset();
+}
+
 class FBlueprintCameraDirectorEvaluator : public FCameraDirectorEvaluator
 {
 	UE_DECLARE_CAMERA_DIRECTOR_EVALUATOR(GAMEPLAYCAMERAS_API, FBlueprintCameraDirectorEvaluator)
@@ -26,15 +36,27 @@ class FBlueprintCameraDirectorEvaluator : public FCameraDirectorEvaluator
 protected:
 
 	virtual void OnInitialize(const FCameraDirectorInitializeParams& Params) override;
+	virtual void OnActivate(const FCameraDirectorActivateParams& Params) override;
+	virtual void OnDeactivate(const FCameraDirectorDeactivateParams& Params) override;
 	virtual void OnRun(const FCameraDirectorEvaluationParams& Params, FCameraDirectorEvaluationResult& OutResult) override;
 	virtual void OnAddReferencedObjects(FReferenceCollector& Collector) override;
 
 private:
 
+	void ActivateTransientCameraRigs(
+			const FCameraDirectorEvaluationParams& Params, 
+			const FBlueprintCameraDirectorEvaluationResult& BlueprintResult, 
+			FCameraDirectorEvaluationResult& OutResult);
+	void ActivateDeactivePersistentCameraRigs(
+			TSharedPtr<FCameraEvaluationContext> EvaluationContext,
+			const FBlueprintCameraDirectorEvaluationResult& BlueprintResult);
+
 	const UCameraRigAsset* FindCameraRigByProxy(const UCameraRigProxyAsset* InProxy);
 
 private:
 
+	FCameraSystemEvaluator* OwningEvaluator = nullptr;
+
 	TObjectPtr<UBlueprintCameraDirectorEvaluator> EvaluatorBlueprint;
 };
 
@@ -65,70 +87,133 @@ void FBlueprintCameraDirectorEvaluator::OnInitialize(const FCameraDirectorInitia
 	}
 }
 
-void FBlueprintCameraDirectorEvaluator::OnRun(const FCameraDirectorEvaluationParams& Params, FCameraDirectorEvaluationResult& OutResult)
+void FBlueprintCameraDirectorEvaluator::OnActivate(const FCameraDirectorActivateParams& Params)
 {
+	OwningEvaluator = Params.Evaluator;
+
 	if (EvaluatorBlueprint)
 	{
-		FBlueprintCameraDirectorEvaluationParams BlueprintParams;
-		BlueprintParams.DeltaTime = Params.DeltaTime;
-		BlueprintParams.EvaluationContext = Params.OwnerContext;
-		if (Params.OwnerContext)
-		{
-			BlueprintParams.EvaluationContextOwner = Params.OwnerContext->GetOwner();
-		}
+		EvaluatorBlueprint->NativeActivateCameraDirector(Params);
+
+		const FBlueprintCameraDirectorEvaluationResult& BlueprintResult = EvaluatorBlueprint->GetEvaluationResult();
+		ActivateDeactivePersistentCameraRigs(Params.OwnerContext, BlueprintResult);
+	}
+	else
+	{
+		UE_LOG(LogCameraSystem, Error, TEXT("Can't activate Blueprint camera director, no Blueprint class was set!"));
+	}
+}
+
+void FBlueprintCameraDirectorEvaluator::OnDeactivate(const FCameraDirectorDeactivateParams& Params)
+{
+	if (EvaluatorBlueprint)
+	{
+		EvaluatorBlueprint->NativeDeactivateCameraDirector(Params);
+
+		const FBlueprintCameraDirectorEvaluationResult& BlueprintResult = EvaluatorBlueprint->GetEvaluationResult();
+		ActivateDeactivePersistentCameraRigs(Params.OwnerContext, BlueprintResult);
+	}
 
-		FBlueprintCameraDirectorEvaluationResult BlueprintResult;
+	OwningEvaluator = nullptr;
+}
 
-		EvaluatorBlueprint->NativeRunCameraDirector(BlueprintParams, BlueprintResult);
+void FBlueprintCameraDirectorEvaluator::OnRun(const FCameraDirectorEvaluationParams& Params, FCameraDirectorEvaluationResult& OutResult)
+{
+	if (EvaluatorBlueprint)
+	{
+		EvaluatorBlueprint->NativeRunCameraDirector(Params);
 
-		TArray<const UCameraRigAsset*, TInlineAllocator<2>> CameraRigs;
-		const UCameraAsset* CameraAsset = Params.OwnerContext->GetCameraAsset();
+		const FBlueprintCameraDirectorEvaluationResult& BlueprintResult = EvaluatorBlueprint->GetEvaluationResult();
+		ActivateTransientCameraRigs(Params, BlueprintResult, OutResult);
+		ActivateDeactivePersistentCameraRigs(Params.OwnerContext, BlueprintResult);
+	}
+}
 
-		// Gather camera rigs.
-		for (const UCameraRigAsset* ActiveCameraRig : BlueprintResult.ActiveCameraRigs)
+void FBlueprintCameraDirectorEvaluator::ActivateTransientCameraRigs(
+		const FCameraDirectorEvaluationParams& Params, 
+		const FBlueprintCameraDirectorEvaluationResult& BlueprintResult, 
+		FCameraDirectorEvaluationResult& OutResult)
+{
+	TArray<const UCameraRigAsset*, TInlineAllocator<2>> CameraRigs;
+	const UCameraAsset* CameraAsset = Params.OwnerContext->GetCameraAsset();
+
+	// Gather camera rigs.
+	for (const UCameraRigAsset* ActiveCameraRig : BlueprintResult.ActiveCameraRigs)
+	{
+		if (ActiveCameraRig)
 		{
-			if (ActiveCameraRig)
-			{
-				CameraRigs.Add(ActiveCameraRig);
-			}
-			else
-			{
-				UE_LOG(
-						LogCameraSystem, 
-						Error, 
-						TEXT("Null camera rig specified in camera director '%s'."),
-						*EvaluatorBlueprint->GetClass()->GetPathName());
-			}
+			CameraRigs.Add(ActiveCameraRig);
 		}
-
-		// Resolve camera rig proxies.
-		for (const UCameraRigProxyAsset* ActiveCameraRigProxy : BlueprintResult.ActiveCameraRigProxies)
+		else
 		{
-			const UCameraRigAsset* ActiveCameraRig = FindCameraRigByProxy(ActiveCameraRigProxy);
-			if (ActiveCameraRig)
-			{
-				CameraRigs.Add(ActiveCameraRig);
-			}
-			else
-			{
-				UE_LOG(
-						LogCameraSystem, 
-						Error, 
-						TEXT("No camera rig found mapped to proxy '%s' in camera '%s'."),
-						*ActiveCameraRigProxy->GetPathName(), *CameraAsset->GetPathName());
-			}
+			UE_LOG(
+					LogCameraSystem, 
+					Error, 
+					TEXT("Null camera rig specified in camera director '%s'."),
+					*EvaluatorBlueprint->GetClass()->GetPathName());
 		}
+	}
 
-		// The BP interface doesn't specify the evaluation context for the chosen camera rigs: we always automatically
-		// make them run in our own owner context.
-		for (const UCameraRigAsset* ActiveCameraRig : CameraRigs)
+	// Resolve camera rig proxies.
+	for (const UCameraRigProxyAsset* ActiveCameraRigProxy : BlueprintResult.ActiveCameraRigProxies)
+	{
+		const UCameraRigAsset* ActiveCameraRig = FindCameraRigByProxy(ActiveCameraRigProxy);
+		if (ActiveCameraRig)
+		{
+			CameraRigs.Add(ActiveCameraRig);
+		}
+		else
 		{
-			OutResult.Add(Params.OwnerContext, ActiveCameraRig);
+			UE_LOG(
+					LogCameraSystem, 
+					Error, 
+					TEXT("No camera rig found mapped to proxy '%s' in camera '%s'."),
+					*ActiveCameraRigProxy->GetPathName(), *CameraAsset->GetPathName());
 		}
 	}
-	else
+
+	// The BP interface doesn't specify the evaluation context for the chosen camera rigs: we always automatically
+	// make them run in our own owner context.
+	for (const UCameraRigAsset* ActiveCameraRig : CameraRigs)
 	{
-		UE_LOG(LogCameraSystem, Error, TEXT("Can't run Blueprint camera director, no Blueprint class was set!"));
+		OutResult.Add(Params.OwnerContext, ActiveCameraRig);
+	}
+}
+
+void FBlueprintCameraDirectorEvaluator::ActivateDeactivePersistentCameraRigs(
+		TSharedPtr<FCameraEvaluationContext> EvaluationContext,
+		const FBlueprintCameraDirectorEvaluationResult& BlueprintResult)
+{
+	if (!ensure(OwningEvaluator))
+	{
+		return;
+	}
+
+	FRootCameraNodeEvaluator* RootNodeEvaluator = OwningEvaluator->GetRootNodeEvaluator();
+
+	APlayerController* PlayerController = EvaluationContext->GetPlayerController();
+	TSharedPtr<const FCameraEvaluationContext> ControllerEvaluationContext;
+	if (ensure(PlayerController))
+	{
+		ControllerEvaluationContext = UControllerGameplayCameraEvaluationComponent::FindOrAddEvaluationContext(PlayerController);
+	}
+
+	for (const FBlueprintPersistentCameraRigInfo& CameraRigInfo : BlueprintResult.InactivePersistentCameraRigs)
+	{
+		FDeactivateCameraRigParams DeactivateParams;
+		DeactivateParams.EvaluationContext = ControllerEvaluationContext;
+		DeactivateParams.CameraRig = CameraRigInfo.CameraRig;
+		DeactivateParams.Layer = CameraRigInfo.Layer;
+		RootNodeEvaluator->DeactivateCameraRig(DeactivateParams);
+	}
+
+	for (const FBlueprintPersistentCameraRigInfo& CameraRigInfo : BlueprintResult.ActivePersistentCameraRigs)
+	{
+		FActivateCameraRigParams ActivateParams;
+		ActivateParams.EvaluationContext = ControllerEvaluationContext;
+		ActivateParams.CameraRig = CameraRigInfo.CameraRig;
+		ActivateParams.Layer = CameraRigInfo.Layer;
+		RootNodeEvaluator->ActivateCameraRig(ActivateParams);
 	}
 }
 
@@ -158,25 +243,40 @@ void FBlueprintCameraDirectorEvaluator::OnAddReferencedObjects(FReferenceCollect
 
 }  // namespace UE::Cameras
 
+void UBlueprintCameraDirectorEvaluator::ActivateBaseCameraRig(UCameraRigAsset* CameraRigPrefab)
+{
+	EvaluationResult.ActivePersistentCameraRigs.Add({ CameraRigPrefab, ECameraRigLayer::Base });
+}
+
+void UBlueprintCameraDirectorEvaluator::ActivateGlobalCameraRig(UCameraRigAsset* CameraRigPrefab)
+{
+	EvaluationResult.ActivePersistentCameraRigs.Add({ CameraRigPrefab, ECameraRigLayer::Global });
+}
+
+void UBlueprintCameraDirectorEvaluator::ActivateVisualCameraRig(UCameraRigAsset* CameraRigPrefab)
+{
+	EvaluationResult.ActivePersistentCameraRigs.Add({ CameraRigPrefab, ECameraRigLayer::Visual });
+}
+
 void UBlueprintCameraDirectorEvaluator::ActivateCameraRig(UCameraRigAsset* CameraRig)
 {
-	CurrentResult.ActiveCameraRigs.Add(CameraRig);
+	EvaluationResult.ActiveCameraRigs.Add(CameraRig);
 }
 
 void UBlueprintCameraDirectorEvaluator::ActivateCameraRigViaProxy(UCameraRigProxyAsset* CameraRigProxy)
 {
-	CurrentResult.ActiveCameraRigProxies.Add(CameraRigProxy);
+	EvaluationResult.ActiveCameraRigProxies.Add(CameraRigProxy);
 }
 
 AActor* UBlueprintCameraDirectorEvaluator::FindEvaluationContextOwnerActor(TSubclassOf<AActor> ActorClass) const
 {
-	if (CurrentContext)
+	if (EvaluationContext)
 	{
-		if (UActorComponent* ContextOwnerAsComponent = Cast<UActorComponent>(CurrentContext->GetOwner()))
+		if (UActorComponent* ContextOwnerAsComponent = Cast<UActorComponent>(EvaluationContext->GetOwner()))
 		{
 			return ContextOwnerAsComponent->GetOwner();
 		}
-		else if (AActor* ContextOwnerAsActor = Cast<AActor>(CurrentContext->GetOwner()))
+		else if (AActor* ContextOwnerAsActor = Cast<AActor>(EvaluationContext->GetOwner()))
 		{
 			return ContextOwnerAsActor;
 		}
@@ -196,9 +296,9 @@ AActor* UBlueprintCameraDirectorEvaluator::FindEvaluationContextOwnerActor(TSubc
 
 FBlueprintCameraPose UBlueprintCameraDirectorEvaluator::GetInitialContextCameraPose() const
 {
-	if (CurrentContext)
+	if (EvaluationContext)
 	{
-		return FBlueprintCameraPose::FromCameraPose(CurrentContext->GetInitialResult().CameraPose);
+		return FBlueprintCameraPose::FromCameraPose(EvaluationContext->GetInitialResult().CameraPose);
 	}
 	else
 	{
@@ -211,9 +311,9 @@ FBlueprintCameraPose UBlueprintCameraDirectorEvaluator::GetInitialContextCameraP
 
 FBlueprintCameraVariableTable UBlueprintCameraDirectorEvaluator::GetInitialContextVariableTable() const
 {
-	if (CurrentContext)
+	if (EvaluationContext)
 	{
-		return FBlueprintCameraVariableTable(&CurrentContext->GetInitialResult().VariableTable);
+		return FBlueprintCameraVariableTable(&EvaluationContext->GetInitialResult().VariableTable);
 	}
 	else
 	{
@@ -226,9 +326,9 @@ FBlueprintCameraVariableTable UBlueprintCameraDirectorEvaluator::GetInitialConte
 
 void UBlueprintCameraDirectorEvaluator::SetInitialContextCameraPose(const FBlueprintCameraPose& InCameraPose)
 {
-	if (CurrentContext)
+	if (EvaluationContext)
 	{
-		InCameraPose.ApplyTo(CurrentContext->GetInitialResult().CameraPose);
+		InCameraPose.ApplyTo(EvaluationContext->GetInitialResult().CameraPose);
 	}
 	else
 	{
@@ -238,16 +338,51 @@ void UBlueprintCameraDirectorEvaluator::SetInitialContextCameraPose(const FBluep
 	}
 }
 
-void UBlueprintCameraDirectorEvaluator::NativeRunCameraDirector(const FBlueprintCameraDirectorEvaluationParams& Params, FBlueprintCameraDirectorEvaluationResult& OutResult)
+void UBlueprintCameraDirectorEvaluator::NativeActivateCameraDirector(const UE::Cameras::FCameraDirectorActivateParams& Params)
 {
-	CurrentContext = Params.EvaluationContext;
-	CurrentResult = OutResult;
+	EvaluationContext = Params.OwnerContext;
+
+	EvaluationResult.Reset();
 	{
-		// Run the Blueprint logic.
-		RunCameraDirector(Params);
+		FBlueprintCameraDirectorActivateParams BlueprintParams;
+		if (Params.OwnerContext)
+		{
+			BlueprintParams.EvaluationContextOwner = Params.OwnerContext->GetOwner();
+		}
+
+		ActivateCameraDirector(BlueprintParams);
+	}
+}
+
+void UBlueprintCameraDirectorEvaluator::NativeDeactivateCameraDirector(const UE::Cameras::FCameraDirectorDeactivateParams& Params)
+{
+	EvaluationResult.Reset();
+	{
+		FBlueprintCameraDirectorDeactivateParams BlueprintParams;
+		if (Params.OwnerContext)
+		{
+			BlueprintParams.EvaluationContextOwner = Params.OwnerContext->GetOwner();
+		}
+
+		DeactivateCameraDirector(BlueprintParams);
+	}
+
+	EvaluationContext = nullptr;
+}
+
+void UBlueprintCameraDirectorEvaluator::NativeRunCameraDirector(const UE::Cameras::FCameraDirectorEvaluationParams& Params)
+{
+	EvaluationResult.Reset();
+	{
+		FBlueprintCameraDirectorEvaluationParams BlueprintParams;
+		BlueprintParams.DeltaTime = Params.DeltaTime;
+		if (Params.OwnerContext)
+		{
+			BlueprintParams.EvaluationContextOwner = Params.OwnerContext->GetOwner();
+		}
+
+		RunCameraDirector(BlueprintParams);
 	}
-	OutResult = CurrentResult;
-	CurrentContext.Reset();
 }
 
 FCameraDirectorEvaluatorPtr UBlueprintCameraDirector::OnBuildEvaluator(FCameraDirectorEvaluatorBuilder& Builder) const
diff --git a/Engine/Plugins/Cameras/GameplayCameras/Source/GameplayCameras/Public/Directors/BlueprintCameraDirector.h b/Engine/Plugins/Cameras/GameplayCameras/Source/GameplayCameras/Public/Directors/BlueprintCameraDirector.h
--- a/Engine/Plugins/Cameras/GameplayCameras/Source/GameplayCameras/Public/Directors/BlueprintCameraDirector.h
+++ b/Engine/Plugins/Cameras/GameplayCameras/Source/GameplayCameras/Public/Directors/BlueprintCameraDirector.h
@@ -10,44 +10,85 @@
 
 #include "BlueprintCameraDirector.generated.h"
 
+class UCameraRigAsset;
 class UCameraRigProxyAsset;
 class UCameraRigProxyTable;
+enum class ECameraRigLayer : uint8;
+
+namespace UE::Cameras
+{
+
+/** Information about a persitent camera rig to be activated or deactivated. */
+struct FBlueprintPersistentCameraRigInfo
+{
+	UCameraRigAsset* CameraRig;
+	ECameraRigLayer Layer;
+};
+
+/**
+ * The evaluation result for the Blueprint camera director evaluator.
+ */
+struct FBlueprintCameraDirectorEvaluationResult
+{
+	/** The list of camera rigs that should be active this frame. */
+	TArray<UCameraRigProxyAsset*> ActiveCameraRigProxies;
+
+	/** The list of camera rigs that should be active this frame. */
+	TArray<UCameraRigAsset*> ActiveCameraRigs;
+
+	/** The list of persistent camera rigs to activate. */
+	TArray<FBlueprintPersistentCameraRigInfo> ActivePersistentCameraRigs;
+
+	/** The list of persistent camera rigs to deactivate. */
+	TArray<FBlueprintPersistentCameraRigInfo> InactivePersistentCameraRigs;
+
+	/** Reset this result for a new evaluation. */
+	void Reset();
+};
+
+}  // namespace UE::Cameras
 
 /**
- * Parameter struct for the Blueprint camera director evaluator.
+ * Parameter struct for activating the Blueprint camera director evaluator.
  */
 USTRUCT(BlueprintType)
-struct FBlueprintCameraDirectorEvaluationParams
+struct FBlueprintCameraDirectorActivateParams
 {
 	GENERATED_BODY()
 
-	/** The elapsed time since the last evaluation. */
-	UPROPERTY(BlueprintReadWrite, Category="Evaluation")
-	float DeltaTime = 0.f;
-
 	/** The owner (if any) of the evaluation context we are running inside of. */
 	UPROPERTY(BlueprintReadWrite, Category="Evaluation")
 	TObjectPtr<UObject> EvaluationContextOwner;
+};
+
+/**
+ * Parameter struct for deactivating the Blueprint camera director evaluator.
+ */
+USTRUCT(BlueprintType)
+struct FBlueprintCameraDirectorDeactivateParams
+{
+	GENERATED_BODY()
 
-	/** The evaluation context we are running inside of. */
-	TSharedPtr<UE::Cameras::FCameraEvaluationContext> EvaluationContext;
+	/** The owner (if any) of the evaluation context we were running inside of. */
+	UPROPERTY(BlueprintReadWrite, Category="Evaluation")
+	TObjectPtr<UObject> EvaluationContextOwner;
 };
 
 /**
- * The evaluation result for the Blueprint camera director evaluator.
+ * Parameter struct for running the Blueprint camera director evaluator.
  */
 USTRUCT(BlueprintType)
-struct FBlueprintCameraDirectorEvaluationResult
+struct FBlueprintCameraDirectorEvaluationParams
 {
 	GENERATED_BODY()
 
-	/** The list of camera rigs that should be active this frame. */
+	/** The elapsed time since the last evaluation. */
 	UPROPERTY(BlueprintReadWrite, Category="Evaluation")
-	TArray<TObjectPtr<UCameraRigProxyAsset>> ActiveCameraRigProxies;
+	float DeltaTime = 0.f;
 
-	/** The list of camera rigs that should be active this frame. */
+	/** The owner (if any) of the evaluation context we are running inside of. */
 	UPROPERTY(BlueprintReadWrite, Category="Evaluation")
-	TArray<TObjectPtr<UCameraRigAsset>> ActiveCameraRigs;
+	TObjectPtr<UObject> EvaluationContextOwner;
 };
 
 /**
@@ -59,6 +100,20 @@ class UBlueprintCameraDirectorEvaluator : public UObject
 	GENERATED_BODY()
 
 public:
+
+	/**
+	 * Override this method in Blueprint to execute custom logic when this
+	 * camera director gets activated.
+	 */
+	UFUNCTION(BlueprintCallable, BlueprintImplementableEvent, Category="Activation")
+	void ActivateCameraDirector(const FBlueprintCameraDirectorActivateParams& Params);
+
+	/**
+	 * Override this method in Blueprint to execute custom logic when this
+	 * camera director gets deactivated.
+	 */
+	UFUNCTION(BlueprintCallable, BlueprintImplementableEvent, Category="Activation")
+	void DeactivateCameraDirector(const FBlueprintCameraDirectorDeactivateParams& Params);
 	
 	/**
 	 * Override this method in Blueprint to execute the custom logic that determines
@@ -67,6 +122,22 @@ public:
 	UFUNCTION(BlueprintCallable, BlueprintImplementableEvent, Category="Evaluation")
 	void RunCameraDirector(const FBlueprintCameraDirectorEvaluationParams& Params);
 
+public:
+
+	/** Activates the given camera rig prefab in the base layer. */
+	UFUNCTION(BlueprintCallable, Category="Activation")
+	void ActivateBaseCameraRig(UCameraRigAsset* CameraRigPrefab);
+
+	/** Activates the given camera rig prefab in the global layer. */
+	UFUNCTION(BlueprintCallable, Category="Activation")
+	void ActivateGlobalCameraRig(UCameraRigAsset* CameraRigPrefab);
+
+	/** Activates the given camera rig prefab in the visual layer. */
+	UFUNCTION(BlueprintCallable, Category="Activation")
+	void ActivateVisualCameraRig(UCameraRigAsset* CameraRigPrefab);
+
+public:
+
 	/** Specifies a camera rig to be active this frame. */
 	UFUNCTION(BlueprintCallable, Category="Evaluation")
 	void ActivateCameraRig(
@@ -112,20 +183,28 @@ public:
 
 public:
 
-	/** Native wrapper for RunCameraDirector. */
-	void NativeRunCameraDirector(
-			const FBlueprintCameraDirectorEvaluationParams& Params,
-			FBlueprintCameraDirectorEvaluationResult& OutResult);
+	using FCameraEvaluationContext = UE::Cameras::FCameraEvaluationContext;
+	using FBlueprintCameraDirectorEvaluationResult = UE::Cameras::FBlueprintCameraDirectorEvaluationResult;
 
-protected:
+	/** Native wrapper for ActivateCameraDirector. */
+	void NativeActivateCameraDirector(const UE::Cameras::FCameraDirectorActivateParams& Params);
 
-	/** The current camera director evaluation result. */
-	UPROPERTY(BlueprintReadWrite, Category="Evaluation")
-	FBlueprintCameraDirectorEvaluationResult CurrentResult;
+	/** Native wrapper for DeactivateCameraDirector. */
+	void NativeDeactivateCameraDirector(const UE::Cameras::FCameraDirectorDeactivateParams& Params);
+
+	/** Native wrapper for RunCameraDirector. */
+	void NativeRunCameraDirector(const UE::Cameras::FCameraDirectorEvaluationParams& Params);
+
+	/** Get the last result for this camera director. */
+	const FBlueprintCameraDirectorEvaluationResult& GetEvaluationResult() const { return EvaluationResult; }
 
 private:
 
-	TSharedPtr<UE::Cameras::FCameraEvaluationContext> CurrentContext;
+	/** The current camera director evaluation result. */
+	FBlueprintCameraDirectorEvaluationResult EvaluationResult;
+
+	/** The current evaluation context. */
+	TSharedPtr<FCameraEvaluationContext> EvaluationContext;
 };
 
 /**
