Cameras: fix issues with the graph pin factory

- By default, let camera rig properties use the normal asset picker, for picking prefabs.
- Add more metadata to specify when to make a custom picker.
- Remove duplicated bit of code.

diff --git a/Engine/Plugins/Cameras/GameplayCameras/Source/GameplayCamerasEditor/Private/Editors/GameplayCamerasGraphPanelPinFactory.cpp b/Engine/Plugins/Cameras/GameplayCameras/Source/GameplayCamerasEditor/Private/Editors/GameplayCamerasGraphPanelPinFactory.cpp
--- a/Engine/Plugins/Cameras/GameplayCameras/Source/GameplayCamerasEditor/Private/Editors/GameplayCamerasGraphPanelPinFactory.cpp
+++ b/Engine/Plugins/Cameras/GameplayCameras/Source/GameplayCamerasEditor/Private/Editors/GameplayCamerasGraphPanelPinFactory.cpp
@@ -28,15 +28,6 @@ TSharedPtr<SGraphPin> FGameplayCamerasGraphPanelPinFactory::CreatePin(UEdGraphPi
 		}
 	}
 
-	if (Pin->PinType.PinCategory == UEdGraphSchema_K2::PC_Object && 
-			Pin->PinType.PinSubCategoryObject == UCameraRigAsset::StaticClass())
-	{
-		if (TSharedPtr<SGraphPin> PinWidget = CreateCameraRigPickerPin(Pin))
-		{
-			return PinWidget;
-		}
-	}
-
 	const FEdGraphPinType& PinType = Pin->PinType;
 	const UClass* PinPropertyClass = Cast<const UClass>(PinType.PinSubCategoryObject);
 	if (PinType.PinCategory == UEdGraphSchema_K2::PC_Object && PinPropertyClass)
@@ -84,7 +75,25 @@ TSharedPtr<SGraphPin> FGameplayCamerasGraphPanelPinFactory::CreateFunctionParame
 
 TSharedPtr<SGraphPin> FGameplayCamerasGraphPanelPinFactory::CreateCameraRigPickerPin(UEdGraphPin* Pin) const
 {
-	return SNew(SCameraRigNameGraphPin, Pin);
+	UEdGraphNode* OwningNode = Pin->GetOwningNode();
+	if (!OwningNode)
+	{
+		return nullptr;
+	}
+	const UClass* OwningNodeClass = OwningNode->GetClass();
+	const FString& UseCameraRigPickerForPinsMetaData = OwningNodeClass->GetMetaData(TEXT("UseCameraRigPickerForPins"));
+	if (UseCameraRigPickerForPinsMetaData.IsEmpty())
+	{
+		return nullptr;
+	}
+
+	TArray<FString> CameraRigPickerPinNames;
+	UseCameraRigPickerForPinsMetaData.ParseIntoArray(CameraRigPickerPinNames, TEXT(","));
+	if (CameraRigPickerPinNames.Contains(Pin->GetName()))
+	{
+		return SNew(SCameraRigNameGraphPin, Pin);
+	}
+	return nullptr;
 }
 
 TSharedPtr<SGraphPin> FGameplayCamerasGraphPanelPinFactory::CreateCameraVariablePickerPin(UEdGraphPin* Pin) const
diff --git a/Engine/Plugins/Cameras/GameplayCameras/Source/GameplayCamerasUncookedOnly/Private/BlueprintGraph/K2Node_SetCameraRigParameters.h b/Engine/Plugins/Cameras/GameplayCameras/Source/GameplayCamerasUncookedOnly/Private/BlueprintGraph/K2Node_SetCameraRigParameters.h
--- a/Engine/Plugins/Cameras/GameplayCameras/Source/GameplayCamerasUncookedOnly/Private/BlueprintGraph/K2Node_SetCameraRigParameters.h
+++ b/Engine/Plugins/Cameras/GameplayCameras/Source/GameplayCamerasUncookedOnly/Private/BlueprintGraph/K2Node_SetCameraRigParameters.h
@@ -17,7 +17,7 @@ class UCameraRigAsset;
  * found on it. On compile, this node gets expanded into the appropriate number of individual setter
  * function calls for each parameter (see UCameraRigParameterInterop).
  */
-UCLASS(MinimalAPI)
+UCLASS(MinimalAPI, meta=(UseCameraRigPickerForPins="CameraRig"))
 class UK2Node_SetCameraRigParameters : public UK2Node
 {
 	GENERATED_BODY()
