Cameras: properly reset the Blueprint camera director result every frame

diff --git a/Engine/Plugins/Cameras/GameplayCameras/Source/GameplayCameras/Private/Directors/BlueprintCameraDirector.cpp b/Engine/Plugins/Cameras/GameplayCameras/Source/GameplayCameras/Private/Directors/BlueprintCameraDirector.cpp
--- a/Engine/Plugins/Cameras/GameplayCameras/Source/GameplayCameras/Private/Directors/BlueprintCameraDirector.cpp
+++ b/Engine/Plugins/Cameras/GameplayCameras/Source/GameplayCameras/Private/Directors/BlueprintCameraDirector.cpp
@@ -27,6 +27,8 @@ void FBlueprintCameraDirectorEvaluationResult::Reset()
 {
 	ActiveCameraRigProxies.Reset();
 	ActiveCameraRigs.Reset();
+	ActivePersistentCameraRigs.Reset();
+	InactivePersistentCameraRigs.Reset();
 }
 
 class FBlueprintCameraDirectorEvaluator : public FCameraDirectorEvaluator
@@ -184,6 +186,11 @@ void FBlueprintCameraDirectorEvaluator::ActivateDeactivePersistentCameraRigs(
 		TSharedPtr<FCameraEvaluationContext> EvaluationContext,
 		const FBlueprintCameraDirectorEvaluationResult& BlueprintResult)
 {
+	if (BlueprintResult.InactivePersistentCameraRigs.IsEmpty() && BlueprintResult.ActivePersistentCameraRigs.IsEmpty())
+	{
+		return;
+	}
+
 	if (!ensure(OwningEvaluator))
 	{
 		return;
