Cameras: rename some BP functions to make it more obvious base/global/visual rigs are persistent

diff --git a/Engine/Plugins/Cameras/GameplayCameras/Config/BaseGameplayCameras.ini b/Engine/Plugins/Cameras/GameplayCameras/Config/BaseGameplayCameras.ini
new file mode 100644
--- /dev/null
+++ b/Engine/Plugins/Cameras/GameplayCameras/Config/BaseGameplayCameras.ini
@@ -0,0 +1,7 @@
+[CoreRedirects]
+
+; UE 5.5 redirects.
++FunctionRedirects=(OldName="ActivateCameraRigFunctions.ActivateBaseRig", NewName="ActivatePersistentBaseRig")
++FunctionRedirects=(OldName="ActivateCameraRigFunctions.ActivateGlobalRig", NewName="ActivatePersistentGlobalRig")
++FunctionRedirects=(OldName="ActivateCameraRigFunctions.ActivateVisualRig", NewName="ActivatePersistentVisualRig")
+
diff --git a/Engine/Plugins/Cameras/GameplayCameras/Source/GameplayCameras/Private/Directors/BlueprintCameraDirector.cpp b/Engine/Plugins/Cameras/GameplayCameras/Source/GameplayCameras/Private/Directors/BlueprintCameraDirector.cpp
--- a/Engine/Plugins/Cameras/GameplayCameras/Source/GameplayCameras/Private/Directors/BlueprintCameraDirector.cpp
+++ b/Engine/Plugins/Cameras/GameplayCameras/Source/GameplayCameras/Private/Directors/BlueprintCameraDirector.cpp
@@ -250,21 +250,36 @@ void FBlueprintCameraDirectorEvaluator::OnAddReferencedObjects(FReferenceCollect
 
 }  // namespace UE::Cameras
 
-void UBlueprintCameraDirectorEvaluator::ActivateBaseCameraRig(UCameraRigAsset* CameraRigPrefab)
+void UBlueprintCameraDirectorEvaluator::ActivatePersistentBaseCameraRig(UCameraRigAsset* CameraRigPrefab)
 {
 	EvaluationResult.ActivePersistentCameraRigs.Add({ CameraRigPrefab, ECameraRigLayer::Base });
 }
 
-void UBlueprintCameraDirectorEvaluator::ActivateGlobalCameraRig(UCameraRigAsset* CameraRigPrefab)
+void UBlueprintCameraDirectorEvaluator::ActivatePersistentGlobalCameraRig(UCameraRigAsset* CameraRigPrefab)
 {
 	EvaluationResult.ActivePersistentCameraRigs.Add({ CameraRigPrefab, ECameraRigLayer::Global });
 }
 
-void UBlueprintCameraDirectorEvaluator::ActivateVisualCameraRig(UCameraRigAsset* CameraRigPrefab)
+void UBlueprintCameraDirectorEvaluator::ActivatePersistentVisualCameraRig(UCameraRigAsset* CameraRigPrefab)
 {
 	EvaluationResult.ActivePersistentCameraRigs.Add({ CameraRigPrefab, ECameraRigLayer::Visual });
 }
 
+void UBlueprintCameraDirectorEvaluator::DeactivatePersistentBaseCameraRig(UCameraRigAsset* CameraRigPrefab)
+{
+	EvaluationResult.InactivePersistentCameraRigs.Add({ CameraRigPrefab, ECameraRigLayer::Base });
+}
+
+void UBlueprintCameraDirectorEvaluator::DeactivatePersistentGlobalCameraRig(UCameraRigAsset* CameraRigPrefab)
+{
+	EvaluationResult.InactivePersistentCameraRigs.Add({ CameraRigPrefab, ECameraRigLayer::Global });
+}
+
+void UBlueprintCameraDirectorEvaluator::DeactivatePersistentVisualCameraRig(UCameraRigAsset* CameraRigPrefab)
+{
+	EvaluationResult.InactivePersistentCameraRigs.Add({ CameraRigPrefab, ECameraRigLayer::Visual });
+}
+
 void UBlueprintCameraDirectorEvaluator::ActivateCameraRig(UCameraRigAsset* CameraRig)
 {
 	EvaluationResult.ActiveCameraRigs.Add(CameraRig);
@@ -275,6 +290,11 @@ void UBlueprintCameraDirectorEvaluator::ActivateCameraRigViaProxy(UCameraRigProx
 	EvaluationResult.ActiveCameraRigProxies.Add(CameraRigProxy);
 }
 
+void UBlueprintCameraDirectorEvaluator::ActivateCameraRigPrefab(UCameraRigAsset* CameraRig)
+{
+	EvaluationResult.ActiveCameraRigs.Add(CameraRig);
+}
+
 AActor* UBlueprintCameraDirectorEvaluator::FindEvaluationContextOwnerActor(TSubclassOf<AActor> ActorClass) const
 {
 	if (EvaluationContext)
diff --git a/Engine/Plugins/Cameras/GameplayCameras/Source/GameplayCameras/Private/GameFramework/ActivateCameraRigFunctions.cpp b/Engine/Plugins/Cameras/GameplayCameras/Source/GameplayCameras/Private/GameFramework/ActivateCameraRigFunctions.cpp
--- a/Engine/Plugins/Cameras/GameplayCameras/Source/GameplayCameras/Private/GameFramework/ActivateCameraRigFunctions.cpp
+++ b/Engine/Plugins/Cameras/GameplayCameras/Source/GameplayCameras/Private/GameFramework/ActivateCameraRigFunctions.cpp
@@ -8,22 +8,22 @@
 
 #include UE_INLINE_GENERATED_CPP_BY_NAME(ActivateCameraRigFunctions)
 
-void UActivateCameraRigFunctions::ActivateBaseCameraRig(UObject* WorldContextObject, APlayerController* PlayerController, UCameraRigAsset* CameraRig)
+void UActivateCameraRigFunctions::ActivatePersistentBaseCameraRig(UObject* WorldContextObject, APlayerController* PlayerController, UCameraRigAsset* CameraRig)
 {
-	ActivateCameraRig(WorldContextObject, PlayerController, CameraRig, ECameraRigLayer::Base);
+	ActivateCameraRigImpl(WorldContextObject, PlayerController, CameraRig, ECameraRigLayer::Base);
 }
 
-void UActivateCameraRigFunctions::ActivateGlobalCameraRig(UObject* WorldContextObject, APlayerController* PlayerController, UCameraRigAsset* CameraRig)
+void UActivateCameraRigFunctions::ActivatePersistentGlobalCameraRig(UObject* WorldContextObject, APlayerController* PlayerController, UCameraRigAsset* CameraRig)
 {
-	ActivateCameraRig(WorldContextObject, PlayerController, CameraRig, ECameraRigLayer::Global);
+	ActivateCameraRigImpl(WorldContextObject, PlayerController, CameraRig, ECameraRigLayer::Global);
 }
 
-void UActivateCameraRigFunctions::ActivateVisualCameraRig(UObject* WorldContextObject, APlayerController* PlayerController, UCameraRigAsset* CameraRig)
+void UActivateCameraRigFunctions::ActivatePersistentVisualCameraRig(UObject* WorldContextObject, APlayerController* PlayerController, UCameraRigAsset* CameraRig)
 {
-	ActivateCameraRig(WorldContextObject, PlayerController, CameraRig, ECameraRigLayer::Visual);
+	ActivateCameraRigImpl(WorldContextObject, PlayerController, CameraRig, ECameraRigLayer::Visual);
 }
 
-void UActivateCameraRigFunctions::ActivateCameraRig(UObject* WorldContextObject, APlayerController* PlayerController, UCameraRigAsset* CameraRig, ECameraRigLayer EvaluationLayer)
+void UActivateCameraRigFunctions::ActivateCameraRigImpl(UObject* WorldContextObject, APlayerController* PlayerController, UCameraRigAsset* CameraRig, ECameraRigLayer EvaluationLayer)
 {
 	using namespace UE::Cameras;
 
diff --git a/Engine/Plugins/Cameras/GameplayCameras/Source/GameplayCameras/Public/Directors/BlueprintCameraDirector.h b/Engine/Plugins/Cameras/GameplayCameras/Source/GameplayCameras/Public/Directors/BlueprintCameraDirector.h
--- a/Engine/Plugins/Cameras/GameplayCameras/Source/GameplayCameras/Public/Directors/BlueprintCameraDirector.h
+++ b/Engine/Plugins/Cameras/GameplayCameras/Source/GameplayCameras/Public/Directors/BlueprintCameraDirector.h
@@ -126,15 +126,27 @@ public:
 
 	/** Activates the given camera rig prefab in the base layer. */
 	UFUNCTION(BlueprintCallable, Category="Activation")
-	void ActivateBaseCameraRig(UCameraRigAsset* CameraRigPrefab);
+	void ActivatePersistentBaseCameraRig(UCameraRigAsset* CameraRigPrefab);
 
 	/** Activates the given camera rig prefab in the global layer. */
 	UFUNCTION(BlueprintCallable, Category="Activation")
-	void ActivateGlobalCameraRig(UCameraRigAsset* CameraRigPrefab);
+	void ActivatePersistentGlobalCameraRig(UCameraRigAsset* CameraRigPrefab);
 
 	/** Activates the given camera rig prefab in the visual layer. */
 	UFUNCTION(BlueprintCallable, Category="Activation")
-	void ActivateVisualCameraRig(UCameraRigAsset* CameraRigPrefab);
+	void ActivatePersistentVisualCameraRig(UCameraRigAsset* CameraRigPrefab);
+
+	/** Deactivates the given camera rig prefab in the base layer. */
+	UFUNCTION(BlueprintCallable, Category="Activation")
+	void DeactivatePersistentBaseCameraRig(UCameraRigAsset* CameraRigPrefab);
+
+	/** Deactivates the given camera rig prefab in the global layer. */
+	UFUNCTION(BlueprintCallable, Category="Activation")
+	void DeactivatePersistentGlobalCameraRig(UCameraRigAsset* CameraRigPrefab);
+
+	/** Deactivates the given camera rig prefab in the visual layer. */
+	UFUNCTION(BlueprintCallable, Category="Activation")
+	void DeactivatePersistentVisualCameraRig(UCameraRigAsset* CameraRigPrefab);
 
 public:
 
@@ -151,6 +163,12 @@ public:
 	UFUNCTION(BlueprintCallable, Category="Evaluation")
 	void ActivateCameraRigViaProxy(UCameraRigProxyAsset* CameraRigProxy);
 
+	/**
+	 * Specifies an external camera rig prefab asset to be active this frame.
+	 */
+	UFUNCTION(BlueprintCallable, Category="Evaluation")
+	void ActivateCameraRigPrefab(UCameraRigAsset* CameraRig);
+
 public:
 
 	/**
diff --git a/Engine/Plugins/Cameras/GameplayCameras/Source/GameplayCameras/Public/GameFramework/ActivateCameraRigFunctions.h b/Engine/Plugins/Cameras/GameplayCameras/Source/GameplayCameras/Public/GameFramework/ActivateCameraRigFunctions.h
--- a/Engine/Plugins/Cameras/GameplayCameras/Source/GameplayCameras/Public/GameFramework/ActivateCameraRigFunctions.h
+++ b/Engine/Plugins/Cameras/GameplayCameras/Source/GameplayCameras/Public/GameFramework/ActivateCameraRigFunctions.h
@@ -27,20 +27,19 @@ public:
 
 	/** Activates the given camera rig prefab in the base layer. */
 	UFUNCTION(BlueprintCallable, Category="Camera", meta=(WorldContext="WorldContextObject"))
-	static void ActivateBaseCameraRig(UObject* WorldContextObject, APlayerController* PlayerController, UCameraRigAsset* CameraRig);
+	static void ActivatePersistentBaseCameraRig(UObject* WorldContextObject, APlayerController* PlayerController, UCameraRigAsset* CameraRig);
 
 	/** Activates the given camera rig prefab in the global layer. */
 	UFUNCTION(BlueprintCallable, Category="Camera", meta = (WorldContext = "WorldContextObject"))
-	static void ActivateGlobalCameraRig(UObject* WorldContextObject, APlayerController* PlayerController, UCameraRigAsset* CameraRig);
+	static void ActivatePersistentGlobalCameraRig(UObject* WorldContextObject, APlayerController* PlayerController, UCameraRigAsset* CameraRig);
 
 	/** Activates the given camera rig prefab in the visual layer. */
 	UFUNCTION(BlueprintCallable, Category="Camera", meta = (WorldContext = "WorldContextObject"))
-	static void ActivateVisualCameraRig(UObject* WorldContextObject, APlayerController* PlayerController, UCameraRigAsset* CameraRig);
+	static void ActivatePersistentVisualCameraRig(UObject* WorldContextObject, APlayerController* PlayerController, UCameraRigAsset* CameraRig);
 
 private:
 
 	/** Activates the given camera rig in the given layer. Should not be used with Main layer. */
-	UFUNCTION(BlueprintCallable, Category="Camera", meta = (WorldContext = "WorldContextObject"))
-	static void ActivateCameraRig(UObject* WorldContextObject, APlayerController* PlayerController, UCameraRigAsset* CameraRig, ECameraRigLayer EvaluationLayer);
+	static void ActivateCameraRigImpl(UObject* WorldContextObject, APlayerController* PlayerController, UCameraRigAsset* CameraRig, ECameraRigLayer EvaluationLayer);
 };
 
