Cameras: rename standalone camera rigs to "camera rig prefabs"

diff --git a/Engine/Plugins/Cameras/GameplayCameras/Source/GameplayCameras/Public/Nodes/Common/CameraRigCameraNode.h b/Engine/Plugins/Cameras/GameplayCameras/Source/GameplayCameras/Public/Nodes/Common/CameraRigCameraNode.h
--- a/Engine/Plugins/Cameras/GameplayCameras/Source/GameplayCameras/Public/Nodes/Common/CameraRigCameraNode.h
+++ b/Engine/Plugins/Cameras/GameplayCameras/Source/GameplayCameras/Public/Nodes/Common/CameraRigCameraNode.h
@@ -11,7 +11,7 @@
 /**
  * A camera node that runs a camera rig's own node tree.
  */
-UCLASS(MinimalAPI, meta=(CameraNodeCategories="Common,Utility"))
+UCLASS(MinimalAPI, meta=(DisplayName="Camera Rig Prefab", CameraNodeCategories="Common,Utility"))
 class UCameraRigCameraNode : public UCameraNode
 {
 	GENERATED_BODY()
diff --git a/Engine/Plugins/Cameras/GameplayCameras/Source/GameplayCamerasEditor/Private/AssetTools/AssetDefinition_CameraRigAsset.cpp b/Engine/Plugins/Cameras/GameplayCameras/Source/GameplayCamerasEditor/Private/AssetTools/AssetDefinition_CameraRigAsset.cpp
--- a/Engine/Plugins/Cameras/GameplayCameras/Source/GameplayCamerasEditor/Private/AssetTools/AssetDefinition_CameraRigAsset.cpp
+++ b/Engine/Plugins/Cameras/GameplayCameras/Source/GameplayCamerasEditor/Private/AssetTools/AssetDefinition_CameraRigAsset.cpp
@@ -14,7 +14,7 @@
 
 FText UAssetDefinition_CameraRigAsset::GetAssetDisplayName() const
 {
-	return LOCTEXT("AssetDisplayName", "Camera Rig");
+	return LOCTEXT("AssetDisplayName", "Camera Rig Prefab");
 }
 
 FLinearColor UAssetDefinition_CameraRigAsset::GetAssetColor() const
