Cameras: fix loading older assets

diff --git a/Engine/Plugins/Cameras/GameplayCameras/Source/GameplayCamerasEditor/Private/Editors/CameraNodeGraphSchema.cpp b/Engine/Plugins/Cameras/GameplayCameras/Source/GameplayCamerasEditor/Private/Editors/CameraNodeGraphSchema.cpp
--- a/Engine/Plugins/Cameras/GameplayCameras/Source/GameplayCamerasEditor/Private/Editors/CameraNodeGraphSchema.cpp
+++ b/Engine/Plugins/Cameras/GameplayCameras/Source/GameplayCamerasEditor/Private/Editors/CameraNodeGraphSchema.cpp
@@ -95,6 +95,7 @@ void UCameraNodeGraphSchema::CollectAllObjects(UObjectTreeGraph* InGraph, TSet<U
 			for (UObject* MissingObject : MissingNodeTreeObjects)
 			{
 				((IObjectTreeGraphRootObject*)CameraRig)->AddConnectableObject(UCameraRigAsset::NodeTreeGraphName, MissingObject);
+				OutAllObjects.Add(MissingObject);
 			}
 		}
 	}
