Cameras: restore the option on the camera system actor to auto-set the control rotation

The Gameplay Control Rotation component disables it for auto-spawned camera system actors.

diff --git a/Engine/Plugins/Cameras/GameplayCameras/Source/GameplayCameras/Private/GameFramework/GameplayCameraSystemComponent.cpp b/Engine/Plugins/Cameras/GameplayCameras/Source/GameplayCameras/Private/GameFramework/GameplayCameraSystemComponent.cpp
--- a/Engine/Plugins/Cameras/GameplayCameras/Source/GameplayCameras/Private/GameFramework/GameplayCameraSystemComponent.cpp
+++ b/Engine/Plugins/Cameras/GameplayCameras/Source/GameplayCameras/Private/GameFramework/GameplayCameraSystemComponent.cpp
@@ -56,6 +56,14 @@ void UGameplayCameraSystemComponent::GetCameraView(float DeltaTime, FMinimalView
 		Evaluator->Update(UpdateParams);
 
 		Evaluator->GetEvaluatedCameraView(DesiredView);
+
+		if (bSetPlayerControllerRotation)
+		{
+			if (APlayerController* PlayerController = WeakPlayerController.Get())
+			{
+				PlayerController->SetControlRotation(Evaluator->GetEvaluatedResult().CameraPose.GetRotation());
+			}
+		}
 	}
 }
 
diff --git a/Engine/Plugins/Cameras/GameplayCameras/Source/GameplayCameras/Private/GameFramework/GameplayControlRotationComponent.cpp b/Engine/Plugins/Cameras/GameplayCameras/Source/GameplayCameras/Private/GameFramework/GameplayControlRotationComponent.cpp
--- a/Engine/Plugins/Cameras/GameplayCameras/Source/GameplayCameras/Private/GameFramework/GameplayControlRotationComponent.cpp
+++ b/Engine/Plugins/Cameras/GameplayCameras/Source/GameplayCameras/Private/GameFramework/GameplayControlRotationComponent.cpp
@@ -9,6 +9,8 @@
 #include "Engine/World.h"
 #include "GameFramework/Controller.h"
 #include "GameFramework/GameplayCameraComponent.h"
+#include "GameFramework/GameplayCameraSystemActor.h"
+#include "GameFramework/GameplayCameraSystemComponent.h"
 #include "GameFramework/GameplayCameraSystemHost.h"
 #include "GameFramework/Pawn.h"
 #include "GameplayCameras.h"
@@ -70,6 +72,13 @@ void UGameplayControlRotationComponent::BeginPlay()
 		return;
 	}
 
+	// Make sure that if there's an auto-spawned camera system actor, it doesn't set the control rotation.
+	AGameplayCameraSystemActor* AutoSpawnedCameraSystemActor = AGameplayCameraSystemActor::GetAutoSpawnedCameraSystemActor(PlayerController);
+	if (AutoSpawnedCameraSystemActor)
+	{
+		AutoSpawnedCameraSystemActor->GetCameraSystemComponent()->bSetPlayerControllerRotation = false;
+	}
+
 	// Create the evaluation service, with a copy of our parameters.
 	FPlayerControlRotationParams ServiceParams;
 	ServiceParams.AxisActionAngularSpeedThreshold = AxisActionAngularSpeedThreshold;
diff --git a/Engine/Plugins/Cameras/GameplayCameras/Source/GameplayCameras/Public/GameFramework/GameplayCameraSystemComponent.h b/Engine/Plugins/Cameras/GameplayCameras/Source/GameplayCameras/Public/GameFramework/GameplayCameraSystemComponent.h
--- a/Engine/Plugins/Cameras/GameplayCameras/Source/GameplayCameras/Public/GameFramework/GameplayCameraSystemComponent.h
+++ b/Engine/Plugins/Cameras/GameplayCameras/Source/GameplayCameras/Public/GameFramework/GameplayCameraSystemComponent.h
@@ -90,6 +90,13 @@ public:
 	UPROPERTY(EditAnywhere, Category=Camera)
 	TEnumAsByte<EAutoReceiveInput::Type> AutoActivateForPlayer;
 
+	/**
+	 * If enabled, sets the evaluated camera orientation as the player controller rotation every frame.
+	 * This is set on the player controller that this component was activated for.
+	 */
+	UPROPERTY(EditAnywhere, Category=Camera)
+	bool bSetPlayerControllerRotation = true;
+
 private:
 	
 	UPROPERTY()
