Cameras: add camera variable for specifying a custom framing target

diff --git a/Engine/Plugins/Cameras/GameplayCameras/Source/GameplayCameras/Private/Nodes/Framing/BaseFramingCameraNode.cpp b/Engine/Plugins/Cameras/GameplayCameras/Source/GameplayCameras/Private/Nodes/Framing/BaseFramingCameraNode.cpp
--- a/Engine/Plugins/Cameras/GameplayCameras/Source/GameplayCameras/Private/Nodes/Framing/BaseFramingCameraNode.cpp
+++ b/Engine/Plugins/Cameras/GameplayCameras/Source/GameplayCameras/Private/Nodes/Framing/BaseFramingCameraNode.cpp
@@ -9,6 +9,8 @@
 #include "Debug/CameraDebugBlock.h"
 #include "Debug/CameraDebugBlockBuilder.h"
 #include "Debug/CameraDebugRenderer.h"
+#include "GameFramework/Pawn.h"
+#include "GameFramework/PlayerController.h"
 #include "GameplayCameras.h"
 #include "HAL/IConsoleManager.h"
 #include "Math/CameraPoseMath.h"
@@ -66,6 +68,25 @@ void FBaseFramingCameraNodeEvaluator::OnInitialize(const FCameraNodeEvaluatorIni
 	Readers.SoftZoneMargin.Initialize(BaseFramingNode->SoftZone);
 }
 
+TOptional<FVector3d> FBaseFramingCameraNodeEvaluator::AcquireTargetLocation(const FCameraNodeEvaluationParams& Params, const FCameraNodeEvaluationResult& InResult)
+{
+	const UBaseFramingCameraNode* DollyNode = GetCameraNodeAs<UBaseFramingCameraNode>();
+	if (DollyNode->TargetLocation)
+	{
+		FVector3d TargetLocation;
+		const bool bGotTargetLocation = InResult.VariableTable.TryGetValue(DollyNode->TargetLocation.Get(), TargetLocation);
+		return bGotTargetLocation ? TOptional<FVector3d>(TargetLocation) : TOptional<FVector3d>();
+	}
+	else if (APlayerController* PlayerController = Params.EvaluationContext->GetPlayerController())
+	{
+		APawn* Pawn = PlayerController->GetPawn();
+		FVector3d TargetLocation = Pawn->GetActorLocation();
+		return TOptional<FVector3d>(TargetLocation);
+	}
+
+	return TOptional<FVector3d>();
+}
+
 void FBaseFramingCameraNodeEvaluator::UpdateFramingState(const FCameraNodeEvaluationParams& Params, const FCameraNodeEvaluationResult& OutResult, const FVector3d& TargetLocation, const FTransform3d& LastFraming)
 {
 	// Get screen-space coordinates of the ideal framing point. These are in 0..1 UI space.
diff --git a/Engine/Plugins/Cameras/GameplayCameras/Source/GameplayCameras/Private/Nodes/Framing/DollyFramingCameraNode.cpp b/Engine/Plugins/Cameras/GameplayCameras/Source/GameplayCameras/Private/Nodes/Framing/DollyFramingCameraNode.cpp
--- a/Engine/Plugins/Cameras/GameplayCameras/Source/GameplayCameras/Private/Nodes/Framing/DollyFramingCameraNode.cpp
+++ b/Engine/Plugins/Cameras/GameplayCameras/Source/GameplayCameras/Private/Nodes/Framing/DollyFramingCameraNode.cpp
@@ -4,6 +4,7 @@
 
 #include "Core/CameraEvaluationContext.h"
 #include "Core/CameraParameterReader.h"
+#include "Core/CameraVariableAssets.h"
 #include "Debug/CameraDebugBlock.h"
 #include "Debug/CameraDebugBlockBuilder.h"
 #include "Debug/CameraDebugRenderer.h"
@@ -74,12 +75,14 @@ void FDollyFramingCameraNodeEvaluator::OnInitialize(const FCameraNodeEvaluatorIn
 
 void FDollyFramingCameraNodeEvaluator::OnRun(const FCameraNodeEvaluationParams& Params, FCameraNodeEvaluationResult& OutResult)
 {
-	// TODO get the target some other way
-	APlayerController* PlayerController = Params.EvaluationContext->GetPlayerController();
-	APawn* Pawn = PlayerController->GetPawn();
-	const FVector3d TargetLocation = Pawn->GetActorLocation();
+	TOptional<FVector3d> OptTargetLocation = AcquireTargetLocation(Params, OutResult);
+	if (!OptTargetLocation.IsSet())
+	{
+		return;
+	}
 
 	// Let the base class figure out all the screen-space framing stuff.
+	const FVector3d TargetLocation = OptTargetLocation.GetValue();
 	const FTransform3d LastShotTransform = BuildDollyShotTransform(OutResult.CameraPose);
 	UpdateFramingState(Params, OutResult, TargetLocation, LastShotTransform);
 	ComputeDesiredState(Params.DeltaTime);
@@ -91,6 +94,7 @@ void FDollyFramingCameraNodeEvaluator::OnRun(const FCameraNodeEvaluationParams&
 		FCameraPose LastShotPose(OutResult.CameraPose);
 		LastShotPose.SetTransform(LastShotTransform);
 
+		APlayerController* PlayerController = Params.EvaluationContext->GetPlayerController();
 		FVector3d DesiredLocalOffset = ComputeFramingTranslation(LastShotPose, PlayerController);
 
 		// We never bring the dolly forward or backward (we only move it vertically or horizontally).
diff --git a/Engine/Plugins/Cameras/GameplayCameras/Source/GameplayCameras/Private/Nodes/Framing/PanningFramingCameraNode.cpp b/Engine/Plugins/Cameras/GameplayCameras/Source/GameplayCameras/Private/Nodes/Framing/PanningFramingCameraNode.cpp
--- a/Engine/Plugins/Cameras/GameplayCameras/Source/GameplayCameras/Private/Nodes/Framing/PanningFramingCameraNode.cpp
+++ b/Engine/Plugins/Cameras/GameplayCameras/Source/GameplayCameras/Private/Nodes/Framing/PanningFramingCameraNode.cpp
@@ -4,11 +4,10 @@
 
 #include "Core/CameraEvaluationContext.h"
 #include "Core/CameraParameterReader.h"
+#include "Core/CameraVariableAssets.h"
 #include "Debug/CameraDebugBlock.h"
 #include "Debug/CameraDebugBlockBuilder.h"
 #include "Debug/CameraDebugRenderer.h"
-#include "GameFramework/Pawn.h"
-#include "GameFramework/PlayerController.h"
 #include "GameplayCameras.h"
 #include "Math/CameraFramingMath.h"
 #include "Math/CameraPoseMath.h"
@@ -70,12 +69,14 @@ void FPanningFramingCameraNodeEvaluator::OnInitialize(const FCameraNodeEvaluator
 
 void FPanningFramingCameraNodeEvaluator::OnRun(const FCameraNodeEvaluationParams& Params, FCameraNodeEvaluationResult& OutResult)
 {
-	// TODO get the target some other way
-	APlayerController* PlayerController = Params.EvaluationContext->GetPlayerController();
-	APawn* Pawn = PlayerController->GetPawn();
-	const FVector3d TargetLocation = Pawn->GetActorLocation();
+	TOptional<FVector3d> OptTargetLocation = AcquireTargetLocation(Params, OutResult);
+	if (!OptTargetLocation.IsSet())
+	{
+		return;
+	}
 
 	// Let the base class figure out all the screen-space framing stuff.
+	const FVector3d TargetLocation = OptTargetLocation.GetValue();
 	const FTransform3d LastShotTransform = BuildPanningShotTransform(OutResult.CameraPose);
 	UpdateFramingState(Params, OutResult, TargetLocation, LastShotTransform);
 	ComputeDesiredState(Params.DeltaTime);
@@ -87,6 +88,7 @@ void FPanningFramingCameraNodeEvaluator::OnRun(const FCameraNodeEvaluationParams
 		FCameraPose LastShotPose(OutResult.CameraPose);
 		LastShotPose.SetTransform(LastShotTransform);
 
+		APlayerController* PlayerController = Params.EvaluationContext->GetPlayerController();
 		const float AspectRatio = FCameraPoseMath::GetEffectiveAspectRatio(LastShotPose, PlayerController);
 		const FCameraFieldsOfView FOVs(FCameraPoseMath::GetEffectiveFieldsOfView(LastShotPose, AspectRatio));
 
diff --git a/Engine/Plugins/Cameras/GameplayCameras/Source/GameplayCameras/Public/Core/CameraVariableTable.h b/Engine/Plugins/Cameras/GameplayCameras/Source/GameplayCameras/Public/Core/CameraVariableTable.h
--- a/Engine/Plugins/Cameras/GameplayCameras/Source/GameplayCameras/Public/Core/CameraVariableTable.h
+++ b/Engine/Plugins/Cameras/GameplayCameras/Source/GameplayCameras/Public/Core/CameraVariableTable.h
@@ -103,9 +103,15 @@ public:
 	template<typename ValueType>
 	ValueType GetValue(FCameraVariableID VariableID, typename TCallTraits<ValueType>::ParamType DefaultValue) const;
 
+	template<typename VariableAssetType>
+	typename VariableAssetType::ValueType GetValue(const VariableAssetType* VariableAsset) const;
+
 	template<typename ValueType>
 	bool TryGetValue(FCameraVariableID VariableID, ValueType& OutValue) const;
 
+	template<typename VariableAssetType>
+	bool TryGetValue(const VariableAssetType* VariableAsset, typename VariableAssetType::ValueType& OutValue) const;
+
 	bool ContainsValue(FCameraVariableID VariableID) const;
 
 public:
@@ -266,6 +272,12 @@ ValueType FCameraVariableTable::GetValue(FCameraVariableID VariableID, typename
 	return DefaultValue;
 }
 
+template<typename VariableAssetType>
+typename VariableAssetType::ValueType FCameraVariableTable::GetValue(const VariableAssetType* VariableAsset) const
+{
+	return GetValue(VariableAsset->GetVariableID(), VariableAsset->GetDefaultValue());
+}
+
 template<typename ValueType>
 bool FCameraVariableTable::TryGetValue(FCameraVariableID VariableID, ValueType& OutValue) const
 {
@@ -277,6 +289,12 @@ bool FCameraVariableTable::TryGetValue(FCameraVariableID VariableID, ValueType&
 	return false;
 }
 
+template<typename VariableAssetType>
+bool FCameraVariableTable::TryGetValue(const VariableAssetType* VariableAsset, typename VariableAssetType::ValueType& OutValue) const
+{
+	return TryGetValue(VariableAsset->GetVariableID(), OutValue);
+}
+
 template<typename ValueType>
 void FCameraVariableTable::SetValue(FCameraVariableID VariableID, typename TCallTraits<ValueType>::ParamType Value)
 {
diff --git a/Engine/Plugins/Cameras/GameplayCameras/Source/GameplayCameras/Public/Nodes/Framing/BaseFramingCameraNode.h b/Engine/Plugins/Cameras/GameplayCameras/Source/GameplayCameras/Public/Nodes/Framing/BaseFramingCameraNode.h
--- a/Engine/Plugins/Cameras/GameplayCameras/Source/GameplayCameras/Public/Nodes/Framing/BaseFramingCameraNode.h
+++ b/Engine/Plugins/Cameras/GameplayCameras/Source/GameplayCameras/Public/Nodes/Framing/BaseFramingCameraNode.h
@@ -12,6 +12,7 @@
 #include "BaseFramingCameraNode.generated.h"
 
 class FArchive;
+class UVector3dCameraVariable;
 
 /**
  * The base class for a standard scren-space framing camera node.
@@ -23,6 +24,13 @@ class UBaseFramingCameraNode : public UCameraNode
 
 public:
 
+	/**
+	 * A camera variable providing the location of the target to frame. If unspecified,
+	 * the player pawn's location will be used by default.
+	 */
+	UPROPERTY(EditAnywhere, Category="Target")
+	TObjectPtr<UVector3dCameraVariable> TargetLocation;
+
 	/** The ideal horizontal screen-space position of the target. */
 	UPROPERTY(EditAnywhere, Category="Framing Target")
 	FDoubleCameraParameter HorizontalFraming;
@@ -133,6 +141,8 @@ protected:
 
 protected:
 
+	/** Gets the target location. */
+	TOptional<FVector3d> AcquireTargetLocation(const FCameraNodeEvaluationParams& Params, const FCameraNodeEvaluationResult& InResult);
 	/** Updates the framing state for the current tick, see State member field. */
 	void UpdateFramingState(const FCameraNodeEvaluationParams& Params, const FCameraNodeEvaluationResult& OutResult, const FVector3d& TargetLocation, const FTransform3d& LastFraming);
 	/** Computes the desired reframing for the current tick, see Desired member field. */
