Cameras: set the camera cut flag when the gameplay camera system gets teleported

diff --git a/Engine/Plugins/Cameras/GameplayCameras/Source/GameplayCameras/Private/GameFramework/GameplayCameraComponent.cpp b/Engine/Plugins/Cameras/GameplayCameras/Source/GameplayCameras/Private/GameFramework/GameplayCameraComponent.cpp
--- a/Engine/Plugins/Cameras/GameplayCameras/Source/GameplayCameras/Private/GameFramework/GameplayCameraComponent.cpp
+++ b/Engine/Plugins/Cameras/GameplayCameras/Source/GameplayCameras/Private/GameFramework/GameplayCameraComponent.cpp
@@ -22,6 +22,7 @@
 UGameplayCameraComponent::UGameplayCameraComponent(const FObjectInitializer& ObjectInit)
 	: Super(ObjectInit)
 {
+	bWantsOnUpdateTransform = true;
 	PrimaryComponentTick.bCanEverTick = true;
 
 #if WITH_EDITORONLY_DATA
@@ -241,6 +242,12 @@ void UGameplayCameraComponent::TickComponent(float DeltaTime, enum ELevelTick Ti
 	if (EvaluationContext)
 	{
 		EvaluationContext->Update(this);
+
+		if (bIsCameraCutNextFrame)
+		{
+			EvaluationContext->GetInitialResult().bIsCameraCut = true;
+			bIsCameraCutNextFrame = false;
+		}
 	}
 }
 
@@ -256,6 +263,16 @@ void UGameplayCameraComponent::OnComponentDestroyed(bool bDestroyingHierarchy)
 #endif  // WITH_EDITORONLY_DATA
 }
 
+void UGameplayCameraComponent::OnUpdateTransform(EUpdateTransformFlags UpdateTransformFlags, ETeleportType Teleport)
+{
+	Super::OnUpdateTransform(UpdateTransformFlags, Teleport);
+
+	if (EvaluationContext && Teleport != ETeleportType::None)
+	{
+		bIsCameraCutNextFrame = true;
+	}
+}
+
 #if WITH_EDITORONLY_DATA
 
 void UGameplayCameraComponent::UpdatePreviewMeshTransform()
@@ -280,6 +297,7 @@ void FGameplayCameraComponentEvaluationContext::Update(UGameplayCameraComponent*
 {
 	const FTransform& OwnerTransform = Owner->GetComponentTransform();
 	InitialResult.CameraPose.SetTransform(OwnerTransform);
+	InitialResult.bIsCameraCut = false;
 	InitialResult.bIsValid = true;
 }
 
diff --git a/Engine/Plugins/Cameras/GameplayCameras/Source/GameplayCameras/Public/GameFramework/GameplayCameraComponent.h b/Engine/Plugins/Cameras/GameplayCameras/Source/GameplayCameras/Public/GameFramework/GameplayCameraComponent.h
--- a/Engine/Plugins/Cameras/GameplayCameras/Source/GameplayCameras/Public/GameFramework/GameplayCameraComponent.h
+++ b/Engine/Plugins/Cameras/GameplayCameras/Source/GameplayCameras/Public/GameFramework/GameplayCameraComponent.h
@@ -76,6 +76,9 @@ public:
 	virtual void TickComponent(float DeltaTime, enum ELevelTick TickType, FActorComponentTickFunction *ThisTickFunction) override;
 	virtual void OnComponentDestroyed(bool bDestroyingHierarchy) override;
 
+	// USceneComponent interface.
+	virtual void OnUpdateTransform(EUpdateTransformFlags UpdateTransformFlags, ETeleportType Teleport) override;
+
 private:
 
 	void ActivateCameraEvaluationContext(int32 PlayerIndex);
@@ -107,6 +110,8 @@ protected:
 
 	TSharedPtr<FGameplayCameraComponentEvaluationContext> EvaluationContext;
 
+	bool bIsCameraCutNextFrame = false;
+
 	UPROPERTY()
 	TObjectPtr<UGameplayCameraSystemHost> CameraSystemHost;
 
