Cameras: make the boom arm length interpolation gracefully handle camera cuts

diff --git a/Engine/Plugins/Cameras/GameplayCameras/Source/GameplayCameras/Private/Nodes/Common/BoomArmCameraNode.cpp b/Engine/Plugins/Cameras/GameplayCameras/Source/GameplayCameras/Private/Nodes/Common/BoomArmCameraNode.cpp
--- a/Engine/Plugins/Cameras/GameplayCameras/Source/GameplayCameras/Private/Nodes/Common/BoomArmCameraNode.cpp
+++ b/Engine/Plugins/Cameras/GameplayCameras/Source/GameplayCameras/Private/Nodes/Common/BoomArmCameraNode.cpp
@@ -137,7 +137,7 @@ void FBoomArmCameraNodeEvaluator::OnRun(const FCameraNodeEvaluationParams& Param
 	const double DefaultBoomLength = BoomOffset.Length();
 	if (BoomLengthInterpolator && DefaultBoomLength > 0)
 	{
-		if (!Params.bIsFirstFrame)
+		if (!Params.bIsFirstFrame && !OutResult.bIsCameraCut)
 		{
 			// The pull this frame is how much the base (pivot) of the boom arm has moved. We add that
 			// to our cumulative tally of the pull.
