Cameras: fix redirectors

diff --git a/Engine/Plugins/Cameras/GameplayCameras/Config/BaseGameplayCameras.ini b/Engine/Plugins/Cameras/GameplayCameras/Config/BaseGameplayCameras.ini
--- a/Engine/Plugins/Cameras/GameplayCameras/Config/BaseGameplayCameras.ini
+++ b/Engine/Plugins/Cameras/GameplayCameras/Config/BaseGameplayCameras.ini
@@ -1,7 +1,7 @@
 [CoreRedirects]
 
 ; UE 5.5 redirects.
-+FunctionRedirects=(OldName="ActivateCameraRigFunctions.ActivateBaseRig", NewName="ActivatePersistentBaseRig")
-+FunctionRedirects=(OldName="ActivateCameraRigFunctions.ActivateGlobalRig", NewName="ActivatePersistentGlobalRig")
-+FunctionRedirects=(OldName="ActivateCameraRigFunctions.ActivateVisualRig", NewName="ActivatePersistentVisualRig")
++FunctionRedirects=(OldName="ActivateCameraRigFunctions.ActivateBaseCameraRig", NewName="ActivatePersistentBaseCameraRig")
++FunctionRedirects=(OldName="ActivateCameraRigFunctions.ActivateGlobalCameraRig", NewName="ActivatePersistentGlobalCameraRig")
++FunctionRedirects=(OldName="ActivateCameraRigFunctions.ActivateVisualCameraRig", NewName="ActivatePersistentVisualCameraRig")
 
