Cameras: add camera variable references

These are for when a camera node just needs a camera variable directly, without a "static value" like with camera parameters.

diff --git a/Engine/Plugins/Cameras/GameplayCameras/Source/GameplayCameras/Private/Core/CameraRigAssetBuilder.cpp b/Engine/Plugins/Cameras/GameplayCameras/Source/GameplayCameras/Private/Core/CameraRigAssetBuilder.cpp
--- a/Engine/Plugins/Cameras/GameplayCameras/Source/GameplayCameras/Private/Core/CameraRigAssetBuilder.cpp
+++ b/Engine/Plugins/Cameras/GameplayCameras/Source/GameplayCameras/Private/Core/CameraRigAssetBuilder.cpp
@@ -5,9 +5,10 @@
 #include "Core/CameraNode.h"
 #include "Core/CameraNodeEvaluatorStorage.h"
 #include "Core/CameraParameters.h"
-#include "Core/CameraRigBuildContext.h"
 #include "Core/CameraRigAsset.h"
+#include "Core/CameraRigBuildContext.h"
 #include "Core/CameraVariableAssets.h"
+#include "Core/CameraVariableReferences.h"
 #include "GameplayCamerasDelegates.h"
 #include "Logging/TokenizedMessage.h"
 #include "Nodes/Common/CameraRigCameraNode.h"
@@ -285,6 +286,19 @@ void SetupPrivateVariable(
 	DoSetupPrivateVariable(Builder, InterfaceParameter, &ParameterOverride->Value, ReusedVariable);
 }
 
+void AddCameraVariableToAllocationInfo(UCameraVariableAsset* Variable, FCameraVariableTableAllocationInfo& AllocationInfo)
+{
+	if (Variable)
+	{
+		FCameraVariableDefinition VariableDefinition = Variable->GetVariableDefinition();
+		AllocationInfo.VariableDefinitions.Add(VariableDefinition);
+		if (Variable->bAutoReset)
+		{
+			AllocationInfo.AutoResetVariables.Add(Variable);
+		}
+	}
+}
+
 }  // namespace Internal
 
 FCameraRigAssetBuilder::FCameraRigAssetBuilder(FCameraBuildLog& InBuildLog)
@@ -819,6 +833,8 @@ void FCameraRigAssetBuilder::BuildAllocationInfo()
 
 void FCameraRigAssetBuilder::BuildAllocationInfo(UCameraNode* CameraNode)
 {
+	using namespace UE::Cameras::Internal;
+
 	// Look for properties that are camera parameters, and gather what camera variables they reference. 
 	// This is for both exposed rig parameters (which we just built in BuildNewDrivenParameters) and 
 	// for parameters driven by user-defined variables.
@@ -835,15 +851,12 @@ void FCameraRigAssetBuilder::BuildAllocationInfo(UCameraNode* CameraNode)
 		if (StructProperty->Struct == F##ValueName##CameraParameter::StaticStruct())\
 		{\
 			auto* CameraParameterPtr = StructProperty->ContainerPtrToValuePtr<F##ValueName##CameraParameter>(CameraNode);\
-			if (CameraParameterPtr->Variable)\
-			{\
-				FCameraVariableDefinition VariableDefinition = CameraParameterPtr->Variable->GetVariableDefinition();\
-				AllocationInfo.VariableTableInfo.VariableDefinitions.Add(VariableDefinition);\
-				if (CameraParameterPtr->Variable->bAutoReset)\
-				{\
-					AllocationInfo.VariableTableInfo.AutoResetVariables.Add(CameraParameterPtr->Variable);\
-				}\
-			}\
+			AddCameraVariableToAllocationInfo(CameraParameterPtr->Variable, AllocationInfo.VariableTableInfo);\
+		}\
+		else if (StructProperty->Struct == F##ValueName##CameraVariableReference::StaticStruct())\
+		{\
+			auto* CameraVariableReferencePtr = StructProperty->ContainerPtrToValuePtr<F##ValueName##CameraVariableReference>(CameraNode);\
+			AddCameraVariableToAllocationInfo(CameraVariableReferencePtr->Variable, AllocationInfo.VariableTableInfo);\
 		}\
 		else
 UE_CAMERA_VARIABLE_FOR_ALL_TYPES()
diff --git a/Engine/Plugins/Cameras/GameplayCameras/Source/GameplayCameras/Public/Core/CameraVariableReferences.h b/Engine/Plugins/Cameras/GameplayCameras/Source/GameplayCameras/Public/Core/CameraVariableReferences.h
new file mode 100644
--- /dev/null
+++ b/Engine/Plugins/Cameras/GameplayCameras/Source/GameplayCameras/Public/Core/CameraVariableReferences.h
@@ -0,0 +1,196 @@
+// Copyright Epic Games, Inc. All Rights Reserved.
+
+#pragma once
+
+#include "Core/CameraVariableAssets.h"
+
+#include "CameraVariableReferences.generated.h"
+
+#define UE_DEFINE_CAMERA_VARIABLE_REFERENCE(ValueName)\
+	F##ValueName##CameraVariableReference() {}\
+	F##ValueName##CameraVariableReference(VariableAssetType* InVariable) : Variable(InVariable) {}\
+	U##ValueName##CameraVariable* Get() const { return Variable.Get(); }\
+	operator U##ValueName##CameraVariable* () const { return Variable.Get(); }
+
+USTRUCT()
+struct FBooleanCameraVariableReference
+{
+	GENERATED_BODY()
+
+	using VariableAssetType = UBooleanCameraVariable;
+
+	UPROPERTY(EditAnywhere, Category="Variable")
+	TObjectPtr<UBooleanCameraVariable> Variable;
+
+	UE_DEFINE_CAMERA_VARIABLE_REFERENCE(Boolean)
+};
+
+USTRUCT()
+struct FInteger32CameraVariableReference
+{
+	GENERATED_BODY()
+
+	using VariableAssetType = UInteger32CameraVariable;
+
+	UPROPERTY(EditAnywhere, Category="Variable")
+	TObjectPtr<UInteger32CameraVariable> Variable;
+
+	UE_DEFINE_CAMERA_VARIABLE_REFERENCE(Integer32)
+};
+
+USTRUCT()
+struct FFloatCameraVariableReference
+{
+	GENERATED_BODY()
+
+	using VariableAssetType = UFloatCameraVariable;
+
+	UPROPERTY(EditAnywhere, Category="Variable")
+	TObjectPtr<UFloatCameraVariable> Variable;
+
+	UE_DEFINE_CAMERA_VARIABLE_REFERENCE(Float)
+};
+
+USTRUCT()
+struct FDoubleCameraVariableReference
+{
+	GENERATED_BODY()
+
+	using VariableAssetType = UDoubleCameraVariable;
+
+	UPROPERTY(EditAnywhere, Category="Variable")
+	TObjectPtr<UDoubleCameraVariable> Variable;
+
+	UE_DEFINE_CAMERA_VARIABLE_REFERENCE(Double)
+};
+
+USTRUCT()
+struct FVector2fCameraVariableReference
+{
+	GENERATED_BODY()
+
+	using VariableAssetType = UVector2fCameraVariable;
+
+	UPROPERTY(EditAnywhere, Category="Variable")
+	TObjectPtr<UVector2fCameraVariable> Variable;
+
+	UE_DEFINE_CAMERA_VARIABLE_REFERENCE(Vector2f)
+};
+
+USTRUCT()
+struct FVector2dCameraVariableReference
+{
+	GENERATED_BODY()
+
+	using VariableAssetType = UVector2dCameraVariable;
+
+	UPROPERTY(EditAnywhere, Category="Variable")
+	TObjectPtr<UVector2dCameraVariable> Variable;
+
+	UE_DEFINE_CAMERA_VARIABLE_REFERENCE(Vector2d)
+};
+
+USTRUCT()
+struct FVector3fCameraVariableReference
+{
+	GENERATED_BODY()
+
+	using VariableAssetType = UVector3fCameraVariable;
+
+	UPROPERTY(EditAnywhere, Category="Variable")
+	TObjectPtr<UVector3fCameraVariable> Variable;
+
+	UE_DEFINE_CAMERA_VARIABLE_REFERENCE(Vector3f)
+};
+
+USTRUCT()
+struct FVector3dCameraVariableReference
+{
+	GENERATED_BODY()
+
+	using VariableAssetType = UVector3dCameraVariable;
+
+	UPROPERTY(EditAnywhere, Category="Variable")
+	TObjectPtr<UVector3dCameraVariable> Variable;
+
+	UE_DEFINE_CAMERA_VARIABLE_REFERENCE(Vector3d)
+};
+
+USTRUCT()
+struct FVector4fCameraVariableReference
+{
+	GENERATED_BODY()
+
+	using VariableAssetType = UVector4fCameraVariable;
+
+	UPROPERTY(EditAnywhere, Category="Variable")
+	TObjectPtr<UVector4fCameraVariable> Variable;
+
+	UE_DEFINE_CAMERA_VARIABLE_REFERENCE(Vector4f)
+};
+
+USTRUCT()
+struct FVector4dCameraVariableReference
+{
+	GENERATED_BODY()
+
+	using VariableAssetType = UVector4dCameraVariable;
+
+	UPROPERTY(EditAnywhere, Category="Variable")
+	TObjectPtr<UVector4dCameraVariable> Variable;
+
+	UE_DEFINE_CAMERA_VARIABLE_REFERENCE(Vector4d)
+};
+
+USTRUCT()
+struct FRotator3fCameraVariableReference
+{
+	GENERATED_BODY()
+
+	using VariableAssetType = URotator3fCameraVariable;
+
+	UPROPERTY(EditAnywhere, Category="Variable")
+	TObjectPtr<URotator3fCameraVariable> Variable;
+
+	UE_DEFINE_CAMERA_VARIABLE_REFERENCE(Rotator3f)
+};
+
+USTRUCT()
+struct FRotator3dCameraVariableReference
+{
+	GENERATED_BODY()
+
+	using VariableAssetType = URotator3dCameraVariable;
+
+	UPROPERTY(EditAnywhere, Category="Variable")
+	TObjectPtr<URotator3dCameraVariable> Variable;
+
+	UE_DEFINE_CAMERA_VARIABLE_REFERENCE(Rotator3d)
+};
+
+USTRUCT()
+struct FTransform3fCameraVariableReference
+{
+	GENERATED_BODY()
+
+	using VariableAssetType = UTransform3fCameraVariable;
+
+	UPROPERTY(EditAnywhere, Category="Variable")
+	TObjectPtr<UTransform3fCameraVariable> Variable;
+
+	UE_DEFINE_CAMERA_VARIABLE_REFERENCE(Transform3f)
+};
+
+USTRUCT()
+struct FTransform3dCameraVariableReference
+{
+	GENERATED_BODY()
+
+	using VariableAssetType = UTransform3dCameraVariable;
+
+	UPROPERTY(EditAnywhere, Category="Variable")
+	TObjectPtr<UTransform3dCameraVariable> Variable;
+
+	UE_DEFINE_CAMERA_VARIABLE_REFERENCE(Transform3d)
+};
+
diff --git a/Engine/Plugins/Cameras/GameplayCameras/Source/GameplayCamerasEditor/Private/Customizations/CameraVariableReferenceDetailsCustomizations.cpp b/Engine/Plugins/Cameras/GameplayCameras/Source/GameplayCamerasEditor/Private/Customizations/CameraVariableReferenceDetailsCustomizations.cpp
new file mode 100644
--- /dev/null
+++ b/Engine/Plugins/Cameras/GameplayCameras/Source/GameplayCamerasEditor/Private/Customizations/CameraVariableReferenceDetailsCustomizations.cpp
@@ -0,0 +1,235 @@
+// Copyright Epic Games, Inc. All Rights Reserved.
+
+#include "Customizations/CameraVariableReferenceDetailsCustomizations.h"
+
+#include "ContentBrowserModule.h"
+#include "Core/CameraVariableAssets.h"
+#include "Core/CameraVariableCollection.h"
+#include "Core/CameraVariableReferences.h"
+#include "Editors/CameraVariablePickerConfig.h"
+#include "Framework/MultiBox/MultiBoxBuilder.h"
+#include "IGameplayCamerasEditorModule.h"
+#include "IPropertyUtilities.h"
+#include "PropertyCustomizationHelpers.h"
+#include "PropertyEditorModule.h"
+#include "ScopedTransaction.h"
+#include "Styles/GameplayCamerasEditorStyle.h"
+#include "Styling/StyleColors.h"
+#include "Widgets/Images/SImage.h"
+#include "Widgets/Input/SComboButton.h"
+#include "Widgets/SBoxPanel.h"
+#include "Widgets/Text/STextBlock.h"
+
+#define LOCTEXT_NAMESPACE "CameraVariableReferenceDetailsCustomization"
+
+namespace UE::Cameras
+{
+
+void FCameraVariableReferenceDetailsCustomization::Register(FPropertyEditorModule& PropertyEditorModule)
+{
+#define UE_CAMERA_VARIABLE_FOR_TYPE(ValueType, ValueName)\
+	PropertyEditorModule.RegisterCustomPropertyTypeLayout(\
+			F##ValueName##CameraVariableReference::StaticStruct()->GetFName(),\
+			FOnGetPropertyTypeCustomizationInstance::CreateLambda(\
+				[]{ return MakeShared<F##ValueName##CameraVariableReferenceDetailsCustomization>(); }));
+UE_CAMERA_VARIABLE_FOR_ALL_TYPES()
+#undef UE_CAMERA_VARIABLE_FOR_TYPE
+}
+
+void FCameraVariableReferenceDetailsCustomization::Unregister(FPropertyEditorModule& PropertyEditorModule)
+{
+	if (UObjectInitialized())
+	{
+#define UE_CAMERA_VARIABLE_FOR_TYPE(ValueType, ValueName)\
+		PropertyEditorModule.UnregisterCustomPropertyTypeLayout(\
+				F##ValueName##CameraVariableReference::StaticStruct()->GetFName());
+UE_CAMERA_VARIABLE_FOR_ALL_TYPES()
+#undef UE_CAMERA_VARIABLE_FOR_TYPE
+	}
+}
+
+void FCameraVariableReferenceDetailsCustomization::CustomizeHeader(TSharedRef<IPropertyHandle> PropertyHandle, FDetailWidgetRow& HeaderRow, IPropertyTypeCustomizationUtils& CustomizationUtils)
+{
+	PropertyUtilities = CustomizationUtils.GetPropertyUtilities();
+
+	StructProperty = PropertyHandle;
+	StructProperty->SetOnPropertyResetToDefault(
+			FSimpleDelegate::CreateSP(this, &FCameraVariableReferenceDetailsCustomization::OnResetToDefault));
+
+	// All references should have a "Variable" property.
+	VariableProperty = PropertyHandle->GetChildHandle("Variable");
+	ensure(VariableProperty);
+
+	VariableClass = nullptr;
+	if (FObjectProperty* VariableObjectProperty = CastField<FObjectProperty>(VariableProperty->GetProperty()))
+	{
+		VariableClass = VariableObjectProperty->PropertyClass;
+	}
+	ensure(VariableClass);
+
+	TSharedRef<FGameplayCamerasEditorStyle> GameplayCamerasStyle = FGameplayCamerasEditorStyle::Get();
+
+	HeaderRow
+	.NameContent()
+	[
+		StructProperty->CreatePropertyNameWidget()
+	]
+	.ValueContent()
+	.MinDesiredWidth(100.f)
+	.HAlign(HAlign_Fill)
+	[
+		SAssignNew(VariableBrowserButton, SComboButton)
+		.HasDownArrow(true)
+		.ContentPadding(1.f)
+		.ToolTipText(LOCTEXT("SetVariable_ToolTip", "Selects the camera variable"))
+		.ButtonContent()
+		[
+			SNew(SHorizontalBox)
+			+SHorizontalBox::Slot()
+			.AutoWidth()
+			.Padding(2.f)
+			.VAlign(VAlign_Center)
+			[
+				SNew(SImage)
+				.ColorAndOpacity(FSlateColor::UseForeground())
+				.Image(GameplayCamerasStyle->GetBrush("CameraParameter.VariableBrowser"))
+			]
+			+SHorizontalBox::Slot()
+			.AutoWidth()
+			.Padding(2.f)
+			.VAlign(VAlign_Center)
+			[
+				SNew(SBox)
+				.VAlign(VAlign_Center)
+				[
+					SNew(STextBlock)
+					.Text(this, &FCameraVariableReferenceDetailsCustomization::GetVariableName)
+				]
+			]
+		]
+		.OnGetMenuContent(this, &FCameraVariableReferenceDetailsCustomization::BuildCameraVariableBrowser)
+	];
+}
+
+void FCameraVariableReferenceDetailsCustomization::CustomizeChildren(TSharedRef<IPropertyHandle> PropertyHandle, IDetailChildrenBuilder& ChildBuilder, IPropertyTypeCustomizationUtils& CustomizationUtils)
+{
+}
+
+TSharedRef<SWidget> FCameraVariableReferenceDetailsCustomization::BuildCameraVariableBrowser()
+{
+	const bool bShouldCloseWindowAfterMenuSelection = true;
+	const bool bCloseSelfOnly = true;
+	const bool bSearchable = false;
+	FMenuBuilder MenuBuilder(true, nullptr, nullptr, bCloseSelfOnly, &FCoreStyle::Get(), bSearchable);
+
+	MenuBuilder.BeginSection(NAME_None, LOCTEXT("CameraVariableOperations", "Current Variable Reference"));
+	{
+		MenuBuilder.AddMenuEntry(
+			LOCTEXT("ClearVariable", "Clear"),
+			LOCTEXT("ClearVariable_ToolTip", "Clears the variable"),
+			FSlateIcon(FAppStyle::GetAppStyleSetName(), "GenericCommands.Delete"),
+			FUIAction(
+				FExecuteAction::CreateSP(this, &FCameraVariableReferenceDetailsCustomization::OnClearVariable),
+				FCanExecuteAction::CreateSP(this, &FCameraVariableReferenceDetailsCustomization::CanClearVariable))
+			);
+	}
+	MenuBuilder.EndSection();
+
+	UCameraVariableAsset* CommonVariable = nullptr;
+	{
+		UObject* VariableObject = nullptr;
+		FPropertyAccess::Result PropertyAccessResult = VariableProperty->GetValue(VariableObject);
+		if (PropertyAccessResult == FPropertyAccess::Success)
+		{
+			CommonVariable = Cast<UCameraVariableAsset>(VariableObject);
+		}
+	}
+
+	FCameraVariablePickerConfig PickerConfig;
+	PickerConfig.CameraVariableClass = VariableClass;
+	PickerConfig.InitialCameraVariableSelection = CommonVariable;
+	PickerConfig.CameraVariableCollectionSaveSettingsName = TEXT("CameraVariableReferencePicker");
+	PickerConfig.OnCameraVariableSelected = FOnCameraVariableSelected::CreateSP(
+			this, &FCameraVariableReferenceDetailsCustomization::OnSetVariable);
+	IGameplayCamerasEditorModule& GameplayCamerasEditorModule = IGameplayCamerasEditorModule::Get();
+	TSharedRef<SWidget> PickerWidget = GameplayCamerasEditorModule.CreateCameraVariablePicker(PickerConfig);
+
+	MenuBuilder.BeginSection(NAME_None, LOCTEXT("CameraVariableBrowser", "Browse"));
+	{
+		TSharedRef<SWidget> VariableBrowser = SNew(SBox)
+			.MinDesiredWidth(300.f)
+			.MinDesiredHeight(300.f)
+			[
+				PickerWidget
+			];
+		MenuBuilder.AddWidget(VariableBrowser, FText(), true, false);
+	}
+	MenuBuilder.EndSection();
+
+	return MenuBuilder.MakeWidget();
+}
+
+FText FCameraVariableReferenceDetailsCustomization::GetVariableName() const
+{
+	UObject* VariableObject = nullptr;
+	FPropertyAccess::Result PropertyAccessResult = VariableProperty->GetValue(VariableObject);
+	if (PropertyAccessResult == FPropertyAccess::Success)
+	{
+		if (VariableObject)
+		{
+			if (UCameraVariableAsset* Variable = Cast<UCameraVariableAsset>(VariableObject))
+			{
+				return Variable->DisplayName.IsEmpty() ?
+					FText::FromName(Variable->GetFName()) :
+					FText::FromString(Variable->DisplayName);
+			}
+			else
+			{
+				return LOCTEXT("InvalidVariableObject", "Invalid Variable");
+			}
+		}
+		return LOCTEXT("NullVariable", "None");
+	}
+	else if (PropertyAccessResult == FPropertyAccess::MultipleValues)
+	{
+		return LOCTEXT("MultipleVariableValues", "Multiple Variables");
+	}
+	else
+	{
+		return LOCTEXT("ErrorReadingVariable", "Error Reading Variable");
+	}
+}
+
+bool FCameraVariableReferenceDetailsCustomization::CanClearVariable() const
+{
+	return VariableProperty->CanResetToDefault();
+}
+
+void FCameraVariableReferenceDetailsCustomization::OnClearVariable()
+{
+	VariableProperty->ResetToDefault();
+	PropertyUtilities->RequestForceRefresh();
+}
+
+void FCameraVariableReferenceDetailsCustomization::OnSetVariable(UCameraVariableAsset* InVariable)
+{
+	{
+		FScopedTransaction Transaction(FText::Format(
+					LOCTEXT("SetPropertyValue", "Set {0}"), StructProperty->GetPropertyDisplayName()));
+
+		VariableProperty->SetValue(InVariable);
+	}
+
+	PropertyUtilities->RequestForceRefresh();
+	VariableBrowserButton->SetIsOpen(false);
+}
+
+void FCameraVariableReferenceDetailsCustomization::OnResetToDefault()
+{
+	PropertyUtilities->RequestForceRefresh();
+}
+
+}  // namespace UE::Cameras
+
+#undef LOCTEXT_NAMESPACE
+
diff --git a/Engine/Plugins/Cameras/GameplayCameras/Source/GameplayCamerasEditor/Private/GameplayCamerasEditorModule.cpp b/Engine/Plugins/Cameras/GameplayCameras/Source/GameplayCamerasEditor/Private/GameplayCamerasEditorModule.cpp
--- a/Engine/Plugins/Cameras/GameplayCameras/Source/GameplayCamerasEditor/Private/GameplayCamerasEditorModule.cpp
+++ b/Engine/Plugins/Cameras/GameplayCameras/Source/GameplayCamerasEditor/Private/GameplayCamerasEditorModule.cpp
@@ -16,6 +16,7 @@
 #include "Customizations/CameraRigAssetReferenceDetailsCustomization.h"
 #include "Customizations/CameraRigPtrDetailsCustomization.h"
 #include "Customizations/CameraProxyTableDetailsCustomization.h"
+#include "Customizations/CameraVariableReferenceDetailsCustomizations.h"
 #include "Customizations/SingleCameraDirectorDetailsCustomization.h"
 #include "Debug/CameraDebugCategories.h"
 #include "Debugger/SBlendStacksDebugPanel.h"
@@ -471,6 +472,7 @@ private:
 		FPropertyEditorModule& PropertyEditorModule = FModuleManager::GetModuleChecked<FPropertyEditorModule>("PropertyEditor");
 
 		FCameraParameterDetailsCustomization::Register(PropertyEditorModule);
+		FCameraVariableReferenceDetailsCustomization::Register(PropertyEditorModule);
 
 		PropertyEditorModule.RegisterCustomPropertyTypeLayout("CameraRigProxyTableEntry", FOnGetPropertyTypeCustomizationInstance::CreateStatic(
 					&FCameraProxyTableEntryDetailsCustomization::MakeInstance));
@@ -489,6 +491,7 @@ private:
 		if (PropertyEditorModule)
 		{
 			FCameraParameterDetailsCustomization::Unregister(*PropertyEditorModule);
+			FCameraVariableReferenceDetailsCustomization::Unregister(*PropertyEditorModule);
 
 			PropertyEditorModule->UnregisterCustomPropertyTypeLayout("CameraRigProxyTableEntry");
 			PropertyEditorModule->UnregisterCustomPropertyTypeLayout("CameraRigAsset");
diff --git a/Engine/Plugins/Cameras/GameplayCameras/Source/GameplayCamerasEditor/Public/Customizations/CameraVariableReferenceDetailsCustomizations.h b/Engine/Plugins/Cameras/GameplayCameras/Source/GameplayCamerasEditor/Public/Customizations/CameraVariableReferenceDetailsCustomizations.h
new file mode 100644
--- /dev/null
+++ b/Engine/Plugins/Cameras/GameplayCameras/Source/GameplayCamerasEditor/Public/Customizations/CameraVariableReferenceDetailsCustomizations.h
@@ -0,0 +1,68 @@
+// Copyright Epic Games, Inc. All Rights Reserved.
+
+#pragma once
+
+#include "Core/CameraVariableTableFwd.h"
+#include "IPropertyTypeCustomization.h"
+
+class FPropertyEditorModule;
+class SComboButton;
+class SWidget;
+class UCameraVariableAsset;
+
+namespace UE::Cameras
+{
+
+/**
+ * Base details customization for camera variable references.
+ */
+class FCameraVariableReferenceDetailsCustomization : public IPropertyTypeCustomization
+{
+public:
+
+	/** Registers details customizations for all camera variable reference types. */
+	static void Register(FPropertyEditorModule& PropertyEditorModule);
+	/** Unregisters details customizations for all camera variable reference types. */
+	static void Unregister(FPropertyEditorModule& PropertyEditorModule);
+
+public:
+
+	/** IPropertyTypeCustomization interface */
+	virtual void CustomizeHeader(TSharedRef<IPropertyHandle> PropertyHandle, FDetailWidgetRow& HeaderRow, IPropertyTypeCustomizationUtils& CustomizationUtils) override;
+	virtual void CustomizeChildren(TSharedRef<IPropertyHandle> PropertyHandle, IDetailChildrenBuilder& ChildBuilder, IPropertyTypeCustomizationUtils& CustomizationUtils) override;
+
+private:
+
+	TSharedRef<SWidget> BuildCameraVariableBrowser();
+
+	FText GetVariableName() const;
+
+	bool CanClearVariable() const;
+	void OnClearVariable();
+
+	void OnSetVariable(UCameraVariableAsset* InVariable);
+	void OnResetToDefault();
+
+protected:
+
+	UClass* VariableClass = nullptr;
+
+	TSharedPtr<IPropertyUtilities> PropertyUtilities;
+
+	TSharedPtr<IPropertyHandle> StructProperty;
+	TSharedPtr<IPropertyHandle> VariableProperty;
+
+	TSharedPtr<SComboButton> VariableBrowserButton;
+};
+
+// Create all the individual classes.
+#define UE_CAMERA_VARIABLE_FOR_TYPE(ValueType, ValueName)\
+class F##ValueName##CameraVariableReferenceDetailsCustomization : public FCameraVariableReferenceDetailsCustomization\
+{\
+};
+UE_CAMERA_VARIABLE_FOR_ALL_TYPES()
+#undef UE_CAMERA_VARIABLE_FOR_TYPE
+
+}  // namespace UE::Cameras
+ 
+
