Cameras: fix collision node's safe position lookup

Always using the pivot and/or the owner evaluation context doesn't work well when the collision node is in the global stack. This change adds an option to specify where the safe position should be by default, and in particular allows getting it from the active context (from the main blend stack) or from the player pawn. There is also a new option to activate/deactivate collision in a smooth way.

Debug display has been improved, and potential log spam reduced.

diff --git a/Engine/Plugins/Cameras/GameplayCameras/Source/GameplayCameras/Private/Nodes/Collision/CollisionPushCameraNode.cpp b/Engine/Plugins/Cameras/GameplayCameras/Source/GameplayCameras/Private/Nodes/Collision/CollisionPushCameraNode.cpp
--- a/Engine/Plugins/Cameras/GameplayCameras/Source/GameplayCameras/Private/Nodes/Collision/CollisionPushCameraNode.cpp
+++ b/Engine/Plugins/Cameras/GameplayCameras/Source/GameplayCameras/Private/Nodes/Collision/CollisionPushCameraNode.cpp
@@ -10,6 +10,7 @@
 #include "Core/CameraValueInterpolator.h"
 #include "Debug/CameraDebugBlock.h"
 #include "Debug/CameraDebugBlockBuilder.h"
+#include "Debug/CameraDebugColors.h"
 #include "Debug/CameraDebugRenderer.h"
 #include "Engine/Engine.h"
 #include "Engine/HitResult.h"
@@ -40,10 +41,20 @@ protected:
 
 private:
 
-	TOptional<FVector3d> GetSafePosition(const FCameraNodeEvaluationParams& Params, const FCameraNodeEvaluationResult& OutResult);
+	TOptional<FVector3d> GetFinalSafePosition(
+			const FCameraNodeEvaluationParams& Params, 
+			const FCameraNodeEvaluationResult& OutResult);
+	TOptional<FVector3d> GetSafePosition(
+			const FCameraNodeEvaluationParams& Params, 
+			TSharedPtr<const FCameraEvaluationContext> ActiveContext,
+			const FCameraRigJoint* PivotJoint,
+			const FCameraNodeEvaluationResult& OutResult);
+
 	void RunCollisionTrace(UWorld* World, APlayerController* PlayerController, const FVector3d& SafePosition, const FCameraNodeEvaluationParams& Params, FCameraNodeEvaluationResult& OutResult);
 	void HandleAsyncCollisionTraceResult(UWorld* World, const FVector3d& SafePosition, const FCameraNodeEvaluationParams& Params, FCameraNodeEvaluationResult& OutResult);
 	void HandleCollisionTraceResult(UWorld* World, TArrayView<const FHitResult> HitResults, const FVector3d& SafePosition, const FCameraNodeEvaluationParams& Params, FCameraNodeEvaluationResult& OutResult);
+	void HandleDisabledCollision(const FVector3d& SafePosition, const FCameraNodeEvaluationParams& Params, FCameraNodeEvaluationResult& OutResult);
+	void UpdatePushFactor(bool bFoundHit, float CurrentPushFactor, const FVector3d& SafePosition, const FCameraNodeEvaluationParams& Params, FCameraNodeEvaluationResult& OutResult);
 
 private:
 
@@ -62,7 +73,10 @@ private:
 	ECameraCollisionDirection LastDirection = ECameraCollisionDirection::Pushing;
 
 #if UE_GAMEPLAY_CAMERAS_DEBUG
+	bool bDebugCollisionEnabled = false;
 	bool bDebugFoundHit = false;
+	bool bDebugGotSafePosition = false;
+	bool bDebugGotSafePositionOffsetSpace = false;
 	FString DebugHitObjectName;
 	FVector3d DebugSafePosition;
 #endif
@@ -71,6 +85,11 @@ private:
 UE_DEFINE_CAMERA_NODE_EVALUATOR(FCollisionPushCameraNodeEvaluator)
 
 UE_DECLARE_CAMERA_DEBUG_BLOCK_START(GAMEPLAYCAMERAS_API, FCollisionPushCameraDebugBlock)
+	UE_DECLARE_CAMERA_DEBUG_BLOCK_FIELD(bool, bCollisionEnabled)
+	UE_DECLARE_CAMERA_DEBUG_BLOCK_FIELD(bool, bGotSafePosition)
+	UE_DECLARE_CAMERA_DEBUG_BLOCK_FIELD(bool, bGotSafePositionOffsetSpace)
+	UE_DECLARE_CAMERA_DEBUG_BLOCK_FIELD(ECollisionSafePosition, SafePositionType)
+	UE_DECLARE_CAMERA_DEBUG_BLOCK_FIELD(ECollisionSafePositionOffsetSpace, SafePositionOffsetSpace)
 	UE_DECLARE_CAMERA_DEBUG_BLOCK_FIELD(float, PushFactor);
 	UE_DECLARE_CAMERA_DEBUG_BLOCK_FIELD(float, DampedPushFactor);
 	UE_DECLARE_CAMERA_DEBUG_BLOCK_FIELD(bool, bIsPulling)
@@ -110,13 +129,26 @@ void FCollisionPushCameraNodeEvaluator::OnRun(const FCameraNodeEvaluationParams&
 		return;
 	}
 
-	TOptional<FVector3d> SafePosition = GetSafePosition(Params, OutResult);
+	// Get the safe position... bail out if we don't have any.
+	TOptional<FVector3d> SafePosition = GetFinalSafePosition(Params, OutResult);
 	if (!SafePosition.IsSet())
 	{
-		const UCameraNode* ThisNode = GetCameraNode();
-		UE_LOG(LogCameraSystem, Error, 
-				TEXT("Can't find safe position for CollisionPushCameraNode '%s'."),
-				*GetNameSafe(ThisNode));
+		return;
+	}
+
+	// See if collision is enabled. If not, handle it as if we didn't collide with anything.
+	bool bEnableCollision = true;
+	const UCollisionPushCameraNode* ThisNode = GetCameraNodeAs<UCollisionPushCameraNode>();
+	if (const UBooleanCameraVariable* EnableCollisionVariable = ThisNode->EnableCollision.Get())
+	{
+		bEnableCollision = OutResult.VariableTable.GetValue(EnableCollisionVariable);
+	}
+#if UE_GAMEPLAY_CAMERAS_DEBUG
+	bDebugCollisionEnabled = bEnableCollision;
+#endif
+	if (!bEnableCollision)
+	{
+		HandleDisabledCollision(SafePosition.GetValue(), Params, OutResult);
 		return;
 	}
 
@@ -133,104 +165,207 @@ void FCollisionPushCameraNodeEvaluator::OnRun(const FCameraNodeEvaluationParams&
 		return;
 	}
 
+	// Actually run some collision tests!
 	HandleAsyncCollisionTraceResult(World, SafePosition.GetValue(), Params, OutResult);
 	RunCollisionTrace(World, PlayerController, SafePosition.GetValue(), Params, OutResult);
 }
 
-TOptional<FVector3d> FCollisionPushCameraNodeEvaluator::GetSafePosition(const FCameraNodeEvaluationParams& Params, const FCameraNodeEvaluationResult& OutResult)
+TOptional<FVector3d> FCollisionPushCameraNodeEvaluator::GetFinalSafePosition(const FCameraNodeEvaluationParams& Params, const FCameraNodeEvaluationResult& OutResult)
 {
-	FTransform3d SafePositionTransform;
-	bool bGotSafePosition = false;
+	if (!ensure(Params.Evaluator))
+	{
+		return TOptional<FVector3d>();
+	}
 
-	// See if we can use the camera rig pivot as the safe position.
-	// Otherwise, use the context's initial position.
+	if (!ensure(Params.EvaluationContext))
+	{
+		return TOptional<FVector3d>();
+	}
+
+	const UCollisionPushCameraNode* ThisNode = GetCameraNodeAs<UCollisionPushCameraNode>();
+
+	// Get the active context.
+	const FCameraEvaluationContextStack& ContextStack = Params.Evaluator->GetEvaluationContextStack();
+	TSharedPtr<const FCameraEvaluationContext> ActiveContext = ContextStack.GetActiveContext();
+	
+	// See if we have a pivot.
 	const FBuiltInCameraVariables& BuiltInVariables = FBuiltInCameraVariables::Get();
 	TArrayView<const FCameraRigJoint> Joints = OutResult.CameraRigJoints.GetJoints();
 	const FCameraRigJoint* PivotJoint = Joints.FindByPredicate([&BuiltInVariables](const FCameraRigJoint& Joint)
 			{
 				return Joint.VariableID == BuiltInVariables.YawPitchDefinition;
 			});
-	if (PivotJoint)
+
+	// Get the safe position itself first.
+	TOptional<FVector3d> OptSafePosition = GetSafePosition(Params, ActiveContext, PivotJoint, OutResult);	
+	if (!OptSafePosition.IsSet())
 	{
-		SafePositionTransform = PivotJoint->Transform;
-		bGotSafePosition = true;
+		return TOptional<FVector3d>();
 	}
-	else if (Params.EvaluationContext)
+
+	// Apply the offset in the specified space.
+	const FVector3d SafePositionOffset = SafePositionOffsetReader.Get(OutResult.VariableTable);
+
+	FVector3d WorldSafePositionOffset = SafePositionOffset;
+	bool bGotSafePositionOffsetSpace = false;
+
+	switch (ThisNode->SafePositionOffsetSpace)
 	{
-		const FCameraNodeEvaluationResult& InitialResult = Params.EvaluationContext->GetInitialResult();
-		if (InitialResult.bIsValid)
-		{
-			SafePositionTransform = InitialResult.CameraPose.GetTransform();
-			bGotSafePosition = true;
-		}
+		case ECollisionSafePositionOffsetSpace::ActiveContext:
+			if (ActiveContext)
+			{
+				const FCameraNodeEvaluationResult& InitialResult = ActiveContext->GetInitialResult();
+				ensure(InitialResult.bIsValid);
+				WorldSafePositionOffset = InitialResult.CameraPose.GetRotation().RotateVector(SafePositionOffset);
+				bGotSafePositionOffsetSpace = true;
+			}
+			break;
+		case ECollisionSafePositionOffsetSpace::OwningContext:
+			{
+				const FCameraNodeEvaluationResult& InitialResult = Params.EvaluationContext->GetInitialResult();
+				ensure(InitialResult.bIsValid);
+				WorldSafePositionOffset = InitialResult.CameraPose.GetRotation().RotateVector(SafePositionOffset);
+				bGotSafePositionOffsetSpace = true;
+			}
+			break;
+		case ECollisionSafePositionOffsetSpace::Pivot:
+			if (PivotJoint)
+			{
+				WorldSafePositionOffset = PivotJoint->Transform.TransformVectorNoScale(SafePositionOffset);
+				bGotSafePositionOffsetSpace = true;
+			}
+			else if (ActiveContext)
+			{
+				const FCameraNodeEvaluationResult& InitialResult = ActiveContext->GetInitialResult();
+				ensure(InitialResult.bIsValid);
+				WorldSafePositionOffset = InitialResult.CameraPose.GetRotation().RotateVector(SafePositionOffset);
+				bGotSafePositionOffsetSpace = true;
+			}
+			break;
+		case ECollisionSafePositionOffsetSpace::CameraPose:
+			{
+				WorldSafePositionOffset = OutResult.CameraPose.GetRotation().RotateVector(SafePositionOffset);
+				bGotSafePositionOffsetSpace = true;
+			}
+			break;
+		case ECollisionSafePositionOffsetSpace::Pawn:
+			if (ActiveContext)
+			{
+				if (APlayerController* PlayerController = ActiveContext->GetPlayerController())
+				{
+					if (APawn* Pawn = PlayerController->GetPawnOrSpectator())
+					{
+						const FTransform3d& PawnTransform = Pawn->GetActorTransform();
+						WorldSafePositionOffset = PawnTransform.TransformVectorNoScale(SafePositionOffset);
+						bGotSafePositionOffsetSpace = true;
+					}
+				}
+			}
+			break;
 	}
 
-	// Apply the offset in the specified space.
-	if (bGotSafePosition)
+#if UE_GAMEPLAY_CAMERAS_DEBUG
+	bDebugGotSafePositionOffsetSpace = bGotSafePositionOffsetSpace;
+#endif
+
+	if (bGotSafePositionOffsetSpace)
 	{
-		const UCollisionPushCameraNode* ThisNode = GetCameraNodeAs<UCollisionPushCameraNode>();
-		const FVector3d SafePositionOffset = SafePositionOffsetReader.Get(OutResult.VariableTable);
+		return OptSafePosition.GetValue() + WorldSafePositionOffset;
+	}
+	else
+	{
+		return OptSafePosition.GetValue();
+	}
+}
 
-		FVector3d WorldSafePositionOffset = SafePositionOffset;
-		switch (ThisNode->SafePositionOffsetSpace)
+TOptional<FVector3d> FCollisionPushCameraNodeEvaluator::GetSafePosition(
+		const FCameraNodeEvaluationParams& Params, 
+		TSharedPtr<const FCameraEvaluationContext> ActiveContext,
+		const FCameraRigJoint* PivotJoint,
+		const FCameraNodeEvaluationResult& OutResult)
+{
+	const UCollisionPushCameraNode* ThisNode = GetCameraNodeAs<UCollisionPushCameraNode>();
+
+#if UE_GAMEPLAY_CAMERAS_DEBUG
+	bDebugGotSafePosition = false;
+#endif
+
+	// See if we have a custom safe position.
+	if (UVector3dCameraVariable* CustomSafePositionVariable = ThisNode->CustomSafePosition.Get())
+	{
+		FVector3d SafePosition;
+		if (OutResult.VariableTable.TryGetValue(CustomSafePositionVariable, SafePosition))
 		{
-			case ECollisionSafePositionOffsetSpace::ActiveContext:
-				if (Params.Evaluator)
+#if UE_GAMEPLAY_CAMERAS_DEBUG
+			bDebugGotSafePosition = false;
+#endif
+			return TOptional<FVector3d>(SafePosition);
+		}
+	}
+	
+	// Compute the base safe position.
+	FVector3d SafePosition;
+	bool bGotSafePosition = false;
+
+	switch (ThisNode->SafePosition)
+	{
+		case ECollisionSafePosition::ActiveContext:
+			if (ActiveContext)
+			{
+				const FCameraNodeEvaluationResult& InitialResult = ActiveContext->GetInitialResult();
+				ensure(InitialResult.bIsValid);
+				SafePosition = InitialResult.CameraPose.GetLocation();
+				bGotSafePosition = true;
+			}
+			break;
+		case ECollisionSafePosition::OwningContext:
+			{
+				const FCameraNodeEvaluationResult& InitialResult = Params.EvaluationContext->GetInitialResult();
+				ensure(InitialResult.bIsValid);
+				SafePosition = InitialResult.CameraPose.GetLocation();
+				bGotSafePosition = true;
+			}
+			break;
+		case ECollisionSafePosition::Pivot:
+			if (PivotJoint)
+			{
+				SafePosition = PivotJoint->Transform.GetLocation();
+				bGotSafePosition = true;
+			}
+			else if (ActiveContext)
+			{
+				const FCameraNodeEvaluationResult& InitialResult = ActiveContext->GetInitialResult();
+				ensure(InitialResult.bIsValid);
+				SafePosition = InitialResult.CameraPose.GetLocation();
+				bGotSafePosition = true;
+			}
+			break;
+		case ECollisionSafePosition::Pawn:
+			if (ActiveContext)
+			{
+				if (APlayerController* PlayerController = ActiveContext->GetPlayerController())
 				{
-					const FCameraEvaluationContextStack& ContextStack = Params.Evaluator->GetEvaluationContextStack();
-					TSharedPtr<const FCameraEvaluationContext> ActiveContext = ContextStack.GetActiveContext();
-					if (ActiveContext)
+					if (APawn* Pawn = PlayerController->GetPawnOrSpectator())
 					{
-						const FCameraNodeEvaluationResult& InitialResult = Params.EvaluationContext->GetInitialResult();
-						ensure(InitialResult.bIsValid);
-						WorldSafePositionOffset = InitialResult.CameraPose.GetRotation().RotateVector(SafePositionOffset);
+						SafePosition = Pawn->GetActorTransform().GetLocation();
+						bGotSafePosition = true;
 					}
-					else
-					{
-						UE_LOG(LogCameraSystem, Warning, 
-								TEXT("Can't offset safe collision position in active context space because "
-									"no active context is active on the camera system evaluator for node '%s'!"),
-								*GetNameSafe(ThisNode));
-					}
-				}
-				else
-				{
-					UE_LOG(LogCameraSystem, Warning, 
-							TEXT("Can't offset safe collision position in active context space because "
-								"no camera system evaluator was given for node '%s'!"),
-							*GetNameSafe(ThisNode));
-				}
-			case ECollisionSafePositionOffsetSpace::OwningContext:
-				if (Params.EvaluationContext)
-				{
-					const FCameraNodeEvaluationResult& InitialResult = Params.EvaluationContext->GetInitialResult();
-					ensure(InitialResult.bIsValid);
-					WorldSafePositionOffset = InitialResult.CameraPose.GetRotation().RotateVector(SafePositionOffset);
-				}
-				else
-				{
-					UE_LOG(LogCameraSystem, Warning, 
-							TEXT("Can't offset safe collision position in owning context space because "
-								"there is no owning context for node '%s'!"),
-							*GetNameSafe(ThisNode));
-				}
-				break;
-			case ECollisionSafePositionOffsetSpace::Pivot:
-				if (PivotJoint)
-				{
-					WorldSafePositionOffset = PivotJoint->Transform.TransformVectorNoScale(SafePositionOffset);
 				}
-				break;
-			case ECollisionSafePositionOffsetSpace::CameraPose:
-				{
-					WorldSafePositionOffset = OutResult.CameraPose.GetRotation().RotateVector(SafePositionOffset);
-				}
-				break;
-		}
+			}
+			break;
+	}
 
-		return SafePositionTransform.GetLocation() + WorldSafePositionOffset;
+#if UE_GAMEPLAY_CAMERAS_DEBUG
+	bDebugGotSafePosition = bGotSafePosition;
+#endif
+	if (bGotSafePosition)
+	{
+		return TOptional<FVector3d>(SafePosition);
+	}
+	else
+	{
+		return TOptional<FVector3d>();
 	}
-	return TOptional<FVector3d>();
 }
 
 void FCollisionPushCameraNodeEvaluator::RunCollisionTrace(UWorld* World, APlayerController* PlayerController, const FVector3d& SafePosition, const FCameraNodeEvaluationParams& Params, FCameraNodeEvaluationResult& OutResult)
@@ -280,6 +415,8 @@ void FCollisionPushCameraNodeEvaluator::RunCollisionTrace(UWorld* World, APlayer
 			SweepShape,
 			QueryParams,
 			FCollisionResponseParams::DefaultResponseParam);
+
+		// Handle results right away.
 		HandleCollisionTraceResult(World, HitResults, SafePosition, Params, OutResult);
 	}
 }
@@ -344,6 +481,21 @@ void FCollisionPushCameraNodeEvaluator::HandleCollisionTraceResult(UWorld* World
 		}
 	}
 
+	UpdatePushFactor(bFoundHit, CurrentPushFactor, SafePosition, Params, OutResult);
+}
+
+void FCollisionPushCameraNodeEvaluator::HandleDisabledCollision(const FVector3d& SafePosition, const FCameraNodeEvaluationParams& Params, FCameraNodeEvaluationResult& OutResult)
+{
+#if UE_GAMEPLAY_CAMERAS_DEBUG
+	DebugHitObjectName.Empty();
+	DebugSafePosition = SafePosition;
+#endif
+
+	UpdatePushFactor(false, 0.f, SafePosition, Params, OutResult);
+}
+
+void FCollisionPushCameraNodeEvaluator::UpdatePushFactor(bool bFoundHit, float CurrentPushFactor, const FVector3d& SafePosition, const FCameraNodeEvaluationParams& Params, FCameraNodeEvaluationResult& OutResult)
+{
 #if UE_GAMEPLAY_CAMERAS_DEBUG
 	bDebugFoundHit = bFoundHit;
 #endif
@@ -409,6 +561,14 @@ void FCollisionPushCameraNodeEvaluator::OnBuildDebugBlocks(const FCameraDebugBlo
 {
 	FCollisionPushCameraDebugBlock& DebugBlock = Builder.AttachDebugBlock<FCollisionPushCameraDebugBlock>();
 
+	DebugBlock.bCollisionEnabled = bDebugCollisionEnabled;
+
+	const UCollisionPushCameraNode* ThisNode = GetCameraNodeAs<UCollisionPushCameraNode>();
+	DebugBlock.bGotSafePosition = bDebugGotSafePosition;
+	DebugBlock.bGotSafePositionOffsetSpace = bDebugGotSafePositionOffsetSpace;
+	DebugBlock.SafePositionType = ThisNode->SafePosition;
+	DebugBlock.SafePositionOffsetSpace = ThisNode->SafePositionOffsetSpace;
+
 	DebugBlock.bFoundHit = bDebugFoundHit;
 	DebugBlock.HitObjectName = DebugHitObjectName;
 	DebugBlock.SafePosition = DebugSafePosition;
@@ -419,10 +579,16 @@ void FCollisionPushCameraNodeEvaluator::OnBuildDebugBlocks(const FCameraDebugBlo
 
 void FCollisionPushCameraDebugBlock::OnDebugDraw(const FCameraDebugBlockDrawParams& Params, FCameraDebugRenderer& Renderer)
 {
+	if (!bCollisionEnabled)
+	{
+		Renderer.AddText(TEXT("collision disabled"));
+		return;
+	}
+
 	if (DampedPushFactor > 0)
 	{
 		Renderer.AddText(TEXT("need to push by %.2f%%, currently %.2f%% [%s]"),
-				PushFactor, DampedPushFactor,
+				PushFactor * 100.f, DampedPushFactor * 100.f,
 				bIsPulling ? TEXT("pulling") : TEXT("pushing"));
 		if (bFoundHit)
 		{
@@ -434,7 +600,57 @@ void FCollisionPushCameraDebugBlock::OnDebugDraw(const FCameraDebugBlockDrawPara
 		Renderer.AddText(TEXT("not pushing"));
 	}
 
-	Renderer.DrawText(SafePosition, TEXT("Safe Position"), FLinearColor::Gray, GEngine->GetTinyFont());
+	const FCameraDebugColors& Colors = FCameraDebugColors::Get();
+	if (!bGotSafePosition)
+	{
+		Renderer.NewLine();
+		Renderer.SetTextColor(Colors.Error);
+		Renderer.AddText(TEXT("can't get safe position: "));
+		switch (SafePositionType)
+		{
+			case ECollisionSafePosition::ActiveContext:
+				Renderer.AddText(TEXT("no active context"));
+				break;
+			case ECollisionSafePosition::Pivot:
+				Renderer.AddText(TEXT("no pivot nor active context"));
+				break;
+			case ECollisionSafePosition::Pawn:
+				Renderer.AddText(TEXT("no active context nor player controller"));
+				break;
+			default:
+				Renderer.AddText(TEXT("unknown error"));
+				break;
+		}
+		Renderer.SetTextColor(Colors.Default);
+	}
+
+	if (!bGotSafePositionOffsetSpace)
+	{
+		Renderer.NewLine();
+		Renderer.SetTextColor(Colors.Error);
+		Renderer.AddText(TEXT("can't get safe position offset space: "));
+		switch (SafePositionOffsetSpace)
+		{
+			case ECollisionSafePositionOffsetSpace::ActiveContext:
+				Renderer.AddText(TEXT("no active context"));
+				break;
+			case ECollisionSafePositionOffsetSpace::Pivot:
+				Renderer.AddText(TEXT("no pivot nor active context"));
+				break;
+			case ECollisionSafePositionOffsetSpace::Pawn:
+				Renderer.AddText(TEXT("no active context nor player controller"));
+				break;
+			default:
+				Renderer.AddText(TEXT("unknown error"));
+				break;
+		}
+		Renderer.SetTextColor(Colors.Default);
+	}
+
+	if (bGotSafePosition)
+	{
+		Renderer.DrawText(SafePosition, TEXT("Safe Position"), FLinearColor::Gray, GEngine->GetTinyFont());
+	}
 }
 
 #endif  // UE_GAMEPLAY_CAMERAS_DEBUG
diff --git a/Engine/Plugins/Cameras/GameplayCameras/Source/GameplayCameras/Public/Core/CameraVariableTable.h b/Engine/Plugins/Cameras/GameplayCameras/Source/GameplayCameras/Public/Core/CameraVariableTable.h
--- a/Engine/Plugins/Cameras/GameplayCameras/Source/GameplayCameras/Public/Core/CameraVariableTable.h
+++ b/Engine/Plugins/Cameras/GameplayCameras/Source/GameplayCameras/Public/Core/CameraVariableTable.h
@@ -275,7 +275,7 @@ ValueType FCameraVariableTable::GetValue(FCameraVariableID VariableID, typename
 template<typename VariableAssetType>
 typename VariableAssetType::ValueType FCameraVariableTable::GetValue(const VariableAssetType* VariableAsset) const
 {
-	return GetValue(VariableAsset->GetVariableID(), VariableAsset->GetDefaultValue());
+	return GetValue<typename VariableAssetType::ValueType>(VariableAsset->GetVariableID(), VariableAsset->GetDefaultValue());
 }
 
 template<typename ValueType>
diff --git a/Engine/Plugins/Cameras/GameplayCameras/Source/GameplayCameras/Public/Nodes/Collision/CollisionPushCameraNode.h b/Engine/Plugins/Cameras/GameplayCameras/Source/GameplayCameras/Public/Nodes/Collision/CollisionPushCameraNode.h
--- a/Engine/Plugins/Cameras/GameplayCameras/Source/GameplayCameras/Public/Nodes/Collision/CollisionPushCameraNode.h
+++ b/Engine/Plugins/Cameras/GameplayCameras/Source/GameplayCameras/Public/Nodes/Collision/CollisionPushCameraNode.h
@@ -4,6 +4,7 @@
 
 #include "Core/CameraNode.h"
 #include "Core/CameraParameters.h"
+#include "Core/CameraVariableReferences.h"
 #include "Nodes/CameraNodeTypes.h"
 #include "Engine/EngineTypes.h"
 
@@ -11,13 +12,49 @@
 
 class UCameraValueInterpolator;
 
+/**
+ * Specifies how to compute the default safe position for the collision camera node
+ * to push towards.
+ */
+UENUM()
+enum class ECollisionSafePosition : uint8
+{
+	/**
+	 * The initial result location of the active evaluation context on the main 
+	 * layer's blend stack.
+	 */
+	ActiveContext,
+	/**
+	 * The initial result location of the evaluation context of the collision camera node.
+	 */
+	OwningContext,
+	/**
+	 * The current pivot. If no pivot is found, fallback to ActiveContext.
+	 */
+	Pivot,
+	/**
+	 * The location of the player's controlled pawn.
+	 */
+	Pawn
+};
+
+/**
+ * Describes the coordinate system in which to offset the collision camera node's
+ * safe position.
+ */
 UENUM()
-enum class ECollisionSafePositionOffsetSpace
+enum class ECollisionSafePositionOffsetSpace : uint8
 {
+	/** The space of the active evaluation context on the main layer's blend stack. */
 	ActiveContext,
+	/** The space of the evaluation context of the collision camera node. */
 	OwningContext,
+	/** The space of the current pivot. If no pivot is found, fallback to ActiveContext. */
 	Pivot,
-	CameraPose
+	/** The local space of the current camera pose. */
+	CameraPose,
+	/** The space of the player's controlled pawn. */
+	Pawn
 };
 
 /**
@@ -32,22 +69,41 @@ class UCollisionPushCameraNode : public UCameraNode
 
 public:
 
-	/** Radius of the sphere used for collision testing. */
-	UPROPERTY(EditAnywhere, Category="Collision")
-	FFloatCameraParameter CollisionSphereRadius;
+	/** How to compute the safe position. */
+	UPROPERTY(EditAnywhere, Category="Safe Position")
+	ECollisionSafePosition SafePosition = ECollisionSafePosition::Pivot;
 
-	/** Collision channel to use for the line trace. */
-	UPROPERTY(EditAnywhere, Category="Collision")
-	TEnumAsByte<ECollisionChannel> CollisionChannel;
+	/**
+	 * An optional camera variable to query for a safe position. If null, or if the variable
+	 * isn't set, fallback to the value defined by SafePosition.
+	 */
+	UPROPERTY(EditAnywhere, Category="Safe Position")
+	FVector3dCameraVariableReference CustomSafePosition;
 
 	/** World-space offset from the target to the line trace's end. */
-	UPROPERTY(EditAnywhere, Category="Occlusion")
+	UPROPERTY(EditAnywhere, Category="Safe Position")
 	FVector3dCameraParameter SafePositionOffset;
 
 	/** What space the safe position offset should be in. */
-	UPROPERTY(EditAnywhere, Category=Damping)
+	UPROPERTY(EditAnywhere, Category="Safe Position")
 	ECollisionSafePositionOffsetSpace SafePositionOffsetSpace = ECollisionSafePositionOffsetSpace::Pivot;
 
+	/**
+	 * An optional boolean camera variable that specifies whether collision should be enabled.
+	 * When enabled/disabled, the collision push amount will interpolate as per the PushInterpolator
+	 * and PullInterpolator.
+	 */
+	UPROPERTY(EditAnywhere, Category="Collision")
+	FBooleanCameraVariableReference EnableCollision;
+
+	/** Radius of the sphere used for collision testing. */
+	UPROPERTY(EditAnywhere, Category="Collision")
+	FFloatCameraParameter CollisionSphereRadius;
+
+	/** Collision channel to use for the line trace. */
+	UPROPERTY(EditAnywhere, Category="Collision")
+	TEnumAsByte<ECollisionChannel> CollisionChannel;
+
 	/** The interpolation to use when pushing the camera towards the safe position. */
 	UPROPERTY(EditAnywhere, Category="Collision")
 	TObjectPtr<UCameraValueInterpolator> PushInterpolator;
