Cameras: camera variable table wrapper for Blueprints now handles registering auto-reset variables

Note that the auto-reset service now uses weak pointers, with auto-cleanup of dead pointers, since, unlike camera rigs, we have no way to know when variables used by Blueprint are not used anymore.

diff --git a/Engine/Plugins/Cameras/GameplayCameras/Source/GameplayCameras/Private/Core/CameraEvaluationContext.cpp b/Engine/Plugins/Cameras/GameplayCameras/Source/GameplayCameras/Private/Core/CameraEvaluationContext.cpp
--- a/Engine/Plugins/Cameras/GameplayCameras/Source/GameplayCameras/Private/Core/CameraEvaluationContext.cpp
+++ b/Engine/Plugins/Cameras/GameplayCameras/Source/GameplayCameras/Private/Core/CameraEvaluationContext.cpp
@@ -100,6 +100,8 @@ void FCameraEvaluationContext::Activate(const FCameraEvaluationContextActivatePa
 		return;
 	}
 
+	CameraSystemEvaluator = Params.Evaluator;
+
 	OnActivate(Params);
 
 	AutoCreateDirectorEvaluator();
@@ -134,6 +136,8 @@ void FCameraEvaluationContext::Deactivate(const FCameraEvaluationContextDeactiva
 
 	OnDeactivate(Params);
 
+	CameraSystemEvaluator = nullptr;
+
 	bActivated = false;
 }
 
diff --git a/Engine/Plugins/Cameras/GameplayCameras/Source/GameplayCameras/Private/Core/CameraSystemEvaluator.cpp b/Engine/Plugins/Cameras/GameplayCameras/Source/GameplayCameras/Private/Core/CameraSystemEvaluator.cpp
--- a/Engine/Plugins/Cameras/GameplayCameras/Source/GameplayCameras/Private/Core/CameraSystemEvaluator.cpp
+++ b/Engine/Plugins/Cameras/GameplayCameras/Source/GameplayCameras/Private/Core/CameraSystemEvaluator.cpp
@@ -149,6 +149,18 @@ void FCameraSystemEvaluator::GetEvaluationServices(TArray<TSharedPtr<FCameraEval
 	OutEvaluationServices = EvaluationServices;
 }
 
+TSharedPtr<FCameraEvaluationService> FCameraSystemEvaluator::FindEvaluationService(const FCameraObjectTypeID& TypeID) const
+{
+	for (TSharedPtr<FCameraEvaluationService> EvaluationService : EvaluationServices)
+	{
+		if (EvaluationService.Get()->IsKindOf(TypeID))
+		{
+			return EvaluationService;
+		}
+	}
+	return nullptr;
+}
+
 void FCameraSystemEvaluator::NotifyRootCameraNodeEvent(const FRootCameraNodeCameraRigEvent& InEvent)
 {
 	for (TSharedPtr<FCameraEvaluationService> EvaluationService : EvaluationServices)
diff --git a/Engine/Plugins/Cameras/GameplayCameras/Source/GameplayCameras/Private/Directors/BlueprintCameraDirector.cpp b/Engine/Plugins/Cameras/GameplayCameras/Source/GameplayCameras/Private/Directors/BlueprintCameraDirector.cpp
--- a/Engine/Plugins/Cameras/GameplayCameras/Source/GameplayCameras/Private/Directors/BlueprintCameraDirector.cpp
+++ b/Engine/Plugins/Cameras/GameplayCameras/Source/GameplayCameras/Private/Directors/BlueprintCameraDirector.cpp
@@ -15,6 +15,7 @@
 #include "GameFramework/Actor.h"
 #include "GameFramework/ControllerGameplayCameraEvaluationComponent.h"
 #include "GameplayCameras.h"
+#include "Services/AutoResetCameraVariableService.h"
 
 #include UE_INLINE_GENERATED_CPP_BY_NAME(BlueprintCameraDirector)
 
@@ -338,9 +339,12 @@ FBlueprintCameraPose UBlueprintCameraDirectorEvaluator::GetInitialContextCameraP
 
 FBlueprintCameraVariableTable UBlueprintCameraDirectorEvaluator::GetInitialContextVariableTable() const
 {
+	using namespace UE::Cameras;
+
 	if (EvaluationContext)
 	{
-		return FBlueprintCameraVariableTable(&EvaluationContext->GetInitialResult().VariableTable);
+		FCameraVariableTable& VariableTable = EvaluationContext->GetInitialResult().VariableTable;
+		return FBlueprintCameraVariableTable(&VariableTable, VariableAutoResetService);
 	}
 	else
 	{
@@ -367,7 +371,10 @@ void UBlueprintCameraDirectorEvaluator::SetInitialContextCameraPose(const FBluep
 
 void UBlueprintCameraDirectorEvaluator::NativeActivateCameraDirector(const UE::Cameras::FCameraDirectorActivateParams& Params)
 {
+	using namespace UE::Cameras;
+
 	EvaluationContext = Params.OwnerContext;
+	VariableAutoResetService = Params.Evaluator->FindEvaluationService<FAutoResetCameraVariableService>();
 
 	EvaluationResult.Reset();
 	{
@@ -394,6 +401,7 @@ void UBlueprintCameraDirectorEvaluator::NativeDeactivateCameraDirector(const UE:
 		DeactivateCameraDirector(BlueprintParams);
 	}
 
+	VariableAutoResetService = nullptr;
 	EvaluationContext = nullptr;
 }
 
diff --git a/Engine/Plugins/Cameras/GameplayCameras/Source/GameplayCameras/Private/GameFramework/BlueprintCameraVariableTable.cpp b/Engine/Plugins/Cameras/GameplayCameras/Source/GameplayCameras/Private/GameFramework/BlueprintCameraVariableTable.cpp
--- a/Engine/Plugins/Cameras/GameplayCameras/Source/GameplayCameras/Private/GameFramework/BlueprintCameraVariableTable.cpp
+++ b/Engine/Plugins/Cameras/GameplayCameras/Source/GameplayCameras/Private/GameFramework/BlueprintCameraVariableTable.cpp
@@ -4,6 +4,7 @@
 
 #include "Core/CameraVariableAssets.h"
 #include "Core/CameraVariableTable.h"
+#include "Services/AutoResetCameraVariableService.h"
 
 #include UE_INLINE_GENERATED_CPP_BY_NAME(BlueprintCameraVariableTable)
 
@@ -30,15 +31,20 @@
 #define UE_PRIVATE_BLUEPRINT_CAMERA_VARIABLE_TABLE_SET_VARIABLE(VariableType)\
 	UE_PRIVATE_BLUEPRINT_CAMERA_VARIABLE_TABLE_VALIDATE()\
 	UE_PRIVATE_BLUEPRINT_CAMERA_VARIABLE_TABLE_VALIDATE_PARAM()\
-	VariableTable.GetVariableTable()->SetValue(Variable, Value, true);
+	VariableTable.GetVariableTable()->SetValue(Variable, Value, true);\
+	if (Variable->bAutoReset && VariableTable.VariableAutoResetService)\
+	{\
+		VariableTable.VariableAutoResetService->AddAutoResetVariable(Variable);\
+	}
 
 FBlueprintCameraVariableTable::FBlueprintCameraVariableTable()
 	: PrivateVariableTable(nullptr)
 {
 }
 
-FBlueprintCameraVariableTable::FBlueprintCameraVariableTable(FCameraVariableTable* InVariableTable)
+FBlueprintCameraVariableTable::FBlueprintCameraVariableTable(FCameraVariableTable* InVariableTable, TSharedPtr<UE::Cameras::FAutoResetCameraVariableService> InVariableAutoResetService)
 	: PrivateVariableTable(InVariableTable)
+	, VariableAutoResetService(InVariableAutoResetService)
 {
 }
 
diff --git a/Engine/Plugins/Cameras/GameplayCameras/Source/GameplayCameras/Private/GameFramework/GameplayCameraComponent.cpp b/Engine/Plugins/Cameras/GameplayCameras/Source/GameplayCameras/Private/GameFramework/GameplayCameraComponent.cpp
--- a/Engine/Plugins/Cameras/GameplayCameras/Source/GameplayCameras/Private/GameFramework/GameplayCameraComponent.cpp
+++ b/Engine/Plugins/Cameras/GameplayCameras/Source/GameplayCameras/Private/GameFramework/GameplayCameraComponent.cpp
@@ -13,6 +13,7 @@
 #include "GameFramework/GameplayCameraSystemHost.h"
 #include "GameplayCameras.h"
 #include "Kismet/GameplayStatics.h"
+#include "Services/AutoResetCameraVariableService.h"
 #include "UObject/ConstructorHelpers.h"
 
 #include UE_INLINE_GENERATED_CPP_BY_NAME(GameplayCameraComponent)
@@ -169,9 +170,15 @@ void UGameplayCameraComponent::SetInitialPose(const FBlueprintCameraPose& Camera
 
 FBlueprintCameraVariableTable UGameplayCameraComponent::GetInitialVariableTable() const
 {
+	using namespace UE::Cameras;
+
 	if (EvaluationContext)
 	{
-		return FBlueprintCameraVariableTable(&EvaluationContext->GetInitialResult().VariableTable);
+		FCameraVariableTable& VariableTable = EvaluationContext->GetInitialResult().VariableTable;
+		FCameraSystemEvaluator* CameraSystemEvaluator = EvaluationContext->GetCameraSystemEvaluator();
+		TSharedPtr<FAutoResetCameraVariableService> VariableAutoResetService = 
+			CameraSystemEvaluator->FindEvaluationService<FAutoResetCameraVariableService>();
+		return FBlueprintCameraVariableTable(&VariableTable, VariableAutoResetService);
 	}
 	else
 	{
diff --git a/Engine/Plugins/Cameras/GameplayCameras/Source/GameplayCameras/Private/Services/AutoResetCameraVariableService.cpp b/Engine/Plugins/Cameras/GameplayCameras/Source/GameplayCameras/Private/Services/AutoResetCameraVariableService.cpp
--- a/Engine/Plugins/Cameras/GameplayCameras/Source/GameplayCameras/Private/Services/AutoResetCameraVariableService.cpp
+++ b/Engine/Plugins/Cameras/GameplayCameras/Source/GameplayCameras/Private/Services/AutoResetCameraVariableService.cpp
@@ -24,17 +24,25 @@ void FAutoResetCameraVariableService::OnPreUpdate(const FCameraEvaluationService
 {
 	FCameraVariableTable& VariableTable = OutResult.EvaluationResult.VariableTable;
 
-	for (const TPair<const UCameraVariableAsset*, uint32>& Pair : AutoResetVariables)
+	for (auto It = AutoResetVariables.CreateIterator(); It; ++It)
 	{
-		const FCameraVariableID VariableID = Pair.Key->GetVariableID();
-		const uint8* DefaultValuePtr = Pair.Key->GetDefaultValuePtr();
-		const ECameraVariableType VariableType = Pair.Key->GetVariableType();
-		// Using TrySetValue instead of SetValue because we only know of variables *possibly* used by camera rigs. 
-		// This doesn't mean these variables have been added to the table and written to. For instance, a camera
-		// parameter on a node might be configured to use a variable that isn't set, which makes it fallback to
-		// the variable's default value when it's not found in the variable table. So we silently ignore variables
-		// not present in the variable table here.
-		VariableTable.TrySetValue(VariableID, VariableType, DefaultValuePtr);
+		const UCameraVariableAsset* Variable = It.Key().Get();
+		if (Variable)
+		{
+			const FCameraVariableID VariableID = Variable->GetVariableID();
+			const uint8* DefaultValuePtr = Variable->GetDefaultValuePtr();
+			const ECameraVariableType VariableType = Variable->GetVariableType();
+			// Using TrySetValue instead of SetValue because we only know of variables *possibly* used by camera rigs. 
+			// This doesn't mean these variables have been added to the table and written to. For instance, a camera
+			// parameter on a node might be configured to use a variable that isn't set, which makes it fallback to
+			// the variable's default value when it's not found in the variable table. So we silently ignore variables
+			// not present in the variable table here.
+			VariableTable.TrySetValue(VariableID, VariableType, DefaultValuePtr);
+		}
+		else
+		{
+			It.RemoveCurrent();
+		}
 	}
 }
 
diff --git a/Engine/Plugins/Cameras/GameplayCameras/Source/GameplayCameras/Private/Services/AutoResetCameraVariableService.h b/Engine/Plugins/Cameras/GameplayCameras/Source/GameplayCameras/Private/Services/AutoResetCameraVariableService.h
--- a/Engine/Plugins/Cameras/GameplayCameras/Source/GameplayCameras/Private/Services/AutoResetCameraVariableService.h
+++ b/Engine/Plugins/Cameras/GameplayCameras/Source/GameplayCameras/Private/Services/AutoResetCameraVariableService.h
@@ -3,7 +3,7 @@
 #pragma once
 
 #include "Core/CameraEvaluationService.h"
-#include "UObject/ObjectPtr.h"
+#include "UObject/WeakObjectPtr.h"
 
 class UCameraVariableAsset;
 
@@ -14,6 +14,13 @@ class FAutoResetCameraVariableService : public FCameraEvaluationService
 {
 	UE_DECLARE_CAMERA_EVALUATION_SERVICE(GAMEPLAYCAMERAS_API, FAutoResetCameraVariableService)
 
+public:
+
+	/** Adds a variable to the list of variables to reset every update. */
+	void AddAutoResetVariable(const UCameraVariableAsset* InVariable);
+	/** Remove a variable from the reset list. */
+	void RemoveAutoResetVariable(const UCameraVariableAsset* InVariable);
+
 protected:
 
 	// FCameraEvaluationService interface.
@@ -23,12 +30,7 @@ protected:
 
 private:
 
-	void AddAutoResetVariable(const UCameraVariableAsset* InVariable);
-	void RemoveAutoResetVariable(const UCameraVariableAsset* InVariable);
-
-private:
-
-	TMap<const UCameraVariableAsset*, uint32> AutoResetVariables;
+	TMap<TWeakObjectPtr<const UCameraVariableAsset>, uint32> AutoResetVariables;
 };
 
 }  // namespace UE::Cameras
diff --git a/Engine/Plugins/Cameras/GameplayCameras/Source/GameplayCameras/Public/Core/CameraEvaluationContext.h b/Engine/Plugins/Cameras/GameplayCameras/Source/GameplayCameras/Public/Core/CameraEvaluationContext.h
--- a/Engine/Plugins/Cameras/GameplayCameras/Source/GameplayCameras/Public/Core/CameraEvaluationContext.h
+++ b/Engine/Plugins/Cameras/GameplayCameras/Source/GameplayCameras/Public/Core/CameraEvaluationContext.h
@@ -86,6 +86,9 @@ public:
 	/** Gets the camera asset that is hosted in this context. */
 	const UCameraAsset* GetCameraAsset() const { return CameraAsset; }
 
+	/** Gets the camera system evaluator. Only valid if this evaluation context is active. */
+	FCameraSystemEvaluator* GetCameraSystemEvaluator() const { return CameraSystemEvaluator; }
+
 	/** Gets the initial evaluation result for all camera rigs in this context. */
 	const FCameraNodeEvaluationResult& GetInitialResult() const { return InitialResult; }
 
@@ -148,6 +151,8 @@ private:
 	FCameraDirectorEvaluatorStorage DirectorEvaluatorStorage;
 	FCameraDirectorEvaluator* DirectorEvaluator = nullptr;
 
+	FCameraSystemEvaluator* CameraSystemEvaluator = nullptr;
+
 	TWeakPtr<FCameraEvaluationContext> WeakParent;
 
 	using FChildrenContexts = TArray<TSharedPtr<FCameraEvaluationContext>>;
diff --git a/Engine/Plugins/Cameras/GameplayCameras/Source/GameplayCameras/Public/Core/CameraSystemEvaluator.h b/Engine/Plugins/Cameras/GameplayCameras/Source/GameplayCameras/Public/Core/CameraSystemEvaluator.h
--- a/Engine/Plugins/Cameras/GameplayCameras/Source/GameplayCameras/Public/Core/CameraSystemEvaluator.h
+++ b/Engine/Plugins/Cameras/GameplayCameras/Source/GameplayCameras/Public/Core/CameraSystemEvaluator.h
@@ -127,6 +127,21 @@ public:
 	/** Get currently registered evaluation services. */
 	GAMEPLAYCAMERAS_API void GetEvaluationServices(TArray<TSharedPtr<FCameraEvaluationService>>& OutEvaluationServices) const;
 
+	/** Finds an evaluation service of the given type. */
+	GAMEPLAYCAMERAS_API TSharedPtr<FCameraEvaluationService> FindEvaluationService(const FCameraObjectTypeID& TypeID) const;
+	
+	/** Finds an evaluation service of the given type. */
+	template<typename EvaluationServiceType>
+	TSharedPtr<EvaluationServiceType> FindEvaluationService() const
+	{
+		TSharedPtr<FCameraEvaluationService> EvaluationService = FindEvaluationService(EvaluationServiceType::StaticTypeID());
+		if (EvaluationService)
+		{
+			return StaticCastSharedPtr<EvaluationServiceType>(EvaluationService);
+		}
+		return nullptr;
+	}
+
 public:
 
 	/** Run an update of the camera system. */
diff --git a/Engine/Plugins/Cameras/GameplayCameras/Source/GameplayCameras/Public/Directors/BlueprintCameraDirector.h b/Engine/Plugins/Cameras/GameplayCameras/Source/GameplayCameras/Public/Directors/BlueprintCameraDirector.h
--- a/Engine/Plugins/Cameras/GameplayCameras/Source/GameplayCameras/Public/Directors/BlueprintCameraDirector.h
+++ b/Engine/Plugins/Cameras/GameplayCameras/Source/GameplayCameras/Public/Directors/BlueprintCameraDirector.h
@@ -18,6 +18,8 @@ enum class ECameraRigLayer : uint8;
 namespace UE::Cameras
 {
 
+class FAutoResetCameraVariableService;
+
 /** Information about a persitent camera rig to be activated or deactivated. */
 struct FBlueprintPersistentCameraRigInfo
 {
@@ -223,6 +225,9 @@ private:
 
 	/** The current evaluation context. */
 	TSharedPtr<FCameraEvaluationContext> EvaluationContext;
+
+	/** The variable auto-reset service, for using when returning the variable table. */
+	TSharedPtr<UE::Cameras::FAutoResetCameraVariableService> VariableAutoResetService;
 };
 
 /**
diff --git a/Engine/Plugins/Cameras/GameplayCameras/Source/GameplayCameras/Public/GameFramework/BlueprintCameraVariableTable.h b/Engine/Plugins/Cameras/GameplayCameras/Source/GameplayCameras/Public/GameFramework/BlueprintCameraVariableTable.h
--- a/Engine/Plugins/Cameras/GameplayCameras/Source/GameplayCameras/Public/GameFramework/BlueprintCameraVariableTable.h
+++ b/Engine/Plugins/Cameras/GameplayCameras/Source/GameplayCameras/Public/GameFramework/BlueprintCameraVariableTable.h
@@ -23,6 +23,7 @@ namespace UE::Cameras
 {
 
 class FCameraVariableTable;
+class FAutoResetCameraVariableService;
 
 }  // namespace UE::Cameras
 
@@ -37,7 +38,7 @@ public:
 	using FCameraVariableTable = UE::Cameras::FCameraVariableTable;
 
 	FBlueprintCameraVariableTable();
-	FBlueprintCameraVariableTable(FCameraVariableTable* InVariableTable);
+	FBlueprintCameraVariableTable(FCameraVariableTable* InVariableTable, TSharedPtr<UE::Cameras::FAutoResetCameraVariableService> InVariableAutoResetService);
 
 	/** Gets the underlying variable table. */
 	FCameraVariableTable* GetVariableTable() const { return PrivateVariableTable; }
@@ -48,6 +49,9 @@ public:
 private:
 
 	FCameraVariableTable* PrivateVariableTable = nullptr;
+	TSharedPtr<UE::Cameras::FAutoResetCameraVariableService> VariableAutoResetService;
+
+	friend class UBlueprintCameraVariableTableFunctionLibrary;
 };
 
 /**
