Cameras: expose a few gameplay camera component settings to Blueprints

diff --git a/Engine/Plugins/Cameras/GameplayCameras/Source/GameplayCameras/Public/GameFramework/GameplayCameraComponent.h b/Engine/Plugins/Cameras/GameplayCameras/Source/GameplayCameras/Public/GameFramework/GameplayCameraComponent.h
--- a/Engine/Plugins/Cameras/GameplayCameras/Source/GameplayCameras/Public/GameFramework/GameplayCameraComponent.h
+++ b/Engine/Plugins/Cameras/GameplayCameras/Source/GameplayCameras/Public/GameFramework/GameplayCameraComponent.h
@@ -94,14 +94,14 @@ private:
 public:
 
 	/** The camera asset to run. */
-	UPROPERTY(EditAnywhere, BlueprintReadWrite, Category=Camera)
+	UPROPERTY(BlueprintReadWrite, EditAnywhere, Category=Camera)
 	TObjectPtr<UCameraAsset> Camera;
 
 	/**
 	 * If set, auto-activates this component's camera for the given player.
 	 * This is equivalent to calling ActivateCamera on BeginPlay.
 	 */
-	UPROPERTY(EditAnywhere, Category=Camera)
+	UPROPERTY(BlueprintReadWrite, EditAnywhere, Category=Camera)
 	TEnumAsByte<EAutoReceiveInput::Type> AutoActivateForPlayer;
 
 protected:
diff --git a/Engine/Plugins/Cameras/GameplayCameras/Source/GameplayCameras/Public/GameFramework/GameplayCameraSystemComponent.h b/Engine/Plugins/Cameras/GameplayCameras/Source/GameplayCameras/Public/GameFramework/GameplayCameraSystemComponent.h
--- a/Engine/Plugins/Cameras/GameplayCameras/Source/GameplayCameras/Public/GameFramework/GameplayCameraSystemComponent.h
+++ b/Engine/Plugins/Cameras/GameplayCameras/Source/GameplayCameras/Public/GameFramework/GameplayCameraSystemComponent.h
@@ -87,14 +87,14 @@ public:
 	 * If set, auto-activates the camera system for the given player.
 	 * This sets this actor as the view target, and is equivalent to calling ActivateCameraSystem on BeginPlay.
 	 */
-	UPROPERTY(EditAnywhere, Category=Camera)
+	UPROPERTY(BlueprintReadWrite, EditAnywhere, Category=Camera)
 	TEnumAsByte<EAutoReceiveInput::Type> AutoActivateForPlayer;
 
 	/**
 	 * If enabled, sets the evaluated camera orientation as the player controller rotation every frame.
 	 * This is set on the player controller that this component was activated for.
 	 */
-	UPROPERTY(EditAnywhere, Category=Camera)
+	UPROPERTY(BlueprintReadWrite, EditAnywhere, Category=Camera)
 	bool bSetPlayerControllerRotation = true;
 
 private:
