Cameras: fix actor components activation

- Auto-activate the component when activating the camera asset, if the component isn't active.
- Only enable the auto-activate combox box if component auto-activation is enabled.
- Make the control rotation component use the same conventions as the other two components (activate by player index or by player controller, remove looking up a sibling camera component, etc.)

diff --git a/Engine/Plugins/Cameras/GameplayCameras/Source/GameplayCameras/Private/GameFramework/GameplayCameraComponent.cpp b/Engine/Plugins/Cameras/GameplayCameras/Source/GameplayCameras/Private/GameFramework/GameplayCameraComponent.cpp
--- a/Engine/Plugins/Cameras/GameplayCameras/Source/GameplayCameras/Private/GameFramework/GameplayCameraComponent.cpp
+++ b/Engine/Plugins/Cameras/GameplayCameras/Source/GameplayCameras/Private/GameFramework/GameplayCameraComponent.cpp
@@ -93,6 +93,9 @@ void UGameplayCameraComponent::DeactivateCameraEvaluationContext()
 		TSharedPtr<FCameraSystemEvaluator> Evaluator = CameraSystemHost->GetCameraSystemEvaluator();
 		Evaluator->RemoveEvaluationContext(EvaluationContext.ToSharedRef());
 	}
+
+	// Don't deactivate the component: we still need to update our evaluation context while any
+	// running camera rigs blend out.
 }
 
 void UGameplayCameraComponent::ActivateCameraEvaluationContext(APlayerController* PlayerController)
@@ -132,6 +135,9 @@ void UGameplayCameraComponent::ActivateCameraEvaluationContext(APlayerController
 
 	TSharedPtr<FCameraSystemEvaluator> CameraSystemEvaluator = CameraSystemHost->GetCameraSystemEvaluator();
 	CameraSystemEvaluator->PushEvaluationContext(EvaluationContext.ToSharedRef());
+
+	// Make sure the component is active so it receives tick updates to maintain the evaluation context.
+	Activate();
 }
 
 FBlueprintCameraPose UGameplayCameraComponent::GetInitialPose() const
@@ -228,7 +234,7 @@ void UGameplayCameraComponent::BeginPlay()
 	}
 #endif
 
-	if (AutoActivateForPlayer != EAutoReceiveInput::Disabled && GetNetMode() != NM_DedicatedServer)
+	if (IsActive() && AutoActivateForPlayer != EAutoReceiveInput::Disabled && GetNetMode() != NM_DedicatedServer)
 	{
 		const int32 PlayerIndex = AutoActivateForPlayer.GetIntValue() - 1;
 		ActivateCameraForPlayerIndex(PlayerIndex);
diff --git a/Engine/Plugins/Cameras/GameplayCameras/Source/GameplayCameras/Private/GameFramework/GameplayCameraSystemComponent.cpp b/Engine/Plugins/Cameras/GameplayCameras/Source/GameplayCameras/Private/GameFramework/GameplayCameraSystemComponent.cpp
--- a/Engine/Plugins/Cameras/GameplayCameras/Source/GameplayCameras/Private/GameFramework/GameplayCameraSystemComponent.cpp
+++ b/Engine/Plugins/Cameras/GameplayCameras/Source/GameplayCameras/Private/GameFramework/GameplayCameraSystemComponent.cpp
@@ -21,9 +21,6 @@ UGameplayCameraSystemComponent::UGameplayCameraSystemComponent(const FObjectInit
 	: Super(ObjectInit)
 {
 #if WITH_EDITORONLY_DATA
-	bTickInEditor = true;
-	PrimaryComponentTick.bCanEverTick = true;
-
 	if (GIsEditor && !IsRunningCommandlet())
 	{
 		static ConstructorHelpers::FObjectFinder<UStaticMesh> EditorCameraMesh(
@@ -190,7 +187,7 @@ void UGameplayCameraSystemComponent::BeginPlay()
 {
 	Super::BeginPlay();
 
-	if (AutoActivateForPlayer != EAutoReceiveInput::Disabled && GetNetMode() != NM_DedicatedServer)
+	if (IsActive() && AutoActivateForPlayer != EAutoReceiveInput::Disabled && GetNetMode() != NM_DedicatedServer)
 	{
 		const int32 PlayerIndex = AutoActivateForPlayer.GetIntValue() - 1;
 		ActivateCameraSystemForPlayerIndex(PlayerIndex);
@@ -204,11 +201,6 @@ void UGameplayCameraSystemComponent::EndPlay(const EEndPlayReason::Type EndPlayR
 	Super::EndPlay(EndPlayReason);
 }
 
-void UGameplayCameraSystemComponent::TickComponent(float DeltaTime, ELevelTick TickType, FActorComponentTickFunction* ThisTickFunction)
-{
-	Super::TickComponent(DeltaTime, TickType, ThisTickFunction);
-}
-
 void UGameplayCameraSystemComponent::OnComponentDestroyed(bool bDestroyingHierarchy)
 {
 	Super::OnComponentDestroyed(bDestroyingHierarchy);
@@ -229,20 +221,6 @@ void UGameplayCameraSystemComponent::OnComponentDestroyed(bool bDestroyingHierar
 #endif  // UE_GAMEPLAY_CAMERAS_DEBUG
 }
 
-#if WITH_EDITOR
-
-bool UGameplayCameraSystemComponent::GetEditorPreviewInfo(float DeltaTime, FMinimalViewInfo& ViewOut)
-{
-	const bool bIsCameraSystemActive = IsActive();
-	if (bIsCameraSystemActive)
-	{
-		GetCameraView(DeltaTime, ViewOut);
-	}
-	return bIsCameraSystemActive;
-}
-
-#endif  // WITH_EDITOR
-
 void UGameplayCameraSystemComponent::OnBecomeViewTarget()
 {
 }
diff --git a/Engine/Plugins/Cameras/GameplayCameras/Source/GameplayCameras/Private/GameFramework/GameplayControlRotationComponent.cpp b/Engine/Plugins/Cameras/GameplayCameras/Source/GameplayCameras/Private/GameFramework/GameplayControlRotationComponent.cpp
--- a/Engine/Plugins/Cameras/GameplayCameras/Source/GameplayCameras/Private/GameFramework/GameplayControlRotationComponent.cpp
+++ b/Engine/Plugins/Cameras/GameplayCameras/Source/GameplayCameras/Private/GameFramework/GameplayControlRotationComponent.cpp
@@ -33,45 +33,81 @@ void UGameplayControlRotationComponent::BeginPlay()
 {
 	Super::BeginPlay();
 
+	if (IsActive() && AutoActivateForPlayer != EAutoReceiveInput::Disabled && GetNetMode() != NM_DedicatedServer)
+	{
+		const int32 PlayerIndex = AutoActivateForPlayer.GetIntValue() - 1;
+		ActivateControlRotationManagementForPlayerIndex(PlayerIndex);
+	}
+}
+
+void UGameplayControlRotationComponent::EndPlay(const EEndPlayReason::Type EndPlayReason)
+{
 	using namespace UE::Cameras;
 
-	// We need to grab our owner actor, and then find the gameplay camera component and the enhanced
-	// input component that we will be coordinating.
-	AActor* OwnerActor = GetOwner();
-	if (!OwnerActor)
+	TeardownControlRotationService(true);
+
+	Super::EndPlay(EndPlayReason);
+}
+
+void UGameplayControlRotationComponent::TickComponent(float DeltaTime, ELevelTick TickType, FActorComponentTickFunction* ThisTickFunction)
+{
+	Super::TickComponent(DeltaTime, TickType, ThisTickFunction);
+
+	if (PlayerController && ControlRotationService)
 	{
-		UE_LOG(LogCameraSystem, Error, 
-				TEXT("GameplayControlRotationComponent '%s' doesn't belong to any actor"),
-				*GetNameSafe(this));
-		return;
+		// This may be technically one frame late (i.e. we set the control rotation computed 
+		// late last tick) unless the camera system is setup to process player input into camera
+		// rotation early in the frame.
+		PlayerController->SetControlRotation(ControlRotationService->GetCurrentControlRotation());
 	}
+}
+
+void UGameplayControlRotationComponent::ActivateControlRotationManagementForPlayerIndex(int32 PlayerIndex)
+{
+	APlayerController* ForPlayerController = UGameplayStatics::GetPlayerController(this, PlayerIndex);
+	ActivateControlRotationManagementForPlayerController(ForPlayerController);
+}
+
+void UGameplayControlRotationComponent::ActivateControlRotationManagementForPlayerController(APlayerController* InPlayerController)
+{
+	InitializeControlRotationService(InPlayerController);
+}
 
-	UGameplayCameraComponent* GameplayCameraComponent = OwnerActor->FindComponentByClass<UGameplayCameraComponent>();
-	if (!GameplayCameraComponent)
+void UGameplayControlRotationComponent::DeactivateControlRotationManagement()
+{
+	TeardownControlRotationService(false);
+}
+
+void UGameplayControlRotationComponent::InitializeControlRotationService(APlayerController* InPlayerController)
+{
+	using namespace UE::Cameras;
+
+	if (ControlRotationService)
 	{
 		UE_LOG(LogCameraSystem, Error,
-				TEXT("GameplayControlRotationComponent '%s' couldn't find a GameplayCameraComponent on owner '%s'"),
-				*GetNameSafe(this), *GetNameSafe(OwnerActor));
+				TEXT("GameplayControlRotationComponent '%s' has already been activated"),
+				*GetNameSafe(this));
 		return;
 	}
 
-	// Find where the camera system is running from.
-	PlayerController = GameplayCameraComponent->GetPlayerController();
-	if (!ensure(PlayerController))
+	if (!InPlayerController)
 	{
-		UE_LOG(LogCameraSystem, Error,
-				TEXT("GameplayCameraComponent '%s' hasn't activated for any player yet"),
-				*GetNameSafe(GameplayCameraComponent));
+		UE_LOG(LogCameraSystem, Error, 
+				TEXT("GameplayControlRotationComponent '%s' can't activate: no player controller given or found!"),
+				*GetNameSafe(this));
 		return;
 	}
 
-	CameraSystemHost = UGameplayCameraSystemHost::FindHost(PlayerController);
-	if (!CameraSystemHost)
+	UGameplayCameraSystemHost* FoundHost = UGameplayCameraSystemHost::FindHost(InPlayerController);
+	if (!FoundHost)
 	{
 		UE_LOG(LogCameraSystem, Error, TEXT("Can't find camera system host on the player controller."));
 		return;
 	}
 
+	PlayerController = InPlayerController;
+	CameraSystemHost = FoundHost;
+
 	// Make sure that if there's an auto-spawned camera system actor, it doesn't set the control rotation.
 	AGameplayCameraSystemActor* AutoSpawnedCameraSystemActor = AGameplayCameraSystemActor::GetAutoSpawnedCameraSystemActor(PlayerController);
 	if (AutoSpawnedCameraSystemActor)
@@ -93,32 +129,27 @@ void UGameplayControlRotationComponent::BeginPlay()
 	CameraSystem->RegisterEvaluationService(ControlRotationService.ToSharedRef());
 }
 
-void UGameplayControlRotationComponent::EndPlay(const EEndPlayReason::Type EndPlayReason)
+void UGameplayControlRotationComponent::TeardownControlRotationService(bool bAllowUninitialized)
 {
 	using namespace UE::Cameras;
 
-	if (CameraSystemHost)
+	if (!ControlRotationService || !CameraSystemHost)
 	{
-		TSharedPtr<FCameraSystemEvaluator> CameraSystem = CameraSystemHost->GetCameraSystemEvaluator();
-		CameraSystem->UnregisterEvaluationService(ControlRotationService.ToSharedRef());
+		if (!bAllowUninitialized)
+		{
+			UE_LOG(LogCameraSystem, Error,
+					TEXT("GameplayCameraComponent '%s' isn't active"),
+					*GetNameSafe(this));
+		}
+		return;
 	}
 
-	ControlRotationService.Reset();
-
-	Super::EndPlay(EndPlayReason);
-}
-
-void UGameplayControlRotationComponent::TickComponent(float DeltaTime, ELevelTick TickType, FActorComponentTickFunction* ThisTickFunction)
-{
-	Super::TickComponent(DeltaTime, TickType, ThisTickFunction);
+	TSharedPtr<FCameraSystemEvaluator> CameraSystem = CameraSystemHost->GetCameraSystemEvaluator();
+	CameraSystem->UnregisterEvaluationService(ControlRotationService.ToSharedRef());
 
-	if (PlayerController)
-	{
-		// This may be technically one frame late (i.e. we set the control rotation computed 
-		// late last tick) unless the camera system is setup to process player input into camera
-		// rotation early in the frame.
-		PlayerController->SetControlRotation(ControlRotationService->GetCurrentControlRotation());
-	}
+	ControlRotationService = nullptr;
+	CameraSystemHost = nullptr;
+	PlayerController = nullptr;
 }
 
 #undef LOCTEXT_NAMESPACE
diff --git a/Engine/Plugins/Cameras/GameplayCameras/Source/GameplayCameras/Public/GameFramework/GameplayCameraComponent.h b/Engine/Plugins/Cameras/GameplayCameras/Source/GameplayCameras/Public/GameFramework/GameplayCameraComponent.h
--- a/Engine/Plugins/Cameras/GameplayCameras/Source/GameplayCameras/Public/GameFramework/GameplayCameraComponent.h
+++ b/Engine/Plugins/Cameras/GameplayCameras/Source/GameplayCameras/Public/GameFramework/GameplayCameraComponent.h
@@ -98,10 +98,10 @@ public:
 	TObjectPtr<UCameraAsset> Camera;
 
 	/**
-	 * If set, auto-activates this component's camera for the given player.
+	 * If AutoActivate is set, auto-activates this component's camera for the given player.
 	 * This is equivalent to calling ActivateCamera on BeginPlay.
 	 */
-	UPROPERTY(BlueprintReadWrite, EditAnywhere, Category=Camera)
+	UPROPERTY(BlueprintReadWrite, EditAnywhere, Category=Activation, meta=(EditCondition="bAutoActivate"))
 	TEnumAsByte<EAutoReceiveInput::Type> AutoActivateForPlayer;
 
 protected:
diff --git a/Engine/Plugins/Cameras/GameplayCameras/Source/GameplayCameras/Public/GameFramework/GameplayCameraSystemComponent.h b/Engine/Plugins/Cameras/GameplayCameras/Source/GameplayCameras/Public/GameFramework/GameplayCameraSystemComponent.h
--- a/Engine/Plugins/Cameras/GameplayCameras/Source/GameplayCameras/Public/GameFramework/GameplayCameraSystemComponent.h
+++ b/Engine/Plugins/Cameras/GameplayCameras/Source/GameplayCameras/Public/GameFramework/GameplayCameraSystemComponent.h
@@ -61,14 +61,8 @@ public:
 	virtual void OnRegister() override;
 	virtual void BeginPlay() override;
 	virtual void EndPlay(const EEndPlayReason::Type EndPlayReason) override;
-	virtual void TickComponent(float DeltaTime, ELevelTick TickType, FActorComponentTickFunction* ThisTickFunction) override;
 	virtual void OnComponentDestroyed(bool bDestroyingHierarchy) override;
 
-	// USceneComponent interface
-#if WITH_EDITOR
-	virtual bool GetEditorPreviewInfo(float DeltaTime, FMinimalViewInfo& ViewOut) override;
-#endif  // WITH_EDITOR
-
 public:
 
 	// Internal API
@@ -84,10 +78,10 @@ private:
 public:
 
 	/**
-	 * If set, auto-activates the camera system for the given player.
+	 * If AutoActivate is set, auto-activates the camera system for the given player.
 	 * This sets this actor as the view target, and is equivalent to calling ActivateCameraSystem on BeginPlay.
 	 */
-	UPROPERTY(BlueprintReadWrite, EditAnywhere, Category=Camera)
+	UPROPERTY(BlueprintReadWrite, EditAnywhere, Category=Activation, meta=(EditCondition="bAutoActivate"))
 	TEnumAsByte<EAutoReceiveInput::Type> AutoActivateForPlayer;
 
 	/**
diff --git a/Engine/Plugins/Cameras/GameplayCameras/Source/GameplayCameras/Public/GameFramework/GameplayControlRotationComponent.h b/Engine/Plugins/Cameras/GameplayCameras/Source/GameplayCameras/Public/GameFramework/GameplayControlRotationComponent.h
--- a/Engine/Plugins/Cameras/GameplayCameras/Source/GameplayCameras/Public/GameFramework/GameplayControlRotationComponent.h
+++ b/Engine/Plugins/Cameras/GameplayCameras/Source/GameplayCameras/Public/GameFramework/GameplayControlRotationComponent.h
@@ -32,6 +32,26 @@ public:
 	/** Create a new control rotation component. */
 	UGameplayControlRotationComponent(const FObjectInitializer& ObjectInit);
 
+	/**
+	 * Activates management of a player controller's control rotation. The component will set
+	 * the control rotation every frame based on the latest camera system update.
+	 */
+	UFUNCTION(BlueprintCallable, Category="Control Rotation")
+	void ActivateControlRotationManagementForPlayerIndex(int32 PlayerIndex);
+
+	/**
+	 * Activates management of a player controller's control rotation. The component will set
+	 * the control rotation every frame based on the latest camera system update.
+	 */
+	UFUNCTION(BlueprintCallable, Category="Control Rotation")
+	void ActivateControlRotationManagementForPlayerController(APlayerController* PlayerController);
+
+	/**
+	 * Deactivates management of a player controller's control rotation.
+	 */
+	UFUNCTION(BlueprintCallable, Category="Control Rotation")
+	void DeactivateControlRotationManagement();
+
 public:
 
 	// UActorComponent interface
@@ -58,6 +78,18 @@ public:
 	UPROPERTY(EditAnywhere, Category="Input")
 	float AxisActionMagnitudeThreshold = 0.1f;
 
+	/**
+	 * If AutoActivate is set, auto-activates control rotation management for the given player.
+	 * This is equivalent to calling ActivateControlRotationManagement on BeginPlay.
+	 */
+	UPROPERTY(BlueprintReadWrite, EditAnywhere, Category="Activation", meta=(EditCondition="bAutoActivate"))
+	TEnumAsByte<EAutoReceiveInput::Type> AutoActivateForPlayer;
+
+private:
+
+	void InitializeControlRotationService(APlayerController* PlayerController);
+	void TeardownControlRotationService(bool bAllowUninitialized);
+
 private:
 
 	UPROPERTY()
