Cameras: more fixes to camera component activation, prevent assert when enabling debug info

diff --git a/Engine/Plugins/Cameras/GameplayCameras/Source/GameplayCameras/Private/GameFramework/GameplayCameraSystemComponent.cpp b/Engine/Plugins/Cameras/GameplayCameras/Source/GameplayCameras/Private/GameFramework/GameplayCameraSystemComponent.cpp
--- a/Engine/Plugins/Cameras/GameplayCameras/Source/GameplayCameras/Private/GameFramework/GameplayCameraSystemComponent.cpp
+++ b/Engine/Plugins/Cameras/GameplayCameras/Source/GameplayCameras/Private/GameFramework/GameplayCameraSystemComponent.cpp
@@ -142,6 +142,9 @@ void UGameplayCameraSystemComponent::ActivateCameraSystemForPlayerController(APl
 
 	PlayerController->SetViewTarget(OwningActor);
 	WeakPlayerController = PlayerController;
+
+	// Make sure the component is active.
+	Activate();
 }
 
 bool UGameplayCameraSystemComponent::IsCameraSystemActiveForPlayController(APlayerController* PlayerController) const
@@ -235,7 +238,7 @@ void UGameplayCameraSystemComponent::DebugDraw(UCanvas* Canvas, APlayerControlle
 {
 	using namespace UE::Cameras;
 
-	TSharedPtr<FCameraSystemEvaluator> Evaluator = GetCameraSystemEvaluator();
+	TSharedPtr<FCameraSystemEvaluator> Evaluator = GetCameraSystemEvaluator(false);
 	if (Evaluator.IsValid())
 	{
 		FCameraSystemDebugUpdateParams DebugUpdateParams;
