Cameras: use camera variable references in input slot nodes

diff --git a/Engine/Plugins/Cameras/GameplayCameras/Source/GameplayCameras/Private/Nodes/Input/CameraRigInput1DSlot.cpp b/Engine/Plugins/Cameras/GameplayCameras/Source/GameplayCameras/Private/Nodes/Input/CameraRigInput1DSlot.cpp
--- a/Engine/Plugins/Cameras/GameplayCameras/Source/GameplayCameras/Private/Nodes/Input/CameraRigInput1DSlot.cpp
+++ b/Engine/Plugins/Cameras/GameplayCameras/Source/GameplayCameras/Private/Nodes/Input/CameraRigInput1DSlot.cpp
@@ -100,9 +100,9 @@ void UCameraRigInput1DSlot::OnBuild(FCameraRigBuildContext& BuildContext)
 	{
 		VariableDefinition = FBuiltInCameraVariables::Get().GetDefinition(BuiltInVariable);
 	}
-	else if (Variable)
+	else if (CustomVariable.IsValid())
 	{
-		VariableDefinition = Variable->GetVariableDefinition();
+		VariableDefinition = CustomVariable.Get()->GetVariableDefinition();
 	}
 	else if (InputSlotParameters.bIsPreBlended)
 	{
diff --git a/Engine/Plugins/Cameras/GameplayCameras/Source/GameplayCameras/Private/Nodes/Input/CameraRigInput2DSlot.cpp b/Engine/Plugins/Cameras/GameplayCameras/Source/GameplayCameras/Private/Nodes/Input/CameraRigInput2DSlot.cpp
--- a/Engine/Plugins/Cameras/GameplayCameras/Source/GameplayCameras/Private/Nodes/Input/CameraRigInput2DSlot.cpp
+++ b/Engine/Plugins/Cameras/GameplayCameras/Source/GameplayCameras/Private/Nodes/Input/CameraRigInput2DSlot.cpp
@@ -107,9 +107,9 @@ void UCameraRigInput2DSlot::OnBuild(FCameraRigBuildContext& BuildContext)
 	{
 		VariableDefinition = FBuiltInCameraVariables::Get().GetDefinition(BuiltInVariable);
 	}
-	else if (Variable)
+	else if (CustomVariable.IsValid())
 	{
-		VariableDefinition = Variable->GetVariableDefinition();
+		VariableDefinition = CustomVariable.Get()->GetVariableDefinition();
 	}
 	else if (InputSlotParameters.bIsPreBlended)
 	{
diff --git a/Engine/Plugins/Cameras/GameplayCameras/Source/GameplayCameras/Public/Core/CameraVariableReferences.h b/Engine/Plugins/Cameras/GameplayCameras/Source/GameplayCameras/Public/Core/CameraVariableReferences.h
--- a/Engine/Plugins/Cameras/GameplayCameras/Source/GameplayCameras/Public/Core/CameraVariableReferences.h
+++ b/Engine/Plugins/Cameras/GameplayCameras/Source/GameplayCameras/Public/Core/CameraVariableReferences.h
@@ -9,6 +9,7 @@
 #define UE_DEFINE_CAMERA_VARIABLE_REFERENCE(ValueName)\
 	F##ValueName##CameraVariableReference() {}\
 	F##ValueName##CameraVariableReference(VariableAssetType* InVariable) : Variable(InVariable) {}\
+	bool IsValid() const { return (bool)Variable; }\
 	U##ValueName##CameraVariable* Get() const { return Variable.Get(); }\
 	operator U##ValueName##CameraVariable* () const { return Variable.Get(); }
 
diff --git a/Engine/Plugins/Cameras/GameplayCameras/Source/GameplayCameras/Public/Nodes/Input/CameraRigInput1DSlot.h b/Engine/Plugins/Cameras/GameplayCameras/Source/GameplayCameras/Public/Nodes/Input/CameraRigInput1DSlot.h
--- a/Engine/Plugins/Cameras/GameplayCameras/Source/GameplayCameras/Public/Nodes/Input/CameraRigInput1DSlot.h
+++ b/Engine/Plugins/Cameras/GameplayCameras/Source/GameplayCameras/Public/Nodes/Input/CameraRigInput1DSlot.h
@@ -4,6 +4,7 @@
 
 #include "Core/BuiltInCameraVariables.h"
 #include "Core/CameraNodeEvaluator.h"
+#include "Core/CameraVariableReferences.h"
 #include "Nodes/Input/CameraRigInputSlotTypes.h"
 #include "Nodes/Input/Input1DCameraNode.h"
 #include "UObject/ObjectPtr.h"
@@ -40,7 +41,7 @@ public:
 
 	/** The variable to use to blend with other input slots. */
 	UPROPERTY(EditAnywhere, Category="Input", meta=(EditCondition="BuiltInVariable == EBuiltInDoubleCameraVariable::None"))
-	TObjectPtr<UDoubleCameraVariable> Variable;
+	FDoubleCameraVariableReference CustomVariable;
 
 public:
 
diff --git a/Engine/Plugins/Cameras/GameplayCameras/Source/GameplayCameras/Public/Nodes/Input/CameraRigInput2DSlot.h b/Engine/Plugins/Cameras/GameplayCameras/Source/GameplayCameras/Public/Nodes/Input/CameraRigInput2DSlot.h
--- a/Engine/Plugins/Cameras/GameplayCameras/Source/GameplayCameras/Public/Nodes/Input/CameraRigInput2DSlot.h
+++ b/Engine/Plugins/Cameras/GameplayCameras/Source/GameplayCameras/Public/Nodes/Input/CameraRigInput2DSlot.h
@@ -4,6 +4,7 @@
 
 #include "Core/BuiltInCameraVariables.h"
 #include "Core/CameraNodeEvaluator.h"
+#include "Core/CameraVariableReferences.h"
 #include "Nodes/Input/CameraRigInputSlotTypes.h"
 #include "Nodes/Input/Input2DCameraNode.h"
 #include "UObject/ObjectPtr.h"
@@ -48,7 +49,7 @@ public:
 
 	/** The variable to use to blend with other input slots. */
 	UPROPERTY(EditAnywhere, Category="Input", meta=(EditCondition="BuiltInVariable == EBuiltInVector2dCameraVariable::None"))
-	TObjectPtr<UVector2dCameraVariable> Variable;
+	FVector2dCameraVariableReference CustomVariable;
 
 public:
 
