Cameras: variable table getter BP functions should return default values for unset variables

diff --git a/Engine/Plugins/Cameras/GameplayCameras/Source/GameplayCameras/Private/GameFramework/BlueprintCameraVariableTable.cpp b/Engine/Plugins/Cameras/GameplayCameras/Source/GameplayCameras/Private/GameFramework/BlueprintCameraVariableTable.cpp
--- a/Engine/Plugins/Cameras/GameplayCameras/Source/GameplayCameras/Private/GameFramework/BlueprintCameraVariableTable.cpp
+++ b/Engine/Plugins/Cameras/GameplayCameras/Source/GameplayCameras/Private/GameFramework/BlueprintCameraVariableTable.cpp
@@ -26,7 +26,7 @@
 	static VariableType ErrorResult {};\
 	UE_PRIVATE_BLUEPRINT_CAMERA_VARIABLE_TABLE_VALIDATE(ErrorResult)\
 	UE_PRIVATE_BLUEPRINT_CAMERA_VARIABLE_TABLE_VALIDATE_PARAM(ErrorResult)\
-	return VariableTable.GetVariableTable()->GetValue<VariableType>(Variable->GetVariableID());
+	return VariableTable.GetVariableTable()->GetValue<VariableType>(Variable->GetVariableID(), Variable->GetDefaultValue());
 
 #define UE_PRIVATE_BLUEPRINT_CAMERA_VARIABLE_TABLE_SET_VARIABLE(VariableType)\
 	UE_PRIVATE_BLUEPRINT_CAMERA_VARIABLE_TABLE_VALIDATE()\
