Cameras: fix handling of camera cuts for damping logic

We handle all camera cuts as "teleport physics", effectively: we just keep the same damping amount as the last frame.

diff --git a/Engine/Plugins/Cameras/GameplayCameras/Source/GameplayCameras/Private/Nodes/Common/BoomArmCameraNode.cpp b/Engine/Plugins/Cameras/GameplayCameras/Source/GameplayCameras/Private/Nodes/Common/BoomArmCameraNode.cpp
--- a/Engine/Plugins/Cameras/GameplayCameras/Source/GameplayCameras/Private/Nodes/Common/BoomArmCameraNode.cpp
+++ b/Engine/Plugins/Cameras/GameplayCameras/Source/GameplayCameras/Private/Nodes/Common/BoomArmCameraNode.cpp
@@ -187,9 +187,16 @@ void FBoomArmCameraNodeEvaluator::OnRun(const FCameraNodeEvaluationParams& Param
 
 			CumulativePull = ClampedPull;
 		}
-		else
+		else if (!Params.bIsFirstFrame && OutResult.bIsCameraCut)
+		{
+			// On camera cuts, we re-use last frame's cumulative pull without updating it.
+			const FVector3d ForwardBoomOrientation = BoomRotation.RotateVector(FVector3d::ForwardVector);
+			FinalTransform.SetLocation(FinalTransform.GetLocation() - ForwardBoomOrientation * CumulativePull);
+
+			// Leave bDebugDidClampPull to what it was last frame.
+		}
+		else if (Params.bIsFirstFrame)
 		{
-			LastPivotLocation = BoomPivot.GetLocation();
 			CumulativePull = 0.0;
 		}
 
diff --git a/Engine/Plugins/Cameras/GameplayCameras/Source/GameplayCameras/Private/Nodes/Common/DampenPositionCameraNode.cpp b/Engine/Plugins/Cameras/GameplayCameras/Source/GameplayCameras/Private/Nodes/Common/DampenPositionCameraNode.cpp
--- a/Engine/Plugins/Cameras/GameplayCameras/Source/GameplayCameras/Private/Nodes/Common/DampenPositionCameraNode.cpp
+++ b/Engine/Plugins/Cameras/GameplayCameras/Source/GameplayCameras/Private/Nodes/Common/DampenPositionCameraNode.cpp
@@ -31,6 +31,12 @@ protected:
 	virtual void OnBuildDebugBlocks(const FCameraDebugBlockBuildParams& Params, FCameraDebugBlockBuilder& Builder) override;
 #endif  // UE_GAMEPLAY_CAMERAS_DEBUG
 
+private:
+
+	using FAxisDamper = TTuple<FVector3d, FCriticalDamper*>;
+	using FAxisDampers = FAxisDamper[3];
+	void ComputeAxisDampers(const FCameraNodeEvaluationParams& Params, const FCameraNodeEvaluationResult& OutResult, FAxisDampers& OutAxisDampers);
+
 private:
 
 	TCameraParameterReader<float> ForwardDampingFactorReader;
@@ -42,6 +48,7 @@ private:
 	FCriticalDamper VerticalDamper;
 
 	FVector3d PreviousLocation;
+	FVector3d PreviousLagVector;
 
 #if UE_GAMEPLAY_CAMERAS_DEBUG
 	FVector3d DebugLastUndampedPosition;
@@ -89,8 +96,6 @@ void FDampenPositionCameraNodeEvaluator::OnInitialize(const FCameraNodeEvaluator
 
 void FDampenPositionCameraNodeEvaluator::OnRun(const FCameraNodeEvaluationParams& Params, FCameraNodeEvaluationResult& OutResult)
 {
-	const UDampenPositionCameraNode* DampenNode = GetCameraNodeAs<UDampenPositionCameraNode>();
-
 	// Update our damping factors if they're driven by a camera variable (which means they
 	// could change every frame). Also update them if we are in the editor, since the 
 	// user might tweak them live.
@@ -125,50 +130,12 @@ void FDampenPositionCameraNodeEvaluator::OnRun(const FCameraNodeEvaluationParams
 	DebugLastDampingRotation = FRotator3d::ZeroRotator;
 #endif
 
+	// Figure out the coordinate system in which we are damping movement.
+	FAxisDampers AxisDampers;
+	ComputeAxisDampers(Params, OutResult, AxisDampers);
+
 	if (!Params.bIsFirstFrame && !OutResult.bIsCameraCut)
 	{
-		// Figure out the desired coordinate system.
-		FRotator3d AxesRotation(EForceInit::ForceInit);
-		switch (DampenNode->DampenSpace)
-		{
-			case ECameraNodeSpace::CameraPose:
-				{
-					const FRotator3d CameraRotation = OutResult.CameraPose.GetRotation();
-					AxesRotation = CameraRotation;
-				}
-				break;
-			case ECameraNodeSpace::OwningContext:
-				if (Params.EvaluationContext)
-				{
-					const FCameraNodeEvaluationResult& InitialResult = Params.EvaluationContext->GetInitialResult();
-					ensureMsgf(InitialResult.bIsValid,
-							TEXT("DampenPositionCameraNode: using invalid context result as damping space!"));
-
-					const FRotator3d ContextRotation = InitialResult.CameraPose.GetRotation();
-					AxesRotation = ContextRotation;
-				}
-				else
-				{
-					UE_LOG(LogCameraSystem, Error,
-							TEXT("DampenPositionCameraNode: cannot dampen in context space when there is "
-								"no current context set."));
-				}
-				break;
-			case ECameraNodeSpace::World:
-				break;
-		}
-#if UE_GAMEPLAY_CAMERAS_DEBUG
-		DebugLastDampingRotation = AxesRotation;
-#endif
-
-		using FAxisDamper = TTuple<FVector3d, FCriticalDamper*>;
-		FAxisDamper AxisDampers[3]
-		{
-			{ AxesRotation.RotateVector(FVector3d::ForwardVector), &ForwardDamper },
-			{ AxesRotation.RotateVector(FVector3d::RightVector),& LateralDamper },
-			{ AxesRotation.RotateVector(FVector3d::UpVector), &VerticalDamper }
-		};
-
 		// The next target has moved further away compared to the previous target,
 		// so we're lagging behind even more than before. Compute this new lag vector.
 		const FVector3d NewLagVector = NextTarget - PreviousLocation;
@@ -191,6 +158,13 @@ void FDampenPositionCameraNodeEvaluator::OnRun(const FCameraNodeEvaluationParams
 		}
 		
 		NextLocation = NewDampedLocation;
+		PreviousLagVector = NextTarget - NextLocation;
+	}
+	else if (!Params.bIsFirstFrame && OutResult.bIsCameraCut)
+	{
+		// On camera cuts, we don't update the damping, and just re-use whatever lag 
+		// we previously had.
+		NextLocation = NextTarget - PreviousLagVector;
 	}
 
 #if UE_GAMEPLAY_CAMERAS_DEBUG
@@ -203,6 +177,48 @@ void FDampenPositionCameraNodeEvaluator::OnRun(const FCameraNodeEvaluationParams
 	OutResult.CameraPose.SetLocation(NextLocation);
 }
 
+void FDampenPositionCameraNodeEvaluator::ComputeAxisDampers(const FCameraNodeEvaluationParams& Params, const FCameraNodeEvaluationResult& Result, FAxisDampers& OutAxisDampers)
+{
+	// Figure out the desired coordinate system.
+	FRotator3d AxesRotation(EForceInit::ForceInit);
+	const UDampenPositionCameraNode* DampenNode = GetCameraNodeAs<UDampenPositionCameraNode>();
+	switch (DampenNode->DampenSpace)
+	{
+		case ECameraNodeSpace::CameraPose:
+			{
+				const FRotator3d CameraRotation = Result.CameraPose.GetRotation();
+				AxesRotation = CameraRotation;
+			}
+			break;
+		case ECameraNodeSpace::OwningContext:
+			if (Params.EvaluationContext)
+			{
+				const FCameraNodeEvaluationResult& InitialResult = Params.EvaluationContext->GetInitialResult();
+				ensureMsgf(InitialResult.bIsValid,
+						TEXT("DampenPositionCameraNode: using invalid context result as damping space!"));
+
+				const FRotator3d ContextRotation = InitialResult.CameraPose.GetRotation();
+				AxesRotation = ContextRotation;
+			}
+			else
+			{
+				UE_LOG(LogCameraSystem, Error,
+						TEXT("DampenPositionCameraNode: cannot dampen in context space when there is "
+							"no current context set."));
+			}
+			break;
+		case ECameraNodeSpace::World:
+			break;
+	}
+#if UE_GAMEPLAY_CAMERAS_DEBUG
+	DebugLastDampingRotation = AxesRotation;
+#endif
+
+	OutAxisDampers[0] = { AxesRotation.RotateVector(FVector3d::ForwardVector), &ForwardDamper };
+	OutAxisDampers[1] = { AxesRotation.RotateVector(FVector3d::RightVector),& LateralDamper };
+	OutAxisDampers[2] = { AxesRotation.RotateVector(FVector3d::UpVector), &VerticalDamper };
+}
+
 void FDampenPositionCameraNodeEvaluator::OnSerialize(const FCameraNodeEvaluatorSerializeParams& Params, FArchive& Ar)
 {
 	Ar << ForwardDamper;
