Cameras: improve debug natvis of camera variables

- Use their display name if they have one
- Show the IDs in the entry list

diff --git a/Engine/Plugins/Cameras/GameplayCameras/Source/GameplayCameras/Private/Core/CameraVariableAssets.cpp b/Engine/Plugins/Cameras/GameplayCameras/Source/GameplayCameras/Private/Core/CameraVariableAssets.cpp
--- a/Engine/Plugins/Cameras/GameplayCameras/Source/GameplayCameras/Private/Core/CameraVariableAssets.cpp
+++ b/Engine/Plugins/Cameras/GameplayCameras/Source/GameplayCameras/Private/Core/CameraVariableAssets.cpp
@@ -23,7 +23,7 @@ FCameraVariableDefinition UCameraVariableAsset::GetVariableDefinition() const
 	VariableDefinition.bIsPrivate = bIsPrivate;
 	VariableDefinition.bIsInput = bIsInput;
 #if WITH_EDITORONLY_DATA
-	VariableDefinition.VariableName = GetName();
+	VariableDefinition.VariableName = GetDisplayName();
 #endif
 	return VariableDefinition;
 }
@@ -59,7 +59,7 @@ void UCameraVariableAsset::PostDuplicate(EDuplicateMode::Type DuplicateMode)
 	}
 }
 
-#if WITH_EDITOR
+#if WITH_EDITORONLY_DATA
 
 FString UCameraVariableAsset::GetDisplayName() const
 {
@@ -70,6 +70,10 @@ FString UCameraVariableAsset::GetDisplayName() const
 	return GetName();
 }
 
+#endif  // WITH_EDITORONLY_DATA
+
+#if WITH_EDITOR
+
 FText UCameraVariableAsset::GetDisplayText() const
 {
 	if (!DisplayName.IsEmpty())
@@ -80,3 +84,4 @@ FText UCameraVariableAsset::GetDisplayText() const
 }
 
 #endif  // WITH_EDITOR
+
diff --git a/Engine/Plugins/Cameras/GameplayCameras/Source/GameplayCameras/Public/Core/CameraVariableAssets.h b/Engine/Plugins/Cameras/GameplayCameras/Source/GameplayCameras/Public/Core/CameraVariableAssets.h
--- a/Engine/Plugins/Cameras/GameplayCameras/Source/GameplayCameras/Public/Core/CameraVariableAssets.h
+++ b/Engine/Plugins/Cameras/GameplayCameras/Source/GameplayCameras/Public/Core/CameraVariableAssets.h
@@ -30,8 +30,11 @@ public:
 	virtual ECameraVariableType GetVariableType() const PURE_VIRTUAL(UCameraVariableAsset::GetVariableType, return ECameraVariableType::Boolean;);
 	virtual const uint8* GetDefaultValuePtr() const PURE_VIRTUAL(UCameraVariableAsset::GetDefaultValuePtr, return nullptr;);
 
-#if WITH_EDITOR
+#if WITH_EDITORONLY_DATA
 	FString GetDisplayName() const;
+#endif  // WITH_EDITORONLY_DATA
+
+#if WITH_EDITOR
 	FText GetDisplayText() const;
 
 	virtual FString FormatDefaultValue() const PURE_VIRTUAL(UCameraVariableAsset::FormatDefaultValue, return FString(););
