Cameras: fix a couple of egregious memory bugs in the variable table

Oops. Amazing that everything seemed to work except for one test... Good thing we have QA!

diff --git a/Engine/Plugins/Cameras/GameplayCameras/Source/GameplayCameras/Private/Core/CameraVariableTable.cpp b/Engine/Plugins/Cameras/GameplayCameras/Source/GameplayCameras/Private/Core/CameraVariableTable.cpp
--- a/Engine/Plugins/Cameras/GameplayCameras/Source/GameplayCameras/Private/Core/CameraVariableTable.cpp
+++ b/Engine/Plugins/Cameras/GameplayCameras/Source/GameplayCameras/Private/Core/CameraVariableTable.cpp
@@ -112,7 +112,7 @@ void FCameraVariableTable::Initialize(const FCameraVariableTableAllocationInfo&
 	MaxAlignOf = FMath::Max(32u, MaxAlignOf);
 	Memory = reinterpret_cast<uint8*>(FMemory::Malloc(TotalSizeOf, MaxAlignOf));
 	Capacity = TotalSizeOf;
-	Used = 0;
+	Used = TotalSizeOf;
 
 	// Go back to our entries and initialize each entry to the default value for that variable type.
 	for (const FEntry& Entry : Entries)
@@ -214,7 +214,7 @@ bool FCameraVariableTable::GetVariableTypeAllocationInfo(ECameraVariableType Var
 #define UE_CAMERA_VARIABLE_FOR_TYPE(ValueType, ValueName)\
 		case ECameraVariableType::ValueName:\
 			OutSizeOf = sizeof(ValueType);\
-			OutAlignOf = sizeof(ValueType);\
+			OutAlignOf = alignof(ValueType);\
 			return true;
 UE_CAMERA_VARIABLE_FOR_ALL_TYPES()
 #undef UE_CAMERA_VARIABLE_FOR_TYPE
