Cameras: more fixes for auto-resetting variables used from Blueprint

Multiple calls to SetValue from Blueprint would result in ever rising ref-counts for auto-resetting variables. This change keeps a separate flag for uses from BP/Verse/etc. The flag is cleared every frame, so when BP/Verse/etc. stops writing to the variable, it doesn't count as a reference anymore and the variable can be deleted if nothing else references it.

This change required doing a pass on the WrittenThisFrame flag:

- Clean-up how these flags are cleared in the camera system update.
- Correctly propagate the flag when overriding/interpolating variable tables.
- Option for not setting it inside low-level variable table API.

diff --git a/Engine/Plugins/Cameras/GameplayCameras/Source/GameplayCameras/Private/Core/BlendStackCameraNode.cpp b/Engine/Plugins/Cameras/GameplayCameras/Source/GameplayCameras/Private/Core/BlendStackCameraNode.cpp
--- a/Engine/Plugins/Cameras/GameplayCameras/Source/GameplayCameras/Private/Core/BlendStackCameraNode.cpp
+++ b/Engine/Plugins/Cameras/GameplayCameras/Source/GameplayCameras/Private/Core/BlendStackCameraNode.cpp
@@ -215,6 +215,11 @@ void FBlendStackCameraNodeEvaluator::OnRun(const FCameraNodeEvaluationParams& Pa
 
 				continue;
 			}
+
+			// Reset this entry's flags for this frame.
+			FCameraNodeEvaluationResult& CurResult = Entry.Result;
+			CurResult.CameraPose.ClearAllChangedFlags();
+			CurResult.VariableTable.ClearAllWrittenThisFrameFlags();
 		}
 		// else: frozen entries may have null contexts or invalid initial results
 		//       because we're not going to update them anyway. We will however blend
@@ -250,8 +255,6 @@ void FBlendStackCameraNodeEvaluator::OnRun(const FCameraNodeEvaluationParams& Pa
 		const FCameraNodeEvaluationResult& ContextResult(ResolvedEntry.Context->GetInitialResult());
 		CurResult.VariableTable.OverrideAll(ContextResult.VariableTable);
 
-		CurResult.VariableTable.ClearAllWrittenThisFrameFlags();
-
 		// Gather input parameters.
 		if (!Entry.bInputRunThisFrame)
 		{
@@ -342,8 +345,6 @@ void FBlendStackCameraNodeEvaluator::OnRun(const FCameraNodeEvaluationParams& Pa
 		CurResult.bIsCameraCut = OutResult.bIsCameraCut || ContextResult.bIsCameraCut;
 		CurResult.bIsValid = true;
 
-		CurResult.CameraPose.ClearAllChangedFlags();
-
 		// Run the camera rig's root node.
 		FCameraNodeEvaluator* RootEvaluator = Entry.RootEvaluator->GetRootEvaluator();
 		if (RootEvaluator)
diff --git a/Engine/Plugins/Cameras/GameplayCameras/Source/GameplayCameras/Private/Core/CameraEvaluationContextStack.cpp b/Engine/Plugins/Cameras/GameplayCameras/Source/GameplayCameras/Private/Core/CameraEvaluationContextStack.cpp
--- a/Engine/Plugins/Cameras/GameplayCameras/Source/GameplayCameras/Private/Core/CameraEvaluationContextStack.cpp
+++ b/Engine/Plugins/Cameras/GameplayCameras/Source/GameplayCameras/Private/Core/CameraEvaluationContextStack.cpp
@@ -5,6 +5,7 @@
 #include "Core/CameraDirectorEvaluator.h"
 #include "Core/CameraEvaluationContext.h"
 #include "Core/CameraSystemEvaluator.h"
+#include "Core/CameraVariableTable.h"
 
 namespace UE::Cameras
 {
@@ -103,6 +104,17 @@ void FCameraEvaluationContextStack::PopContext()
 	Entries.Pop();
 }
 
+void FCameraEvaluationContextStack::GetAllContexts(TArray<TSharedPtr<FCameraEvaluationContext>>& OutContexts) const
+{
+	for (const FContextEntry& Entry : Entries)
+	{
+		if (TSharedPtr<FCameraEvaluationContext> Context = Entry.WeakContext.Pin())
+		{
+			OutContexts.Add(Context);
+		}
+	}
+}
+
 void FCameraEvaluationContextStack::Reset()
 {
 	for (FContextEntry& Entry : Entries)
@@ -132,5 +144,34 @@ void FCameraEvaluationContextStack::AddReferencedObjects(FReferenceCollector& Co
 	}
 }
 
+void FCameraEvaluationContextStack::OnEndCameraSystemUpdate()
+{
+	// Reset all written-this-frame flags on evaluation contexts, so we properly get those flags set
+	// regardless of when, during next frame, they set their variables. This is because various 
+	// gameplay systems, Blueprint scripting, whatever, might set variables at any time.
+	TArray<TSharedPtr<FCameraEvaluationContext>> ContextsToVisit;
+	for (const FContextEntry& Entry : Entries)
+	{
+		if (TSharedPtr<FCameraEvaluationContext> Context = Entry.WeakContext.Pin())
+		{
+			ContextsToVisit.Add(Context);
+		}
+	}
+	while (!ContextsToVisit.IsEmpty())
+	{
+		TSharedPtr<FCameraEvaluationContext> Context = ContextsToVisit.Pop();
+		Context->GetInitialResult().VariableTable.ClearAllWrittenThisFrameFlags();
+
+		TArrayView<const TSharedPtr<FCameraEvaluationContext>> ChildrenContexts(Context->GetChildrenContexts());
+		for (TSharedPtr<FCameraEvaluationContext> ChildContext : ReverseIterate(ChildrenContexts))
+		{
+			if (ChildContext)
+			{
+				ContextsToVisit.Add(ChildContext);
+			}
+		}
+	}
+}
+
 }  // namespace UE::Cameras
 
diff --git a/Engine/Plugins/Cameras/GameplayCameras/Source/GameplayCameras/Private/Core/CameraNodeEvaluator.cpp b/Engine/Plugins/Cameras/GameplayCameras/Source/GameplayCameras/Private/Core/CameraNodeEvaluator.cpp
--- a/Engine/Plugins/Cameras/GameplayCameras/Source/GameplayCameras/Private/Core/CameraNodeEvaluator.cpp
+++ b/Engine/Plugins/Cameras/GameplayCameras/Source/GameplayCameras/Private/Core/CameraNodeEvaluator.cpp
@@ -18,17 +18,12 @@ FCameraNodeEvaluatorInitializeParams::FCameraNodeEvaluatorInitializeParams(FCame
 {
 }
 
-void FCameraNodeEvaluationResult::Reset(bool bResetVariableTable)
+void FCameraNodeEvaluationResult::Reset()
 {
 	CameraPose.Reset();
 	CameraRigJoints.Reset();
 	PostProcessSettings.Reset();
 
-	if (bResetVariableTable)
-	{
-		VariableTable.UnsetAllValues();
-	}
-
 	bIsCameraCut = false;
 	bIsValid = false;
 }
diff --git a/Engine/Plugins/Cameras/GameplayCameras/Source/GameplayCameras/Private/Core/CameraPose.cpp b/Engine/Plugins/Cameras/GameplayCameras/Source/GameplayCameras/Private/Core/CameraPose.cpp
--- a/Engine/Plugins/Cameras/GameplayCameras/Source/GameplayCameras/Private/Core/CameraPose.cpp
+++ b/Engine/Plugins/Cameras/GameplayCameras/Source/GameplayCameras/Private/Core/CameraPose.cpp
@@ -79,14 +79,11 @@ FCameraPose::FCameraPose()
 {
 }
 
-void FCameraPose::Reset(bool bSetAllChangedFlags)
+void FCameraPose::Reset()
 {
 	*this = FCameraPose();
 
-	if (bSetAllChangedFlags)
-	{
-		SetAllChangedFlags();
-	}
+	ClearAllChangedFlags();
 }
 
 void FCameraPose::SetAllChangedFlags()
@@ -291,7 +288,6 @@ void FCameraPose::InternalLerpChanged(const FCameraPose& ToPose, float Factor, c
 		{\
 			if (!bChangedOnly || ToPoseChangedFlags.PropName)\
 			{\
-				ensureMsgf(ChangedFlags.PropName, TEXT("Interpolating " #PropName " from default value!"));\
 				Set##PropName(FMath::Lerp(Get##PropName(), ToPose.Get##PropName(), Factor));\
 			}\
 			OutMask.PropName = true;\
@@ -352,7 +348,6 @@ void FCameraPose::InternalLerpChanged(const FCameraPose& ToPose, float Factor, c
 		{\
 			if ((!bChangedOnly || ToPoseChangedFlags.PropName) && Factor >= 0.5f)\
 			{\
-				ensureMsgf(ChangedFlags.PropName, TEXT("Interpolating " #PropName " from default value!"));\
 				Set##PropName(ToPose.Get##PropName());\
 			}\
 			OutMask.PropName = true;\
diff --git a/Engine/Plugins/Cameras/GameplayCameras/Source/GameplayCameras/Private/Core/CameraSystemEvaluator.cpp b/Engine/Plugins/Cameras/GameplayCameras/Source/GameplayCameras/Private/Core/CameraSystemEvaluator.cpp
--- a/Engine/Plugins/Cameras/GameplayCameras/Source/GameplayCameras/Private/Core/CameraSystemEvaluator.cpp
+++ b/Engine/Plugins/Cameras/GameplayCameras/Source/GameplayCameras/Private/Core/CameraSystemEvaluator.cpp
@@ -27,6 +27,27 @@ namespace UE::Cameras
 
 extern bool GGameplayCamerasDebugEnable;
 
+void FCameraSystemEvaluationResult::Reset()
+{
+	CameraPose.ClearAllChangedFlags();
+	VariableTable.ClearAllWrittenThisFrameFlags();
+	bIsCameraCut = false;
+	bIsValid = false;
+}
+
+void FCameraSystemEvaluationResult::Reset(const FCameraNodeEvaluationResult& NodeResult)
+{
+	Reset();
+
+	// Make the camera poses actually equal, so that we get the exact same changed-flags.
+	CameraPose = NodeResult.CameraPose;
+
+	VariableTable.OverrideAll(NodeResult.VariableTable);
+
+	bIsCameraCut = NodeResult.bIsCameraCut;
+	bIsValid = true;
+}
+
 FCameraSystemEvaluator::FCameraSystemEvaluator()
 {
 }
@@ -176,6 +197,17 @@ void FCameraSystemEvaluator::Update(const FCameraSystemEvaluationParams& Params)
 {
 	SCOPE_CYCLE_COUNTER(CameraSystemEval_Total);
 
+	// Reset our result' flags.
+	RootNodeResult.CameraPose.ClearAllChangedFlags();
+	RootNodeResult.VariableTable.ClearAllWrittenThisFrameFlags();
+
+	// Run the variable auto-reset service here, because the other (third party) services
+	// should get the reset variable values.
+	if (VariableAutoResetService)
+	{
+		VariableAutoResetService->PerformVariableResets(RootNodeResult.VariableTable, ContextStack);
+	}
+
 	// Pre-update all services.
 	PreUpdateServices(Params.DeltaTime, ECameraEvaluationServiceFlags::None);
 
@@ -248,7 +280,7 @@ void FCameraSystemEvaluator::Update(const FCameraSystemEvaluationParams& Params)
 		NodeParams.Evaluator = this;
 		NodeParams.DeltaTime = Params.DeltaTime;
 
-		RootNodeResult.Reset(false);
+		RootNodeResult.Reset();
 
 		// Run the root camera node.
 		RootEvaluator->Run(NodeParams, RootNodeResult);
@@ -260,10 +292,10 @@ void FCameraSystemEvaluator::Update(const FCameraSystemEvaluationParams& Params)
 	PostUpdateServices(Params.DeltaTime, ECameraEvaluationServiceFlags::None);
 
 	// Harvest the result.
-	Result.CameraPose.OverrideAll(RootNodeResult.CameraPose);
-	Result.VariableTable.OverrideAll(RootNodeResult.VariableTable);
-	Result.bIsCameraCut = RootNodeResult.bIsCameraCut;
-	Result.bIsValid = true;
+	Result.Reset(RootNodeResult);
+
+	// End of update things...
+	ContextStack.OnEndCameraSystemUpdate();
 }
 
 void FCameraSystemEvaluator::PreUpdateServices(float DeltaTime, ECameraEvaluationServiceFlags ExtraFlags)
diff --git a/Engine/Plugins/Cameras/GameplayCameras/Source/GameplayCameras/Private/Core/CameraVariableTable.cpp b/Engine/Plugins/Cameras/GameplayCameras/Source/GameplayCameras/Private/Core/CameraVariableTable.cpp
--- a/Engine/Plugins/Cameras/GameplayCameras/Source/GameplayCameras/Private/Core/CameraVariableTable.cpp
+++ b/Engine/Plugins/Cameras/GameplayCameras/Source/GameplayCameras/Private/Core/CameraVariableTable.cpp
@@ -227,7 +227,11 @@ bool FCameraVariableTable::ContainsValue(FCameraVariableID VariableID) const
 	return EntryLookup.Contains(VariableID);
 }
 
-void FCameraVariableTable::SetValue(FCameraVariableID VariableID, ECameraVariableType ExpectedVariableType, const uint8* InRawValuePtr)
+void FCameraVariableTable::SetValue(
+		FCameraVariableID VariableID, 
+		ECameraVariableType ExpectedVariableType, 
+		const uint8* InRawValuePtr,
+		bool bMarkAsWrittenThisFrame)
 {
 	FEntry* Entry = FindEntry(VariableID);
 	if (ensureMsgf(Entry, TEXT("Can't set camera variable (ID '%d') because it doesn't exist in the table."), VariableID.GetValue()))
@@ -237,11 +241,19 @@ void FCameraVariableTable::SetValue(FCameraVariableID VariableID, ECameraVariabl
 		GetVariableTypeAllocationInfo(Entry->Type, SizeOf, AlignOf);
 		uint8* ValuePtr = Memory + Entry->Offset;
 		FMemory::Memcpy(ValuePtr, InRawValuePtr, SizeOf);
-		Entry->Flags |= EEntryFlags::Written | EEntryFlags::WrittenThisFrame;
+		Entry->Flags |= EEntryFlags::Written;
+		if (bMarkAsWrittenThisFrame)
+		{
+			Entry->Flags |= EEntryFlags::WrittenThisFrame;
+		}
 	}
 }
 
-bool FCameraVariableTable::TrySetValue(FCameraVariableID VariableID, ECameraVariableType ExpectedVariableType, const uint8* InRawValuePtr)
+bool FCameraVariableTable::TrySetValue(
+		FCameraVariableID VariableID,
+		ECameraVariableType ExpectedVariableType,
+		const uint8* InRawValuePtr,
+		bool bMarkAsWrittenThisFrame)
 {
 	if (FEntry* Entry = FindEntry(VariableID))
 	{
@@ -250,7 +262,11 @@ bool FCameraVariableTable::TrySetValue(FCameraVariableID VariableID, ECameraVari
 		GetVariableTypeAllocationInfo(Entry->Type, SizeOf, AlignOf);
 		uint8* ValuePtr = Memory + Entry->Offset;
 		FMemory::Memcpy(ValuePtr, InRawValuePtr, SizeOf);
-		Entry->Flags |= EEntryFlags::Written | EEntryFlags::WrittenThisFrame;
+		Entry->Flags |= EEntryFlags::Written;
+		if (bMarkAsWrittenThisFrame)
+		{
+			Entry->Flags |= EEntryFlags::WrittenThisFrame;
+		}
 		return true;
 	}
 	return false;
@@ -412,7 +428,7 @@ void FCameraVariableTable::InternalOverride(const FCameraVariableTable& OtherTab
 				uint8* ThisValuePtr = Memory + ThisEntry->Offset;
 				const uint8* OtherValuePtr = OtherTable.Memory + OtherEntry.Offset;
 				FMemory::Memcpy(ThisValuePtr, OtherValuePtr, ValueSize);
-				EnumAddFlags(ThisEntry->Flags, EEntryFlags::Written);
+				EnumAddFlags(ThisEntry->Flags, EEntryFlags::Written | (OtherFlags & EEntryFlags::WrittenThisFrame));
 
 				if (OutMask)
 				{
@@ -466,7 +482,7 @@ void FCameraVariableTable::InternalLerp(const FCameraVariableTable& ToTable, ECa
 				// We already have the other table's variable in our table. Let's check
 				// that the types match, and then interpolate the values.
 #if WITH_EDITORONLY_DATA
-				checkf(FromEntry->DebugName == ToEntry.DebugName,
+				ensureMsgf(FromEntry->DebugName == ToEntry.DebugName,
 						TEXT("Camera variable name collision! Expected variable '%d' to be named '%s', but other table has '%s'!"),
 						FromEntry->ID.GetValue(), *FromEntry->DebugName, *ToEntry.DebugName);
 #endif
@@ -476,7 +492,7 @@ void FCameraVariableTable::InternalLerp(const FCameraVariableTable& ToTable, ECa
 #else
 				const FString& DebugName = GUnavailableVariableDebugName;
 #endif
-				checkf(FromEntry->Type == ToEntry.Type, 
+				ensureMsgf(FromEntry->Type == ToEntry.Type, 
 						TEXT("Camera variable name collision! Expected '%d' (%s) to be of type '%s' but other table has type '%s'!"),
 						FromEntry->ID.GetValue(), *DebugName,
 						*UEnum::GetValueAsString(FromEntry->Type), *UEnum::GetValueAsString(ToEntry.Type));
@@ -497,7 +513,11 @@ void FCameraVariableTable::InternalLerp(const FCameraVariableTable& ToTable, ECa
 UE_CAMERA_VARIABLE_FOR_ALL_TYPES()
 #undef UE_CAMERA_VARIABLE_FOR_TYPE
 				}
-				EnumAddFlags(FromEntry->Flags, EEntryFlags::Written);
+				// We consider this variable "written to this frame" if it was written in either variable tables this frame.
+				// If the value interpolates because the from/to values are different, but neither was written this frame, we
+				// don't consider the interpolated value written this frame either.
+				const EEntryFlags FromFlags = FromEntry->Flags;
+				EnumAddFlags(FromEntry->Flags, EEntryFlags::Written | (FromFlags & EEntryFlags::WrittenThisFrame) | (ToFlags & EEntryFlags::WrittenThisFrame));
 			}
 			else
 			{
@@ -521,7 +541,7 @@ UE_CAMERA_VARIABLE_FOR_ALL_TYPES()
 				uint8* FromValuePtr = Memory + FromEntry->Offset;
 				const uint8* ToValuePtr = ToTable.Memory + ToEntry.Offset;
 				FMemory::Memcpy(FromValuePtr, ToValuePtr, ValueSize);
-				EnumAddFlags(FromEntry->Flags, EEntryFlags::Written);
+				EnumAddFlags(FromEntry->Flags, EEntryFlags::Written | (ToFlags & EEntryFlags::WrittenThisFrame));
 			}
 
 			if (OutMask)
diff --git a/Engine/Plugins/Cameras/GameplayCameras/Source/GameplayCameras/Private/Core/RootCameraNode.cpp b/Engine/Plugins/Cameras/GameplayCameras/Source/GameplayCameras/Private/Core/RootCameraNode.cpp
--- a/Engine/Plugins/Cameras/GameplayCameras/Source/GameplayCameras/Private/Core/RootCameraNode.cpp
+++ b/Engine/Plugins/Cameras/GameplayCameras/Source/GameplayCameras/Private/Core/RootCameraNode.cpp
@@ -39,13 +39,8 @@ void FRootCameraNodeEvaluator::RunSingleCameraRig(const FSingleCameraRigEvaluati
 	// an update result that isn't representative of what would happen normally.
 	if (FCameraSystemEvaluator* Evaluator = Params.EvaluationParams.Evaluator)
 	{
-		FCameraEvaluationServiceUpdateParams ServiceUpdateParams;
-		ServiceUpdateParams.Evaluator = Evaluator;
-		ServiceUpdateParams.DeltaTime = Params.EvaluationParams.DeltaTime;
-
-		FCameraEvaluationServiceUpdateResult ServiceUpdateResult(OutResult);
-
-		Evaluator->VariableAutoResetService->PreUpdate(ServiceUpdateParams, ServiceUpdateResult);
+		// TODO: we might have to reset variables on the context's initial result too?
+		Evaluator->VariableAutoResetService->PerformVariableResets(OutResult.VariableTable);
 	}
 
 	OnRunSingleCameraRig(Params, OutResult);
diff --git a/Engine/Plugins/Cameras/GameplayCameras/Source/GameplayCameras/Private/GameFramework/BlueprintCameraVariableTable.cpp b/Engine/Plugins/Cameras/GameplayCameras/Source/GameplayCameras/Private/GameFramework/BlueprintCameraVariableTable.cpp
--- a/Engine/Plugins/Cameras/GameplayCameras/Source/GameplayCameras/Private/GameFramework/BlueprintCameraVariableTable.cpp
+++ b/Engine/Plugins/Cameras/GameplayCameras/Source/GameplayCameras/Private/GameFramework/BlueprintCameraVariableTable.cpp
@@ -34,7 +34,7 @@
 	VariableTable.GetVariableTable()->SetValue(Variable, Value, true);\
 	if (Variable->bAutoReset && VariableTable.VariableAutoResetService)\
 	{\
-		VariableTable.VariableAutoResetService->AddAutoResetVariable(Variable);\
+		VariableTable.VariableAutoResetService->RegisterVariableUseFromScripting(Variable);\
 	}
 
 FBlueprintCameraVariableTable::FBlueprintCameraVariableTable()
diff --git a/Engine/Plugins/Cameras/GameplayCameras/Source/GameplayCameras/Private/Services/AutoResetCameraVariableService.cpp b/Engine/Plugins/Cameras/GameplayCameras/Source/GameplayCameras/Private/Services/AutoResetCameraVariableService.cpp
--- a/Engine/Plugins/Cameras/GameplayCameras/Source/GameplayCameras/Private/Services/AutoResetCameraVariableService.cpp
+++ b/Engine/Plugins/Cameras/GameplayCameras/Source/GameplayCameras/Private/Services/AutoResetCameraVariableService.cpp
@@ -2,44 +2,120 @@
 
 #include "Services/AutoResetCameraVariableService.h"
 
+#include "Core/CameraEvaluationContext.h"
+#include "Core/CameraEvaluationContextStack.h"
 #include "Core/CameraNodeEvaluator.h"
 #include "Core/CameraRigAsset.h"
 #include "Core/CameraVariableAssets.h"
 #include "Core/CameraVariableTable.h"
 #include "Core/RootCameraNodeCameraRigEvent.h"
+#include "Debug/CameraDebugBlock.h"
+#include "Debug/CameraDebugBlockBuilder.h"
+#include "Debug/CameraDebugRenderer.h"
 
 namespace UE::Cameras
 {
 
 UE_DEFINE_CAMERA_EVALUATION_SERVICE(FAutoResetCameraVariableService)
 
-void FAutoResetCameraVariableService::OnInitialize(const FCameraEvaluationServiceInitializeParams& Params)
+FAutoResetCameraVariableService::FAutoResetCameraVariableService()
 {
-	SetEvaluationServiceFlags(
-			ECameraEvaluationServiceFlags::NeedsPreUpdate |
-			ECameraEvaluationServiceFlags::NeedsRootCameraNodeEvents);
+	SetEvaluationServiceFlags(ECameraEvaluationServiceFlags::NeedsRootCameraNodeEvents);
 }
 
-void FAutoResetCameraVariableService::OnPreUpdate(const FCameraEvaluationServiceUpdateParams& Params, FCameraEvaluationServiceUpdateResult& OutResult)
+void FAutoResetCameraVariableService::PerformVariableResets(FCameraVariableTable& RootVariableTable)
 {
-	FCameraVariableTable& VariableTable = OutResult.EvaluationResult.VariableTable;
+	DoPerformVariableResets(RootVariableTable, false);
+}
+
+void FAutoResetCameraVariableService::PerformVariableResets(FCameraVariableTable& RootVariableTable, const FCameraEvaluationContextStack& ContextStack)
+{
+	// Reset all variables in the root table.
+	DoPerformVariableResets(RootVariableTable, false);
+
+	// For evaluation contexts, only reset what wasn't set this frame by arbitrary code/logic (see the
+	// comment later in DoPerformVariableResets).
+	TArray<TSharedPtr<FCameraEvaluationContext>> Contexts;
+	ContextStack.GetAllContexts(Contexts);
+	for (TSharedPtr<FCameraEvaluationContext> Context : Contexts)
+	{
+		DoPerformVariableResets(Context);
+	}
 
+	RemoveMarkedVariablesAndClearFlags();
+}
+
+void FAutoResetCameraVariableService::DoPerformVariableResets(FCameraVariableTable& VariableTable, bool bOnlyNotWrittenThisFrame)
+{
 	for (auto It = AutoResetVariables.CreateIterator(); It; ++It)
 	{
+		FEntry& Entry = It.Value();
 		const UCameraVariableAsset* Variable = It.Key().Get();
-		if (Variable)
+
+		// If the variable was GC'ed, remove it from the list.
+		if (!Variable)
+		{
+			Entry.bMarkedForRemoval = true;
+			continue;
+		}
+
+		// Don't reset variables that were written this frame. For evaluation contexts' initial result, these
+		// variables may have been written by gameplay systems or Blueprint logic or whatever. We only want
+		// to reset them if they haven't been touched this frame.
+		const bool bValueWasWrittenLastFrame = VariableTable.IsValueWrittenThisFrame(Variable->GetVariableID());
+		if (bOnlyNotWrittenThisFrame && bValueWasWrittenLastFrame)
+		{
+			continue;
+		}
+
+		const FCameraVariableID VariableID = Variable->GetVariableID();
+		const uint8* DefaultValuePtr = Variable->GetDefaultValuePtr();
+		const ECameraVariableType VariableType = Variable->GetVariableType();
+
+		// We use TrySetValue instead of SetValue because we only know of variables *possibly* used by camera rigs. 
+		// This doesn't mean these variables have been added to the table and written to. For instance, a camera
+		// parameter on a node might be configured to use a variable that isn't set, which makes it fallback to
+		// the variable's default value when it's not found in the variable table. So we silently ignore variables
+		// not present in the variable table here.
+		//
+		// Also, note that we don't mark the variable as written this frame. We're resetting it to its default 
+		// value so that doesn't qualify.
+		VariableTable.TrySetValue(VariableID, VariableType, DefaultValuePtr, false);
+
+		// See if this variable is still used by anybody. If not, we remove it from our list.
+		if (Entry.RefCount == 0 && !Entry.bUsedByScripting)
+		{
+			Entry.bMarkedForRemoval = true;
+		}
+	}
+}
+
+void FAutoResetCameraVariableService::DoPerformVariableResets(TSharedPtr<FCameraEvaluationContext> EvaluationContext)
+{
+	if (EvaluationContext)
+	{
+		FCameraNodeEvaluationResult& ContextResult = EvaluationContext->GetInitialResult();
+		DoPerformVariableResets(ContextResult.VariableTable, true);
+
+		for (TSharedPtr<FCameraEvaluationContext> ChildContext : EvaluationContext->GetChildrenContexts())
 		{
-			const FCameraVariableID VariableID = Variable->GetVariableID();
-			const uint8* DefaultValuePtr = Variable->GetDefaultValuePtr();
-			const ECameraVariableType VariableType = Variable->GetVariableType();
-			// Using TrySetValue instead of SetValue because we only know of variables *possibly* used by camera rigs. 
-			// This doesn't mean these variables have been added to the table and written to. For instance, a camera
-			// parameter on a node might be configured to use a variable that isn't set, which makes it fallback to
-			// the variable's default value when it's not found in the variable table. So we silently ignore variables
-			// not present in the variable table here.
-			VariableTable.TrySetValue(VariableID, VariableType, DefaultValuePtr);
+			DoPerformVariableResets(ChildContext);
 		}
-		else
+	}
+}
+
+void FAutoResetCameraVariableService::RemoveMarkedVariablesAndClearFlags()
+{
+	for (auto It = AutoResetVariables.CreateIterator(); It; ++It)
+	{
+		FEntry& Entry = It.Value();
+
+#if UE_GAMEPLAY_CAMERAS_DEBUG
+		Entry.bDebugWasUsedByScripting = Entry.bUsedByScripting;
+#endif
+		Entry.bUsedByScripting = false;
+
+		if (Entry.bMarkedForRemoval)
 		{
 			It.RemoveCurrent();
 		}
@@ -77,24 +153,136 @@ void FAutoResetCameraVariableService::OnRootCameraNodeEvent(const FRootCameraNod
 	}
 }
 
+void FAutoResetCameraVariableService::RegisterVariableUseFromScripting(const UCameraVariableAsset* InVariable)
+{
+	if (ensure(InVariable))
+	{
+		FEntry& Entry = AutoResetVariables.FindOrAdd(InVariable);
+		Entry.bUsedByScripting = true;
+	}
+}
+
 void FAutoResetCameraVariableService::AddAutoResetVariable(const UCameraVariableAsset* InVariable)
 {
-	uint32& RefCount = AutoResetVariables.FindOrAdd(InVariable, 0);
-	++RefCount;
+	if (ensure(InVariable))
+	{
+		FEntry& Entry = AutoResetVariables.FindOrAdd(InVariable);
+		Entry.RefCount += 1;
+	}
 }
 
 void FAutoResetCameraVariableService::RemoveAutoResetVariable(const UCameraVariableAsset* InVariable)
 {
-	uint32* RefCount = AutoResetVariables.Find(InVariable);
-	if (ensure(RefCount))
+	if (ensure(InVariable))
 	{
-		--(*RefCount);
-		if (*RefCount == 0)
+		FEntry* Entry = AutoResetVariables.Find(InVariable);
+		if (ensure(Entry && Entry->RefCount > 0))
 		{
-			AutoResetVariables.Remove(InVariable);
+			Entry->RefCount -= 1;
 		}
 	}
 }
 
+#if UE_GAMEPLAY_CAMERAS_DEBUG
+
+class FAutoResetCameraVariableDebugBlock : public FCameraDebugBlock
+{
+	UE_DECLARE_CAMERA_DEBUG_BLOCK(GAMEPLAYCAMERAS_API, FAutoResetCameraVariableDebugBlock)
+
+public:
+
+	FAutoResetCameraVariableDebugBlock() = default;
+	FAutoResetCameraVariableDebugBlock(const FAutoResetCameraVariableService& InService);
+	
+protected:
+
+	// FCameraDebugBlock interface.
+	virtual void OnDebugDraw(const FCameraDebugBlockDrawParams& Params, FCameraDebugRenderer& Renderer) override;
+	virtual void OnSerialize(FArchive& Ar) override;
+
+private:
+
+	struct FDebugEntry
+	{
+		FString VariableName;
+		uint32 RefCount = 0;
+		bool bUsedByScripting;
+	};
+	TArray<FDebugEntry> AutoResetVariables;
+
+	friend FArchive& operator<< (FArchive& Ar, FDebugEntry& Entry);
+};
+
+UE_DEFINE_CAMERA_DEBUG_BLOCK(FAutoResetCameraVariableDebugBlock)
+
+void FAutoResetCameraVariableService::OnBuildDebugBlocks(const FCameraDebugBlockBuildParams& Params, FCameraDebugBlockBuilder& Builder)
+{
+	Builder.AttachDebugBlock<FAutoResetCameraVariableDebugBlock>(*this);
+}
+
+FAutoResetCameraVariableDebugBlock::FAutoResetCameraVariableDebugBlock(const FAutoResetCameraVariableService& InService)
+{
+	using FEntry = FAutoResetCameraVariableService::FEntry;
+	for (const TPair<TWeakObjectPtr<const UCameraVariableAsset>, FEntry>& Pair : InService.AutoResetVariables)
+	{
+		const UCameraVariableAsset* Variable = Pair.Key.Get();
+		const FEntry& Entry = Pair.Value;
+		const FString VariableName = Variable ? 
+#if WITH_EDITORONLY_DATA
+			Variable->GetDisplayName() : 
+#else
+			Variable->GetName() :
+#endif  // WITH_EDITORONLY_DATA
+			TEXT("<None>");
+		AutoResetVariables.Add({ VariableName, Entry.RefCount, Entry.bDebugWasUsedByScripting });
+	}
+}
+
+void FAutoResetCameraVariableDebugBlock::OnDebugDraw(const FCameraDebugBlockDrawParams& Params, FCameraDebugRenderer& Renderer)
+{
+	Renderer.AddText(TEXT("%d auto-reset variables\n"), AutoResetVariables.Num());
+	Renderer.AddIndent();
+	{
+		for (const FDebugEntry& DebugEntry : AutoResetVariables)
+		{
+			Renderer.AddText(TEXT("{cam_notice}%s{cam_default} :"), *DebugEntry.VariableName);
+			if (DebugEntry.RefCount > 0)
+			{
+				Renderer.AddText(TEXT(" used {cam_notice2}%d{cam_default} times"), DebugEntry.RefCount);
+				if (DebugEntry.bUsedByScripting)
+				{
+					Renderer.AddText(TEXT(" ;"));
+				}
+			}
+			if (DebugEntry.bUsedByScripting)
+			{
+				Renderer.AddText(TEXT(" used by {cam_notice2}scripting{cam_default}"));
+			}
+			if (DebugEntry.RefCount == 0 && !DebugEntry.bUsedByScripting)
+			{
+				Renderer.AddText(TEXT(" {cam_passive}not used{cam_default}"));
+			}
+			Renderer.AddText(TEXT("\n"));
+		}
+	}
+	Renderer.RemoveIndent();
+}
+
+void FAutoResetCameraVariableDebugBlock::OnSerialize(FArchive& Ar)
+{
+	Ar << AutoResetVariables;
+}
+
+FArchive& operator<< (FArchive& Ar, FAutoResetCameraVariableDebugBlock::FDebugEntry& Entry)
+{
+	Ar << Entry.VariableName;
+	Ar << Entry.RefCount;
+	Ar << Entry.bUsedByScripting;
+
+	return Ar;
+}
+
+#endif  // UE_GAMEPLAY_CAMERAS_DEBUG
+
 }  // namespace UE::Cameras
 
diff --git a/Engine/Plugins/Cameras/GameplayCameras/Source/GameplayCameras/Private/Services/AutoResetCameraVariableService.h b/Engine/Plugins/Cameras/GameplayCameras/Source/GameplayCameras/Private/Services/AutoResetCameraVariableService.h
--- a/Engine/Plugins/Cameras/GameplayCameras/Source/GameplayCameras/Private/Services/AutoResetCameraVariableService.h
+++ b/Engine/Plugins/Cameras/GameplayCameras/Source/GameplayCameras/Private/Services/AutoResetCameraVariableService.h
@@ -10,27 +10,69 @@ class UCameraVariableAsset;
 namespace UE::Cameras
 {
 
+class FCameraEvaluationContext;
+class FCameraVariableTable;
+struct FCameraEvaluationContextStack;
+
 class FAutoResetCameraVariableService : public FCameraEvaluationService
 {
 	UE_DECLARE_CAMERA_EVALUATION_SERVICE(GAMEPLAYCAMERAS_API, FAutoResetCameraVariableService)
 
 public:
 
+	FAutoResetCameraVariableService();
+
 	/** Adds a variable to the list of variables to reset every update. */
 	void AddAutoResetVariable(const UCameraVariableAsset* InVariable);
 	/** Remove a variable from the reset list. */
 	void RemoveAutoResetVariable(const UCameraVariableAsset* InVariable);
 
+	/**
+	 * Marks a variable as used from Blueprint, Verse, etc. This is an alternative
+	 * to AddAutoResetVariable, for when it is unknown when the variable isn't used
+	 * anymore. The variable will be removed from the list after it is seen that it
+	 * hasn't been written to during a frame.
+	 */
+	void RegisterVariableUseFromScripting(const UCameraVariableAsset* InVariable);
+
+public:
+
+	// Internal API.
+	void PerformVariableResets(FCameraVariableTable& RootVariableTable);
+	void PerformVariableResets(FCameraVariableTable& RootVariableTable, const FCameraEvaluationContextStack& ContextStack);
+
 protected:
 
 	// FCameraEvaluationService interface.
-	virtual void OnInitialize(const FCameraEvaluationServiceInitializeParams& Params) override;
-	virtual void OnPreUpdate(const FCameraEvaluationServiceUpdateParams& Params, FCameraEvaluationServiceUpdateResult& OutResult) override;
 	virtual void OnRootCameraNodeEvent(const FRootCameraNodeCameraRigEvent& InEvent) override;
 
+#if UE_GAMEPLAY_CAMERAS_DEBUG
+	virtual void OnBuildDebugBlocks(const FCameraDebugBlockBuildParams& Params, FCameraDebugBlockBuilder& Builder) override;
+#endif  // UE_GAMEPLAY_CAMERAS_DEBUG
+
+private:
+
+	void DoPerformVariableResets(FCameraVariableTable& VariableTable, bool bOnlyNotWrittenThisFrame);
+	void DoPerformVariableResets(TSharedPtr<FCameraEvaluationContext> EvaluationContext);
+	void RemoveMarkedVariablesAndClearFlags();
+
 private:
 
-	TMap<TWeakObjectPtr<const UCameraVariableAsset>, uint32> AutoResetVariables;
+	struct FEntry
+	{
+		uint32 RefCount = 0;
+		bool bUsedByScripting = false;
+		bool bMarkedForRemoval = false;
+#if UE_GAMEPLAY_CAMERAS_DEBUG
+		bool bDebugWasUsedByScripting = false;
+#endif
+	};
+
+	TMap<TWeakObjectPtr<const UCameraVariableAsset>, FEntry> AutoResetVariables;
+
+#if UE_GAMEPLAY_CAMERAS_DEBUG
+	friend class FAutoResetCameraVariableDebugBlock;
+#endif  // UE_GAMEPLAY_CAMERAS_DEBUG
 };
 
 }  // namespace UE::Cameras
diff --git a/Engine/Plugins/Cameras/GameplayCameras/Source/GameplayCameras/Public/Core/CameraEvaluationContextStack.h b/Engine/Plugins/Cameras/GameplayCameras/Source/GameplayCameras/Public/Core/CameraEvaluationContextStack.h
--- a/Engine/Plugins/Cameras/GameplayCameras/Source/GameplayCameras/Public/Core/CameraEvaluationContextStack.h
+++ b/Engine/Plugins/Cameras/GameplayCameras/Source/GameplayCameras/Public/Core/CameraEvaluationContextStack.h
@@ -36,6 +36,12 @@ public:
 	/** Push a new context on the stack and instantiate its director. */
 	void PushContext(TSharedRef<FCameraEvaluationContext> Context);
 
+	/**
+	 * Tries to add a context inside the active context. This will query the active context's
+	 * director in order to find an "available spot" for the child context.
+	 * 
+	 * @return Whether the child context was acccepted.
+	 */
 	bool AddChildContext(TSharedRef<FCameraEvaluationContext> Context);
 
 	/** Remove an existing context from the stack. */
@@ -47,6 +53,9 @@ public:
 	/** The number of contexts on the stack. */
 	int32 NumContexts() const { return Entries.Num(); }
 
+	/** Gets all the contexts in the stack, from bottom to top. */
+	void GetAllContexts(TArray<TSharedPtr<FCameraEvaluationContext>>& OutContexts) const;
+
 	/** Empties the stack of all contexts. */
 	void Reset();
 
@@ -55,6 +64,7 @@ public:
 	// Internal API
 	void Initialize(FCameraSystemEvaluator& InEvaluator);
 	void AddReferencedObjects(FReferenceCollector& Collector);
+	void OnEndCameraSystemUpdate();
 
 private:
 
diff --git a/Engine/Plugins/Cameras/GameplayCameras/Source/GameplayCameras/Public/Core/CameraNodeEvaluator.h b/Engine/Plugins/Cameras/GameplayCameras/Source/GameplayCameras/Public/Core/CameraNodeEvaluator.h
--- a/Engine/Plugins/Cameras/GameplayCameras/Source/GameplayCameras/Public/Core/CameraNodeEvaluator.h
+++ b/Engine/Plugins/Cameras/GameplayCameras/Source/GameplayCameras/Public/Core/CameraNodeEvaluator.h
@@ -179,7 +179,7 @@ struct GAMEPLAYCAMERAS_API FCameraNodeEvaluationResult
 public:
 
 	/** Reset this result to its default (non-valid) state.  */
-	void Reset(bool bResetVariableTable);
+	void Reset();
 
 	/** Override this result with the given other result. */
 	void OverrideAll(const FCameraNodeEvaluationResult& OtherResult);
diff --git a/Engine/Plugins/Cameras/GameplayCameras/Source/GameplayCameras/Public/Core/CameraPose.h b/Engine/Plugins/Cameras/GameplayCameras/Source/GameplayCameras/Public/Core/CameraPose.h
--- a/Engine/Plugins/Cameras/GameplayCameras/Source/GameplayCameras/Public/Core/CameraPose.h
+++ b/Engine/Plugins/Cameras/GameplayCameras/Source/GameplayCameras/Public/Core/CameraPose.h
@@ -89,8 +89,8 @@ public:
 
 	FCameraPose();
 
-	/** Resets this camera pose to its default values. */
-	void Reset(bool bSetAllChangedFlags = true);
+	/** Resets this camera pose to its default values, with all changed flags off. */
+	void Reset();
 
 public:
 
diff --git a/Engine/Plugins/Cameras/GameplayCameras/Source/GameplayCameras/Public/Core/CameraSystemEvaluator.h b/Engine/Plugins/Cameras/GameplayCameras/Source/GameplayCameras/Public/Core/CameraSystemEvaluator.h
--- a/Engine/Plugins/Cameras/GameplayCameras/Source/GameplayCameras/Public/Core/CameraSystemEvaluator.h
+++ b/Engine/Plugins/Cameras/GameplayCameras/Source/GameplayCameras/Public/Core/CameraSystemEvaluator.h
@@ -73,6 +73,14 @@ struct FCameraSystemEvaluationResult
 
 	/** Whether this result is valid. */
 	bool bIsValid = false;
+
+public:
+
+	/** Reset this result to its default (non-valid) state.  */
+	void Reset();
+
+	/** Set this result to be equivalent to the given evaluation result. */
+	void Reset(const FCameraNodeEvaluationResult& NodeResult);
 };
 
 #if UE_GAMEPLAY_CAMERAS_DEBUG
diff --git a/Engine/Plugins/Cameras/GameplayCameras/Source/GameplayCameras/Public/Core/CameraVariableTable.h b/Engine/Plugins/Cameras/GameplayCameras/Source/GameplayCameras/Public/Core/CameraVariableTable.h
--- a/Engine/Plugins/Cameras/GameplayCameras/Source/GameplayCameras/Public/Core/CameraVariableTable.h
+++ b/Engine/Plugins/Cameras/GameplayCameras/Source/GameplayCameras/Public/Core/CameraVariableTable.h
@@ -152,8 +152,18 @@ public:
 
 	// Lower level API.
 
-	void SetValue(FCameraVariableID VariableID, ECameraVariableType ExpectedVariableType, const uint8* InRawValuePtr);
-	bool TrySetValue(FCameraVariableID VariableID, ECameraVariableType ExpectedVariableType, const uint8* InRawValuePtr);
+	void SetValue(
+			FCameraVariableID VariableID, 
+			ECameraVariableType ExpectedVariableType, 
+			const uint8* InRawValuePtr,
+			bool bMarkAsWrittenThisFrame = true);
+
+	bool TrySetValue(
+			FCameraVariableID VariableID, 
+			ECameraVariableType ExpectedVariableType, 
+			const uint8* InRawValuePtr,
+			bool bMarkAsWrittenThisFrame = true);
+
 	bool IsValueWritten(FCameraVariableID VariableID) const;
 	void UnsetValue(FCameraVariableID VariableID);
 	void UnsetAllValues();
