Cameras: fix IK aiming not resetting values between iterations

diff --git a/Engine/Plugins/Cameras/GameplayCameras/Source/GameplayCameras/Private/Core/CameraIKAim.cpp b/Engine/Plugins/Cameras/GameplayCameras/Source/GameplayCameras/Private/Core/CameraIKAim.cpp
--- a/Engine/Plugins/Cameras/GameplayCameras/Source/GameplayCameras/Private/Core/CameraIKAim.cpp
+++ b/Engine/Plugins/Cameras/GameplayCameras/Source/GameplayCameras/Private/Core/CameraIKAim.cpp
@@ -61,8 +61,6 @@ bool FCameraIKAim::DoRun(const FCameraIKAimParams& Params, const FCameraRigEvalu
 
 	// Initialize our scratch result.
 	ScratchResult.VariableTable.Initialize(CameraRigInfo.CameraRig->AllocationInfo.VariableTableInfo);
-	ScratchResult.VariableTable.OverrideAll(CameraSystemEvaluator->GetEvaluatedResult().VariableTable);
-	ScratchResult.VariableTable.Override(CameraRigInfo.LastResult->VariableTable, ECameraVariableTableFilter::AllPublic | ECameraVariableTableFilter::Private);
 
 	// Initialize our hierarchy caches.
 	FRootCameraNodeEvaluator* CameraSystemRootEvaluator = CameraSystemEvaluator->GetRootNodeEvaluator();
@@ -137,6 +135,14 @@ void FCameraIKAim::RunRootCameraNode(const FCameraIKAimParams& Params, const FCa
 	const FCameraSystemEvaluationResult& LastResult = Params.Evaluator->GetEvaluatedResult();
 	const FCameraNodeEvaluationResult& ContextInitialResult = CameraRigInfo.EvaluationContext->GetInitialResult();
 
+	// Reset the scratch result the same way the camera system does it at the beginning of each frame.
+	ScratchResult.Reset();
+
+	// Make sure the camera rig will get its private variables, such as rig interface parameter overrides.
+	ScratchResult.VariableTable.Override(
+			CameraRigInfo.LastResult->VariableTable, 
+			ECameraVariableTableFilter::AllPublic | ECameraVariableTableFilter::Private);
+
 	FRootCameraNodeEvaluator* RootEvaluator = Params.Evaluator->GetRootNodeEvaluator();
 
 	// Save the initial state of the camera rig.
diff --git a/Engine/Plugins/Cameras/GameplayCameras/Source/GameplayCameras/Private/Core/DefaultRootCameraNode.cpp b/Engine/Plugins/Cameras/GameplayCameras/Source/GameplayCameras/Private/Core/DefaultRootCameraNode.cpp
--- a/Engine/Plugins/Cameras/GameplayCameras/Source/GameplayCameras/Private/Core/DefaultRootCameraNode.cpp
+++ b/Engine/Plugins/Cameras/GameplayCameras/Source/GameplayCameras/Private/Core/DefaultRootCameraNode.cpp
@@ -163,6 +163,8 @@ void FDefaultRootCameraNodeEvaluator::OnRunSingleCameraRig(const FSingleCameraRi
 
 	FCameraNodeEvaluator* RootEvaluator = Params.CameraRigInfo.RootEvaluator;
 
+	// Emulate what the main blend stack does.
+
 	{
 		const FCameraNodeEvaluationResult* CameraRigResult = Params.CameraRigInfo.LastResult;
 		FCameraBlendedParameterUpdateParams InputParams(Params.EvaluationParams, CameraRigResult->CameraPose);
@@ -172,21 +174,20 @@ void FDefaultRootCameraNodeEvaluator::OnRunSingleCameraRig(const FSingleCameraRi
 		Hierarchy.CallUpdateParameters(InputParams, InputResult);
 	}
 
-	{
-		OutResult.CameraPose.ClearAllChangedFlags();
-		OutResult.VariableTable.ClearAllWrittenThisFrameFlags();
+	// No parameter blending: we are running this camera rig in isolation.
 
+	{
 		const FCameraNodeEvaluationResult& InitialResult = Params.CameraRigInfo.EvaluationContext->GetInitialResult();
 		OutResult.CameraPose.OverrideChanged(InitialResult.CameraPose);
 		OutResult.VariableTable.OverrideAll(InitialResult.VariableTable);
 
-		OutResult.bIsValid = true;
-
 		RootEvaluator->Run(Params.EvaluationParams, OutResult);
 	}
 
 	GlobalLayer->Run(Params.EvaluationParams, OutResult);
 	// Don't run the visual layer.
+
+	OutResult.bIsValid = true;
 }
 
 void FDefaultRootCameraNodeEvaluator::OnBlendStackEvent(const FBlendStackCameraRigEvent& InEvent)
