Cameras: fix boom arm incorrectly executing operations on its input node

This could lead to the input node executing the operation twice. This didn't cause any particular bug since the yaw-pitch-correction operation consumes the angles, so the second time would only correct by zero degrees, but it is still incorrect.

diff --git a/Engine/Plugins/Cameras/GameplayCameras/Source/GameplayCameras/Private/Nodes/Common/BoomArmCameraNode.cpp b/Engine/Plugins/Cameras/GameplayCameras/Source/GameplayCameras/Private/Nodes/Common/BoomArmCameraNode.cpp
--- a/Engine/Plugins/Cameras/GameplayCameras/Source/GameplayCameras/Private/Nodes/Common/BoomArmCameraNode.cpp
+++ b/Engine/Plugins/Cameras/GameplayCameras/Source/GameplayCameras/Private/Nodes/Common/BoomArmCameraNode.cpp
@@ -210,11 +210,7 @@ void FBoomArmCameraNodeEvaluator::OnRun(const FCameraNodeEvaluationParams& Param
 
 void FBoomArmCameraNodeEvaluator::OnExecuteOperation(const FCameraOperationParams& Params, FCameraOperation& Operation)
 {
-	if (InputSlotEvaluator)
-	{
-		InputSlotEvaluator->ExecuteOperation(Params, Operation);
-	}
-	else
+	if (!InputSlotEvaluator)
 	{
 		// If we don't have an input slot, we use the pawn rotation directly in OnRun. So let's handle
 		// some operations by affecting that pawn rotation ourselves.
