Cameras: add debug display for camera rig joints in an evaluated result object

diff --git a/Engine/Plugins/Cameras/GameplayCameras/Source/GameplayCameras/Private/Core/CameraVariableTable.cpp b/Engine/Plugins/Cameras/GameplayCameras/Source/GameplayCameras/Private/Core/CameraVariableTable.cpp
--- a/Engine/Plugins/Cameras/GameplayCameras/Source/GameplayCameras/Private/Core/CameraVariableTable.cpp
+++ b/Engine/Plugins/Cameras/GameplayCameras/Source/GameplayCameras/Private/Core/CameraVariableTable.cpp
@@ -314,6 +314,22 @@ void FCameraVariableTable::ClearAllWrittenThisFrameFlags()
 	}
 }
 
+bool FCameraVariableTable::TryGetVariableDefinition(FCameraVariableID VariableID, FCameraVariableDefinition& OutVariableDefinition) const
+{
+	if (const FEntry* Entry = FindEntry(VariableID))
+	{
+		OutVariableDefinition.VariableID = Entry->ID;
+		OutVariableDefinition.VariableType = Entry->Type;
+		OutVariableDefinition.bIsPrivate = EnumHasAnyFlags(Entry->Flags, EEntryFlags::Private);
+		OutVariableDefinition.bIsInput = EnumHasAnyFlags(Entry->Flags, EEntryFlags::Input);
+#if WITH_EDITORONLY_DATA
+		OutVariableDefinition.VariableName = Entry->DebugName;
+#endif  // WITH_EDITORONLY_DATA
+		return true;
+	}
+	return false;
+}
+
 void FCameraVariableTable::Serialize(FArchive& Ar)
 {
 	if (Ar.IsSaving())
diff --git a/Engine/Plugins/Cameras/GameplayCameras/Source/GameplayCameras/Private/Debug/CameraNodeEvaluationResultDebugBlock.cpp b/Engine/Plugins/Cameras/GameplayCameras/Source/GameplayCameras/Private/Debug/CameraNodeEvaluationResultDebugBlock.cpp
--- a/Engine/Plugins/Cameras/GameplayCameras/Source/GameplayCameras/Private/Debug/CameraNodeEvaluationResultDebugBlock.cpp
+++ b/Engine/Plugins/Cameras/GameplayCameras/Source/GameplayCameras/Private/Debug/CameraNodeEvaluationResultDebugBlock.cpp
@@ -8,6 +8,7 @@
 #include "Debug/CameraDebugColors.h"
 #include "Debug/CameraDebugRenderer.h"
 #include "Debug/CameraPoseDebugBlock.h"
+#include "Debug/CameraRigJointsDebugBlock.h"
 #include "Debug/DebugTextRenderer.h"
 #include "Debug/VariableTableDebugBlock.h"
 #include "Math/ColorList.h"
@@ -25,11 +26,13 @@ FCameraNodeEvaluationResultDebugBlock::FCameraNodeEvaluationResultDebugBlock()
 
 void FCameraNodeEvaluationResultDebugBlock::Initialize(const FCameraNodeEvaluationResult& InResult, FCameraDebugBlockBuilder& Builder)
 {
+	NumPostProcessSettings = InResult.PostProcessSettings.GetEntries().Num();
 	bIsCameraCut = InResult.bIsCameraCut;
 	bIsValid = InResult.bIsValid;
 
 	AddChild(&Builder.BuildDebugBlock<FCameraPoseDebugBlock>(InResult.CameraPose));
 	AddChild(&Builder.BuildDebugBlock<FVariableTableDebugBlock>(InResult.VariableTable));
+	AddChild(&Builder.BuildDebugBlock<FCameraRigJointsDebugBlock>(InResult.CameraRigJoints, InResult.VariableTable));
 }
 
 void FCameraNodeEvaluationResultDebugBlock::Initialize(const FCameraSystemEvaluationResult& InResult, FCameraDebugBlockBuilder& Builder)
@@ -39,6 +42,7 @@ void FCameraNodeEvaluationResultDebugBlock::Initialize(const FCameraSystemEvalua
 
 	AddChild(&Builder.BuildDebugBlock<FCameraPoseDebugBlock>(InResult.CameraPose));
 	AddChild(&Builder.BuildDebugBlock<FVariableTableDebugBlock>(InResult.VariableTable));
+	// No rig joints exposed on the camera system result.
 }
 
 FCameraPoseDebugBlock* FCameraNodeEvaluationResultDebugBlock::GetCameraPoseDebugBlock()
@@ -79,6 +83,11 @@ void FCameraNodeEvaluationResultDebugBlock::OnDebugDraw(const FCameraDebugBlockD
 		Renderer.AddText(TEXT("  {cam_warning}IsCameraCut"));
 	}
 
+	if (NumPostProcessSettings > 0)
+	{
+		Renderer.AddText(TEXT("  {cam_notice}%d post-FX"), NumPostProcessSettings);
+	}
+
 	Renderer.NewLine();
 	Renderer.SetTextColor(Colors.Default);
 
@@ -99,6 +108,14 @@ void FCameraNodeEvaluationResultDebugBlock::OnDebugDraw(const FCameraDebugBlockD
 		ChildrenView[1]->DebugDraw(Params, Renderer);
 		Renderer.RemoveIndent();
 	}
+	if (ChildrenView.IsValidIndex(2))
+	{
+		Renderer.AddText(TEXT("{cam_title}Camera Rig Joints:"));
+		Renderer.AddIndent();
+		Renderer.SetTextColor(Colors.Default);
+		ChildrenView[2]->DebugDraw(Params, Renderer);
+		Renderer.RemoveIndent();
+	}
 
 	// We already manually rendered our children.
 	Renderer.SkipAllBlocks();
diff --git a/Engine/Plugins/Cameras/GameplayCameras/Source/GameplayCameras/Private/Debug/CameraRigJointsDebugBlock.cpp b/Engine/Plugins/Cameras/GameplayCameras/Source/GameplayCameras/Private/Debug/CameraRigJointsDebugBlock.cpp
new file mode 100644
--- /dev/null
+++ b/Engine/Plugins/Cameras/GameplayCameras/Source/GameplayCameras/Private/Debug/CameraRigJointsDebugBlock.cpp
@@ -0,0 +1,82 @@
+// Copyright Epic Games, Inc. All Rights Reserved.
+
+#include "Debug/CameraRigJointsDebugBlock.h"
+
+#include "Core/CameraVariableTable.h"
+#include "Debug/CameraDebugColors.h"
+#include "Debug/CameraDebugRenderer.h"
+
+#if UE_GAMEPLAY_CAMERAS_DEBUG
+
+namespace UE::Cameras
+{
+
+UE_DEFINE_CAMERA_DEBUG_BLOCK(FCameraRigJointsDebugBlock)
+
+FCameraRigJointsDebugBlock::FCameraRigJointsDebugBlock()
+{
+}
+
+FCameraRigJointsDebugBlock::FCameraRigJointsDebugBlock(const FCameraRigJoints& InCameraRigJoints, const FCameraVariableTable& InVariableTable)
+{
+	for (const FCameraRigJoint& Joint : InCameraRigJoints.GetJoints())
+	{
+		FEntry Entry;
+		Entry.Location = Joint.Transform.GetLocation();
+		Entry.Rotation = Joint.Transform.Rotator();
+		Entry.VariableID = Joint.VariableID;
+
+#if WITH_EDITORONLY_DATA
+		FCameraVariableDefinition VariableDefinition;
+		if (InVariableTable.TryGetVariableDefinition(Joint.VariableID, VariableDefinition))
+		{
+			Entry.VariableName = VariableDefinition.VariableName;
+		}
+		else
+		{
+			Entry.VariableName = TEXT("<Unknown>");
+		}
+#else
+		Entry.VariableName = TEXT("<Unknown>");
+#endif  // WITH_EDITORONLY_DATA
+	}
+}
+
+void FCameraRigJointsDebugBlock::OnDebugDraw(const FCameraDebugBlockDrawParams& Params, FCameraDebugRenderer& Renderer)
+{
+	Renderer.AddText(TEXT("%d joints"), Entries.Num());
+	Renderer.AddIndent();
+	{
+		for (const FEntry& Entry : Entries)
+		{
+			Renderer.AddText(
+					TEXT("location {cam_notice}%s{cam_default}  "
+						"rotation {cam_notice}%s{cam_default}  "
+						"variable {cam_notice2}%s{cam_passive} [%d]{cam_default}"),
+					*Entry.Location.ToString(),
+					*Entry.Rotation.ToString(),
+					*Entry.VariableName,
+					Entry.VariableID);
+		}
+	}
+	Renderer.RemoveIndent();
+}
+
+void FCameraRigJointsDebugBlock::OnSerialize(FArchive& Ar)
+{
+	Ar << Entries;
+}
+
+FArchive& operator<< (FArchive& Ar, FCameraRigJointsDebugBlock::FEntry& Entry)
+{
+	Ar << Entry.Location;
+	Ar << Entry.Rotation;
+	Ar << Entry.VariableID;
+	Ar << Entry.VariableName;
+	return Ar;
+}
+
+}  // namespace UE::Cameras
+
+#endif  // UE_GAMEPLAY_CAMERAS_DEBUG
+
diff --git a/Engine/Plugins/Cameras/GameplayCameras/Source/GameplayCameras/Public/Core/CameraVariableTable.h b/Engine/Plugins/Cameras/GameplayCameras/Source/GameplayCameras/Public/Core/CameraVariableTable.h
--- a/Engine/Plugins/Cameras/GameplayCameras/Source/GameplayCameras/Public/Core/CameraVariableTable.h
+++ b/Engine/Plugins/Cameras/GameplayCameras/Source/GameplayCameras/Public/Core/CameraVariableTable.h
@@ -171,6 +171,8 @@ public:
 	bool IsValueWrittenThisFrame(FCameraVariableID VariableID) const;
 	void ClearAllWrittenThisFrameFlags();
 
+	bool TryGetVariableDefinition(FCameraVariableID VariableID, FCameraVariableDefinition& OutVariableDefinition) const;
+
 	void Serialize(FArchive& Ar);
 
 private:
diff --git a/Engine/Plugins/Cameras/GameplayCameras/Source/GameplayCameras/Public/Debug/CameraNodeEvaluationResultDebugBlock.h b/Engine/Plugins/Cameras/GameplayCameras/Source/GameplayCameras/Public/Debug/CameraNodeEvaluationResultDebugBlock.h
--- a/Engine/Plugins/Cameras/GameplayCameras/Source/GameplayCameras/Public/Debug/CameraNodeEvaluationResultDebugBlock.h
+++ b/Engine/Plugins/Cameras/GameplayCameras/Source/GameplayCameras/Public/Debug/CameraNodeEvaluationResultDebugBlock.h
@@ -57,6 +57,7 @@ protected:
 
 private:
 
+	int32 NumPostProcessSettings = 0;
 	bool bIsCameraCut;
 	bool bIsValid;
 };
diff --git a/Engine/Plugins/Cameras/GameplayCameras/Source/GameplayCameras/Public/Debug/CameraRigJointsDebugBlock.h b/Engine/Plugins/Cameras/GameplayCameras/Source/GameplayCameras/Public/Debug/CameraRigJointsDebugBlock.h
new file mode 100644
--- /dev/null
+++ b/Engine/Plugins/Cameras/GameplayCameras/Source/GameplayCameras/Public/Debug/CameraRigJointsDebugBlock.h
@@ -0,0 +1,51 @@
+// Copyright Epic Games, Inc. All Rights Reserved.
+
+#pragma once
+
+#include "Core/CameraRigJoints.h"
+#include "Debug/CameraDebugBlock.h"
+
+#if UE_GAMEPLAY_CAMERAS_DEBUG
+
+namespace UE::Cameras
+{
+
+class FCameraVariableTable;
+
+/**
+ * A debug block that displays information about a camera rig's joints.
+ */
+class FCameraRigJointsDebugBlock : public FCameraDebugBlock
+{
+	UE_DECLARE_CAMERA_DEBUG_BLOCK(GAMEPLAYCAMERAS_API, FCameraRigJointsDebugBlock)
+
+public:
+
+	/** Creates a new camera rig joints debug block. */
+	FCameraRigJointsDebugBlock();
+	/** Creates a new camera rig joints debug block. */
+	FCameraRigJointsDebugBlock(const FCameraRigJoints& InCameraRigJoints, const FCameraVariableTable& InVariableTable);
+
+protected:
+
+	virtual void OnDebugDraw(const FCameraDebugBlockDrawParams& Params, FCameraDebugRenderer& Renderer) override;
+	virtual void OnSerialize(FArchive& Ar) override;
+
+private:
+
+	struct FEntry
+	{
+		FVector3d Location;
+		FRotator3d Rotation;
+		FCameraVariableID VariableID;
+		FString VariableName;
+	};
+	TArray<FEntry, TInlineAllocator<2>> Entries;
+
+	friend FArchive& operator<< (FArchive& Ar, FEntry& Entry);
+};
+
+}  // namespace UE::Cameras
+
+#endif  // UE_GAMEPLAY_CAMERAS_DEBUG
+
