Cameras: add debugging display for IK aiming and target preservation

- Debug display for IK aiming iteration
- Debugging panel for enabling that rendering
- Add missing panel for final pose stats

diff --git a/Engine/Plugins/Cameras/GameplayCameras/Source/GameplayCameras/Private/Core/CameraIKAim.cpp b/Engine/Plugins/Cameras/GameplayCameras/Source/GameplayCameras/Private/Core/CameraIKAim.cpp
--- a/Engine/Plugins/Cameras/GameplayCameras/Source/GameplayCameras/Private/Core/CameraIKAim.cpp
+++ b/Engine/Plugins/Cameras/GameplayCameras/Source/GameplayCameras/Private/Core/CameraIKAim.cpp
@@ -12,8 +12,11 @@
 #include "Core/CameraSystemEvaluator.h"
 #include "Core/CameraVariableTable.h"
 #include "Core/RootCameraNode.h"
+#include "Debug/CameraDebugRenderer.h"
+#include "Engine/Engine.h"
 #include "GameplayCameras.h"
 #include "GameplayCamerasSettings.h"
+#include "Math/Ray.h"
 #include "Serialization/MemoryReader.h"
 #include "Serialization/MemoryWriter.h"
 
@@ -97,15 +100,32 @@ bool FCameraIKAim::DoRun(const FCameraIKAimParams& Params, const FCameraRigEvalu
 		++IterationInfo.IterationIndex;
 	}
 
+#if UE_GAMEPLAY_CAMERAS_DEBUG
+	LastRunDebugInfo.DesiredTarget = Params.TargetLocation;
+	LastRunDebugInfo.bSucceeded = (IterationInfo.Result == EAimResult::Completed);
+#endif  // UE_GAMEPLAY_CAMERAS_DEBUG
+
 	return IterationInfo.Result == EAimResult::Completed || IterationInfo.Result == EAimResult::Corrected;
 }
 
 void FCameraIKAim::DoRunIteration(const FCameraIKAimParams& Params, const FCameraRigEvaluationInfo& CameraRigInfo, FAimIterationInfo& IterationInfo)
 {
+#if UE_GAMEPLAY_CAMERAS_DEBUG
+	FCameraIKAimIterationDebugInfo& IterationDebugInfo = LastRunDebugInfo.Iterations.Emplace_GetRef();
+#endif  // UE_GAMEPLAY_CAMERAS_DEBUG
+
 	// Run the system. We restore its state after each run.
 	RunRootCameraNode(Params, CameraRigInfo);
 
+	// Check how far we are from the desired target.
 	const bool bContinueIteration = CheckTolerance(Params, CameraRigInfo, IterationInfo);
+#if UE_GAMEPLAY_CAMERAS_DEBUG
+	IterationDebugInfo.CameraPoseLocation = ScratchResult.CameraPose.GetLocation();
+	IterationDebugInfo.CameraPoseRotation = ScratchResult.CameraPose.GetRotation();
+	IterationDebugInfo.ErrorAngle = IterationInfo.ErrorAngle;
+	IterationDebugInfo.ErrorDistance = IterationInfo.ErrorDistance;
+	IterationDebugInfo.bNeededSolver = bContinueIteration;
+#endif  // UE_GAMEPLAY_CAMERAS_DEBUG
 	if (!bContinueIteration)
 	{
 		return;
@@ -119,6 +139,9 @@ void FCameraIKAim::DoRunIteration(const FCameraIKAimParams& Params, const FCamer
 			{
 				return Item.VariableID == YawPitchDefinition.VariableID;
 			});
+#if UE_GAMEPLAY_CAMERAS_DEBUG
+	IterationDebugInfo.bFoundSolver = (FoundItem != nullptr);
+#endif  // UE_GAMEPLAY_CAMERAS_DEBUG
 	if (!FoundItem)
 	{
 		UE_LOG(LogCameraSystem, Error, TEXT("Can't figure out how to aim camera rig '%s'."), *CameraRigInfo.CameraRig->GetPathName());
@@ -220,14 +243,16 @@ bool FCameraIKAim::CheckTolerance(const FCameraIKAimParams& Params, const FCamer
 
 	const double ErrorAngle = FMath::RadiansToDegrees(FMath::Asin(SinAngle));
 	IterationInfo.ErrorAngle = ErrorAngle;
+
+	const double ErrorDistance = OrthLength / TargetDistance;
+	IterationInfo.ErrorDistance = ErrorDistance;
+
 	if (ErrorAngle <= Params.AngleTolerance)
 	{
 		IterationInfo.Result = EAimResult::Completed;
 		return false;
 	}
 
-	const double ErrorDistance = OrthLength / TargetDistance;
-	IterationInfo.ErrorDistance = ErrorDistance;
 	if (ErrorDistance <= Params.DistanceTolerance)
 	{
 		IterationInfo.Result = EAimResult::Completed;
@@ -242,6 +267,11 @@ void FCameraIKAim::AimTwoBonesCameraRig(const FCameraIKAimParams& Params, const
 	FRotator3d Correction;
 	const bool bGotCorrection = ComputeTwoBonesCorrection(
 			ScratchResult.CameraPose, PivotTransform.GetLocation(), Params.TargetLocation, Correction);
+#if UE_GAMEPLAY_CAMERAS_DEBUG
+	FCameraIKAimIterationDebugInfo& IterationDebugInfo = LastRunDebugInfo.Iterations.Last();
+	IterationDebugInfo.PivotJointLocation = PivotTransform.GetLocation();
+	IterationDebugInfo.bSolvingSuccess = bGotCorrection;
+#endif  // UE_GAMEPLAY_CAMERAS_DEBUG
 	if (!bGotCorrection)
 	{
 		IterationInfo.Result = EAimResult::Failed;
@@ -256,6 +286,10 @@ void FCameraIKAim::AimTwoBonesCameraRig(const FCameraIKAimParams& Params, const
 	Operation.Yaw = FConsumableDouble::Delta(Correction.Yaw);
 	Operation.Pitch = FConsumableDouble::Delta(Correction.Pitch);
 
+#if UE_GAMEPLAY_CAMERAS_DEBUG
+	IterationDebugInfo.YawPitchCorrection = FVector2d(Correction.Yaw, Correction.Pitch);
+#endif  // UE_GAMEPLAY_CAMERAS_DEBUG
+
 	CameraSystemHierarchy.ForEachEvaluator(TEXT("ActiveCameraRig"), ECameraNodeEvaluatorFlags::SupportsOperations,
 			[&OperationParams, &Operation](FCameraNodeEvaluator* Evaluator)
 			{
@@ -264,13 +298,18 @@ void FCameraIKAim::AimTwoBonesCameraRig(const FCameraIKAimParams& Params, const
 
 	if (Operation.Yaw.HasValue() || Operation.Pitch.HasValue())
 	{
-		UE_LOG(LogCameraSystem, Warning, TEXT("Aborting aiming of camera rig '%s': not all corrections were consumed by the camera nodes."),
+		UE_LOG(LogCameraSystem, Warning, 
+				TEXT("Aborting aiming of camera rig '%s': not all corrections were consumed by the camera nodes."),
 				*CameraRigInfo.CameraRig->GetPathName());
 		IterationInfo.Result = EAimResult::Aborted;
-		return;
 	}
-
-	IterationInfo.Result = EAimResult::Corrected;
+	else
+	{
+		IterationInfo.Result = EAimResult::Corrected;
+	}
+#if UE_GAMEPLAY_CAMERAS_DEBUG
+	IterationDebugInfo.bSolvingSuccess = (IterationInfo.Result == EAimResult::Corrected);
+#endif  // UE_GAMEPLAY_CAMERAS_DEBUG
 }
 
 bool FCameraIKAim::ComputeTwoBonesCorrection(const FCameraPose& CurrentPose, const FVector3d& PivotLocation, const FVector3d& DesiredTarget, FRotator3d& OutCorrection)
@@ -436,5 +475,113 @@ bool FCameraIKAim::RaySphereIntersectExit(const FVector3d& RayStart, const FVect
 	return false;
 }
 
+#if UE_GAMEPLAY_CAMERAS_DEBUG
+
+void FCameraIKAim::GetLastRunDebugInfo(FCameraIKAimDebugInfo& OutDebugInfo) const
+{
+	OutDebugInfo = LastRunDebugInfo;
+}
+
+void FCameraIKAimDebugInfo::DebugDraw(const FCameraDebugBlockDrawParams& Params, FCameraDebugRenderer& Renderer) const
+{
+	Renderer.AddText(TEXT("IK aiming "));
+	if (bSucceeded)
+	{
+		Renderer.AddText(TEXT("{cam_good}succeeded"));
+	}
+	else
+	{
+		Renderer.AddText(TEXT("{cam_error}failed"));
+	}
+	Renderer.AddText(TEXT("{cam_default} in %d iterations\n"), Iterations.Num());
+
+	UFont* TinyFont = GEngine->GetTinyFont();
+	Renderer.DrawSphere(DesiredTarget, 1.f, 8, FLinearColor::Yellow, 1.f);
+
+	Renderer.AddIndent();
+	{
+		int32 IterationIndex = 1;
+		for (const FCameraIKAimIterationDebugInfo& IterationDebugInfo : Iterations)
+		{
+			const bool bDirectionIsNormalized = true;
+			const FVector3d TargetDir = FVector3d::ForwardVector;
+			FRay3d DirectionRay(
+					IterationDebugInfo.CameraPoseLocation, 
+					IterationDebugInfo.CameraPoseRotation.RotateVector(TargetDir), bDirectionIsNormalized);
+
+			Renderer.DrawLine(
+					IterationDebugInfo.CameraPoseLocation,
+					DirectionRay.PointAt(1000.0),
+					FLinearColor::Yellow);
+			Renderer.DrawText(
+					IterationDebugInfo.CameraPoseLocation,
+					FString::Format(TEXT("Iteration {0}"), { IterationIndex }),
+					FLinearColor::Yellow,
+					TinyFont);
+
+			Renderer.AddText(
+					TEXT("%d : error angle %.2fdeg, error distance %.1fcm, "),
+					IterationIndex,
+					IterationDebugInfo.ErrorAngle, IterationDebugInfo.ErrorDistance);
+
+			++IterationIndex;
+
+			if (!IterationDebugInfo.bNeededSolver)
+			{
+				Renderer.AddText(TEXT(" {cam_good}reached tolerance{cam_default}\n"));
+				continue;
+			}
+			if (!IterationDebugInfo.bFoundSolver)
+			{
+				Renderer.AddText(TEXT(" {cam_error}couldn't find solver{cam_default}\n"));
+				continue;
+			}
+			
+			Renderer.AddText(TEXT(" pivot %s"), *IterationDebugInfo.PivotJointLocation.ToString());
+			Renderer.AddText(
+					TEXT(" correction Yaw=%.1f Pitch=%.1f"), 
+					IterationDebugInfo.YawPitchCorrection.X,
+					IterationDebugInfo.YawPitchCorrection.Y);
+
+			if (!IterationDebugInfo.bSolvingSuccess)
+			{
+				Renderer.AddText(TEXT(", {cam_error}couldn't compute correction{cam_default}\n"));
+				continue;
+			}
+
+			Renderer.NewLine();
+		}
+	}
+	Renderer.RemoveIndent();
+}
+
+FArchive& operator<< (FArchive& Ar, FCameraIKAimIterationDebugInfo& IterationDebugInfo)
+{
+	Ar << IterationDebugInfo.CameraPoseLocation;
+	Ar << IterationDebugInfo.CameraPoseRotation;
+	Ar << IterationDebugInfo.ErrorAngle;
+	Ar << IterationDebugInfo.ErrorDistance;
+
+	Ar << IterationDebugInfo.PivotJointLocation;
+	Ar << IterationDebugInfo.YawPitchCorrection;
+
+	Ar << IterationDebugInfo.bNeededSolver;
+	Ar << IterationDebugInfo.bFoundSolver;
+	Ar << IterationDebugInfo.bSolvingSuccess;
+
+	return Ar;
+}
+
+FArchive& operator<< (FArchive& Ar, FCameraIKAimDebugInfo& DebugInfo)
+{
+	Ar << DebugInfo.Iterations;
+	Ar << DebugInfo.DesiredTarget;
+	Ar << DebugInfo.bSucceeded;
+
+	return Ar;
+}
+
+#endif  // UE_GAMEPLAY_CAMERAS_DEBUG
+
 }  // namespace UE::Cameras
 
diff --git a/Engine/Plugins/Cameras/GameplayCameras/Source/GameplayCameras/Private/Debug/RootCameraDebugBlock.cpp b/Engine/Plugins/Cameras/GameplayCameras/Source/GameplayCameras/Private/Debug/RootCameraDebugBlock.cpp
--- a/Engine/Plugins/Cameras/GameplayCameras/Source/GameplayCameras/Private/Debug/RootCameraDebugBlock.cpp
+++ b/Engine/Plugins/Cameras/GameplayCameras/Source/GameplayCameras/Private/Debug/RootCameraDebugBlock.cpp
@@ -45,6 +45,12 @@ static FAutoConsoleVariableRef CVarGameplayCamerasDebugPoseStatsShowUnchanged(
 	GGameplayCamerasDebugPoseStatsShowUnchanged,
 	TEXT(""));
 
+bool GGameplayCamerasDebugPoseStatsShowVariableIDs = false;
+static FAutoConsoleVariableRef CVarGameplayCamerasDebugPoseStatsShowVariableIDs(
+	TEXT("GameplayCameras.Debug.PoseStats.ShowVariableIDs"),
+	GGameplayCamerasDebugPoseStatsShowVariableIDs,
+	TEXT(""));
+
 bool GGameplayCamerasDebugDrawBackground = true;
 static FAutoConsoleVariableRef CVarGameplayCamerasDebugDrawBackground(
 	TEXT("GameplayCameras.Debug.DrawBackground"),
@@ -125,7 +131,7 @@ void FRootCameraDebugBlock::BuildDebugBlocks(const FCameraSystemEvaluator& Camer
 		{
 			ResultDebugBlock.Initialize(CameraSystem.GetEvaluatedResult(), Builder);
 			ResultDebugBlock.GetCameraPoseDebugBlock()->WithShowUnchangedCVar(TEXT("GameplayCameras.Debug.PoseStats.ShowUnchanged"));
-			ResultDebugBlock.GetVariableTableDebugBlock()->WithShowVariableIDsCVar(TEXT("GameplayCameras.Debug.PoseStats.ShowUnchanged"));
+			ResultDebugBlock.GetVariableTableDebugBlock()->WithShowVariableIDsCVar(TEXT("GameplayCameras.Debug.PoseStats.ShowVariableIDs"));
 		}
 		FPlayerControllersDebugBlock& PlayerControllersDebugBlock = Builder.BuildDebugBlock<FPlayerControllersDebugBlock>();
 		PoseStatsCategory.AddChild(&PlayerControllersDebugBlock);
diff --git a/Engine/Plugins/Cameras/GameplayCameras/Source/GameplayCameras/Private/Services/OrientationInitializationService.cpp b/Engine/Plugins/Cameras/GameplayCameras/Source/GameplayCameras/Private/Services/OrientationInitializationService.cpp
--- a/Engine/Plugins/Cameras/GameplayCameras/Source/GameplayCameras/Private/Services/OrientationInitializationService.cpp
+++ b/Engine/Plugins/Cameras/GameplayCameras/Source/GameplayCameras/Private/Services/OrientationInitializationService.cpp
@@ -4,16 +4,25 @@
 
 #include "Core/CameraEvaluationContext.h"
 #include "Core/CameraEvaluationContextStack.h"
-#include "Core/CameraIKAim.h"
 #include "Core/CameraOperation.h"
 #include "Core/CameraRigAsset.h"
 #include "Core/CameraSystemEvaluator.h"
 #include "Core/RootCameraNode.h"
 #include "Core/RootCameraNodeCameraRigEvent.h"
+#include "Debug/CameraDebugBlock.h"
+#include "Debug/CameraDebugBlockBuilder.h"
+#include "Debug/CameraDebugRenderer.h"
+#include "HAL/IConsoleManager.h"
 
 namespace UE::Cameras
 {
 
+bool GGameplayCamerasDebugOrientationInitializationShowLastTargetPreservation = false;
+static FAutoConsoleVariableRef CVarGameplayCamerasDebugPoseStatsShowUnchanged(
+	TEXT("GameplayCameras.Debug.OrientationInitialization.ShowLastTargetPreservation"),
+	GGameplayCamerasDebugOrientationInitializationShowLastTargetPreservation,
+	TEXT(""));
+
 UE_DEFINE_CAMERA_EVALUATION_SERVICE(FOrientationInitializationService)
 
 void FOrientationInitializationService::OnInitialize(const FCameraEvaluationServiceInitializeParams& Params)
@@ -173,6 +182,8 @@ void FOrientationInitializationService::TryPreserveTarget(const FCameraRigEvalua
 
 	FVector3d TargetToPreserve = LastResult.CameraPose.GetTarget();
 
+	// Relative target preservation means that if the context turned since last frame, we want to preserve
+	// a target that has "turned" along with it.
 	if (bUseRelativeTarget && bHasPreviousContextTransform &&
 			(PreviousEvaluationContext == nullptr || CameraRigInfo.EvaluationContext == PreviousEvaluationContext))
 	{
@@ -193,7 +204,65 @@ void FOrientationInitializationService::TryPreserveTarget(const FCameraRigEvalua
 
 	FCameraIKAim CameraAim;
 	CameraAim.Run(AimParams, CameraRigInfo);
+
+#if UE_GAMEPLAY_CAMERAS_DEBUG
+	DebugLastEvaluatedTarget = LastResult.CameraPose.GetTarget();
+	CameraAim.GetLastRunDebugInfo(LastAimDebugInfo);
+#endif  // UE_GAMEPLAY_CAMERAS_DEBUG
+}
+
+#if UE_GAMEPLAY_CAMERAS_DEBUG
+
+class FOrientationInitializationDebugBlock : public FCameraDebugBlock
+{
+	UE_DECLARE_CAMERA_DEBUG_BLOCK(GAMEPLAYCAMERAS_API, FOrientationInitializationDebugBlock)
+
+public:
+
+	FOrientationInitializationDebugBlock() = default;
+	FOrientationInitializationDebugBlock(const FOrientationInitializationService& InService);
+
+protected:
+
+	// FCameraDebugBlock interface.
+	virtual void OnDebugDraw(const FCameraDebugBlockDrawParams& Params, FCameraDebugRenderer& Renderer) override;
+	virtual void OnSerialize(FArchive& Ar) override;
+
+private:
+
+	FVector3d LastEvaluatedTarget;
+	FCameraIKAimDebugInfo AimDebugInfo;
+};
+
+UE_DEFINE_CAMERA_DEBUG_BLOCK(FOrientationInitializationDebugBlock)
+
+void FOrientationInitializationService::OnBuildDebugBlocks(const FCameraDebugBlockBuildParams& Params, FCameraDebugBlockBuilder& Builder)
+{
+	Builder.AttachDebugBlock<FOrientationInitializationDebugBlock>(*this);
+}
+
+FOrientationInitializationDebugBlock::FOrientationInitializationDebugBlock(const FOrientationInitializationService& InService)
+{
+	LastEvaluatedTarget = InService.DebugLastEvaluatedTarget;
+	AimDebugInfo = InService.LastAimDebugInfo;
+}
+
+void FOrientationInitializationDebugBlock::OnDebugDraw(const FCameraDebugBlockDrawParams& Params, FCameraDebugRenderer& Renderer)
+{
+	if (GGameplayCamerasDebugOrientationInitializationShowLastTargetPreservation)
+	{
+		AimDebugInfo.DebugDraw(Params, Renderer);
+
+		Renderer.DrawSphere(LastEvaluatedTarget, 1.f, 8, FLinearColor::Green, 1.f);
+	}
 }
 
+void FOrientationInitializationDebugBlock::OnSerialize(FArchive& Ar)
+{
+	Ar << AimDebugInfo;
+}
+
+#endif  // UE_GAMEPLAY_CAMERAS_DEBUG
+
 }  // namespace UE::Cameras
 
diff --git a/Engine/Plugins/Cameras/GameplayCameras/Source/GameplayCameras/Private/Services/OrientationInitializationService.h b/Engine/Plugins/Cameras/GameplayCameras/Source/GameplayCameras/Private/Services/OrientationInitializationService.h
--- a/Engine/Plugins/Cameras/GameplayCameras/Source/GameplayCameras/Private/Services/OrientationInitializationService.h
+++ b/Engine/Plugins/Cameras/GameplayCameras/Source/GameplayCameras/Private/Services/OrientationInitializationService.h
@@ -3,6 +3,7 @@
 #pragma once
 
 #include "Core/CameraEvaluationService.h"
+#include "Core/CameraIKAim.h"
 #include "Core/CameraRigTransition.h"
 
 namespace UE::Cameras
@@ -22,6 +23,10 @@ protected:
 	virtual void OnPostUpdate(const FCameraEvaluationServiceUpdateParams& Params, FCameraEvaluationServiceUpdateResult& OutResult) override;
 	virtual void OnRootCameraNodeEvent(const FRootCameraNodeCameraRigEvent& InEvent) override;
 
+#if UE_GAMEPLAY_CAMERAS_DEBUG
+	virtual void OnBuildDebugBlocks(const FCameraDebugBlockBuildParams& Params, FCameraDebugBlockBuilder& Builder) override;
+#endif  // UE_GAMEPLAY_CAMERAS_DEBUG
+
 private:
 
 	void TryInitializeContextYawPitch(const FCameraRigEvaluationInfo& CameraRigInfo);
@@ -37,6 +42,13 @@ private:
 	FVector3d PreviousContextLocation;
 	FRotator3d PreviousContextRotation;
 	bool bHasPreviousContextTransform = false;
+
+#if UE_GAMEPLAY_CAMERAS_DEBUG
+	FVector3d DebugLastEvaluatedTarget;
+	FCameraIKAimDebugInfo LastAimDebugInfo;
+
+	friend class FOrientationInitializationDebugBlock;
+#endif  // UE_GAMEPLAY_CAMERAS_DEBUG
 };
 
 }  // namespace UE::Cameras
diff --git a/Engine/Plugins/Cameras/GameplayCameras/Source/GameplayCameras/Public/Core/CameraIKAim.h b/Engine/Plugins/Cameras/GameplayCameras/Source/GameplayCameras/Public/Core/CameraIKAim.h
--- a/Engine/Plugins/Cameras/GameplayCameras/Source/GameplayCameras/Public/Core/CameraIKAim.h
+++ b/Engine/Plugins/Cameras/GameplayCameras/Source/GameplayCameras/Public/Core/CameraIKAim.h
@@ -6,6 +6,7 @@
 #include "Core/CameraNodeEvaluatorHierarchy.h"
 #include "Core/CameraNodeEvaluatorStorage.h"
 #include "Math/MathFwd.h"
+#include "GameplayCameras.h"
 
 namespace UE::Cameras
 {
@@ -38,6 +39,38 @@ struct FCameraIKAimParams
 	uint8 MaxIterations = 0;
 };
 
+#if UE_GAMEPLAY_CAMERAS_DEBUG
+
+struct FCameraIKAimIterationDebugInfo
+{
+	FVector3d CameraPoseLocation;
+	FRotator3d CameraPoseRotation;
+	double ErrorAngle = 0;
+	double ErrorDistance = 0;
+ 
+	FVector3d PivotJointLocation = FVector3d::ZeroVector;
+	FVector2d YawPitchCorrection = FVector2d::ZeroVector;
+
+	bool bNeededSolver = false;
+	bool bFoundSolver = false;
+	bool bSolvingSuccess = false;
+};
+
+struct FCameraIKAimDebugInfo
+{
+	TArray<FCameraIKAimIterationDebugInfo, TInlineAllocator<4>> Iterations;
+
+	FVector3d DesiredTarget = FVector3d::ZeroVector;
+	bool bSucceeded = false;
+
+	void DebugDraw(const FCameraDebugBlockDrawParams& Params, FCameraDebugRenderer& Renderer) const;
+};
+
+FArchive& operator<< (FArchive& Ar, FCameraIKAimIterationDebugInfo& IterationDebugInfo);
+FArchive& operator<< (FArchive& Ar, FCameraIKAimDebugInfo& DebugInfo);
+
+#endif  // UE_GAMEPLAY_CAMERAS_DEBUG
+
 /**
  * A class that can manipulate a camera rig in order to aim it at a desired target.
  */
@@ -48,6 +81,10 @@ public:
 	/** Executes the aiming. */
 	bool Run(const FCameraIKAimParams& Params, const FCameraRigEvaluationInfo& CameraRigInfo);
 
+#if UE_GAMEPLAY_CAMERAS_DEBUG
+	void GetLastRunDebugInfo(FCameraIKAimDebugInfo& OutDebugInfo) const;
+#endif  // UE_GAMEPLAY_CAMERAS_DEBUG
+
 private:
 
 	enum class EAimResult
@@ -85,6 +122,10 @@ private:
 	
 	FCameraNodeEvaluatorHierarchy CameraSystemHierarchy;
 	TArray<uint8> EvaluatorSnapshot;
+
+#if UE_GAMEPLAY_CAMERAS_DEBUG
+	FCameraIKAimDebugInfo LastRunDebugInfo;
+#endif  // UE_GAMEPLAY_CAMERAS_DEBUG
 };
 
 }  // namespace UE::Cameras
diff --git a/Engine/Plugins/Cameras/GameplayCameras/Source/GameplayCamerasEditor/Private/Debugger/SEvaluationServicesDebugPanel.cpp b/Engine/Plugins/Cameras/GameplayCameras/Source/GameplayCamerasEditor/Private/Debugger/SEvaluationServicesDebugPanel.cpp
new file mode 100644
--- /dev/null
+++ b/Engine/Plugins/Cameras/GameplayCameras/Source/GameplayCamerasEditor/Private/Debugger/SEvaluationServicesDebugPanel.cpp
@@ -0,0 +1,33 @@
+// Copyright Epic Games, Inc. All Rights Reserved.
+
+#include "Debugger/SEvaluationServicesDebugPanel.h"
+
+#include "Debugger/SDebugWidgetUtils.h"
+#include "Widgets/Input/SCheckBox.h"
+#include "Widgets/Layout/SBox.h"
+#include "Widgets/SBoxPanel.h"
+
+#define LOCTEXT_NAMESPACE "SEvaluationServicesDebugPanel"
+
+namespace UE::Cameras
+{
+
+void SEvaluationServicesDebugPanel::Construct(const FArguments& InArgs)
+{
+	ChildSlot
+	[
+		SNew(SVerticalBox)
+		+ SVerticalBox::Slot()
+			.AutoHeight()
+			[
+				SDebugWidgetUtils::CreateConsoleVariableCheckBox(
+						LOCTEXT("ShowLastTargetPreservation", "Show last target preservation"),
+						TEXT("GameplayCameras.Debug.OrientationInitialization.ShowLastTargetPreservation"))
+			]
+	];
+}
+
+}  // namespace UE::Cameras
+
+#undef LOCTEXT_NAMESPACE
+
diff --git a/Engine/Plugins/Cameras/GameplayCameras/Source/GameplayCamerasEditor/Private/Debugger/SEvaluationServicesDebugPanel.h b/Engine/Plugins/Cameras/GameplayCameras/Source/GameplayCamerasEditor/Private/Debugger/SEvaluationServicesDebugPanel.h
new file mode 100644
--- /dev/null
+++ b/Engine/Plugins/Cameras/GameplayCameras/Source/GameplayCamerasEditor/Private/Debugger/SEvaluationServicesDebugPanel.h
@@ -0,0 +1,24 @@
+// Copyright Epic Games, Inc. All Rights Reserved.
+
+#pragma once
+
+#include "CoreTypes.h"
+#include "Widgets/SCompoundWidget.h"
+
+class FTabManager;
+
+namespace UE::Cameras
+{
+
+class SEvaluationServicesDebugPanel : public SCompoundWidget
+{
+public:
+
+	SLATE_BEGIN_ARGS(SEvaluationServicesDebugPanel) {}
+	SLATE_END_ARGS();
+
+	void Construct(const FArguments& InArgs);
+};
+
+}  // namespace UE::Cameras
+
diff --git a/Engine/Plugins/Cameras/GameplayCameras/Source/GameplayCamerasEditor/Private/GameplayCamerasEditorModule.cpp b/Engine/Plugins/Cameras/GameplayCameras/Source/GameplayCamerasEditor/Private/GameplayCamerasEditorModule.cpp
--- a/Engine/Plugins/Cameras/GameplayCameras/Source/GameplayCamerasEditor/Private/GameplayCamerasEditorModule.cpp
+++ b/Engine/Plugins/Cameras/GameplayCameras/Source/GameplayCamerasEditor/Private/GameplayCamerasEditorModule.cpp
@@ -21,6 +21,8 @@
 #include "Debug/CameraDebugCategories.h"
 #include "Debugger/SBlendStacksDebugPanel.h"
 #include "Debugger/SCameraNodeTreeDebugPanel.h"
+#include "Debugger/SEvaluationServicesDebugPanel.h"
+#include "Debugger/SCameraPoseStatsDebugPanel.h"
 #include "Debugger/SGameplayCamerasDebugger.h"
 #include "Directors/BlueprintCameraDirector.h"
 #include "EdGraph/EdGraph.h"
@@ -418,6 +420,14 @@ private:
 					{
 						return SNew(SBlendStacksDebugPanel);
 					}));
+		RegisterDebugCategoryPanel(FCameraDebugCategories::Services, FOnCreateDebugCategoryPanel::CreateLambda([](const FString&)
+					{
+						return SNew(SEvaluationServicesDebugPanel);
+					}));
+		RegisterDebugCategoryPanel(FCameraDebugCategories::PoseStats, FOnCreateDebugCategoryPanel::CreateLambda([](const FString&)
+					{
+						return SNew(SCameraPoseStatsDebugPanel);
+					}));
 	}
 
 	void UnregisterCoreDebugCategories()
