Camreas: fix static anaylisis build

diff --git a/Engine/Plugins/Cameras/GameplayCameras/Source/GameplayCameras/Private/Debug/CameraRigJointsDebugBlock.cpp b/Engine/Plugins/Cameras/GameplayCameras/Source/GameplayCameras/Private/Debug/CameraRigJointsDebugBlock.cpp
--- a/Engine/Plugins/Cameras/GameplayCameras/Source/GameplayCameras/Private/Debug/CameraRigJointsDebugBlock.cpp
+++ b/Engine/Plugins/Cameras/GameplayCameras/Source/GameplayCameras/Private/Debug/CameraRigJointsDebugBlock.cpp
@@ -56,7 +56,7 @@ void FCameraRigJointsDebugBlock::OnDebugDraw(const FCameraDebugBlockDrawParams&
 					*Entry.Location.ToString(),
 					*Entry.Rotation.ToString(),
 					*Entry.VariableName,
-					Entry.VariableID);
+					Entry.VariableID.GetValue());
 		}
 	}
 	Renderer.RemoveIndent();
