Cameras: more static analysis build fixes

diff --git a/Engine/Plugins/Cameras/GameplayCameras/Source/GameplayCameras/Private/Services/OrientationInitializationService.cpp b/Engine/Plugins/Cameras/GameplayCameras/Source/GameplayCameras/Private/Services/OrientationInitializationService.cpp
--- a/Engine/Plugins/Cameras/GameplayCameras/Source/GameplayCameras/Private/Services/OrientationInitializationService.cpp
+++ b/Engine/Plugins/Cameras/GameplayCameras/Source/GameplayCameras/Private/Services/OrientationInitializationService.cpp
@@ -18,7 +18,7 @@ namespace UE::Cameras
 {
 
 bool GGameplayCamerasDebugOrientationInitializationShowLastTargetPreservation = false;
-static FAutoConsoleVariableRef CVarGameplayCamerasDebugPoseStatsShowUnchanged(
+static FAutoConsoleVariableRef CVarGameplayCamerasDebugOrientationInitializationShowLastTargetPreservation(
 	TEXT("GameplayCameras.Debug.OrientationInitialization.ShowLastTargetPreservation"),
 	GGameplayCamerasDebugOrientationInitializationShowLastTargetPreservation,
 	TEXT(""));
