Cameras: make the non-main stacks evaluate additively

The main stack evaluates camera rigs "in isolation", i.e. they all run with a mix of the base stack's output and their own evaluation result's output. Each camera rig's result is then blended on top of each other. But it turns out we don't want that for the other stacks.

For instance, let's consider a global stack running 2 camera rigs, one with a Collision Push node and the other with a Target RayCast node. This would result in the Target Distance determined by the second rig (via ray-cast) to be blended with the first rig's default 1000cm Target Distance. And the camera position would be blended between the "pushed" position (via collision) of the first rig, and the "untouched" position of the second rig. We effectively end up with neither the raycast nor the collision actually working! This didn't show up in previous testing because all global nodes like this were in the same camera rig, instead of split between multiple rigs.

To fix this, we make the non-main stacks operate additively: camera rigs inside those stacks aren't run in isolation anymore, they are run on top of each other. The second camera rig's Target RayCast node receives the output of the first rig's Collision Push. It gets the "last word" on the Target Distance, and doesn't touch the (possibly "pushed") camera pose position.

diff --git a/Engine/Plugins/Cameras/GameplayCameras/Source/GameplayCameras/Private/Core/BlendStackCameraNode.cpp b/Engine/Plugins/Cameras/GameplayCameras/Source/GameplayCameras/Private/Core/BlendStackCameraNode.cpp
--- a/Engine/Plugins/Cameras/GameplayCameras/Source/GameplayCameras/Private/Core/BlendStackCameraNode.cpp
+++ b/Engine/Plugins/Cameras/GameplayCameras/Source/GameplayCameras/Private/Core/BlendStackCameraNode.cpp
@@ -30,9 +30,9 @@
 
 	switch (BlendStackType)
 	{
-		case ECameraBlendStackType::Persistent:
+		case ECameraBlendStackType::AdditivePersistent:
 			return Builder.BuildEvaluator<FPersistentBlendStackCameraNodeEvaluator>();
-		case ECameraBlendStackType::Transient:
+		case ECameraBlendStackType::IsolatedTransient:
 			return Builder.BuildEvaluator<FTransientBlendStackCameraNodeEvaluator>();
 		default:
 			ensure(false);
@@ -155,29 +155,18 @@
 	OwningEvaluator = Params.Evaluator;
 }
 
-void FBlendStackCameraNodeEvaluator::OnRun(const FCameraNodeEvaluationParams& Params, FCameraNodeEvaluationResult& OutResult)
+void FBlendStackCameraNodeEvaluator::ResolveEntries(TArray<FResolvedEntry>& OutResolvedEntries)
 {
-	const UBlendStackCameraNode* BlendStackNode = GetCameraNodeAs<UBlendStackCameraNode>();
-
 	// Build up these structures so we don't re-resolve evaluation context weak-pointers
 	// multiple times in this function..
-	struct FResolvedEntry
-	{
-		FCameraRigEntry& Entry;
-		TSharedPtr<const FCameraEvaluationContext> Context;
-		int32 EntryIndex;
-	};
-
-	TArray<FResolvedEntry> ResolvedEntries;
-
-	// While we make these resolved entries, emit warnings and errors as needed.
 	for (int32 Index = 0; Index < Entries.Num(); ++Index)
 	{
 		FCameraRigEntry& Entry(Entries[Index]);
 		TSharedPtr<const FCameraEvaluationContext> CurContext = Entry.EvaluationContext.Pin();
 
-		ResolvedEntries.Add({ Entry, CurContext, Index });
+		OutResolvedEntries.Add({ Entry, CurContext, Index });
 
+		// While we make these resolved entries, emit warnings and errors as needed.
 		if (!Entry.bIsFrozen)
 		{
 			// Check that we still have a valid context. If not, let's freeze the entry, since
@@ -231,177 +220,10 @@
 		Entry.bLogWarnings = true;
 #endif  // UE_GAMEPLAY_CAMERAS_TRACE
 	}
+}
 
-	// Gather parameters to pre-blend, and evaluate blend nodes.
-	for (FResolvedEntry& ResolvedEntry : ResolvedEntries)
-	{
-		FCameraRigEntry& Entry(ResolvedEntry.Entry);
-
-		if (UNLIKELY(Entry.bIsFrozen))
-		{
-			continue;
-		}
-
-		FCameraNodeEvaluationParams CurParams(Params);
-		CurParams.EvaluationContext = ResolvedEntry.Context;
-		CurParams.bIsFirstFrame = Entry.bIsFirstFrame;
-
-		FCameraNodeEvaluationResult& CurResult(Entry.Result);
-
-		// Start with the input given to us.
-		CurResult.VariableTable.OverrideAll(OutResult.VariableTable);
-
-		// Override it with whatever the evaluation context has set on its result.
-		const FCameraNodeEvaluationResult& ContextResult(ResolvedEntry.Context->GetInitialResult());
-		CurResult.VariableTable.OverrideAll(ContextResult.VariableTable);
-
-		// Gather input parameters.
-		if (!Entry.bInputRunThisFrame)
-		{
-			FCameraBlendedParameterUpdateParams InputParams(CurParams, CurResult.CameraPose);
-			FCameraBlendedParameterUpdateResult InputResult(CurResult.VariableTable);
-
-			Entry.EvaluatorHierarchy.ForEachEvaluator(ECameraNodeEvaluatorFlags::NeedsParameterUpdate,
-					[&InputParams, &InputResult](FCameraNodeEvaluator* ParameterEvaluator)
-					{
-						ParameterEvaluator->UpdateParameters(InputParams, InputResult);
-					});
-
-			Entry.bInputRunThisFrame = true;
-		}
-
-		// Run blends.
-		// Note that we pass last frame's camera pose to the Run() method. This may change.
-		// Blends aren't expected to use the camera pose to do any logic until BlendResults().
-		if (!Entry.bBlendRunThisFrame)
-		{
-			FBlendCameraNodeEvaluator* EntryBlendEvaluator = Entry.RootEvaluator->GetBlendEvaluator();
-			if (EntryBlendEvaluator)
-			{
-				EntryBlendEvaluator->Run(CurParams, CurResult);
-			}
-
-			Entry.bBlendRunThisFrame = true;
-		}
-	}
-
-	// Blend input variables.
-	for (FResolvedEntry& ResolvedEntry : ResolvedEntries)
-	{
-		FCameraRigEntry& Entry(ResolvedEntry.Entry);
-		FCameraNodeEvaluationResult& CurResult(Entry.Result);
-
-		if (!Entry.bIsFrozen)
-		{
-			FCameraNodeEvaluationParams CurParams(Params);
-			CurParams.EvaluationContext = ResolvedEntry.Context;
-			CurParams.bIsFirstFrame = Entry.bIsFirstFrame;
-			FCameraNodePreBlendParams PreBlendParams(CurParams, CurResult.CameraPose, CurResult.VariableTable);
-
-			FCameraNodePreBlendResult PreBlendResult(OutResult.VariableTable);
-
-			FBlendCameraNodeEvaluator* EntryBlendEvaluator = Entry.RootEvaluator->GetBlendEvaluator();
-			if (EntryBlendEvaluator)
-			{
-				EntryBlendEvaluator->BlendParameters(PreBlendParams, PreBlendResult);
-			}
-			else
-			{
-				OutResult.VariableTable.Override(CurResult.VariableTable, ECameraVariableTableFilter::Input);
-			}
-		}
-		else
-		{
-			// Frozen entries still contribute to the blend using their last evaluated values.
-			OutResult.VariableTable.Override(CurResult.VariableTable, ECameraVariableTableFilter::Input);
-		}
-	}
-
-	// Run the root nodes. They will use the pre-blended inputs from the last step.
-	// Frozen entries are skipped, since they only ever use the last result they produced.
-	for (FResolvedEntry& ResolvedEntry : ResolvedEntries)
-	{
-		FCameraRigEntry& Entry(ResolvedEntry.Entry);
-
-		if (UNLIKELY(Entry.bIsFrozen))
-		{
-			continue;
-		}
-
-		FCameraNodeEvaluationParams CurParams(Params);
-		CurParams.EvaluationContext = ResolvedEntry.Context;
-		CurParams.bIsFirstFrame = Entry.bIsFirstFrame;
-
-		FCameraNodeEvaluationResult& CurResult(Entry.Result);
-
-		// Start with the input given to us.
-		CurResult.CameraPose = OutResult.CameraPose;
-		CurResult.CameraRigJoints.OverrideAll(OutResult.CameraRigJoints);
-		CurResult.PostProcessSettings.OverrideAll(OutResult.PostProcessSettings);
-
-		// Override it with whatever the evaluation context has set on its result.
-		const FCameraNodeEvaluationResult& ContextResult(ResolvedEntry.Context->GetInitialResult());
-		CurResult.CameraPose.OverrideChanged(ContextResult.CameraPose);
-		CurResult.bIsCameraCut = OutResult.bIsCameraCut || ContextResult.bIsCameraCut;
-		CurResult.bIsValid = true;
-
-		// Run the camera rig's root node.
-		FCameraNodeEvaluator* RootEvaluator = Entry.RootEvaluator->GetRootEvaluator();
-		if (RootEvaluator)
-		{
-			RootEvaluator->Run(CurParams, CurResult);
-		}
-	}
-
-	// Now blend all the results, keeping track of blends that have reached 100% so
-	// that we can remove any camera rigs below (since they would have been completely
-	// blended out by that).
-	int32 PopEntriesBelow = INDEX_NONE;
-	for (FResolvedEntry& ResolvedEntry : ResolvedEntries)
-	{
-		FCameraRigEntry& Entry(ResolvedEntry.Entry);
-		FCameraNodeEvaluationResult& CurResult(Entry.Result);
-
-		if (!Entry.bIsFrozen)
-		{
-			FCameraNodeEvaluationParams CurParams(Params);
-			CurParams.EvaluationContext = ResolvedEntry.Context;
-			CurParams.bIsFirstFrame = Entry.bIsFirstFrame;
-			FCameraNodeBlendParams BlendParams(CurParams, CurResult);
-
-			FCameraNodeBlendResult BlendResult(OutResult);
-
-			FBlendCameraNodeEvaluator* EntryBlendEvaluator = Entry.RootEvaluator->GetBlendEvaluator();
-			if (EntryBlendEvaluator)
-			{
-				EntryBlendEvaluator->BlendResults(BlendParams, BlendResult);
-
-				if (BlendResult.bIsBlendFull && BlendResult.bIsBlendFinished)
-				{
-					PopEntriesBelow = ResolvedEntry.EntryIndex;
-				}
-			}
-			else
-			{
-				OutResult.OverrideAll(CurResult);
-
-				PopEntriesBelow = ResolvedEntry.EntryIndex;
-			}
-		}
-		else
-		{
-			OutResult.OverrideAll(CurResult);
-
-			PopEntriesBelow = ResolvedEntry.EntryIndex;
-		}
-	}
-
-	// Pop out camera rigs that have been blended out.
-	if (BlendStackNode->BlendStackType == ECameraBlendStackType::Transient && PopEntriesBelow != INDEX_NONE)
-	{
-		PopEntries(PopEntriesBelow);
-	}
-
+void FBlendStackCameraNodeEvaluator::OnRunFinished()
+{
 	// Reset transient flags.
 	for (FCameraRigEntry& Entry : Entries)
 	{
@@ -753,6 +575,210 @@
 	}
 }
 
+void FTransientBlendStackCameraNodeEvaluator::OnRun(const FCameraNodeEvaluationParams& Params, FCameraNodeEvaluationResult& OutResult)
+{
+	// Validate our entries and resolve evaluation context weak pointers.
+	TArray<FResolvedEntry> ResolvedEntries;
+	ResolveEntries(ResolvedEntries);
+
+	// Gather parameters to pre-blend, and evaluate blend nodes.
+	InternalPreBlendPrepare(ResolvedEntries, Params, OutResult);
+
+	// Blend input variables.
+	InternalPreBlendExecute(ResolvedEntries, Params, OutResult);
+
+	// Run the root nodes. They will use the pre-blended inputs from the last step.
+	// Frozen entries are skipped, since they only ever use the last result they produced.
+	InternalUpdate(ResolvedEntries, Params, OutResult);
+
+	// Now blend all the results, keeping track of blends that have reached 100% so
+	// that we can remove any camera rigs below (since they would have been completely
+	// blended out by that).
+	InternalPostBlendExecute(ResolvedEntries, Params, OutResult);
+
+	// Tidy up.
+	OnRunFinished();
+}
+
+void FTransientBlendStackCameraNodeEvaluator::InternalPreBlendPrepare(TArrayView<FResolvedEntry> ResolvedEntries, const FCameraNodeEvaluationParams& Params, FCameraNodeEvaluationResult& OutResult)
+{
+	for (FResolvedEntry& ResolvedEntry : ResolvedEntries)
+	{
+		FCameraRigEntry& Entry(ResolvedEntry.Entry);
+
+		if (UNLIKELY(Entry.bIsFrozen))
+		{
+			continue;
+		}
+
+		FCameraNodeEvaluationParams CurParams(Params);
+		CurParams.EvaluationContext = ResolvedEntry.Context;
+		CurParams.bIsFirstFrame = Entry.bIsFirstFrame;
+
+		FCameraNodeEvaluationResult& CurResult(Entry.Result);
+
+		// Start with the input given to us.
+		CurResult.VariableTable.OverrideAll(OutResult.VariableTable);
+
+		// Override it with whatever the evaluation context has set on its result.
+		const FCameraNodeEvaluationResult& ContextResult(ResolvedEntry.Context->GetInitialResult());
+		CurResult.VariableTable.OverrideAll(ContextResult.VariableTable);
+
+		// Gather input parameters if needed (and remember if it was indeed needed).
+		if (!Entry.bInputRunThisFrame)
+		{
+			bool bHasPreBlendedParameters = false;
+			FCameraBlendedParameterUpdateParams InputParams(CurParams, CurResult.CameraPose);
+			FCameraBlendedParameterUpdateResult InputResult(CurResult.VariableTable);
+
+			Entry.EvaluatorHierarchy.ForEachEvaluator(ECameraNodeEvaluatorFlags::NeedsParameterUpdate,
+					[&bHasPreBlendedParameters, &InputParams, &InputResult](FCameraNodeEvaluator* ParameterEvaluator)
+					{
+						ParameterEvaluator->UpdateParameters(InputParams, InputResult);
+						bHasPreBlendedParameters = true;
+					});
+
+			ResolvedEntry.bHasPreBlendedParameters = bHasPreBlendedParameters;
+			Entry.bInputRunThisFrame = true;
+		}
+
+		// Run blends.
+		// Note that we pass last frame's camera pose to the Run() method. This may change.
+		// Blends aren't expected to use the camera pose to do any logic until BlendResults().
+		if (!Entry.bBlendRunThisFrame)
+		{
+			FBlendCameraNodeEvaluator* EntryBlendEvaluator = Entry.RootEvaluator->GetBlendEvaluator();
+			if (EntryBlendEvaluator)
+			{
+				EntryBlendEvaluator->Run(CurParams, CurResult);
+			}
+
+			Entry.bBlendRunThisFrame = true;
+		}
+	}
+}
+
+void FTransientBlendStackCameraNodeEvaluator::InternalPreBlendExecute(TArrayView<FResolvedEntry> ResolvedEntries, const FCameraNodeEvaluationParams& Params, FCameraNodeEvaluationResult& OutResult)
+{
+	for (FResolvedEntry& ResolvedEntry : ResolvedEntries)
+	{
+		FCameraRigEntry& Entry(ResolvedEntry.Entry);
+		FCameraNodeEvaluationResult& CurResult(Entry.Result);
+
+		if (!Entry.bIsFrozen)
+		{
+			FCameraNodeEvaluationParams CurParams(Params);
+			CurParams.EvaluationContext = ResolvedEntry.Context;
+			CurParams.bIsFirstFrame = Entry.bIsFirstFrame;
+			FCameraNodePreBlendParams PreBlendParams(CurParams, CurResult.CameraPose, CurResult.VariableTable);
+
+			FCameraNodePreBlendResult PreBlendResult(OutResult.VariableTable);
+
+			FBlendCameraNodeEvaluator* EntryBlendEvaluator = Entry.RootEvaluator->GetBlendEvaluator();
+			if (EntryBlendEvaluator)
+			{
+				EntryBlendEvaluator->BlendParameters(PreBlendParams, PreBlendResult);
+			}
+			else
+			{
+				OutResult.VariableTable.Override(CurResult.VariableTable, ECameraVariableTableFilter::Input);
+			}
+		}
+		else
+		{
+			// Frozen entries still contribute to the blend using their last evaluated values.
+			OutResult.VariableTable.Override(CurResult.VariableTable, ECameraVariableTableFilter::Input);
+		}
+	}
+}
+
+void FTransientBlendStackCameraNodeEvaluator::InternalUpdate(TArrayView<FResolvedEntry> ResolvedEntries, const FCameraNodeEvaluationParams& Params, FCameraNodeEvaluationResult& OutResult)
+{
+	for (FResolvedEntry& ResolvedEntry : ResolvedEntries)
+	{
+		FCameraRigEntry& Entry(ResolvedEntry.Entry);
+
+		if (UNLIKELY(Entry.bIsFrozen))
+		{
+			continue;
+		}
+
+		FCameraNodeEvaluationParams CurParams(Params);
+		CurParams.EvaluationContext = ResolvedEntry.Context;
+		CurParams.bIsFirstFrame = Entry.bIsFirstFrame;
+
+		FCameraNodeEvaluationResult& CurResult(Entry.Result);
+
+		// Start with the input given to us.
+		CurResult.CameraPose = OutResult.CameraPose;
+		CurResult.CameraRigJoints.OverrideAll(OutResult.CameraRigJoints);
+		CurResult.PostProcessSettings.OverrideAll(OutResult.PostProcessSettings);
+
+		// Override it with whatever the evaluation context has set on its result.
+		const FCameraNodeEvaluationResult& ContextResult(ResolvedEntry.Context->GetInitialResult());
+		CurResult.CameraPose.OverrideChanged(ContextResult.CameraPose);
+		CurResult.bIsCameraCut = OutResult.bIsCameraCut || ContextResult.bIsCameraCut;
+		CurResult.bIsValid = true;
+
+		// Run the camera rig's root node.
+		FCameraNodeEvaluator* RootEvaluator = Entry.RootEvaluator->GetRootEvaluator();
+		if (RootEvaluator)
+		{
+			RootEvaluator->Run(CurParams, CurResult);
+		}
+	}
+}
+
+void FTransientBlendStackCameraNodeEvaluator::InternalPostBlendExecute(TArrayView<FResolvedEntry> ResolvedEntries, const FCameraNodeEvaluationParams& Params, FCameraNodeEvaluationResult& OutResult)
+{
+	int32 PopEntriesBelow = INDEX_NONE;
+	for (FResolvedEntry& ResolvedEntry : ResolvedEntries)
+	{
+		FCameraRigEntry& Entry(ResolvedEntry.Entry);
+		FCameraNodeEvaluationResult& CurResult(Entry.Result);
+
+		if (!Entry.bIsFrozen)
+		{
+			FCameraNodeEvaluationParams CurParams(Params);
+			CurParams.EvaluationContext = ResolvedEntry.Context;
+			CurParams.bIsFirstFrame = Entry.bIsFirstFrame;
+			FCameraNodeBlendParams BlendParams(CurParams, CurResult);
+
+			FCameraNodeBlendResult BlendResult(OutResult);
+
+			FBlendCameraNodeEvaluator* EntryBlendEvaluator = Entry.RootEvaluator->GetBlendEvaluator();
+			if (EntryBlendEvaluator)
+			{
+				EntryBlendEvaluator->BlendResults(BlendParams, BlendResult);
+
+				if (BlendResult.bIsBlendFull && BlendResult.bIsBlendFinished)
+				{
+					PopEntriesBelow = ResolvedEntry.EntryIndex;
+				}
+			}
+			else
+			{
+				OutResult.OverrideAll(CurResult);
+
+				PopEntriesBelow = ResolvedEntry.EntryIndex;
+			}
+		}
+		else
+		{
+			OutResult.OverrideAll(CurResult);
+
+			PopEntriesBelow = ResolvedEntry.EntryIndex;
+		}
+	}
+
+	// Pop out camera rigs that have been blended out.
+	const UBlendStackCameraNode* BlendStackNode = GetCameraNodeAs<UBlendStackCameraNode>();
+	if (BlendStackNode->BlendStackType == ECameraBlendStackType::IsolatedTransient && PopEntriesBelow != INDEX_NONE)
+	{
+		PopEntries(PopEntriesBelow);
+	}
+}
+
 const UCameraRigTransition* FTransientBlendStackCameraNodeEvaluator::FindTransition(const FBlendStackCameraPushParams& Params) const
 {
 	const UBlendStackCameraNode* BlendStackNode = GetCameraNodeAs<UBlendStackCameraNode>();
@@ -825,13 +851,7 @@
 			}
 		}
 	}
-	else if (BlendStackNode->bBlendFirstCameraRig)
-	{
-		return FindTransition(
-				ToCameraRig->EnterTransitions,
-				nullptr, nullptr, false,
-				ToCameraRig, ToCameraAsset);
-	}
+	// else: make the first camera rig in the stack start at 100% blend immediately.
 
 	return nullptr;
 }
@@ -917,6 +937,112 @@
 	}
 }
 
+void FPersistentBlendStackCameraNodeEvaluator::OnRun(const FCameraNodeEvaluationParams& Params, FCameraNodeEvaluationResult& OutResult)
+{
+	// Validate our entries and resolve evaluation context weak pointers.
+	TArray<FResolvedEntry> ResolvedEntries;
+	ResolveEntries(ResolvedEntries);
+
+	// Run the stack!
+	InternalUpdate(ResolvedEntries, Params, OutResult);
+
+	// Tidy things up.
+	OnRunFinished();
+}
+
+void FPersistentBlendStackCameraNodeEvaluator::InternalUpdate(TArrayView<FResolvedEntry> ResolvedEntries, const FCameraNodeEvaluationParams& Params, FCameraNodeEvaluationResult& OutResult)
+{
+	for (FResolvedEntry& ResolvedEntry : ResolvedEntries)
+	{
+		FCameraRigEntry& Entry(ResolvedEntry.Entry);
+
+		if (!Entry.bIsFrozen)
+		{
+			FCameraNodeEvaluationParams CurParams(Params);
+			CurParams.EvaluationContext = ResolvedEntry.Context;
+			CurParams.bIsFirstFrame = Entry.bIsFirstFrame;
+
+			FCameraNodeEvaluationResult& CurResult(Entry.Result);
+
+			// Start with the input given to us.
+			{
+				CurResult.CameraPose = OutResult.CameraPose;
+				CurResult.VariableTable.OverrideAll(OutResult.VariableTable);
+				CurResult.CameraRigJoints.OverrideAll(OutResult.CameraRigJoints);
+				CurResult.PostProcessSettings.OverrideAll(OutResult.PostProcessSettings);
+
+				// Override it with whatever the evaluation context has set on its result.
+				const FCameraNodeEvaluationResult& ContextResult(ResolvedEntry.Context->GetInitialResult());
+				CurResult.CameraPose.OverrideChanged(ContextResult.CameraPose);
+				CurResult.VariableTable.OverrideAll(ContextResult.VariableTable);
+
+				// Setup flags.
+				CurResult.bIsCameraCut = OutResult.bIsCameraCut || ContextResult.bIsCameraCut;
+				CurResult.bIsValid = true;
+			}
+
+			// Update pre-blended parameters.
+			{
+				FCameraBlendedParameterUpdateParams InputParams(CurParams, CurResult.CameraPose);
+				FCameraBlendedParameterUpdateResult InputResult(CurResult.VariableTable);
+
+				Entry.EvaluatorHierarchy.ForEachEvaluator(ECameraNodeEvaluatorFlags::NeedsParameterUpdate,
+						[&InputParams, &InputResult](FCameraNodeEvaluator* ParameterEvaluator)
+						{
+							ParameterEvaluator->UpdateParameters(InputParams, InputResult);
+						});
+			}
+
+			// Run the blend node.
+			FBlendCameraNodeEvaluator* EntryBlendEvaluator = Entry.RootEvaluator->GetBlendEvaluator();
+			if (EntryBlendEvaluator)
+			{
+				EntryBlendEvaluator->Run(CurParams, CurResult);
+			}
+
+			// Blend pre-blended parameters.
+			if (EntryBlendEvaluator)
+			{
+				FCameraNodePreBlendParams PreBlendParams(CurParams, CurResult.CameraPose, CurResult.VariableTable);
+				FCameraNodePreBlendResult PreBlendResult(OutResult.VariableTable);
+
+				EntryBlendEvaluator->BlendParameters(PreBlendParams, PreBlendResult);
+			}
+			else
+			{
+				OutResult.VariableTable.Override(CurResult.VariableTable, ECameraVariableTableFilter::Input);
+			}
+
+			// Run the camera rig's root node.
+			FCameraNodeEvaluator* RootEvaluator = Entry.RootEvaluator->GetRootEvaluator();
+			if (RootEvaluator)
+			{
+				RootEvaluator->Run(CurParams, CurResult);
+			}
+
+			// Blend the results.
+			if (EntryBlendEvaluator)
+			{
+				FCameraNodeBlendParams BlendParams(CurParams, CurResult);
+				FCameraNodeBlendResult BlendResult(OutResult);
+
+				EntryBlendEvaluator->BlendResults(BlendParams, BlendResult);
+			}
+			else
+			{
+				OutResult.OverrideAll(CurResult);
+			}
+		}
+		else
+		{
+			FCameraNodeEvaluationResult& CurResult(Entry.Result);
+
+			OutResult.VariableTable.Override(CurResult.VariableTable, ECameraVariableTableFilter::Input);
+			OutResult.OverrideAll(CurResult);
+		}
+	}
+}
+
 #if UE_GAMEPLAY_CAMERAS_DEBUG
 
 void FBlendStackCameraNodeEvaluator::OnBuildDebugBlocks(const FCameraDebugBlockBuildParams& Params, FCameraDebugBlockBuilder& Builder)
diff --git a/Engine/Plugins/Cameras/GameplayCameras/Source/GameplayCameras/Private/Core/DefaultRootCameraNode.cpp b/Engine/Plugins/Cameras/GameplayCameras/Source/GameplayCameras/Private/Core/DefaultRootCameraNode.cpp
--- a/Engine/Plugins/Cameras/GameplayCameras/Source/GameplayCameras/Private/Core/DefaultRootCameraNode.cpp
+++ b/Engine/Plugins/Cameras/GameplayCameras/Source/GameplayCameras/Private/Core/DefaultRootCameraNode.cpp
@@ -17,12 +17,11 @@ namespace UE::Cameras::Private
 
 TObjectPtr<UBlendStackCameraNode> CreateBlendStack(
 		UObject* This, const FObjectInitializer& ObjectInit,
-		const FName& Name, ECameraBlendStackType BlendStackType, bool bBlendFirstCameraRig)
+		const FName& Name, ECameraBlendStackType BlendStackType)
 {
 	TObjectPtr<UBlendStackCameraNode> NewBlendStack = ObjectInit.CreateDefaultSubobject<UBlendStackCameraNode>(
 			This, Name);
 	NewBlendStack->BlendStackType = BlendStackType;
-	NewBlendStack->bBlendFirstCameraRig = bBlendFirstCameraRig;
 	return NewBlendStack;
 }
 
@@ -33,10 +32,10 @@ UDefaultRootCameraNode::UDefaultRootCameraNode(const FObjectInitializer& ObjectI
 {
 	using namespace UE::Cameras::Private;
 
-	BaseLayer = CreateBlendStack(this, ObjectInit, TEXT("BaseLayer"), ECameraBlendStackType::Persistent, true);
-	MainLayer = CreateBlendStack(this, ObjectInit, TEXT("MainLayer"), ECameraBlendStackType::Transient, false);
-	GlobalLayer = CreateBlendStack(this, ObjectInit, TEXT("GlobalLayer"), ECameraBlendStackType::Persistent, true);
-	VisualLayer = CreateBlendStack(this, ObjectInit, TEXT("VisualLayer"), ECameraBlendStackType::Persistent, true);
+	BaseLayer = CreateBlendStack(this, ObjectInit, TEXT("BaseLayer"), ECameraBlendStackType::AdditivePersistent);
+	MainLayer = CreateBlendStack(this, ObjectInit, TEXT("MainLayer"), ECameraBlendStackType::IsolatedTransient);
+	GlobalLayer = CreateBlendStack(this, ObjectInit, TEXT("GlobalLayer"), ECameraBlendStackType::AdditivePersistent);
+	VisualLayer = CreateBlendStack(this, ObjectInit, TEXT("VisualLayer"), ECameraBlendStackType::AdditivePersistent);
 }
 
 FCameraNodeEvaluatorPtr UDefaultRootCameraNode::OnBuildEvaluator(FCameraNodeEvaluatorBuilder& Builder) const
diff --git a/Engine/Plugins/Cameras/GameplayCameras/Source/GameplayCameras/Public/Core/BlendStackCameraNode.h b/Engine/Plugins/Cameras/GameplayCameras/Source/GameplayCameras/Public/Core/BlendStackCameraNode.h
--- a/Engine/Plugins/Cameras/GameplayCameras/Source/GameplayCameras/Public/Core/BlendStackCameraNode.h
+++ b/Engine/Plugins/Cameras/GameplayCameras/Source/GameplayCameras/Public/Core/BlendStackCameraNode.h
@@ -37,18 +37,24 @@ class IGameplayCamerasLiveEditManager;
 
 }  // namespace UE::Cameras
 
+/**
+ * Describes a type of blend stack.
+ */
 UENUM()
 enum class ECameraBlendStackType
 {
 	/**
-	 * Camera rigs in a transient blend stack get automatically popped out of the stack 
-	 * when another rig has reached 100% blend above them.
+	 * Camera rigs are evaluated in isolation before being blended together, and get 
+	 * automatically popped out of the stack when another rig has reached 100% blend above 
+	 * them.
 	 */
-	Transient,
+	IsolatedTransient,
 	/**
-	 * Camera rigs in a persistent blend stack stay in the stack until explicitly removed.
+	 * Camera rigs are evaluated in an additive way, i.e. the result of a lower camera rig
+	 * becomes the input of the next ones. Also, camera rigs stay in the stack until explicitly 
+	 * removed.
 	 */
-	Persistent
+	AdditivePersistent
 };
 
 /**
@@ -70,13 +76,7 @@ public:
 	 * The type of blend stack this should run as.
 	 */
 	UPROPERTY()
-	ECameraBlendStackType BlendStackType = ECameraBlendStackType::Transient;
-
-	/**
-	 * Whether to blend-in the first camera rig when the stack is previously empty.
-	 */
-	UPROPERTY()
-	bool bBlendFirstCameraRig = false;
+	ECameraBlendStackType BlendStackType = ECameraBlendStackType::IsolatedTransient;
 };
 
 namespace UE::Cameras
@@ -116,7 +116,6 @@ protected:
 	// FCameraNodeEvaluator interface
 	virtual FCameraNodeEvaluatorChildrenView OnGetChildren() override;
 	virtual void OnInitialize(const FCameraNodeEvaluatorInitializeParams& Params, FCameraNodeEvaluationResult& OutResult) override;
-	virtual void OnRun(const FCameraNodeEvaluationParams& Params, FCameraNodeEvaluationResult& OutResult) override;
 	virtual void OnAddReferencedObjects(FReferenceCollector& Collector) override;
 	virtual void OnSerialize(const FCameraNodeEvaluatorSerializeParams& Params, FArchive& Ar) override;
 
@@ -152,6 +151,19 @@ protected:
 	void RemoveListenedPackages(TSharedPtr<IGameplayCamerasLiveEditManager> LiveEditManager, FCameraRigEntry& Entry);
 #endif
 
+protected:
+
+	struct FResolvedEntry
+	{
+		FCameraRigEntry& Entry;
+		TSharedPtr<const FCameraEvaluationContext> Context;
+		int32 EntryIndex;
+		bool bHasPreBlendedParameters;
+	};
+
+	void ResolveEntries(TArray<FResolvedEntry>& OutResolvedEntries);
+	void OnRunFinished();
+
 protected:
 
 	struct FCameraRigEntry
@@ -251,8 +263,19 @@ public:
 	/** Freeze all camera rigs that belong to a given evaluation context. */
 	void FreezeAll(TSharedPtr<FCameraEvaluationContext> EvaluationContext);
 
+protected:
+
+	// FCameraNodeEvaluator interface.
+	virtual void OnRun(const FCameraNodeEvaluationParams& Params, FCameraNodeEvaluationResult& OutResult) override;
+
 private:
 
+	// Update methods.
+	void InternalPreBlendPrepare(TArrayView<FResolvedEntry> ResolvedEntries, const FCameraNodeEvaluationParams& Params, FCameraNodeEvaluationResult& OutResult);
+	void InternalPreBlendExecute(TArrayView<FResolvedEntry> ResolvedEntries, const FCameraNodeEvaluationParams& Params, FCameraNodeEvaluationResult& OutResult);
+	void InternalUpdate(TArrayView<FResolvedEntry> ResolvedEntries, const FCameraNodeEvaluationParams& Params, FCameraNodeEvaluationResult& OutResult);
+	void InternalPostBlendExecute(TArrayView<FResolvedEntry> ResolvedEntries, const FCameraNodeEvaluationParams& Params, FCameraNodeEvaluationResult& OutResult);
+
 	// Utility functions for finding an appropriate transition.
 	const UCameraRigTransition* FindTransition(const FBlendStackCameraPushParams& Params) const;
 	const UCameraRigTransition* FindTransition(
@@ -305,8 +328,16 @@ public:
 	/** Remove an existing camera rig from the blend stack. */
 	void Remove(const FBlendStackCameraRemoveParams& Params);
 
+protected:
+
+	// FCameraNodeEvaluator interface.
+	virtual void OnRun(const FCameraNodeEvaluationParams& Params, FCameraNodeEvaluationResult& OutResult) override;
+
 private:
 
+	// Update methods.
+	void InternalUpdate(TArrayView<FResolvedEntry> ResolvedEntries, const FCameraNodeEvaluationParams& Params, FCameraNodeEvaluationResult& OutResult);
+
 	// Utility functions for finding an appropriate transition.
 	const UCameraRigTransition* FindTransition(const FBlendStackCameraPushParams& Params) const;
 };
