Cameras: make sure the occlusion node doesn't leave objects with the transparent material on deactivation

Also make it use weak pointers in case objects get GC'ed while it's still active. This can happen for example when shutting down PIE.

diff --git a/Engine/Plugins/Cameras/GameplayCameras/Source/GameplayCameras/Private/Nodes/Collision/OcclusionMaterialCameraNode.cpp b/Engine/Plugins/Cameras/GameplayCameras/Source/GameplayCameras/Private/Nodes/Collision/OcclusionMaterialCameraNode.cpp
--- a/Engine/Plugins/Cameras/GameplayCameras/Source/GameplayCameras/Private/Nodes/Collision/OcclusionMaterialCameraNode.cpp
+++ b/Engine/Plugins/Cameras/GameplayCameras/Source/GameplayCameras/Private/Nodes/Collision/OcclusionMaterialCameraNode.cpp
@@ -31,6 +31,10 @@ class FOcclusionMaterialCameraNodeEvaluator : public FCameraNodeEvaluator
 {
 	UE_DECLARE_CAMERA_NODE_EVALUATOR(GAMEPLAYCAMERAS_API, FOcclusionMaterialCameraNodeEvaluator)
 
+public:
+
+	~FOcclusionMaterialCameraNodeEvaluator();
+
 protected:
 
 	virtual void OnInitialize(const FCameraNodeEvaluatorInitializeParams& Params, FCameraNodeEvaluationResult& OutResult) override;
@@ -44,18 +48,28 @@ private:
 	void ApplyOcclusionMaterial(TSet<UStaticMeshComponent*> MeshComponents);
 	void RemoveOcclusionMaterial(TSet<UStaticMeshComponent*> MeshComponents);
 
+	void ResolveWeakMeshComponents(TSet<TWeakObjectPtr<UStaticMeshComponent>> WeakMeshComponents, TSet<UStaticMeshComponent*>& OutMeshComponents);
+
 private:
 
 	TCameraParameterReader<float> OcclusionSphereRadiusReader;
 	TCameraParameterReader<FVector3d> OcclusionTargetOffsetReader;
 
 	FTraceHandle OcclusionTraceHandle;
-	TSet<UStaticMeshComponent*> CurrentlyOccludedMeshComponents;
-	TMap<UStaticMeshComponent*, FOcclusionMaterialOverrideInfo> AppliedMaterialOverrides;
+	TSet<TWeakObjectPtr<UStaticMeshComponent>> CurrentlyOccludedMeshComponents;
+	TMap<TWeakObjectPtr<UStaticMeshComponent>, FOcclusionMaterialOverrideInfo> AppliedMaterialOverrides;
 };
 
 UE_DEFINE_CAMERA_NODE_EVALUATOR(FOcclusionMaterialCameraNodeEvaluator)
 
+FOcclusionMaterialCameraNodeEvaluator::~FOcclusionMaterialCameraNodeEvaluator()
+{
+	// Make sure any occluded meshes are released when our camera rig is deactivated it.
+	TSet<UStaticMeshComponent*> MeshComponents;
+	ResolveWeakMeshComponents(CurrentlyOccludedMeshComponents, MeshComponents);
+	RemoveOcclusionMaterial(MeshComponents);
+}
+
 void FOcclusionMaterialCameraNodeEvaluator::OnInitialize(const FCameraNodeEvaluatorInitializeParams& Params, FCameraNodeEvaluationResult& OutResult)
 {
 	const UOcclusionMaterialCameraNode* OcclusionMaterialNode = GetCameraNodeAs<UOcclusionMaterialCameraNode>();
@@ -151,10 +165,17 @@ void FOcclusionMaterialCameraNodeEvaluator::HandleOcclusionTraceResult(UWorld* W
 		}
 	}
 
-	TSet<UStaticMeshComponent*> NewMeshComponents = MeshComponents.Difference(CurrentlyOccludedMeshComponents);
-	TSet<UStaticMeshComponent*> OldMeshComponents = CurrentlyOccludedMeshComponents.Difference(MeshComponents);
+	TSet<UStaticMeshComponent*> CurrentMeshComponents;
+	ResolveWeakMeshComponents(CurrentlyOccludedMeshComponents, CurrentMeshComponents);
 
-	CurrentlyOccludedMeshComponents = MeshComponents;
+	TSet<UStaticMeshComponent*> NewMeshComponents = MeshComponents.Difference(CurrentMeshComponents);
+	TSet<UStaticMeshComponent*> OldMeshComponents = CurrentMeshComponents.Difference(MeshComponents);
+
+	CurrentlyOccludedMeshComponents.Reset();
+	for (UStaticMeshComponent* MeshComponent : MeshComponents)
+	{
+		CurrentlyOccludedMeshComponents.Add(MeshComponent);
+	}
 
 	// Apply occlusion material changes to new/old components.
 	ApplyOcclusionMaterial(NewMeshComponents);
@@ -213,6 +234,17 @@ void FOcclusionMaterialCameraNodeEvaluator::RemoveOcclusionMaterial(TSet<UStatic
 	}
 }
 
+void FOcclusionMaterialCameraNodeEvaluator::ResolveWeakMeshComponents(TSet<TWeakObjectPtr<UStaticMeshComponent>> WeakMeshComponents, TSet<UStaticMeshComponent*>& OutMeshComponents)
+{
+	for (TWeakObjectPtr<UStaticMeshComponent> WeakMeshComponent : WeakMeshComponents)
+	{
+		if (UStaticMeshComponent* MeshComponent = WeakMeshComponent.Get())
+		{
+			OutMeshComponents.Add(MeshComponent);
+		}
+	}
+}
+
 }  // namespace UE::Cameras
 
 UOcclusionMaterialCameraNode::UOcclusionMaterialCameraNode(const FObjectInitializer& ObjectInit)
