Cameras: extract position/offset math from collision node into a math utility class

diff --git a/Engine/Plugins/Cameras/GameplayCameras/Source/GameplayCameras/Private/Math/CameraNodeSpaceMath.cpp b/Engine/Plugins/Cameras/GameplayCameras/Source/GameplayCameras/Private/Math/CameraNodeSpaceMath.cpp
new file mode 100644
--- /dev/null
+++ b/Engine/Plugins/Cameras/GameplayCameras/Source/GameplayCameras/Private/Math/CameraNodeSpaceMath.cpp
@@ -0,0 +1,290 @@
+// Copyright Epic Games, Inc. All Rights Reserved.
+
+#include "Math/CameraNodeSpaceMath.h"
+
+#include "Core/BuiltInCameraVariables.h"
+#include "Core/CameraEvaluationContext.h"
+#include "Core/CameraNodeEvaluator.h"
+#include "Core/CameraPose.h"
+#include "Core/CameraRigJoints.h"
+#include "Core/CameraSystemEvaluator.h"
+#include "GameFramework/Pawn.h"
+#include "GameFramework/PlayerController.h"
+
+namespace UE::Cameras
+{
+
+FCameraNodeSpaceParams::FCameraNodeSpaceParams(const FCameraNodeEvaluationParams& InEvaluationParams, const FCameraNodeEvaluationResult& InEvaluationResult)
+	: EvaluationParams(InEvaluationParams)
+	, EvaluationResult(InEvaluationResult)
+{
+}
+
+TSharedPtr<const FCameraEvaluationContext> FCameraNodeSpaceParams::GetActiveContext() const
+{
+	if (ensure(EvaluationParams.Evaluator))
+	{
+		const FCameraEvaluationContextStack& ContextStack = EvaluationParams.Evaluator->GetEvaluationContextStack();
+		return ContextStack.GetActiveContext();
+	}
+	return nullptr;
+}
+
+TSharedPtr<const FCameraEvaluationContext> FCameraNodeSpaceParams::GetOwningContext() const
+{
+	return EvaluationParams.EvaluationContext;
+}
+
+bool FCameraNodeSpaceParams::FindPivotTransform(FTransform3d& OutPivotTransform) const
+{
+	const FBuiltInCameraVariables& BuiltInVariables = FBuiltInCameraVariables::Get();
+	TArrayView<const FCameraRigJoint> Joints = EvaluationResult.CameraRigJoints.GetJoints();
+	const FCameraRigJoint* PivotJoint = Joints.FindByPredicate([&BuiltInVariables](const FCameraRigJoint& Joint)
+			{
+				return Joint.VariableID == BuiltInVariables.YawPitchDefinition;
+			});
+	if (PivotJoint)
+	{
+		OutPivotTransform = PivotJoint->Transform;
+		return true;
+	}
+	return false;
+}
+
+bool FCameraNodeSpaceMath::GetCameraNodeOriginPosition(const FCameraNodeEvaluationParams& Params, const FCameraNodeEvaluationResult& Result, ECameraNodeOriginPosition InOriginPosition, FVector3d& OutOrigin)
+{
+	FCameraNodeSpaceParams SpaceParams(Params, Result);
+	return GetCameraNodeOriginPosition(SpaceParams, InOriginPosition, OutOrigin);
+}
+
+bool FCameraNodeSpaceMath::GetCameraNodeOriginPosition(const FCameraNodeSpaceParams& Params, ECameraNodeOriginPosition InOriginPosition, FVector3d& OutOrigin)
+{
+	switch (InOriginPosition)
+	{
+		case ECameraNodeOriginPosition::CameraPose:
+			{
+				OutOrigin = Params.EvaluationResult.CameraPose.GetLocation();
+				return true;
+			}
+		case ECameraNodeOriginPosition::ActiveContext:
+			if (TSharedPtr<const FCameraEvaluationContext> ActiveContext = Params.GetActiveContext())
+			{
+				const FCameraNodeEvaluationResult& InitialResult = ActiveContext->GetInitialResult();
+				ensure(InitialResult.bIsValid);
+				OutOrigin = InitialResult.CameraPose.GetLocation();
+				return true;
+			}
+			break;
+		case ECameraNodeOriginPosition::OwningContext:
+			if (TSharedPtr<const FCameraEvaluationContext> OwningContext = Params.GetOwningContext())
+			{
+				const FCameraNodeEvaluationResult& InitialResult = OwningContext->GetInitialResult();
+				ensure(InitialResult.bIsValid);
+				OutOrigin = InitialResult.CameraPose.GetLocation();
+				return true;
+			}
+			break;
+		case ECameraNodeOriginPosition::Pivot:
+			{
+				FTransform PivotTransform;
+				if (Params.FindPivotTransform(PivotTransform))
+				{
+					OutOrigin = PivotTransform.GetLocation();
+					return true;
+				}
+				else if (TSharedPtr<const FCameraEvaluationContext> ActiveContext = Params.GetActiveContext())
+				{
+					const FCameraNodeEvaluationResult& InitialResult = ActiveContext->GetInitialResult();
+					ensure(InitialResult.bIsValid);
+					OutOrigin = InitialResult.CameraPose.GetLocation();
+					return true;
+				}
+			}
+			break;
+		case ECameraNodeOriginPosition::Pawn:
+			if (TSharedPtr<const FCameraEvaluationContext> ActiveContext = Params.GetActiveContext())
+			{
+				if (APlayerController* PlayerController = ActiveContext->GetPlayerController())
+				{
+					if (APawn* Pawn = PlayerController->GetPawnOrSpectator())
+					{
+						OutOrigin = Pawn->GetActorTransform().GetLocation();
+						return true;
+					}
+				}
+			}
+			break;
+		default:
+			ensureMsgf(false, TEXT("Unsupported camera node space."));
+			break;
+	}
+
+	return false;
+}
+
+bool FCameraNodeSpaceMath::GetCameraNodeSpaceTransform(const FCameraNodeEvaluationParams& Params, const FCameraNodeEvaluationResult& Result, ECameraNodeSpace InSpace, FTransform3d& OutTransform)
+{
+	FCameraNodeSpaceParams SpaceParams(Params, Result);
+	return GetCameraNodeSpaceTransform(SpaceParams, InSpace, OutTransform);
+}
+
+bool FCameraNodeSpaceMath::GetCameraNodeSpaceTransform(const FCameraNodeSpaceParams& Params, ECameraNodeSpace InSpace, FTransform3d& OutTransform)
+{
+	switch (InSpace)
+	{
+		case ECameraNodeSpace::CameraPose:
+			{
+				OutTransform = Params.EvaluationResult.CameraPose.GetTransform();
+				return true;
+			}
+		case ECameraNodeSpace::ActiveContext:
+			if (TSharedPtr<const FCameraEvaluationContext> ActiveContext = Params.GetActiveContext())
+			{
+				const FCameraNodeEvaluationResult& InitialResult = ActiveContext->GetInitialResult();
+				ensure(InitialResult.bIsValid);
+				OutTransform = InitialResult.CameraPose.GetTransform();
+				return true;
+			}
+			break;
+		case ECameraNodeSpace::OwningContext:
+			if (TSharedPtr<const FCameraEvaluationContext> OwningContext = Params.GetOwningContext())
+			{
+				const FCameraNodeEvaluationResult& InitialResult = OwningContext->GetInitialResult();
+				ensure(InitialResult.bIsValid);
+				OutTransform = InitialResult.CameraPose.GetTransform();
+				return true;
+			}
+			break;
+		case ECameraNodeSpace::Pivot:
+			{
+				FTransform PivotTransform;
+				if (Params.FindPivotTransform(PivotTransform))
+				{
+					OutTransform = PivotTransform;
+					return true;
+				}
+				else if (TSharedPtr<const FCameraEvaluationContext> ActiveContext = Params.GetActiveContext())
+				{
+					const FCameraNodeEvaluationResult& InitialResult = ActiveContext->GetInitialResult();
+					ensure(InitialResult.bIsValid);
+					OutTransform = InitialResult.CameraPose.GetTransform();
+					return true;
+				}
+			}
+			break;
+		case ECameraNodeSpace::Pawn:
+			if (TSharedPtr<const FCameraEvaluationContext> ActiveContext = Params.GetActiveContext())
+			{
+				if (APlayerController* PlayerController = ActiveContext->GetPlayerController())
+				{
+					if (APawn* Pawn = PlayerController->GetPawnOrSpectator())
+					{
+						OutTransform = Pawn->GetActorTransform();
+						return true;
+					}
+				}
+			}
+			break;
+		case ECameraNodeSpace::World:
+			{
+				OutTransform = FTransform3d::Identity;
+				return true;
+			}
+			break;
+		default:
+			ensureMsgf(false, TEXT("Unsupported camera node space."));
+			break;
+	}
+
+	return false;
+}
+
+bool FCameraNodeSpaceMath::OffsetCameraNodeSpacePosition(const FCameraNodeEvaluationParams& Params, const FCameraNodeEvaluationResult& Result, const FVector3d& InPosition, const FVector3d& InOffset, ECameraNodeSpace InSpace, FVector3d& OutPosition)
+{
+	FCameraNodeSpaceParams SpaceParams(Params, Result);
+	return OffsetCameraNodeSpacePosition(SpaceParams, InPosition, InOffset, InSpace, OutPosition);
+}
+
+bool FCameraNodeSpaceMath::OffsetCameraNodeSpacePosition(const FCameraNodeSpaceParams& Params, const FVector3d& InPosition, const FVector3d& InOffset, ECameraNodeSpace InSpace, FVector3d& OutPosition)
+{
+	FVector3d WorldOffset = FVector3d::ZeroVector;
+	bool bGotWorldOffset = false;
+
+	switch (InSpace)
+	{
+		case ECameraNodeSpace::CameraPose:
+			{
+				WorldOffset = Params.EvaluationResult.CameraPose.GetRotation().RotateVector(InOffset);
+				bGotWorldOffset = true;
+			}
+			break;
+		case ECameraNodeSpace::ActiveContext:
+			if (TSharedPtr<const FCameraEvaluationContext> ActiveContext = Params.GetActiveContext())
+			{
+				const FCameraNodeEvaluationResult& InitialResult = ActiveContext->GetInitialResult();
+				ensure(InitialResult.bIsValid);
+				WorldOffset = InitialResult.CameraPose.GetRotation().RotateVector(InOffset);
+				bGotWorldOffset = true;
+			}
+			break;
+		case ECameraNodeSpace::OwningContext:
+			if (TSharedPtr<const FCameraEvaluationContext> OwningContext = Params.GetOwningContext())
+			{
+				const FCameraNodeEvaluationResult& InitialResult = OwningContext->GetInitialResult();
+				ensure(InitialResult.bIsValid);
+				WorldOffset = InitialResult.CameraPose.GetRotation().RotateVector(InOffset);
+				bGotWorldOffset = true;
+			}
+			break;
+		case ECameraNodeSpace::Pivot:
+			{
+				FTransform PivotTransform;
+				if (Params.FindPivotTransform(PivotTransform))
+				{
+					WorldOffset = PivotTransform.TransformVectorNoScale(InOffset);
+					bGotWorldOffset = true;
+				}
+				else if (TSharedPtr<const FCameraEvaluationContext> ActiveContext = Params.GetActiveContext())
+				{
+					const FCameraNodeEvaluationResult& InitialResult = ActiveContext->GetInitialResult();
+					ensure(InitialResult.bIsValid);
+					WorldOffset = InitialResult.CameraPose.GetRotation().RotateVector(InOffset);
+					bGotWorldOffset = true;
+				}
+			}
+			break;
+		case ECameraNodeSpace::Pawn:
+			if (TSharedPtr<const FCameraEvaluationContext> ActiveContext = Params.GetActiveContext())
+			{
+				if (APlayerController* PlayerController = ActiveContext->GetPlayerController())
+				{
+					if (APawn* Pawn = PlayerController->GetPawnOrSpectator())
+					{
+						const FTransform3d& PawnTransform = Pawn->GetActorTransform();
+						WorldOffset = PawnTransform.TransformVectorNoScale(InOffset);
+						bGotWorldOffset = true;
+					}
+				}
+			}
+			break;
+		case ECameraNodeSpace::World:
+			{
+				WorldOffset = InOffset;
+				bGotWorldOffset = true;
+			}
+			break;
+		default:
+			ensureMsgf(false, TEXT("Unsupported camera node space."));
+			break;
+	}
+
+	if (bGotWorldOffset)
+	{
+		OutPosition = InPosition + WorldOffset;
+	}
+
+	return bGotWorldOffset;
+}
+
+}  // namespace UE::Cameras
+
diff --git a/Engine/Plugins/Cameras/GameplayCameras/Source/GameplayCameras/Private/Nodes/Collision/CollisionPushCameraNode.cpp b/Engine/Plugins/Cameras/GameplayCameras/Source/GameplayCameras/Private/Nodes/Collision/CollisionPushCameraNode.cpp
--- a/Engine/Plugins/Cameras/GameplayCameras/Source/GameplayCameras/Private/Nodes/Collision/CollisionPushCameraNode.cpp
+++ b/Engine/Plugins/Cameras/GameplayCameras/Source/GameplayCameras/Private/Nodes/Collision/CollisionPushCameraNode.cpp
@@ -18,6 +18,7 @@
 #include "GameFramework/Pawn.h"
 #include "GameFramework/PlayerController.h"
 #include "GameplayCameras.h"
+#include "Math/CameraNodeSpaceMath.h"
 #include "Misc/AssertionMacros.h"
 #include "WorldCollision.h"
 
@@ -41,14 +42,8 @@ protected:
 
 private:
 
-	TOptional<FVector3d> GetFinalSafePosition(
-			const FCameraNodeEvaluationParams& Params, 
-			const FCameraNodeEvaluationResult& OutResult);
-	TOptional<FVector3d> GetSafePosition(
-			const FCameraNodeEvaluationParams& Params, 
-			TSharedPtr<const FCameraEvaluationContext> ActiveContext,
-			const FCameraRigJoint* PivotJoint,
-			const FCameraNodeEvaluationResult& OutResult);
+	TOptional<FVector3d> GetFinalSafePosition(const FCameraNodeEvaluationParams& Params, const FCameraNodeEvaluationResult& OutResult);
+	TOptional<FVector3d> GetSafePosition(const FCameraNodeEvaluationParams& Params, const FCameraNodeEvaluationResult& OutResult);
 
 	void RunCollisionTrace(UWorld* World, APlayerController* PlayerController, const FVector3d& SafePosition, const FCameraNodeEvaluationParams& Params, FCameraNodeEvaluationResult& OutResult);
 	void HandleAsyncCollisionTraceResult(UWorld* World, const FVector3d& SafePosition, const FCameraNodeEvaluationParams& Params, FCameraNodeEvaluationResult& OutResult);
@@ -76,7 +71,7 @@ private:
 	bool bDebugCollisionEnabled = false;
 	bool bDebugFoundHit = false;
 	bool bDebugGotSafePosition = false;
-	bool bDebugGotSafePositionOffsetSpace = false;
+	bool bDebugGotSafePositionOffset = false;
 	FString DebugHitObjectName;
 	FVector3d DebugSafePosition;
 #endif
@@ -87,7 +82,7 @@ UE_DEFINE_CAMERA_NODE_EVALUATOR(FCollisionPushCameraNodeEvaluator)
 UE_DECLARE_CAMERA_DEBUG_BLOCK_START(GAMEPLAYCAMERAS_API, FCollisionPushCameraDebugBlock)
 	UE_DECLARE_CAMERA_DEBUG_BLOCK_FIELD(bool, bCollisionEnabled)
 	UE_DECLARE_CAMERA_DEBUG_BLOCK_FIELD(bool, bGotSafePosition)
-	UE_DECLARE_CAMERA_DEBUG_BLOCK_FIELD(bool, bGotSafePositionOffsetSpace)
+	UE_DECLARE_CAMERA_DEBUG_BLOCK_FIELD(bool, bGotSafePositionOffset)
 	UE_DECLARE_CAMERA_DEBUG_BLOCK_FIELD(ECollisionSafePosition, SafePositionType)
 	UE_DECLARE_CAMERA_DEBUG_BLOCK_FIELD(ECollisionSafePositionOffsetSpace, SafePositionOffsetSpace)
 	UE_DECLARE_CAMERA_DEBUG_BLOCK_FIELD(float, PushFactor);
@@ -182,107 +177,56 @@ TOptional<FVector3d> FCollisionPushCameraNodeEvaluator::GetFinalSafePosition(con
 		return TOptional<FVector3d>();
 	}
 
-	const UCollisionPushCameraNode* ThisNode = GetCameraNodeAs<UCollisionPushCameraNode>();
-
-	// Get the active context.
-	const FCameraEvaluationContextStack& ContextStack = Params.Evaluator->GetEvaluationContextStack();
-	TSharedPtr<const FCameraEvaluationContext> ActiveContext = ContextStack.GetActiveContext();
-	
-	// See if we have a pivot.
-	const FBuiltInCameraVariables& BuiltInVariables = FBuiltInCameraVariables::Get();
-	TArrayView<const FCameraRigJoint> Joints = OutResult.CameraRigJoints.GetJoints();
-	const FCameraRigJoint* PivotJoint = Joints.FindByPredicate([&BuiltInVariables](const FCameraRigJoint& Joint)
-			{
-				return Joint.VariableID == BuiltInVariables.YawPitchDefinition;
-			});
-
 	// Get the safe position itself first.
-	TOptional<FVector3d> OptSafePosition = GetSafePosition(Params, ActiveContext, PivotJoint, OutResult);	
+	TOptional<FVector3d> OptSafePosition = GetSafePosition(Params, OutResult);	
 	if (!OptSafePosition.IsSet())
 	{
 		return TOptional<FVector3d>();
 	}
 
 	// Apply the offset in the specified space.
+	FVector3d SafePosition = OptSafePosition.GetValue();
+	bool bGotSafePositionOffset = true;
 	const FVector3d SafePositionOffset = SafePositionOffsetReader.Get(OutResult.VariableTable);
-
-	FVector3d WorldSafePositionOffset = SafePositionOffset;
-	bool bGotSafePositionOffsetSpace = false;
-
-	switch (ThisNode->SafePositionOffsetSpace)
+	if (!SafePositionOffset.IsZero())
 	{
-		case ECollisionSafePositionOffsetSpace::ActiveContext:
-			if (ActiveContext)
-			{
-				const FCameraNodeEvaluationResult& InitialResult = ActiveContext->GetInitialResult();
-				ensure(InitialResult.bIsValid);
-				WorldSafePositionOffset = InitialResult.CameraPose.GetRotation().RotateVector(SafePositionOffset);
-				bGotSafePositionOffsetSpace = true;
-			}
-			break;
-		case ECollisionSafePositionOffsetSpace::OwningContext:
-			{
-				const FCameraNodeEvaluationResult& InitialResult = Params.EvaluationContext->GetInitialResult();
-				ensure(InitialResult.bIsValid);
-				WorldSafePositionOffset = InitialResult.CameraPose.GetRotation().RotateVector(SafePositionOffset);
-				bGotSafePositionOffsetSpace = true;
-			}
-			break;
-		case ECollisionSafePositionOffsetSpace::Pivot:
-			if (PivotJoint)
-			{
-				WorldSafePositionOffset = PivotJoint->Transform.TransformVectorNoScale(SafePositionOffset);
-				bGotSafePositionOffsetSpace = true;
-			}
-			else if (ActiveContext)
-			{
-				const FCameraNodeEvaluationResult& InitialResult = ActiveContext->GetInitialResult();
-				ensure(InitialResult.bIsValid);
-				WorldSafePositionOffset = InitialResult.CameraPose.GetRotation().RotateVector(SafePositionOffset);
-				bGotSafePositionOffsetSpace = true;
-			}
-			break;
-		case ECollisionSafePositionOffsetSpace::CameraPose:
-			{
-				WorldSafePositionOffset = OutResult.CameraPose.GetRotation().RotateVector(SafePositionOffset);
-				bGotSafePositionOffsetSpace = true;
-			}
-			break;
-		case ECollisionSafePositionOffsetSpace::Pawn:
-			if (ActiveContext)
-			{
-				if (APlayerController* PlayerController = ActiveContext->GetPlayerController())
-				{
-					if (APawn* Pawn = PlayerController->GetPawnOrSpectator())
-					{
-						const FTransform3d& PawnTransform = Pawn->GetActorTransform();
-						WorldSafePositionOffset = PawnTransform.TransformVectorNoScale(SafePositionOffset);
-						bGotSafePositionOffsetSpace = true;
-					}
-				}
-			}
-			break;
+		ECameraNodeSpace SafePositionOffsetSpace;
+		const UCollisionPushCameraNode* ThisNode = GetCameraNodeAs<UCollisionPushCameraNode>();
+		switch (ThisNode->SafePositionOffsetSpace)
+		{
+			case ECollisionSafePositionOffsetSpace::ActiveContext:
+				SafePositionOffsetSpace = ECameraNodeSpace::ActiveContext;
+				break;
+			case ECollisionSafePositionOffsetSpace::OwningContext:
+				SafePositionOffsetSpace = ECameraNodeSpace::OwningContext;
+				break;
+			case ECollisionSafePositionOffsetSpace::Pivot:
+				SafePositionOffsetSpace = ECameraNodeSpace::Pivot;
+				break;
+			case ECollisionSafePositionOffsetSpace::CameraPose:
+				SafePositionOffsetSpace = ECameraNodeSpace::CameraPose;
+				break;
+			case ECollisionSafePositionOffsetSpace::Pawn:
+				SafePositionOffsetSpace = ECameraNodeSpace::Pawn;
+				break;
+			default:
+				ensure(false);
+				SafePositionOffsetSpace = ECameraNodeSpace::Pivot;
+				break;
+		}
+
+		bGotSafePositionOffset = FCameraNodeSpaceMath::OffsetCameraNodeSpacePosition(
+				Params, OutResult, SafePosition, SafePositionOffset, SafePositionOffsetSpace, SafePosition);
 	}
 
 #if UE_GAMEPLAY_CAMERAS_DEBUG
-	bDebugGotSafePositionOffsetSpace = bGotSafePositionOffsetSpace;
+	bDebugGotSafePositionOffset = bGotSafePositionOffset;
 #endif
 
-	if (bGotSafePositionOffsetSpace)
-	{
-		return OptSafePosition.GetValue() + WorldSafePositionOffset;
-	}
-	else
-	{
-		return OptSafePosition.GetValue();
-	}
+	return SafePosition;
 }
 
-TOptional<FVector3d> FCollisionPushCameraNodeEvaluator::GetSafePosition(
-		const FCameraNodeEvaluationParams& Params, 
-		TSharedPtr<const FCameraEvaluationContext> ActiveContext,
-		const FCameraRigJoint* PivotJoint,
-		const FCameraNodeEvaluationResult& OutResult)
+TOptional<FVector3d> FCollisionPushCameraNodeEvaluator::GetSafePosition(const FCameraNodeEvaluationParams& Params, const FCameraNodeEvaluationResult& OutResult)
 {
 	const UCollisionPushCameraNode* ThisNode = GetCameraNodeAs<UCollisionPushCameraNode>();
 
@@ -304,57 +248,30 @@ TOptional<FVector3d> FCollisionPushCameraNodeEvaluator::GetSafePosition(
 	}
 	
 	// Compute the base safe position.
-	FVector3d SafePosition;
-	bool bGotSafePosition = false;
-
+	ECameraNodeOriginPosition OriginPosition;
 	switch (ThisNode->SafePosition)
 	{
 		case ECollisionSafePosition::ActiveContext:
-			if (ActiveContext)
-			{
-				const FCameraNodeEvaluationResult& InitialResult = ActiveContext->GetInitialResult();
-				ensure(InitialResult.bIsValid);
-				SafePosition = InitialResult.CameraPose.GetLocation();
-				bGotSafePosition = true;
-			}
+			OriginPosition = ECameraNodeOriginPosition::ActiveContext;
 			break;
 		case ECollisionSafePosition::OwningContext:
-			{
-				const FCameraNodeEvaluationResult& InitialResult = Params.EvaluationContext->GetInitialResult();
-				ensure(InitialResult.bIsValid);
-				SafePosition = InitialResult.CameraPose.GetLocation();
-				bGotSafePosition = true;
-			}
+			OriginPosition = ECameraNodeOriginPosition::OwningContext;
 			break;
 		case ECollisionSafePosition::Pivot:
-			if (PivotJoint)
-			{
-				SafePosition = PivotJoint->Transform.GetLocation();
-				bGotSafePosition = true;
-			}
-			else if (ActiveContext)
-			{
-				const FCameraNodeEvaluationResult& InitialResult = ActiveContext->GetInitialResult();
-				ensure(InitialResult.bIsValid);
-				SafePosition = InitialResult.CameraPose.GetLocation();
-				bGotSafePosition = true;
-			}
+			OriginPosition = ECameraNodeOriginPosition::Pivot;
 			break;
 		case ECollisionSafePosition::Pawn:
-			if (ActiveContext)
-			{
-				if (APlayerController* PlayerController = ActiveContext->GetPlayerController())
-				{
-					if (APawn* Pawn = PlayerController->GetPawnOrSpectator())
-					{
-						SafePosition = Pawn->GetActorTransform().GetLocation();
-						bGotSafePosition = true;
-					}
-				}
-			}
+			OriginPosition = ECameraNodeOriginPosition::Pawn;
+			break;
+		default:
+			ensure(false);
+			OriginPosition = ECameraNodeOriginPosition::Pivot;
 			break;
 	}
 
+	FVector3d SafePosition;
+	const bool bGotSafePosition = FCameraNodeSpaceMath::GetCameraNodeOriginPosition(Params, OutResult, OriginPosition, SafePosition);
+
 #if UE_GAMEPLAY_CAMERAS_DEBUG
 	bDebugGotSafePosition = bGotSafePosition;
 #endif
@@ -565,7 +482,7 @@ void FCollisionPushCameraNodeEvaluator::OnBuildDebugBlocks(const FCameraDebugBlo
 
 	const UCollisionPushCameraNode* ThisNode = GetCameraNodeAs<UCollisionPushCameraNode>();
 	DebugBlock.bGotSafePosition = bDebugGotSafePosition;
-	DebugBlock.bGotSafePositionOffsetSpace = bDebugGotSafePositionOffsetSpace;
+	DebugBlock.bGotSafePositionOffset = bDebugGotSafePositionOffset;
 	DebugBlock.SafePositionType = ThisNode->SafePosition;
 	DebugBlock.SafePositionOffsetSpace = ThisNode->SafePositionOffsetSpace;
 
@@ -624,7 +541,7 @@ void FCollisionPushCameraDebugBlock::OnDebugDraw(const FCameraDebugBlockDrawPara
 		Renderer.SetTextColor(Colors.Default);
 	}
 
-	if (!bGotSafePositionOffsetSpace)
+	if (!bGotSafePositionOffset)
 	{
 		Renderer.NewLine();
 		Renderer.SetTextColor(Colors.Error);
diff --git a/Engine/Plugins/Cameras/GameplayCameras/Source/GameplayCameras/Public/Math/CameraNodeSpaceMath.h b/Engine/Plugins/Cameras/GameplayCameras/Source/GameplayCameras/Public/Math/CameraNodeSpaceMath.h
new file mode 100644
--- /dev/null
+++ b/Engine/Plugins/Cameras/GameplayCameras/Source/GameplayCameras/Public/Math/CameraNodeSpaceMath.h
@@ -0,0 +1,62 @@
+// Copyright Epic Games, Inc. All Rights Reserved.
+
+#pragma once
+
+#include "Nodes/CameraNodeTypes.h"
+
+#include "Math/Transform.h"
+
+namespace UE::Cameras
+{
+
+class FCameraEvaluationContext;
+struct FCameraNodeEvaluationParams;
+struct FCameraNodeEvaluationResult;
+
+/**
+ * A structure that can wrap node evaluation parameters and results, and spit out information needed
+ * to compute positions and offsets in various evaluation coordinate systems.
+ */
+struct FCameraNodeSpaceParams
+{
+	/** Creates a new FCameraNodeSpaceParams. */
+	FCameraNodeSpaceParams(const FCameraNodeEvaluationParams& InEvaluationParams, const FCameraNodeEvaluationResult& InEvaluationResult);
+
+	/** Gets the active context on the root node's main stack. */
+	TSharedPtr<const FCameraEvaluationContext> GetActiveContext() const;
+	/** Gets the context owning the caller camera node. */
+	TSharedPtr<const FCameraEvaluationContext> GetOwningContext() const;
+	/** Finds a pivot in the evaluation result. */
+	bool FindPivotTransform(FTransform3d& OutPivotTransform) const;
+
+public:
+
+	const FCameraNodeEvaluationParams& EvaluationParams;
+	const FCameraNodeEvaluationResult& EvaluationResult;
+};
+
+class FCameraNodeSpaceMath
+{
+public:
+
+	/** Gets the location of the specified camera node origin. */
+	static bool GetCameraNodeOriginPosition(const FCameraNodeEvaluationParams& Params, const FCameraNodeEvaluationResult& Result, ECameraNodeOriginPosition InOriginPosition, FVector3d& OutOrigin);
+
+	/** Gets the location of the specified camera node origin. */
+	static bool GetCameraNodeOriginPosition(const FCameraNodeSpaceParams& Params, ECameraNodeOriginPosition InOriginPosition, FVector3d& OutOrigin);
+
+	/** Gets the transform of the specified camera node space. */
+	static bool GetCameraNodeSpaceTransform(const FCameraNodeEvaluationParams& Params, const FCameraNodeEvaluationResult& Result, ECameraNodeSpace InSpace, FTransform3d& OutTransform);
+
+	/** Gets the transform of the specified camera node space. */
+	static bool GetCameraNodeSpaceTransform(const FCameraNodeSpaceParams& Params, ECameraNodeSpace InSpace, FTransform3d& OutTransform);
+
+	/** Offsets a given world position in the specified camera node space. */
+	static bool OffsetCameraNodeSpacePosition(const FCameraNodeEvaluationParams& Params, const FCameraNodeEvaluationResult& Result, const FVector3d& InPosition, const FVector3d& InOffset, ECameraNodeSpace InSpace, FVector3d& OutPosition);
+
+	/** Offsets a given world position in the specified camera node space. */
+	static bool OffsetCameraNodeSpacePosition(const FCameraNodeSpaceParams& Params, const FVector3d& InPosition, const FVector3d& InOffset, ECameraNodeSpace InSpace, FVector3d& OutPosition);
+};
+
+}  // namespace UE::Cameras
+
diff --git a/Engine/Plugins/Cameras/GameplayCameras/Source/GameplayCameras/Public/Nodes/CameraNodeTypes.h b/Engine/Plugins/Cameras/GameplayCameras/Source/GameplayCameras/Public/Nodes/CameraNodeTypes.h
--- a/Engine/Plugins/Cameras/GameplayCameras/Source/GameplayCameras/Public/Nodes/CameraNodeTypes.h
+++ b/Engine/Plugins/Cameras/GameplayCameras/Source/GameplayCameras/Public/Nodes/CameraNodeTypes.h
@@ -7,16 +7,40 @@
 #include "CameraNodeTypes.generated.h"
 
 /**
- * Defines what space a camera node should operate in.
+ * Defines an origin position for a camera node to operate at.
  */
 UENUM()
-enum class ECameraNodeSpace
+enum class ECameraNodeOriginPosition : uint8
 {
-	/** Operates in the local camera pose space. */
+	/** The position of the current camera pose. */
 	CameraPose,
-	/** Operates in the space of the owning evaluation context's initial result. */
+	/** The origin of the active evaluation context on the main layer's blend stack. */
+	ActiveContext,
+	/** The origin of the evaluation context of the current camera node. */
 	OwningContext,
-	/** Operates in world space. */
+	/** The location of the current pivot. If no pivot is found, fallback to ActiveContext. */
+	Pivot,
+	/** The location of the player's controlled pawn. */
+	Pawn
+};
+
+/**
+ * Defines what space a camera node, or one of its features, should operate in.
+ */
+UENUM()
+enum class ECameraNodeSpace : uint8
+{
+	/** The local space of the current camera pose. */
+	CameraPose,
+	/** The space of the active evaluation context on the main layer's blend stack. */
+	ActiveContext,
+	/** The space of the evaluation context of the current camera node. */
+	OwningContext,
+	/** The space of the current pivot. If no pivot is found, fallback to ActiveContext. */
+	Pivot,
+	/** The space of the player's controlled pawn. */
+	Pawn,
+	/** The space of the world in which the camera rig evaluates. */
 	World
 };
 
