Cameras: add options for specifying the position and offset space of the occlusion node's target

diff --git a/Engine/Plugins/Cameras/GameplayCameras/Source/GameplayCameras/Private/Nodes/Collision/OcclusionMaterialCameraNode.cpp b/Engine/Plugins/Cameras/GameplayCameras/Source/GameplayCameras/Private/Nodes/Collision/OcclusionMaterialCameraNode.cpp
--- a/Engine/Plugins/Cameras/GameplayCameras/Source/GameplayCameras/Private/Nodes/Collision/OcclusionMaterialCameraNode.cpp
+++ b/Engine/Plugins/Cameras/GameplayCameras/Source/GameplayCameras/Private/Nodes/Collision/OcclusionMaterialCameraNode.cpp
@@ -13,6 +13,7 @@
 #include "GameFramework/Pawn.h"
 #include "GameplayCameras.h"
 #include "Materials/MaterialInstanceDynamic.h"
+#include "Math/CameraNodeSpaceMath.h"
 #include "Misc/AssertionMacros.h"
 #include "WorldCollision.h"
 
@@ -42,7 +43,7 @@ protected:
 
 private:
 
-	void RunOcclusionTrace(UWorld* World, APlayerController* PlayerController, FCameraNodeEvaluationResult& OutResult);
+	void RunOcclusionTrace(UWorld* World, const FCameraNodeEvaluationParams& Params, FCameraNodeEvaluationResult& OutResult);
 	void HandleOcclusionTraceResult(UWorld* World);
 
 	void ApplyOcclusionMaterial(TSet<UStaticMeshComponent*> MeshComponents);
@@ -91,37 +92,59 @@ void FOcclusionMaterialCameraNodeEvaluator::OnRun(const FCameraNodeEvaluationPar
 	}
 
 	UWorld* World = Params.EvaluationContext->GetWorld();
-	APlayerController* PlayerController = Params.EvaluationContext->GetPlayerController();
-	if (!World || !PlayerController)
+	if (!World)
 	{
 		return;
 	}
 
 	HandleOcclusionTraceResult(World);
-	RunOcclusionTrace(World, PlayerController, OutResult);
+	RunOcclusionTrace(World, Params, OutResult);
 }
 
-void FOcclusionMaterialCameraNodeEvaluator::RunOcclusionTrace(UWorld* World, APlayerController* PlayerController, FCameraNodeEvaluationResult& OutResult)
+void FOcclusionMaterialCameraNodeEvaluator::RunOcclusionTrace(UWorld* World, const FCameraNodeEvaluationParams& Params, FCameraNodeEvaluationResult& OutResult)
 {
 	static FName OcclusionTraceTag(TEXT("CameraOcclusion"));
 	static FName OcclusionTraceOwnerTag(TEXT("OcclusionMaterialCameraNode"));
 
-	APawn* Pawn = PlayerController->GetPawn();
+	const UOcclusionMaterialCameraNode* OcclusionMaterialNode = GetCameraNodeAs<UOcclusionMaterialCameraNode>();
 
-	// TODO fix this hardcoded nonsense
-	const FVector3d OcclusionTarget = Pawn->GetActorLocation();
+	const FCameraNodeSpaceParams SpaceParams(Params, OutResult);
+
+	FVector3d OcclusionTarget;
+	const bool bGotOcclusionTarget = FCameraNodeSpaceMath::GetCameraNodeOriginPosition(
+			SpaceParams, OcclusionMaterialNode->OcclusionTargetPosition, OcclusionTarget);
+	if (!bGotOcclusionTarget)
+	{
+		return;
+	}
+
+	const FVector3d OcclusionTargetOffset = OcclusionTargetOffsetReader.Get(OutResult.VariableTable);
+	if (!OcclusionTargetOffset.IsZero())
+	{
+		FCameraNodeSpaceMath::OffsetCameraNodeSpacePosition(
+				SpaceParams,
+				OcclusionTarget, OcclusionTargetOffset, OcclusionMaterialNode->OcclusionTargetOffsetSpace,
+				OcclusionTarget);
+	}
 
-	const UOcclusionMaterialCameraNode* OcclusionMaterialNode = GetCameraNodeAs<UOcclusionMaterialCameraNode>();
 	ECollisionChannel OcclusionChannel = OcclusionMaterialNode->OcclusionChannel;
 
 	const float OcclusionSphereRadius = OcclusionSphereRadiusReader.Get(OutResult.VariableTable);
-	const FVector3d TraceEndOffset = OcclusionTargetOffsetReader.Get(OutResult.VariableTable);
 
 	const FVector3d TraceStart(OutResult.CameraPose.GetLocation());
-	const FVector3d TraceEnd(OcclusionTarget + TraceEndOffset);
+	const FVector3d TraceEnd(OcclusionTarget);
 
-	FCollisionShape SweepShape = FCollisionShape::MakeSphere(OcclusionSphereRadius);
 	// Ignore the player pawn by default.
+	APawn* Pawn = nullptr;
+	if (TSharedPtr<const FCameraEvaluationContext> ActiveContext = SpaceParams.GetActiveContext())
+	{
+		if (APlayerController* PlayerController = ActiveContext->GetPlayerController())
+		{
+			Pawn = PlayerController->GetPawn();
+		}
+	}
+
+	FCollisionShape SweepShape = FCollisionShape::MakeSphere(OcclusionSphereRadius);
 	FCollisionQueryParams QueryParams(SCENE_QUERY_STAT(StartOcclusionSweep), false, Pawn);
 	QueryParams.TraceTag = OcclusionTraceTag;
 	QueryParams.OwnerTag = OcclusionTraceOwnerTag;
diff --git a/Engine/Plugins/Cameras/GameplayCameras/Source/GameplayCameras/Public/Nodes/Collision/OcclusionMaterialCameraNode.h b/Engine/Plugins/Cameras/GameplayCameras/Source/GameplayCameras/Public/Nodes/Collision/OcclusionMaterialCameraNode.h
--- a/Engine/Plugins/Cameras/GameplayCameras/Source/GameplayCameras/Public/Nodes/Collision/OcclusionMaterialCameraNode.h
+++ b/Engine/Plugins/Cameras/GameplayCameras/Source/GameplayCameras/Public/Nodes/Collision/OcclusionMaterialCameraNode.h
@@ -5,6 +5,7 @@
 #include "Core/CameraNode.h"
 #include "Core/CameraParameters.h"
 #include "Engine/EngineTypes.h"
+#include "Nodes/CameraNodeTypes.h"
 
 #include "OcclusionMaterialCameraNode.generated.h"
 
@@ -27,10 +28,22 @@ public:
 
 	/** Collision channel to use for the occlusion testing. */
 	UPROPERTY(EditAnywhere, Category="Occlusion")
-	TEnumAsByte<ECollisionChannel> OcclusionChannel;
-
-	/** World-space offset from the target to the occlusion trace's end. */
-	UPROPERTY(EditAnywhere, Category="Occlusion")
+	TEnumAsByte<ECollisionChannel> OcclusionChannel = ECollisionChannel::ECC_Camera;
+
+	/** 
+	 * The position of the occlusion trace's target. Objects encountered between the current
+	 * camera pose location and this target will have the transparency material applied to
+	 * them until they move out of the way.
+	 */
+	UPROPERTY(EditAnywhere, Category="Occlusion Target")
+	ECameraNodeOriginPosition OcclusionTargetPosition = ECameraNodeOriginPosition::Pawn;
+
+	/** The space in which to apply the occlusion target offset. */
+	UPROPERTY(EditAnywhere, Category="Occlusion Target")
+	ECameraNodeSpace OcclusionTargetOffsetSpace = ECameraNodeSpace::World;
+
+	/** Offset from the target to the occlusion trace's end. */
+	UPROPERTY(EditAnywhere, Category="Occlusion Target")
 	FVector3dCameraParameter OcclusionTargetOffset;
 
 public:
