Cameras: allow choosing the target ray-cast node's physics trace channel

diff --git a/Engine/Plugins/Cameras/GameplayCameras/Source/GameplayCameras/Private/Nodes/Common/TargetRayCastCameraNode.cpp b/Engine/Plugins/Cameras/GameplayCameras/Source/GameplayCameras/Private/Nodes/Common/TargetRayCastCameraNode.cpp
--- a/Engine/Plugins/Cameras/GameplayCameras/Source/GameplayCameras/Private/Nodes/Common/TargetRayCastCameraNode.cpp
+++ b/Engine/Plugins/Cameras/GameplayCameras/Source/GameplayCameras/Private/Nodes/Common/TargetRayCastCameraNode.cpp
@@ -98,6 +98,9 @@ void FTargetRayCastCameraNodeEvaluator::RunLineTrace(UWorld* World, APlayerContr
 	const FVector3d RayStart = CameraPose.GetLocation();
 	const FVector3d RayEnd = RayStart + CameraPose.GetAimDir() * GTargetRayCastLength;
 
+	const UTargetRayCastCameraNode* RayCastNode = GetCameraNodeAs<UTargetRayCastCameraNode>();
+	ECollisionChannel TraceChannel = RayCastNode->TraceChannel;
+
 	FHitResult HitResult;
 
 	// Ignore the player pawn by default.
@@ -108,7 +111,7 @@ void FTargetRayCastCameraNodeEvaluator::RunLineTrace(UWorld* World, APlayerContr
 	const bool bGotHit = World->LineTraceSingleByChannel(
 			HitResult, 
 			RayStart, RayEnd,
-			ECollisionChannel::ECC_Camera,
+			TraceChannel,
 			QueryParams);
 
 #if UE_GAMEPLAY_CAMERAS_DEBUG
diff --git a/Engine/Plugins/Cameras/GameplayCameras/Source/GameplayCameras/Public/Nodes/Common/TargetRayCastCameraNode.h b/Engine/Plugins/Cameras/GameplayCameras/Source/GameplayCameras/Public/Nodes/Common/TargetRayCastCameraNode.h
--- a/Engine/Plugins/Cameras/GameplayCameras/Source/GameplayCameras/Public/Nodes/Common/TargetRayCastCameraNode.h
+++ b/Engine/Plugins/Cameras/GameplayCameras/Source/GameplayCameras/Public/Nodes/Common/TargetRayCastCameraNode.h
@@ -15,6 +15,12 @@ class UTargetRayCastCameraNode : public UCameraNode
 {
 	GENERATED_BODY()
 
+public:
+
+	/** Trace channel to use for the ray-cast. */
+	UPROPERTY(EditAnywhere, Category="Ray-Cast")
+	TEnumAsByte<ECollisionChannel> TraceChannel = ECollisionChannel::ECC_Camera;
+
 protected:
 
 	// UCameraNode interface.
