Cameras: better default property values for some nodes

- Collision trace channel defaults to camera channel
- Attach-to-Pawn node defaults to attaching position

diff --git a/Engine/Plugins/Cameras/GameplayCameras/Source/GameplayCameras/Public/Nodes/Attach/AttachToPlayerPawnCameraNode.h b/Engine/Plugins/Cameras/GameplayCameras/Source/GameplayCameras/Public/Nodes/Attach/AttachToPlayerPawnCameraNode.h
--- a/Engine/Plugins/Cameras/GameplayCameras/Source/GameplayCameras/Public/Nodes/Attach/AttachToPlayerPawnCameraNode.h
+++ b/Engine/Plugins/Cameras/GameplayCameras/Source/GameplayCameras/Public/Nodes/Attach/AttachToPlayerPawnCameraNode.h
@@ -24,10 +24,10 @@ public:
 
 	/** Whether to move the camera to the pawn's location. */
 	UPROPERTY(EditAnywhere, Category=Common)
-	FBooleanCameraParameter AttachToLocation;
+	FBooleanCameraParameter AttachToLocation = true;
 
 	/** Whether to align the camera rotation to the pawn's orientation. */
 	UPROPERTY(EditAnywhere, Category=Common)
-	FBooleanCameraParameter AttachToRotation;
+	FBooleanCameraParameter AttachToRotation = false;
 };
 
diff --git a/Engine/Plugins/Cameras/GameplayCameras/Source/GameplayCameras/Public/Nodes/Collision/CollisionPushCameraNode.h b/Engine/Plugins/Cameras/GameplayCameras/Source/GameplayCameras/Public/Nodes/Collision/CollisionPushCameraNode.h
--- a/Engine/Plugins/Cameras/GameplayCameras/Source/GameplayCameras/Public/Nodes/Collision/CollisionPushCameraNode.h
+++ b/Engine/Plugins/Cameras/GameplayCameras/Source/GameplayCameras/Public/Nodes/Collision/CollisionPushCameraNode.h
@@ -102,7 +102,7 @@ public:
 
 	/** Collision channel to use for the line trace. */
 	UPROPERTY(EditAnywhere, Category="Collision")
-	TEnumAsByte<ECollisionChannel> CollisionChannel;
+	TEnumAsByte<ECollisionChannel> CollisionChannel = ECollisionChannel::ECC_Camera;
 
 	/** The interpolation to use when pushing the camera towards the safe position. */
 	UPROPERTY(EditAnywhere, Category="Collision")
