Cameras: look for transition settings on a combination node's main rig if needed

diff --git a/Engine/Plugins/Cameras/GameplayCameras/Source/GameplayCameras/Private/Core/BlendStackCameraNode.cpp b/Engine/Plugins/Cameras/GameplayCameras/Source/GameplayCameras/Private/Core/BlendStackCameraNode.cpp
--- a/Engine/Plugins/Cameras/GameplayCameras/Source/GameplayCameras/Private/Core/BlendStackCameraNode.cpp
+++ b/Engine/Plugins/Cameras/GameplayCameras/Source/GameplayCameras/Private/Core/BlendStackCameraNode.cpp
@@ -8,6 +8,8 @@
 #include "Core/CameraAsset.h"
 #include "Core/CameraEvaluationContext.h"
 #include "Core/CameraRigAsset.h"
+#include "Core/CameraRigAssetReference.h"
+#include "Core/CameraRigCombinationRegistry.h"
 #include "Core/CameraSystemEvaluator.h"
 #include "Debug/CameraDebugBlockBuilder.h"
 #include "Debug/CameraDebugRenderer.h"
@@ -787,6 +789,9 @@ const UCameraRigTransition* FTransientBlendStackCameraNodeEvaluator::FindTransit
 	const UCameraAsset* ToCameraAsset = ToContext ? ToContext->GetCameraAsset() : nullptr;
 	const UCameraRigAsset* ToCameraRig = Params.CameraRig;
 
+	// If the new entry is a combination, look for transitions on its "main" camera rig.
+	ToCameraRig = UCombinedCameraRigsCameraNode::GetMainCameraRigIfCombination(ToCameraRig);
+
 	// Find a transition that works for blending towards ToCameraRig.
 	// If the stack isn't empty, we need to find a transition that works between the previous and 
 	// next camera rigs. If the stack is empty, we blend the new camera rig in from nothing if
@@ -802,6 +807,9 @@ const UCameraRigTransition* FTransientBlendStackCameraNodeEvaluator::FindTransit
 		const UCameraAsset* FromCameraAsset = FromContext ? FromContext->GetCameraAsset() : nullptr;
 		const UCameraRigAsset* FromCameraRig = TopEntry.CameraRig;
 
+		// If the top entry is a combination, look for transitions on its "main" camera rig.
+		FromCameraRig = UCombinedCameraRigsCameraNode::GetMainCameraRigIfCombination(FromCameraRig);
+
 		if (!TopEntry.bIsFrozen)
 		{
 			// Look for exit transitions on the last active camera rig itself.
diff --git a/Engine/Plugins/Cameras/GameplayCameras/Source/GameplayCameras/Private/Core/CameraRigCombinationRegistry.cpp b/Engine/Plugins/Cameras/GameplayCameras/Source/GameplayCameras/Private/Core/CameraRigCombinationRegistry.cpp
--- a/Engine/Plugins/Cameras/GameplayCameras/Source/GameplayCameras/Private/Core/CameraRigCombinationRegistry.cpp
+++ b/Engine/Plugins/Cameras/GameplayCameras/Source/GameplayCameras/Private/Core/CameraRigCombinationRegistry.cpp
@@ -130,6 +130,27 @@ void FCombinedCameraRigsCameraNodeEvaluator::ApplyParameterOverrides(FCameraVari
 
 }  // namespace UE::Cameras
 
+const UCameraRigAsset* UCombinedCameraRigsCameraNode::GetMainCameraRigIfCombination(const UCameraRigAsset* InCameraRig)
+{
+	if (!InCameraRig || !InCameraRig->RootNode)
+	{
+		return InCameraRig;
+	}
+
+	if (const UCombinedCameraRigsCameraNode* CameraRigCombinationNode = Cast<UCombinedCameraRigsCameraNode>(InCameraRig->RootNode))
+	{
+		if (!CameraRigCombinationNode->CameraRigReferences.IsEmpty())
+		{
+			if (const UCameraRigAsset* MainCameraRig = CameraRigCombinationNode->CameraRigReferences[0].GetCameraRig())
+			{
+				return MainCameraRig;
+			}
+		}
+	}
+
+	return InCameraRig;
+}
+
 FCameraNodeEvaluatorPtr UCombinedCameraRigsCameraNode::OnBuildEvaluator(FCameraNodeEvaluatorBuilder& Builder) const
 {
 	using namespace UE::Cameras;
diff --git a/Engine/Plugins/Cameras/GameplayCameras/Source/GameplayCameras/Private/Core/CameraRigCombinationRegistry.h b/Engine/Plugins/Cameras/GameplayCameras/Source/GameplayCameras/Private/Core/CameraRigCombinationRegistry.h
--- a/Engine/Plugins/Cameras/GameplayCameras/Source/GameplayCameras/Private/Core/CameraRigCombinationRegistry.h
+++ b/Engine/Plugins/Cameras/GameplayCameras/Source/GameplayCameras/Private/Core/CameraRigCombinationRegistry.h
@@ -54,6 +54,10 @@ class UCombinedCameraRigsCameraNode : public UCameraNode
 {
 	GENERATED_BODY()
 
+public:
+
+	static const UCameraRigAsset* GetMainCameraRigIfCombination(const UCameraRigAsset* InCameraRig);
+
 protected:
 
 	// UCameraNode interface.
diff --git a/Engine/Plugins/Cameras/GameplayCameras/Source/GameplayCameras/Private/Services/OrientationInitializationService.cpp b/Engine/Plugins/Cameras/GameplayCameras/Source/GameplayCameras/Private/Services/OrientationInitializationService.cpp
--- a/Engine/Plugins/Cameras/GameplayCameras/Source/GameplayCameras/Private/Services/OrientationInitializationService.cpp
+++ b/Engine/Plugins/Cameras/GameplayCameras/Source/GameplayCameras/Private/Services/OrientationInitializationService.cpp
@@ -6,6 +6,7 @@
 #include "Core/CameraEvaluationContextStack.h"
 #include "Core/CameraOperation.h"
 #include "Core/CameraRigAsset.h"
+#include "Core/CameraRigCombinationRegistry.h"
 #include "Core/CameraSystemEvaluator.h"
 #include "Core/RootCameraNode.h"
 #include "Core/RootCameraNodeCameraRigEvent.h"
@@ -56,9 +57,13 @@ void FOrientationInitializationService::OnRootCameraNodeEvent(const FRootCameraN
 	{
 		ECameraRigInitialOrientation InitialOrientation = ECameraRigInitialOrientation::None;
 
-		if (InEvent.CameraRigInfo.CameraRig)
+		if (const UCameraRigAsset* NewCameraRig = InEvent.CameraRigInfo.CameraRig)
 		{
-			InitialOrientation = InEvent.CameraRigInfo.CameraRig->InitialOrientation;
+			// If the new camera rig is a combination, find its initial orientation settings
+			// on its main rig.
+			NewCameraRig = UCombinedCameraRigsCameraNode::GetMainCameraRigIfCombination(NewCameraRig);
+
+			InitialOrientation = NewCameraRig->InitialOrientation;
 		}
 
 		if (InEvent.Transition && InEvent.Transition->bOverrideInitialOrientation)
