Cameras: remove Class Viewer warnings

diff --git a/Engine/Plugins/Cameras/GameplayCameras/Source/GameplayCameras/Public/Core/BlendCameraNode.h b/Engine/Plugins/Cameras/GameplayCameras/Source/GameplayCameras/Public/Core/BlendCameraNode.h
--- a/Engine/Plugins/Cameras/GameplayCameras/Source/GameplayCameras/Public/Core/BlendCameraNode.h
+++ b/Engine/Plugins/Cameras/GameplayCameras/Source/GameplayCameras/Public/Core/BlendCameraNode.h
@@ -10,7 +10,7 @@
 /**
  * Base class for blend camera nodes.
  */
-UCLASS(MinimalAPI, Hidden, meta=(CameraNodeCategories="Blends"))
+UCLASS(MinimalAPI, Abstract, meta=(CameraNodeCategories="Blends"))
 class UBlendCameraNode : public UCameraNode
 {
 	GENERATED_BODY()
