Cameras: fix pin direction for camera rig parameter Blueprint setters

diff --git a/Engine/Plugins/Cameras/GameplayCameras/Source/GameplayCameras/Public/GameFramework/CameraRigParameterInterop.h b/Engine/Plugins/Cameras/GameplayCameras/Source/GameplayCameras/Public/GameFramework/CameraRigParameterInterop.h
--- a/Engine/Plugins/Cameras/GameplayCameras/Source/GameplayCameras/Public/GameFramework/CameraRigParameterInterop.h
+++ b/Engine/Plugins/Cameras/GameplayCameras/Source/GameplayCameras/Public/GameFramework/CameraRigParameterInterop.h
@@ -32,31 +32,31 @@ public:
 public:
 
 	UFUNCTION(BlueprintCallable, Category="Camera", meta=(BlueprintInternalUseOnly="true"))
-	static void SetBooleanParameter(FBlueprintCameraVariableTable& VariableTable, UCameraRigAsset* CameraRig, const FString& ParameterName, bool bParameterValue);
+	static void SetBooleanParameter(UPARAM(Ref) FBlueprintCameraVariableTable& VariableTable, UCameraRigAsset* CameraRig, const FString& ParameterName, bool bParameterValue);
 
 	UFUNCTION(BlueprintCallable, Category="Camera", meta=(BlueprintInternalUseOnly="true"))
-	static void SetIntegerParameter(FBlueprintCameraVariableTable& VariableTable, UCameraRigAsset* CameraRig, const FString& ParameterName, int32 ParameterValue);
+	static void SetIntegerParameter(UPARAM(Ref) FBlueprintCameraVariableTable& VariableTable, UCameraRigAsset* CameraRig, const FString& ParameterName, int32 ParameterValue);
 
 	UFUNCTION(BlueprintCallable, Category="Camera", meta=(BlueprintInternalUseOnly="true"))
-	static void SetFloatParameter(FBlueprintCameraVariableTable& VariableTable, UCameraRigAsset* CameraRig, const FString& ParameterName, double ParameterValue);
+	static void SetFloatParameter(UPARAM(Ref) FBlueprintCameraVariableTable& VariableTable, UCameraRigAsset* CameraRig, const FString& ParameterName, double ParameterValue);
 
 	UFUNCTION(BlueprintCallable, Category="Camera", meta=(BlueprintInternalUseOnly="true"))
-	static void SetDoubleParameter(FBlueprintCameraVariableTable& VariableTable, UCameraRigAsset* CameraRig, const FString& ParameterName, double ParameterValue);
+	static void SetDoubleParameter(UPARAM(Ref) FBlueprintCameraVariableTable& VariableTable, UCameraRigAsset* CameraRig, const FString& ParameterName, double ParameterValue);
 
 	UFUNCTION(BlueprintCallable, Category="Camera", meta=(BlueprintInternalUseOnly="true"))
-	static void SetVector2Parameter(FBlueprintCameraVariableTable& VariableTable, UCameraRigAsset* CameraRig, const FString& ParameterName, FVector2D ParameterValue);
+	static void SetVector2Parameter(UPARAM(Ref) FBlueprintCameraVariableTable& VariableTable, UCameraRigAsset* CameraRig, const FString& ParameterName, FVector2D ParameterValue);
 
 	UFUNCTION(BlueprintCallable, Category="Camera", meta=(BlueprintInternalUseOnly="true"))
-	static void SetVector3Parameter(FBlueprintCameraVariableTable& VariableTable, UCameraRigAsset* CameraRig, const FString& ParameterName, FVector ParameterValue);
+	static void SetVector3Parameter(UPARAM(Ref) FBlueprintCameraVariableTable& VariableTable, UCameraRigAsset* CameraRig, const FString& ParameterName, FVector ParameterValue);
 
 	UFUNCTION(BlueprintCallable, Category="Camera", meta=(BlueprintInternalUseOnly="true"))
-	static void SetVector4Parameter(FBlueprintCameraVariableTable& VariableTable, UCameraRigAsset* CameraRig, const FString& ParameterName, FVector4 ParameterValue);
+	static void SetVector4Parameter(UPARAM(Ref) FBlueprintCameraVariableTable& VariableTable, UCameraRigAsset* CameraRig, const FString& ParameterName, FVector4 ParameterValue);
 
 	UFUNCTION(BlueprintCallable, Category="Camera", meta=(BlueprintInternalUseOnly="true"))
-	static void SetRotatorParameter(FBlueprintCameraVariableTable& VariableTable, UCameraRigAsset* CameraRig, const FString& ParameterName, FRotator ParameterValue);
+	static void SetRotatorParameter(UPARAM(Ref) FBlueprintCameraVariableTable& VariableTable, UCameraRigAsset* CameraRig, const FString& ParameterName, FRotator ParameterValue);
 
 	UFUNCTION(BlueprintCallable, Category="Camera", meta=(BlueprintInternalUseOnly="true"))
-	static void SetTransformParameter(FBlueprintCameraVariableTable& VariableTable, UCameraRigAsset* CameraRig, const FString& ParameterName, FTransform ParameterValue);
+	static void SetTransformParameter(UPARAM(Ref) FBlueprintCameraVariableTable& VariableTable, UCameraRigAsset* CameraRig, const FString& ParameterName, FTransform ParameterValue);
 
 private:
 
