Cameras: creating a new camera rig opens it in the graph editor and focuses the field to rename it

diff --git a/Engine/Plugins/Cameras/GameplayCameras/Source/GameplayCamerasEditor/Private/Editors/SCameraRigList.cpp b/Engine/Plugins/Cameras/GameplayCameras/Source/GameplayCamerasEditor/Private/Editors/SCameraRigList.cpp
--- a/Engine/Plugins/Cameras/GameplayCameras/Source/GameplayCamerasEditor/Private/Editors/SCameraRigList.cpp
+++ b/Engine/Plugins/Cameras/GameplayCameras/Source/GameplayCamerasEditor/Private/Editors/SCameraRigList.cpp
@@ -231,6 +231,23 @@ void SCameraRigList::Tick(const FGeometry& AllottedGeometry, const double InCurr
 		ListView->RequestListRefresh();
 	}
 
+	if (DeferredFinishAddCameraRig)
+	{
+		// If we just added a new camera rig, find it in the list of items. We're going 
+		// to request for it to be open in the graph editor, and enter rename mode.
+		TSharedPtr<FCameraRigListItem> AddedListItem = FindListItem(DeferredFinishAddCameraRig);
+		DeferredFinishAddCameraRig = nullptr;
+
+		if (AddedListItem)
+		{
+			ListView->SetSelection(AddedListItem);
+			OnRequestEditCameraRig.ExecuteIfBound(AddedListItem->CameraRigAsset);
+
+			DeferredRequestRenameItem = AddedListItem;
+			ListView->RequestScrollIntoView(AddedListItem);
+		}
+	}
+
 	SCompoundWidget::Tick(AllottedGeometry, InCurrentTime, InDeltaTime);
 }
 
@@ -297,6 +314,7 @@ void SCameraRigList::OnAddCameraRig()
 			);
 	CameraAsset->AddCameraRig(NewCameraRig);
 
+	DeferredFinishAddCameraRig = NewCameraRig;
 	bUpdateItemSource = true;
 }
 
@@ -433,6 +451,20 @@ void SCameraRigList::UpdateFilteredItemSource()
 	}
 }
 
+TSharedPtr<FCameraRigListItem> SCameraRigList::FindListItem(UCameraRigAsset* InCameraRig)
+{
+	TSharedPtr<FCameraRigListItem>* FoundItem = FilteredItemSource.FindByPredicate(
+			[InCameraRig](TSharedPtr<FCameraRigListItem> Item)
+			{
+				return Item->CameraRigAsset == InCameraRig;
+			});
+	if (FoundItem)
+	{
+		return *FoundItem;
+	}
+	return nullptr;
+}
+
 TSharedRef<ITableRow> SCameraRigList::OnListGenerateItemRow(TSharedPtr<FCameraRigListItem> Item, const TSharedRef<STableViewBase>& OwnerTable)
 {
 	return SNew(SCameraRigListEntry, OwnerTable)
diff --git a/Engine/Plugins/Cameras/GameplayCameras/Source/GameplayCamerasEditor/Public/Editors/SCameraRigList.h b/Engine/Plugins/Cameras/GameplayCameras/Source/GameplayCamerasEditor/Public/Editors/SCameraRigList.h
--- a/Engine/Plugins/Cameras/GameplayCameras/Source/GameplayCamerasEditor/Public/Editors/SCameraRigList.h
+++ b/Engine/Plugins/Cameras/GameplayCameras/Source/GameplayCamerasEditor/Public/Editors/SCameraRigList.h
@@ -110,6 +110,8 @@ private:
 	void UpdateItemSource();
 	void UpdateFilteredItemSource();
 
+	TSharedPtr<FCameraRigListItem> FindListItem(UCameraRigAsset* InCameraRig);
+
 	TSharedRef<ITableRow> OnListGenerateItemRow(TSharedPtr<FCameraRigListItem> Item, const TSharedRef<STableViewBase>& OwnerTable);
 	void OnListItemScrolledIntoView(TSharedPtr<FCameraRigListItem> Item, const TSharedPtr<ITableRow>& ItemWidget);
 	void OnListMouseButtonDoubleClick(TSharedPtr<FCameraRigListItem> Item);
@@ -138,6 +140,7 @@ private:
 
 	TArray<TSharedPtr<FCameraRigListItem>> FilteredItemSource;
 	TSharedPtr<FCameraRigListItem> DeferredRequestRenameItem;
+	UCameraRigAsset* DeferredFinishAddCameraRig = nullptr;
 
 	bool bUpdateItemSource = false;
 	bool bUpdateFilteredItemSource = false;
