Cameras: new icons!

- Set a billboard icon for the camera system actor/component, since its position and orientation don't matter.
- Set icons for camera and system actors/components.
- Set the icon for the cameras debugger tool.

#virtualized

diff --git a/Engine/Plugins/Cameras/GameplayCameras/Source/GameplayCameras/Private/GameFramework/GameplayCameraComponent.cpp b/Engine/Plugins/Cameras/GameplayCameras/Source/GameplayCameras/Private/GameFramework/GameplayCameraComponent.cpp
--- a/Engine/Plugins/Cameras/GameplayCameras/Source/GameplayCameras/Private/GameFramework/GameplayCameraComponent.cpp
+++ b/Engine/Plugins/Cameras/GameplayCameras/Source/GameplayCameras/Private/GameFramework/GameplayCameraComponent.cpp
@@ -299,7 +299,21 @@ void UGameplayCameraComponent::UpdatePreviewMeshTransform()
 	}
 }
 
-#endif
+#endif  // WITH_EDITORONLY_DATA
+
+#if WITH_EDITOR
+
+bool UGameplayCameraComponent::GetEditorPreviewInfo(float DeltaTime, FMinimalViewInfo& ViewOut)
+{
+	// TODO: in the future, run the camera asset in a private camera system evaluator, with a UI
+	//		 to pick which camera rig to preview.
+	const FTransform3d& ComponentTransform = GetComponentTransform();
+	ViewOut.Location = ComponentTransform.GetLocation();
+	ViewOut.Rotation = ComponentTransform.Rotator();
+	return true;
+}
+
+#endif  // WITH_EDITOR
 
 namespace UE::Cameras
 {
diff --git a/Engine/Plugins/Cameras/GameplayCameras/Source/GameplayCameras/Private/GameFramework/GameplayCameraSystemComponent.cpp b/Engine/Plugins/Cameras/GameplayCameras/Source/GameplayCameras/Private/GameFramework/GameplayCameraSystemComponent.cpp
--- a/Engine/Plugins/Cameras/GameplayCameras/Source/GameplayCameras/Private/GameFramework/GameplayCameraSystemComponent.cpp
+++ b/Engine/Plugins/Cameras/GameplayCameras/Source/GameplayCameras/Private/GameFramework/GameplayCameraSystemComponent.cpp
@@ -2,11 +2,11 @@
 
 #include "GameFramework/GameplayCameraSystemComponent.h"
 
-#include "Components/StaticMeshComponent.h"
+#include "Components/BillboardComponent.h"
 #include "Core/CameraSystemEvaluator.h"
 #include "Debug/DebugDrawService.h"
 #include "Engine/Canvas.h"
-#include "Engine/StaticMesh.h"
+#include "Engine/Texture2D.h"
 #include "Engine/World.h"
 #include "GameFramework/Controller.h"
 #include "GameFramework/GameplayCameraSystemHost.h"
@@ -21,11 +21,27 @@ UGameplayCameraSystemComponent::UGameplayCameraSystemComponent(const FObjectInit
 	: Super(ObjectInit)
 {
 #if WITH_EDITORONLY_DATA
+	EditorSpriteComponent = CreateEditorOnlyDefaultSubobject<UBillboardComponent>(TEXT("Sprite"));
 	if (GIsEditor && !IsRunningCommandlet())
 	{
-		static ConstructorHelpers::FObjectFinder<UStaticMesh> EditorCameraMesh(
-				TEXT("/Engine/EditorMeshes/Camera/SM_CineCam.SM_CineCam"));
-		PreviewMesh = EditorCameraMesh.Object;
+		static ConstructorHelpers::FObjectFinder<UTexture2D> EditorSpriteTextureFinder(
+				TEXT("/GameplayCameras/Textures/S_GameplayCameraSystem.S_GameplayCameraSystem"));
+
+		UTexture2D* EditorSpriteTexture = EditorSpriteTextureFinder.Object;
+
+		if (EditorSpriteTexture && EditorSpriteComponent)
+		{
+			EditorSpriteComponent->SpriteInfo.Category = TEXT("Cameras");
+			EditorSpriteComponent->SpriteInfo.DisplayName = NSLOCTEXT("SpriteCategory", "Cameras", "Cameras");
+
+			EditorSpriteComponent->SetSprite(EditorSpriteTexture);
+			EditorSpriteComponent->SetRelativeScale3D(FVector3d(EditorSpriteTextureScale));
+
+			EditorSpriteComponent->bHiddenInGame = true;
+			EditorSpriteComponent->bIsScreenSizeScaled = true;
+			EditorSpriteComponent->SetUsingAbsoluteScale(true);
+			EditorSpriteComponent->SetupAttachment(this);
+		}
 	}
 #endif  // WITH_EDITORONLY_DATA
 }
@@ -84,21 +100,6 @@ void UGameplayCameraSystemComponent::OnRegister()
 				TEXT("Game"), FDebugDrawDelegate::CreateUObject(this, &UGameplayCameraSystemComponent::DebugDraw));
 	}
 #endif  // UE_GAMEPLAY_CAMERAS_DEBUG
-
-#if WITH_EDITORONLY_DATA
-	if (PreviewMesh && !PreviewMeshComponent)
-	{
-		PreviewMeshComponent = NewObject<UStaticMeshComponent>(this, NAME_None, RF_Transactional | RF_TextExportTransient);
-		PreviewMeshComponent->SetupAttachment(this);
-		PreviewMeshComponent->SetIsVisualizationComponent(true);
-		PreviewMeshComponent->SetStaticMesh(PreviewMesh);
-		PreviewMeshComponent->SetCollisionProfileName(UCollisionProfile::NoCollision_ProfileName);
-		PreviewMeshComponent->bHiddenInGame = true;
-		PreviewMeshComponent->CastShadow = false;
-		PreviewMeshComponent->CreationMethod = CreationMethod;
-		PreviewMeshComponent->RegisterComponentWithWorld(GetWorld());
-	}
-#endif	// WITH_EDITORONLY_DATA
 }
 
 void UGameplayCameraSystemComponent::ActivateCameraSystemForPlayerIndex(int32 PlayerIndex)
@@ -208,13 +209,6 @@ void UGameplayCameraSystemComponent::OnComponentDestroyed(bool bDestroyingHierar
 {
 	Super::OnComponentDestroyed(bDestroyingHierarchy);
 
-#if WITH_EDITORONLY_DATA
-	if (PreviewMeshComponent)
-	{
-		PreviewMeshComponent->DestroyComponent();
-	}
-#endif  // WITH_EDITORONLY_DATA
-
 #if UE_GAMEPLAY_CAMERAS_DEBUG
 	if (DebugDrawDelegateHandle.IsValid())
 	{
diff --git a/Engine/Plugins/Cameras/GameplayCameras/Source/GameplayCameras/Public/GameFramework/GameplayCameraComponent.h b/Engine/Plugins/Cameras/GameplayCameras/Source/GameplayCameras/Public/GameFramework/GameplayCameraComponent.h
--- a/Engine/Plugins/Cameras/GameplayCameras/Source/GameplayCameras/Public/GameFramework/GameplayCameraComponent.h
+++ b/Engine/Plugins/Cameras/GameplayCameras/Source/GameplayCameras/Public/GameFramework/GameplayCameraComponent.h
@@ -75,6 +75,9 @@ public:
 	virtual void EndPlay(const EEndPlayReason::Type EndPlayReason) override;
 	virtual void TickComponent(float DeltaTime, enum ELevelTick TickType, FActorComponentTickFunction *ThisTickFunction) override;
 	virtual void OnComponentDestroyed(bool bDestroyingHierarchy) override;
+#if WITH_EDITOR
+	virtual bool GetEditorPreviewInfo(float DeltaTime, FMinimalViewInfo& ViewOut) override;
+#endif 
 
 	// USceneComponent interface.
 	virtual void OnUpdateTransform(EUpdateTransformFlags UpdateTransformFlags, ETeleportType Teleport) override;
diff --git a/Engine/Plugins/Cameras/GameplayCameras/Source/GameplayCameras/Public/GameFramework/GameplayCameraSystemComponent.h b/Engine/Plugins/Cameras/GameplayCameras/Source/GameplayCameras/Public/GameFramework/GameplayCameraSystemComponent.h
--- a/Engine/Plugins/Cameras/GameplayCameras/Source/GameplayCameras/Public/GameFramework/GameplayCameraSystemComponent.h
+++ b/Engine/Plugins/Cameras/GameplayCameras/Source/GameplayCameras/Public/GameFramework/GameplayCameraSystemComponent.h
@@ -9,6 +9,7 @@
 #include "GameplayCameraSystemComponent.generated.h"
 
 class APlayerController;
+class UBillboardComponent;
 class UCameraRigAsset;
 class UCanvas;
 class UGameplayCameraSystemHost;
@@ -105,11 +106,13 @@ private:
 
 #if WITH_EDITORONLY_DATA
 
-	UPROPERTY()
-	TObjectPtr<UStaticMesh> PreviewMesh;
+	/** Sprite scaling for the editor. */
+	UPROPERTY(transient)
+	float EditorSpriteTextureScale = 0.5f;
 
+	/** Sprite component for the editor. */
 	UPROPERTY()
-	TObjectPtr<UStaticMeshComponent> PreviewMeshComponent;
+	TObjectPtr<UBillboardComponent> EditorSpriteComponent;
 
 #endif	// WITH_EDITORONLY_DATA
 };
diff --git a/Engine/Plugins/Cameras/GameplayCameras/Source/GameplayCamerasEditor/Private/Debugger/SGameplayCamerasDebugger.cpp b/Engine/Plugins/Cameras/GameplayCameras/Source/GameplayCamerasEditor/Private/Debugger/SGameplayCamerasDebugger.cpp
--- a/Engine/Plugins/Cameras/GameplayCameras/Source/GameplayCamerasEditor/Private/Debugger/SGameplayCamerasDebugger.cpp
+++ b/Engine/Plugins/Cameras/GameplayCameras/Source/GameplayCamerasEditor/Private/Debugger/SGameplayCamerasDebugger.cpp
@@ -40,12 +40,15 @@ const FName SGameplayCamerasDebugger::ToolbarName(TEXT("GameplayCamerasDebugger.
 
 void SGameplayCamerasDebugger::RegisterTabSpawners()
 {
+	TSharedRef<FGameplayCamerasEditorStyle> CamerasEditorStyle = FGameplayCamerasEditorStyle::Get();
+
 	FGlobalTabmanager::Get()->RegisterNomadTabSpawner(
 		SGameplayCamerasDebugger::WindowName,
 		FOnSpawnTab::CreateStatic(&SGameplayCamerasDebugger::SpawnGameplayCamerasDebugger)
 	)
 	.SetDisplayName(LOCTEXT("TabDisplayName", "Cameras Debugger"))
 	.SetTooltipText(LOCTEXT("TabTooltipText", "Open the Cameras Debugger tab."))
+	.SetIcon(FSlateIcon(CamerasEditorStyle->GetStyleSetName(), "Debugger.TabIcon"))
 	.SetGroup(WorkspaceMenu::GetMenuStructure().GetDeveloperToolsDebugCategory())
 	.SetCanSidebarTab(false);
 }
diff --git a/Engine/Plugins/Cameras/GameplayCameras/Source/GameplayCamerasEditor/Private/Styles/GameplayCamerasEditorStyle.cpp b/Engine/Plugins/Cameras/GameplayCameras/Source/GameplayCamerasEditor/Private/Styles/GameplayCamerasEditorStyle.cpp
--- a/Engine/Plugins/Cameras/GameplayCameras/Source/GameplayCamerasEditor/Private/Styles/GameplayCamerasEditorStyle.cpp
+++ b/Engine/Plugins/Cameras/GameplayCameras/Source/GameplayCamerasEditor/Private/Styles/GameplayCamerasEditorStyle.cpp
@@ -46,6 +46,16 @@ FGameplayCamerasEditorStyle::FGameplayCamerasEditorStyle()
 	Set("ClassIcon.CameraVariableCollection", new IMAGE_BRUSH_SVG("Icons/ContentBrowser-CameraVariableCollection", Icon16x16));
 	Set("ClassThumbnail.CameraVariableCollection", new IMAGE_BRUSH_SVG("Icons/ContentBrowser-CameraVariableCollection", Icon64x64));
 
+	// Camera actors and components.
+	Set("ClassIcon.GameplayCameraComponent", new IMAGE_BRUSH_SVG("Icons/GameplayCamera_16", Icon16x16));
+	Set("ClassThumbnail.GameplayCameraComponent", new IMAGE_BRUSH_SVG("Icons/GameplayCamera_64", Icon64x64));
+	Set("ClassIcon.GameplayCameraActor", new IMAGE_BRUSH_SVG("Icons/GameplayCamera_16", Icon16x16));
+	Set("ClassThumbnail.GameplayCameraActor", new IMAGE_BRUSH_SVG("Icons/GameplayCamera_64", Icon64x64));
+	Set("ClassIcon.GameplayCameraSystemComponent", new IMAGE_BRUSH_SVG("Icons/GameplayCameraSystem_16", Icon16x16));
+	Set("ClassThumbnail.GameplayCameraSystemComponent", new IMAGE_BRUSH_SVG("Icons/GameplayCameraSystem_64", Icon64x64));
+	Set("ClassIcon.GameplayCameraSystemActor", new IMAGE_BRUSH_SVG("Icons/GameplayCameraSystem_16", Icon16x16));
+	Set("ClassThumbnail.GameplayCameraSystemActor", new IMAGE_BRUSH_SVG("Icons/GameplayCameraSystem_64", Icon64x64));
+
 	// Camera asset editor icons.
 	Set("CameraAssetEditor.ShowCameraDirector", new IMAGE_BRUSH_SVG("Icons/CameraEditor-CameraDirector", Icon20x20));
 	Set("CameraAssetEditor.ShowCameraRigs", new IMAGE_BRUSH_SVG("Icons/CameraEditor-CameraRigs", Icon20x20));
@@ -100,6 +110,8 @@ FGameplayCamerasEditorStyle::FGameplayCamerasEditorStyle()
 	Set("CameraVariableCollectionEditor.Entry.Value", FTextBlockStyle(NormalText).SetFont(DEFAULT_FONT("Regular", 10)));
 
 	// Debugger tool icons.
+	Set("Debugger.TabIcon", new IMAGE_BRUSH_SVG("Icons/GameplayCameraSystem_16", Icon16x16));
+
 	Set("Debugger.DebugInfoEnabled.Icon", new IMAGE_BRUSH_SVG("Icons/DebugInfo-ToggleCheck", Icon16x16, FStyleColors::AccentGreen));
 	Set("Debugger.DebugInfoDisabled.Icon", new IMAGE_BRUSH_SVG("Icons/DebugInfo-ToggleCheck", Icon16x16, FStyleColors::AccentGray));
 
