Cameras: convert lens parameters into post-process DoF settings

- Implement blending of post-process settings, so the camera system only maintains one structure through the whole evaluation (and since that's what the engine wants out of it).
- Convert Aperture, FocusDistance, etc. into the appropriate post-process settings at the end of the evaluation.
- Add a couple of missing properties (ISO, ShutterSpeed), and add an option for the raycast node to set FocusDistance automatically (basically an auto-focus).
- Adjust debug displays and BP wrappers.

diff --git a/Engine/Plugins/Cameras/GameplayCameras/Source/GameplayCameras/Private/Core/CameraPose.cpp b/Engine/Plugins/Cameras/GameplayCameras/Source/GameplayCameras/Private/Core/CameraPose.cpp
--- a/Engine/Plugins/Cameras/GameplayCameras/Source/GameplayCameras/Private/Core/CameraPose.cpp
+++ b/Engine/Plugins/Cameras/GameplayCameras/Source/GameplayCameras/Private/Core/CameraPose.cpp
@@ -3,11 +3,19 @@
 #include "Core/CameraPose.h"
 
 #include "Engine/EngineTypes.h"
+#include "Engine/Scene.h"
 #include "GameplayCameras.h"
+#include "HAL/IConsoleManager.h"
 #include "Math/Ray.h"
 
 #include UE_INLINE_GENERATED_CPP_BY_NAME(CameraPose)
 
+float GGameplayCamerasDefaultMinFstop = 0.f;
+static FAutoConsoleVariableRef CVarGameplayCamerasDefaultMinFstop(
+	TEXT("GameplayCameras.DefaultMinFstop"),
+	GGameplayCamerasDefaultMinFstop,
+	TEXT("(Default: 0. Minimum camera lens aperture (f-stop) that defines the curvature of the diaphragm blades."));
+
 const FCameraPoseFlags& FCameraPoseFlags::All()
 {
 	static FCameraPoseFlags Instance(true);
@@ -168,6 +176,41 @@ void FCameraPose::GetDefaultSensorSize(float& OutSensorWidth, float& OutSensorHe
 	OutSensorHeight = 18.67f;
 }
 
+bool FCameraPose::ApplyPhysicalCameraSettings(FPostProcessSettings& PostProcessSettings, bool bOverwriteSettings) const
+{
+	if (!EnablePhysicalCamera || PhysicalCameraBlendWeight <= 0.f)
+	{
+		return false;
+	}
+
+#define UE_LERP_PP(SettingName, Value)\
+	if (!PostProcessSettings.bOverride_##SettingName || bOverwriteSettings)\
+	{\
+		PostProcessSettings.bOverride_##SettingName = true;\
+		PostProcessSettings.SettingName = FMath::Lerp(PostProcessSettings.SettingName, Value, PhysicalCameraBlendWeight);\
+	}
+
+	UE_LERP_PP(CameraISO, ISO);
+	UE_LERP_PP(CameraShutterSpeed, ShutterSpeed);
+
+	UE_LERP_PP(DepthOfFieldFstop, Aperture);
+	UE_LERP_PP(DepthOfFieldBladeCount, DiaphragmBladeCount);
+
+	// TODO: add this to the camera pose?
+	UE_LERP_PP(DepthOfFieldMinFstop, GGameplayCamerasDefaultMinFstop);
+
+	// TODO: support minimum-focus-distance?
+	UE_LERP_PP(DepthOfFieldFocalDistance, FocusDistance);
+
+	// TODO: support overscan?
+	UE_LERP_PP(DepthOfFieldSensorWidth, SensorWidth);
+	UE_LERP_PP(DepthOfFieldSqueezeFactor, SqueezeFactor);
+
+#undef UE_LERP_PP
+
+	return true;
+}
+
 FRay3d FCameraPose::GetAimRay() const
 {
 	const bool bDirectionIsNormalized = false;
diff --git a/Engine/Plugins/Cameras/GameplayCameras/Source/GameplayCameras/Private/Core/CameraSystemEvaluator.cpp b/Engine/Plugins/Cameras/GameplayCameras/Source/GameplayCameras/Private/Core/CameraSystemEvaluator.cpp
--- a/Engine/Plugins/Cameras/GameplayCameras/Source/GameplayCameras/Private/Core/CameraSystemEvaluator.cpp
+++ b/Engine/Plugins/Cameras/GameplayCameras/Source/GameplayCameras/Private/Core/CameraSystemEvaluator.cpp
@@ -347,13 +347,11 @@ void FCameraSystemEvaluator::GetEvaluatedCameraView(FMinimalViewInfo& DesiredVie
 	// TODO: add support for ortho cameras.
 	DesiredView.PerspectiveNearClipPlane = CameraPose.GetNearClippingPlane();
 
-	for (const FPostProcessSettingsCollectionEntry& PostProcessEntry : RootNodeResult.PostProcessSettings.GetEntries())
-	{
-		DesiredView.PostProcessSettings = PostProcessEntry.PostProcessSettings;
-		DesiredView.PostProcessBlendWeight = PostProcessEntry.PostProcessBlendWeight;
-		// TODO: support multiple post-process settings by getting the PlayerCameraManager
-		break;
-	}
+	const FPostProcessSettingsCollection& PostProcessSettings = RootNodeResult.PostProcessSettings;
+	DesiredView.PostProcessSettings = PostProcessSettings.Get();
+	DesiredView.PostProcessBlendWeight = 1.f;
+	// Create the physical camera settings if needed. Don't overwrite settings that were set by hand.
+	CameraPose.ApplyPhysicalCameraSettings(DesiredView.PostProcessSettings, false);
 }
 
 #if UE_GAMEPLAY_CAMERAS_DEBUG
diff --git a/Engine/Plugins/Cameras/GameplayCameras/Source/GameplayCameras/Private/Core/PostProcessSettingsCollection.cpp b/Engine/Plugins/Cameras/GameplayCameras/Source/GameplayCameras/Private/Core/PostProcessSettingsCollection.cpp
--- a/Engine/Plugins/Cameras/GameplayCameras/Source/GameplayCameras/Private/Core/PostProcessSettingsCollection.cpp
+++ b/Engine/Plugins/Cameras/GameplayCameras/Source/GameplayCameras/Private/Core/PostProcessSettingsCollection.cpp
@@ -5,57 +5,604 @@
 namespace UE::Cameras
 {
 
-FArchive& operator<< (FArchive& Ar, FPostProcessSettingsCollectionEntry& Entry)
+const FPostProcessSettings FPostProcessSettingsCollection::DefaultPostProcessSettings;
+
+void FPostProcessSettingsCollection::Reset()
 {
-	static FPostProcessSettings DefaultPostProcessSettings;
+	PostProcessSettings = DefaultPostProcessSettings;
+	bHasAnySetting = false;
+}
 
-	UScriptStruct* PostProcessSettingsStruct = FPostProcessSettings::StaticStruct();
-	PostProcessSettingsStruct->SerializeItem(Ar, &Entry.PostProcessSettings, &DefaultPostProcessSettings);
+void FPostProcessSettingsCollection::OverrideAll(const FPostProcessSettingsCollection& OtherCollection)
+{
+	PostProcessSettings = OtherCollection.PostProcessSettings;
+	bHasAnySetting = OtherCollection.bHasAnySetting;
+}
+
+void FPostProcessSettingsCollection::OverrideChanged(const FPostProcessSettingsCollection& OtherCollection)
+{
+	OverrideChanged(OtherCollection.PostProcessSettings);
+}
+
+void FPostProcessSettingsCollection::OverrideChanged(const FPostProcessSettings& OtherPostProcessSettings)
+{
+	FPostProcessSettings& ThisPP = PostProcessSettings;
+	const FPostProcessSettings& OtherPP = OtherPostProcessSettings;
+
+#define UE_SET_PP(Name)\
+	if (OtherPP.bOverride_##Name)\
+	{\
+		ThisPP.bOverride_##Name = true;\
+		ThisPP.Name = OtherPP.Name;\
+		bHasAnySetting = true;\
+	}
+
+	{
+		UE_SET_PP(TemperatureType);
+		UE_SET_PP(WhiteTemp);
+		UE_SET_PP(WhiteTint);
+
+		UE_SET_PP(ColorSaturation);
+		UE_SET_PP(ColorContrast);
+		UE_SET_PP(ColorGamma);
+		UE_SET_PP(ColorGain);
+		UE_SET_PP(ColorOffset);
+
+		UE_SET_PP(ColorSaturationShadows);
+		UE_SET_PP(ColorContrastShadows);
+		UE_SET_PP(ColorGammaShadows);
+		UE_SET_PP(ColorGainShadows);
+		UE_SET_PP(ColorOffsetShadows);
+
+		UE_SET_PP(ColorSaturationMidtones);
+		UE_SET_PP(ColorContrastMidtones);
+		UE_SET_PP(ColorGammaMidtones);
+		UE_SET_PP(ColorGainMidtones);
+		UE_SET_PP(ColorOffsetMidtones);
+
+		UE_SET_PP(ColorSaturationHighlights);
+		UE_SET_PP(ColorContrastHighlights);
+		UE_SET_PP(ColorGammaHighlights);
+		UE_SET_PP(ColorGainHighlights);
+		UE_SET_PP(ColorOffsetHighlights);
+
+		UE_SET_PP(ColorCorrectionShadowsMax);
+		UE_SET_PP(ColorCorrectionHighlightsMin);
+		UE_SET_PP(ColorCorrectionHighlightsMax);
+
+		UE_SET_PP(BlueCorrection);
+		UE_SET_PP(ExpandGamut);
+		UE_SET_PP(ToneCurveAmount);
+
+		UE_SET_PP(FilmSlope);
+		UE_SET_PP(FilmToe);
+		UE_SET_PP(FilmShoulder);
+		UE_SET_PP(FilmBlackClip);
+		UE_SET_PP(FilmWhiteClip);
+
+		UE_SET_PP(SceneColorTint);
+		UE_SET_PP(SceneFringeIntensity);
+		UE_SET_PP(ChromaticAberrationStartOffset);
+		UE_SET_PP(BloomIntensity);
+		UE_SET_PP(BloomThreshold);
+		UE_SET_PP(Bloom1Tint);
+		UE_SET_PP(BloomSizeScale);
+		UE_SET_PP(Bloom1Size);
+		UE_SET_PP(Bloom2Tint);
+		UE_SET_PP(Bloom2Size);
+		UE_SET_PP(Bloom3Tint);
+		UE_SET_PP(Bloom3Size);
+		UE_SET_PP(Bloom4Tint);
+		UE_SET_PP(Bloom4Size);
+		UE_SET_PP(Bloom5Tint);
+		UE_SET_PP(Bloom5Size);
+		UE_SET_PP(Bloom6Tint);
+		UE_SET_PP(Bloom6Size);
+		UE_SET_PP(BloomDirtMaskIntensity);
+		UE_SET_PP(BloomDirtMaskTint);
+		UE_SET_PP(BloomConvolutionScatterDispersion);
+		UE_SET_PP(BloomConvolutionSize);
+		UE_SET_PP(BloomConvolutionCenterUV);
+		UE_SET_PP(BloomConvolutionPreFilterMin);
+		UE_SET_PP(BloomConvolutionPreFilterMax);
+		UE_SET_PP(BloomConvolutionPreFilterMult);
+		UE_SET_PP(AmbientCubemapIntensity);
+		UE_SET_PP(AmbientCubemapTint);
+		UE_SET_PP(CameraShutterSpeed);
+		UE_SET_PP(CameraISO);
+		UE_SET_PP(AutoExposureLowPercent);
+		UE_SET_PP(AutoExposureHighPercent);
+		UE_SET_PP(AutoExposureMinBrightness);
+		UE_SET_PP(AutoExposureMaxBrightness);
+		UE_SET_PP(AutoExposureSpeedUp);
+		UE_SET_PP(AutoExposureSpeedDown);
+		UE_SET_PP(AutoExposureBias);
+		UE_SET_PP(HistogramLogMin);
+		UE_SET_PP(HistogramLogMax);
+		UE_SET_PP(LocalExposureMethod);
+		UE_SET_PP(LocalExposureContrastScale_DEPRECATED);
+		UE_SET_PP(LocalExposureHighlightContrastScale);
+		UE_SET_PP(LocalExposureShadowContrastScale);
+		UE_SET_PP(LocalExposureHighlightThreshold);
+		UE_SET_PP(LocalExposureShadowThreshold);
+		UE_SET_PP(LocalExposureDetailStrength);
+		UE_SET_PP(LocalExposureBlurredLuminanceBlend);
+		UE_SET_PP(LocalExposureBlurredLuminanceKernelSizePercent);
+		UE_SET_PP(LocalExposureMiddleGreyBias);
+		UE_SET_PP(LensFlareIntensity);
+		UE_SET_PP(LensFlareTint);
+		UE_SET_PP(LensFlareBokehSize);
+		UE_SET_PP(LensFlareThreshold);
+		UE_SET_PP(VignetteIntensity);
+		UE_SET_PP(Sharpen);
+		UE_SET_PP(FilmGrainIntensity);
+		UE_SET_PP(FilmGrainIntensityShadows);
+		UE_SET_PP(FilmGrainIntensityMidtones);
+		UE_SET_PP(FilmGrainIntensityHighlights);
+		UE_SET_PP(FilmGrainShadowsMax);
+		UE_SET_PP(FilmGrainHighlightsMin);
+		UE_SET_PP(FilmGrainHighlightsMax);
+		UE_SET_PP(FilmGrainTexelSize);
+		UE_SET_PP(AmbientOcclusionIntensity);
+		UE_SET_PP(AmbientOcclusionStaticFraction);
+		UE_SET_PP(AmbientOcclusionRadius);
+		UE_SET_PP(AmbientOcclusionFadeDistance);
+		UE_SET_PP(AmbientOcclusionFadeRadius);
+		UE_SET_PP(AmbientOcclusionDistance_DEPRECATED);
+		UE_SET_PP(AmbientOcclusionPower);
+		UE_SET_PP(AmbientOcclusionBias);
+		UE_SET_PP(AmbientOcclusionQuality);
+		UE_SET_PP(AmbientOcclusionMipBlend);
+		UE_SET_PP(AmbientOcclusionMipScale);
+		UE_SET_PP(AmbientOcclusionMipThreshold);
+		UE_SET_PP(AmbientOcclusionTemporalBlendWeight);
+		UE_SET_PP(IndirectLightingColor);
+		UE_SET_PP(IndirectLightingIntensity);
+
+		UE_SET_PP(DepthOfFieldFocalDistance);
 
-	Ar << Entry.PostProcessBlendWeight;
+		UE_SET_PP(DepthOfFieldFstop);
+		UE_SET_PP(DepthOfFieldMinFstop);
+		UE_SET_PP(DepthOfFieldSensorWidth);
+		UE_SET_PP(DepthOfFieldSqueezeFactor);
+		UE_SET_PP(DepthOfFieldDepthBlurRadius);
+		UE_SET_PP(DepthOfFieldUseHairDepth)
+		UE_SET_PP(DepthOfFieldDepthBlurAmount);
+		UE_SET_PP(DepthOfFieldFocalRegion);
+		UE_SET_PP(DepthOfFieldNearTransitionRegion);
+		UE_SET_PP(DepthOfFieldFarTransitionRegion);
+		UE_SET_PP(DepthOfFieldScale);
+		UE_SET_PP(DepthOfFieldNearBlurSize);
+		UE_SET_PP(DepthOfFieldFarBlurSize);
+		UE_SET_PP(DepthOfFieldOcclusion);
+		UE_SET_PP(DepthOfFieldSkyFocusDistance);
+		UE_SET_PP(DepthOfFieldVignetteSize);
+		UE_SET_PP(MotionBlurAmount);
+		UE_SET_PP(MotionBlurMax);
+		UE_SET_PP(MotionBlurPerObjectSize);
+		UE_SET_PP(ScreenSpaceReflectionQuality);
+		UE_SET_PP(ScreenSpaceReflectionIntensity);
+		UE_SET_PP(ScreenSpaceReflectionMaxRoughness);
 
-	return Ar;
+		UE_SET_PP(TranslucencyType);
+		UE_SET_PP(RayTracingTranslucencyMaxRoughness);
+		UE_SET_PP(RayTracingTranslucencyRefractionRays);
+		UE_SET_PP(RayTracingTranslucencySamplesPerPixel);
+		UE_SET_PP(RayTracingTranslucencyShadows);
+		UE_SET_PP(RayTracingTranslucencyRefraction);
+
+		UE_SET_PP(DynamicGlobalIlluminationMethod);
+		UE_SET_PP(LumenSurfaceCacheResolution);
+		UE_SET_PP(LumenSceneLightingQuality);
+		UE_SET_PP(LumenSceneDetail);
+		UE_SET_PP(LumenSceneViewDistance);
+		UE_SET_PP(LumenSceneLightingUpdateSpeed);
+		UE_SET_PP(LumenFinalGatherQuality);
+		UE_SET_PP(LumenFinalGatherLightingUpdateSpeed);
+		UE_SET_PP(LumenFinalGatherScreenTraces);
+		UE_SET_PP(LumenMaxTraceDistance);
+
+		UE_SET_PP(LumenDiffuseColorBoost);
+		UE_SET_PP(LumenSkylightLeaking);
+		UE_SET_PP(LumenFullSkylightLeakingDistance);
+
+		UE_SET_PP(LumenRayLightingMode);
+		UE_SET_PP(LumenReflectionsScreenTraces);
+		UE_SET_PP(LumenFrontLayerTranslucencyReflections);
+		UE_SET_PP(LumenMaxRoughnessToTraceReflections);
+		UE_SET_PP(LumenMaxReflectionBounces);
+		UE_SET_PP(LumenMaxRefractionBounces);
+		UE_SET_PP(ReflectionMethod);
+		UE_SET_PP(LumenReflectionQuality);
+		UE_SET_PP(RayTracingAO);
+		UE_SET_PP(RayTracingAOSamplesPerPixel);
+		UE_SET_PP(RayTracingAOIntensity);
+		UE_SET_PP(RayTracingAORadius);
+
+		UE_SET_PP(PathTracingMaxBounces);
+		UE_SET_PP(PathTracingSamplesPerPixel);
+		UE_SET_PP(PathTracingMaxPathIntensity);
+		UE_SET_PP(PathTracingEnableEmissiveMaterials);
+		UE_SET_PP(PathTracingEnableReferenceDOF);
+		UE_SET_PP(PathTracingEnableReferenceAtmosphere);
+		UE_SET_PP(PathTracingEnableDenoiser);
+		UE_SET_PP(PathTracingIncludeEmissive);
+		UE_SET_PP(PathTracingIncludeDiffuse);
+		UE_SET_PP(PathTracingIncludeIndirectDiffuse);
+		UE_SET_PP(PathTracingIncludeSpecular);
+		UE_SET_PP(PathTracingIncludeIndirectSpecular);
+		UE_SET_PP(PathTracingIncludeVolume);
+		UE_SET_PP(PathTracingIncludeIndirectVolume);
+
+		UE_SET_PP(DepthOfFieldBladeCount);
+
+		// This is no bOverride_AmbientCubemap so just see if it is set.
+		if (OtherPP.AmbientCubemap)
+		{
+			ThisPP.AmbientCubemap = OtherPP.AmbientCubemap;
+		}
+
+		UE_SET_PP(ColorGradingIntensity);
+		UE_SET_PP(ColorGradingLUT);
+
+		UE_SET_PP(BloomDirtMask);
+		UE_SET_PP(BloomMethod);
+		UE_SET_PP(BloomConvolutionTexture)
+		UE_SET_PP(FilmGrainTexture)
+
+		UE_SET_PP(BloomConvolutionBufferScale);
+
+		UE_SET_PP(AutoExposureBiasCurve);
+		UE_SET_PP(AutoExposureMeterMask);
+		UE_SET_PP(LocalExposureHighlightContrastCurve);
+		UE_SET_PP(LocalExposureShadowContrastCurve);
+		UE_SET_PP(LensFlareBokehShape);
+
+		if (OtherPP.bOverride_LensFlareTints)
+		{
+			for (uint32 i = 0; i < 8; ++i)
+			{
+				ThisPP.LensFlareTints[i] = OtherPP.LensFlareTints[i];
+			}
+		}
+
+		if (OtherPP.bOverride_MobileHQGaussian)
+		{
+			ThisPP.bMobileHQGaussian = OtherPP.bMobileHQGaussian;
+		}
+
+		UE_SET_PP(AutoExposureMethod);
+		UE_SET_PP(AmbientOcclusionRadiusInWS);
+		UE_SET_PP(MotionBlurTargetFPS);
+		UE_SET_PP(AutoExposureApplyPhysicalCameraExposure);
+		UE_SET_PP(UserFlags);
+	}
+
+#undef UE_SET_PP
 }
 
-void FPostProcessSettingsCollection::Add(const FPostProcessSettings& InPostProcess, float InBlendWeight)
+void FPostProcessSettingsCollection::LerpAll(const FPostProcessSettingsCollection& ToCollection, float BlendFactor)
 {
-	Entries.Add(FEntry{ InPostProcess, InBlendWeight });
+	LerpAll(ToCollection.PostProcessSettings, BlendFactor);
 }
 
-void FPostProcessSettingsCollection::Reset()
+void FPostProcessSettingsCollection::LerpAll(const FPostProcessSettings& ToPostProcessSettings, float BlendFactor)
 {
-	Entries.Reset();
+	InternalLerpChanged(ToPostProcessSettings, BlendFactor, false);
 }
 
-void FPostProcessSettingsCollection::OverrideAll(const FPostProcessSettingsCollection& OtherCollection)
+void FPostProcessSettingsCollection::LerpChanged(const FPostProcessSettingsCollection& ToCollection, float BlendFactor)
 {
-	Entries = OtherCollection.Entries;
+	LerpChanged(ToCollection.PostProcessSettings, BlendFactor);
 }
 
-void FPostProcessSettingsCollection::LerpAll(const FPostProcessSettingsCollection& ToCollection, float BlendFactor)
+void FPostProcessSettingsCollection::LerpChanged(const FPostProcessSettings& ToPostProcessSettings, float BlendFactor)
+{
+	InternalLerpChanged(ToPostProcessSettings, BlendFactor, true);
+}
+
+void FPostProcessSettingsCollection::InternalLerpChanged(const FPostProcessSettings& ToPostProcessSettings, float BlendFactor, bool bChangedOnly)
 {
-	// Blend down our entries. Remove any entries with zero weight.
-	for (auto It = Entries.CreateIterator(); It; ++It)
+	if (BlendFactor <= 0.f)
 	{
-		FEntry& ThisEntry = *It;
-		ThisEntry.PostProcessBlendWeight *= (1.f - BlendFactor);
-		
-		if (ThisEntry.PostProcessBlendWeight < UE_SMALL_NUMBER)
-		{
-			It.RemoveCurrent();
-		}
+		return;
+	}
+
+	BlendFactor = FMath::Clamp(BlendFactor, 0.f, 1.f);
+	const bool bShouldFlip = (BlendFactor >= 0.5f);
+
+	FPostProcessSettings& ThisPP = PostProcessSettings;
+	const FPostProcessSettings& ToPP = ToPostProcessSettings;
+
+	// We need to a poor man's version of FSceneView::OverridePostProcessSettings... differences include:
+	//
+	// 1) Flipping non-interpolatable properties at 50% blend, instead of always overwriting them.
+	// 2) A few things not being supported, such as accumulating ambient cubmaps.
+	// 3) No support for blendable objects.
+	// 4) Ability to blend _away_ from the current values, towards possibly default (non-overriden) values.
+	//
+	// TODO: refactor post-process settings to make blending them possible by default.
+
+#define UE_SET_PP(Name)\
+	if ((ToPP.bOverride_##Name || (!bChangedOnly && ThisPP.bOverride_##Name)) && bShouldFlip)\
+	{\
+		ThisPP.bOverride_##Name = true;\
+		ThisPP.Name = ToPP.Name;\
+		bHasAnySetting = true;\
+	}
+
+#define UE_LERP_PP(Name)\
+	if (ToPP.bOverride_##Name || (!bChangedOnly && ThisPP.bOverride_##Name))\
+	{\
+		ThisPP.bOverride_##Name = true;\
+		ThisPP.Name = FMath::Lerp(ThisPP.Name, ToPP.Name, BlendFactor);\
+		bHasAnySetting = true;\
 	}
 
-	// Blend up the other collection's entries.
-	for (const FEntry& OtherEntry : ToCollection.Entries)
 	{
-		Entries.Add(FEntry{ OtherEntry.PostProcessSettings, OtherEntry.PostProcessBlendWeight * BlendFactor });
+		UE_SET_PP(TemperatureType);
+		UE_LERP_PP(WhiteTemp);
+		UE_LERP_PP(WhiteTint);
+
+		UE_LERP_PP(ColorSaturation);
+		UE_LERP_PP(ColorContrast);
+		UE_LERP_PP(ColorGamma);
+		UE_LERP_PP(ColorGain);
+		UE_LERP_PP(ColorOffset);
+
+		UE_LERP_PP(ColorSaturationShadows);
+		UE_LERP_PP(ColorContrastShadows);
+		UE_LERP_PP(ColorGammaShadows);
+		UE_LERP_PP(ColorGainShadows);
+		UE_LERP_PP(ColorOffsetShadows);
+
+		UE_LERP_PP(ColorSaturationMidtones);
+		UE_LERP_PP(ColorContrastMidtones);
+		UE_LERP_PP(ColorGammaMidtones);
+		UE_LERP_PP(ColorGainMidtones);
+		UE_LERP_PP(ColorOffsetMidtones);
+
+		UE_LERP_PP(ColorSaturationHighlights);
+		UE_LERP_PP(ColorContrastHighlights);
+		UE_LERP_PP(ColorGammaHighlights);
+		UE_LERP_PP(ColorGainHighlights);
+		UE_LERP_PP(ColorOffsetHighlights);
+
+		UE_LERP_PP(ColorCorrectionShadowsMax);
+		UE_LERP_PP(ColorCorrectionHighlightsMin);
+		UE_LERP_PP(ColorCorrectionHighlightsMax);
+
+		UE_LERP_PP(BlueCorrection);
+		UE_LERP_PP(ExpandGamut);
+		UE_LERP_PP(ToneCurveAmount);
+
+		UE_LERP_PP(FilmSlope);
+		UE_LERP_PP(FilmToe);
+		UE_LERP_PP(FilmShoulder);
+		UE_LERP_PP(FilmBlackClip);
+		UE_LERP_PP(FilmWhiteClip);
+
+		UE_LERP_PP(SceneColorTint);
+		UE_LERP_PP(SceneFringeIntensity);
+		UE_LERP_PP(ChromaticAberrationStartOffset);
+		UE_LERP_PP(BloomIntensity);
+		UE_LERP_PP(BloomThreshold);
+		UE_LERP_PP(Bloom1Tint);
+		UE_LERP_PP(BloomSizeScale);
+		UE_LERP_PP(Bloom1Size);
+		UE_LERP_PP(Bloom2Tint);
+		UE_LERP_PP(Bloom2Size);
+		UE_LERP_PP(Bloom3Tint);
+		UE_LERP_PP(Bloom3Size);
+		UE_LERP_PP(Bloom4Tint);
+		UE_LERP_PP(Bloom4Size);
+		UE_LERP_PP(Bloom5Tint);
+		UE_LERP_PP(Bloom5Size);
+		UE_LERP_PP(Bloom6Tint);
+		UE_LERP_PP(Bloom6Size);
+		UE_LERP_PP(BloomDirtMaskIntensity);
+		UE_LERP_PP(BloomDirtMaskTint);
+		UE_LERP_PP(BloomConvolutionScatterDispersion);
+		UE_LERP_PP(BloomConvolutionSize);
+		UE_LERP_PP(BloomConvolutionCenterUV);
+		UE_LERP_PP(BloomConvolutionPreFilterMin);
+		UE_LERP_PP(BloomConvolutionPreFilterMax);
+		UE_LERP_PP(BloomConvolutionPreFilterMult);
+		UE_LERP_PP(AmbientCubemapIntensity);
+		UE_LERP_PP(AmbientCubemapTint);
+		UE_LERP_PP(CameraShutterSpeed);
+		UE_LERP_PP(CameraISO);
+		UE_LERP_PP(AutoExposureLowPercent);
+		UE_LERP_PP(AutoExposureHighPercent);
+		UE_LERP_PP(AutoExposureMinBrightness);
+		UE_LERP_PP(AutoExposureMaxBrightness);
+		UE_LERP_PP(AutoExposureSpeedUp);
+		UE_LERP_PP(AutoExposureSpeedDown);
+		UE_LERP_PP(AutoExposureBias);
+		UE_LERP_PP(HistogramLogMin);
+		UE_LERP_PP(HistogramLogMax);
+		UE_SET_PP(LocalExposureMethod);
+		UE_LERP_PP(LocalExposureContrastScale_DEPRECATED);
+		UE_LERP_PP(LocalExposureHighlightContrastScale);
+		UE_LERP_PP(LocalExposureShadowContrastScale);
+		UE_LERP_PP(LocalExposureHighlightThreshold);
+		UE_LERP_PP(LocalExposureShadowThreshold);
+		UE_LERP_PP(LocalExposureDetailStrength);
+		UE_LERP_PP(LocalExposureBlurredLuminanceBlend);
+		UE_LERP_PP(LocalExposureBlurredLuminanceKernelSizePercent);
+		UE_LERP_PP(LocalExposureMiddleGreyBias);
+		UE_LERP_PP(LensFlareIntensity);
+		UE_LERP_PP(LensFlareTint);
+		UE_LERP_PP(LensFlareBokehSize);
+		UE_LERP_PP(LensFlareThreshold);
+		UE_LERP_PP(VignetteIntensity);
+		UE_LERP_PP(Sharpen);
+		UE_LERP_PP(FilmGrainIntensity);
+		UE_LERP_PP(FilmGrainIntensityShadows);
+		UE_LERP_PP(FilmGrainIntensityMidtones);
+		UE_LERP_PP(FilmGrainIntensityHighlights);
+		UE_LERP_PP(FilmGrainShadowsMax);
+		UE_LERP_PP(FilmGrainHighlightsMin);
+		UE_LERP_PP(FilmGrainHighlightsMax);
+		UE_LERP_PP(FilmGrainTexelSize);
+		UE_LERP_PP(AmbientOcclusionIntensity);
+		UE_LERP_PP(AmbientOcclusionStaticFraction);
+		UE_LERP_PP(AmbientOcclusionRadius);
+		UE_LERP_PP(AmbientOcclusionFadeDistance);
+		UE_LERP_PP(AmbientOcclusionFadeRadius);
+		UE_LERP_PP(AmbientOcclusionDistance_DEPRECATED);
+		UE_LERP_PP(AmbientOcclusionPower);
+		UE_LERP_PP(AmbientOcclusionBias);
+		UE_LERP_PP(AmbientOcclusionQuality);
+		UE_LERP_PP(AmbientOcclusionMipBlend);
+		UE_LERP_PP(AmbientOcclusionMipScale);
+		UE_LERP_PP(AmbientOcclusionMipThreshold);
+		UE_LERP_PP(AmbientOcclusionTemporalBlendWeight);
+		UE_LERP_PP(IndirectLightingColor);
+		UE_LERP_PP(IndirectLightingIntensity);
+
+		if (ToPP.bOverride_DepthOfFieldFocalDistance)
+		{
+			if (ThisPP.DepthOfFieldFocalDistance == 0.0f || ToPP.DepthOfFieldFocalDistance == 0.0f)
+			{
+				ThisPP.DepthOfFieldFocalDistance = ToPP.DepthOfFieldFocalDistance;
+			}
+			else
+			{
+				ThisPP.DepthOfFieldFocalDistance = FMath::Lerp(ThisPP.DepthOfFieldFocalDistance, ToPP.DepthOfFieldFocalDistance, BlendFactor);
+			}
+		}
+		UE_LERP_PP(DepthOfFieldFstop);
+		UE_LERP_PP(DepthOfFieldMinFstop);
+		UE_LERP_PP(DepthOfFieldSensorWidth);
+		UE_LERP_PP(DepthOfFieldSqueezeFactor);
+		UE_LERP_PP(DepthOfFieldDepthBlurRadius);
+		UE_SET_PP(DepthOfFieldUseHairDepth)
+		UE_LERP_PP(DepthOfFieldDepthBlurAmount);
+		UE_LERP_PP(DepthOfFieldFocalRegion);
+		UE_LERP_PP(DepthOfFieldNearTransitionRegion);
+		UE_LERP_PP(DepthOfFieldFarTransitionRegion);
+		UE_LERP_PP(DepthOfFieldScale);
+		UE_LERP_PP(DepthOfFieldNearBlurSize);
+		UE_LERP_PP(DepthOfFieldFarBlurSize);
+		UE_LERP_PP(DepthOfFieldOcclusion);
+		UE_LERP_PP(DepthOfFieldSkyFocusDistance);
+		UE_LERP_PP(DepthOfFieldVignetteSize);
+		UE_LERP_PP(MotionBlurAmount);
+		UE_LERP_PP(MotionBlurMax);
+		UE_LERP_PP(MotionBlurPerObjectSize);
+		UE_LERP_PP(ScreenSpaceReflectionQuality);
+		UE_LERP_PP(ScreenSpaceReflectionIntensity);
+		UE_LERP_PP(ScreenSpaceReflectionMaxRoughness);
+
+		UE_SET_PP(TranslucencyType);
+		UE_SET_PP(RayTracingTranslucencyMaxRoughness);
+		UE_SET_PP(RayTracingTranslucencyRefractionRays);
+		UE_SET_PP(RayTracingTranslucencySamplesPerPixel);
+		UE_SET_PP(RayTracingTranslucencyShadows);
+		UE_SET_PP(RayTracingTranslucencyRefraction);
+
+		UE_SET_PP(DynamicGlobalIlluminationMethod);
+		UE_SET_PP(LumenSurfaceCacheResolution);
+		UE_SET_PP(LumenSceneLightingQuality);
+		UE_SET_PP(LumenSceneDetail);
+		UE_SET_PP(LumenSceneViewDistance);
+		UE_SET_PP(LumenSceneLightingUpdateSpeed);
+		UE_SET_PP(LumenFinalGatherQuality);
+		UE_SET_PP(LumenFinalGatherLightingUpdateSpeed);
+		UE_SET_PP(LumenFinalGatherScreenTraces);
+		UE_SET_PP(LumenMaxTraceDistance);
+
+		UE_LERP_PP(LumenDiffuseColorBoost);
+		UE_LERP_PP(LumenSkylightLeaking);
+		UE_LERP_PP(LumenFullSkylightLeakingDistance);
+
+		UE_SET_PP(LumenRayLightingMode);
+		UE_SET_PP(LumenReflectionsScreenTraces);
+		UE_SET_PP(LumenFrontLayerTranslucencyReflections);
+		UE_SET_PP(LumenMaxRoughnessToTraceReflections);
+		UE_SET_PP(LumenMaxReflectionBounces);
+		UE_SET_PP(LumenMaxRefractionBounces);
+		UE_SET_PP(ReflectionMethod);
+		UE_SET_PP(LumenReflectionQuality);
+		UE_SET_PP(RayTracingAO);
+		UE_SET_PP(RayTracingAOSamplesPerPixel);
+		UE_SET_PP(RayTracingAOIntensity);
+		UE_SET_PP(RayTracingAORadius);
+
+		UE_SET_PP(PathTracingMaxBounces);
+		UE_SET_PP(PathTracingSamplesPerPixel);
+		UE_LERP_PP(PathTracingMaxPathIntensity);
+		UE_SET_PP(PathTracingEnableEmissiveMaterials);
+		UE_SET_PP(PathTracingEnableReferenceDOF);
+		UE_SET_PP(PathTracingEnableReferenceAtmosphere);
+		UE_SET_PP(PathTracingEnableDenoiser);
+		UE_SET_PP(PathTracingIncludeEmissive);
+		UE_SET_PP(PathTracingIncludeDiffuse);
+		UE_SET_PP(PathTracingIncludeIndirectDiffuse);
+		UE_SET_PP(PathTracingIncludeSpecular);
+		UE_SET_PP(PathTracingIncludeIndirectSpecular);
+		UE_SET_PP(PathTracingIncludeVolume);
+		UE_SET_PP(PathTracingIncludeIndirectVolume);
+
+		UE_SET_PP(DepthOfFieldBladeCount);
+
+		// No cubemap blending (only supported for FFinalPostProcessSettings)
+		if (bShouldFlip)
+		{
+			ThisPP.AmbientCubemap = ToPP.AmbientCubemap;
+		}
+
+		// No color grading texture blending (only supported for FFinalPostProcessSettings)
+		UE_LERP_PP(ColorGradingIntensity);
+		UE_SET_PP(ColorGradingLUT);
+
+		UE_SET_PP(BloomDirtMask);
+		UE_SET_PP(BloomMethod);
+		UE_SET_PP(BloomConvolutionTexture)
+		UE_SET_PP(FilmGrainTexture)
+
+		// Flipping this instead of blending, as per the comment in SceneView.cpp
+		UE_SET_PP(BloomConvolutionBufferScale);
+
+		UE_SET_PP(AutoExposureBiasCurve);
+		UE_SET_PP(AutoExposureMeterMask);
+		UE_SET_PP(LocalExposureHighlightContrastCurve);
+		UE_SET_PP(LocalExposureShadowContrastCurve);
+		UE_SET_PP(LensFlareBokehShape);
+
+		if (ToPP.bOverride_LensFlareTints)
+		{
+			for (uint32 i = 0; i < 8; ++i)
+			{
+				ThisPP.LensFlareTints[i] = FMath::Lerp(ThisPP.LensFlareTints[i], ToPP.LensFlareTints[i], BlendFactor);
+			}
+		}
+
+		if (ToPP.bOverride_MobileHQGaussian && bShouldFlip)
+		{
+			ThisPP.bMobileHQGaussian = ToPP.bMobileHQGaussian;
+		}
+
+		UE_SET_PP(AutoExposureMethod);
+		UE_SET_PP(AmbientOcclusionRadiusInWS);
+		UE_SET_PP(MotionBlurTargetFPS);
+		UE_SET_PP(AutoExposureApplyPhysicalCameraExposure);
+		UE_SET_PP(UserFlags);
 	}
+
+	// No support for blendable objects for now.
+
+#undef UE_SET_PP
+#undef UE_LERP_PP
 }
 
 void FPostProcessSettingsCollection::Serialize(FArchive& Ar)
 {
-	Ar << Entries;
+	UScriptStruct* PostProcessSettingsStruct = FPostProcessSettings::StaticStruct();
+	PostProcessSettingsStruct->SerializeItem(Ar, &PostProcessSettings, &DefaultPostProcessSettings);
+
+	Ar << bHasAnySetting;
 }
 
 }  // namespace UE::Cameras
diff --git a/Engine/Plugins/Cameras/GameplayCameras/Source/GameplayCameras/Private/Debug/CameraNodeEvaluationResultDebugBlock.cpp b/Engine/Plugins/Cameras/GameplayCameras/Source/GameplayCameras/Private/Debug/CameraNodeEvaluationResultDebugBlock.cpp
--- a/Engine/Plugins/Cameras/GameplayCameras/Source/GameplayCameras/Private/Debug/CameraNodeEvaluationResultDebugBlock.cpp
+++ b/Engine/Plugins/Cameras/GameplayCameras/Source/GameplayCameras/Private/Debug/CameraNodeEvaluationResultDebugBlock.cpp
@@ -26,7 +26,7 @@ FCameraNodeEvaluationResultDebugBlock::FCameraNodeEvaluationResultDebugBlock()
 
 void FCameraNodeEvaluationResultDebugBlock::Initialize(const FCameraNodeEvaluationResult& InResult, FCameraDebugBlockBuilder& Builder)
 {
-	NumPostProcessSettings = InResult.PostProcessSettings.GetEntries().Num();
+	bHasAnyPostProcessSettings = InResult.PostProcessSettings.HasAnyPostProcessSettings();
 	bIsCameraCut = InResult.bIsCameraCut;
 	bIsValid = InResult.bIsValid;
 
@@ -83,9 +83,9 @@ void FCameraNodeEvaluationResultDebugBlock::OnDebugDraw(const FCameraDebugBlockD
 		Renderer.AddText(TEXT("  {cam_warning}IsCameraCut"));
 	}
 
-	if (NumPostProcessSettings > 0)
+	if (bHasAnyPostProcessSettings)
 	{
-		Renderer.AddText(TEXT("  {cam_notice}%d post-FX"), NumPostProcessSettings);
+		Renderer.AddText(TEXT("  {cam_notice}post-FX set"));
 	}
 
 	Renderer.NewLine();
diff --git a/Engine/Plugins/Cameras/GameplayCameras/Source/GameplayCameras/Private/Nodes/Common/FilmbackCameraNode.cpp b/Engine/Plugins/Cameras/GameplayCameras/Source/GameplayCameras/Private/Nodes/Common/FilmbackCameraNode.cpp
--- a/Engine/Plugins/Cameras/GameplayCameras/Source/GameplayCameras/Private/Nodes/Common/FilmbackCameraNode.cpp
+++ b/Engine/Plugins/Cameras/GameplayCameras/Source/GameplayCameras/Private/Nodes/Common/FilmbackCameraNode.cpp
@@ -23,6 +23,7 @@ private:
 
 	TCameraParameterReader<float> SensorWidthReader;
 	TCameraParameterReader<float> SensorHeightReader;
+	TCameraParameterReader<float> ISOReader;
 	TCameraParameterReader<bool> ConstrainAspectRatioReader;
 	TCameraParameterReader<bool> OverrideAspectRatioAxisConstraintReader;
 };
@@ -34,6 +35,7 @@ void FFilmbackCameraNodeEvaluator::OnInitialize(const FCameraNodeEvaluatorInitia
 	const UFilmbackCameraNode* FilmbackNode = GetCameraNodeAs<UFilmbackCameraNode>();
 	SensorWidthReader.Initialize(FilmbackNode->SensorWidth);
 	SensorHeightReader.Initialize(FilmbackNode->SensorHeight);
+	ISOReader.Initialize(FilmbackNode->ISO);
 	ConstrainAspectRatioReader.Initialize(FilmbackNode->ConstrainAspectRatio);
 	OverrideAspectRatioAxisConstraintReader.Initialize(FilmbackNode->OverrideAspectRatioAxisConstraint);
 }
@@ -52,6 +54,11 @@ void FFilmbackCameraNodeEvaluator::OnRun(const FCameraNodeEvaluationParams& Para
 	{
 		OutPose.SetSensorHeight(SensorHeight);
 	}
+	float ISO = ISOReader.Get(OutResult.VariableTable);
+	if (ISO > 0)
+	{
+		OutPose.SetISO(ISO);
+	}
 
 	OutPose.SetConstrainAspectRatio(ConstrainAspectRatioReader.Get(OutResult.VariableTable));
 	OutPose.SetOverrideAspectRatioAxisConstraint(OverrideAspectRatioAxisConstraintReader.Get(OutResult.VariableTable));
diff --git a/Engine/Plugins/Cameras/GameplayCameras/Source/GameplayCameras/Private/Nodes/Common/LensParametersCameraNode.cpp b/Engine/Plugins/Cameras/GameplayCameras/Source/GameplayCameras/Private/Nodes/Common/LensParametersCameraNode.cpp
--- a/Engine/Plugins/Cameras/GameplayCameras/Source/GameplayCameras/Private/Nodes/Common/LensParametersCameraNode.cpp
+++ b/Engine/Plugins/Cameras/GameplayCameras/Source/GameplayCameras/Private/Nodes/Common/LensParametersCameraNode.cpp
@@ -22,9 +22,11 @@ protected:
 
 private:
 
-	TCameraParameterReader<float> FocusDistanceReader;
 	TCameraParameterReader<float> FocalLengthReader;
+	TCameraParameterReader<float> FocusDistanceReader;
 	TCameraParameterReader<float> ApertureReader;
+	TCameraParameterReader<float> ShutterSpeedReader;
+	TCameraParameterReader<bool> EnablePhysicalCameraReader;
 };
 
 UE_DEFINE_CAMERA_NODE_EVALUATOR(FLensParametersCameraNodeEvaluator)
@@ -32,31 +34,42 @@ UE_DEFINE_CAMERA_NODE_EVALUATOR(FLensParametersCameraNodeEvaluator)
 void FLensParametersCameraNodeEvaluator::OnInitialize(const FCameraNodeEvaluatorInitializeParams& Params, FCameraNodeEvaluationResult& OutResult)
 {
 	const ULensParametersCameraNode* LensParametersNode = GetCameraNodeAs<ULensParametersCameraNode>();
-	FocusDistanceReader.Initialize(LensParametersNode->FocusDistance);
 	FocalLengthReader.Initialize(LensParametersNode->FocalLength);
+	FocusDistanceReader.Initialize(LensParametersNode->FocusDistance);
 	ApertureReader.Initialize(LensParametersNode->Aperture);
+	ShutterSpeedReader.Initialize(LensParametersNode->ShutterSpeed);
+	EnablePhysicalCameraReader.Initialize(LensParametersNode->EnablePhysicalCamera);
 }
 
 void FLensParametersCameraNodeEvaluator::OnRun(const FCameraNodeEvaluationParams& Params, FCameraNodeEvaluationResult& OutResult)
 {
 	FCameraPose& OutPose = OutResult.CameraPose;
 
-	float FocusDistance = FocusDistanceReader.Get(OutResult.VariableTable);
-	if (FocusDistance > 0)
-	{
-		OutPose.SetFocusDistance(FocusDistance);
-	}
 	float FocalLength = FocalLengthReader.Get(OutResult.VariableTable);
 	if (FocalLength > 0)
 	{
 		OutPose.SetFocalLength(FocalLength);
 		OutPose.SetFieldOfView(-1);
 	}
+	float FocusDistance = FocusDistanceReader.Get(OutResult.VariableTable);
+	if (FocusDistance > 0)
+	{
+		OutPose.SetFocusDistance(FocusDistance);
+	}
 	float Aperture = ApertureReader.Get(OutResult.VariableTable);
 	if (Aperture > 0)
 	{
 		OutPose.SetAperture(Aperture);
 	}
+	float ShutterSpeed = ShutterSpeedReader.Get(OutResult.VariableTable);
+	if (ShutterSpeed > 0)
+	{
+		OutPose.SetShutterSpeed(ShutterSpeed);
+	}
+
+	const bool bEnablePhysicalCamera = EnablePhysicalCameraReader.Get(OutResult.VariableTable);
+	OutPose.SetEnablePhysicalCamera(bEnablePhysicalCamera);
+	OutPose.SetPhysicalCameraBlendWeight(bEnablePhysicalCamera ? 1.f : 0.f);
 }
 
 }  // namespace UE::Cameras
diff --git a/Engine/Plugins/Cameras/GameplayCameras/Source/GameplayCameras/Private/Nodes/Common/PostProcessCameraNode.cpp b/Engine/Plugins/Cameras/GameplayCameras/Source/GameplayCameras/Private/Nodes/Common/PostProcessCameraNode.cpp
--- a/Engine/Plugins/Cameras/GameplayCameras/Source/GameplayCameras/Private/Nodes/Common/PostProcessCameraNode.cpp
+++ b/Engine/Plugins/Cameras/GameplayCameras/Source/GameplayCameras/Private/Nodes/Common/PostProcessCameraNode.cpp
@@ -2,7 +2,6 @@
 
 #include "Nodes/Common/PostProcessCameraNode.h"
 
-#include "Core/CameraParameterReader.h"
 #include "Core/CameraPose.h"
 
 #include UE_INLINE_GENERATED_CPP_BY_NAME(PostProcessCameraNode)
@@ -16,32 +15,17 @@ class FPostProcessCameraNodeEvaluator : public FCameraNodeEvaluator
 
 protected:
 
-	virtual void OnInitialize(const FCameraNodeEvaluatorInitializeParams& Params, FCameraNodeEvaluationResult& OutResult) override;
 	virtual void OnRun(const FCameraNodeEvaluationParams& Params, FCameraNodeEvaluationResult& OutResult) override;
-
-private:
-
-	TCameraParameterReader<float> PostProcessBlendWeightReader;
 };
 
 UE_DEFINE_CAMERA_NODE_EVALUATOR(FPostProcessCameraNodeEvaluator)
 
-void FPostProcessCameraNodeEvaluator::OnInitialize(const FCameraNodeEvaluatorInitializeParams& Params, FCameraNodeEvaluationResult& OutResult)
-{
-	const UPostProcessCameraNode* PostProcessNode = GetCameraNodeAs<UPostProcessCameraNode>();
-	PostProcessBlendWeightReader.Initialize(PostProcessNode->PostProcessBlendWeight);
-}
-
 void FPostProcessCameraNodeEvaluator::OnRun(const FCameraNodeEvaluationParams& Params, FCameraNodeEvaluationResult& OutResult)
 {
 	FCameraPose& OutPose = OutResult.CameraPose;
 
-	float PostProcessBlendWeight = PostProcessBlendWeightReader.Get(OutResult.VariableTable);
-	if (PostProcessBlendWeight > 0)
-	{
-		const UPostProcessCameraNode* PostProcessNode = GetCameraNodeAs<UPostProcessCameraNode>();
-		OutResult.PostProcessSettings.Add(PostProcessNode->PostProcessSettings, PostProcessBlendWeight);
-	}
+	const UPostProcessCameraNode* PostProcessNode = GetCameraNodeAs<UPostProcessCameraNode>();
+	OutResult.PostProcessSettings.OverrideChanged(PostProcessNode->PostProcessSettings);
 }
 
 }  // namespace UE::Cameras
diff --git a/Engine/Plugins/Cameras/GameplayCameras/Source/GameplayCameras/Private/Nodes/Common/TargetRayCastCameraNode.cpp b/Engine/Plugins/Cameras/GameplayCameras/Source/GameplayCameras/Private/Nodes/Common/TargetRayCastCameraNode.cpp
--- a/Engine/Plugins/Cameras/GameplayCameras/Source/GameplayCameras/Private/Nodes/Common/TargetRayCastCameraNode.cpp
+++ b/Engine/Plugins/Cameras/GameplayCameras/Source/GameplayCameras/Private/Nodes/Common/TargetRayCastCameraNode.cpp
@@ -5,6 +5,7 @@
 #include "CollisionQueryParams.h"
 #include "Core/CameraEvaluationContext.h"
 #include "Core/CameraNodeEvaluator.h"
+#include "Core/CameraParameterReader.h"
 #include "Debug/CameraDebugBlock.h"
 #include "Debug/CameraDebugBlockBuilder.h"
 #include "Debug/CameraDebugRenderer.h"
@@ -32,6 +33,7 @@ class FTargetRayCastCameraNodeEvaluator : public FCameraNodeEvaluator
 protected:
 
 	// FCameraNodeEvaluator interface.
+	virtual void OnInitialize(const FCameraNodeEvaluatorInitializeParams& Params, FCameraNodeEvaluationResult& OutResult) override;
 	virtual void OnRun(const FCameraNodeEvaluationParams& Params, FCameraNodeEvaluationResult& OutResult) override;
 #if UE_GAMEPLAY_CAMERAS_DEBUG
 	virtual void OnBuildDebugBlocks(const FCameraDebugBlockBuildParams& Params, FCameraDebugBlockBuilder& Builder) override;
@@ -43,6 +45,8 @@ private:
 
 private:
 
+	TCameraParameterReader<bool> AutoFocusReader;
+
 	double LastHitResultDistance = 1000.0; // Same as the default CameraPose target distance.
 
 #if UE_GAMEPLAY_CAMERAS_DEBUG
@@ -63,6 +67,12 @@ UE_DECLARE_CAMERA_DEBUG_BLOCK_END()
 
 UE_DEFINE_CAMERA_DEBUG_BLOCK_WITH_FIELDS(FTargetRayCastCameraDebugBlock)
 
+void FTargetRayCastCameraNodeEvaluator::OnInitialize(const FCameraNodeEvaluatorInitializeParams& Params, FCameraNodeEvaluationResult& OutResult)
+{
+	const UTargetRayCastCameraNode* RayCastNode = GetCameraNodeAs<UTargetRayCastCameraNode>();
+	AutoFocusReader.Initialize(RayCastNode->AutoFocus);
+}
+
 void FTargetRayCastCameraNodeEvaluator::OnRun(const FCameraNodeEvaluationParams& Params, FCameraNodeEvaluationResult& OutResult)
 {
 	if (!ensure(Params.EvaluationContext))
@@ -133,6 +143,11 @@ void FTargetRayCastCameraNodeEvaluator::RunLineTrace(UWorld* World, APlayerContr
 	{
 		CameraPose.SetTargetDistance(GTargetRayCastLength);
 	}
+
+	if (AutoFocusReader.Get(OutResult.VariableTable))
+	{
+		CameraPose.SetFocusDistance(CameraPose.GetTargetDistance());
+	}
 }
 
 #if UE_GAMEPLAY_CAMERAS_DEBUG
diff --git a/Engine/Plugins/Cameras/GameplayCameras/Source/GameplayCameras/Public/Core/CameraPose.h b/Engine/Plugins/Cameras/GameplayCameras/Source/GameplayCameras/Public/Core/CameraPose.h
--- a/Engine/Plugins/Cameras/GameplayCameras/Source/GameplayCameras/Public/Core/CameraPose.h
+++ b/Engine/Plugins/Cameras/GameplayCameras/Source/GameplayCameras/Public/Core/CameraPose.h
@@ -10,6 +10,7 @@
 #include "CameraPose.generated.h"
 
 class FArchive;
+struct FPostProcessSettings;
 enum EAspectRatioAxisConstraint : int;
 
 #define UE_CAMERA_POSE_FOR_TRANSFORM_PROPERTIES()\
@@ -19,10 +20,14 @@ enum EAspectRatioAxisConstraint : int;
 #define UE_CAMERA_POSE_FOR_INTERPOLABLE_PROPERTIES()\
 	UE_CAMERA_POSE_FOR_PROPERTY(double, TargetDistance)\
 	UE_CAMERA_POSE_FOR_PROPERTY(float,  Aperture)\
+	UE_CAMERA_POSE_FOR_PROPERTY(float,  ShutterSpeed)\
 	UE_CAMERA_POSE_FOR_PROPERTY(float,  FocusDistance)\
 	UE_CAMERA_POSE_FOR_PROPERTY(float,  SensorWidth)\
 	UE_CAMERA_POSE_FOR_PROPERTY(float,  SensorHeight)\
+	UE_CAMERA_POSE_FOR_PROPERTY(float,  ISO)\
 	UE_CAMERA_POSE_FOR_PROPERTY(float,  SqueezeFactor)\
+	UE_CAMERA_POSE_FOR_PROPERTY(int32,  DiaphragmBladeCount)\
+	UE_CAMERA_POSE_FOR_PROPERTY(float,  PhysicalCameraBlendWeight)\
 	UE_CAMERA_POSE_FOR_PROPERTY(float,  NearClippingPlane)\
 	UE_CAMERA_POSE_FOR_PROPERTY(float,  FarClippingPlane)
 
@@ -31,6 +36,7 @@ enum EAspectRatioAxisConstraint : int;
 	UE_CAMERA_POSE_FOR_PROPERTY(float, FocalLength)
 
 #define UE_CAMERA_POSE_FOR_FLIPPING_PROPERTIES()\
+	UE_CAMERA_POSE_FOR_PROPERTY(bool, EnablePhysicalCamera)\
 	UE_CAMERA_POSE_FOR_PROPERTY(bool, ConstrainAspectRatio)\
 	UE_CAMERA_POSE_FOR_PROPERTY(bool, OverrideAspectRatioAxisConstraint)\
 	UE_CAMERA_POSE_FOR_PROPERTY(EAspectRatioAxisConstraint, AspectRatioAxisConstraint)
@@ -173,6 +179,19 @@ public:
 	/** Gets the default sensor size. */
 	static void GetDefaultSensorSize(float& OutSensorWidth, float& OutSensorHeight);
 
+	/**
+	 * Applies the necessary post-process settings given the current values
+	 * on this camera pose.
+	 *
+	 * This function doesn't do anything if EnablePhysicalCamera is false, or if
+	 * PhysicalCameraBlendWeight is zero or less.
+	 *
+	 * @param PostProcessSettings  The post-process settings to modify
+	 * @param bOverwriteSettings   Whether to overwrite values found to already be set
+	 * @return  Whether post-process settings were created.
+	 */
+	bool ApplyPhysicalCameraSettings(FPostProcessSettings& PostProcessSettings, bool bOverwriteSettings = false) const;
+
 public:
 
 	// Interpolation
@@ -233,6 +252,10 @@ private:
 	UPROPERTY()
 	float Aperture = 2.8f;
 
+	/** The shutter speed of the camera's lens, in 1/seconds */
+	UPROPERTY()
+	float ShutterSpeed = 60.f;
+
 	/** The focus distance of the camera's lens, in world units */
 	UPROPERTY()
 	float FocusDistance = -1.f;
@@ -245,10 +268,18 @@ private:
 	UPROPERTY()
 	float SensorHeight = 18.67f;
 
+	/** The camera sensor sensitivity in ISO. */
+	UPROPERTY()
+	float ISO = 100.f;
+
 	/** Squeeze factor for anamorphic lenses */
 	UPROPERTY()
 	float SqueezeFactor = 1.f;
 
+	/** Number of blades in the lens diaphragm */
+	UPROPERTY()
+	int32 DiaphragmBladeCount = 8;
+
 	/** The distance to the near clipping plane, in world units */
 	UPROPERTY()
 	float NearClippingPlane = 10.f;
@@ -257,6 +288,20 @@ private:
 	UPROPERTY()
 	float FarClippingPlane = -1.f;
 
+	/** 
+	 * An internal weight for the physical camera post-process settings, used when blending between 
+	 * cameras with EnablePhysicalCamera enabled/disabled.
+	 */
+	UPROPERTY()
+	float PhysicalCameraBlendWeight = 0.f;
+
+	/** 
+	 * Whether to setup post-process settings based on physical camera properties such as Aperture,
+	 * FocusDistance, DiaphragmBladeCount, and so on.
+	 */
+	UPROPERTY()
+	bool EnablePhysicalCamera = false;
+
 	/** Whether to constrain aspect ratio */
 	UPROPERTY()
 	bool ConstrainAspectRatio = false;
diff --git a/Engine/Plugins/Cameras/GameplayCameras/Source/GameplayCameras/Public/Core/PostProcessSettingsCollection.h b/Engine/Plugins/Cameras/GameplayCameras/Source/GameplayCameras/Public/Core/PostProcessSettingsCollection.h
--- a/Engine/Plugins/Cameras/GameplayCameras/Source/GameplayCameras/Public/Core/PostProcessSettingsCollection.h
+++ b/Engine/Plugins/Cameras/GameplayCameras/Source/GameplayCameras/Public/Core/PostProcessSettingsCollection.h
@@ -7,69 +7,67 @@
 namespace UE::Cameras
 {
 
-/** An entry into an FPostProcessSettingsCollection. */
-struct FPostProcessSettingsCollectionEntry
-{
-	/** The post-process settings. */
-	FPostProcessSettings PostProcessSettings;
-	/** The blend factor for the post-process settings. */
-	float PostProcessBlendWeight = 1.f;
-};
-
-FArchive& operator<< (FArchive& Ar, FPostProcessSettingsCollectionEntry& Entry);
-
 /**
- * A collection of post-process settings, with associated blend weights.
- * TODO: the camera system currently supports only one post-process setting.
+ * A class that can collect post-process settings, combining them with their associated
+ * blend weights.
  */
 struct GAMEPLAYCAMERAS_API FPostProcessSettingsCollection
 {
-	/** Returns whether this collection has any entries. */
-	bool IsEmpty() const
-	{
-		return Entries.IsEmpty();
-	}
-
-	/** Gets the entries in this collection. */
-	TArrayView<const FPostProcessSettingsCollectionEntry> GetEntries() const
-	{
-		return Entries;
-	}
-
-	/** Gets the entries in this collection. */
-	TArrayView<FPostProcessSettingsCollectionEntry> GetEntries() 
-	{
-		return Entries;
-	}
-
-	/** Adds a new entry in the collection. */
-	void Add(const FPostProcessSettings& InPostProcess, float InBlendWeight);
-
-	/** Removes all entries in the collection. */
+	/** Gets the effective post-process settings. */
+	FPostProcessSettings& Get() { return PostProcessSettings; }
+
+	/** Gets the effective post-process settings. */
+	const FPostProcessSettings& Get() const { return PostProcessSettings; }
+
+	/** Returns whether this collection has any set post-process settings. */
+	bool HasAnyPostProcessSettings() const { return bHasAnySetting; }
+
+	/** Resets this collection to the default post-process settings. */
 	void Reset();
 
 	/** 
-	 * Removes all entries in the collection, replacing them with the entries in the
-	 * given other collection.
+	 * Overwrites the post-process settings in this collection with the values in the other.
 	 */
 	void OverrideAll(const FPostProcessSettingsCollection& OtherCollection);
 
+	/** 
+	 * Overwrites the post-process settings in this collection with any changed values in the other.
+	 * Changed values are those whose bOverride_Xxx flag is true. Functionally equivalent to
+	 * LerpAll with a blend factor of 100%.
+	 */
+	void OverrideChanged(const FPostProcessSettingsCollection& OtherCollection);
+	void OverrideChanged(const FPostProcessSettings& OtherPostProcessSettings);
+
 	/**
-	 * Interpolates the entries in this collection towards the entries in the given other
-	 * collection. This effectively merges the given collection into this collection, ramping
-	 * down the blend weight of existing entries, and ramping up the blend weight of the
-	 * new entries.
+	 * Interpolates the post-process settings towards the values in the given other collection. All
+	 * values are interpolated if either post-process settings have the bOverride_Xxx flag set.
+	 * This means that some values will interpolate to and/or from default values.
 	 */
 	void LerpAll(const FPostProcessSettingsCollection& ToCollection, float BlendFactor);
+	void LerpAll(const FPostProcessSettings& ToPostProcessSettings, float BlendFactor);
+
+	/**
+	 * Interpolates the post-process settings towards the values in the given other collection. Only
+	 * values from the other collection that have their bOverride_Xxx set to true will be interpolated.
+	 * If the current value isn't overriden, it will interpolate from the default value. All affected
+	 * values will get their bOverride_Xxx flag set.
+	 */
+	void LerpChanged(const FPostProcessSettingsCollection& ToCollection, float BlendFactor);
+	void LerpChanged(const FPostProcessSettings& ToPostProcessSettings, float BlendFactor);
 
 	/** Serializes this collection into the given archive. */
 	void Serialize(FArchive& Ar);
 
 private:
 
-	using FEntry = FPostProcessSettingsCollectionEntry;
-	using FCollection = TArray<FEntry, TInlineAllocator<1>>;
-	FCollection Entries;
+	void InternalLerpChanged(const FPostProcessSettings& ToPostProcessSettings, float BlendFactor, bool bChangedOnly);
+
+private:
+
+	static const FPostProcessSettings DefaultPostProcessSettings;
+
+	FPostProcessSettings PostProcessSettings;
+	bool bHasAnySetting = false;
 };
 
 }  // namespace UE::Cameras
diff --git a/Engine/Plugins/Cameras/GameplayCameras/Source/GameplayCameras/Public/Debug/CameraNodeEvaluationResultDebugBlock.h b/Engine/Plugins/Cameras/GameplayCameras/Source/GameplayCameras/Public/Debug/CameraNodeEvaluationResultDebugBlock.h
--- a/Engine/Plugins/Cameras/GameplayCameras/Source/GameplayCameras/Public/Debug/CameraNodeEvaluationResultDebugBlock.h
+++ b/Engine/Plugins/Cameras/GameplayCameras/Source/GameplayCameras/Public/Debug/CameraNodeEvaluationResultDebugBlock.h
@@ -57,7 +57,7 @@ protected:
 
 private:
 
-	int32 NumPostProcessSettings = 0;
+	bool bHasAnyPostProcessSettings;
 	bool bIsCameraCut;
 	bool bIsValid;
 };
diff --git a/Engine/Plugins/Cameras/GameplayCameras/Source/GameplayCameras/Public/GameFramework/BlueprintCameraPose.h b/Engine/Plugins/Cameras/GameplayCameras/Source/GameplayCameras/Public/GameFramework/BlueprintCameraPose.h
--- a/Engine/Plugins/Cameras/GameplayCameras/Source/GameplayCameras/Public/GameFramework/BlueprintCameraPose.h
+++ b/Engine/Plugins/Cameras/GameplayCameras/Source/GameplayCameras/Public/GameFramework/BlueprintCameraPose.h
@@ -52,6 +52,10 @@ public:
 	UPROPERTY(BlueprintReadWrite, EditAnywhere, Category="Camera")
 	float Aperture = 2.8f;
 
+	/** The shutter speed of the camera's lens, in 1/seconds */
+	UPROPERTY(BlueprintReadWrite, EditAnywhere, Category="Camera")
+	float ShutterSpeed = 60.f;
+
 	/** The focus distance of the camera, if different from target distance. */
 	UPROPERTY(BlueprintReadWrite, EditAnywhere, Category="Camera")
 	float FocusDistance = -1.f;
@@ -64,10 +68,18 @@ public:
 	UPROPERTY(BlueprintReadWrite, EditAnywhere, Category="Camera")
 	float SensorHeight = 18.67f;
 
+	/** The camera sensor sensitivity in ISO. */
+	UPROPERTY(BlueprintReadWrite, EditAnywhere, Category="Camera")
+	float ISO = 100.f;
+
 	/** The squeeze factor of the camera's lens. */
 	UPROPERTY(BlueprintReadWrite, EditAnywhere, Category="Camera")
 	float SqueezeFactor = 1.f;
 
+	/** Number of blades in the lens diaphragm */
+	UPROPERTY(BlueprintReadWrite, EditAnywhere, Category="Camera")
+	int32 DiaphragmBladeCount = 8;
+
 	/** The camera's near clipping plane. */
 	UPROPERTY(BlueprintReadWrite, EditAnywhere, Category="Camera")
 	float NearClippingPlane = 10.f;
@@ -76,6 +88,17 @@ public:
 	UPROPERTY(BlueprintReadWrite, EditAnywhere, Category="Camera")
 	float FarClippingPlane = -1.f;
 
+	/** Internal weight for physical camera post-process settings. */
+	UPROPERTY()
+	float PhysicalCameraBlendWeight = 0.f;
+
+	/** 
+	 * Whether to setup post-process settings based on physical camera properties such as Aperture,
+	 * FocusDistance, DiaphragmBladeCount, and so on.
+	 */
+	UPROPERTY(BlueprintReadWrite, EditAnywhere, Category="Camera")
+	bool EnablePhysicalCamera = false;
+
 	/** Whether the camera should constrain aspect ratio. */
 	UPROPERTY(BlueprintReadWrite, EditAnywhere, Category="Camera")
 	bool ConstrainAspectRatio = false;
diff --git a/Engine/Plugins/Cameras/GameplayCameras/Source/GameplayCameras/Public/Nodes/Common/FilmbackCameraNode.h b/Engine/Plugins/Cameras/GameplayCameras/Source/GameplayCameras/Public/Nodes/Common/FilmbackCameraNode.h
--- a/Engine/Plugins/Cameras/GameplayCameras/Source/GameplayCameras/Public/Nodes/Common/FilmbackCameraNode.h
+++ b/Engine/Plugins/Cameras/GameplayCameras/Source/GameplayCameras/Public/Nodes/Common/FilmbackCameraNode.h
@@ -34,6 +34,10 @@ public:
 	UPROPERTY(EditAnywhere, Category="Filmback", meta=(ClampMin="0.001", ForceUnits=mm))
 	FFloatCameraParameter SensorHeight;
 
+	/** The camera sensor sensitivity in ISO. */
+	UPROPERTY(EditAnywhere, Category="Filmback")
+	FFloatCameraParameter ISO;
+
 	/** Whether to constrain the aspect ratio of the evaluated camera. */
 	UPROPERTY(EditAnywhere, Category="Filmback")
 	FBooleanCameraParameter ConstrainAspectRatio = false;
diff --git a/Engine/Plugins/Cameras/GameplayCameras/Source/GameplayCameras/Public/Nodes/Common/LensParametersCameraNode.h b/Engine/Plugins/Cameras/GameplayCameras/Source/GameplayCameras/Public/Nodes/Common/LensParametersCameraNode.h
--- a/Engine/Plugins/Cameras/GameplayCameras/Source/GameplayCameras/Public/Nodes/Common/LensParametersCameraNode.h
+++ b/Engine/Plugins/Cameras/GameplayCameras/Source/GameplayCameras/Public/Nodes/Common/LensParametersCameraNode.h
@@ -23,16 +23,27 @@ protected:
 
 public:
 
-	/** Manually-controlled focus distance (manual focus mode only) */
-	UPROPERTY(EditAnywhere, Category="Lens Parameters", meta=(Units=cm))
-	FFloatCameraParameter FocusDistance;
-
 	/** Current focal length of the camera (i.e. controls FoV, zoom) */
 	UPROPERTY(EditAnywhere, Category="Lens Parameters")
 	FFloatCameraParameter FocalLength;
 
+	/** Manually-controlled focus distance (manual focus mode only) */
+	UPROPERTY(EditAnywhere, Category="Lens Parameters", meta=(Units=cm))
+	FFloatCameraParameter FocusDistance;
+
 	/** Current aperture, in terms of f-stop (e.g. 2.8 for f/2.8) */
 	UPROPERTY(EditAnywhere, Category="Lens Parameters")
 	FFloatCameraParameter Aperture;
+
+	/** Current shutter speed, in 1/seconds */
+	UPROPERTY(EditAnywhere, Category="Lens Parameters", DisplayName="Shutter Speed (1/s)")
+    FFloatCameraParameter ShutterSpeed;
+
+	/** 
+	 * Whether post-process settings should be automatically applied to reflect
+	 * the lens settings.
+	 */
+	UPROPERTY(EditAnywhere, Category="Lens Parameters")
+	FBooleanCameraParameter EnablePhysicalCamera = false;
 };
 
diff --git a/Engine/Plugins/Cameras/GameplayCameras/Source/GameplayCameras/Public/Nodes/Common/PostProcessCameraNode.h b/Engine/Plugins/Cameras/GameplayCameras/Source/GameplayCameras/Public/Nodes/Common/PostProcessCameraNode.h
--- a/Engine/Plugins/Cameras/GameplayCameras/Source/GameplayCameras/Public/Nodes/Common/PostProcessCameraNode.h
+++ b/Engine/Plugins/Cameras/GameplayCameras/Source/GameplayCameras/Public/Nodes/Common/PostProcessCameraNode.h
@@ -25,8 +25,5 @@ public:
 
 	UPROPERTY(EditAnywhere, Category="Rendering")
 	FPostProcessSettings PostProcessSettings;
-
-	UPROPERTY(EditAnywhere, Category="Rendering")
-	FFloatCameraParameter PostProcessBlendWeight = 1.f;
 };
 
diff --git a/Engine/Plugins/Cameras/GameplayCameras/Source/GameplayCameras/Public/Nodes/Common/TargetRayCastCameraNode.h b/Engine/Plugins/Cameras/GameplayCameras/Source/GameplayCameras/Public/Nodes/Common/TargetRayCastCameraNode.h
--- a/Engine/Plugins/Cameras/GameplayCameras/Source/GameplayCameras/Public/Nodes/Common/TargetRayCastCameraNode.h
+++ b/Engine/Plugins/Cameras/GameplayCameras/Source/GameplayCameras/Public/Nodes/Common/TargetRayCastCameraNode.h
@@ -3,6 +3,7 @@
 #pragma once
 
 #include "Core/CameraNode.h"
+#include "Core/CameraParameters.h"
 
 #include "TargetRayCastCameraNode.generated.h"
 
@@ -21,6 +22,10 @@ public:
 	UPROPERTY(EditAnywhere, Category="Ray-Cast")
 	TEnumAsByte<ECollisionChannel> TraceChannel = ECollisionChannel::ECC_Camera;
 
+	/** Whether to set the focus distance to the ray-cast hit result. */
+	UPROPERTY(EditAnywhere, Category="Auto-Focus")
+	FBooleanCameraParameter AutoFocus = true;
+
 protected:
 
 	// UCameraNode interface.
