Cameras: remove unnecessary ensure

diff --git a/Engine/Plugins/Cameras/GameplayCameras/Source/GameplayCameras/Private/Core/CameraPose.cpp b/Engine/Plugins/Cameras/GameplayCameras/Source/GameplayCameras/Private/Core/CameraPose.cpp
--- a/Engine/Plugins/Cameras/GameplayCameras/Source/GameplayCameras/Private/Core/CameraPose.cpp
+++ b/Engine/Plugins/Cameras/GameplayCameras/Source/GameplayCameras/Private/Core/CameraPose.cpp
@@ -354,10 +354,6 @@ void FCameraPose::InternalLerpChanged(const FCameraPose& ToPose, float Factor, c
 
 			if (!bChangedOnly || (ToPoseChangedFlags.FocalLength || ToPoseChangedFlags.FieldOfView))
 			{
-				ensureMsgf(
-						ChangedFlags.FieldOfView || ChangedFlags.FocalLength, 
-						TEXT("Interpolating FieldOfView or FocalLength from default value!"));
-
 				// Interpolate FocalLength, or FieldOfView, if both poses use the same.
 				// If there's a mix, interpolate the effective FieldOfView.
 				//
