Cameras: enable Rename/Delete keyboard shortcuts

diff --git a/Engine/Plugins/Cameras/GameplayCameras/Source/GameplayCamerasEditor/Private/Commands/CameraAssetEditorCommands.cpp b/Engine/Plugins/Cameras/GameplayCameras/Source/GameplayCamerasEditor/Private/Commands/CameraAssetEditorCommands.cpp
--- a/Engine/Plugins/Cameras/GameplayCameras/Source/GameplayCamerasEditor/Private/Commands/CameraAssetEditorCommands.cpp
+++ b/Engine/Plugins/Cameras/GameplayCameras/Source/GameplayCamerasEditor/Private/Commands/CameraAssetEditorCommands.cpp
@@ -41,9 +41,9 @@ void FCameraAssetEditorCommands::RegisterCommands()
 	UI_COMMAND(AddCameraRig, "Add Camera Rig", "Adds a new camera rig",
 			EUserInterfaceActionType::Button, FInputChord());
 	UI_COMMAND(DeleteCameraRig, "Delete Camera Rig", "Deletes the selected camera rig(s)",
-			EUserInterfaceActionType::Button, FInputChord());
+			EUserInterfaceActionType::Button, FInputChord(EKeys::Delete));
 	UI_COMMAND(RenameCameraRig, "Rename Camera Rig", "Renames the selected camera rig",
-			EUserInterfaceActionType::Button, FInputChord());
+			EUserInterfaceActionType::Button, FInputChord(EKeys::F2));
 
 	UI_COMMAND(ShowMessages, "Messages", "Shows the message log for this camera asset",
 			EUserInterfaceActionType::ToggleButton, FInputChord());
diff --git a/Engine/Plugins/Cameras/GameplayCameras/Source/GameplayCamerasEditor/Private/Editors/SCameraRigList.cpp b/Engine/Plugins/Cameras/GameplayCameras/Source/GameplayCamerasEditor/Private/Editors/SCameraRigList.cpp
--- a/Engine/Plugins/Cameras/GameplayCameras/Source/GameplayCamerasEditor/Private/Editors/SCameraRigList.cpp
+++ b/Engine/Plugins/Cameras/GameplayCameras/Source/GameplayCamerasEditor/Private/Editors/SCameraRigList.cpp
@@ -251,6 +251,15 @@ void SCameraRigList::Tick(const FGeometry& AllottedGeometry, const double InCurr
 	SCompoundWidget::Tick(AllottedGeometry, InCurrentTime, InDeltaTime);
 }
 
+FReply SCameraRigList::OnKeyDown(const FGeometry& MyGeometry, const FKeyEvent& InKeyEvent)
+{
+	if (CommandList->ProcessCommandBindings(InKeyEvent))
+	{
+		return FReply::Handled();
+	}
+	return SCompoundWidget::OnKeyDown(MyGeometry, InKeyEvent);
+}
+
 TSharedPtr<SWidget> SCameraRigList::GenerateToolbar()
 {
 	static const FName ToolbarName("CameraRigList.ToolBar");
diff --git a/Engine/Plugins/Cameras/GameplayCameras/Source/GameplayCamerasEditor/Public/Editors/SCameraRigList.h b/Engine/Plugins/Cameras/GameplayCameras/Source/GameplayCamerasEditor/Public/Editors/SCameraRigList.h
--- a/Engine/Plugins/Cameras/GameplayCameras/Source/GameplayCamerasEditor/Public/Editors/SCameraRigList.h
+++ b/Engine/Plugins/Cameras/GameplayCameras/Source/GameplayCamerasEditor/Public/Editors/SCameraRigList.h
@@ -89,6 +89,7 @@ protected:
 
 	// SWidget interface
 	virtual void Tick(const FGeometry& AllottedGeometry, const double InCurrentTime, const float InDeltaTime) override;
+	virtual FReply OnKeyDown(const FGeometry& MyGeometry, const FKeyEvent& InKeyEvent) override;
 
 private:
 
