Cameras: fix camera rig picker for single camera director

diff --git a/Engine/Plugins/Cameras/GameplayCameras/Source/GameplayCameras/Public/Directors/SingleCameraDirector.h b/Engine/Plugins/Cameras/GameplayCameras/Source/GameplayCameras/Public/Directors/SingleCameraDirector.h
--- a/Engine/Plugins/Cameras/GameplayCameras/Source/GameplayCameras/Public/Directors/SingleCameraDirector.h
+++ b/Engine/Plugins/Cameras/GameplayCameras/Source/GameplayCameras/Public/Directors/SingleCameraDirector.h
@@ -30,7 +30,7 @@ protected:
 public:
 
 	/** The camera rig to run every frame. */
-	UPROPERTY(EditAnywhere, Category=Common)
+	UPROPERTY(EditAnywhere, Category=Common, meta=(UseSelfCameraRigPicker=true))
 	TObjectPtr<UCameraRigAsset> CameraRig;
 };
 
