Cameras: fix layout and reset-to-default logic for camera parameter detail customizations

- Redo the layout in order to make the variable-picker combo-button more gracefully resize with the Details View, giving more space to the value widget.
- Move some of the variable-dependent bits of UI to be more dynamic, so the UI reacts to variables being set/unset by, for instance, the camera rig builder.
- An exposed camera parameter shouldn't be allowed to be bound to a variable, so disable the combo-button in that case.
- Resetting a camera parameter to its default value shouldn't clear the private variable set by the camera rig builder.

diff --git a/Engine/Plugins/Cameras/GameplayCameras/Source/GameplayCamerasEditor/Private/Customizations/CameraParameterDetailsCustomizations.cpp b/Engine/Plugins/Cameras/GameplayCameras/Source/GameplayCamerasEditor/Private/Customizations/CameraParameterDetailsCustomizations.cpp
--- a/Engine/Plugins/Cameras/GameplayCameras/Source/GameplayCamerasEditor/Private/Customizations/CameraParameterDetailsCustomizations.cpp
+++ b/Engine/Plugins/Cameras/GameplayCameras/Source/GameplayCamerasEditor/Private/Customizations/CameraParameterDetailsCustomizations.cpp
@@ -20,9 +20,11 @@
 #include "ScopedTransaction.h"
 #include "Styles/GameplayCamerasEditorStyle.h"
 #include "Styling/StyleColors.h"
+#include "UObject/Object.h"
 #include "Widgets/Images/SImage.h"
 #include "Widgets/Input/SComboButton.h"
 #include "Widgets/SBoxPanel.h"
+#include "Widgets/SCompoundWidget.h"
 #include "Widgets/Text/STextBlock.h"
 
 #define LOCTEXT_NAMESPACE "CameraParameterDetailsCustomization"
@@ -30,6 +32,21 @@
 namespace UE::Cameras
 {
 
+//class SCameraParameterLayoutBox : public SCompoundWidget
+//{
+//public:
+//
+//protected:
+//
+//	// SCompoundWidget interface.
+//	virtual int32 OnPaint(const FPaintArgs& Args, const FGeometry& AllottedGeometry, const FSlateRect& MyCullingRect, FSlateWindowElementList& OutDrawElements, int32 LayerId, const FWidgetStyle& InWidgetStyle, bool bParentEnabled) const override
+//	{
+//		SCompoundWidget::OnPaint(Args, AllottedGeometry, MyCullingRect, OutDrawElements, LayerId, InWidgetStyle, bParentEnabled);
+//
+//		this->GetPaintSpaceGeometry();
+//	}
+//};
+
 void FCameraParameterDetailsCustomization::Register(FPropertyEditorModule& PropertyEditorModule)
 {
 #define UE_CAMERA_VARIABLE_FOR_TYPE(ValueType, ValueName)\
@@ -55,39 +72,35 @@ UE_CAMERA_VARIABLE_FOR_ALL_TYPES()
 
 void FCameraParameterDetailsCustomization::CustomizeHeader(TSharedRef<IPropertyHandle> PropertyHandle, FDetailWidgetRow& HeaderRow, IPropertyTypeCustomizationUtils& CustomizationUtils)
 {
+	// Gather up the things we need.
 	PropertyUtilities = CustomizationUtils.GetPropertyUtilities();
-
 	StructProperty = PropertyHandle;
-	StructProperty->SetOnPropertyResetToDefault(
-			FSimpleDelegate::CreateSP(this, &FCameraParameterDetailsCustomization::OnResetToDefault));
 
+	// Figure out the parameter value and driving variable property names for the actual camera
+	// parameter type we are displaying. In theory these are always called "Value" and "Variable", but
+	// we do this in a statically-typed way via polyphormism, instead of in a duck-typing way.
 	FName ValuePropertyName, VariablePropertyName;
 	GetValueAndVariablePropertyNames(ValuePropertyName, VariablePropertyName);
 
 	ValueProperty = PropertyHandle->GetChildHandle(ValuePropertyName);
 	VariableProperty = PropertyHandle->GetChildHandle(VariablePropertyName);
 
+	// Get the type of camera variable we need for this camera parameter (bool variable, float variable, etc.)
 	VariableClass = nullptr;
 	if (FObjectProperty* VariableObjectProperty = CastField<FObjectProperty>(VariableProperty->GetProperty()))
 	{
 		VariableClass = VariableObjectProperty->PropertyClass;
 	}
 
+	// Update our variable info once now. We will then update it every tick, since the UI needs it
+	// for various things.
 	UpdateVariableInfo();
 
-	// The value widget is enabled (i.e. the user can change the value) if the parameter isn't driven by
-	// a variable, or if that variable is a private variable meant to expose the parameter on the rig interface.
-	const bool bShowVariableText = !VariableInfoText.IsEmpty() && !bIsExposedParameterVariable;
-	const bool bShowVariableError = !VariableErrorText.IsEmpty() && !bIsExposedParameterVariable;
-	const bool bEnableVariableWidget = bShowVariableText || bShowVariableError;
-
+	// Create the parameter value editor (float editor, vector editor, etc.)
 	TSharedRef<SWidget> ValueWidget = ValueProperty->CreatePropertyValueWidgetWithCustomization(nullptr);
-	ValueWidget->SetEnabled(!bEnableVariableWidget);
-
-	// TODO: change SStandaloneCustomizedValueWidget so that it can tell us about the layout requirements
-	//		 of the value widget (min/max desired width/height, H/V alignment, etc.)
-	const float MaxValueWidgetDesiredWidth = bEnableVariableWidget ? 300.f : 0.f;
+	ValueWidget->SetEnabled(TAttribute<bool>::CreateSP(this, &FCameraParameterDetailsCustomization::IsValueEditorEnabled));
 
+	// Create the whole UI layout.
 	TSharedRef<FGameplayCamerasEditorStyle> GameplayCamerasStyle = FGameplayCamerasEditorStyle::Get();
 
 	HeaderRow
@@ -97,10 +110,9 @@ void FCameraParameterDetailsCustomization::CustomizeHeader(TSharedRef<IPropertyH
 	]
 	.ValueContent()
 	.MinDesiredWidth(100.f)
-	.MaxDesiredWidth(MaxValueWidgetDesiredWidth)
 	.HAlign(HAlign_Fill)
 	[
-		SNew(SHorizontalBox)
+		SAssignNew(LayoutBox, SHorizontalBox)
 		+SHorizontalBox::Slot()
 		.Padding(0)
 		.FillWidth(1.f)
@@ -114,11 +126,6 @@ void FCameraParameterDetailsCustomization::CustomizeHeader(TSharedRef<IPropertyH
 		[
 			ValueProperty->CreateDefaultPropertyButtonWidgets()
 		]
-	]
-	.ExtensionContent()
-	.HAlign(HAlign_Right)
-	[
-		SNew(SHorizontalBox)
 		+SHorizontalBox::Slot()
 		.Padding(2.f)
 		.AutoWidth()
@@ -127,6 +134,7 @@ void FCameraParameterDetailsCustomization::CustomizeHeader(TSharedRef<IPropertyH
 			.HasDownArrow(true)
 			.ContentPadding(1.f)
 			.ButtonStyle(FAppStyle::Get(), "SimpleButton")
+			.IsEnabled(this, &FCameraParameterDetailsCustomization::IsCameraVariableBrowserEnabled)
 			.ToolTipText(LOCTEXT("SetVariable_ToolTip", "Selects a camera variable to drive this parameter"))
 			.ButtonContent()
 			[
@@ -141,31 +149,35 @@ void FCameraParameterDetailsCustomization::CustomizeHeader(TSharedRef<IPropertyH
 					.Image(GameplayCamerasStyle->GetBrush("CameraParameter.VariableBrowser"))
 				]
 				+SHorizontalBox::Slot()
-				.AutoWidth()
+				.FillWidth(0.3f)
 				.Padding(2.f)
 				.VAlign(VAlign_Center)
 				[
 					SNew(SBox)
 					.VAlign(VAlign_Center)
-					.Visibility(bShowVariableText ? EVisibility::Visible : EVisibility::Collapsed)
+					.Visibility(this, &FCameraParameterDetailsCustomization::GetVariableInfoTextVisibility)
+					.MaxDesiredWidth(this, &FCameraParameterDetailsCustomization::GetVariableInfoTextMaxWidth)
 					[
 						SNew(STextBlock)
-						.Text(VariableInfoText)
+						.Text(this, &FCameraParameterDetailsCustomization::GetVariableInfoText)
 						.MinDesiredWidth(20.f)
-						.OverflowPolicy(ETextOverflowPolicy::Clip)
+						.OverflowPolicy(ETextOverflowPolicy::Ellipsis)
 					]
 				]
 				+SHorizontalBox::Slot()
-				.AutoWidth()
+				.FillWidth(0.3f)
 				.Padding(2.f)
 				.VAlign(VAlign_Center)
 				[
 					SNew(SBox)
 					.VAlign(VAlign_Center)
-					.Visibility(bShowVariableError ? EVisibility::Visible : EVisibility::Collapsed)
+					.Visibility(this, &FCameraParameterDetailsCustomization::GetVariableErrorTextVisibility)
+					.MaxDesiredWidth(this, &FCameraParameterDetailsCustomization::GetVariableErrorTextMaxWidth)
 					[
 						SNew(STextBlock)
-						.Text(VariableErrorText)
+						.Text(this, &FCameraParameterDetailsCustomization::GetVariableErrorText)
+						.MinDesiredWidth(20.f)
+						.OverflowPolicy(ETextOverflowPolicy::Ellipsis)
 						.ColorAndOpacity(FStyleColors::Error)
 					]
 				]
@@ -173,12 +185,68 @@ void FCameraParameterDetailsCustomization::CustomizeHeader(TSharedRef<IPropertyH
 			.OnGetMenuContent(this, &FCameraParameterDetailsCustomization::BuildCameraVariableBrowser)
 		]
 	];
+
+	// Setup some custom reset-to-default behavior, if we are allowed to.
+	// (see the code in FCameraRigAssetReferenceDetailsCustomization that sets this metadata)
+	const bool bNoResetToDefault = StructProperty->GetBoolMetaData("NoCustomCameraParameterResetToDefault");
+	if (!bNoResetToDefault)
+	{
+		HeaderRow.OverrideResetToDefault(
+				FResetToDefaultOverride::Create(
+					FIsResetToDefaultVisible::CreateSP(this, &FCameraParameterDetailsCustomization::IsResetToDefaultVisible),
+					FResetToDefaultHandler::CreateSP(this, &FCameraParameterDetailsCustomization::OnResetToDefault)));
+	}
 }
 
 void FCameraParameterDetailsCustomization::CustomizeChildren(TSharedRef<IPropertyHandle> PropertyHandle, IDetailChildrenBuilder& ChildBuilder, IPropertyTypeCustomizationUtils& CustomizationUtils)
 {
 }
 
+void FCameraParameterDetailsCustomization::Tick(float DeltaTime)
+{
+	// Use the editor tick to query the property values only once per frame.
+	UpdateVariableInfo();
+}
+
+void FCameraParameterDetailsCustomization::UpdateVariableInfo()
+{
+	VariableInfo = FCameraVariableInfo();
+
+	UObject* VariableObject = nullptr;
+	FPropertyAccess::Result PropertyAccessResult = VariableProperty->GetValue(VariableObject);
+	if (PropertyAccessResult == FPropertyAccess::Success)
+	{
+		if (VariableObject)
+		{
+			if (UCameraVariableAsset* Variable = Cast<UCameraVariableAsset>(VariableObject))
+			{
+				VariableInfo.VariableValue = ECameraVariableValue::Set;
+				VariableInfo.CommonVariable = Variable;
+				VariableInfo.InfoText = Variable->DisplayName.IsEmpty() ?
+					FText::FromName(Variable->GetFName()) :
+					FText::FromString(Variable->DisplayName);
+				VariableInfo.bIsExposedParameterVariable = VariableInfo.CommonVariable->bIsPrivate;
+			}
+			else
+			{
+				VariableInfo.VariableValue = ECameraVariableValue::Invalid;
+				VariableInfo.ErrorText = LOCTEXT("InvalidVariableObject", "Invalid Variable");
+			}
+		}
+		// else: variable is not set
+	}
+	else if (PropertyAccessResult == FPropertyAccess::MultipleValues)
+	{
+		VariableInfo.VariableValue = ECameraVariableValue::MultipleSet;
+		VariableInfo.InfoText = LOCTEXT("MultipleVariableValues", "Multiple Variables");
+	}
+	else
+	{
+		VariableInfo.VariableValue = ECameraVariableValue::Invalid;
+		VariableInfo.ErrorText = LOCTEXT("ErrorReadingVariable", "Error Reading Variable");
+	}
+}
+
 TSharedRef<SWidget> FCameraParameterDetailsCustomization::BuildCameraVariableBrowser()
 {
 	const bool bShouldCloseWindowAfterMenuSelection = true;
@@ -201,7 +269,7 @@ TSharedRef<SWidget> FCameraParameterDetailsCustomization::BuildCameraVariableBro
 
 	FCameraVariablePickerConfig PickerConfig;
 	PickerConfig.CameraVariableClass = VariableClass;
-	PickerConfig.InitialCameraVariableSelection = CommonVariable;
+	PickerConfig.InitialCameraVariableSelection = VariableInfo.CommonVariable;
 	PickerConfig.CameraVariableCollectionSaveSettingsName = TEXT("CameraParameterVariablePropertyPicker");
 	PickerConfig.OnCameraVariableSelected = FOnCameraVariableSelected::CreateSP(
 			this, &FCameraParameterDetailsCustomization::OnSetVariable);
@@ -223,47 +291,62 @@ TSharedRef<SWidget> FCameraParameterDetailsCustomization::BuildCameraVariableBro
 	return MenuBuilder.MakeWidget();
 }
 
-void FCameraParameterDetailsCustomization::UpdateVariableInfo()
+bool FCameraParameterDetailsCustomization::IsValueEditorEnabled() const
 {
-	CommonVariable = nullptr;
-	VariableInfoText = FText::GetEmpty();
-	VariableErrorText = FText::GetEmpty();
-	bIsExposedParameterVariable = false;
+	// The value widget is enabled (i.e. the user can change the value) if the parameter isn't driven by
+	// a variable, or if that variable is a private variable meant to expose the parameter on the rig interface.
+	return (VariableInfo.VariableValue == ECameraVariableValue::NotSet || VariableInfo.bIsExposedParameterVariable);
+}
 
-	UObject* VariableObject = nullptr;
-	FPropertyAccess::Result PropertyAccessResult = VariableProperty->GetValue(VariableObject);
-	if (PropertyAccessResult == FPropertyAccess::Success)
-	{
-		if (VariableObject)
-		{
-			if (UCameraVariableAsset* Variable = Cast<UCameraVariableAsset>(VariableObject))
-			{
-				CommonVariable = Variable;
-				VariableInfoText = Variable->DisplayName.IsEmpty() ?
-					FText::FromName(Variable->GetFName()) :
-					FText::FromString(Variable->DisplayName);
-				bIsExposedParameterVariable = CommonVariable->bIsPrivate;
-			}
-			else
-			{
-				VariableErrorText = LOCTEXT("InvalidVariableObject", "Invalid Variable");
-			}
-		}
-		// else: no variable set
-	}
-	else if (PropertyAccessResult == FPropertyAccess::MultipleValues)
-	{
-		VariableInfoText = LOCTEXT("MultipleVariableValues", "Multiple Variables");
-	}
-	else
-	{
-		VariableErrorText = LOCTEXT("ErrorReadingVariable", "Error Reading Variable");
-	}
+bool FCameraParameterDetailsCustomization::IsCameraVariableBrowserEnabled() const
+{
+	// The variable picker is enabled if the parameter isn't exposed to the rig interface via a private variable,
+	// since we can't drive a value with both an interface parameter and a user-defined variable.
+	return !VariableInfo.bIsExposedParameterVariable;
+}
+
+FText FCameraParameterDetailsCustomization::GetVariableInfoText() const
+{
+	return VariableInfo.InfoText;
+}
+
+EVisibility FCameraParameterDetailsCustomization::GetVariableInfoTextVisibility() const
+{
+	const bool bShowVariableInfoText = !VariableInfo.InfoText.IsEmpty() && !VariableInfo.bIsExposedParameterVariable;
+	return bShowVariableInfoText ? EVisibility::Visible : EVisibility::Collapsed;
+}
+
+FOptionalSize FCameraParameterDetailsCustomization::GetVariableInfoTextMaxWidth() const
+{
+	// We want this text to take at most 30% of the free-standing space of the combo button.
+	// Free-standing space excludes fixed things like the combo button icon, the dropdown icon, paddings, etc.
+	// IMPORTANT: update this if the main layout changes inside Construct()
+	const float FixedSpace = 1.f + (2.f+ 16.f + 2.f) + (2.f + 16.f + 2.f) + 1.f;
+
+	const bool bShowVariableInfoText = !VariableInfo.InfoText.IsEmpty() && !VariableInfo.bIsExposedParameterVariable;
+	const float LayoutBoxWidth = LayoutBox ? LayoutBox->GetPaintSpaceGeometry().GetLocalSize().X : 0.f;
+	return bShowVariableInfoText ? FOptionalSize((LayoutBoxWidth - FixedSpace) / 3.f) : FOptionalSize(0);
+}
+
+FText FCameraParameterDetailsCustomization::GetVariableErrorText() const
+{
+	return VariableInfo.ErrorText;
+}
+
+EVisibility FCameraParameterDetailsCustomization::GetVariableErrorTextVisibility() const
+{
+	const bool bShowVariableErrorText = !VariableInfo.ErrorText.IsEmpty() && !VariableInfo.bIsExposedParameterVariable;
+	return bShowVariableErrorText ? EVisibility::Visible : EVisibility::Collapsed;
 }
 
-bool FCameraParameterDetailsCustomization::HasVariableInfoText() const
+FOptionalSize FCameraParameterDetailsCustomization::GetVariableErrorTextMaxWidth() const
 {
-	return !VariableInfoText.IsEmpty() || !VariableErrorText.IsEmpty();
+	// See comments in GetVariableInfoTextMaxWidth.
+	const float FixedSpace = 1.f + (2.f+ 16.f + 2.f) + (2.f + 16.f + 2.f) + 1.f;
+
+	const bool bShowVariableErrorText = !VariableInfo.ErrorText.IsEmpty() && !VariableInfo.bIsExposedParameterVariable;
+	const float LayoutBoxWidth = LayoutBox ? LayoutBox->GetPaintSpaceGeometry().GetLocalSize().X : 0.f;
+	return bShowVariableErrorText ? FOptionalSize((LayoutBoxWidth - FixedSpace) / 3.f) : FOptionalSize(0);
 }
 
 bool FCameraParameterDetailsCustomization::CanClearVariable() const
@@ -308,8 +391,30 @@ void FCameraParameterDetailsCustomization::OnSetVariable(UCameraVariableAsset* I
 	VariableBrowserButton->SetIsOpen(false);
 }
 
-void FCameraParameterDetailsCustomization::OnResetToDefault()
+bool FCameraParameterDetailsCustomization::IsResetToDefaultVisible(TSharedPtr<IPropertyHandle> InPropertyHandle) const
+{
+	// The user can reset the camera parameter to its default if the value is non-default, and/or the
+	// variable is a user-defined variable. In other words, if the variable is private because the parameter
+	// is exposed, then we don't want to reset that -- we only want to reset the value.
+	return 
+		ValueProperty->CanResetToDefault() || 
+		(VariableProperty->CanResetToDefault() && !VariableInfo.bIsExposedParameterVariable);
+}
+
+void FCameraParameterDetailsCustomization::OnResetToDefault(TSharedPtr<IPropertyHandle> InPropertyHandle)
 {
+	// As mentioned above, if the camera parameter is exposed publicly on the camera rig, we only want
+	// to reset the value to the default (and keep the private variable set on it).
+	// Otherwise, we can reset both the value and the variable.
+	if (VariableInfo.bIsExposedParameterVariable)
+	{
+		ValueProperty->ResetToDefault();
+	}
+	else
+	{
+		StructProperty->ResetToDefault();
+	}
+
 	PropertyUtilities->RequestForceRefresh();
 }
 
@@ -317,7 +422,7 @@ void FCameraParameterDetailsCustomization::OnResetToDefault()
 void F##ValueName##CameraParameterDetailsCustomization::SetParameterVariable(void* InRawData, UCameraVariableAsset* InVariable)\
 {\
 	F##ValueName##CameraParameter* TypedData = reinterpret_cast<F##ValueName##CameraParameter*>(InRawData);\
-	TypedData->Variable = CastChecked<U##ValueName##CameraVariable>(InVariable);\
+	TypedData->Variable = CastChecked<U##ValueName##CameraVariable>(InVariable, ECastCheckedType::NullAllowed);\
 }
 UE_CAMERA_VARIABLE_FOR_ALL_TYPES()
 #undef UE_CAMERA_VARIABLE_FOR_TYPE
diff --git a/Engine/Plugins/Cameras/GameplayCameras/Source/GameplayCamerasEditor/Private/Customizations/CameraRigAssetReferenceDetailsCustomization.cpp b/Engine/Plugins/Cameras/GameplayCameras/Source/GameplayCamerasEditor/Private/Customizations/CameraRigAssetReferenceDetailsCustomization.cpp
--- a/Engine/Plugins/Cameras/GameplayCameras/Source/GameplayCamerasEditor/Private/Customizations/CameraRigAssetReferenceDetailsCustomization.cpp
+++ b/Engine/Plugins/Cameras/GameplayCameras/Source/GameplayCamerasEditor/Private/Customizations/CameraRigAssetReferenceDetailsCustomization.cpp
@@ -273,6 +273,12 @@ void FCameraRigParameterOverrideDetailRow::BuildDetailPropertyRowImpl(IDetailChi
 	ParameterOverrideRow->GetPropertyHandle()->SetOnChildPropertyValueChanged(FSimpleDelegate::CreateSP(
 				this, &FCameraRigParameterOverrideDetailRow::OnPropertyValueChanged<ParameterOverrideType>));
 
+	// We need to add our own custom reset-to-default logic, which sets the default value as determined
+	// by the values inside the camera rig prefab. However, camera parameters also have a detail 
+	// customization that has a custom reset-to-default logic. Slate complaines when both the outer and
+	// inner rows have custom reset-to-default logic, so suppress the inner one.
+	ParameterOverrideRow->GetPropertyHandle()->SetInstanceMetaData("NoCustomCameraParameterResetToDefault", "true");
+
 	FResetToDefaultOverride ResetToDefault = FResetToDefaultOverride::Create(
 			FIsResetToDefaultVisible::CreateSP(
 				this, &FCameraRigParameterOverrideDetailRow::OnIsResetToDefaultVisible<ParameterOverrideType>),
diff --git a/Engine/Plugins/Cameras/GameplayCameras/Source/GameplayCamerasEditor/Public/Customizations/CameraParameterDetailsCustomizations.h b/Engine/Plugins/Cameras/GameplayCameras/Source/GameplayCamerasEditor/Public/Customizations/CameraParameterDetailsCustomizations.h
--- a/Engine/Plugins/Cameras/GameplayCameras/Source/GameplayCamerasEditor/Public/Customizations/CameraParameterDetailsCustomizations.h
+++ b/Engine/Plugins/Cameras/GameplayCameras/Source/GameplayCamerasEditor/Public/Customizations/CameraParameterDetailsCustomizations.h
@@ -3,13 +3,16 @@
 #pragma once
 
 #include "Core/CameraParameters.h"
-#include "CoreTypes.h"
 #include "IPropertyTypeCustomization.h"
+#include "Layout/Visibility.h"
+#include "TickableEditorObject.h"
+#include "Types/SlateStructs.h"
 
 class FPropertyEditorModule;
 class IDetailLayoutBuilder;
 class IPropertyUtilities;
 class SComboButton;
+class SHorizontalBox;
 class SWidget;
 
 namespace UE::Cameras
@@ -18,7 +21,9 @@ namespace UE::Cameras
 /**
  * Base details customization for camera parameters.
  */
-class FCameraParameterDetailsCustomization : public IPropertyTypeCustomization
+class FCameraParameterDetailsCustomization 
+	: public IPropertyTypeCustomization
+	, public FTickableEditorObject
 {
 public:
 
@@ -29,10 +34,15 @@ public:
 
 public:
 
-	/** IPropertyTypeCustomization interface */
+	// IPropertyTypeCustomization interface.
 	virtual void CustomizeHeader(TSharedRef<IPropertyHandle> PropertyHandle, FDetailWidgetRow& HeaderRow, IPropertyTypeCustomizationUtils& CustomizationUtils) override;
 	virtual void CustomizeChildren(TSharedRef<IPropertyHandle> PropertyHandle, IDetailChildrenBuilder& ChildBuilder, IPropertyTypeCustomizationUtils& CustomizationUtils) override;
 
+	// FTickableEditorObject interface.
+	virtual void Tick(float DeltaTime) override;
+	virtual ETickableTickType GetTickableTickType() const override { return ETickableTickType::Always; }
+	virtual TStatId GetStatId() const override { RETURN_QUICK_DECLARE_CYCLE_STAT(FCameraParameterDetailsCustomization, STATGROUP_Tickables); }
+
 protected:
 
 	virtual void GetValueAndVariablePropertyNames(FName& OutValueProperty, FName& OutVariableProperty) = 0;
@@ -40,25 +50,52 @@ protected:
 
 private:
 
-	TSharedRef<SWidget> BuildCameraVariableBrowser();
+	enum class ECameraVariableValue
+	{
+		NotSet,
+		Set,
+		MultipleSet,
+		Invalid
+	};
+
+	struct FCameraVariableInfo
+	{
+		UCameraVariableAsset* CommonVariable = nullptr;
+		ECameraVariableValue VariableValue = ECameraVariableValue::NotSet;
+		bool bIsExposedParameterVariable = false;
+
+		FText InfoText;
+		FText ErrorText;
+	};
 
 	void UpdateVariableInfo();
-	bool HasVariableInfoText() const;
+
+	TSharedRef<SWidget> BuildCameraVariableBrowser();
+
+	bool IsValueEditorEnabled() const;
+	bool IsCameraVariableBrowserEnabled() const;
+
+	FText GetVariableInfoText() const;
+	EVisibility GetVariableInfoTextVisibility() const;
+	FOptionalSize GetVariableInfoTextMaxWidth() const;
+
+	FText GetVariableErrorText() const;
+	EVisibility GetVariableErrorTextVisibility() const;
+	FOptionalSize GetVariableErrorTextMaxWidth() const;
 
 	bool CanClearVariable() const;
 	void OnClearVariable();
 
 	void OnSetVariable(UCameraVariableAsset* InVariable);
-	void OnResetToDefault();
+
+	bool IsResetToDefaultVisible(TSharedPtr<IPropertyHandle> InPropertyHandle) const;
+	void OnResetToDefault(TSharedPtr<IPropertyHandle> InPropertyHandle);
 
 protected:
 
 	UClass* VariableClass = nullptr;
 
-	UCameraVariableAsset* CommonVariable = nullptr;
-	FText VariableInfoText;
-	FText VariableErrorText;
-	bool bIsExposedParameterVariable = false;
+	FCameraVariableInfo VariableInfo;
 
 	TSharedPtr<IPropertyUtilities> PropertyUtilities;
 
@@ -66,6 +103,7 @@ protected:
 	TSharedPtr<IPropertyHandle> ValueProperty;
 	TSharedPtr<IPropertyHandle> VariableProperty;
 
+	TSharedPtr<SHorizontalBox> LayoutBox;
 	TSharedPtr<SComboButton> VariableBrowserButton;
 };
 
