Cameras: pick up any stray camera variables in the camera rig building process

This prevents possible name collisions.

diff --git a/Engine/Plugins/Cameras/GameplayCameras/Source/GameplayCameras/Private/Core/CameraRigAssetBuilder.cpp b/Engine/Plugins/Cameras/GameplayCameras/Source/GameplayCameras/Private/Core/CameraRigAssetBuilder.cpp
--- a/Engine/Plugins/Cameras/GameplayCameras/Source/GameplayCameras/Private/Core/CameraRigAssetBuilder.cpp
+++ b/Engine/Plugins/Cameras/GameplayCameras/Source/GameplayCameras/Private/Core/CameraRigAssetBuilder.cpp
@@ -507,6 +507,22 @@ UE_CAMERA_VARIABLE_FOR_ALL_TYPES()
 		}
 	}
 
+	// Sanity check: see if we have any stray camera variables, possibly introduced by incorrect
+	// editor code or dynamic data manipulation.
+	const int32 PreviousNumGatheredVariables = GatheredVariables.Num();
+	ForEachObjectWithOuter(CameraRig, [&GatheredVariables](UObject* Obj)
+			{
+				if (UCameraVariableAsset* CameraVariable = Cast<UCameraVariableAsset>(Obj))
+				{
+					GatheredVariables.Add(CameraVariable);
+				}
+			});
+	if (GatheredVariables.Num() > PreviousNumGatheredVariables)
+	{
+		UE_LOG(LogCameraSystem, Verbose, TEXT("Collected %d stray camera variables while building camera rig '%s'."),
+				(GatheredVariables.Num() - PreviousNumGatheredVariables), *GetPathNameSafe(CameraRig));
+	}
+
 	// Temporarily rename all old camera variables, so their names are available to the new
 	// driven parameters.
 	for (UCameraVariableAsset* GatheredVariable : GatheredVariables)
