Cameras: remove unnecessary assert

diff --git a/Engine/Plugins/Cameras/GameplayCameras/Source/GameplayCameras/Private/Nodes/Collision/CollisionPushCameraNode.cpp b/Engine/Plugins/Cameras/GameplayCameras/Source/GameplayCameras/Private/Nodes/Collision/CollisionPushCameraNode.cpp
--- a/Engine/Plugins/Cameras/GameplayCameras/Source/GameplayCameras/Private/Nodes/Collision/CollisionPushCameraNode.cpp
+++ b/Engine/Plugins/Cameras/GameplayCameras/Source/GameplayCameras/Private/Nodes/Collision/CollisionPushCameraNode.cpp
@@ -300,7 +300,7 @@ void FCollisionPushCameraNodeEvaluator::RunCollisionTrace(UWorld* World, APlayer
 	const FVector3d TraceEnd(OutResult.CameraPose.GetLocation());
 
 	double TraceLength = FVector3d::Distance(TraceStart, TraceEnd);
-	if (!ensure(TraceLength > 0))
+	if (TraceLength <= 0)
 	{
 		return;
 	}
