Cameras: improve camera rig picker UI look

diff --git a/Engine/Plugins/Cameras/GameplayCameras/Source/GameplayCamerasEditor/Private/Customizations/CameraRigPtrDetailsCustomization.cpp b/Engine/Plugins/Cameras/GameplayCameras/Source/GameplayCamerasEditor/Private/Customizations/CameraRigPtrDetailsCustomization.cpp
--- a/Engine/Plugins/Cameras/GameplayCameras/Source/GameplayCamerasEditor/Private/Customizations/CameraRigPtrDetailsCustomization.cpp
+++ b/Engine/Plugins/Cameras/GameplayCameras/Source/GameplayCamerasEditor/Private/Customizations/CameraRigPtrDetailsCustomization.cpp
@@ -50,8 +50,9 @@ void FCameraRigPtrDetailsCustomization::CustomizeHeader(TSharedRef<IPropertyHand
 		}
 
 		CameraRigPickerButton = SNew(SComboButton)
-		.ButtonStyle(FAppStyle::Get(), "PropertyEditor.AssetComboStyle")
+		.HasDownArrow(true)
 		.ContentPadding(FMargin(2.f, 2.f, 2.f, 1.f))
+		.ToolTipText(LOCTEXT("SetRig_ToolTip", "Selects the camera rig"))
 		.OnGetMenuContent(OnGetComboMenuContent)
 		.ButtonContent()
 		[
