Cameras: add a "Get Camera Rig" node for BP directors

This is so people can do selection logic in rigs from the referencing camera asset.

diff --git a/Engine/Plugins/Cameras/GameplayCameras/Source/GameplayCameras/Private/Directors/BlueprintCameraDirector.cpp b/Engine/Plugins/Cameras/GameplayCameras/Source/GameplayCameras/Private/Directors/BlueprintCameraDirector.cpp
--- a/Engine/Plugins/Cameras/GameplayCameras/Source/GameplayCameras/Private/Directors/BlueprintCameraDirector.cpp
+++ b/Engine/Plugins/Cameras/GameplayCameras/Source/GameplayCameras/Private/Directors/BlueprintCameraDirector.cpp
@@ -296,6 +296,13 @@ void UBlueprintCameraDirectorEvaluator::ActivateCameraRigPrefab(UCameraRigAsset*
 	EvaluationResult.ActiveCameraRigs.Add(CameraRig);
 }
 
+UCameraRigAsset* UBlueprintCameraDirectorEvaluator::GetCameraRig(UCameraRigAsset* CameraRig) const
+{
+	// This function is only here to provide an easy way to pick a camera rig from the referencing 
+	// camera asset, using the custom rig picker. Then we just return it.
+	return CameraRig;
+}
+
 AActor* UBlueprintCameraDirectorEvaluator::FindEvaluationContextOwnerActor(TSubclassOf<AActor> ActorClass) const
 {
 	if (EvaluationContext)
diff --git a/Engine/Plugins/Cameras/GameplayCameras/Source/GameplayCameras/Public/Directors/BlueprintCameraDirector.h b/Engine/Plugins/Cameras/GameplayCameras/Source/GameplayCameras/Public/Directors/BlueprintCameraDirector.h
--- a/Engine/Plugins/Cameras/GameplayCameras/Source/GameplayCameras/Public/Directors/BlueprintCameraDirector.h
+++ b/Engine/Plugins/Cameras/GameplayCameras/Source/GameplayCameras/Public/Directors/BlueprintCameraDirector.h
@@ -171,6 +171,12 @@ public:
 	UFUNCTION(BlueprintCallable, Category="Evaluation")
 	void ActivateCameraRigPrefab(UCameraRigAsset* CameraRig);
 
+	/** Gets a camera rig from the referencing camera asset. */
+	UFUNCTION(BlueprintPure, Category="Evaluation", meta=(HideSelfPin=true))
+	UCameraRigAsset* GetCameraRig(
+			UPARAM(meta=(UseBlueprintCameraDirectorRigPicker=true))
+			UCameraRigAsset* CameraRig) const;
+
 public:
 
 	/**
