Cameras: clean-up dead code

diff --git a/Engine/Plugins/Cameras/GameplayCameras/Source/GameplayCamerasEditor/Private/Customizations/CameraParameterDetailsCustomizations.cpp b/Engine/Plugins/Cameras/GameplayCameras/Source/GameplayCamerasEditor/Private/Customizations/CameraParameterDetailsCustomizations.cpp
--- a/Engine/Plugins/Cameras/GameplayCameras/Source/GameplayCamerasEditor/Private/Customizations/CameraParameterDetailsCustomizations.cpp
+++ b/Engine/Plugins/Cameras/GameplayCameras/Source/GameplayCamerasEditor/Private/Customizations/CameraParameterDetailsCustomizations.cpp
@@ -24,7 +24,6 @@
 #include "Widgets/Images/SImage.h"
 #include "Widgets/Input/SComboButton.h"
 #include "Widgets/SBoxPanel.h"
-#include "Widgets/SCompoundWidget.h"
 #include "Widgets/Text/STextBlock.h"
 
 #define LOCTEXT_NAMESPACE "CameraParameterDetailsCustomization"
@@ -32,21 +31,6 @@
 namespace UE::Cameras
 {
 
-//class SCameraParameterLayoutBox : public SCompoundWidget
-//{
-//public:
-//
-//protected:
-//
-//	// SCompoundWidget interface.
-//	virtual int32 OnPaint(const FPaintArgs& Args, const FGeometry& AllottedGeometry, const FSlateRect& MyCullingRect, FSlateWindowElementList& OutDrawElements, int32 LayerId, const FWidgetStyle& InWidgetStyle, bool bParentEnabled) const override
-//	{
-//		SCompoundWidget::OnPaint(Args, AllottedGeometry, MyCullingRect, OutDrawElements, LayerId, InWidgetStyle, bParentEnabled);
-//
-//		this->GetPaintSpaceGeometry();
-//	}
-//};
-
 void FCameraParameterDetailsCustomization::Register(FPropertyEditorModule& PropertyEditorModule)
 {
 #define UE_CAMERA_VARIABLE_FOR_TYPE(ValueType, ValueName)\
