Cameras: fixes to object tree graph schema logic

- Don't try to create a connection if the 2 pins are already connected. Typically it doesn't do anything, but for array properties it can lead to asserts/crashes since we would try to reconnecte a removed pin.
- Give array pins different names. This fixes cases where the base schema class matches pins by name, and was ending up matching more than one.
- Implement breaking all links on a node. The default implementation doesn't play well since array properties remove pins while iterating on the list of pins.
- Implement relinking pins.

diff --git a/Engine/Plugins/Cameras/GameplayCameras/Source/GameplayCamerasEditor/Private/Editors/ObjectTreeGraphNode.cpp b/Engine/Plugins/Cameras/GameplayCameras/Source/GameplayCamerasEditor/Private/Editors/ObjectTreeGraphNode.cpp
--- a/Engine/Plugins/Cameras/GameplayCameras/Source/GameplayCamerasEditor/Private/Editors/ObjectTreeGraphNode.cpp
+++ b/Engine/Plugins/Cameras/GameplayCameras/Source/GameplayCamerasEditor/Private/Editors/ObjectTreeGraphNode.cpp
@@ -192,7 +192,9 @@ void UObjectTreeGraphNode::CreateNewItemPin(UEdGraphPin* InParentArrayPin)
 
 	InParentArrayPin->Modify();
 
-	UEdGraphPin* ChildPin = CreatePin(PinDirection, ChildPinType, PropertyName);
+	FName ChildPinName = PropertyName;
+	ChildPinName.SetNumber(NewIndex);
+	UEdGraphPin* ChildPin = CreatePin(PinDirection, ChildPinType, ChildPinName);
 	ChildPin->PinFriendlyName = FText::Format(LOCTEXT("ArrayPinFriendlyNameFmt", "{0} {1}"), FText::FromName(PropertyName), NewIndex);
 
 	ChildPin->ParentPin = InParentArrayPin;
@@ -243,6 +245,7 @@ void UObjectTreeGraphNode::RefreshArrayPropertyPinNames()
 			for (int32 PinIndex = 0; PinIndex < Pin->SubPins.Num(); ++PinIndex)
 			{
 				UEdGraphPin* ChildPin = Pin->SubPins[PinIndex];
+				ChildPin->PinName.SetNumber(PinIndex);
 				ChildPin->PinFriendlyName = FText::Format(LOCTEXT("ArrayPinFriendlyNameFmt", "{0} {1}"), FText::FromName(PropertyName), PinIndex);
 			}
 		}
diff --git a/Engine/Plugins/Cameras/GameplayCameras/Source/GameplayCamerasEditor/Private/Editors/ObjectTreeGraphSchema.cpp b/Engine/Plugins/Cameras/GameplayCameras/Source/GameplayCamerasEditor/Private/Editors/ObjectTreeGraphSchema.cpp
--- a/Engine/Plugins/Cameras/GameplayCameras/Source/GameplayCamerasEditor/Private/Editors/ObjectTreeGraphSchema.cpp
+++ b/Engine/Plugins/Cameras/GameplayCameras/Source/GameplayCamerasEditor/Private/Editors/ObjectTreeGraphSchema.cpp
@@ -629,6 +629,11 @@ bool UObjectTreeGraphSchema::TryCreateConnection(UEdGraphPin* A, UEdGraphPin* B)
 {
 	const FPinConnectionResponse Response = CanCreateConnection(A, B);
 
+	if (A->LinkedTo.Contains(B) && B->LinkedTo.Contains(A))
+	{
+		return false;
+	}
+
 	FScopedTransaction Transaction(LOCTEXT("CreateConnection", "Create Connection"));
 
 	bool bModified = true;
@@ -715,6 +720,46 @@ bool UObjectTreeGraphSchema::TryCreateConnection(UEdGraphPin* A, UEdGraphPin* B)
 	return true;
 }
 
+void UObjectTreeGraphSchema::BreakNodeLinks(UEdGraphNode& TargetNode) const
+{
+	const FScopedTransaction Transaction(LOCTEXT("BreakNodeLinks", "Break Node Links"));
+
+	FDelayedPinActions Actions;
+	TArray<UEdGraphPin*> CachedPins = TargetNode.Pins;
+
+#if WITH_EDITOR
+	TSet<UEdGraphNode*> NodeList;
+	NodeList.Add(&TargetNode);
+#endif
+	
+	for (UEdGraphPin* TargetPin : CachedPins)
+	{
+		if (TargetPin && TargetPin->SubPins.Num() == 0)
+		{
+#if WITH_EDITOR
+			for (UEdGraphPin* OtherPin : TargetPin->LinkedTo)
+			{
+				UEdGraphNode* OtherNode = OtherPin ? OtherPin->GetOwningNode() : nullptr;
+				if (OtherNode)
+				{
+					OtherNode->PinConnectionListChanged(OtherPin);
+					NodeList.Add(OtherNode);
+				}
+			}
+#endif
+
+			BreakPinLinks(*TargetPin, false);
+		}
+	}
+	
+#if WITH_EDITOR
+	for (UEdGraphNode* Node : NodeList)
+	{
+		Node->NodeConnectionListChanged();
+	}
+#endif
+}
+
 void UObjectTreeGraphSchema::BreakPinLinks(UEdGraphPin& TargetPin, bool bSendsNodeNotification) const
 {
 	if (TargetPin.LinkedTo.IsEmpty())
diff --git a/Engine/Plugins/Cameras/GameplayCameras/Source/GameplayCamerasEditor/Public/Editors/ObjectTreeGraphSchema.h b/Engine/Plugins/Cameras/GameplayCameras/Source/GameplayCamerasEditor/Public/Editors/ObjectTreeGraphSchema.h
--- a/Engine/Plugins/Cameras/GameplayCameras/Source/GameplayCamerasEditor/Public/Editors/ObjectTreeGraphSchema.h
+++ b/Engine/Plugins/Cameras/GameplayCameras/Source/GameplayCamerasEditor/Public/Editors/ObjectTreeGraphSchema.h
@@ -71,6 +71,7 @@ public:
 	virtual FPinConnectionResponse CanCreateNewNodes(UEdGraphPin* InSourcePin) const override;
 	virtual const FPinConnectionResponse CanCreateConnection(const UEdGraphPin* A, const UEdGraphPin* B) const override;
 	virtual bool TryCreateConnection(UEdGraphPin* A, UEdGraphPin* B) const override;
+	virtual void BreakNodeLinks(UEdGraphNode& TargetNode) const override;
 	virtual void BreakPinLinks(UEdGraphPin& TargetPin, bool bSendsNodeNotification) const override;
 	virtual void BreakSinglePinLink(UEdGraphPin* SourcePin, UEdGraphPin* TargetPin) const override;
 	virtual bool SupportsDropPinOnNode(UEdGraphNode* InTargetNode, const FEdGraphPinType& InSourcePinType, EEdGraphPinDirection InSourcePinDirection, FText& OutErrorMessage) const override;
