Cameras: don't set player controller rotation by default

Since boom arm nodes use controller rotation by default (unless they have an input node specified), then we shouldn't also set controller rotation, otherwise that creates a feedback loop.

diff --git a/Engine/Plugins/Cameras/GameplayCameras/Source/GameplayCameras/Private/GameFramework/GameplayCameraSystemActor.cpp b/Engine/Plugins/Cameras/GameplayCameras/Source/GameplayCameras/Private/GameFramework/GameplayCameraSystemActor.cpp
--- a/Engine/Plugins/Cameras/GameplayCameras/Source/GameplayCameras/Private/GameFramework/GameplayCameraSystemActor.cpp
+++ b/Engine/Plugins/Cameras/GameplayCameras/Source/GameplayCameras/Private/GameFramework/GameplayCameraSystemActor.cpp
@@ -75,6 +75,12 @@ AGameplayCameraSystemActor* AGameplayCameraSystemActor::GetAutoSpawnedCameraSyst
 			SpawnedActor = World->SpawnActor<AGameplayCameraSystemActor>(SpawnParams);
 
 			SpawnedActor->Rename(nullptr, PlayerController);
+			
+			UGameplayCameraSystemComponent* CameraSystemComponent = SpawnedActor->CameraSystemComponent;
+			if (ensure(CameraSystemComponent))
+			{
+				CameraSystemComponent->bSetPlayerControllerRotation = Settings->bAutoSpawnCameraSystemActorSetsControlRotation;
+			}
 		}
 		else
 		{
diff --git a/Engine/Plugins/Cameras/GameplayCameras/Source/GameplayCameras/Public/GameFramework/GameplayCameraSystemComponent.h b/Engine/Plugins/Cameras/GameplayCameras/Source/GameplayCameras/Public/GameFramework/GameplayCameraSystemComponent.h
--- a/Engine/Plugins/Cameras/GameplayCameras/Source/GameplayCameras/Public/GameFramework/GameplayCameraSystemComponent.h
+++ b/Engine/Plugins/Cameras/GameplayCameras/Source/GameplayCameras/Public/GameFramework/GameplayCameraSystemComponent.h
@@ -90,7 +90,7 @@ public:
 	 * This is set on the player controller that this component was activated for.
 	 */
 	UPROPERTY(BlueprintReadWrite, EditAnywhere, Category=Camera)
-	bool bSetPlayerControllerRotation = true;
+	bool bSetPlayerControllerRotation = false;
 
 private:
 	
diff --git a/Engine/Plugins/Cameras/GameplayCameras/Source/GameplayCameras/Public/GameplayCamerasSettings.h b/Engine/Plugins/Cameras/GameplayCameras/Source/GameplayCameras/Public/GameplayCamerasSettings.h
--- a/Engine/Plugins/Cameras/GameplayCameras/Source/GameplayCameras/Public/GameplayCamerasSettings.h
+++ b/Engine/Plugins/Cameras/GameplayCameras/Source/GameplayCameras/Public/GameplayCamerasSettings.h
@@ -27,6 +27,16 @@ public:
 	UPROPERTY(EditAnywhere, Config, Category="General")
 	bool bAutoSpawnCameraSystemActor = true;
 
+	/**
+	 * Whether the automatically spawned camera system actor should set the control rotation on the
+	 * associated player controller. This is useful if camera rigs are managing their own rotation, e.g.
+	 * by specifying input slots on boom arms instead of using the existing control rotation.
+	 * Do not mix handling control rotation via camera nodes, and handling control rotation by calling
+	 * methods like AddYawInput/AddPitchInput/AddRollInput, as this can lead to a feedback loop.
+	 */
+	UPROPERTY(EditAnywhere, Config, Category="General", meta=(EditCondition="bAutoSpawnCameraSystemActor"))
+	bool bAutoSpawnCameraSystemActorSetsControlRotation = false;
+
 	/**
 	 * The number of camera rigs combined in one frame past which the camera system emits a warning.
 	 */
