Cameras: fix rotation roll in the orbit blend

diff --git a/Engine/Plugins/Cameras/GameplayCameras/Source/GameplayCameras/Private/Nodes/Blends/OrbitBlendCameraNode.cpp b/Engine/Plugins/Cameras/GameplayCameras/Source/GameplayCameras/Private/Nodes/Blends/OrbitBlendCameraNode.cpp
--- a/Engine/Plugins/Cameras/GameplayCameras/Source/GameplayCameras/Private/Nodes/Blends/OrbitBlendCameraNode.cpp
+++ b/Engine/Plugins/Cameras/GameplayCameras/Source/GameplayCameras/Private/Nodes/Blends/OrbitBlendCameraNode.cpp
@@ -168,8 +168,12 @@ void FOrbitBlendCameraNodeEvaluator::OnBlendResults(const FCameraNodeBlendParams
 			const FVector3d ToOrbitCenter = ToAim.PointAt(ToClosestParam);
 			const FVector3d BlendedOrbitCenter = FMath::Lerp(FromOrbitCenter, ToOrbitCenter, Factor);
 
+			// The aim direction will get us a yaw/pitch orientation only, so we need to also get the
+			// blended roll from the underlying blend's result.
 			const FVector3d BlendedAimDir = (BlendedOrbitCenter - BlendedLocation).GetUnsafeNormal();
-			const FRotator3d BlendedRotation = BlendedAimDir.ToOrientationRotator();
+			const FRotator3d BlendedRotationNoRoll = BlendedAimDir.ToOrientationRotator();
+			const float BlendedRoll = BlendedResult.CameraPose.GetRotation().Roll;
+			const FRotator3d BlendedRotation(BlendedRotationNoRoll.Pitch, BlendedRotationNoRoll.Yaw, BlendedRoll);
 
 			const double FromOrbitCenterDistance = FVector3d::Distance(FromLocation, FromOrbitCenter);
 			const double ToOrbitCenterDistance = FVector3d::Distance(ToLocation, ToOrbitCenter);
