Cameras: override private variables set on the evaluation context, in case rig parameter overrides are set there

diff --git a/Engine/Plugins/Cameras/GameplayCameras/Source/GameplayCameras/Private/Core/BlendStackCameraNode.cpp b/Engine/Plugins/Cameras/GameplayCameras/Source/GameplayCameras/Private/Core/BlendStackCameraNode.cpp
--- a/Engine/Plugins/Cameras/GameplayCameras/Source/GameplayCameras/Private/Core/BlendStackCameraNode.cpp
+++ b/Engine/Plugins/Cameras/GameplayCameras/Source/GameplayCameras/Private/Core/BlendStackCameraNode.cpp
@@ -623,8 +623,10 @@ void FTransientBlendStackCameraNodeEvaluator::InternalPreBlendPrepare(TArrayView
 		CurResult.VariableTable.OverrideAll(OutResult.VariableTable);
 
 		// Override it with whatever the evaluation context has set on its result.
+		// Evaluation contexts may have private variables we need to pass along, such as when rig parameter
+		// overrides have been set on them.
 		const FCameraNodeEvaluationResult& ContextResult(ResolvedEntry.Context->GetInitialResult());
-		CurResult.VariableTable.OverrideAll(ContextResult.VariableTable);
+		CurResult.VariableTable.Override(ContextResult.VariableTable, ECameraVariableTableFilter::AllPublic | ECameraVariableTableFilter::Private);
 
 		// Gather input parameters if needed (and remember if it was indeed needed).
 		if (!Entry.bInputRunThisFrame)
@@ -980,9 +982,11 @@ void FPersistentBlendStackCameraNodeEvaluator::InternalUpdate(TArrayView<FResolv
 				CurResult.PostProcessSettings.OverrideAll(OutResult.PostProcessSettings);
 
 				// Override it with whatever the evaluation context has set on its result.
+				// Evaluation contexts may have private variables we need to pass along, such as when rig parameter
+				// overrides have been set on them.
 				const FCameraNodeEvaluationResult& ContextResult(ResolvedEntry.Context->GetInitialResult());
 				CurResult.CameraPose.OverrideChanged(ContextResult.CameraPose);
-				CurResult.VariableTable.OverrideAll(ContextResult.VariableTable);
+				CurResult.VariableTable.Override(ContextResult.VariableTable, ECameraVariableTableFilter::AllPublic | ECameraVariableTableFilter::Private);
 
 				// Setup flags.
 				CurResult.bIsCameraCut = OutResult.bIsCameraCut || ContextResult.bIsCameraCut;
