Cameras: fix SetCameraRigParameters BP node not restoring pins correctly on load

diff --git a/Engine/Plugins/Cameras/GameplayCameras/Source/GameplayCamerasUncookedOnly/Private/BlueprintGraph/K2Node_SetCameraRigParameters.cpp b/Engine/Plugins/Cameras/GameplayCameras/Source/GameplayCamerasUncookedOnly/Private/BlueprintGraph/K2Node_SetCameraRigParameters.cpp
--- a/Engine/Plugins/Cameras/GameplayCameras/Source/GameplayCamerasUncookedOnly/Private/BlueprintGraph/K2Node_SetCameraRigParameters.cpp
+++ b/Engine/Plugins/Cameras/GameplayCameras/Source/GameplayCamerasUncookedOnly/Private/BlueprintGraph/K2Node_SetCameraRigParameters.cpp
@@ -48,8 +48,21 @@ void UK2Node_SetCameraRigParameters::ReallocatePinsDuringReconstruction(TArray<U
 	TArrayView<UEdGraphPin* const> PinsToSearch = OldPins;
 	if (UCameraRigAsset* CameraRig = GetCameraRig(&PinsToSearch))
 	{
+		// The camera rig might not be loaded yet.
+		PreloadObject(CameraRig);
+		for (UCameraRigInterfaceParameter* InterfaceParameter : CameraRig->Interface.InterfaceParameters)
+		{
+			PreloadObject(InterfaceParameter);
+			if (InterfaceParameter)
+			{
+				PreloadObject(InterfaceParameter->PrivateVariable);
+			}
+		}
+
 		CreatePinsForCameraRig(CameraRig);
 	}
+
+	RestoreSplitPins(OldPins);
 }
 
 void UK2Node_SetCameraRigParameters::PostPlacedNewNode()
@@ -202,14 +215,14 @@ void UK2Node_SetCameraRigParameters::ExpandNode(class FKismetCompilerContext& Co
 
 		// Set or connect the parameter value argument.
 		UEdGraphPin* CallSetParameterValuePin = CallSetParameter->FindPinChecked(TEXT("ParameterValue"));
+		CallSetParameterValuePin->DefaultValue = RigParameterPin->DefaultValue;
+		CallSetParameterValuePin->DefaultTextValue = RigParameterPin->DefaultTextValue;
+		CallSetParameterValuePin->AutogeneratedDefaultValue = RigParameterPin->AutogeneratedDefaultValue;
+		CallSetParameterValuePin->DefaultObject = RigParameterPin->DefaultObject;
 		if (RigParameterPin->LinkedTo.Num() > 0)
 		{
 			CompilerContext.MovePinLinksToIntermediate(*RigParameterPin, *CallSetParameterValuePin);
 		}
-		else
-		{
-			CallSetParameterValuePin->DefaultObject = RigParameterPin->DefaultObject;
-		}
 
 		// Connect the SetXxxParameter node to the chain of other similar nodes. The SetCameraRigParameters node
 		// effectively transforms into a chain of individual setter function calls.
