Cameras: add standard shortcut support to camera variable collection editor

diff --git a/Engine/Plugins/Cameras/GameplayCameras/Source/GameplayCamerasEditor/Private/Commands/CameraVariableCollectionEditorCommands.cpp b/Engine/Plugins/Cameras/GameplayCameras/Source/GameplayCamerasEditor/Private/Commands/CameraVariableCollectionEditorCommands.cpp
--- a/Engine/Plugins/Cameras/GameplayCameras/Source/GameplayCamerasEditor/Private/Commands/CameraVariableCollectionEditorCommands.cpp
+++ b/Engine/Plugins/Cameras/GameplayCameras/Source/GameplayCamerasEditor/Private/Commands/CameraVariableCollectionEditorCommands.cpp
@@ -29,9 +29,9 @@ void FCameraVariableCollectionEditorCommands::RegisterCommands()
 	UI_COMMAND(CreateVariable, "Create Variable", "Adds a new camera variable to the collection",
 			EUserInterfaceActionType::Button, FInputChord());
 	UI_COMMAND(RenameVariable, "Rename Variable", "Renames a camera variable",
-			EUserInterfaceActionType::Button, FInputChord());
+			EUserInterfaceActionType::Button, FInputChord(EKeys::F2));
 	UI_COMMAND(DeleteVariable, "Delete Variable", "Removes a camera variable from the collection",
-			EUserInterfaceActionType::Button, FInputChord());
+			EUserInterfaceActionType::Button, FInputChord(EKeys::Delete));
 }
 
 }  // namespace UE::Cameras
