Cameras: add utility function to safely get the name of an RTTI-enabled class

diff --git a/Engine/Plugins/Cameras/GameplayCameras/Source/GameplayCameras/Private/Core/CameraObjectRtti.cpp b/Engine/Plugins/Cameras/GameplayCameras/Source/GameplayCameras/Private/Core/CameraObjectRtti.cpp
--- a/Engine/Plugins/Cameras/GameplayCameras/Source/GameplayCameras/Private/Core/CameraObjectRtti.cpp
+++ b/Engine/Plugins/Cameras/GameplayCameras/Source/GameplayCameras/Private/Core/CameraObjectRtti.cpp
@@ -26,16 +26,16 @@ void FCameraObjectTypeRegistry::RegisterType(FCameraObjectTypeID TypeID, FCamera
 	TypeInfos.Insert(TypeID.GetTypeID(), MoveTemp(TypeInfo));
 }
 
-FCameraObjectTypeID FCameraObjectTypeRegistry::FindTypeByName(const FName& TypeName)
+FCameraObjectTypeID FCameraObjectTypeRegistry::FindTypeByName(const FName& TypeName) const
 {
-	if (uint32* RegisteredTypeID = TypeIDsByName.Find(TypeName))
+	if (const uint32* RegisteredTypeID = TypeIDsByName.Find(TypeName))
 	{
 		return FCameraObjectTypeID{ *RegisteredTypeID };
 	}
 	return FCameraObjectTypeID::Invalid();
 }
 
-const FCameraObjectTypeInfo* FCameraObjectTypeRegistry::GetTypeInfo(FCameraObjectTypeID TypeID)
+const FCameraObjectTypeInfo* FCameraObjectTypeRegistry::GetTypeInfo(FCameraObjectTypeID TypeID) const
 {
 	if (ensureMsgf(
 				TypeID.IsValid() && TypeInfos.IsValidIndex(TypeID.GetTypeID()),
@@ -47,6 +47,19 @@ const FCameraObjectTypeInfo* FCameraObjectTypeRegistry::GetTypeInfo(FCameraObjec
 	return nullptr;
 }
 
+FName FCameraObjectTypeRegistry::GetTypeNameSafe(FCameraObjectTypeID TypeID) const
+{
+	if (TypeID.IsValid())
+	{
+		if (const FCameraObjectTypeInfo* TypeInfo = GetTypeInfo(TypeID))
+		{
+			return TypeInfo->TypeName;
+		}
+	}
+
+	return NAME_None;
+}
+
 void FCameraObjectTypeRegistry::ConstructObject(FCameraObjectTypeID TypeID, void* Ptr)
 {
 	if (ensureMsgf(TypeInfos.IsValidIndex(TypeID.GetTypeID()), TEXT("Invalid camera object type ID!")))
diff --git a/Engine/Plugins/Cameras/GameplayCameras/Source/GameplayCameras/Public/Core/CameraObjectRtti.h b/Engine/Plugins/Cameras/GameplayCameras/Source/GameplayCameras/Public/Core/CameraObjectRtti.h
--- a/Engine/Plugins/Cameras/GameplayCameras/Source/GameplayCameras/Public/Core/CameraObjectRtti.h
+++ b/Engine/Plugins/Cameras/GameplayCameras/Source/GameplayCameras/Public/Core/CameraObjectRtti.h
@@ -85,8 +85,9 @@ public:
 	GAMEPLAYCAMERAS_API static FCameraObjectTypeRegistry& Get();
 
 	GAMEPLAYCAMERAS_API void RegisterType(FCameraObjectTypeID TypeID, FCameraObjectTypeInfo&& TypeInfo);
-	GAMEPLAYCAMERAS_API FCameraObjectTypeID FindTypeByName(const FName& TypeName);
-	GAMEPLAYCAMERAS_API const FCameraObjectTypeInfo* GetTypeInfo(FCameraObjectTypeID TypeID);
+	GAMEPLAYCAMERAS_API FCameraObjectTypeID FindTypeByName(const FName& TypeName) const;
+	GAMEPLAYCAMERAS_API const FCameraObjectTypeInfo* GetTypeInfo(FCameraObjectTypeID TypeID) const;
+	GAMEPLAYCAMERAS_API FName GetTypeNameSafe(FCameraObjectTypeID TypeID) const;
 
 	GAMEPLAYCAMERAS_API void ConstructObject(FCameraObjectTypeID TypeID, void* Ptr);
 
