Cameras: tweak the look of some debug display

diff --git a/Engine/Plugins/Cameras/GameplayCameras/Source/GameplayCameras/Private/Core/BlendStackCameraNode.cpp b/Engine/Plugins/Cameras/GameplayCameras/Source/GameplayCameras/Private/Core/BlendStackCameraNode.cpp
--- a/Engine/Plugins/Cameras/GameplayCameras/Source/GameplayCameras/Private/Core/BlendStackCameraNode.cpp
+++ b/Engine/Plugins/Cameras/GameplayCameras/Source/GameplayCameras/Private/Core/BlendStackCameraNode.cpp
@@ -1114,7 +1114,7 @@ FBlendStackSummaryCameraDebugBlock::FBlendStackSummaryCameraDebugBlock(const FBl
 
 void FBlendStackSummaryCameraDebugBlock::OnDebugDraw(const FCameraDebugBlockDrawParams& Params, FCameraDebugRenderer& Renderer)
 {
-	Renderer.AddText(TEXT("%d entries, %s stack"), NumEntries, *UEnum::GetValueAsString(BlendStackType));
+	Renderer.AddText(TEXT("%d entries"), NumEntries);
 }
 
 void FBlendStackSummaryCameraDebugBlock::OnSerialize(FArchive& Ar)
@@ -1147,7 +1147,7 @@ void FBlendStackCameraDebugBlock::OnDebugDraw(const FCameraDebugBlockDrawParams&
 	{
 		const FEntryDebugInfo& Entry(Entries[Index]);
 
-		Renderer.AddText(TEXT("{cam_passive}[%d]{cam_highlighted} %s\n"), Index + 1, *Entry.CameraRigName);
+		Renderer.AddText(TEXT("{cam_passive}[%d] {cam_notice}%s{cam_default}\n"), Index + 1, *Entry.CameraRigName);
 
 		if (ChildrenView.IsValidIndex(Index))
 		{
diff --git a/Engine/Plugins/Cameras/GameplayCameras/Source/GameplayCameras/Private/Core/BlendStackRootCameraNode.cpp b/Engine/Plugins/Cameras/GameplayCameras/Source/GameplayCameras/Private/Core/BlendStackRootCameraNode.cpp
--- a/Engine/Plugins/Cameras/GameplayCameras/Source/GameplayCameras/Private/Core/BlendStackRootCameraNode.cpp
+++ b/Engine/Plugins/Cameras/GameplayCameras/Source/GameplayCameras/Private/Core/BlendStackRootCameraNode.cpp
@@ -381,14 +381,14 @@ void FBlendStackRootCameraDebugBlock::OnDebugDraw(const FCameraDebugBlockDrawPar
 {
 	TArrayView<FCameraDebugBlock*> ChildrenView(GetChildren());
 
-	Renderer.AddText(TEXT("{cam_passive}<Blend>\n"));
+	Renderer.AddText(TEXT("{cam_passive}<Blend>{cam_default}\n"));
 	Renderer.AddIndent();
 	ChildrenView[0]->DebugDraw(Params, Renderer);
 	Renderer.RemoveIndent();
 
 	if (!BlendedParameterOverridesEntries.IsEmpty())
 	{
-		Renderer.AddText(TEXT("{cam_passive}<%d Merged Camera Rigs>\n"), BlendedParameterOverridesEntries.Num());
+		Renderer.AddText(TEXT("{cam_passive}<%d Merged Camera Rigs>{cam_default}\n"), BlendedParameterOverridesEntries.Num());
 		Renderer.AddIndent();
 		{
 			int32 ChildIndex = 0;
@@ -402,7 +402,7 @@ void FBlendStackRootCameraDebugBlock::OnDebugDraw(const FCameraDebugBlockDrawPar
 		Renderer.RemoveIndent();
 	}
 
-	Renderer.AddText(TEXT("{cam_passive}<CameraRig %s>\n"), *CameraRigAssetName);
+	Renderer.AddText(TEXT("{cam_passive}<CameraRig> {cam_default}Running {cam_notice}%s{cam_default}\n"), *CameraRigAssetName);
 	Renderer.AddIndent();
 	ChildrenView[2]->DebugDraw(Params, Renderer);
 	Renderer.RemoveIndent();
diff --git a/Engine/Plugins/Cameras/GameplayCameras/Source/GameplayCameras/Private/Debug/BlendStacksCameraDebugBlock.cpp b/Engine/Plugins/Cameras/GameplayCameras/Source/GameplayCameras/Private/Debug/BlendStacksCameraDebugBlock.cpp
--- a/Engine/Plugins/Cameras/GameplayCameras/Source/GameplayCameras/Private/Debug/BlendStacksCameraDebugBlock.cpp
+++ b/Engine/Plugins/Cameras/GameplayCameras/Source/GameplayCameras/Private/Debug/BlendStacksCameraDebugBlock.cpp
@@ -46,10 +46,8 @@ void FBlendStacksCameraDebugBlock::OnDebugDraw(const FCameraDebugBlockDrawParams
 
 	for (int32 Index = 0; Index < MaxIndex; ++Index)
 	{
-		Renderer.SetTextColor(Colors.Notice);
 		Renderer.AddText(BlendStackNames.IsValidIndex(Index) ? BlendStackNames[Index] : TEXT("<unnamed blend stack>"));
 		Renderer.NewLine();
-		Renderer.SetTextColor(Colors.Default);
 
 		if (BlendStackBlocks.IsValidIndex(Index))
 		{
diff --git a/Engine/Plugins/Cameras/GameplayCameras/Source/GameplayCameras/Private/Debug/CameraDirectorTreeDebugBlock.cpp b/Engine/Plugins/Cameras/GameplayCameras/Source/GameplayCameras/Private/Debug/CameraDirectorTreeDebugBlock.cpp
--- a/Engine/Plugins/Cameras/GameplayCameras/Source/GameplayCameras/Private/Debug/CameraDirectorTreeDebugBlock.cpp
+++ b/Engine/Plugins/Cameras/GameplayCameras/Source/GameplayCameras/Private/Debug/CameraDirectorTreeDebugBlock.cpp
@@ -3,6 +3,8 @@
 #include "Debug/CameraDirectorTreeDebugBlock.h"
 
 #include "Core/CameraAsset.h"
+#include "Core/CameraDirector.h"
+#include "Core/CameraDirectorEvaluator.h"
 #include "Core/CameraEvaluationContext.h"
 #include "Core/CameraEvaluationContextStack.h"
 #include "Debug/CameraDebugBlockBuilder.h"
@@ -56,11 +58,17 @@ void FCameraDirectorTreeDebugBlock::InitializeEntry(TSharedPtr<FCameraEvaluation
 {
 	if (Context)
 	{
+		const FCameraObjectTypeRegistry& TypeRegistry = FCameraObjectTypeRegistry::Get();
+		const FName ContextTypeName = TypeRegistry.GetTypeNameSafe(Context->GetTypeID());
+
 		const UObject* ContextOwner = Context->GetOwner();
+		const FCameraDirectorEvaluator* DirectorEvaluator = Context->GetDirectorEvaluator();
 
+		EntryDebugInfo.ContextClassName = ContextTypeName;
+		EntryDebugInfo.OwnerName = *GetPathNameSafe(ContextOwner);
+		EntryDebugInfo.OwnerClassName = ContextOwner ? ContextOwner->GetClass()->GetFName() : NAME_None;
 		EntryDebugInfo.CameraAssetName = GetNameSafe(Context->GetCameraAsset());
-		EntryDebugInfo.OwnerName = *GetNameSafe(ContextOwner);
-		EntryDebugInfo.OwnerClassName = *GetNameSafe(ContextOwner ? ContextOwner->GetClass() : nullptr);
+		EntryDebugInfo.CameraDirectorClassName = DirectorEvaluator->GetCameraDirector()->GetFName();
 		EntryDebugInfo.InitialContextTransform = Context->GetInitialResult().CameraPose.GetTransform();
 		EntryDebugInfo.bIsValid = true;
 
@@ -95,8 +103,8 @@ void FCameraDirectorTreeDebugBlock::OnDebugDraw(const FCameraDebugBlockDrawParam
 
 	const int32 MinNum = FMath::Min(ChildrenView.Num(), CameraDirectors.Num());
 
-	Renderer.SetTextColor(Colors.Hightlighted);
-	Renderer.AddText("Inactive Directors\n\n");
+	Renderer.SetTextColor(Colors.Notice);
+	Renderer.AddText("Inactive Directors\n");
 	Renderer.SetTextColor(Colors.Default);
 	Renderer.AddIndent();
 
@@ -105,18 +113,36 @@ void FCameraDirectorTreeDebugBlock::OnDebugDraw(const FCameraDebugBlockDrawParam
 		if (Index == CameraDirectors.Num() - 1)
 		{
 			Renderer.RemoveIndent();
-			Renderer.NewLine();
 
 			Renderer.SetTextColor(Colors.Notice);
-			Renderer.AddText("Active Director\n\n");
+			Renderer.AddText("Active Director\n");
 			Renderer.SetTextColor(Colors.Default);
 			Renderer.AddIndent();
 		}
 
+		Renderer.AddText(TEXT("{cam_passive}[%d]{cam_default} "), Index + 1);
+
 		const FDirectorDebugInfo& EntryDebugInfo(CameraDirectors[Index]);
 		if (EntryDebugInfo.bIsValid)
 		{
-			Renderer.AddText(TEXT("Camera asset: {cam_notice}%s{cam_default}\n"), *EntryDebugInfo.CameraAssetName);
+			Renderer.AddText(TEXT("{cam_passive}[%s]{cam_default}"), 
+					*EntryDebugInfo.CameraDirectorClassName.ToString());
+			Renderer.AddIndent();
+			{
+				Renderer.AddText(TEXT("Context {cam_passive}[%s]{cam_default}\n"), *EntryDebugInfo.ContextClassName.ToString());
+
+				Renderer.AddText(TEXT("Owned by {cam_passive}[%s]{cam_default}\n"), *EntryDebugInfo.OwnerClassName.ToString());
+				Renderer.AddIndent();
+				{
+					Renderer.AddText(*EntryDebugInfo.OwnerName);
+				}
+				Renderer.RemoveIndent();
+
+				Renderer.AddText(TEXT("{cam_passive}From camera asset {cam_notice}%s{cam_default}\n"), 
+						*EntryDebugInfo.CameraAssetName);
+			}
+			Renderer.RemoveIndent();
+
 			Renderer.DrawCoordinateSystem(EntryDebugInfo.InitialContextTransform);
 		}
 		else
@@ -144,9 +170,11 @@ void FCameraDirectorTreeDebugBlock::OnSerialize(FArchive& Ar)
 
 FArchive& operator<< (FArchive& Ar, FCameraDirectorTreeDebugBlock::FDirectorDebugInfo& DirectorDebugInfo)
 {
-	Ar << DirectorDebugInfo.CameraAssetName;
+	Ar << DirectorDebugInfo.ContextClassName;
 	Ar << DirectorDebugInfo.OwnerClassName;
 	Ar << DirectorDebugInfo.OwnerName;
+	Ar << DirectorDebugInfo.CameraAssetName;
+	Ar << DirectorDebugInfo.CameraDirectorClassName;
 	Ar << DirectorDebugInfo.InitialContextTransform;
 	Ar << DirectorDebugInfo.bIsValid;
 	return Ar;
diff --git a/Engine/Plugins/Cameras/GameplayCameras/Source/GameplayCameras/Public/Debug/CameraDirectorTreeDebugBlock.h b/Engine/Plugins/Cameras/GameplayCameras/Source/GameplayCameras/Public/Debug/CameraDirectorTreeDebugBlock.h
--- a/Engine/Plugins/Cameras/GameplayCameras/Source/GameplayCameras/Public/Debug/CameraDirectorTreeDebugBlock.h
+++ b/Engine/Plugins/Cameras/GameplayCameras/Source/GameplayCameras/Public/Debug/CameraDirectorTreeDebugBlock.h
@@ -43,9 +43,11 @@ private:
 
 	struct FDirectorDebugInfo
 	{
-		FString CameraAssetName;
-		FString OwnerClassName;
+		FName ContextClassName;
+		FName OwnerClassName;
 		FString OwnerName;
+		FString CameraAssetName;
+		FName CameraDirectorClassName;
 		FTransform3d InitialContextTransform = FTransform3d::Identity;
 		bool bIsValid = false;
 	};
