Cameras: add detail customization to offer camera filmback presets to the user

diff --git a/Engine/Plugins/Cameras/GameplayCameras/Source/GameplayCamerasEditor/GameplayCamerasEditor.Build.cs b/Engine/Plugins/Cameras/GameplayCameras/Source/GameplayCamerasEditor/GameplayCamerasEditor.Build.cs
--- a/Engine/Plugins/Cameras/GameplayCameras/Source/GameplayCamerasEditor/GameplayCamerasEditor.Build.cs
+++ b/Engine/Plugins/Cameras/GameplayCameras/Source/GameplayCamerasEditor/GameplayCamerasEditor.Build.cs
@@ -23,6 +23,7 @@ public class GameplayCamerasEditor : ModuleRules
 				"AssetDefinition",
 				"AssetRegistry",
 				"BlueprintGraph",
+				"CinematicCamera",
 				"Core",
 				"CoreUObject",
 				"DeveloperSettings",
diff --git a/Engine/Plugins/Cameras/GameplayCameras/Source/GameplayCamerasEditor/Private/Customizations/FilmbackCameraNodeDetailsCustomization.cpp b/Engine/Plugins/Cameras/GameplayCameras/Source/GameplayCamerasEditor/Private/Customizations/FilmbackCameraNodeDetailsCustomization.cpp
new file mode 100644
--- /dev/null
+++ b/Engine/Plugins/Cameras/GameplayCameras/Source/GameplayCamerasEditor/Private/Customizations/FilmbackCameraNodeDetailsCustomization.cpp
@@ -0,0 +1,192 @@
+// Copyright Epic Games, Inc. All Rights Reserved.
+
+#include "Customizations/FilmbackCameraNodeDetailsCustomization.h"
+
+#include "CineCameraSettings.h"
+#include "DetailCategoryBuilder.h"
+#include "DetailLayoutBuilder.h"
+#include "DetailWidgetRow.h"
+#include "IDetailChildrenBuilder.h"
+#include "IDetailGroup.h"
+#include "Nodes/Common/FilmbackCameraNode.h"
+#include "PropertyCustomizationHelpers.h"
+#include "PropertyEditorModule.h"
+#include "PropertyHandle.h"
+#include "ScopedTransaction.h"
+#include "Widgets/Layout/SBorder.h"
+#include "Widgets/SBoxPanel.h"
+#include "Widgets/SNullWidget.h"
+#include "Widgets/Text/STextBlock.h"
+
+#define LOCTEXT_NAMESPACE "FilmbackCameraNodeDetailsCustomization"
+
+namespace UE::Cameras
+{
+
+TSharedRef<IDetailCustomization> FFilmbackCameraNodeDetailsCustomization::MakeInstance()
+{
+	return MakeShared<FFilmbackCameraNodeDetailsCustomization>();
+}
+
+void FFilmbackCameraNodeDetailsCustomization::CustomizeDetails(IDetailLayoutBuilder& DetailBuilder)
+{
+	const FName SensorWidthPropertyName = GET_MEMBER_NAME_CHECKED(UFilmbackCameraNode, SensorWidth);
+	const FName SensorHeightPropertyName = GET_MEMBER_NAME_CHECKED(UFilmbackCameraNode, SensorHeight);
+
+	SensorWidthProperty = DetailBuilder.GetProperty(SensorWidthPropertyName);
+	SensorHeightProperty = DetailBuilder.GetProperty(SensorHeightPropertyName);
+
+	DetailBuilder.HideProperty(SensorWidthProperty);
+	DetailBuilder.HideProperty(SensorHeightProperty);
+
+	IDetailCategoryBuilder& FilmbackCategory = DetailBuilder.EditCategory(TEXT("Filmback"));
+
+	IDetailGroup& SensorSizeGroup = FilmbackCategory.AddGroup(
+			TEXT("SensorSize"), 
+			LOCTEXT("SensorSizeRow", "Sensor Size"),
+			false,
+			true);
+
+	FDetailWidgetRow& FilmbackPresetsRow = SensorSizeGroup.AddWidgetRow();
+	FilmbackPresetsRow
+		.NameContent()
+		[
+			SNew(SHorizontalBox)
+			+SHorizontalBox::Slot()
+			.Padding(FMargin(0.0f, 1.0f))
+			.FillWidth(1)
+			[
+				SNew(SBorder)
+				.BorderImage(FAppStyle::GetBrush(TEXT("PropertyWindow.NoOverlayColor")))
+				.Padding(FMargin(0.0f, 2.0f))
+				.VAlign(VAlign_Center)
+				[
+					SNew(STextBlock)
+					.Font(FAppStyle::GetFontStyle(TEXT("PropertyWindow.NormalFont")))
+					.Text(LOCTEXT("SensorSizePresets", "Sensor Size Presets"))
+				]
+			]
+		]
+		.ValueContent()
+		.MaxDesiredWidth(0.f)
+		[
+			SAssignNew(PresetComboBox, SComboBox<TSharedPtr<FText>>)
+			.OptionsSource(&PresetComboList)
+			.IsEnabled(this, &FFilmbackCameraNodeDetailsCustomization::IsPresetEnabled)
+			.OnGenerateWidget(this, &FFilmbackCameraNodeDetailsCustomization::MakePresetComboWidget)
+			.OnSelectionChanged(this, &FFilmbackCameraNodeDetailsCustomization::OnPresetChanged)
+			.ContentPadding(2)
+			.Content()
+			[
+				SNew(STextBlock)
+				.Font(IDetailLayoutBuilder::GetDetailFont())
+				.Text(this, &FFilmbackCameraNodeDetailsCustomization::GetPresetComboBoxContent)
+				.ToolTipText(this, &FFilmbackCameraNodeDetailsCustomization::GetPresetComboBoxContent)
+			]
+		];
+
+	SensorSizeGroup.AddPropertyRow(SensorWidthProperty.ToSharedRef());
+	SensorSizeGroup.AddPropertyRow(SensorHeightProperty.ToSharedRef());
+
+	BuildPresetComboList();
+}
+
+void FFilmbackCameraNodeDetailsCustomization::BuildPresetComboList()
+{
+	const TArray<FNamedFilmbackPreset>& Presets = UCineCameraSettings::GetFilmbackPresets();
+
+	const int32 NumPresets = Presets.Num();
+	PresetComboList.Reserve(NumPresets + 1);
+
+	PresetComboList.Add(MakeShared<FText>(LOCTEXT("CustomPreset", "Custom...")));
+	for (const FNamedFilmbackPreset& Preset : Presets)
+	{
+		PresetComboList.Add(MakeShared<FText>(FText::FromString(Preset.Name)));
+	}
+}
+
+bool FFilmbackCameraNodeDetailsCustomization::IsPresetEnabled() const
+{
+	if (SensorHeightProperty && SensorWidthProperty)
+	{
+		return 
+			SensorHeightProperty->IsEditable() &&
+			SensorWidthProperty->IsEditable() &&
+			FSlateApplication::Get().GetNormalExecutionAttribute().Get();
+	}
+	return false;
+}
+
+TSharedRef<SWidget> FFilmbackCameraNodeDetailsCustomization::MakePresetComboWidget(TSharedPtr<FText> InItem)
+{
+	return
+		SNew(STextBlock)
+		.Text(*InItem)
+		.Font(IDetailLayoutBuilder::GetDetailFont());
+}
+
+void FFilmbackCameraNodeDetailsCustomization::OnPresetChanged(TSharedPtr<FText> NewSelection, ESelectInfo::Type SelectInfo)
+{
+	if (SelectInfo == ESelectInfo::Direct)
+	{
+		return;
+	}
+
+	const FString NewPresetName = NewSelection->ToString();
+
+	const TArray<FNamedFilmbackPreset>& Presets = UCineCameraSettings::GetFilmbackPresets();
+	for (const FNamedFilmbackPreset& Preset : Presets)
+	{
+		if (Preset.Name == NewPresetName)
+		{
+			const FName ParameterValueName = GET_MEMBER_NAME_CHECKED(FFloatCameraParameter, Value);
+
+			const FScopedTransaction Transaction(LOCTEXT("ChangeFilmbackPreset", "Change Filmback Preset"));
+			
+			TSharedPtr<IPropertyHandle> SensorWidthValueProperty = SensorWidthProperty->GetChildHandle(ParameterValueName);
+			TSharedPtr<IPropertyHandle> SensorHeightValueProperty = SensorHeightProperty->GetChildHandle(ParameterValueName);
+
+			SensorWidthValueProperty->SetValue(Preset.FilmbackSettings.SensorWidth);
+			SensorHeightValueProperty->SetValue(Preset.FilmbackSettings.SensorHeight);
+			
+			break;
+		}
+	}
+}
+
+FText FFilmbackCameraNodeDetailsCustomization::GetPresetComboBoxContent() const
+{
+	const FName ParameterValueName = GET_MEMBER_NAME_CHECKED(FFloatCameraParameter, Value);
+
+	TSharedPtr<IPropertyHandle> SensorWidthValueProperty = SensorWidthProperty->GetChildHandle(ParameterValueName);
+	TSharedPtr<IPropertyHandle> SensorHeightValueProperty = SensorHeightProperty->GetChildHandle(ParameterValueName);
+
+	float CurSensorWidth;
+	float CurSensorHeight;
+	if (SensorWidthValueProperty->GetValue(CurSensorWidth) == FPropertyAccess::MultipleValues ||
+			SensorHeightValueProperty->GetValue(CurSensorHeight) == FPropertyAccess::MultipleValues)
+	{
+		return LOCTEXT("MultipleValues", "Multiple Values");
+	}
+
+	const TArray<FNamedFilmbackPreset>& Presets = UCineCameraSettings::GetFilmbackPresets();
+	for (int32 PresetIndex = 0, NumPresets = Presets.Num(); PresetIndex < NumPresets; ++PresetIndex)
+	{
+		const FNamedFilmbackPreset& Preset(Presets[PresetIndex]);
+		if ((Preset.FilmbackSettings.SensorWidth == CurSensorWidth) && 
+				(Preset.FilmbackSettings.SensorHeight == CurSensorHeight))
+		{
+			if (PresetComboList.IsValidIndex(PresetIndex + 1))
+			{
+				return *PresetComboList[PresetIndex + 1];
+			}
+		}
+	}
+
+	return *PresetComboList[0];
+}
+
+}  // namespace UE::Cameras
+
+#undef LOCTEXT_NAMESPACE
+
diff --git a/Engine/Plugins/Cameras/GameplayCameras/Source/GameplayCamerasEditor/Private/Customizations/SingleCameraDirectorDetailsCustomization.cpp b/Engine/Plugins/Cameras/GameplayCameras/Source/GameplayCamerasEditor/Private/Customizations/SingleCameraDirectorDetailsCustomization.cpp
deleted file mode 100644
diff --git a/Engine/Plugins/Cameras/GameplayCameras/Source/GameplayCamerasEditor/Private/GameplayCamerasEditorModule.cpp b/Engine/Plugins/Cameras/GameplayCameras/Source/GameplayCamerasEditor/Private/GameplayCamerasEditorModule.cpp
--- a/Engine/Plugins/Cameras/GameplayCameras/Source/GameplayCamerasEditor/Private/GameplayCamerasEditorModule.cpp
+++ b/Engine/Plugins/Cameras/GameplayCameras/Source/GameplayCamerasEditor/Private/GameplayCamerasEditorModule.cpp
@@ -17,7 +17,7 @@
 #include "Customizations/CameraRigPtrDetailsCustomization.h"
 #include "Customizations/CameraProxyTableDetailsCustomization.h"
 #include "Customizations/CameraVariableReferenceDetailsCustomizations.h"
-#include "Customizations/SingleCameraDirectorDetailsCustomization.h"
+#include "Customizations/FilmbackCameraNodeDetailsCustomization.h"
 #include "Debug/CameraDebugCategories.h"
 #include "Debugger/SBlendStacksDebugPanel.h"
 #include "Debugger/SCameraNodeTreeDebugPanel.h"
@@ -490,6 +490,9 @@ private:
 					&FCameraRigPtrDetailsCustomization::MakeInstance));
 		PropertyEditorModule.RegisterCustomPropertyTypeLayout("CameraRigAssetReference", FOnGetPropertyTypeCustomizationInstance::CreateStatic(
 					&FCameraRigAssetReferenceDetailsCustomization::MakeInstance));
+
+		PropertyEditorModule.RegisterCustomClassLayout("FilmbackCameraNode", FOnGetDetailCustomizationInstance::CreateStatic(
+					&FFilmbackCameraNodeDetailsCustomization::MakeInstance));
 	}
 
 	void UnregisterDetailsCustomizations()
@@ -506,6 +509,8 @@ private:
 			PropertyEditorModule->UnregisterCustomPropertyTypeLayout("CameraRigProxyTableEntry");
 			PropertyEditorModule->UnregisterCustomPropertyTypeLayout("CameraRigAsset");
 			PropertyEditorModule->UnregisterCustomPropertyTypeLayout("CameraRigAssetReference");
+
+			PropertyEditorModule->UnregisterCustomClassLayout("FilmbackCameraNode");
 		}
 	}
 
diff --git a/Engine/Plugins/Cameras/GameplayCameras/Source/GameplayCamerasEditor/Public/Customizations/FilmbackCameraNodeDetailsCustomization.h b/Engine/Plugins/Cameras/GameplayCameras/Source/GameplayCamerasEditor/Public/Customizations/FilmbackCameraNodeDetailsCustomization.h
new file mode 100644
--- /dev/null
+++ b/Engine/Plugins/Cameras/GameplayCameras/Source/GameplayCamerasEditor/Public/Customizations/FilmbackCameraNodeDetailsCustomization.h
@@ -0,0 +1,50 @@
+// Copyright Epic Games, Inc. All Rights Reserved.
+
+#pragma once
+
+#include "IDetailCustomization.h"
+#include "Templates/SharedPointerFwd.h"
+#include "Widgets/Input/SComboBox.h"
+
+class IPropertyHandle;
+class SComboButton;
+class SWidget;
+
+namespace UE::Cameras
+{
+
+/**
+ * Details customization for the filmback camera node.
+ */
+class FFilmbackCameraNodeDetailsCustomization : public IDetailCustomization
+{
+public:
+
+	static TSharedRef<IDetailCustomization> MakeInstance();
+
+public:
+
+	/** IDetailCustomization interface */
+	virtual void CustomizeDetails(IDetailLayoutBuilder& DetailBuilder) override;
+
+private:
+
+	void BuildPresetComboList();
+
+	bool IsPresetEnabled() const;
+	TSharedRef<SWidget> MakePresetComboWidget(TSharedPtr<FText> InItem);
+	void OnPresetChanged(TSharedPtr<FText> NewSelection, ESelectInfo::Type SelectInfo);
+
+	FText GetPresetComboBoxContent() const;
+
+private:
+
+	TSharedPtr<IPropertyHandle> SensorWidthProperty;
+	TSharedPtr<IPropertyHandle> SensorHeightProperty;
+
+	TSharedPtr<SComboBox<TSharedPtr<FText>>> PresetComboBox;
+	TArray<TSharedPtr<FText>> PresetComboList;
+};
+
+}  // namespace UE::Cameras
+
diff --git a/Engine/Plugins/Cameras/GameplayCameras/Source/GameplayCamerasEditor/Public/Customizations/SingleCameraDirectorDetailsCustomization.h b/Engine/Plugins/Cameras/GameplayCameras/Source/GameplayCamerasEditor/Public/Customizations/SingleCameraDirectorDetailsCustomization.h
deleted file mode 100644
