Cameras: fix problems with the way public inner objects were handled

- Public objects must have the RF_Public flag removed before they can be removed from a package.
- When they are renamed, public objects introduce object redirectors in the package, which are also RF_Public.
- Because of these two previous points, we could have deleted objects (camera rigs, camera variables) left behind inside packages. We have to clean those up on load.
- Deleting objects now removes the RF_Public flag and any associated redirectors.
- Renaming objects only renames their user-defined names, now, because changing their UObject name can litter a package with redirector after many renames.
- As a corollary of the previous point, we leave all objects with their default name because those don't matter now anymore except for referencing.

#rnx

diff --git a/Engine/Plugins/Cameras/GameplayCameras/Source/GameplayCameras/Private/Core/CameraAsset.cpp b/Engine/Plugins/Cameras/GameplayCameras/Source/GameplayCameras/Private/Core/CameraAsset.cpp
--- a/Engine/Plugins/Cameras/GameplayCameras/Source/GameplayCameras/Private/Core/CameraAsset.cpp
+++ b/Engine/Plugins/Cameras/GameplayCameras/Source/GameplayCameras/Private/Core/CameraAsset.cpp
@@ -5,6 +5,8 @@
 #include "Core/CameraAssetBuilder.h"
 #include "Core/CameraBuildLog.h"
 #include "Core/CameraDirector.h"
+#include "Core/CameraRigAsset.h"
+#include "UObject/ObjectRedirector.h"
 #include "UObject/ObjectSaveContext.h"
 
 #include UE_INLINE_GENERATED_CPP_BY_NAME(CameraAsset)
@@ -110,6 +112,7 @@ void UCameraAsset::PostLoad()
 {
 	Super::PostLoad();
 
+#if WITH_EDITOR
 	if (CameraDirector)
 	{
 		EObjectFlags Flags = CameraDirector->GetFlags();
@@ -122,10 +125,68 @@ void UCameraAsset::PostLoad()
 			CameraDirector->ClearFlags(RF_Public | RF_Standalone);
 		}
 	}
+
+	CleanUpStrayObjects();
+#endif  // WITH_EDITOR
 }
 
 #if WITH_EDITOR
 
+void UCameraAsset::CleanUpStrayObjects()
+{
+	UPackage* CameraAssetPackage = GetOutermost();
+	if (!CameraAssetPackage || CameraAssetPackage == GetTransientPackage())
+	{
+		return;
+	}
+
+	// Some older versions of the camera editors had a bug that could lead to
+	// stray deleted camera rigs being left in the package. Let's clean them up.
+	TSet<UObject*> StrayObjects;
+	TSet<UCameraRigAsset*> KnownCameraRigs(CameraRigs);
+
+	TArray<UObject*> ObjectsInPackage;
+	GetObjectsWithPackage(CameraAssetPackage, ObjectsInPackage);
+	for (UObject* Object : ObjectsInPackage)
+	{
+		UCameraRigAsset* CameraRig = Cast<UCameraRigAsset>(Object);
+		if (!CameraRig)
+		{
+			continue;
+		}
+		if (KnownCameraRigs.Contains(CameraRig))
+		{
+			continue;
+		}
+
+		Modify();
+
+		CameraRig->ClearFlags(RF_Public | RF_Standalone);
+		StrayObjects.Add(CameraRig);
+	}
+
+	if (StrayObjects.Num() > 0)
+	{
+		// Also clean-up any redirectors to these objects.
+		for (UObject* Object : ObjectsInPackage)
+		{
+			if (UObjectRedirector* Redirector = Cast<UObjectRedirector>(Object))
+			{
+				if (StrayObjects.Contains(Redirector->DestinationObject))
+				{
+					Redirector->ClearFlags(RF_Public | RF_Standalone);
+					Redirector->DestinationObject = nullptr;
+				}
+			}
+		}
+
+		UE_LOG(LogCameraSystem, Warning,
+				TEXT("Cleaned up %d stray camera rigs in camera asset '%s'. Please resave the asset."),
+				StrayObjects.Num(), *GetPathNameSafe(this));
+	}
+}
+
+
 void UCameraAsset::PostEditChangeProperty(FPropertyChangedEvent& PropertyChangedEvent)
 {
 	using namespace UE::Cameras;
@@ -186,7 +247,7 @@ void UCameraAsset::PreSave(FObjectPreSaveContext ObjectSaveContext)
 
 	if (!HasAnyFlags(RF_ClassDefaultObject | RF_ArchetypeObject))
 	{
-		// Build on save.
+		// Build when saving/cooking.
 		BuildCamera();
 	}
 
diff --git a/Engine/Plugins/Cameras/GameplayCameras/Source/GameplayCameras/Private/Core/CameraVariableCollection.cpp b/Engine/Plugins/Cameras/GameplayCameras/Source/GameplayCameras/Private/Core/CameraVariableCollection.cpp
--- a/Engine/Plugins/Cameras/GameplayCameras/Source/GameplayCameras/Private/Core/CameraVariableCollection.cpp
+++ b/Engine/Plugins/Cameras/GameplayCameras/Source/GameplayCameras/Private/Core/CameraVariableCollection.cpp
@@ -4,6 +4,7 @@
 
 #include "Core/CameraVariableAssets.h"
 #include "GameplayCameras.h"
+#include "UObject/ObjectRedirector.h"
 
 #include UE_INLINE_GENERATED_CPP_BY_NAME(CameraVariableCollection)
 
@@ -25,6 +26,63 @@ void UCameraVariableCollection::PostLoad()
 			Variable->SetFlags(RF_Public);
 		}
 	}
+
+	CleanUpStrayObjects();
 #endif  // WITH_EDITOR
 }
 
+#if WITH_EDITOR
+
+void UCameraVariableCollection::CleanUpStrayObjects()
+{
+	UPackage* CollectionPackage = GetOutermost();
+	if (!CollectionPackage || CollectionPackage == GetTransientPackage())
+	{
+		return;
+	}
+
+	TSet<UObject*> StrayObjects;
+	TSet<UCameraVariableAsset*> KnownVariables(Variables);
+
+	TArray<UObject*> ObjectsInPackage;
+	GetObjectsWithPackage(CollectionPackage, ObjectsInPackage);
+	for (UObject* Object : ObjectsInPackage)
+	{
+		UCameraVariableAsset* Variable = Cast<UCameraVariableAsset>(Object);
+		if (!Variable)
+		{
+			continue;
+		}
+		if (KnownVariables.Contains(Variable))
+		{
+			continue;
+		}
+
+		Modify();
+
+		Variable->ClearFlags(RF_Public | RF_Standalone);
+		StrayObjects.Add(Variable);
+	}
+
+	if (StrayObjects.Num() > 0)
+	{
+		// Also clean-up any redirectors to these objects.
+		for (UObject* Object : ObjectsInPackage)
+		{
+			if (UObjectRedirector* Redirector = Cast<UObjectRedirector>(Object))
+			{
+				if (StrayObjects.Contains(Redirector->DestinationObject))
+				{
+					Redirector->ClearFlags(RF_Public | RF_Standalone);
+					Redirector->DestinationObject = nullptr;
+				}
+			}
+		}
+
+		UE_LOG(LogCameraSystem, Warning,
+				TEXT("Cleaned up %d stray camera variables in camera variable collection '%s'. Please resave the asset."),
+				StrayObjects.Num(), *GetPathNameSafe(this));
+	}
+}
+
+#endif  // WITH_EDITOR
diff --git a/Engine/Plugins/Cameras/GameplayCameras/Source/GameplayCameras/Public/Core/CameraAsset.h b/Engine/Plugins/Cameras/GameplayCameras/Source/GameplayCameras/Public/Core/CameraAsset.h
--- a/Engine/Plugins/Cameras/GameplayCameras/Source/GameplayCameras/Public/Core/CameraAsset.h
+++ b/Engine/Plugins/Cameras/GameplayCameras/Source/GameplayCameras/Public/Core/CameraAsset.h
@@ -135,6 +135,12 @@ public:
 	/** Event handlers to be notified of data changes. */
 	UE::Cameras::TCameraEventHandlerContainer<UE::Cameras::ICameraAssetEventHandler> EventHandlers;
 
+private:
+
+#if WITH_EDITOR
+	void CleanUpStrayObjects();
+#endif
+
 private:
 
 	/** The camera director to use in this camera. */
diff --git a/Engine/Plugins/Cameras/GameplayCameras/Source/GameplayCameras/Public/Core/CameraVariableCollection.h b/Engine/Plugins/Cameras/GameplayCameras/Source/GameplayCameras/Public/Core/CameraVariableCollection.h
--- a/Engine/Plugins/Cameras/GameplayCameras/Source/GameplayCameras/Public/Core/CameraVariableCollection.h
+++ b/Engine/Plugins/Cameras/GameplayCameras/Source/GameplayCameras/Public/Core/CameraVariableCollection.h
@@ -25,6 +25,12 @@ protected:
 	// UObject interface
 	virtual void PostLoad() override;
 
+private:
+
+#if WITH_EDITOR
+	void CleanUpStrayObjects();
+#endif
+
 public:
 
 	/** The variables in this collection. */
diff --git a/Engine/Plugins/Cameras/GameplayCameras/Source/GameplayCamerasEditor/Private/Editors/SCameraRigList.cpp b/Engine/Plugins/Cameras/GameplayCameras/Source/GameplayCamerasEditor/Private/Editors/SCameraRigList.cpp
--- a/Engine/Plugins/Cameras/GameplayCameras/Source/GameplayCamerasEditor/Private/Editors/SCameraRigList.cpp
+++ b/Engine/Plugins/Cameras/GameplayCameras/Source/GameplayCamerasEditor/Private/Editors/SCameraRigList.cpp
@@ -108,15 +108,10 @@ void SCameraRigListEntry::OnTextCommitted(const FText& Text, ETextCommit::Type C
 		FScopedTransaction Transaction(LOCTEXT("RenameCameraRig", "Rename Camera Rig"));
 
 		UCameraRigAsset* CameraRigAsset = Item->CameraRigAsset;
-		CameraRigAsset->Modify();
 
-		// Set the interface name, but also rename the object itself, which helps for debugging.
 		const FString NewDisplayName = Text.ToString();
+		CameraRigAsset->Modify();
 		CameraRigAsset->Interface.DisplayName = NewDisplayName;
-
-		FName NewObjectName = MakeObjectNameFromDisplayLabel(NewDisplayName, CameraRigAsset->GetFName());
-		NewObjectName = MakeUniqueObjectName(CameraRigAsset->GetOuter(), UCameraRigAsset::StaticClass(), NewObjectName);
-		CameraRigAsset->Rename(*NewObjectName.ToString());
 	}
 }
 
@@ -315,10 +310,9 @@ void SCameraRigList::OnAddCameraRig()
 
 	CameraAsset->Modify();
 
-	const FName NewCameraRigName = MakeUniqueObjectName(CameraAsset, UCameraRigAsset::StaticClass(), TEXT("NewCameraRig"));
 	UCameraRigAsset* NewCameraRig = NewObject<UCameraRigAsset>(
-			CameraAsset, 
-			NewCameraRigName, 
+			CameraAsset,
+			NAME_None,
 			RF_Transactional | RF_Public  // Must be referenceable from camera directors.
 			);
 	CameraAsset->AddCameraRig(NewCameraRig);
diff --git a/Engine/Plugins/Cameras/GameplayCameras/Source/GameplayCamerasEditor/Private/Editors/SCameraVariableCollectionEditor.cpp b/Engine/Plugins/Cameras/GameplayCameras/Source/GameplayCamerasEditor/Private/Editors/SCameraVariableCollectionEditor.cpp
--- a/Engine/Plugins/Cameras/GameplayCameras/Source/GameplayCamerasEditor/Private/Editors/SCameraVariableCollectionEditor.cpp
+++ b/Engine/Plugins/Cameras/GameplayCameras/Source/GameplayCamerasEditor/Private/Editors/SCameraVariableCollectionEditor.cpp
@@ -162,11 +162,6 @@ void SCameraVariableCollectionListRow::OnVariableNameCommitted(const FText& Text
 		const FString NewDisplayName = Text.ToString();
 		CameraVariable->Modify();
 		CameraVariable->DisplayName = NewDisplayName;
-
-		// We just set the display name, but also rename the object itself, which helps for debugging.
-		FName NewObjectName = MakeObjectNameFromDisplayLabel(NewDisplayName, CameraVariable->GetFName());
-		NewObjectName = MakeUniqueObjectName(CameraVariable->GetOuter(), CameraVariable->GetClass(), NewObjectName);
-		CameraVariable->Rename(*NewObjectName.ToString());
 	}
 
 	if (OnTextComitted.IsBound())
diff --git a/Engine/Plugins/Cameras/GameplayCameras/Source/GameplayCamerasEditor/Private/Factories/CameraAssetFactory.cpp b/Engine/Plugins/Cameras/GameplayCameras/Source/GameplayCamerasEditor/Private/Factories/CameraAssetFactory.cpp
--- a/Engine/Plugins/Cameras/GameplayCameras/Source/GameplayCamerasEditor/Private/Factories/CameraAssetFactory.cpp
+++ b/Engine/Plugins/Cameras/GameplayCameras/Source/GameplayCamerasEditor/Private/Factories/CameraAssetFactory.cpp
@@ -60,7 +60,7 @@ UObject* UCameraAssetFactory::FactoryCreateNew(UClass* Class, UObject* Parent, F
 {
 	UCameraAsset* NewCameraAsset = NewObject<UCameraAsset>(Parent, Class, Name, Flags | RF_Transactional);
 
-	UCameraRigAsset* FirstCameraRig = NewObject<UCameraRigAsset>(NewCameraAsset, TEXT("NewCameraRig"), RF_Transactional | RF_Public);
+	UCameraRigAsset* FirstCameraRig = NewObject<UCameraRigAsset>(NewCameraAsset, NAME_None, RF_Transactional | RF_Public);
 	NewCameraAsset->AddCameraRig(FirstCameraRig);
 
 	if (CameraDirectorClass)
diff --git a/Engine/Plugins/Cameras/GameplayCameras/Source/GameplayCamerasEditor/Private/Widgets/SDeleteCameraObjectDialog.cpp b/Engine/Plugins/Cameras/GameplayCameras/Source/GameplayCamerasEditor/Private/Widgets/SDeleteCameraObjectDialog.cpp
--- a/Engine/Plugins/Cameras/GameplayCameras/Source/GameplayCamerasEditor/Private/Widgets/SDeleteCameraObjectDialog.cpp
+++ b/Engine/Plugins/Cameras/GameplayCameras/Source/GameplayCamerasEditor/Private/Widgets/SDeleteCameraObjectDialog.cpp
@@ -10,6 +10,7 @@
 #include "ObjectTools.h"
 #include "Serialization/FindReferencersArchive.h"
 #include "Subsystems/AssetEditorSubsystem.h"
+#include "UObject/ObjectRedirector.h"
 #include "Widgets/Input/SButton.h"
 #include "Widgets/Layout/SBorder.h"
 #include "Widgets/Notifications/SNotificationList.h"
@@ -173,6 +174,7 @@ void SDeleteCameraObjectDialog::PerformReferenceReplacement()
 	// fails to recurse into the package (that's because it's using object references, and packages don't really
 	// reference their asset directly).
 	TArray<UPackage*> DirtiedPackages;
+	TArray<UPackage*> PackagesToDeleteFrom;
 	TSet<UObject*> ObjectsToReplaceWithin;
 	for (const FName PackageName : ReferencingPackages)
 	{
@@ -194,6 +196,7 @@ void SDeleteCameraObjectDialog::PerformReferenceReplacement()
 		if (ensure(Package))
 		{
 			DirtiedPackages.Add(Package);
+			PackagesToDeleteFrom.Add(Package);
 
 			UObject* PackageAsset = Package->FindAssetInPackage();
 			if (ensure(PackageAsset))
@@ -217,6 +220,29 @@ void SDeleteCameraObjectDialog::PerformReferenceReplacement()
 		FEditorFileUtils::PromptForCheckoutAndSave(ObjectPtrDecay(DirtiedPackages), SaveParams);
 	}
 
+	// Remove some flags that prevent objects from being collected.
+	for (UObject* Object : ObjectsToDelete)
+	{
+		Object->ClearFlags(RF_Public | RF_Standalone);
+	}
+
+	// Remove any object redirectors. This should be safe since we deleted any references.
+	TSet<UObject*> ObjectsToDeleteSet(ObjectsToDelete);
+	for (UPackage* Package : PackagesToDeleteFrom)
+	{
+		TArray<UObject*> ObjectsInPackage;
+		GetObjectsWithPackage(Package, ObjectsInPackage);
+		for (UObject* Object : ObjectsInPackage)
+		{
+			UObjectRedirector* Redirector = Cast<UObjectRedirector>(Object);
+			if (Redirector && ObjectsToDeleteSet.Contains(Redirector->DestinationObject))
+			{
+				Redirector->ClearFlags(RF_Public | RF_Standalone);
+				Redirector->DestinationObject = nullptr;
+			}
+		}
+	}
+
 	// Optionally rename objects with a TRASH prefix, helpful for debugging sometimes.
 	if (bRenameObjectsAsTrash)
 	{
