Cameras: don't build camera rigs on save/cook if they are already built by an owning camera asset

#rnx

diff --git a/Engine/Plugins/Cameras/GameplayCameras/Source/GameplayCameras/Private/Core/CameraRigAsset.cpp b/Engine/Plugins/Cameras/GameplayCameras/Source/GameplayCameras/Private/Core/CameraRigAsset.cpp
--- a/Engine/Plugins/Cameras/GameplayCameras/Source/GameplayCameras/Private/Core/CameraRigAsset.cpp
+++ b/Engine/Plugins/Cameras/GameplayCameras/Source/GameplayCameras/Private/Core/CameraRigAsset.cpp
@@ -199,9 +199,11 @@ void UCameraRigAsset::PreSave(FObjectPreSaveContext ObjectSaveContext)
 {
 #if WITH_EDITOR
 
-	if (!HasAnyFlags(RF_ClassDefaultObject | RF_ArchetypeObject))
+	if (GetOutermostObject() == this &&
+			!HasAnyFlags(RF_ClassDefaultObject | RF_ArchetypeObject))
 	{
-		// Build on save.
+		// Build when saving/cooking, if we are a standalone camera rig (i.e. not a camera rig inside a 
+		// camera asset, since those are built along with the camera asset).
 		BuildCameraRig();
 	}
 
