Cameras: make the camera rigs panel update when camera rigs change

#rnx

diff --git a/Engine/Plugins/Cameras/GameplayCameras/Source/GameplayCamerasEditor/Private/Editors/SCameraRigList.cpp b/Engine/Plugins/Cameras/GameplayCameras/Source/GameplayCamerasEditor/Private/Editors/SCameraRigList.cpp
--- a/Engine/Plugins/Cameras/GameplayCameras/Source/GameplayCamerasEditor/Private/Editors/SCameraRigList.cpp
+++ b/Engine/Plugins/Cameras/GameplayCameras/Source/GameplayCamerasEditor/Private/Editors/SCameraRigList.cpp
@@ -194,6 +194,8 @@ void SCameraRigList::Construct(const FArguments& InArgs)
 		ListView->SetSelection(FilteredItemSource[0]);
 		OnRequestEditCameraRig.ExecuteIfBound(FilteredItemSource[0]->CameraRigAsset);
 	}
+
+	CameraAsset->EventHandlers.Register(EventHandler, this);
 }
 
 SCameraRigList::~SCameraRigList()
@@ -255,6 +257,11 @@ FReply SCameraRigList::OnKeyDown(const FGeometry& MyGeometry, const FKeyEvent& I
 	return SCompoundWidget::OnKeyDown(MyGeometry, InKeyEvent);
 }
 
+void SCameraRigList::OnCameraRigsChanged(UCameraAsset* InCameraAsset, const TCameraArrayChangedEvent<UCameraRigAsset*>& Event)
+{
+	RequestListRefresh();
+}
+
 TSharedPtr<SWidget> SCameraRigList::GenerateToolbar()
 {
 	static const FName ToolbarName("CameraRigList.ToolBar");
diff --git a/Engine/Plugins/Cameras/GameplayCameras/Source/GameplayCamerasEditor/Public/Editors/SCameraRigList.h b/Engine/Plugins/Cameras/GameplayCameras/Source/GameplayCamerasEditor/Public/Editors/SCameraRigList.h
--- a/Engine/Plugins/Cameras/GameplayCameras/Source/GameplayCamerasEditor/Public/Editors/SCameraRigList.h
+++ b/Engine/Plugins/Cameras/GameplayCameras/Source/GameplayCamerasEditor/Public/Editors/SCameraRigList.h
@@ -2,6 +2,7 @@
 
 #pragma once
 
+#include "Core/CameraAsset.h"
 #include "Core/CameraRigAsset.h"
 #include "Misc/TextFilter.h"
 #include "UObject/ObjectPtr.h"
@@ -67,7 +68,9 @@ private:
  * A panel that shows a list of camera rigs on a camera asset.
  * Can add, rename, delete, etc. camera rigs in the list.
  */
-class SCameraRigList : public SCompoundWidget
+class SCameraRigList 
+	: public SCompoundWidget
+	, public ICameraAssetEventHandler
 {
 public:
 
@@ -91,6 +94,9 @@ protected:
 	virtual void Tick(const FGeometry& AllottedGeometry, const double InCurrentTime, const float InDeltaTime) override;
 	virtual FReply OnKeyDown(const FGeometry& MyGeometry, const FKeyEvent& InKeyEvent) override;
 
+	// ICameraAssetEventHandler interface
+	virtual void OnCameraRigsChanged(UCameraAsset* InCameraAsset, const TCameraArrayChangedEvent<UCameraRigAsset*>& Event) override;
+
 private:
 
 	TSharedPtr<SWidget> GenerateToolbar();
@@ -126,6 +132,8 @@ private:
 
 	TObjectPtr<UCameraAsset> CameraAsset;
 
+	TCameraEventHandler<ICameraAssetEventHandler> EventHandler;
+
 	FOnCameraRigListChanged OnCameraRigListChanged;
 	FOnCameraRigEvent OnRequestEditCameraRig;
 	FOnMultiCameraRigEvent OnCameraRigDeleted;
