Don't crash if LiveEditManager is missing (Ex: Standalone)

#tests AnimSandbox

diff --git a/Engine/Plugins/Cameras/GameplayCameras/Source/GameplayCameras/Private/Core/BlendStackCameraNode.cpp b/Engine/Plugins/Cameras/GameplayCameras/Source/GameplayCameras/Private/Core/BlendStackCameraNode.cpp
--- a/Engine/Plugins/Cameras/GameplayCameras/Source/GameplayCameras/Private/Core/BlendStackCameraNode.cpp
+++ b/Engine/Plugins/Cameras/GameplayCameras/Source/GameplayCameras/Private/Core/BlendStackCameraNode.cpp
@@ -295,6 +295,11 @@ void FBlendStackCameraNodeEvaluator::AddPackageListeners(FCameraRigEntry& Entry)
 
 	IGameplayCamerasModule& GameplayCamerasModule = FModuleManager::GetModuleChecked<IGameplayCamerasModule>("GameplayCameras");
 	TSharedPtr<IGameplayCamerasLiveEditManager> LiveEditManager = GameplayCamerasModule.GetLiveEditManager();
+	if (!LiveEditManager)
+	{
+		return;
+	}
+
 	Entry.CameraRig->GatherPackages(Entry.ListenedPackages);
 	for (const UPackage* ListenPackage : Entry.ListenedPackages)
 	{
@@ -316,6 +321,11 @@ void FBlendStackCameraNodeEvaluator::RemoveListenedPackages(FCameraRigEntry& Ent
 
 void FBlendStackCameraNodeEvaluator::RemoveListenedPackages(TSharedPtr<IGameplayCamerasLiveEditManager> LiveEditManager, FCameraRigEntry& Entry)
 {
+	if (!LiveEditManager)
+	{
+		return;
+	}
+
 	for (const UPackage* ListenPackage : Entry.ListenedPackages)
 	{
 		int32* NumListens = AllListenedPackages.Find(ListenPackage);
