Cameras: properly save/load camera variable tables that have changed size

diff --git a/Engine/Plugins/Cameras/GameplayCameras/Source/GameplayCameras/Private/Core/CameraVariableTable.cpp b/Engine/Plugins/Cameras/GameplayCameras/Source/GameplayCameras/Private/Core/CameraVariableTable.cpp
--- a/Engine/Plugins/Cameras/GameplayCameras/Source/GameplayCameras/Private/Core/CameraVariableTable.cpp
+++ b/Engine/Plugins/Cameras/GameplayCameras/Source/GameplayCameras/Private/Core/CameraVariableTable.cpp
@@ -345,20 +345,38 @@ void FCameraVariableTable::Serialize(FArchive& Ar)
 
 	if (Ar.IsLoading())
 	{
+		// Load the data from the saved buffer. It should have a lesser-or-equal size to our buffer,
+		// because we might have grown since we were saved, but we wouldn't have shrunk (variable tables
+		// don't shrink). So we keep our present capacity and ignore the saved on, except to confirm
+		// our expectations.
 		uint32 LoadedCapacity = 0;
 		Ar << LoadedCapacity;
 
 		uint32 LoadedUsed = 0;
 		Ar << LoadedUsed;
 
-		ensure(LoadedCapacity <= Capacity);
+		ensureMsgf(LoadedCapacity <= Capacity, 
+				TEXT("Loading camera variable table data from a bigger, and probably different, table."));
+		LoadedCapacity = FMath::Min(LoadedCapacity, Capacity);
 		Ar.Serialize(Memory, LoadedCapacity);
+		Used = LoadedUsed;
 
+		// The number of entries may have changed, if new entries were created since we were saved.
+		// Resize down to what we originally had.
 		int32 LoadedNumEntries = 0;
 		Ar << LoadedNumEntries;
 		
-		ensure(LoadedNumEntries == Entries.Num());
+		ensureMsgf(LoadedNumEntries <= Entries.Num(), 
+				TEXT("Loading camera variable table entries from a bigger, and probably different, table."));
+		Entries.SetNum(LoadedNumEntries, EAllowShrinking::No);
 		Ar.Serialize(Entries.GetData(), LoadedNumEntries * sizeof(FEntry));
+
+		// Rebuild the entry lookup table.
+		EntryLookup.Reset();
+		for (auto It = Entries.CreateConstIterator(); It; ++It)
+		{
+			EntryLookup.Add(It->ID, It.GetIndex());
+		}
 	}
 }
 
