Cameras: better handle the case of both FOV and focal length being invalid

diff --git a/Engine/Plugins/Cameras/GameplayCameras/Source/GameplayCameras/Private/Core/CameraPose.cpp b/Engine/Plugins/Cameras/GameplayCameras/Source/GameplayCameras/Private/Core/CameraPose.cpp
--- a/Engine/Plugins/Cameras/GameplayCameras/Source/GameplayCameras/Private/Core/CameraPose.cpp
+++ b/Engine/Plugins/Cameras/GameplayCameras/Source/GameplayCameras/Private/Core/CameraPose.cpp
@@ -125,11 +125,21 @@ double FCameraPose::GetEffectiveFieldOfView() const
 
 double FCameraPose::GetEffectiveFieldOfView(float FocalLength, float FieldOfView, float SensorWidth, float SensorHeight, float SqueezeFactor)
 {
-	checkf((FocalLength > 0.f || FieldOfView > 0.f), TEXT("FocalLength or FieldOfView must have a valid, positive value."));
+	const bool bValidFocalLength = (FocalLength > 0.f);
+	const bool bValidFieldOfView = (FieldOfView > 0.f);
 
 #if !NO_LOGGING
+	static bool GEmitZeroFocalLenthAndFieldOfViewWarning = true;
 	static bool GEmitFocalLengthPrioritizationWarning = true;
-	if ((FocalLength > 0.f && FieldOfView > 0.f) && GEmitFocalLengthPrioritizationWarning)
+
+	if (!bValidFocalLength && !bValidFieldOfView && GEmitZeroFocalLenthAndFieldOfViewWarning)
+	{
+		UE_LOG(LogCameraSystem, Warning,
+				TEXT("Both FocalLength and FieldOfView have a zero or negative value! Using default FocalLength."));
+		GEmitZeroFocalLenthAndFieldOfViewWarning = false;
+	}
+
+	if (bValidFocalLength && bValidFieldOfView && GEmitFocalLengthPrioritizationWarning)
 	{
 		UE_LOG(LogCameraSystem, Warning,
 				TEXT("Both FocalLength and FieldOfView are specified on a camera pose! Using FocalLength first."));
@@ -137,6 +147,11 @@ double FCameraPose::GetEffectiveFieldOfView(float FocalLength, float FieldOfView
 	}
 #endif  // NO_LOGGING	
 
+	if (!bValidFocalLength && !bValidFieldOfView)
+	{
+		FocalLength = 35.f;
+	}
+
 	if (FocalLength > 0.f)
 	{
 		// Compute FOV with similar code to UCineCameraComponent...
