Cameras: fix camera rig names in the debug display

Properly use the rigs' display names, which is even more important now that their UObject names are left alone (from a recent change)

diff --git a/Engine/Plugins/Cameras/GameplayCameras/Source/GameplayCameras/Private/Core/BlendStackCameraNode.cpp b/Engine/Plugins/Cameras/GameplayCameras/Source/GameplayCameras/Private/Core/BlendStackCameraNode.cpp
--- a/Engine/Plugins/Cameras/GameplayCameras/Source/GameplayCameras/Private/Core/BlendStackCameraNode.cpp
+++ b/Engine/Plugins/Cameras/GameplayCameras/Source/GameplayCameras/Private/Core/BlendStackCameraNode.cpp
@@ -1144,7 +1144,7 @@ FBlendStackCameraDebugBlock::FBlendStackCameraDebugBlock(const FBlendStackCamera
 	for (const FBlendStackCameraNodeEvaluator::FCameraRigEntry& Entry : InEvaluator.Entries)
 	{
 		FEntryDebugInfo EntryDebugInfo;
-		EntryDebugInfo.CameraRigName = GetNameSafe(Entry.CameraRig);
+		EntryDebugInfo.CameraRigName = Entry.CameraRig ? Entry.CameraRig->GetDisplayName() : FString("<None>");
 		Entries.Add(EntryDebugInfo);
 	}
 }
diff --git a/Engine/Plugins/Cameras/GameplayCameras/Source/GameplayCameras/Private/Core/BlendStackRootCameraNode.cpp b/Engine/Plugins/Cameras/GameplayCameras/Source/GameplayCameras/Private/Core/BlendStackRootCameraNode.cpp
--- a/Engine/Plugins/Cameras/GameplayCameras/Source/GameplayCameras/Private/Core/BlendStackRootCameraNode.cpp
+++ b/Engine/Plugins/Cameras/GameplayCameras/Source/GameplayCameras/Private/Core/BlendStackRootCameraNode.cpp
@@ -90,10 +90,8 @@ void FBlendStackRootCameraNodeEvaluator::OnInitialize(const FCameraNodeEvaluator
 		}
 
 #if UE_GAMEPLAY_CAMERAS_DEBUG
-		if (const UCameraRigAsset* CameraRig = RootNode->RootNode->GetTypedOuter<UCameraRigAsset>())
-		{
-			CameraRigAssetName = GetNameSafe(CameraRig);
-		}
+		const UCameraRigAsset* CameraRig = RootNode->RootNode->GetTypedOuter<UCameraRigAsset>();
+		CameraRigAssetName = CameraRig ? CameraRig->GetDisplayName() : TEXT("<None>");
 #endif  // UE_GAMEPLAY_CAMERAS_DEBUG
 	}
 }
