Cameras: make camera rig nodes renamable in the graph editor

Making them renamable prevents the graph editor from caching a node title, which means it updates automatically when the object is renamed. And so at this point we might as well also use the renaming feature of the graph editor, which means users can rename a camera rig either from the rigs list/panel, or from the graph editor.

diff --git a/Engine/Plugins/Cameras/GameplayCameras/Source/GameplayCameras/Private/Core/CameraRigAsset.cpp b/Engine/Plugins/Cameras/GameplayCameras/Source/GameplayCameras/Private/Core/CameraRigAsset.cpp
--- a/Engine/Plugins/Cameras/GameplayCameras/Source/GameplayCameras/Private/Core/CameraRigAsset.cpp
+++ b/Engine/Plugins/Cameras/GameplayCameras/Source/GameplayCameras/Private/Core/CameraRigAsset.cpp
@@ -268,6 +268,19 @@ void UCameraRigAsset::OnGraphNodeMoved(FName InGraphName, int32 NodePosX, int32
 	}
 }
 
+EObjectTreeGraphObjectSupportFlags UCameraRigAsset::GetSupportFlags(FName InGraphName) const
+{
+	if (UPackage* OuterPackage = Cast<UPackage>(GetOuter()))
+	{
+		// Can't rename a standalone rig prefab -- you have to rename it in the content browser.
+		return EObjectTreeGraphObjectSupportFlags::CommentText;
+	}
+	else
+	{
+		return EObjectTreeGraphObjectSupportFlags::CommentText | EObjectTreeGraphObjectSupportFlags::CustomRename;
+	}
+}
+
 const FString& UCameraRigAsset::GetGraphNodeCommentText(FName InGraphName) const
 {
 	if (InGraphName == NodeTreeGraphName)
@@ -297,6 +310,16 @@ void UCameraRigAsset::OnUpdateGraphNodeCommentText(FName InGraphName, const FStr
 	}
 }
 
+void UCameraRigAsset::GetGraphNodeName(FName InGraphName, FText& OutName) const
+{
+	OutName = FText::FromString(GetDisplayName());
+}
+
+void UCameraRigAsset::OnRenameGraphNode(FName InGraphName, const FString& NewName)
+{
+	Interface.DisplayName = NewName;
+}
+
 void UCameraRigAsset::GetConnectableObjects(FName InGraphName, TSet<UObject*>& OutObjects) const
 {
 	if (InGraphName == NodeTreeGraphName)
diff --git a/Engine/Plugins/Cameras/GameplayCameras/Source/GameplayCameras/Public/Core/CameraRigAsset.h b/Engine/Plugins/Cameras/GameplayCameras/Source/GameplayCameras/Public/Core/CameraRigAsset.h
--- a/Engine/Plugins/Cameras/GameplayCameras/Source/GameplayCameras/Public/Core/CameraRigAsset.h
+++ b/Engine/Plugins/Cameras/GameplayCameras/Source/GameplayCameras/Public/Core/CameraRigAsset.h
@@ -269,9 +269,11 @@ protected:
 #if WITH_EDITOR
 	virtual void GetGraphNodePosition(FName InGraphName, int32& NodePosX, int32& NodePosY) const override;
 	virtual void OnGraphNodeMoved(FName InGraphName, int32 NodePosX, int32 NodePosY, bool bMarkDirty) override;
-	virtual EObjectTreeGraphObjectSupportFlags GetSupportFlags(FName InGraphName) const override { return EObjectTreeGraphObjectSupportFlags::CommentText; }
+	virtual EObjectTreeGraphObjectSupportFlags GetSupportFlags(FName InGraphName) const override;
 	virtual const FString& GetGraphNodeCommentText(FName InGraphName) const override;
 	virtual void OnUpdateGraphNodeCommentText(FName InGraphName, const FString& NewComment) override;
+	virtual void GetGraphNodeName(FName InGraphName, FText& OutName) const override;
+	virtual void OnRenameGraphNode(FName InGraphName, const FString& NewName) override;
 #endif
 
 	// IObjectTreeGraphRootObject interface.
diff --git a/Engine/Plugins/Cameras/GameplayCameras/Source/GameplayCameras/Public/Core/ObjectTreeGraphObject.h b/Engine/Plugins/Cameras/GameplayCameras/Source/GameplayCameras/Public/Core/ObjectTreeGraphObject.h
--- a/Engine/Plugins/Cameras/GameplayCameras/Source/GameplayCameras/Public/Core/ObjectTreeGraphObject.h
+++ b/Engine/Plugins/Cameras/GameplayCameras/Source/GameplayCameras/Public/Core/ObjectTreeGraphObject.h
@@ -62,7 +62,7 @@ public:
 	virtual void OnUpdateGraphNodeCommentText(FName InGraphName, const FString& NewComment) {}
 
 	/** Gets the custom name for the graph node. */
-	virtual const FString& GetGraphNodeName(FName InGraphName) const { static const FString EmptyString; return EmptyString; }
+	virtual void GetGraphNodeName(FName InGraphName, FText& OutName) const {}
 	/** Called to save a new custom name. */
 	virtual void OnRenameGraphNode(FName InGraphName, const FString& NewName) {}
 
diff --git a/Engine/Plugins/Cameras/GameplayCameras/Source/GameplayCamerasEditor/Private/Editors/CameraNodeGraphSchema.cpp b/Engine/Plugins/Cameras/GameplayCameras/Source/GameplayCamerasEditor/Private/Editors/CameraNodeGraphSchema.cpp
--- a/Engine/Plugins/Cameras/GameplayCameras/Source/GameplayCamerasEditor/Private/Editors/CameraNodeGraphSchema.cpp
+++ b/Engine/Plugins/Cameras/GameplayCameras/Source/GameplayCamerasEditor/Private/Editors/CameraNodeGraphSchema.cpp
@@ -53,14 +53,6 @@ FObjectTreeGraphConfig UCameraNodeGraphSchema::BuildGraphConfig() const
 		.SelfPinName(NAME_None)  // No self pin name, we just want the title
 		.CanCreateNew(false)
 		.GraphNodeClass(UCameraRigInterfaceParameterGraphNode::StaticClass());
-	GraphConfig.OnFormatObjectDisplayName = FOnFormatObjectDisplayName::CreateLambda(
-			[](const UObject* Object, FText& InOutDisplayNameText)
-			{
-				if (const UCameraRigAsset* CameraRigAsset = Cast<UCameraRigAsset>(Object))
-				{
-					InOutDisplayNameText = FText::FromString(CameraRigAsset->GetDisplayName());
-				}
-			});
 
 	return GraphConfig;
 }
diff --git a/Engine/Plugins/Cameras/GameplayCameras/Source/GameplayCamerasEditor/Private/Editors/CameraRigTransitionGraphSchema.cpp b/Engine/Plugins/Cameras/GameplayCameras/Source/GameplayCamerasEditor/Private/Editors/CameraRigTransitionGraphSchema.cpp
--- a/Engine/Plugins/Cameras/GameplayCameras/Source/GameplayCamerasEditor/Private/Editors/CameraRigTransitionGraphSchema.cpp
+++ b/Engine/Plugins/Cameras/GameplayCameras/Source/GameplayCamerasEditor/Private/Editors/CameraRigTransitionGraphSchema.cpp
@@ -24,14 +24,6 @@ void UCameraRigTransitionGraphSchema::OnBuildGraphConfig(FObjectTreeGraphConfig&
 		.OnlyAsRoot()
 		.NodeTitleUsesObjectName(true)
 		.NodeTitleColor(Settings->CameraRigAssetTitleColor);
-	InOutGraphConfig.OnFormatObjectDisplayName = FOnFormatObjectDisplayName::CreateLambda(
-			[](const UObject* Object, FText& InOutDisplayNameText)
-			{
-				if (const UCameraRigAsset* CameraRigAsset = Cast<UCameraRigAsset>(Object))
-				{
-					InOutDisplayNameText = FText::FromString(CameraRigAsset->GetDisplayName());
-				}
-			});
 }
 
 ETransitionGraphContextActions UCameraRigTransitionGraphSchema::GetTransitionGraphContextActions(FGraphContextMenuBuilder& ContextMenuBuilder) const
diff --git a/Engine/Plugins/Cameras/GameplayCameras/Source/GameplayCamerasEditor/Private/Editors/ObjectTreeGraphConfig.cpp b/Engine/Plugins/Cameras/GameplayCameras/Source/GameplayCamerasEditor/Private/Editors/ObjectTreeGraphConfig.cpp
--- a/Engine/Plugins/Cameras/GameplayCameras/Source/GameplayCamerasEditor/Private/Editors/ObjectTreeGraphConfig.cpp
+++ b/Engine/Plugins/Cameras/GameplayCameras/Source/GameplayCamerasEditor/Private/Editors/ObjectTreeGraphConfig.cpp
@@ -163,7 +163,7 @@ FText FObjectTreeGraphConfig::GetDisplayNameText(const UObject* InObject) const
 		const IObjectTreeGraphObject* GraphObject = Cast<IObjectTreeGraphObject>(InObject);
 		if (GraphObject && GraphObject->HasSupportFlags(GraphName, EObjectTreeGraphObjectSupportFlags::CustomRename))
 		{
-			DisplayNameText = FText::FromString(GraphObject->GetGraphNodeName(GraphName));
+			GraphObject->GetGraphNodeName(GraphName, DisplayNameText);
 		}
 
 		if (DisplayNameText.IsEmpty() && ClassConfig.NodeTitleUsesObjectName())
