Cameras: replace ensures and logs with Kismet messages in BP-exposed functions

diff --git a/Engine/Plugins/Cameras/GameplayCameras/Source/GameplayCameras/Private/GameFramework/GameplayCameraComponent.cpp b/Engine/Plugins/Cameras/GameplayCameras/Source/GameplayCameras/Private/GameFramework/GameplayCameraComponent.cpp
--- a/Engine/Plugins/Cameras/GameplayCameras/Source/GameplayCameras/Private/GameFramework/GameplayCameraComponent.cpp
+++ b/Engine/Plugins/Cameras/GameplayCameras/Source/GameplayCameras/Private/GameFramework/GameplayCameraComponent.cpp
@@ -72,7 +72,9 @@ void UGameplayCameraComponent::ActivateCameraEvaluationContext(int32 PlayerIndex
 	APlayerController* PlayerController = UGameplayStatics::GetPlayerController(this, PlayerIndex);
 	if (!PlayerController)
 	{
-		UE_LOG(LogCameraSystem, Error, TEXT("Can't activate gameplay camera: no player controller found!"));
+		FFrame::KismetExecutionMessage(
+				TEXT("Can't activate gameplay camera: no player controller found!"),
+				ELogVerbosity::Error);
 		return;
 	}
 
@@ -102,21 +104,28 @@ void UGameplayCameraComponent::ActivateCameraEvaluationContext(APlayerController
 {
 	using namespace UE::Cameras;
 
-	if (!ensureMsgf(PlayerController, TEXT("Can't activate gameplay camera component: invalid player controller!")))
+	if (!PlayerController)
 	{
+		FFrame::KismetExecutionMessage(
+				TEXT("Can't activate gameplay camera component: invalid player controller!"),
+				ELogVerbosity::Error);
 		return;
 	}
 
 	if (!Camera)
 	{
-		UE_LOG(LogCameraSystem, Error, TEXT("Can't activate gameplay camera component: no camera asset was set!"));
+		FFrame::KismetExecutionMessage(
+				TEXT("Can't activate gameplay camera component: no camera asset was set!"),
+				ELogVerbosity::Error);
 		return;
 	}
 	
 	CameraSystemHost = UGameplayCameraSystemHost::FindOrCreateHost(PlayerController);
 	if (!CameraSystemHost)
 	{
-		UE_LOG(LogCameraSystem, Error, TEXT("Can't activate gameplay camera component: no camera system host found!"));
+		FFrame::KismetExecutionMessage(
+				TEXT("Can't activate gameplay camera component: no camera system host found!"),
+				ELogVerbosity::Error);
 		return;
 	}
 
diff --git a/Engine/Plugins/Cameras/GameplayCameras/Source/GameplayCameras/Private/GameFramework/GameplayCameraSystemActor.cpp b/Engine/Plugins/Cameras/GameplayCameras/Source/GameplayCameras/Private/GameFramework/GameplayCameraSystemActor.cpp
--- a/Engine/Plugins/Cameras/GameplayCameras/Source/GameplayCameras/Private/GameFramework/GameplayCameraSystemActor.cpp
+++ b/Engine/Plugins/Cameras/GameplayCameras/Source/GameplayCameras/Private/GameFramework/GameplayCameraSystemActor.cpp
@@ -58,7 +58,9 @@ AGameplayCameraSystemActor* AGameplayCameraSystemActor::GetAutoSpawnedCameraSyst
 		}
 		else
 		{
-			UE_LOG(LogCameraSystem, Error, TEXT("Can't auto-manage active view target: no camera system host found!"));
+			FFrame::KismetExecutionMessage(
+					TEXT("Can't auto-manage active view target: no camera system host found!"),
+					ELogVerbosity::Error);
 			return nullptr;
 		}
 	}
diff --git a/Engine/Plugins/Cameras/GameplayCameras/Source/GameplayCameras/Private/GameFramework/GameplayCameraSystemComponent.cpp b/Engine/Plugins/Cameras/GameplayCameras/Source/GameplayCameras/Private/GameFramework/GameplayCameraSystemComponent.cpp
--- a/Engine/Plugins/Cameras/GameplayCameras/Source/GameplayCameras/Private/GameFramework/GameplayCameraSystemComponent.cpp
+++ b/Engine/Plugins/Cameras/GameplayCameras/Source/GameplayCameras/Private/GameFramework/GameplayCameraSystemComponent.cpp
@@ -107,7 +107,9 @@ void UGameplayCameraSystemComponent::ActivateCameraSystemForPlayerIndex(int32 Pl
 	APlayerController* PlayerController = UGameplayStatics::GetPlayerController(this, PlayerIndex);
 	if (!PlayerController)
 	{
-		UE_LOG(LogCameraSystem, Error, TEXT("Can't activate gameplay camera system: no player controller found!"));
+		FFrame::KismetExecutionMessage(
+				TEXT("Can't activate gameplay camera system: no player controller found!"),
+				ELogVerbosity::Error);
 		return;
 	}
 
@@ -116,6 +118,14 @@ void UGameplayCameraSystemComponent::ActivateCameraSystemForPlayerIndex(int32 Pl
 
 void UGameplayCameraSystemComponent::ActivateCameraSystemForPlayerController(APlayerController* PlayerController)
 {
+	if (!PlayerController)
+	{
+		FFrame::KismetExecutionMessage(
+				TEXT("Can't activate gameplay camera system: invalid player controller given!"),
+				ELogVerbosity::Error);
+		return;
+	}
+
 	if (APlayerController* ActivePlayerController = WeakPlayerController.Get())
 	{
 		if (ActivePlayerController != PlayerController)
@@ -127,7 +137,9 @@ void UGameplayCameraSystemComponent::ActivateCameraSystemForPlayerController(APl
 	AActor* OwningActor = GetOwner();
 	if (!OwningActor)
 	{
-		UE_LOG(LogCameraSystem, Error, TEXT("Can't activate gameplay camera system: no owning actor found!"));
+		FFrame::KismetExecutionMessage(
+				TEXT("Can't activate gameplay camera system: no owning actor found!"),
+				ELogVerbosity::Error);
 		return;
 	}
 
@@ -136,7 +148,7 @@ void UGameplayCameraSystemComponent::ActivateCameraSystemForPlayerController(APl
 		CameraSystemHost = UGameplayCameraSystemHost::FindOrCreateHost(PlayerController);
 		if (!CameraSystemHost)
 		{
-			UE_LOG(LogCameraSystem, Error, TEXT("can't create camera system host!"));
+			FFrame::KismetExecutionMessage(TEXT("can't create camera system host!"), ELogVerbosity::Error);
 			return;
 		}
 	}
