Cameras: make the camera system component use the scene component sprite

diff --git a/Engine/Plugins/Cameras/GameplayCameras/Source/GameplayCameras/Private/GameFramework/GameplayCameraSystemComponent.cpp b/Engine/Plugins/Cameras/GameplayCameras/Source/GameplayCameras/Private/GameFramework/GameplayCameraSystemComponent.cpp
--- a/Engine/Plugins/Cameras/GameplayCameras/Source/GameplayCameras/Private/GameFramework/GameplayCameraSystemComponent.cpp
+++ b/Engine/Plugins/Cameras/GameplayCameras/Source/GameplayCameras/Private/GameFramework/GameplayCameraSystemComponent.cpp
@@ -12,6 +12,7 @@
 #include "GameFramework/GameplayCameraSystemHost.h"
 #include "Kismet/GameplayStatics.h"
 #include "UObject/ConstructorHelpers.h"
+#include "UObject/ICookInfo.h"
 
 #include UE_INLINE_GENERATED_CPP_BY_NAME(GameplayCameraSystemComponent)
 
@@ -20,30 +21,6 @@
 UGameplayCameraSystemComponent::UGameplayCameraSystemComponent(const FObjectInitializer& ObjectInit)
 	: Super(ObjectInit)
 {
-#if WITH_EDITORONLY_DATA
-	EditorSpriteComponent = CreateEditorOnlyDefaultSubobject<UBillboardComponent>(TEXT("Sprite"));
-	if (GIsEditor && !IsRunningCommandlet())
-	{
-		static ConstructorHelpers::FObjectFinder<UTexture2D> EditorSpriteTextureFinder(
-				TEXT("/GameplayCameras/Textures/S_GameplayCameraSystem.S_GameplayCameraSystem"));
-
-		UTexture2D* EditorSpriteTexture = EditorSpriteTextureFinder.Object;
-
-		if (EditorSpriteTexture && EditorSpriteComponent)
-		{
-			EditorSpriteComponent->SpriteInfo.Category = TEXT("Cameras");
-			EditorSpriteComponent->SpriteInfo.DisplayName = NSLOCTEXT("SpriteCategory", "Cameras", "Cameras");
-
-			EditorSpriteComponent->SetSprite(EditorSpriteTexture);
-			EditorSpriteComponent->SetRelativeScale3D(FVector3d(EditorSpriteTextureScale));
-
-			EditorSpriteComponent->bHiddenInGame = true;
-			EditorSpriteComponent->bIsScreenSizeScaled = true;
-			EditorSpriteComponent->SetUsingAbsoluteScale(true);
-			EditorSpriteComponent->SetupAttachment(this);
-		}
-	}
-#endif  // WITH_EDITORONLY_DATA
 }
 
 TSharedPtr<UE::Cameras::FCameraSystemEvaluator> UGameplayCameraSystemComponent::GetCameraSystemEvaluator(bool bEnsureIfNull)
@@ -86,6 +63,10 @@ void UGameplayCameraSystemComponent::OnRegister()
 
 	Super::OnRegister();
 
+#if WITH_EDITOR
+	CreateCameraSystemSpriteComponent();
+#endif  // WITH_EDITOR
+
 	AActor* OwnerActor = GetOwner();
 	if (!OwnerActor || OwnerActor->HasAnyFlags(RF_ClassDefaultObject))
 	{
@@ -102,6 +83,35 @@ void UGameplayCameraSystemComponent::OnRegister()
 #endif  // UE_GAMEPLAY_CAMERAS_DEBUG
 }
 
+#if WITH_EDITOR
+
+void UGameplayCameraSystemComponent::CreateCameraSystemSpriteComponent()
+{
+	UTexture2D* EditorSpriteTexture = nullptr;
+	{
+		FCookLoadScope EditorOnlyScope(ECookLoadType::EditorOnly);
+		EditorSpriteTexture = LoadObject<UTexture2D>(
+				nullptr,
+				TEXT("/GameplayCameras/Textures/S_GameplayCameraSystem.S_GameplayCameraSystem"));
+	}
+
+	if (EditorSpriteTexture)
+	{
+		bVisualizeComponent = true;
+		CreateSpriteComponent(EditorSpriteTexture);
+	}
+
+	if (SpriteComponent)
+	{
+		SpriteComponent->SpriteInfo.Category = TEXT("Cameras");
+		SpriteComponent->SpriteInfo.DisplayName = NSLOCTEXT("SpriteCategory", "Cameras", "Cameras");
+		SpriteComponent->SetRelativeScale3D(FVector3d(EditorSpriteTextureScale));
+	}
+}
+
+#endif  // WITH_EDITOR
+
+
 void UGameplayCameraSystemComponent::ActivateCameraSystemForPlayerIndex(int32 PlayerIndex)
 {
 	APlayerController* PlayerController = UGameplayStatics::GetPlayerController(this, PlayerIndex);
diff --git a/Engine/Plugins/Cameras/GameplayCameras/Source/GameplayCameras/Public/GameFramework/GameplayCameraSystemComponent.h b/Engine/Plugins/Cameras/GameplayCameras/Source/GameplayCameras/Public/GameFramework/GameplayCameraSystemComponent.h
--- a/Engine/Plugins/Cameras/GameplayCameras/Source/GameplayCameras/Public/GameFramework/GameplayCameraSystemComponent.h
+++ b/Engine/Plugins/Cameras/GameplayCameras/Source/GameplayCameras/Public/GameFramework/GameplayCameraSystemComponent.h
@@ -76,6 +76,10 @@ private:
 	void DebugDraw(UCanvas* Canvas, APlayerController* PlayController);
 #endif  // UE_GAMEPLAY_CAMERAS_DEBUG
 
+#if WITH_EDITOR
+	void CreateCameraSystemSpriteComponent();
+#endif  // WITH_EDITOR
+
 public:
 
 	/**
