Cameras: fix reference to target location variable in framing nodes

References to camera variables were changed to use a custom struct in order to provide a custom picker UI, and the framing nodes were forgotten in the change.

diff --git a/Engine/Plugins/Cameras/GameplayCameras/Source/GameplayCameras/Private/Nodes/Framing/BaseFramingCameraNode.cpp b/Engine/Plugins/Cameras/GameplayCameras/Source/GameplayCameras/Private/Nodes/Framing/BaseFramingCameraNode.cpp
--- a/Engine/Plugins/Cameras/GameplayCameras/Source/GameplayCameras/Private/Nodes/Framing/BaseFramingCameraNode.cpp
+++ b/Engine/Plugins/Cameras/GameplayCameras/Source/GameplayCameras/Private/Nodes/Framing/BaseFramingCameraNode.cpp
@@ -70,11 +70,11 @@ void FBaseFramingCameraNodeEvaluator::OnInitialize(const FCameraNodeEvaluatorIni
 
 TOptional<FVector3d> FBaseFramingCameraNodeEvaluator::AcquireTargetLocation(const FCameraNodeEvaluationParams& Params, const FCameraNodeEvaluationResult& InResult)
 {
-	const UBaseFramingCameraNode* DollyNode = GetCameraNodeAs<UBaseFramingCameraNode>();
-	if (DollyNode->TargetLocation)
+	const UBaseFramingCameraNode* FramingNode = GetCameraNodeAs<UBaseFramingCameraNode>();
+	if (FramingNode->TargetLocation)
 	{
 		FVector3d TargetLocation;
-		const bool bGotTargetLocation = InResult.VariableTable.TryGetValue(DollyNode->TargetLocation.Get(), TargetLocation);
+		const bool bGotTargetLocation = InResult.VariableTable.TryGetValue(FramingNode->TargetLocation.Get(), TargetLocation);
 		return bGotTargetLocation ? TOptional<FVector3d>(TargetLocation) : TOptional<FVector3d>();
 	}
 	else if (APlayerController* PlayerController = Params.EvaluationContext->GetPlayerController())
diff --git a/Engine/Plugins/Cameras/GameplayCameras/Source/GameplayCameras/Public/Nodes/Framing/BaseFramingCameraNode.h b/Engine/Plugins/Cameras/GameplayCameras/Source/GameplayCameras/Public/Nodes/Framing/BaseFramingCameraNode.h
--- a/Engine/Plugins/Cameras/GameplayCameras/Source/GameplayCameras/Public/Nodes/Framing/BaseFramingCameraNode.h
+++ b/Engine/Plugins/Cameras/GameplayCameras/Source/GameplayCameras/Public/Nodes/Framing/BaseFramingCameraNode.h
@@ -5,6 +5,7 @@
 #include "Core/CameraNode.h"
 #include "Core/CameraParameterReader.h"
 #include "Core/CameraParameters.h"
+#include "Core/CameraVariableReferences.h"
 #include "Nodes/Framing/CameraFramingZone.h"
 #include "Math/CameraFramingZoneMath.h"
 #include "Math/CriticalDamper.h"
@@ -29,7 +30,7 @@ public:
 	 * the player pawn's location will be used by default.
 	 */
 	UPROPERTY(EditAnywhere, Category="Target")
-	TObjectPtr<UVector3dCameraVariable> TargetLocation;
+	FVector3dCameraVariableReference TargetLocation;
 
 	/** The ideal horizontal screen-space position of the target. */
 	UPROPERTY(EditAnywhere, Category="Framing Target")
