Cameras: clean-up dead code

diff --git a/Engine/Plugins/Cameras/GameplayCameras/Source/GameplayCameras/Public/GameFramework/GameplayCameraSystemComponent.h b/Engine/Plugins/Cameras/GameplayCameras/Source/GameplayCameras/Public/GameFramework/GameplayCameraSystemComponent.h
--- a/Engine/Plugins/Cameras/GameplayCameras/Source/GameplayCameras/Public/GameFramework/GameplayCameraSystemComponent.h
+++ b/Engine/Plugins/Cameras/GameplayCameras/Source/GameplayCameras/Public/GameFramework/GameplayCameraSystemComponent.h
@@ -9,7 +9,6 @@
 #include "GameplayCameraSystemComponent.generated.h"
 
 class APlayerController;
-class UBillboardComponent;
 class UCameraRigAsset;
 class UCanvas;
 class UGameplayCameraSystemHost;
@@ -114,10 +113,6 @@ private:
 	UPROPERTY(transient)
 	float EditorSpriteTextureScale = 0.5f;
 
-	/** Sprite component for the editor. */
-	UPROPERTY()
-	TObjectPtr<UBillboardComponent> EditorSpriteComponent;
-
 #endif	// WITH_EDITORONLY_DATA
 };
 
