Cameras: replace ensure with error log when occlusion material isn't set on the occlusion node

#rnx

diff --git a/Engine/Plugins/Cameras/GameplayCameras/Source/GameplayCameras/Private/Nodes/Collision/OcclusionMaterialCameraNode.cpp b/Engine/Plugins/Cameras/GameplayCameras/Source/GameplayCameras/Private/Nodes/Collision/OcclusionMaterialCameraNode.cpp
--- a/Engine/Plugins/Cameras/GameplayCameras/Source/GameplayCameras/Private/Nodes/Collision/OcclusionMaterialCameraNode.cpp
+++ b/Engine/Plugins/Cameras/GameplayCameras/Source/GameplayCameras/Private/Nodes/Collision/OcclusionMaterialCameraNode.cpp
@@ -76,6 +76,13 @@ void FOcclusionMaterialCameraNodeEvaluator::OnInitialize(const FCameraNodeEvalua
 	const UOcclusionMaterialCameraNode* OcclusionMaterialNode = GetCameraNodeAs<UOcclusionMaterialCameraNode>();
 	OcclusionSphereRadiusReader.Initialize(OcclusionMaterialNode->OcclusionSphereRadius);
 	OcclusionTargetOffsetReader.Initialize(OcclusionMaterialNode->OcclusionTargetOffset);
+
+	if (!OcclusionMaterialNode->OcclusionTransparencyMaterial)
+	{
+		UE_LOG(LogCameraSystem, Error, 
+				TEXT("OcclusionMaterialCameraNode: no occlusion transparency material set on '%s'"),
+				*GetNameSafe(OcclusionMaterialNode));
+	}
 }
 
 void FOcclusionMaterialCameraNodeEvaluator::OnRun(const FCameraNodeEvaluationParams& Params, FCameraNodeEvaluationResult& OutResult)
@@ -211,7 +218,7 @@ void FOcclusionMaterialCameraNodeEvaluator::ApplyOcclusionMaterial(TSet<UStaticM
 {
 	const UOcclusionMaterialCameraNode* OcclusionMaterialNode = GetCameraNodeAs<UOcclusionMaterialCameraNode>();
 	UMaterialInterface* OcclusionTransparencyMaterial = OcclusionMaterialNode->OcclusionTransparencyMaterial;
-	if (!ensureMsgf(OcclusionTransparencyMaterial, TEXT("No occlusion transparency material set!")))
+	if (!OcclusionTransparencyMaterial)
 	{
 		return;
 	}
