Cameras: use weak pointers for live-edit packages being listened to

#rnx

diff --git a/Engine/Plugins/Cameras/GameplayCameras/Source/GameplayCameras/Private/Core/BlendStackCameraNode.cpp b/Engine/Plugins/Cameras/GameplayCameras/Source/GameplayCameras/Private/Core/BlendStackCameraNode.cpp
--- a/Engine/Plugins/Cameras/GameplayCameras/Source/GameplayCameras/Private/Core/BlendStackCameraNode.cpp
+++ b/Engine/Plugins/Cameras/GameplayCameras/Source/GameplayCameras/Private/Core/BlendStackCameraNode.cpp
@@ -300,8 +300,13 @@ void FBlendStackCameraNodeEvaluator::AddPackageListeners(FCameraRigEntry& Entry)
 		return;
 	}
 
-	Entry.CameraRig->GatherPackages(Entry.ListenedPackages);
-	for (const UPackage* ListenPackage : Entry.ListenedPackages)
+	FCameraRigPackages EntryPackages;
+	Entry.CameraRig->GatherPackages(EntryPackages);
+
+	Entry.ListenedPackages.Reset();
+	Entry.ListenedPackages.Append(EntryPackages);
+
+	for (const UPackage* ListenPackage : EntryPackages)
 	{
 		int32& NumListens = AllListenedPackages.FindOrAdd(ListenPackage, 0);
 		if (NumListens == 0)
@@ -326,16 +331,19 @@ void FBlendStackCameraNodeEvaluator::RemoveListenedPackages(TSharedPtr<IGameplay
 		return;
 	}
 
-	for (const UPackage* ListenPackage : Entry.ListenedPackages)
+	for (TWeakObjectPtr<const UPackage> WeakListenPackage : Entry.ListenedPackages)
 	{
-		int32* NumListens = AllListenedPackages.Find(ListenPackage);
+		int32* NumListens = AllListenedPackages.Find(WeakListenPackage);
 		if (ensure(NumListens))
 		{
 			--(*NumListens);
 			if (*NumListens == 0)
 			{
-				LiveEditManager->RemoveListener(ListenPackage, this);
-				AllListenedPackages.Remove(ListenPackage);
+				if (const UPackage* ListenPackage = WeakListenPackage.Get())
+				{
+					LiveEditManager->RemoveListener(ListenPackage, this);
+				}
+				AllListenedPackages.Remove(WeakListenPackage);
 			}
 		}
 	}
diff --git a/Engine/Plugins/Cameras/GameplayCameras/Source/GameplayCameras/Public/Core/BlendStackCameraNode.h b/Engine/Plugins/Cameras/GameplayCameras/Source/GameplayCameras/Public/Core/BlendStackCameraNode.h
--- a/Engine/Plugins/Cameras/GameplayCameras/Source/GameplayCameras/Public/Core/BlendStackCameraNode.h
+++ b/Engine/Plugins/Cameras/GameplayCameras/Source/GameplayCameras/Public/Core/BlendStackCameraNode.h
@@ -196,7 +196,7 @@ protected:
 #endif  // UE_GAMEPLAY_CAMERAS_TRACE
 
 #if WITH_EDITOR
-		FCameraRigPackages ListenedPackages;
+		TArray<TWeakObjectPtr<const UPackage>, TInlineAllocator<4>> ListenedPackages;
 #endif  // WITH_EDITOR
 	};
 
@@ -210,7 +210,7 @@ protected:
 	FOnBlendStackCameraRigEvent OnCameraRigEventDelegate;
 
 #if WITH_EDITOR
-	TMap<const UPackage*, int32> AllListenedPackages;
+	TMap<TWeakObjectPtr<const UPackage>, int32> AllListenedPackages;
 #endif  // WITH_EDITOR
 
 #if UE_GAMEPLAY_CAMERAS_DEBUG
diff --git a/Engine/Plugins/Cameras/GameplayCameras/Source/GameplayCamerasEditor/Private/GameplayCamerasEditorModule.cpp b/Engine/Plugins/Cameras/GameplayCameras/Source/GameplayCamerasEditor/Private/GameplayCamerasEditorModule.cpp
--- a/Engine/Plugins/Cameras/GameplayCameras/Source/GameplayCamerasEditor/Private/GameplayCamerasEditorModule.cpp
+++ b/Engine/Plugins/Cameras/GameplayCameras/Source/GameplayCamerasEditor/Private/GameplayCamerasEditorModule.cpp
@@ -39,7 +39,6 @@
 #include "ISettingsModule.h"
 #include "K2Node_Event.h"
 #include "Kismet2/KismetEditorUtilities.h"
-#include "Misc/CoreDelegates.h"
 #include "Modules/ModuleManager.h"
 #include "PropertyEditorModule.h"
 #include "Styles/GameplayCamerasEditorStyle.h"
@@ -594,8 +593,6 @@ private:
 
 		IGameplayCamerasModule& CamerasModule = FModuleManager::LoadModuleChecked<IGameplayCamerasModule>("GameplayCameras");
 		CamerasModule.SetLiveEditManager(LiveEditManager);
-
-		FCoreUObjectDelegates::GetPostGarbageCollect().AddRaw(this, &FGameplayCamerasEditorModule::OnPostGarbageCollection);
 	}
 
 	void TeardownLiveEditManager()
@@ -608,10 +605,6 @@ private:
 		LiveEditManager.Reset();
 	}
 	
-	void OnPostGarbageCollection()
-	{
-	}
-
 private:
 
 	TSharedPtr<UE::Cameras::FGameplayCamerasLiveEditManager> LiveEditManager;
diff --git a/Engine/Plugins/Cameras/GameplayCameras/Source/GameplayCamerasEditor/Private/GameplayCamerasLiveEditManager.cpp b/Engine/Plugins/Cameras/GameplayCameras/Source/GameplayCamerasEditor/Private/GameplayCamerasLiveEditManager.cpp
--- a/Engine/Plugins/Cameras/GameplayCameras/Source/GameplayCamerasEditor/Private/GameplayCamerasLiveEditManager.cpp
+++ b/Engine/Plugins/Cameras/GameplayCameras/Source/GameplayCamerasEditor/Private/GameplayCamerasLiveEditManager.cpp
@@ -3,6 +3,7 @@
 #include "GameplayCamerasLiveEditManager.h"
 
 #include "IGameplayCamerasLiveEditListener.h"
+#include "Misc/CoreDelegates.h"
 #include "UObject/Package.h"
 #include "UObject/UObjectGlobals.h"
 
@@ -11,6 +12,12 @@ namespace UE::Cameras
 
 FGameplayCamerasLiveEditManager::FGameplayCamerasLiveEditManager()
 {
+	FCoreUObjectDelegates::GetPostGarbageCollect().AddRaw(this, &FGameplayCamerasLiveEditManager::OnPostGarbageCollection);
+}
+
+FGameplayCamerasLiveEditManager::~FGameplayCamerasLiveEditManager()
+{
+	FCoreUObjectDelegates::GetPostGarbageCollect().RemoveAll(this);
 }
 
 void FGameplayCamerasLiveEditManager::NotifyPostBuildAsset(const UPackage* InAssetPackage) const
@@ -53,6 +60,11 @@ void FGameplayCamerasLiveEditManager::RemoveListener(const UPackage* InAssetPack
 	}
 }
 
+void FGameplayCamerasLiveEditManager::OnPostGarbageCollection()
+{
+	RemoveGarbage();
+}
+
 void FGameplayCamerasLiveEditManager::RemoveGarbage()
 {
 	for (auto It = ListenerMap.CreateIterator(); It; ++It)
diff --git a/Engine/Plugins/Cameras/GameplayCameras/Source/GameplayCamerasEditor/Private/GameplayCamerasLiveEditManager.h b/Engine/Plugins/Cameras/GameplayCameras/Source/GameplayCamerasEditor/Private/GameplayCamerasLiveEditManager.h
--- a/Engine/Plugins/Cameras/GameplayCameras/Source/GameplayCamerasEditor/Private/GameplayCamerasLiveEditManager.h
+++ b/Engine/Plugins/Cameras/GameplayCameras/Source/GameplayCamerasEditor/Private/GameplayCamerasLiveEditManager.h
@@ -17,6 +17,7 @@ class FGameplayCamerasLiveEditManager : public IGameplayCamerasLiveEditManager
 public:
 
 	FGameplayCamerasLiveEditManager();
+	~FGameplayCamerasLiveEditManager();
 
 public:
 
@@ -27,6 +28,8 @@ public:
 
 private:
 
+	void OnPostGarbageCollection();
+
 	void RemoveGarbage();
 
 private:
