Fix filtering of gamplay camera tracks in rewind debugger

diff --git a/Engine/Plugins/Cameras/GameplayCameras/Source/GameplayCamerasEditor/Private/Trace/CameraSystemRewindDebuggerTrack.cpp b/Engine/Plugins/Cameras/GameplayCameras/Source/GameplayCamerasEditor/Private/Trace/CameraSystemRewindDebuggerTrack.cpp
--- a/Engine/Plugins/Cameras/GameplayCameras/Source/GameplayCamerasEditor/Private/Trace/CameraSystemRewindDebuggerTrack.cpp
+++ b/Engine/Plugins/Cameras/GameplayCameras/Source/GameplayCamerasEditor/Private/Trace/CameraSystemRewindDebuggerTrack.cpp
@@ -136,9 +136,18 @@ FName FCameraSystemRewindDebuggerTrackCreator::GetTargetTypeNameInternal() const
 	return TargetTypeName;
 }
 
-FName FCameraSystemRewindDebuggerTrackCreator::GetNameInternal() const
+namespace
 {
 	static const FName TrackName("Gameplay Camera System");
+}
+	
+FName FCameraSystemRewindDebuggerTrackCreator::GetNameInternal() const
+{
+	return TrackName;
+}
+	
+FName FCameraSystemRewindDebuggerTrack::GetNameInternal() const
+{
 	return TrackName;
 }
 
diff --git a/Engine/Plugins/Cameras/GameplayCameras/Source/GameplayCamerasEditor/Private/Trace/CameraSystemRewindDebuggerTrack.h b/Engine/Plugins/Cameras/GameplayCameras/Source/GameplayCamerasEditor/Private/Trace/CameraSystemRewindDebuggerTrack.h
--- a/Engine/Plugins/Cameras/GameplayCameras/Source/GameplayCamerasEditor/Private/Trace/CameraSystemRewindDebuggerTrack.h
+++ b/Engine/Plugins/Cameras/GameplayCameras/Source/GameplayCamerasEditor/Private/Trace/CameraSystemRewindDebuggerTrack.h
@@ -43,6 +43,7 @@ protected:
 
 	virtual FSlateIcon GetIconInternal() override;
 	virtual FText GetDisplayNameInternal() const override;
+	virtual FName GetNameInternal() const override;
 	virtual TSharedPtr<SWidget> GetTimelineViewInternal() override;
 
 private:
