Cameras: remove asserts on initial context being invalid, enable camera-cut flag when that result becomes valid again.

- When a context's initial result is invalid, we already print a warning at the beginning of the blend stack update, so it should be enough and further asserts aren't required.
- When that result becomes valid again, we don't want to run things like damping between the old (invalid) position and the new (valid) one, so set the camera cut flag.

diff --git a/Engine/Plugins/Cameras/GameplayCameras/Source/GameplayCameras/Private/Core/BlendStackCameraNode.cpp b/Engine/Plugins/Cameras/GameplayCameras/Source/GameplayCameras/Private/Core/BlendStackCameraNode.cpp
--- a/Engine/Plugins/Cameras/GameplayCameras/Source/GameplayCameras/Private/Core/BlendStackCameraNode.cpp
+++ b/Engine/Plugins/Cameras/GameplayCameras/Source/GameplayCameras/Private/Core/BlendStackCameraNode.cpp
@@ -101,6 +101,7 @@ bool FBlendStackCameraNodeEvaluator::InitializeEntry(
 	NewEntry.CameraRig = CameraRig;
 	NewEntry.RootNode = EntryRootNode;
 	NewEntry.RootEvaluator = RootEvaluator->CastThisChecked<FBlendStackRootCameraNodeEvaluator>();
+	NewEntry.bWasContextInitialResultValid = EvaluationContext->GetInitialResult().bIsValid;
 	NewEntry.bIsFirstFrame = true;
 
 	return true;
@@ -207,6 +208,14 @@ void FBlendStackCameraNodeEvaluator::ResolveEntries(TArray<FResolvedEntry>& OutR
 				continue;
 			}
 
+			// If the context was previously invalid, and this isn't the first frame, flag
+			// this update as a camera cut.
+			if (UNLIKELY(!Entry.bWasContextInitialResultValid && !Entry.bIsFirstFrame))
+			{
+				Entry.bForceCameraCut = true;
+			}
+			Entry.bWasContextInitialResultValid = true;
+
 			// Reset this entry's flags for this frame.
 			FCameraNodeEvaluationResult& CurResult = Entry.Result;
 			CurResult.CameraPose.ClearAllChangedFlags();
@@ -232,6 +241,7 @@ void FBlendStackCameraNodeEvaluator::OnRunFinished()
 		Entry.bIsFirstFrame = false;
 		Entry.bInputRunThisFrame = false;
 		Entry.bBlendRunThisFrame = false;
+		Entry.bForceCameraCut = false;
 	}
 }
 
@@ -739,7 +749,7 @@ void FTransientBlendStackCameraNodeEvaluator::InternalUpdate(TArrayView<FResolve
 		// Override it with whatever the evaluation context has set on its result.
 		const FCameraNodeEvaluationResult& ContextResult(ResolvedEntry.Context->GetInitialResult());
 		CurResult.CameraPose.OverrideChanged(ContextResult.CameraPose);
-		CurResult.bIsCameraCut = OutResult.bIsCameraCut || ContextResult.bIsCameraCut;
+		CurResult.bIsCameraCut = OutResult.bIsCameraCut || ContextResult.bIsCameraCut || Entry.bForceCameraCut;
 		CurResult.bIsValid = true;
 
 		// Run the camera rig's root node.
@@ -1007,7 +1017,7 @@ void FPersistentBlendStackCameraNodeEvaluator::InternalUpdate(TArrayView<FResolv
 				CurResult.VariableTable.Override(ContextResult.VariableTable, ECameraVariableTableFilter::AllPublic | ECameraVariableTableFilter::Private);
 
 				// Setup flags.
-				CurResult.bIsCameraCut = OutResult.bIsCameraCut || ContextResult.bIsCameraCut;
+				CurResult.bIsCameraCut = OutResult.bIsCameraCut || ContextResult.bIsCameraCut || Entry.bForceCameraCut;
 				CurResult.bIsValid = true;
 			}
 
diff --git a/Engine/Plugins/Cameras/GameplayCameras/Source/GameplayCameras/Private/Math/CameraNodeSpaceMath.cpp b/Engine/Plugins/Cameras/GameplayCameras/Source/GameplayCameras/Private/Math/CameraNodeSpaceMath.cpp
--- a/Engine/Plugins/Cameras/GameplayCameras/Source/GameplayCameras/Private/Math/CameraNodeSpaceMath.cpp
+++ b/Engine/Plugins/Cameras/GameplayCameras/Source/GameplayCameras/Private/Math/CameraNodeSpaceMath.cpp
@@ -70,7 +70,6 @@ bool FCameraNodeSpaceMath::GetCameraNodeOriginPosition(const FCameraNodeSpacePar
 			if (TSharedPtr<const FCameraEvaluationContext> ActiveContext = Params.GetActiveContext())
 			{
 				const FCameraNodeEvaluationResult& InitialResult = ActiveContext->GetInitialResult();
-				ensure(InitialResult.bIsValid);
 				OutOrigin = InitialResult.CameraPose.GetLocation();
 				return true;
 			}
@@ -79,7 +78,6 @@ bool FCameraNodeSpaceMath::GetCameraNodeOriginPosition(const FCameraNodeSpacePar
 			if (TSharedPtr<const FCameraEvaluationContext> OwningContext = Params.GetOwningContext())
 			{
 				const FCameraNodeEvaluationResult& InitialResult = OwningContext->GetInitialResult();
-				ensure(InitialResult.bIsValid);
 				OutOrigin = InitialResult.CameraPose.GetLocation();
 				return true;
 			}
@@ -95,7 +93,6 @@ bool FCameraNodeSpaceMath::GetCameraNodeOriginPosition(const FCameraNodeSpacePar
 				else if (TSharedPtr<const FCameraEvaluationContext> ActiveContext = Params.GetActiveContext())
 				{
 					const FCameraNodeEvaluationResult& InitialResult = ActiveContext->GetInitialResult();
-					ensure(InitialResult.bIsValid);
 					OutOrigin = InitialResult.CameraPose.GetLocation();
 					return true;
 				}
@@ -141,7 +138,6 @@ bool FCameraNodeSpaceMath::GetCameraNodeSpaceTransform(const FCameraNodeSpacePar
 			if (TSharedPtr<const FCameraEvaluationContext> ActiveContext = Params.GetActiveContext())
 			{
 				const FCameraNodeEvaluationResult& InitialResult = ActiveContext->GetInitialResult();
-				ensure(InitialResult.bIsValid);
 				OutTransform = InitialResult.CameraPose.GetTransform();
 				return true;
 			}
@@ -150,7 +146,6 @@ bool FCameraNodeSpaceMath::GetCameraNodeSpaceTransform(const FCameraNodeSpacePar
 			if (TSharedPtr<const FCameraEvaluationContext> OwningContext = Params.GetOwningContext())
 			{
 				const FCameraNodeEvaluationResult& InitialResult = OwningContext->GetInitialResult();
-				ensure(InitialResult.bIsValid);
 				OutTransform = InitialResult.CameraPose.GetTransform();
 				return true;
 			}
@@ -166,7 +161,6 @@ bool FCameraNodeSpaceMath::GetCameraNodeSpaceTransform(const FCameraNodeSpacePar
 				else if (TSharedPtr<const FCameraEvaluationContext> ActiveContext = Params.GetActiveContext())
 				{
 					const FCameraNodeEvaluationResult& InitialResult = ActiveContext->GetInitialResult();
-					ensure(InitialResult.bIsValid);
 					OutTransform = InitialResult.CameraPose.GetTransform();
 					return true;
 				}
@@ -222,7 +216,6 @@ bool FCameraNodeSpaceMath::OffsetCameraNodeSpacePosition(const FCameraNodeSpaceP
 			if (TSharedPtr<const FCameraEvaluationContext> ActiveContext = Params.GetActiveContext())
 			{
 				const FCameraNodeEvaluationResult& InitialResult = ActiveContext->GetInitialResult();
-				ensure(InitialResult.bIsValid);
 				WorldOffset = InitialResult.CameraPose.GetRotation().RotateVector(InOffset);
 				bGotWorldOffset = true;
 			}
@@ -231,7 +224,6 @@ bool FCameraNodeSpaceMath::OffsetCameraNodeSpacePosition(const FCameraNodeSpaceP
 			if (TSharedPtr<const FCameraEvaluationContext> OwningContext = Params.GetOwningContext())
 			{
 				const FCameraNodeEvaluationResult& InitialResult = OwningContext->GetInitialResult();
-				ensure(InitialResult.bIsValid);
 				WorldOffset = InitialResult.CameraPose.GetRotation().RotateVector(InOffset);
 				bGotWorldOffset = true;
 			}
@@ -247,7 +239,6 @@ bool FCameraNodeSpaceMath::OffsetCameraNodeSpacePosition(const FCameraNodeSpaceP
 				else if (TSharedPtr<const FCameraEvaluationContext> ActiveContext = Params.GetActiveContext())
 				{
 					const FCameraNodeEvaluationResult& InitialResult = ActiveContext->GetInitialResult();
-					ensure(InitialResult.bIsValid);
 					WorldOffset = InitialResult.CameraPose.GetRotation().RotateVector(InOffset);
 					bGotWorldOffset = true;
 				}
diff --git a/Engine/Plugins/Cameras/GameplayCameras/Source/GameplayCameras/Private/Nodes/Common/DampenPositionCameraNode.cpp b/Engine/Plugins/Cameras/GameplayCameras/Source/GameplayCameras/Private/Nodes/Common/DampenPositionCameraNode.cpp
--- a/Engine/Plugins/Cameras/GameplayCameras/Source/GameplayCameras/Private/Nodes/Common/DampenPositionCameraNode.cpp
+++ b/Engine/Plugins/Cameras/GameplayCameras/Source/GameplayCameras/Private/Nodes/Common/DampenPositionCameraNode.cpp
@@ -194,9 +194,6 @@ void FDampenPositionCameraNodeEvaluator::ComputeAxisDampers(const FCameraNodeEva
 			if (Params.EvaluationContext)
 			{
 				const FCameraNodeEvaluationResult& InitialResult = Params.EvaluationContext->GetInitialResult();
-				ensureMsgf(InitialResult.bIsValid,
-						TEXT("DampenPositionCameraNode: using invalid context result as damping space!"));
-
 				const FRotator3d ContextRotation = InitialResult.CameraPose.GetRotation();
 				AxesRotation = ContextRotation;
 			}
diff --git a/Engine/Plugins/Cameras/GameplayCameras/Source/GameplayCameras/Private/Nodes/Common/OffsetCameraNode.cpp b/Engine/Plugins/Cameras/GameplayCameras/Source/GameplayCameras/Private/Nodes/Common/OffsetCameraNode.cpp
--- a/Engine/Plugins/Cameras/GameplayCameras/Source/GameplayCameras/Private/Nodes/Common/OffsetCameraNode.cpp
+++ b/Engine/Plugins/Cameras/GameplayCameras/Source/GameplayCameras/Private/Nodes/Common/OffsetCameraNode.cpp
@@ -55,12 +55,9 @@ void FOffsetCameraNodeEvaluator::OnRun(const FCameraNodeEvaluationParams& Params
 		case ECameraNodeSpace::OwningContext:
 			if (Params.EvaluationContext)
 			{ 
-				const FCameraNodeEvaluationResult& InitialResult = Params.EvaluationContext->GetInitialResult();
-				ensureMsgf(InitialResult.bIsValid,
-						TEXT("OffsetCameraNode: using invalid context result as offset space!"));
-
 				// The offsets are meant to be treated as context-local. Let's get the context transform
 				// and apply the offsets using that transform's axes.
+				const FCameraNodeEvaluationResult& InitialResult = Params.EvaluationContext->GetInitialResult();
 				const FTransform3d ContextTransform = InitialResult.CameraPose.GetTransform();
 
 				const FVector3d WorldTranslationOffset = ContextTransform.TransformVector(TranslationOffset);
diff --git a/Engine/Plugins/Cameras/GameplayCameras/Source/GameplayCameras/Public/Core/BlendStackCameraNode.h b/Engine/Plugins/Cameras/GameplayCameras/Source/GameplayCameras/Public/Core/BlendStackCameraNode.h
--- a/Engine/Plugins/Cameras/GameplayCameras/Source/GameplayCameras/Public/Core/BlendStackCameraNode.h
+++ b/Engine/Plugins/Cameras/GameplayCameras/Source/GameplayCameras/Public/Core/BlendStackCameraNode.h
@@ -184,10 +184,14 @@ protected:
 		FCameraNodeEvaluationResult Result;
 		/** Whether this is the first frame this entry runs. */
 		bool bIsFirstFrame = false;
+		/** Whether the context's initial result was valid last frame. */
+		bool bWasContextInitialResultValid = false;
 		/** Whether input slots were run (possibly from a preview update). */
 		bool bInputRunThisFrame = false;
 		/** Whether the blend node was run (possibly from a preview update). */
 		bool bBlendRunThisFrame = false;
+		/** Whether to force a camera cut on this entry this frame. */
+		bool bForceCameraCut = false;
 		/** Whether this entry is frozen. */
 		bool bIsFrozen = false;
 
