Cameras: fix possible crash when getting the active rig info on a frozen rig

diff --git a/Engine/Plugins/Cameras/GameplayCameras/Source/GameplayCameras/Private/Core/BlendStackCameraNode.cpp b/Engine/Plugins/Cameras/GameplayCameras/Source/GameplayCameras/Private/Core/BlendStackCameraNode.cpp
--- a/Engine/Plugins/Cameras/GameplayCameras/Source/GameplayCameras/Private/Core/BlendStackCameraNode.cpp
+++ b/Engine/Plugins/Cameras/GameplayCameras/Source/GameplayCameras/Private/Core/BlendStackCameraNode.cpp
@@ -134,7 +134,7 @@ FCameraRigEvaluationInfo FBlendStackCameraNodeEvaluator::GetActiveCameraRigEvalu
 				ActiveEntry.EvaluationContext.Pin(),
 				ActiveEntry.CameraRig, 
 				&ActiveEntry.Result,
-				ActiveEntry.RootEvaluator->GetRootEvaluator());
+				ActiveEntry.RootEvaluator ? ActiveEntry.RootEvaluator->GetRootEvaluator() : nullptr);
 		return Info;
 	}
 	return FCameraRigEvaluationInfo();
