Cameras: add API for getting the active camera rig info on the root node evaluator

diff --git a/Engine/Plugins/Cameras/GameplayCameras/Source/GameplayCameras/Private/Core/DefaultRootCameraNode.cpp b/Engine/Plugins/Cameras/GameplayCameras/Source/GameplayCameras/Private/Core/DefaultRootCameraNode.cpp
--- a/Engine/Plugins/Cameras/GameplayCameras/Source/GameplayCameras/Private/Core/DefaultRootCameraNode.cpp
+++ b/Engine/Plugins/Cameras/GameplayCameras/Source/GameplayCameras/Private/Core/DefaultRootCameraNode.cpp
@@ -147,6 +147,11 @@ void FDefaultRootCameraNodeEvaluator::OnDeactivateCameraRig(const FDeactivateCam
 	}
 }
 
+void FDefaultRootCameraNodeEvaluator::OnGetActiveCameraRigInfo(FCameraRigEvaluationInfo& OutCameraRigInfo) const
+{
+	OutCameraRigInfo = MainLayer->GetActiveCameraRigEvaluationInfo();
+}
+
 void FDefaultRootCameraNodeEvaluator::OnBuildSingleCameraRigHierarchy(const FSingleCameraRigHierarchyBuildParams& Params, FCameraNodeEvaluatorHierarchy& OutHierarchy)
 {
 	OutHierarchy.Build(BaseLayer);
diff --git a/Engine/Plugins/Cameras/GameplayCameras/Source/GameplayCameras/Private/Core/RootCameraNode.cpp b/Engine/Plugins/Cameras/GameplayCameras/Source/GameplayCameras/Private/Core/RootCameraNode.cpp
--- a/Engine/Plugins/Cameras/GameplayCameras/Source/GameplayCameras/Private/Core/RootCameraNode.cpp
+++ b/Engine/Plugins/Cameras/GameplayCameras/Source/GameplayCameras/Private/Core/RootCameraNode.cpp
@@ -27,6 +27,11 @@ void FRootCameraNodeEvaluator::DeactivateCameraRig(const FDeactivateCameraRigPar
 	OnDeactivateCameraRig(Params);
 }
 
+void FRootCameraNodeEvaluator::GetActiveCameraRigInfo(FCameraRigEvaluationInfo& OutCameraRigInfo) const
+{
+	OnGetActiveCameraRigInfo(OutCameraRigInfo);
+}
+
 void FRootCameraNodeEvaluator::BuildSingleCameraRigHierarchy(const FSingleCameraRigHierarchyBuildParams& Params, FCameraNodeEvaluatorHierarchy& OutHierarchy)
 {
 	OnBuildSingleCameraRigHierarchy(Params, OutHierarchy);
diff --git a/Engine/Plugins/Cameras/GameplayCameras/Source/GameplayCameras/Public/Core/CameraRigEvaluationInfo.h b/Engine/Plugins/Cameras/GameplayCameras/Source/GameplayCameras/Public/Core/CameraRigEvaluationInfo.h
--- a/Engine/Plugins/Cameras/GameplayCameras/Source/GameplayCameras/Public/Core/CameraRigEvaluationInfo.h
+++ b/Engine/Plugins/Cameras/GameplayCameras/Source/GameplayCameras/Public/Core/CameraRigEvaluationInfo.h
@@ -28,7 +28,7 @@ struct FCameraRigEvaluationInfo
 	TObjectPtr<const UCameraRigAsset> CameraRig;
 	/** The last evaluated result for this camera rig. */
 	const FCameraNodeEvaluationResult* LastResult = nullptr;
-	/** The blend node evaluator of the camera rig. */
+	/** The root node evaluator of the camera rig. */
 	FCameraNodeEvaluator* RootEvaluator = nullptr;
 
 	FCameraRigEvaluationInfo()
diff --git a/Engine/Plugins/Cameras/GameplayCameras/Source/GameplayCameras/Public/Core/DefaultRootCameraNode.h b/Engine/Plugins/Cameras/GameplayCameras/Source/GameplayCameras/Public/Core/DefaultRootCameraNode.h
--- a/Engine/Plugins/Cameras/GameplayCameras/Source/GameplayCameras/Public/Core/DefaultRootCameraNode.h
+++ b/Engine/Plugins/Cameras/GameplayCameras/Source/GameplayCameras/Public/Core/DefaultRootCameraNode.h
@@ -67,6 +67,7 @@ protected:
 	// FRootCameraNodeEvaluator interface.
 	virtual void OnActivateCameraRig(const FActivateCameraRigParams& Params) override;
 	virtual void OnDeactivateCameraRig(const FDeactivateCameraRigParams& Params) override;
+	virtual void OnGetActiveCameraRigInfo(FCameraRigEvaluationInfo& OutCameraRigInfo) const override;
 	virtual void OnBuildSingleCameraRigHierarchy(const FSingleCameraRigHierarchyBuildParams& Params, FCameraNodeEvaluatorHierarchy& OutHierarchy) override;
 	virtual void OnRunSingleCameraRig(const FSingleCameraRigEvaluationParams& Params, FCameraNodeEvaluationResult& OutResult) override;
 
diff --git a/Engine/Plugins/Cameras/GameplayCameras/Source/GameplayCameras/Public/Core/RootCameraNode.h b/Engine/Plugins/Cameras/GameplayCameras/Source/GameplayCameras/Public/Core/RootCameraNode.h
--- a/Engine/Plugins/Cameras/GameplayCameras/Source/GameplayCameras/Public/Core/RootCameraNode.h
+++ b/Engine/Plugins/Cameras/GameplayCameras/Source/GameplayCameras/Public/Core/RootCameraNode.h
@@ -109,6 +109,9 @@ public:
 	/** Deactivates a camera rig. */
 	void DeactivateCameraRig(const FDeactivateCameraRigParams& Params);
 
+	/** Gets information about the active camera rig in the main layer. */
+	void GetActiveCameraRigInfo(FCameraRigEvaluationInfo& OutCameraRigInfo) const;
+
 	/**
 	 * Builds the hierarchy of the system for a given single camera rig.
 	 * This is expected to return the nodes of all the layers, except for the main layer which
@@ -140,6 +143,9 @@ protected:
 	/** Deactivates a camera rig. */
 	virtual void OnDeactivateCameraRig(const FDeactivateCameraRigParams& Params) {}
 
+	/** Gets information about the active camera rig in the main layer. */
+	virtual void OnGetActiveCameraRigInfo(FCameraRigEvaluationInfo& OutCameraRigInfo) const {}
+
 	/* Builds the hierarchy of the system for a given single camera rig. */
 	virtual void OnBuildSingleCameraRigHierarchy(const FSingleCameraRigHierarchyBuildParams& Params, FCameraNodeEvaluatorHierarchy& OutHierarchy) {}
 
