Cameras: add test object builders for camera assets and directors, add tests for update flags

#rnx

diff --git a/Engine/Plugins/Cameras/GameplayCameras/Source/GameplayCameras/Private/Tests/CameraSystemUpdateTests.cpp b/Engine/Plugins/Cameras/GameplayCameras/Source/GameplayCameras/Private/Tests/CameraSystemUpdateTests.cpp
new file mode 100644
--- /dev/null
+++ b/Engine/Plugins/Cameras/GameplayCameras/Source/GameplayCameras/Private/Tests/CameraSystemUpdateTests.cpp
@@ -0,0 +1,82 @@
+// Copyright Epic Games, Inc. All Rights Reserved.
+
+#include "Core/CameraRigAsset.h"
+#include "Core/CameraSystemEvaluator.h"
+#include "Core/RootCameraNode.h"
+#include "Misc/AutomationTest.h"
+#include "Tests/GameplayCamerasTestBuilder.h"
+#include "Tests/GameplayCamerasTestObjects.h"
+
+#define LOCTEXT_NAMESPACE "CameraSystemUpdateTests"
+
+IMPLEMENT_SIMPLE_AUTOMATION_TEST(FCameraSystemFrameFlagsTest, "System.Engine.GameplayCameras.CameraSystem.FrameFlags", EAutomationTestFlags::EditorContext | EAutomationTestFlags::EngineFilter)
+bool FCameraSystemFrameFlagsTest::RunTest(const FString& Parameters)
+{
+	using namespace UE::Cameras;
+	using namespace UE::Cameras::Test;
+
+	TSharedRef<FCameraEvaluationContext> EvaluationContext = FCameraEvaluationContextTestBuilder()
+		.AddCameraRig(TEXT("TestRig"))
+			.MakeRootNode<UUpdateTrackerCameraNode>()
+				.Done()
+			.Done()
+		.MakeSingleDirector()
+			.Setup([](USingleCameraDirector* Director, FNamedObjectRegistry* NamedObjectRegistry)
+				{
+					Director->CameraRig = NamedObjectRegistry->Get<UCameraRigAsset>(TEXT("TestRig"));
+				})
+			.Done()
+		.BuildCameraAsset()
+		.Get();
+
+	FCameraSystemEvaluator Evaluator;
+	Evaluator.Initialize();
+	Evaluator.PushEvaluationContext(EvaluationContext);
+
+	FCameraSystemEvaluationParams Params;
+	Params.DeltaTime = 0.3f;
+	Evaluator.Update(Params);
+
+	FCameraRigEvaluationInfo CameraRigInfo;
+	FRootCameraNodeEvaluator* RootEvaluator = Evaluator.GetRootNodeEvaluator();
+	RootEvaluator->GetActiveCameraRigInfo(CameraRigInfo);
+	UTEST_NOT_NULL("RootEvaluator", CameraRigInfo.RootEvaluator);
+
+	FUpdateTrackerCameraNodeEvaluator* UpdateTracker = CameraRigInfo.RootEvaluator->CastThis<FUpdateTrackerCameraNodeEvaluator>();
+	UTEST_NOT_NULL("UpdateTracker", UpdateTracker);
+
+	{
+		UTEST_EQUAL("NumReceivedUpdates", UpdateTracker->ReceivedUpdates.Num(), 1);
+		FTrackedUpdateInfo Update = UpdateTracker->ReceivedUpdates[0];
+		UTEST_EQUAL("DeltaTime", Update.DeltaTime, 0.3f);
+		UTEST_FALSE("IsCameraCut", Update.bIsCameraCut);
+		UTEST_TRUE("IsFirstFrame", Update.bIsFirstFrame);
+	}
+
+	Evaluator.Update(Params);
+
+	{
+		UTEST_EQUAL("NumReceivedUpdates", UpdateTracker->ReceivedUpdates.Num(), 2);
+		FTrackedUpdateInfo Update = UpdateTracker->ReceivedUpdates[1];
+		UTEST_EQUAL("DeltaTime", Update.DeltaTime, 0.3f);
+		UTEST_FALSE("IsCameraCut", Update.bIsCameraCut);
+		UTEST_FALSE("IsFirstFrame", Update.bIsFirstFrame);
+	}
+
+	EvaluationContext->GetInitialResult().bIsCameraCut = true;
+
+	Evaluator.Update(Params);
+
+	{
+		UTEST_EQUAL("NumReceivedUpdates", UpdateTracker->ReceivedUpdates.Num(), 3);
+		FTrackedUpdateInfo Update = UpdateTracker->ReceivedUpdates[2];
+		UTEST_EQUAL("DeltaTime", Update.DeltaTime, 0.3f);
+		UTEST_TRUE("IsCameraCut", Update.bIsCameraCut);
+		UTEST_FALSE("IsFirstFrame", Update.bIsFirstFrame);
+	}
+
+	return true;
+}
+
+#undef LOCTEXT_NAMESPACE
+
diff --git a/Engine/Plugins/Cameras/GameplayCameras/Source/GameplayCameras/Private/Tests/GameplayCamerasTestBuilder.cpp b/Engine/Plugins/Cameras/GameplayCameras/Source/GameplayCameras/Private/Tests/GameplayCamerasTestBuilder.cpp
--- a/Engine/Plugins/Cameras/GameplayCameras/Source/GameplayCameras/Private/Tests/GameplayCamerasTestBuilder.cpp
+++ b/Engine/Plugins/Cameras/GameplayCameras/Source/GameplayCameras/Private/Tests/GameplayCamerasTestBuilder.cpp
@@ -6,62 +6,32 @@ namespace UE::Cameras::Test
 {
 
 FCameraRigAssetTestBuilder::FCameraRigAssetTestBuilder(FName Name, UObject* Outer)
+	: TCameraRigAssetTestBuilderBase<FCameraRigAssetTestBuilder>(nullptr, Name, Outer)
 {
-	Initialize(nullptr, Name, Outer);
 }
 
-FCameraRigAssetTestBuilder::FCameraRigAssetTestBuilder(TSharedRef<FNamedObjectRegistry> InNamedObjectRegistry, FName Name, UObject* Outer)
+FCameraRigAssetTestBuilder::FCameraRigAssetTestBuilder(TSharedPtr<FNamedObjectRegistry> InNamedObjectRegistry, FName Name, UObject* Outer)
+	: TCameraRigAssetTestBuilderBase<FCameraRigAssetTestBuilder>(InNamedObjectRegistry, Name, Outer)
 {
-	Initialize(InNamedObjectRegistry, Name, Outer);
 }
 
-void FCameraRigAssetTestBuilder::Initialize(TSharedPtr<FNamedObjectRegistry> InNamedObjectRegistry, FName Name, UObject* Outer)
+FCameraEvaluationContextTestBuilder::FCameraEvaluationContextTestBuilder(UObject* Owner)
 {
-	if (Outer == nullptr)
+	if (Owner == nullptr)
 	{
-		Outer = GetTransientPackage();
+		Owner = GetTransientPackage();
 	}
 
-	CameraRig = NewObject<UCameraRigAsset>(Outer, Name);
-	TCameraObjectInitializer<UCameraRigAsset>::SetObject(CameraRig);
+	CameraAsset = NewObject<UCameraAsset>(Owner);
 
-	NamedObjectRegistry = InNamedObjectRegistry;
-	if (!NamedObjectRegistry)
-	{
-		NamedObjectRegistry = MakeShared<FNamedObjectRegistry>();
-	}
-}
-
-TCameraRigTransitionTestBuilder<FCameraRigAssetTestBuilder> FCameraRigAssetTestBuilder::AddEnterTransition()
-{
-	TCameraRigTransitionTestBuilder<ThisType> TransitionBuilder(*this, CameraRig);
-	CameraRig->EnterTransitions.Add(TransitionBuilder.Get());
-	return TransitionBuilder;
-}
+	FCameraEvaluationContextInitializeParams InitParams;
+	InitParams.Owner = Owner;
+	InitParams.CameraAsset = CameraAsset;
+	EvaluationContext = MakeShared<FCameraEvaluationContext>(InitParams);
 
-TCameraRigTransitionTestBuilder<FCameraRigAssetTestBuilder> FCameraRigAssetTestBuilder::AddExitTransition()
-{
-	TCameraRigTransitionTestBuilder<ThisType> TransitionBuilder(*this, CameraRig);
-	CameraRig->ExitTransitions.Add(TransitionBuilder.Get());
-	return TransitionBuilder;
-}
-
-FCameraRigAssetTestBuilder& FCameraRigAssetTestBuilder::ExposeParameter(const FString& ParameterName, UCameraNode* Target, FName TargetPropertyName)
-{
-	UCameraRigInterfaceParameter* InterfaceParameter = NewObject<UCameraRigInterfaceParameter>(CameraRig);
-	InterfaceParameter->InterfaceParameterName = ParameterName;
-	InterfaceParameter->Target = Target;
-	InterfaceParameter->TargetPropertyName = TargetPropertyName;
-	NamedObjectRegistry->Register(InterfaceParameter, ParameterName);
-	CameraRig->Interface.InterfaceParameters.Add(InterfaceParameter);
-	return *this;
-}
+	TCameraObjectInitializer<FCameraEvaluationContext>::SetObject(EvaluationContext.Get());
 
-FCameraRigAssetTestBuilder& FCameraRigAssetTestBuilder::ExposeParameter(const FString& ParameterName, const FString& TargetName, FName TargetPropertyName)
-{
-	UCameraNode* Target = NamedObjectRegistry->Get<UCameraNode>(TargetName);
-	ensure(Target);
-	return ExposeParameter(ParameterName, Target, TargetPropertyName);
+	NamedObjectRegistry = MakeShared<FNamedObjectRegistry>();
 }
 
 }  // namespace UE::Cameras::Test
diff --git a/Engine/Plugins/Cameras/GameplayCameras/Source/GameplayCameras/Private/Tests/GameplayCamerasTestBuilder.h b/Engine/Plugins/Cameras/GameplayCameras/Source/GameplayCameras/Private/Tests/GameplayCamerasTestBuilder.h
--- a/Engine/Plugins/Cameras/GameplayCameras/Source/GameplayCameras/Private/Tests/GameplayCamerasTestBuilder.h
+++ b/Engine/Plugins/Cameras/GameplayCameras/Source/GameplayCameras/Private/Tests/GameplayCamerasTestBuilder.h
@@ -3,9 +3,13 @@
 #pragma once
 
 #include "Core/BlendCameraNode.h"
+#include "Core/CameraAsset.h"
+#include "Core/CameraDirector.h"
+#include "Core/CameraEvaluationContext.h"
 #include "Core/CameraNode.h"
 #include "Core/CameraRigAsset.h"
 #include "Core/CameraRigTransition.h"
+#include "Directors/SingleCameraDirector.h"
 #include "Nodes/Common/ArrayCameraNode.h"
 #include "Templates/PointerIsConvertibleFromTo.h"
 #include "Templates/UnrealTypeTraits.h"
@@ -16,7 +20,7 @@
 namespace UE::Cameras::Test
 {
 
-class FCameraRigAssetTestBuilder;
+class FCameraEvaluationContextTestBuilder;
 
 /**
  * Template mix-in for adding "go back to parent" support to a builder class.
@@ -76,7 +80,7 @@ private:
 /**
  * A simple repository matching UObject instances to names.
  */
-class FNamedObjectRegistry
+class FNamedObjectRegistry : public TSharedFromThis<FNamedObjectRegistry>
 {
 public:
 
@@ -116,7 +120,7 @@ struct IHasNamedObjectRegistry
 {
 	virtual ~IHasNamedObjectRegistry() {}
 
-	virtual FNamedObjectRegistry* GetNamedObjectRegistry() = 0;
+	virtual TSharedPtr<FNamedObjectRegistry> GetNamedObjectRegistry() = 0;
 };
 
 /**
@@ -158,7 +162,7 @@ public:
 	template<typename VV = std::enable_if_t<TPointerIsConvertibleFromTo<ParentType, IHasNamedObjectRegistry>::Value>>
 	ThisType& Named(const TCHAR* InName)
 	{
-		FNamedObjectRegistry* NamedObjectRegistry = GetNamedObjectRegistry();
+		TSharedPtr<FNamedObjectRegistry> NamedObjectRegistry = GetNamedObjectRegistry();
 		if (ensure(NamedObjectRegistry))
 		{
 			NamedObjectRegistry->Register(CameraNode, InName);
@@ -191,8 +195,8 @@ public:
 	 */
 	ThisType& Setup(TFunction<void(NodeType*, FNamedObjectRegistry*)> SetupCallback)
 	{
-		FNamedObjectRegistry* NamedObjectRegistry = GetNamedObjectRegistry();
-		SetupCallback(CameraNode, NamedObjectRegistry);
+		TSharedPtr<FNamedObjectRegistry> NamedObjectRegistry = GetNamedObjectRegistry();
+		SetupCallback(CameraNode, NamedObjectRegistry.Get());
 		return *this;
 	}
 
@@ -250,7 +254,7 @@ public:
 	}
 
 	/** Gets the named object registry from the parent. */
-	virtual FNamedObjectRegistry* GetNamedObjectRegistry() override
+	virtual TSharedPtr<FNamedObjectRegistry> GetNamedObjectRegistry() override
 	{
 		if constexpr (TPointerIsConvertibleFromTo<ParentType, IHasNamedObjectRegistry>::Value)
 		{
@@ -314,7 +318,7 @@ public:
 	template<typename VV = std::enable_if_t<TPointerIsConvertibleFromTo<ParentType, IHasNamedObjectRegistry>::Value>>
 	ThisType& Named(const TCHAR* InName)
 	{
-		FNamedObjectRegistry* NamedObjectRegistry = GetNamedObjectRegistry();
+		TSharedPtr<FNamedObjectRegistry> NamedObjectRegistry = GetNamedObjectRegistry();
 		if (ensure(NamedObjectRegistry))
 		{
 			NamedObjectRegistry->Register(Transition, InName);
@@ -363,7 +367,7 @@ public:
 	}
 
 	/** Gets the named object registry from the parent. */
-	virtual FNamedObjectRegistry* GetNamedObjectRegistry() override
+	virtual TSharedPtr<FNamedObjectRegistry> GetNamedObjectRegistry() override
 	{
 		if constexpr (TPointerIsConvertibleFromTo<ParentType, IHasNamedObjectRegistry>::Value)
 		{
@@ -400,20 +404,33 @@ private:
  *				.Done()
  *			.Get();
  */
-class FCameraRigAssetTestBuilder 
+template<typename ThisType>
+class TCameraRigAssetTestBuilderBase 
 	: public TCameraObjectInitializer<UCameraRigAsset>
 	, public IHasNamedObjectRegistry
 {
 public:
 
-	using ThisType = FCameraRigAssetTestBuilder;
-
-	FCameraRigAssetTestBuilder(FName Name = NAME_None, UObject* Outer = nullptr);
-	FCameraRigAssetTestBuilder(TSharedRef<FNamedObjectRegistry> InNamedObjectRegistry, FName Name = NAME_None, UObject* Outer = nullptr);
-
 	/** Gets the built camera rig. */
 	UCameraRigAsset* Get() { return CameraRig; }
 
+	/** Pins the built camera rig to a given pointer, for being able to later refer to it. */
+	ThisType& Pin(UCameraRigAsset*& OutPtr)
+	{
+		OutPtr = CameraRig; 
+		return *static_cast<ThisType*>(this);
+	}
+
+	/** Give a name to the built camera rig, to be recalled later. */
+	ThisType& Named(const TCHAR* InName)
+	{
+		if (ensure(NamedObjectRegistry))
+		{
+			NamedObjectRegistry->Register(CameraRig, InName);
+		}
+		return *static_cast<ThisType*>(this);
+	}
+
 	/**
 	 * Creates a new camera node and sets it as the root node of the rig.
 	 * Returns the builder for the root camera node. You can come back to the rig builder
@@ -422,7 +439,8 @@ public:
 	template<typename NodeType>
 	TCameraNodeTestBuilder<ThisType, NodeType> MakeRootNode()
 	{
-		TCameraNodeTestBuilder<ThisType, NodeType> NodeBuilder(*this, CameraRig);
+		ThisType* ActualThis = static_cast<ThisType*>(this);
+		TCameraNodeTestBuilder<ThisType, NodeType> NodeBuilder(*ActualThis, CameraRig);
 		CameraRig->RootNode = NodeBuilder.Get();
 		return NodeBuilder;
 	}
@@ -439,12 +457,23 @@ public:
 	 * Adds a new enter transition and returns a builder for it. You can come back to the
 	 * rig builder by calling Done() on the transition builder.
 	 */
-	TCameraRigTransitionTestBuilder<ThisType> AddEnterTransition();
+	TCameraRigTransitionTestBuilder<ThisType> AddEnterTransition()
+	{
+		TCameraRigTransitionTestBuilder<ThisType> TransitionBuilder(*this, CameraRig);
+		CameraRig->EnterTransitions.Add(TransitionBuilder.Get());
+		return TransitionBuilder;
+	}
+
 	/**
 	 * Adds a new exit transition and returns a builder for it. You can come back to the
 	 * rig builder by calling Done() on the transition builder.
 	 */
-	TCameraRigTransitionTestBuilder<ThisType> AddExitTransition();
+	TCameraRigTransitionTestBuilder<ThisType> AddExitTransition()
+	{
+		TCameraRigTransitionTestBuilder<ThisType> TransitionBuilder(*this, CameraRig);
+		CameraRig->ExitTransitions.Add(TransitionBuilder.Get());
+		return TransitionBuilder;
+	}
 
 	/**
 	 * Creates a new exposed rig parameter and hooks it up to the given camera node's property.
@@ -453,24 +482,62 @@ public:
 	 *
 	 * The created parameter is automatically stored in the named object registry under its name.
 	 */
-	FCameraRigAssetTestBuilder& ExposeParameter(const FString& ParameterName, UCameraNode* Target, FName TargetPropertyName);
+	ThisType& ExposeParameter(const FString& ParameterName, UCameraNode* Target, FName TargetPropertyName)
+	{
+		UCameraRigInterfaceParameter* InterfaceParameter = NewObject<UCameraRigInterfaceParameter>(CameraRig);
+		InterfaceParameter->InterfaceParameterName = ParameterName;
+		InterfaceParameter->Target = Target;
+		InterfaceParameter->TargetPropertyName = TargetPropertyName;
+		NamedObjectRegistry->Register(InterfaceParameter, ParameterName);
+		CameraRig->Interface.InterfaceParameters.Add(InterfaceParameter);
+		return *static_cast<ThisType*>(this);
+	}
 
 	/**
 	 * A variant of ExposeParameter that retrieves the target node from the named registry.
 	 *
 	 * The created parameter is automatically stored in the named object registry under its name.
 	 */
-	FCameraRigAssetTestBuilder& ExposeParameter(const FString& ParameterName, const FString& TargetName, FName TargetPropertyName);
+	ThisType& ExposeParameter(const FString& ParameterName, const FString& TargetName, FName TargetPropertyName)
+	{
+		UCameraNode* Target = NamedObjectRegistry->Get<UCameraNode>(TargetName);
+		ensure(Target);
+		return ExposeParameter(ParameterName, Target, TargetPropertyName);
+	}
 
 	/** Gets the named object registry. */
-	virtual FNamedObjectRegistry* GetNamedObjectRegistry() override
+	virtual TSharedPtr<FNamedObjectRegistry> GetNamedObjectRegistry() override
+	{
+		return NamedObjectRegistry;
+	}
+
+protected:
+
+	TCameraRigAssetTestBuilderBase(TSharedPtr<FNamedObjectRegistry> InNamedObjectRegistry, FName Name = NAME_None, UObject* Outer = nullptr)
 	{
-		return NamedObjectRegistry.Get();
+		Initialize(InNamedObjectRegistry, Name, Outer);
 	}
 
 private:
 
-	void Initialize(TSharedPtr<FNamedObjectRegistry> InNamedObjectRegistry, FName Name, UObject* Outer);
+	void Initialize(TSharedPtr<FNamedObjectRegistry> InNamedObjectRegistry, FName Name, UObject* Outer)
+	{
+		if (Outer == nullptr)
+		{
+			Outer = GetTransientPackage();
+		}
+
+		CameraRig = NewObject<UCameraRigAsset>(Outer, Name);
+		TCameraObjectInitializer<UCameraRigAsset>::SetObject(CameraRig);
+
+		NamedObjectRegistry = InNamedObjectRegistry;
+		if (!NamedObjectRegistry)
+		{
+			NamedObjectRegistry = MakeShared<FNamedObjectRegistry>();
+		}
+
+		NamedObjectRegistry->Register(CameraRig, Name.ToString());
+	}
 
 private:
 
@@ -479,5 +546,193 @@ private:
 	TSharedPtr<FNamedObjectRegistry> NamedObjectRegistry;
 };
 
+/**
+ * Default version of the camera rig asset builder.
+ */
+class FCameraRigAssetTestBuilder 
+	: public TCameraRigAssetTestBuilderBase<FCameraRigAssetTestBuilder>
+{
+public:
+
+	FCameraRigAssetTestBuilder(FName Name = NAME_None, UObject* Outer = nullptr);
+	FCameraRigAssetTestBuilder(TSharedPtr<FNamedObjectRegistry> InNamedObjectRegistry, FName Name = NAME_None, UObject* Outer = nullptr);
+};
+
+/**
+ * Version of the camera rig asset builder that has a scoped parent, with a Done() method exposed
+ * to go back to it.
+ */
+template<typename ParentType>
+class TScopedCameraRigAssetTestBuilder
+	: public TScopedConstruction<ParentType>
+	, public TCameraRigAssetTestBuilderBase<TScopedCameraRigAssetTestBuilder<ParentType>>
+{
+public:
+
+	TScopedCameraRigAssetTestBuilder(ParentType& InParent, FName Name = NAME_None, UObject* Outer = nullptr)
+		: TScopedConstruction<ParentType>(InParent)
+		, TCameraRigAssetTestBuilderBase<TScopedCameraRigAssetTestBuilder<ParentType>>(Name, Outer)
+	{
+	}
+
+	TScopedCameraRigAssetTestBuilder(ParentType& InParent, TSharedPtr<FNamedObjectRegistry> InNamedObjectRegistry, FName Name = NAME_None, UObject* Outer = nullptr)
+		: TScopedConstruction<ParentType>(InParent)
+		, TCameraRigAssetTestBuilderBase<TScopedCameraRigAssetTestBuilder<ParentType>>(InNamedObjectRegistry, Name, Outer)
+	{
+	}
+};
+
+/**
+ * Builder class for a camera director.
+ */
+template<
+	typename ParentType,
+	typename DirectorType,
+	typename V = std::enable_if_t<TPointerIsConvertibleFromTo<DirectorType, UCameraDirector>::Value>
+	>
+class TCameraDirectorTestBuilder
+	: public TScopedConstruction<ParentType>
+	, public TCameraObjectInitializer<DirectorType>
+	, public IHasNamedObjectRegistry
+{
+public:
+
+	using ThisType = TCameraDirectorTestBuilder<ParentType, DirectorType, V>;
+
+	/** Creates a new instance of this builder class. */
+	TCameraDirectorTestBuilder(ParentType& InParent, UObject* Outer = nullptr)
+		: TScopedConstruction<ParentType>(InParent)
+	{
+		if (Outer == nullptr)
+		{
+			Outer = GetTransientPackage();
+		}
+		CameraDirector = NewObject<DirectorType>(Outer);
+		TCameraObjectInitializer<DirectorType>::SetObject(CameraDirector);
+	}
+
+	/** Gets the build camera director. */
+	UCameraDirector* Get() const { return CameraDirector; }
+
+	/** Pins the built camera director to a given pointer, for being able to later refer to it. */
+	ThisType& Pin(DirectorType*& OutPtr) { OutPtr = CameraDirector; return *this; }
+
+	/** Give a name to the built camera drector, to be recalled later. */
+	template<typename VV = std::enable_if_t<TPointerIsConvertibleFromTo<ParentType, IHasNamedObjectRegistry>::Value>>
+	ThisType& Named(const TCHAR* InName)
+	{
+		TSharedPtr<FNamedObjectRegistry> NamedObjectRegistry = GetNamedObjectRegistry();
+		if (ensure(NamedObjectRegistry))
+		{
+			NamedObjectRegistry->Register(CameraDirector, InName);
+		}
+		return *this;
+	}
+
+	/** Set a parameter on the camera director. */
+	template<typename ParameterType>
+	ThisType& SetParameter(
+			ParameterType DirectorType::*ParameterField,
+			typename TCallTraits<typename ParameterType::ValueType>::ParamType Value)
+	{
+		ParameterType& ParameterRef = (CameraDirector->*ParameterField);
+		ParameterRef.Value = Value;
+		return *this;
+	}
+
+	/** Runs arbitrary setup logic on the camera director. */
+	ThisType& Setup(TFunction<void(DirectorType*)> SetupCallback)
+	{
+		SetupCallback(CameraDirector);
+		return *this;
+	}
+
+	/** Runs arbitrary setup logic on the camera director. */
+	ThisType& Setup(TFunction<void(DirectorType*, FNamedObjectRegistry*)> SetupCallback)
+	{
+		TSharedPtr<FNamedObjectRegistry> NamedObjectRegistry = GetNamedObjectRegistry();
+		SetupCallback(CameraDirector, NamedObjectRegistry.Get());
+		return *this;
+	}
+
+	/** Gets the named object registry from the parent. */
+	virtual TSharedPtr<FNamedObjectRegistry> GetNamedObjectRegistry() override
+	{
+		if constexpr (TPointerIsConvertibleFromTo<ParentType, IHasNamedObjectRegistry>::Value)
+		{
+			return TScopedConstruction<ParentType>::Parent.GetNamedObjectRegistry();
+		}
+		else
+		{
+			return nullptr;
+		}
+	}
+
+private:
+
+	DirectorType* CameraDirector;
+};
+
+/**
+ * Builder class for a camera evaluation context and its camera asset.
+ */
+class FCameraEvaluationContextTestBuilder
+	: public TCameraObjectInitializer<FCameraEvaluationContext>
+	, public IHasNamedObjectRegistry
+{
+public:
+
+	using ThisType = FCameraEvaluationContextTestBuilder;
+
+	/** Creates a new instance of this builder class. */
+	FCameraEvaluationContextTestBuilder(UObject* Owner = nullptr);
+
+	/** Gets the created evaluation context. */
+	TSharedRef<FCameraEvaluationContext> Get() const { return EvaluationContext.ToSharedRef(); }
+
+	/** Pins the created camera asset. */
+	ThisType& PinCameraAsset(UCameraAsset*& OutPtr) { OutPtr = CameraAsset; return *this; }
+
+	/** Builds the camera asset. */
+	ThisType& BuildCameraAsset() { CameraAsset->BuildCamera(); return *this; }
+
+	/** Builds a new camera director of the given type and returns a builder object for its. */
+	template<typename DirectorType>
+	TCameraDirectorTestBuilder<ThisType, DirectorType> MakeDirector()
+	{
+		TCameraDirectorTestBuilder<ThisType, DirectorType> DirectorBuilder(*this, EvaluationContext->GetOwner());
+		CameraAsset->SetCameraDirector(DirectorBuilder.Get());
+		return DirectorBuilder;
+	}
+
+	/** Builds a new single camera director and returns a builder object for its. */
+	TCameraDirectorTestBuilder<ThisType, USingleCameraDirector> MakeSingleDirector()
+	{
+		return MakeDirector<USingleCameraDirector>();
+	}
+
+	/** Creates a new camera rig asset builder and adds its camera rig to the camera asset.*/
+	TScopedCameraRigAssetTestBuilder<ThisType> AddCameraRig(FName Name = NAME_None)
+	{
+		TScopedCameraRigAssetTestBuilder<ThisType> CameraRigBuilder(*this, GetNamedObjectRegistry(), Name, CameraAsset);
+		CameraAsset->AddCameraRig(CameraRigBuilder.Get());
+		return CameraRigBuilder;
+	}
+
+	/** Gets the named object registry. */
+	virtual TSharedPtr<FNamedObjectRegistry> GetNamedObjectRegistry() override
+	{
+		return NamedObjectRegistry;
+	}
+
+private:
+
+	UCameraAsset* CameraAsset;
+
+	TSharedPtr<FCameraEvaluationContext> EvaluationContext;
+
+	TSharedPtr<FNamedObjectRegistry> NamedObjectRegistry;
+};
+
 }  // namespace UE::Cameras::Test
 
diff --git a/Engine/Plugins/Cameras/GameplayCameras/Source/GameplayCameras/Private/Tests/GameplayCamerasTestObjects.cpp b/Engine/Plugins/Cameras/GameplayCameras/Source/GameplayCameras/Private/Tests/GameplayCamerasTestObjects.cpp
new file mode 100644
--- /dev/null
+++ b/Engine/Plugins/Cameras/GameplayCameras/Source/GameplayCameras/Private/Tests/GameplayCamerasTestObjects.cpp
@@ -0,0 +1,27 @@
+// Copyright Epic Games, Inc. All Rights Reserved.
+
+#include "GameplayCamerasTestObjects.h"
+
+namespace UE::Cameras::Test
+{
+
+UE_DEFINE_CAMERA_NODE_EVALUATOR(FUpdateTrackerCameraNodeEvaluator)
+
+void FUpdateTrackerCameraNodeEvaluator::OnRun(const FCameraNodeEvaluationParams& Params, FCameraNodeEvaluationResult& OutResult)
+{
+	ReceivedUpdates.Add(FTrackedUpdateInfo
+			{
+				Params.DeltaTime,
+				Params.bIsFirstFrame,
+				OutResult.bIsCameraCut
+			});
+}
+
+}  // namespace UE::Cameras::Tests
+
+FCameraNodeEvaluatorPtr UUpdateTrackerCameraNode::OnBuildEvaluator(FCameraNodeEvaluatorBuilder& Builder) const
+{
+	using namespace UE::Cameras::Test;
+	return Builder.BuildEvaluator<FUpdateTrackerCameraNodeEvaluator>();
+}
+
diff --git a/Engine/Plugins/Cameras/GameplayCameras/Source/GameplayCameras/Private/Tests/GameplayCamerasTestObjects.h b/Engine/Plugins/Cameras/GameplayCameras/Source/GameplayCameras/Private/Tests/GameplayCamerasTestObjects.h
new file mode 100644
--- /dev/null
+++ b/Engine/Plugins/Cameras/GameplayCameras/Source/GameplayCameras/Private/Tests/GameplayCamerasTestObjects.h
@@ -0,0 +1,56 @@
+// Copyright Epic Games, Inc. All Rights Reserved.
+
+#pragma once
+
+#include "Core/CameraNode.h"
+#include "Core/CameraNodeEvaluator.h"
+#include "Core/CameraParameters.h"
+#include "Nodes/CameraNodeTypes.h"
+
+#include "GameplayCamerasTestObjects.generated.h"
+
+namespace UE::Cameras::Test
+{
+
+struct FTrackedUpdateInfo
+{
+	float DeltaTime = 0.f;
+	bool bIsFirstFrame = false;
+	bool bIsCameraCut = false;
+};
+
+class FUpdateTrackerCameraNodeEvaluator : public FCameraNodeEvaluator
+{
+	UE_DECLARE_CAMERA_NODE_EVALUATOR(GAMEPLAYCAMERAS_API, FUpdateTrackerCameraNodeEvaluator)
+
+public:
+
+	TArray<FTrackedUpdateInfo> ReceivedUpdates;
+
+protected:
+
+	// FCameraNodeEvaluator interface.
+	virtual void OnRun(const FCameraNodeEvaluationParams& Params, FCameraNodeEvaluationResult& OutResult) override;
+};
+
+}  // namespace UE::Cameras::Tests
+
+UCLASS(MinimalAPI, Hidden)
+class UUpdateTrackerCameraNode : public UCameraNode
+{
+	GENERATED_BODY()
+
+protected:
+
+	// UCameraNode interface.
+	virtual FCameraNodeEvaluatorPtr OnBuildEvaluator(FCameraNodeEvaluatorBuilder& Builder) const override;
+
+public:
+
+	UPROPERTY(EditAnywhere, Category=Common)
+	FDoubleCameraParameter DoubleParameter;
+
+	UPROPERTY(EditAnywhere, Category=Common)
+	FVector3dCameraParameter VectorParameter;
+};
+
