Cameras: look for transitions on all the rigs inside a combination rig, not just the first one

diff --git a/Engine/Plugins/Cameras/GameplayCameras/Source/GameplayCameras/Private/Core/BlendStackCameraNode.cpp b/Engine/Plugins/Cameras/GameplayCameras/Source/GameplayCameras/Private/Core/BlendStackCameraNode.cpp
--- a/Engine/Plugins/Cameras/GameplayCameras/Source/GameplayCameras/Private/Core/BlendStackCameraNode.cpp
+++ b/Engine/Plugins/Cameras/GameplayCameras/Source/GameplayCameras/Private/Core/BlendStackCameraNode.cpp
@@ -813,79 +813,105 @@ void FTransientBlendStackCameraNodeEvaluator::InternalPostBlendExecute(TArrayVie
 
 const UCameraRigTransition* FTransientBlendStackCameraNodeEvaluator::FindTransition(const FBlendStackCameraPushParams& Params) const
 {
-	const UBlendStackCameraNode* BlendStackNode = GetCameraNodeAs<UBlendStackCameraNode>();
-
-	TSharedPtr<const FCameraEvaluationContext> ToContext = Params.EvaluationContext;
-	const UCameraAsset* ToCameraAsset = ToContext ? ToContext->GetCameraAsset() : nullptr;
-	const UCameraRigAsset* ToCameraRig = Params.CameraRig;
-
-	// If the new entry is a combination, look for transitions on its "main" camera rig.
-	ToCameraRig = UCombinedCameraRigsCameraNode::GetMainCameraRigIfCombination(ToCameraRig);
-
 	// Find a transition that works for blending towards ToCameraRig.
 	// If the stack isn't empty, we need to find a transition that works between the previous and 
 	// next camera rigs. If the stack is empty, we blend the new camera rig in from nothing if
 	// appropriate.
 	if (!Entries.IsEmpty())
 	{
-		const UCameraRigTransition* TransitionToUse = nullptr;
+		// Grab information about the new entry to push.
+		TSharedPtr<const FCameraEvaluationContext> ToContext = Params.EvaluationContext;
+		const UCameraAsset* ToCameraAsset = ToContext ? ToContext->GetCameraAsset() : nullptr;
 
-		// Start by looking at exit transitions on the last active (top) camera rig.
+		// Grab information about the top entry (i.e. the currently active camera rig).
 		const FCameraRigEntry& TopEntry = Entries.Top();
-
 		TSharedPtr<const FCameraEvaluationContext> FromContext = TopEntry.EvaluationContext.Pin();
 		const UCameraAsset* FromCameraAsset = FromContext ? FromContext->GetCameraAsset() : nullptr;
-		const UCameraRigAsset* FromCameraRig = TopEntry.CameraRig;
 
-		// If the top entry is a combination, look for transitions on its "main" camera rig.
-		FromCameraRig = UCombinedCameraRigsCameraNode::GetMainCameraRigIfCombination(FromCameraRig);
+		// If the new or current top entries are a combination, look for transitions on all 
+		// their combined camera rigs.
+		TArray<const UCameraRigAsset*> ToCombinedCameraRigs;
+		UCombinedCameraRigsCameraNode::GetAllCombinationCameraRigs(Params.CameraRig, ToCombinedCameraRigs);
 
-		if (!TopEntry.bIsFrozen)
+		TArray<const UCameraRigAsset*> FromCombinedCameraRigs;
+		UCombinedCameraRigsCameraNode::GetAllCombinationCameraRigs(TopEntry.CameraRig, FromCombinedCameraRigs);
+
+		const bool bFromFrozen = TopEntry.bIsFrozen;
+		const UCameraRigTransition* TransitionToUse = nullptr;
+
+		// Start by looking at exit transitions on the last active (top) camera rig.
+		for (const UCameraRigAsset* FromCameraRig : FromCombinedCameraRigs)
 		{
-			// Look for exit transitions on the last active camera rig itself.
-			TransitionToUse = FindTransition(
-					FromCameraRig->ExitTransitions,
-					FromCameraRig, FromCameraAsset, false,
-					ToCameraRig, ToCameraAsset);
-			if (TransitionToUse)
+			if (!FromCameraRig->ExitTransitions.IsEmpty())
 			{
-				return TransitionToUse;
+				// Look for exit transitions on the last active camera rig itself.
+				for (const UCameraRigAsset* ToCameraRig : ToCombinedCameraRigs)
+				{
+					TransitionToUse = FindTransition(
+							FromCameraRig->ExitTransitions,
+							FromCameraRig, FromCameraAsset, bFromFrozen,
+							ToCameraRig, ToCameraAsset);
+					if (TransitionToUse)
+					{
+						return TransitionToUse;
+					}
+				}
 			}
-			
-			// Look for exit transitions on its parent camera asset.
-			if (FromCameraAsset)
+		}
+		for (const UCameraRigAsset* FromCameraRig : FromCombinedCameraRigs)
+		{
+			if (FromCameraAsset && !FromCameraAsset->GetExitTransitions().IsEmpty())
 			{
-				TransitionToUse = FindTransition(
-						FromCameraAsset->GetExitTransitions(),
-						FromCameraRig, FromCameraAsset, false,
-						ToCameraRig, ToCameraAsset);
-				if (TransitionToUse)
+				// Look for exit transitions on its parent camera asset.
+				for (const UCameraRigAsset* ToCameraRig : ToCombinedCameraRigs)
 				{
-					return TransitionToUse;
+					TransitionToUse = FindTransition(
+							FromCameraAsset->GetExitTransitions(),
+							FromCameraRig, FromCameraAsset, bFromFrozen,
+							ToCameraRig, ToCameraAsset);
+					if (TransitionToUse)
+					{
+						return TransitionToUse;
+					}
 				}
 			}
 		}
 
 		// Now look at enter transitions on the new camera rig.
-		TransitionToUse = FindTransition(
-				ToCameraRig->EnterTransitions,
-				FromCameraRig, FromCameraAsset, TopEntry.bIsFrozen,
-				ToCameraRig, ToCameraAsset);
-		if (TransitionToUse)
+		for (const UCameraRigAsset* ToCameraRig : ToCombinedCameraRigs)
 		{
-			return TransitionToUse;
+			if (!ToCameraRig->EnterTransitions.IsEmpty())
+			{
+				// Look for enter transitions on the new camera rig itself.
+				for (const UCameraRigAsset* FromCameraRig : FromCombinedCameraRigs)
+				{
+					TransitionToUse = FindTransition(
+							ToCameraRig->EnterTransitions,
+							FromCameraRig, FromCameraAsset, bFromFrozen,
+							ToCameraRig, ToCameraAsset);
+					if (TransitionToUse)
+					{
+						return TransitionToUse;
+					}
+				}
+			}
 		}
-
-		// Look at enter transitions on its parent camera asset.
-		if (ToCameraAsset)
+		for (const UCameraRigAsset* ToCameraRig : ToCombinedCameraRigs)
 		{
-			TransitionToUse = FindTransition(
-					ToCameraAsset->GetEnterTransitions(),
-					FromCameraRig, FromCameraAsset, TopEntry.bIsFrozen,
-					ToCameraRig, ToCameraAsset);
-			if (TransitionToUse)
+			if (ToCameraAsset && !ToCameraAsset->GetEnterTransitions().IsEmpty())
 			{
-				return TransitionToUse;
+				// Look at enter transitions on its parent camera asset.
+				for (const UCameraRigAsset* FromCameraRig : FromCombinedCameraRigs)
+				{
+					TransitionToUse = FindTransition(
+							ToCameraAsset->GetEnterTransitions(),
+							FromCameraRig, FromCameraAsset, bFromFrozen,
+							ToCameraRig, ToCameraAsset);
+					if (TransitionToUse)
+					{
+						return TransitionToUse;
+					}
+				}
 			}
 		}
 	}
diff --git a/Engine/Plugins/Cameras/GameplayCameras/Source/GameplayCameras/Private/Core/CameraRigCombinationRegistry.cpp b/Engine/Plugins/Cameras/GameplayCameras/Source/GameplayCameras/Private/Core/CameraRigCombinationRegistry.cpp
--- a/Engine/Plugins/Cameras/GameplayCameras/Source/GameplayCameras/Private/Core/CameraRigCombinationRegistry.cpp
+++ b/Engine/Plugins/Cameras/GameplayCameras/Source/GameplayCameras/Private/Core/CameraRigCombinationRegistry.cpp
@@ -130,25 +130,27 @@ void FCombinedCameraRigsCameraNodeEvaluator::ApplyParameterOverrides(FCameraVari
 
 }  // namespace UE::Cameras
 
-const UCameraRigAsset* UCombinedCameraRigsCameraNode::GetMainCameraRigIfCombination(const UCameraRigAsset* InCameraRig)
+void UCombinedCameraRigsCameraNode::GetAllCombinationCameraRigs(const UCameraRigAsset* InCameraRig, TArray<const UCameraRigAsset*>& OutCameraRigs)
 {
-	if (!InCameraRig || !InCameraRig->RootNode)
+	if (!InCameraRig)
 	{
-		return InCameraRig;
+		return;
 	}
 
 	if (const UCombinedCameraRigsCameraNode* CameraRigCombinationNode = Cast<UCombinedCameraRigsCameraNode>(InCameraRig->RootNode))
 	{
-		if (!CameraRigCombinationNode->CameraRigReferences.IsEmpty())
+		for (const FCameraRigAssetReference& CameraRigReference : CameraRigCombinationNode->CameraRigReferences)
 		{
-			if (const UCameraRigAsset* MainCameraRig = CameraRigCombinationNode->CameraRigReferences[0].GetCameraRig())
+			if (const UCameraRigAsset* CameraRig = CameraRigReference.GetCameraRig())
 			{
-				return MainCameraRig;
+				OutCameraRigs.Add(CameraRig);
 			}
 		}
 	}
-
-	return InCameraRig;
+	else
+	{
+		OutCameraRigs.Add(InCameraRig);
+	}
 }
 
 FCameraNodeEvaluatorPtr UCombinedCameraRigsCameraNode::OnBuildEvaluator(FCameraNodeEvaluatorBuilder& Builder) const
diff --git a/Engine/Plugins/Cameras/GameplayCameras/Source/GameplayCameras/Private/Core/CameraRigCombinationRegistry.h b/Engine/Plugins/Cameras/GameplayCameras/Source/GameplayCameras/Private/Core/CameraRigCombinationRegistry.h
--- a/Engine/Plugins/Cameras/GameplayCameras/Source/GameplayCameras/Private/Core/CameraRigCombinationRegistry.h
+++ b/Engine/Plugins/Cameras/GameplayCameras/Source/GameplayCameras/Private/Core/CameraRigCombinationRegistry.h
@@ -56,7 +56,7 @@ class UCombinedCameraRigsCameraNode : public UCameraNode
 
 public:
 
-	static const UCameraRigAsset* GetMainCameraRigIfCombination(const UCameraRigAsset* InCameraRig);
+	static void GetAllCombinationCameraRigs(const UCameraRigAsset* InCameraRig, TArray<const UCameraRigAsset*>& OutCameraRigs);
 
 protected:
 
diff --git a/Engine/Plugins/Cameras/GameplayCameras/Source/GameplayCameras/Private/Services/OrientationInitializationService.cpp b/Engine/Plugins/Cameras/GameplayCameras/Source/GameplayCameras/Private/Services/OrientationInitializationService.cpp
--- a/Engine/Plugins/Cameras/GameplayCameras/Source/GameplayCameras/Private/Services/OrientationInitializationService.cpp
+++ b/Engine/Plugins/Cameras/GameplayCameras/Source/GameplayCameras/Private/Services/OrientationInitializationService.cpp
@@ -60,10 +60,17 @@ void FOrientationInitializationService::OnRootCameraNodeEvent(const FRootCameraN
 		if (const UCameraRigAsset* NewCameraRig = InEvent.CameraRigInfo.CameraRig)
 		{
 			// If the new camera rig is a combination, find its initial orientation settings
-			// on its main rig.
-			NewCameraRig = UCombinedCameraRigsCameraNode::GetMainCameraRigIfCombination(NewCameraRig);
-
-			InitialOrientation = NewCameraRig->InitialOrientation;
+			// on one of its combined rigs.
+			TArray<const UCameraRigAsset*> NewCombinedCameraRigs;
+			UCombinedCameraRigsCameraNode::GetAllCombinationCameraRigs(NewCameraRig, NewCombinedCameraRigs);
+			for (const UCameraRigAsset* NewCombinedCameraRig : NewCombinedCameraRigs)
+			{
+				if (NewCombinedCameraRig->InitialOrientation != ECameraRigInitialOrientation::None)
+				{
+					InitialOrientation = NewCombinedCameraRig->InitialOrientation;
+					break;
+				}
+			}
 		}
 
 		if (InEvent.Transition && InEvent.Transition->bOverrideInitialOrientation)
