Cameras: remove the concept of "private" camera rigs inside camera assets

All camera rigs are now shared, standalone assets (not just camera rig prefabs). Camera assets don't own camera rigs inside their package anymore, they all have to be created as separate assets. Everything related to these previously private camera rigs has been removed.

Upon opening a camera asset, the editor will automatically upgrade it, breaking off any previously existing camera rigs into new assets.

diff --git a/Engine/Config/BaseEngine.ini b/Engine/Config/BaseEngine.ini
--- a/Engine/Config/BaseEngine.ini
+++ b/Engine/Config/BaseEngine.ini
@@ -1338,6 +1338,9 @@ ClassName=/Script/OnlineSubsystemUtils.OnlineEngineInterfaceImpl
 +FunctionRedirects=(OldName="SkeletalMeshComponent.SetAnimClass",NewName="SetAnimInstanceClass")
 +FunctionRedirects=(OldName="SkeletalMeshComponent.K2_SetAnimInstanceClass",NewName="SetAnimInstanceClass")
 
+; 5.6
++FunctionRedirects=(OldName="BlueprintCameraDirectorEvaluator.ActivateCameraRigPrefab", NewName="ActivateCameraRig")
+
 ; End of releases redirects
 
 ; Fixing the output of tags for blueprints to be reset.
diff --git a/Engine/Plugins/Cameras/GameplayCameras/Source/GameplayCameras/Private/Core/CameraAsset.cpp b/Engine/Plugins/Cameras/GameplayCameras/Source/GameplayCameras/Private/Core/CameraAsset.cpp
--- a/Engine/Plugins/Cameras/GameplayCameras/Source/GameplayCameras/Private/Core/CameraAsset.cpp
+++ b/Engine/Plugins/Cameras/GameplayCameras/Source/GameplayCameras/Private/Core/CameraAsset.cpp
@@ -27,33 +27,6 @@ void UCameraAsset::SetCameraDirector(UCameraDirector* InCameraDirector)
 	}
 }
 
-void UCameraAsset::AddCameraRig(UCameraRigAsset* InCameraRig)
-{
-	using namespace UE::Cameras;
-
-	ensure(InCameraRig);
-
-	CameraRigs.Add(InCameraRig);
-
-	TCameraArrayChangedEvent<UCameraRigAsset*> ChangedEvent;
-	ChangedEvent.EventType = ECameraArrayChangedEventType::Add;
-	EventHandlers.Notify(&ICameraAssetEventHandler::OnCameraRigsChanged, this, ChangedEvent);
-}
-
-int32 UCameraAsset::RemoveCameraRig(UCameraRigAsset* InCameraRig)
-{
-	using namespace UE::Cameras;
-
-	const int32 NumRemoved = CameraRigs.Remove(InCameraRig);
-	if (NumRemoved > 0)
-	{
-		TCameraArrayChangedEvent<UCameraRigAsset*> ChangedEvent;
-		ChangedEvent.EventType = ECameraArrayChangedEventType::Remove;
-		EventHandlers.Notify(&ICameraAssetEventHandler::OnCameraRigsChanged, this, ChangedEvent);
-	}
-	return NumRemoved;
-}
-
 void UCameraAsset::AddEnterTransition(UCameraRigTransition* InTransition)
 {
 	using namespace UE::Cameras;
@@ -125,68 +98,11 @@ void UCameraAsset::PostLoad()
 			CameraDirector->ClearFlags(RF_Public | RF_Standalone);
 		}
 	}
-
-	CleanUpStrayObjects();
 #endif  // WITH_EDITOR
 }
 
 #if WITH_EDITOR
 
-void UCameraAsset::CleanUpStrayObjects()
-{
-	UPackage* CameraAssetPackage = GetOutermost();
-	if (!CameraAssetPackage || CameraAssetPackage == GetTransientPackage())
-	{
-		return;
-	}
-
-	// Some older versions of the camera editors had a bug that could lead to
-	// stray deleted camera rigs being left in the package. Let's clean them up.
-	TSet<UObject*> StrayObjects;
-	TSet<UCameraRigAsset*> KnownCameraRigs(CameraRigs);
-
-	TArray<UObject*> ObjectsInPackage;
-	GetObjectsWithPackage(CameraAssetPackage, ObjectsInPackage);
-	for (UObject* Object : ObjectsInPackage)
-	{
-		UCameraRigAsset* CameraRig = Cast<UCameraRigAsset>(Object);
-		if (!CameraRig)
-		{
-			continue;
-		}
-		if (KnownCameraRigs.Contains(CameraRig))
-		{
-			continue;
-		}
-
-		Modify();
-
-		CameraRig->ClearFlags(RF_Public | RF_Standalone);
-		StrayObjects.Add(CameraRig);
-	}
-
-	if (StrayObjects.Num() > 0)
-	{
-		// Also clean-up any redirectors to these objects.
-		for (UObject* Object : ObjectsInPackage)
-		{
-			if (UObjectRedirector* Redirector = Cast<UObjectRedirector>(Object))
-			{
-				if (StrayObjects.Contains(Redirector->DestinationObject))
-				{
-					Redirector->ClearFlags(RF_Public | RF_Standalone);
-					Redirector->DestinationObject = nullptr;
-				}
-			}
-		}
-
-		UE_LOG(LogCameraSystem, Warning,
-				TEXT("Cleaned up %d stray camera rigs in camera asset '%s'. Please resave the asset."),
-				StrayObjects.Num(), *GetPathNameSafe(this));
-	}
-}
-
-
 void UCameraAsset::PostEditChangeProperty(FPropertyChangedEvent& PropertyChangedEvent)
 {
 	using namespace UE::Cameras;
@@ -198,11 +114,6 @@ void UCameraAsset::PostEditChangeProperty(FPropertyChangedEvent& PropertyChanged
 		ChangedEvent.NewValue = CameraDirector;
 		EventHandlers.Notify(&ICameraAssetEventHandler::OnCameraDirectorChanged, this, ChangedEvent);
 	}
-	else if (PropertyName == GET_MEMBER_NAME_CHECKED(UCameraAsset, CameraRigs))
-	{
-		TCameraArrayChangedEvent<UCameraRigAsset*> ChangedEvent(PropertyChangedEvent.ChangeType);
-		EventHandlers.Notify(&ICameraAssetEventHandler::OnCameraRigsChanged, this, ChangedEvent);
-	}
 	else if (PropertyName == GET_MEMBER_NAME_CHECKED(UCameraAsset, EnterTransitions))
 	{
 		TCameraArrayChangedEvent<UCameraRigTransition*> ChangedEvent(PropertyChangedEvent.ChangeType);
diff --git a/Engine/Plugins/Cameras/GameplayCameras/Source/GameplayCameras/Private/Core/CameraAssetBuilder.cpp b/Engine/Plugins/Cameras/GameplayCameras/Source/GameplayCameras/Private/Core/CameraAssetBuilder.cpp
--- a/Engine/Plugins/Cameras/GameplayCameras/Source/GameplayCameras/Private/Core/CameraAssetBuilder.cpp
+++ b/Engine/Plugins/Cameras/GameplayCameras/Source/GameplayCameras/Private/Core/CameraAssetBuilder.cpp
@@ -44,25 +44,35 @@ void FCameraAssetBuilder::BuildCamera(UCameraAsset* InCameraAsset, FCustomBuildS
 
 void FCameraAssetBuilder::BuildCameraImpl()
 {
+	TArray<UCameraRigAsset*> CameraRigs;
+
 	if (UCameraDirector* CameraDirector = CameraAsset->GetCameraDirector())
 	{
 		CameraDirector->BuildCameraDirector(BuildLog);
+
+#if WITH_EDITOR
+		FCameraDirectorRigUsageInfo UsageInfo;
+		CameraDirector->GatherRigUsageInfo(UsageInfo);
+		CameraRigs = UsageInfo.CameraRigs;
+#endif
 	}
 	else
 	{
 		BuildLog.AddMessage(EMessageSeverity::Error, LOCTEXT("MissingDirector", "Camera has no director set."));
 	}
 
-	if (CameraAsset->GetCameraRigs().IsEmpty())
+#if WITH_EDITOR
+	if (CameraRigs.IsEmpty())
 	{
-		BuildLog.AddMessage(EMessageSeverity::Warning, LOCTEXT("MissingRigs", "Camera has no camera rigs defined."));
+		BuildLog.AddMessage(EMessageSeverity::Warning, LOCTEXT("MissingRigs", "Camera isn't using any camera rigs."));
 	}
 
-	for (UCameraRigAsset* CameraRig : CameraAsset->GetCameraRigs())
+	for (UCameraRigAsset* CameraRig : CameraRigs)
 	{
 		FCameraRigAssetBuilder CameraRigBuilder(BuildLog);
 		CameraRigBuilder.BuildCameraRig(CameraRig);
 	}
+#endif
 }
 
 void FCameraAssetBuilder::UpdateBuildStatus()
diff --git a/Engine/Plugins/Cameras/GameplayCameras/Source/GameplayCameras/Private/Core/CameraDirector.cpp b/Engine/Plugins/Cameras/GameplayCameras/Source/GameplayCameras/Private/Core/CameraDirector.cpp
--- a/Engine/Plugins/Cameras/GameplayCameras/Source/GameplayCameras/Private/Core/CameraDirector.cpp
+++ b/Engine/Plugins/Cameras/GameplayCameras/Source/GameplayCameras/Private/Core/CameraDirector.cpp
@@ -24,5 +24,10 @@ void UCameraDirector::FactoryCreateAsset(const FCameraDirectorFactoryCreateParam
 	OnFactoryCreateAsset(InParams);
 }
 
+void UCameraDirector::GatherRigUsageInfo(FCameraDirectorRigUsageInfo& UsageInfo)
+{
+	OnGatherRigUsageInfo(UsageInfo);
+}
+
 #endif
 
diff --git a/Engine/Plugins/Cameras/GameplayCameras/Source/GameplayCameras/Private/Directors/BlueprintCameraDirector.cpp b/Engine/Plugins/Cameras/GameplayCameras/Source/GameplayCameras/Private/Directors/BlueprintCameraDirector.cpp
--- a/Engine/Plugins/Cameras/GameplayCameras/Source/GameplayCameras/Private/Directors/BlueprintCameraDirector.cpp
+++ b/Engine/Plugins/Cameras/GameplayCameras/Source/GameplayCameras/Private/Directors/BlueprintCameraDirector.cpp
@@ -12,9 +12,11 @@
 #include "Core/CameraRigProxyTable.h"
 #include "Core/CameraSystemEvaluator.h"
 #include "Core/RootCameraNode.h"
+#include "Engine/Blueprint.h"
 #include "GameFramework/Actor.h"
 #include "GameFramework/ControllerGameplayCameraEvaluationComponent.h"
 #include "GameplayCameras.h"
+#include "Helpers/OutgoingReferenceFinder.h"
 #include "Services/AutoResetCameraVariableService.h"
 
 #include UE_INLINE_GENERATED_CPP_BY_NAME(BlueprintCameraDirector)
@@ -291,18 +293,6 @@ void UBlueprintCameraDirectorEvaluator::ActivateCameraRigViaProxy(UCameraRigProx
 	EvaluationResult.ActiveCameraRigProxies.Add(CameraRigProxy);
 }
 
-void UBlueprintCameraDirectorEvaluator::ActivateCameraRigPrefab(UCameraRigAsset* CameraRig)
-{
-	EvaluationResult.ActiveCameraRigs.Add(CameraRig);
-}
-
-UCameraRigAsset* UBlueprintCameraDirectorEvaluator::GetCameraRig(UCameraRigAsset* CameraRig) const
-{
-	// This function is only here to provide an easy way to pick a camera rig from the referencing 
-	// camera asset, using the custom rig picker. Then we just return it.
-	return CameraRig;
-}
-
 AActor* UBlueprintCameraDirectorEvaluator::FindEvaluationContextOwnerActor(TSubclassOf<AActor> ActorClass) const
 {
 	if (EvaluationContext)
@@ -456,6 +446,26 @@ void UBlueprintCameraDirector::OnFactoryCreateAsset(const FCameraDirectorFactory
 	}
 }
 
+void UBlueprintCameraDirector::OnGatherRigUsageInfo(FCameraDirectorRigUsageInfo& UsageInfo)
+{
+	using namespace UE::Cameras;
+
+	if (!CameraDirectorEvaluatorClass)
+	{
+		return;
+	}
+
+	UBlueprint* EvaluatorBlueprint = Cast<UBlueprint>(CameraDirectorEvaluatorClass->ClassGeneratedBy);
+	if (!ensure(EvaluatorBlueprint))
+	{
+		return;
+	}
+
+	FOutgoingReferenceFinder ReferenceFinder(EvaluatorBlueprint, UCameraRigAsset::StaticClass());
+	ReferenceFinder.CollectReferences();
+	ReferenceFinder.GetReferencesOfClass<UCameraRigAsset>(UsageInfo.CameraRigs);
+}
+
 #endif
 
 #undef LOCTEXT_NAMESPACE
diff --git a/Engine/Plugins/Cameras/GameplayCameras/Source/GameplayCameras/Private/Directors/SingleCameraDirector.cpp b/Engine/Plugins/Cameras/GameplayCameras/Source/GameplayCameras/Private/Directors/SingleCameraDirector.cpp
--- a/Engine/Plugins/Cameras/GameplayCameras/Source/GameplayCameras/Private/Directors/SingleCameraDirector.cpp
+++ b/Engine/Plugins/Cameras/GameplayCameras/Source/GameplayCameras/Private/Directors/SingleCameraDirector.cpp
@@ -52,15 +52,11 @@ void USingleCameraDirector::OnBuildCameraDirector(UE::Cameras::FCameraBuildLog&
 
 #if WITH_EDITOR
 
-void USingleCameraDirector::OnFactoryCreateAsset(const FCameraDirectorFactoryCreateParams& InParams)
+void USingleCameraDirector::OnGatherRigUsageInfo(FCameraDirectorRigUsageInfo& UsageInfo)
 {
-	// Automatically set ourselves to use the first camera rig available.
-	if (UCameraAsset* OuterCameraAsset = GetTypedOuter<UCameraAsset>())
+	if (CameraRig)
 	{
-		if (OuterCameraAsset->GetCameraRigs().Num() > 0)
-		{
-			CameraRig = OuterCameraAsset->GetCameraRigs()[0];
-		}
+		UsageInfo.CameraRigs.Add(CameraRig);
 	}
 }
 
diff --git a/Engine/Plugins/Cameras/GameplayCameras/Source/GameplayCameras/Private/Directors/StateTreeCameraDirector.cpp b/Engine/Plugins/Cameras/GameplayCameras/Source/GameplayCameras/Private/Directors/StateTreeCameraDirector.cpp
--- a/Engine/Plugins/Cameras/GameplayCameras/Source/GameplayCameras/Private/Directors/StateTreeCameraDirector.cpp
+++ b/Engine/Plugins/Cameras/GameplayCameras/Source/GameplayCameras/Private/Directors/StateTreeCameraDirector.cpp
@@ -5,11 +5,14 @@
 #include "Core/CameraAsset.h"
 #include "Core/CameraBuildLog.h"
 #include "Core/CameraEvaluationContext.h"
+#include "Core/CameraRigAsset.h"
 #include "Core/CameraRigProxyAsset.h"
 #include "Core/CameraRigProxyTable.h"
 #include "Directors/CameraDirectorStateTreeSchema.h"
 #include "GameplayCameras.h"
+#include "Helpers/OutgoingReferenceFinder.h"
 #include "Logging/TokenizedMessage.h"
+#include "StateTree.h"
 #include "StateTreeExecutionContext.h"
 #include "StateTreeInstanceData.h"
 
@@ -272,6 +275,21 @@ void UStateTreeCameraDirector::OnFactoryCreateAsset(const FCameraDirectorFactory
 	}
 }
 
+void UStateTreeCameraDirector::OnGatherRigUsageInfo(FCameraDirectorRigUsageInfo& UsageInfo)
+{
+	using namespace UE::Cameras;
+
+	UStateTree* StateTree = StateTreeReference.GetMutableStateTree();
+	if (!StateTree)
+	{
+		return;
+	}
+
+	FOutgoingReferenceFinder ReferenceFinder(StateTree, UCameraRigAsset::StaticClass());
+	ReferenceFinder.CollectReferences();
+	ReferenceFinder.GetReferencesOfClass<UCameraRigAsset>(UsageInfo.CameraRigs);
+}
+
 #endif
 
 #undef LOCTEXT_NAMESPACE
diff --git a/Engine/Plugins/Cameras/GameplayCameras/Source/GameplayCameras/Private/Helpers/ObjectReferenceFinder.cpp b/Engine/Plugins/Cameras/GameplayCameras/Source/GameplayCameras/Private/Helpers/ObjectReferenceFinder.cpp
new file mode 100644
--- /dev/null
+++ b/Engine/Plugins/Cameras/GameplayCameras/Source/GameplayCameras/Private/Helpers/ObjectReferenceFinder.cpp
@@ -0,0 +1,88 @@
+// Copyright Epic Games, Inc. All Rights Reserved.
+
+#include "Helpers/ObjectReferenceFinder.h"
+
+#include "UObject/Object.h"
+#include "UObject/Package.h"
+
+namespace UE::Cameras
+{
+
+FObjectReferenceFinder::FObjectReferenceFinder(UObject* InRootObject, TArrayView<UObject*> InReferencedObjects)
+{
+	Initialize(InRootObject);
+
+	TargetObjects = TSet<UObject*>(InReferencedObjects);
+}
+
+void FObjectReferenceFinder::Initialize(UObject* InRootObject)
+{
+	RootObject = InRootObject;
+	PackageScope = InRootObject->GetOutermost();
+
+	SetIsPersistent(true);
+	SetIsSaving(true);
+	SetFilterEditorOnly(false);
+
+	ArIsObjectReferenceCollector = true;
+	ArShouldSkipBulkData = true;
+}
+
+void FObjectReferenceFinder::CollectReferences()
+{
+	ObjectsToVisit.Reset();
+	VisitedObjects.Reset();
+
+	ObjectsToVisit.Add(RootObject);
+	while (ObjectsToVisit.Num() > 0)
+	{
+		UObject* CurObj = ObjectsToVisit.Pop(EAllowShrinking::No);
+		if (!VisitedObjects.Contains(CurObj))
+		{
+			VisitedObjects.Add(CurObj);
+			CurObj->Serialize(*this);
+		}
+	}
+}
+
+bool FObjectReferenceFinder::HasAnyObjectReference() const
+{
+	for (TPair<UObject*, int32> Pair : ObjectReferenceCounts)
+	{
+		if (Pair.Value > 0)
+		{
+			return true;
+		}
+	}
+	return false;
+}
+
+int32 FObjectReferenceFinder::GetObjectReferenceCount(UObject* InObject) const
+{
+	if (const int32* RefCount = ObjectReferenceCounts.Find(InObject))
+	{
+		return *RefCount;
+	}
+	return 0;
+}
+
+FArchive& FObjectReferenceFinder::operator<<(UObject*& ObjRef)
+{
+	if (ObjRef != nullptr)
+	{
+		if (TargetObjects.Contains(ObjRef))
+		{
+			int32& RefCount = ObjectReferenceCounts.FindOrAdd(ObjRef, 0);
+			++RefCount;
+		}
+
+		if (ObjRef->IsIn(PackageScope) && !VisitedObjects.Contains(ObjRef))
+		{
+			ObjectsToVisit.Add(ObjRef);
+		}
+	}
+	return *this;
+}
+
+}  // namespace UE::Cameras
+
diff --git a/Engine/Plugins/Cameras/GameplayCameras/Source/GameplayCameras/Private/Helpers/OutgoingReferenceFinder.cpp b/Engine/Plugins/Cameras/GameplayCameras/Source/GameplayCameras/Private/Helpers/OutgoingReferenceFinder.cpp
new file mode 100644
--- /dev/null
+++ b/Engine/Plugins/Cameras/GameplayCameras/Source/GameplayCameras/Private/Helpers/OutgoingReferenceFinder.cpp
@@ -0,0 +1,103 @@
+// Copyright Epic Games, Inc. All Rights Reserved.
+
+#include "Helpers/OutgoingReferenceFinder.h"
+
+#include "UObject/Class.h"
+#include "UObject/Object.h"
+#include "UObject/Package.h"
+
+namespace UE::Cameras
+{
+
+FOutgoingReferenceFinder::FOutgoingReferenceFinder(UObject* InRootObject, UClass* InReferencedObjectClass)
+{
+	Initialize(InRootObject);
+
+	TargetObjectClasses.Add(InReferencedObjectClass);
+}
+
+FOutgoingReferenceFinder::FOutgoingReferenceFinder(UObject* InRootObject, TArrayView<UClass*> InReferencedObjectClasses)
+{
+	Initialize(InRootObject);
+
+	TargetObjectClasses.Append(InReferencedObjectClasses);
+}
+
+void FOutgoingReferenceFinder::Initialize(UObject* InRootObject)
+{
+	RootObject = InRootObject;
+	PackageScope = InRootObject->GetOutermost();
+
+	SetIsPersistent(true);
+	SetIsSaving(true);
+	SetFilterEditorOnly(false);
+
+	ArIsObjectReferenceCollector = true;
+	ArShouldSkipBulkData = true;
+}
+
+void FOutgoingReferenceFinder::CollectReferences()
+{
+	ObjectsToVisit.Reset();
+	VisitedObjects.Reset();
+
+	ObjectsToVisit.Add(RootObject);
+	while (ObjectsToVisit.Num() > 0)
+	{
+		UObject* CurObj = ObjectsToVisit.Pop(EAllowShrinking::No);
+		if (!VisitedObjects.Contains(CurObj))
+		{
+			VisitedObjects.Add(CurObj);
+			CurObj->Serialize(*this);
+		}
+	}
+}
+
+bool FOutgoingReferenceFinder::GetAllReferences(TArray<UObject*>& OutReferencedObjects) const
+{
+	bool bHasAnyReference = false;
+	for (TPair<UClass*, TSet<UObject*>> Pair : ReferencedObjects)
+	{
+		for (UObject* Obj : Pair.Value)
+		{
+			OutReferencedObjects.Add(Obj);
+		}
+		bHasAnyReference |= (!Pair.Value.IsEmpty());
+	}
+	return bHasAnyReference;
+}
+
+FArchive& FOutgoingReferenceFinder::operator<<(UObject*& ObjRef)
+{
+	if (ObjRef != nullptr)
+	{
+		UClass* ObjClass = ObjRef->GetClass();
+		if (MatchesAnyTargetClass(ObjClass))
+		{
+			TSet<UObject*>& ReferencesOfClass = ReferencedObjects.FindOrAdd(ObjClass);
+			ReferencesOfClass.Add(ObjRef);
+		}
+
+		if (ObjRef->IsIn(PackageScope) && !VisitedObjects.Contains(ObjRef))
+		{
+			ObjectsToVisit.Add(ObjRef);
+		}
+	}
+	return *this;
+}
+
+bool FOutgoingReferenceFinder::MatchesAnyTargetClass(UClass* InObjClass) const
+{
+	for (UClass* TargetObjectClass : TargetObjectClasses)
+	{
+		if (InObjClass->IsChildOf(TargetObjectClass))
+		{
+			return true;
+		}
+	}
+
+	return false;
+}
+
+}  // namespace UE::Cameras
+
diff --git a/Engine/Plugins/Cameras/GameplayCameras/Source/GameplayCameras/Private/Tests/GameplayCamerasTestBuilder.h b/Engine/Plugins/Cameras/GameplayCameras/Source/GameplayCameras/Private/Tests/GameplayCamerasTestBuilder.h
--- a/Engine/Plugins/Cameras/GameplayCameras/Source/GameplayCameras/Private/Tests/GameplayCamerasTestBuilder.h
+++ b/Engine/Plugins/Cameras/GameplayCameras/Source/GameplayCameras/Private/Tests/GameplayCamerasTestBuilder.h
@@ -715,7 +715,6 @@ public:
 	TScopedCameraRigAssetTestBuilder<ThisType> AddCameraRig(FName Name = NAME_None)
 	{
 		TScopedCameraRigAssetTestBuilder<ThisType> CameraRigBuilder(*this, GetNamedObjectRegistry(), Name, CameraAsset);
-		CameraAsset->AddCameraRig(CameraRigBuilder.Get());
 		return CameraRigBuilder;
 	}
 
diff --git a/Engine/Plugins/Cameras/GameplayCameras/Source/GameplayCameras/Public/Core/CameraAsset.h b/Engine/Plugins/Cameras/GameplayCameras/Source/GameplayCameras/Public/Core/CameraAsset.h
--- a/Engine/Plugins/Cameras/GameplayCameras/Source/GameplayCameras/Public/Core/CameraAsset.h
+++ b/Engine/Plugins/Cameras/GameplayCameras/Source/GameplayCameras/Public/Core/CameraAsset.h
@@ -30,8 +30,6 @@ namespace UE::Cameras
 
 		/** Called when the camera director has been changed. */
 		virtual void OnCameraDirectorChanged(UCameraAsset* InCameraAsset, const TCameraPropertyChangedEvent<UCameraDirector*>& Event) {}
-		/** Changed when the camera rigs have been changed. */
-		virtual void OnCameraRigsChanged(UCameraAsset* InCameraAsset, const TCameraArrayChangedEvent<UCameraRigAsset*>& Event) {}
 		/* Changed when the enter transitions have been changed. */
 		virtual void OnEnterTransitionsChanged(UCameraAsset* InCameraAsset, const TCameraArrayChangedEvent<UCameraRigTransition*>& Event) {}
 		/* Changed when the exit transitions have been changed. */
@@ -58,13 +56,6 @@ public:
 	/** Sets the camera director. */
 	GAMEPLAYCAMERAS_API void SetCameraDirector(UCameraDirector* InCameraDirector);
 
-	/** Gets the camera rigs. */
-	TArrayView<const TObjectPtr<UCameraRigAsset>> GetCameraRigs() const { return CameraRigs; }
-	/** Adds a a camera rig. */
-	GAMEPLAYCAMERAS_API void AddCameraRig(UCameraRigAsset* InCameraRig);
-	/** Removes a camera rig. */
-	GAMEPLAYCAMERAS_API int32 RemoveCameraRig(UCameraRigAsset* InCameraRig);
-
 	/** Gets the enter transitions. */
 	TArrayView<const TObjectPtr<UCameraRigTransition>> GetEnterTransitions() const { return EnterTransitions; }
 	/** Adds an enter transition. */
@@ -135,21 +126,15 @@ public:
 	/** Event handlers to be notified of data changes. */
 	UE::Cameras::TCameraEventHandlerContainer<UE::Cameras::ICameraAssetEventHandler> EventHandlers;
 
-private:
-
-#if WITH_EDITOR
-	void CleanUpStrayObjects();
-#endif
-
 private:
 
 	/** The camera director to use in this camera. */
 	UPROPERTY(Instanced)
 	TObjectPtr<UCameraDirector> CameraDirector;
 
-	/** The list of camera rigs used by this camera. */
+	/** The list of camera rigs used by this camera (deprecated). */
 	UPROPERTY()
-	TArray<TObjectPtr<UCameraRigAsset>> CameraRigs;
+	TArray<TObjectPtr<UCameraRigAsset>> CameraRigs_DEPRECATED;
 
 	/** A list of default enter transitions for all the camera rigs in this asset. */
 	UPROPERTY(meta=(ObjectTreeGraphPinDirection=Input))
diff --git a/Engine/Plugins/Cameras/GameplayCameras/Source/GameplayCameras/Public/Core/CameraDirector.h b/Engine/Plugins/Cameras/GameplayCameras/Source/GameplayCameras/Public/Core/CameraDirector.h
--- a/Engine/Plugins/Cameras/GameplayCameras/Source/GameplayCameras/Public/Core/CameraDirector.h
+++ b/Engine/Plugins/Cameras/GameplayCameras/Source/GameplayCameras/Public/Core/CameraDirector.h
@@ -20,6 +20,14 @@ struct FCameraDirectorFactoryCreateParams
 {
 };
 
+/**
+ * Parameter struct for gathering camera rigs used by a director.
+ */
+struct FCameraDirectorRigUsageInfo
+{
+	TArray<UCameraRigAsset*> CameraRigs;
+};
+
 #endif
 
 /**
@@ -43,6 +51,9 @@ public:
 #if WITH_EDITOR
 	/** Called by the asset factories to setup new data before the editor opens. */
 	void FactoryCreateAsset(const FCameraDirectorFactoryCreateParams& InParams);
+
+	/** Gets the list of camera rigs used by this camera director. */
+	void GatherRigUsageInfo(FCameraDirectorRigUsageInfo& UsageInfo);
 #endif
 
 protected:
@@ -56,6 +67,9 @@ protected:
 #if WITH_EDITOR
 	/** Called by the asset factories to setup new data before the editor opens. */
 	virtual void OnFactoryCreateAsset(const FCameraDirectorFactoryCreateParams& InParams) {}
+
+	/** Gets the list of camera rigs used by this camera director. */
+	virtual void OnGatherRigUsageInfo(FCameraDirectorRigUsageInfo& UsageInfo) {}
 #endif
 };
 
diff --git a/Engine/Plugins/Cameras/GameplayCameras/Source/GameplayCameras/Public/Directors/BlueprintCameraDirector.h b/Engine/Plugins/Cameras/GameplayCameras/Source/GameplayCameras/Public/Directors/BlueprintCameraDirector.h
--- a/Engine/Plugins/Cameras/GameplayCameras/Source/GameplayCameras/Public/Directors/BlueprintCameraDirector.h
+++ b/Engine/Plugins/Cameras/GameplayCameras/Source/GameplayCameras/Public/Directors/BlueprintCameraDirector.h
@@ -154,9 +154,7 @@ public:
 
 	/** Specifies a camera rig to be active this frame. */
 	UFUNCTION(BlueprintCallable, Category="Evaluation")
-	void ActivateCameraRig(
-			UPARAM(meta=(UseBlueprintCameraDirectorRigPicker=true))
-			UCameraRigAsset* CameraRig);
+	void ActivateCameraRig(UCameraRigAsset* CameraRig);
 
 	/**
 	 * Specifies a camera rig to be active this frame, via a proxy which is later resolved
@@ -165,18 +163,6 @@ public:
 	UFUNCTION(BlueprintCallable, Category="Evaluation")
 	void ActivateCameraRigViaProxy(UCameraRigProxyAsset* CameraRigProxy);
 
-	/**
-	 * Specifies an external camera rig prefab asset to be active this frame.
-	 */
-	UFUNCTION(BlueprintCallable, Category="Evaluation")
-	void ActivateCameraRigPrefab(UCameraRigAsset* CameraRig);
-
-	/** Gets a camera rig from the referencing camera asset. */
-	UFUNCTION(BlueprintPure, Category="Evaluation", meta=(HideSelfPin=true))
-	UCameraRigAsset* GetCameraRig(
-			UPARAM(meta=(UseBlueprintCameraDirectorRigPicker=true))
-			UCameraRigAsset* CameraRig) const;
-
 public:
 
 	/**
@@ -264,6 +250,7 @@ protected:
 	virtual void OnBuildCameraDirector(UE::Cameras::FCameraBuildLog& BuildLog) override;
 #if WITH_EDITOR
 	virtual void OnFactoryCreateAsset(const FCameraDirectorFactoryCreateParams& InParams) override;
+	virtual void OnGatherRigUsageInfo(FCameraDirectorRigUsageInfo& UsageInfo) override;
 #endif
 };
 
diff --git a/Engine/Plugins/Cameras/GameplayCameras/Source/GameplayCameras/Public/Directors/SingleCameraDirector.h b/Engine/Plugins/Cameras/GameplayCameras/Source/GameplayCameras/Public/Directors/SingleCameraDirector.h
--- a/Engine/Plugins/Cameras/GameplayCameras/Source/GameplayCameras/Public/Directors/SingleCameraDirector.h
+++ b/Engine/Plugins/Cameras/GameplayCameras/Source/GameplayCameras/Public/Directors/SingleCameraDirector.h
@@ -24,7 +24,7 @@ protected:
 	virtual FCameraDirectorEvaluatorPtr OnBuildEvaluator(FCameraDirectorEvaluatorBuilder& Builder) const override;
 	virtual void OnBuildCameraDirector(UE::Cameras::FCameraBuildLog& BuildLog) override;
 #if WITH_EDITOR
-	virtual void OnFactoryCreateAsset(const FCameraDirectorFactoryCreateParams& InParams) override;
+	virtual void OnGatherRigUsageInfo(FCameraDirectorRigUsageInfo& UsageInfo) override;
 #endif
 
 public:
diff --git a/Engine/Plugins/Cameras/GameplayCameras/Source/GameplayCameras/Public/Directors/StateTreeCameraDirector.h b/Engine/Plugins/Cameras/GameplayCameras/Source/GameplayCameras/Public/Directors/StateTreeCameraDirector.h
--- a/Engine/Plugins/Cameras/GameplayCameras/Source/GameplayCameras/Public/Directors/StateTreeCameraDirector.h
+++ b/Engine/Plugins/Cameras/GameplayCameras/Source/GameplayCameras/Public/Directors/StateTreeCameraDirector.h
@@ -30,6 +30,7 @@ protected:
 	virtual void OnBuildCameraDirector(UE::Cameras::FCameraBuildLog& BuildLog) override;
 #if WITH_EDITOR
 	virtual void OnFactoryCreateAsset(const FCameraDirectorFactoryCreateParams& InParams) override;
+	virtual void OnGatherRigUsageInfo(FCameraDirectorRigUsageInfo& UsageInfo) override;
 #endif
 
 public:
diff --git a/Engine/Plugins/Cameras/GameplayCameras/Source/GameplayCameras/Public/Helpers/ObjectReferenceFinder.h b/Engine/Plugins/Cameras/GameplayCameras/Source/GameplayCameras/Public/Helpers/ObjectReferenceFinder.h
new file mode 100644
--- /dev/null
+++ b/Engine/Plugins/Cameras/GameplayCameras/Source/GameplayCameras/Public/Helpers/ObjectReferenceFinder.h
@@ -0,0 +1,56 @@
+// Copyright Epic Games, Inc. All Rights Reserved.
+
+#pragma once
+
+#include "Containers/Array.h"
+#include "Containers/Map.h"
+#include "Containers/Set.h"
+#include "CoreTypes.h"
+#include "Serialization/ArchiveUObject.h"
+
+namespace UE::Cameras
+{
+
+/**
+ * A utility class for finding references from a given package to a set of UObjects.
+ */
+class GAMEPLAYCAMERAS_API FObjectReferenceFinder : public FArchiveUObject
+{
+public:
+
+	/** Creates a new instance of the reference finder. */
+	FObjectReferenceFinder(UObject* InRootObject, TArrayView<UObject*> InReferencedObjects);
+
+	/** Runs the reference finding. */
+	void CollectReferences();
+
+	/** Whether any object was found in the package that references any of the target objects. */
+	bool HasAnyObjectReference() const;
+
+	/** Returns the number of found references to the given target object. */
+	int32 GetObjectReferenceCount(UObject* InObject) const;
+
+protected:
+
+	// FArchive interface.
+	virtual FArchive& operator<<(UObject*& ObjRef) override;
+
+private:
+
+	void Initialize(UObject* InRootObject);
+
+private:
+
+	UObject* RootObject;
+	UPackage* PackageScope;
+
+	TSet<UObject*> TargetObjects;
+
+	TArray<UObject*> ObjectsToVisit;
+	TSet<UObject*> VisitedObjects;
+
+	TMap<UObject*, int32> ObjectReferenceCounts;
+};
+
+}  // namespace UE::Cameras
+
diff --git a/Engine/Plugins/Cameras/GameplayCameras/Source/GameplayCameras/Public/Helpers/OutgoingReferenceFinder.h b/Engine/Plugins/Cameras/GameplayCameras/Source/GameplayCameras/Public/Helpers/OutgoingReferenceFinder.h
new file mode 100644
--- /dev/null
+++ b/Engine/Plugins/Cameras/GameplayCameras/Source/GameplayCameras/Public/Helpers/OutgoingReferenceFinder.h
@@ -0,0 +1,77 @@
+// Copyright Epic Games, Inc. All Rights Reserved.
+
+#pragma once
+
+#include "Containers/Array.h"
+#include "Containers/Map.h"
+#include "Containers/Set.h"
+#include "CoreTypes.h"
+#include "Serialization/ArchiveUObject.h"
+
+namespace UE::Cameras
+{
+
+/**
+ * A utility class for finding outgoing references from a given package to external UObjects
+ * of a given class.
+ */
+class GAMEPLAYCAMERAS_API FOutgoingReferenceFinder : public FArchiveUObject
+{
+public:
+
+	/** Creates a new instance of the reference finder. */
+	FOutgoingReferenceFinder(UObject* InRootObject, UClass* InReferencedObjectClass);
+
+	/** Creates a new instance of the reference finder. */
+	FOutgoingReferenceFinder(UObject* InRootObject, TArrayView<UClass*> InReferencedObjectClasses);
+
+	/** Runs the reference finding. */
+	void CollectReferences();
+
+	/** Gets the references of a given class. */
+	template<typename ObjectClass>
+	bool GetReferencesOfClass(TArray<ObjectClass*>& OutReferencedObjects) const;
+
+	/** Gets all referenced objects. */
+	bool GetAllReferences(TArray<UObject*>& OutReferencedObjects) const;
+
+protected:
+
+	// FArchive interface.
+	virtual FArchive& operator<<(UObject*& ObjRef) override;
+
+private:
+
+	void Initialize(UObject* InRootObject);
+
+	bool MatchesAnyTargetClass(UClass* InObjClass) const;
+
+private:
+
+	UObject* RootObject;
+	UPackage* PackageScope;
+
+	TSet<UClass*> TargetObjectClasses;
+
+	TArray<UObject*> ObjectsToVisit;
+	TSet<UObject*> VisitedObjects;
+
+	TMap<UClass*, TSet<UObject*>> ReferencedObjects;
+};
+
+template<typename ObjectClass>
+bool FOutgoingReferenceFinder::GetReferencesOfClass(TArray<ObjectClass*>& OutReferencedObjects) const
+{
+	if (const TSet<UObject*>* ReferencesOfClass = ReferencedObjects.Find(ObjectClass::StaticClass()))
+	{
+		for (UObject* Obj : (*ReferencesOfClass))
+		{
+			OutReferencedObjects.Add(CastChecked<ObjectClass>(Obj));
+		}
+		return !ReferencesOfClass->IsEmpty();
+	}
+	return false;
+}
+
+}  // namespace UE::Cameras
+
diff --git a/Engine/Plugins/Cameras/GameplayCameras/Source/GameplayCamerasEditor/Private/AssetTools/AssetDefinition_CameraRigAsset.cpp b/Engine/Plugins/Cameras/GameplayCameras/Source/GameplayCamerasEditor/Private/AssetTools/AssetDefinition_CameraRigAsset.cpp
--- a/Engine/Plugins/Cameras/GameplayCameras/Source/GameplayCamerasEditor/Private/AssetTools/AssetDefinition_CameraRigAsset.cpp
+++ b/Engine/Plugins/Cameras/GameplayCameras/Source/GameplayCamerasEditor/Private/AssetTools/AssetDefinition_CameraRigAsset.cpp
@@ -14,7 +14,7 @@
 
 FText UAssetDefinition_CameraRigAsset::GetAssetDisplayName() const
 {
-	return LOCTEXT("AssetDisplayName", "Camera Rig Prefab");
+	return LOCTEXT("AssetDisplayName", "Camera Rig");
 }
 
 FLinearColor UAssetDefinition_CameraRigAsset::GetAssetColor() const
diff --git a/Engine/Plugins/Cameras/GameplayCameras/Source/GameplayCamerasEditor/Private/Builders/BlueprintCameraDirectorEditorBuilder.cpp b/Engine/Plugins/Cameras/GameplayCameras/Source/GameplayCamerasEditor/Private/Builders/BlueprintCameraDirectorEditorBuilder.cpp
deleted file mode 100644
diff --git a/Engine/Plugins/Cameras/GameplayCameras/Source/GameplayCamerasEditor/Private/Builders/BlueprintCameraDirectorEditorBuilder.h b/Engine/Plugins/Cameras/GameplayCameras/Source/GameplayCamerasEditor/Private/Builders/BlueprintCameraDirectorEditorBuilder.h
deleted file mode 100644
diff --git a/Engine/Plugins/Cameras/GameplayCameras/Source/GameplayCamerasEditor/Private/Customizations/CameraProxyTableDetailsCustomization.cpp b/Engine/Plugins/Cameras/GameplayCameras/Source/GameplayCamerasEditor/Private/Customizations/CameraProxyTableDetailsCustomization.cpp
--- a/Engine/Plugins/Cameras/GameplayCameras/Source/GameplayCamerasEditor/Private/Customizations/CameraProxyTableDetailsCustomization.cpp
+++ b/Engine/Plugins/Cameras/GameplayCameras/Source/GameplayCamerasEditor/Private/Customizations/CameraProxyTableDetailsCustomization.cpp
@@ -2,12 +2,13 @@
 
 #include "Customizations/CameraProxyTableDetailsCustomization.h"
 
+#include "ContentBrowserModule.h"
 #include "Core/CameraAsset.h"
 #include "Core/CameraRigAsset.h"
 #include "Core/CameraRigProxyTable.h"
 #include "DetailCategoryBuilder.h"
 #include "DetailWidgetRow.h"
-#include "Editors/CameraRigPickerConfig.h"
+#include "IContentBrowserSingleton.h"
 #include "IDetailChildrenBuilder.h"
 #include "IGameplayCamerasEditorModule.h"
 #include "Modules/ModuleManager.h"
@@ -108,23 +109,32 @@ TSharedRef<SWidget> FCameraProxyTableEntryDetailsCustomization::OnBuildCameraRig
 		return SNullWidget::NullWidget;
 	}
 
-	IGameplayCamerasEditorModule& CamerasEditorModule = FModuleManager::LoadModuleChecked<IGameplayCamerasEditorModule>("GameplayCamerasEditor");
+	FContentBrowserModule& ContentBrowserModule = FModuleManager::Get().LoadModuleChecked<FContentBrowserModule>(TEXT("ContentBrowser"));
 
 	TObjectPtr<UCameraRigAsset>* CameraRigPtr = (TObjectPtr<UCameraRigAsset>*)RawData[0];
 	UCameraAsset* OuterCameraAsset = ProxyTables[0]->GetTypedOuter<UCameraAsset>();
 
-	FCameraRigPickerConfig CameraRigPickerConfig;
-	CameraRigPickerConfig.bCanSelectCameraAsset = false;
-	CameraRigPickerConfig.InitialCameraAssetSelection = FAssetData(OuterCameraAsset);
-	CameraRigPickerConfig.OnCameraRigSelected = FOnCameraRigSelected::CreateSP(
+	FAssetPickerConfig CameraRigPickerConfig;
+
+	FARFilter ARFilter;
+	ARFilter.ClassPaths.Add(FTopLevelAssetPath(UCameraRigAsset::StaticClass()->GetPathName()));
+
+	CameraRigPickerConfig.bShowBottomToolbar = true;
+	CameraRigPickerConfig.bAllowNullSelection = true;
+	CameraRigPickerConfig.bFocusSearchBoxWhenOpened = true;
+	CameraRigPickerConfig.SelectionMode = ESelectionMode::Single;
+	CameraRigPickerConfig.Filter = ARFilter;
+	CameraRigPickerConfig.SaveSettingsName = TEXT("CameraProxyTableEntryRigPickerSettings");
+	CameraRigPickerConfig.InitialAssetViewType = EAssetViewType::List;
+	CameraRigPickerConfig.InitialAssetSelection = FAssetData(*CameraRigPtr);
+	CameraRigPickerConfig.OnAssetSelected = FOnAssetSelected::CreateSP(
 			this, &FCameraProxyTableEntryDetailsCustomization::OnCameraRigSelected);
-	CameraRigPickerConfig.PropertyToSet = CameraRigPropertyHandle;
-	CameraRigPickerConfig.InitialCameraRigSelection = (*CameraRigPtr);
+	CameraRigPickerConfig.PropertyHandle = CameraRigPropertyHandle;
 
-	return CamerasEditorModule.CreateCameraRigPicker(CameraRigPickerConfig);
+	return ContentBrowserModule.Get().CreateAssetPicker(CameraRigPickerConfig);
 }
 
-void FCameraProxyTableEntryDetailsCustomization::OnCameraRigSelected(UCameraRigAsset* CameraRig)
+void FCameraProxyTableEntryDetailsCustomization::OnCameraRigSelected(const FAssetData& InSelectedAsset)
 {
 	ComboButton->SetIsOpen(false);
 }
diff --git a/Engine/Plugins/Cameras/GameplayCameras/Source/GameplayCamerasEditor/Private/Customizations/CameraRigPtrDetailsCustomization.cpp b/Engine/Plugins/Cameras/GameplayCameras/Source/GameplayCamerasEditor/Private/Customizations/CameraRigPtrDetailsCustomization.cpp
deleted file mode 100644
diff --git a/Engine/Plugins/Cameras/GameplayCameras/Source/GameplayCamerasEditor/Private/Editors/GameplayCamerasGraphPanelPinFactory.cpp b/Engine/Plugins/Cameras/GameplayCameras/Source/GameplayCamerasEditor/Private/Editors/GameplayCamerasGraphPanelPinFactory.cpp
--- a/Engine/Plugins/Cameras/GameplayCameras/Source/GameplayCamerasEditor/Private/Editors/GameplayCamerasGraphPanelPinFactory.cpp
+++ b/Engine/Plugins/Cameras/GameplayCameras/Source/GameplayCamerasEditor/Private/Editors/GameplayCamerasGraphPanelPinFactory.cpp
@@ -5,8 +5,6 @@
 #include "Core/CameraRigAsset.h"
 #include "Core/CameraVariableAssets.h"
 #include "EdGraphSchema_K2.h"
-#include "Editors/SBlueprintCameraDirectorRigNameGraphPin.h"
-#include "Editors/SCameraRigNameGraphPin.h"
 #include "Editors/SCameraVariableNameGraphPin.h"
 #include "K2Node_CallFunction.h"
 
@@ -20,22 +18,10 @@ TSharedPtr<SGraphPin> FGameplayCamerasGraphPanelPinFactory::CreatePin(UEdGraphPi
 		return nullptr;
 	}
 
-	if (UK2Node_CallFunction* CallFunctionNode = Cast<UK2Node_CallFunction>(Pin->GetOwningNode()))
-	{
-		if (TSharedPtr<SGraphPin> PinWidget = CreateFunctionParameterPin(Pin, CallFunctionNode))
-		{
-			return PinWidget;
-		}
-	}
-
 	const FEdGraphPinType& PinType = Pin->PinType;
 	const UClass* PinPropertyClass = Cast<const UClass>(PinType.PinSubCategoryObject);
 	if (PinType.PinCategory == UEdGraphSchema_K2::PC_Object && PinPropertyClass)
 	{
-		if (PinPropertyClass == UCameraRigAsset::StaticClass())
-		{
-			return CreateCameraRigPickerPin(Pin);
-		}
 		if (PinPropertyClass->IsChildOf<UCameraVariableAsset>())
 		{
 			return CreateCameraVariablePickerPin(Pin);
@@ -45,57 +31,6 @@ TSharedPtr<SGraphPin> FGameplayCamerasGraphPanelPinFactory::CreatePin(UEdGraphPi
 	return nullptr;
 }
 
-TSharedPtr<SGraphPin> FGameplayCamerasGraphPanelPinFactory::CreateFunctionParameterPin(UEdGraphPin* Pin, UK2Node_CallFunction* CallFunctionNode) const
-{
-	UClass* BlueprintClass = CallFunctionNode->GetBlueprintClassFromNode();
-	UFunction* ReferencedFunction = CallFunctionNode->FunctionReference.ResolveMember<UFunction>(BlueprintClass);
-	if (!ReferencedFunction)
-	{
-		return nullptr;
-	}
-
-	FProperty* ParameterProperty = ReferencedFunction->FindPropertyByName(Pin->PinName);
-	if (!ParameterProperty)
-	{
-		return nullptr;
-	}
-
-	if (ParameterProperty->HasMetaData(TEXT("UseBlueprintCameraDirectorRigPicker")))
-	{
-		return SNew(SBlueprintCameraDirectorRigNameGraphPin, Pin);
-	}
-
-	if (ParameterProperty->HasMetaData(TEXT("UseCameraRigPicker")))
-	{
-		return SNew(SCameraRigNameGraphPin, Pin);
-	}
-
-	return nullptr;
-}
-
-TSharedPtr<SGraphPin> FGameplayCamerasGraphPanelPinFactory::CreateCameraRigPickerPin(UEdGraphPin* Pin) const
-{
-	UEdGraphNode* OwningNode = Pin->GetOwningNode();
-	if (!OwningNode)
-	{
-		return nullptr;
-	}
-	const UClass* OwningNodeClass = OwningNode->GetClass();
-	const FString& UseCameraRigPickerForPinsMetaData = OwningNodeClass->GetMetaData(TEXT("UseCameraRigPickerForPins"));
-	if (UseCameraRigPickerForPinsMetaData.IsEmpty())
-	{
-		return nullptr;
-	}
-
-	TArray<FString> CameraRigPickerPinNames;
-	UseCameraRigPickerForPinsMetaData.ParseIntoArray(CameraRigPickerPinNames, TEXT(","));
-	if (CameraRigPickerPinNames.Contains(Pin->GetName()))
-	{
-		return SNew(SCameraRigNameGraphPin, Pin);
-	}
-	return nullptr;
-}
-
 TSharedPtr<SGraphPin> FGameplayCamerasGraphPanelPinFactory::CreateCameraVariablePickerPin(UEdGraphPin* Pin) const
 {
 	return SNew(SCameraVariableNameGraphPin, Pin);
diff --git a/Engine/Plugins/Cameras/GameplayCameras/Source/GameplayCamerasEditor/Private/Editors/SBlueprintCameraDirectorRigNameGraphPin.cpp b/Engine/Plugins/Cameras/GameplayCameras/Source/GameplayCamerasEditor/Private/Editors/SBlueprintCameraDirectorRigNameGraphPin.cpp
deleted file mode 100644
diff --git a/Engine/Plugins/Cameras/GameplayCameras/Source/GameplayCamerasEditor/Private/Editors/SCameraRigList.cpp b/Engine/Plugins/Cameras/GameplayCameras/Source/GameplayCamerasEditor/Private/Editors/SCameraRigList.cpp
deleted file mode 100644
diff --git a/Engine/Plugins/Cameras/GameplayCameras/Source/GameplayCamerasEditor/Private/Editors/SCameraRigNameGraphPin.cpp b/Engine/Plugins/Cameras/GameplayCameras/Source/GameplayCamerasEditor/Private/Editors/SCameraRigNameGraphPin.cpp
deleted file mode 100644
diff --git a/Engine/Plugins/Cameras/GameplayCameras/Source/GameplayCamerasEditor/Private/Editors/SCameraRigPicker.cpp b/Engine/Plugins/Cameras/GameplayCameras/Source/GameplayCamerasEditor/Private/Editors/SCameraRigPicker.cpp
deleted file mode 100644
diff --git a/Engine/Plugins/Cameras/GameplayCameras/Source/GameplayCamerasEditor/Private/Editors/SCameraRigPicker.h b/Engine/Plugins/Cameras/GameplayCameras/Source/GameplayCamerasEditor/Private/Editors/SCameraRigPicker.h
deleted file mode 100644
diff --git a/Engine/Plugins/Cameras/GameplayCameras/Source/GameplayCamerasEditor/Private/Editors/SCameraRigPickerButton.cpp b/Engine/Plugins/Cameras/GameplayCameras/Source/GameplayCamerasEditor/Private/Editors/SCameraRigPickerButton.cpp
deleted file mode 100644
diff --git a/Engine/Plugins/Cameras/GameplayCameras/Source/GameplayCamerasEditor/Private/Factories/CameraAssetFactory.cpp b/Engine/Plugins/Cameras/GameplayCameras/Source/GameplayCamerasEditor/Private/Factories/CameraAssetFactory.cpp
--- a/Engine/Plugins/Cameras/GameplayCameras/Source/GameplayCamerasEditor/Private/Factories/CameraAssetFactory.cpp
+++ b/Engine/Plugins/Cameras/GameplayCameras/Source/GameplayCamerasEditor/Private/Factories/CameraAssetFactory.cpp
@@ -60,9 +60,6 @@ UObject* UCameraAssetFactory::FactoryCreateNew(UClass* Class, UObject* Parent, F
 {
 	UCameraAsset* NewCameraAsset = NewObject<UCameraAsset>(Parent, Class, Name, Flags | RF_Transactional);
 
-	UCameraRigAsset* FirstCameraRig = NewObject<UCameraRigAsset>(NewCameraAsset, NAME_None, RF_Transactional | RF_Public);
-	NewCameraAsset->AddCameraRig(FirstCameraRig);
-
 	if (CameraDirectorClass)
 	{
 		UCameraDirector* NewCameraDirector = NewObject<UCameraDirector>(NewCameraAsset, CameraDirectorClass, NAME_None, RF_Transactional);
diff --git a/Engine/Plugins/Cameras/GameplayCameras/Source/GameplayCamerasEditor/Private/GameplayCamerasEditorModule.cpp b/Engine/Plugins/Cameras/GameplayCameras/Source/GameplayCamerasEditor/Private/GameplayCamerasEditorModule.cpp
--- a/Engine/Plugins/Cameras/GameplayCameras/Source/GameplayCamerasEditor/Private/GameplayCamerasEditorModule.cpp
+++ b/Engine/Plugins/Cameras/GameplayCameras/Source/GameplayCamerasEditor/Private/GameplayCamerasEditorModule.cpp
@@ -6,28 +6,25 @@
 #include "AssetTools/CameraRigAssetEditor.h"
 #include "AssetTools/CameraRigProxyAssetEditor.h"
 #include "AssetTools/CameraVariableCollectionEditor.h"
-#include "Builders/BlueprintCameraDirectorEditorBuilder.h"
 #include "Commands/CameraAssetEditorCommands.h"
 #include "Commands/CameraRigAssetEditorCommands.h"
 #include "Commands/CameraRigTransitionEditorCommands.h"
 #include "Commands/CameraVariableCollectionEditorCommands.h"
 #include "Commands/GameplayCamerasDebuggerCommands.h"
 #include "Customizations/CameraParameterDetailsCustomizations.h"
-#include "Customizations/CameraRigAssetReferenceDetailsCustomization.h"
-#include "Customizations/CameraRigPtrDetailsCustomization.h"
 #include "Customizations/CameraProxyTableDetailsCustomization.h"
+#include "Customizations/CameraRigAssetReferenceDetailsCustomization.h"
 #include "Customizations/CameraVariableReferenceDetailsCustomizations.h"
 #include "Customizations/FilmbackCameraNodeDetailsCustomization.h"
 #include "Debug/CameraDebugCategories.h"
 #include "Debugger/SBlendStacksDebugPanel.h"
 #include "Debugger/SCameraNodeTreeDebugPanel.h"
-#include "Debugger/SEvaluationServicesDebugPanel.h"
 #include "Debugger/SCameraPoseStatsDebugPanel.h"
+#include "Debugger/SEvaluationServicesDebugPanel.h"
 #include "Debugger/SGameplayCamerasDebugger.h"
 #include "Directors/BlueprintCameraDirector.h"
 #include "EdGraph/EdGraph.h"
 #include "Editors/GameplayCamerasGraphPanelPinFactory.h"
-#include "Editors/SCameraRigPicker.h"
 #include "Editors/SCameraVariablePicker.h"
 #include "Features/IModularFeatures.h"
 #include "GameplayCameras.h"
@@ -158,14 +155,6 @@ public:
 		return AssetEditor;
 	}
 
-	virtual TSharedRef<SWidget> CreateCameraRigPicker(const FCameraRigPickerConfig& InPickerConfig) override
-	{
-		using namespace UE::Cameras;
-
-		return SNew(SCameraRigPicker)
-			.CameraRigPickerConfig(InPickerConfig);
-	}
-
 	virtual TSharedRef<SWidget> CreateCameraVariablePicker(const FCameraVariablePickerConfig& InPickerConfig) override
 	{
 		using namespace UE::Cameras;
@@ -345,11 +334,7 @@ private:
 
 	void RegisterBuilders()
 	{
-		using namespace UE::Cameras;
-
-		BuiltInCameraAssetBuilders.Add(
-				RegisterCameraAssetBuilder(FOnBuildCameraAsset::CreateStatic(
-						&FBlueprintCameraDirectorEditorBuilder::OnBuildCameraAsset)));
+		// No default builders yet.
 	}
 
 	void UnregisterBuilders()
@@ -485,8 +470,6 @@ private:
 
 		PropertyEditorModule.RegisterCustomPropertyTypeLayout("CameraRigProxyTableEntry", FOnGetPropertyTypeCustomizationInstance::CreateStatic(
 					&FCameraProxyTableEntryDetailsCustomization::MakeInstance));
-		PropertyEditorModule.RegisterCustomPropertyTypeLayout("CameraRigAsset", FOnGetPropertyTypeCustomizationInstance::CreateStatic(
-					&FCameraRigPtrDetailsCustomization::MakeInstance));
 		PropertyEditorModule.RegisterCustomPropertyTypeLayout("CameraRigAssetReference", FOnGetPropertyTypeCustomizationInstance::CreateStatic(
 					&FCameraRigAssetReferenceDetailsCustomization::MakeInstance));
 
diff --git a/Engine/Plugins/Cameras/GameplayCameras/Source/GameplayCamerasEditor/Private/Helpers/CameraAssetReferenceGatherer.cpp b/Engine/Plugins/Cameras/GameplayCameras/Source/GameplayCamerasEditor/Private/Helpers/CameraAssetReferenceGatherer.cpp
deleted file mode 100644
diff --git a/Engine/Plugins/Cameras/GameplayCameras/Source/GameplayCamerasEditor/Private/Toolkits/BlueprintCameraDirectorAssetEditorMode.cpp b/Engine/Plugins/Cameras/GameplayCameras/Source/GameplayCamerasEditor/Private/Toolkits/BlueprintCameraDirectorAssetEditorMode.cpp
--- a/Engine/Plugins/Cameras/GameplayCameras/Source/GameplayCamerasEditor/Private/Toolkits/BlueprintCameraDirectorAssetEditorMode.cpp
+++ b/Engine/Plugins/Cameras/GameplayCameras/Source/GameplayCamerasEditor/Private/Toolkits/BlueprintCameraDirectorAssetEditorMode.cpp
@@ -21,30 +21,6 @@ TSharedPtr<FCameraDirectorAssetEditorMode> FBlueprintCameraDirectorAssetEditorMo
 FBlueprintCameraDirectorAssetEditorMode::FBlueprintCameraDirectorAssetEditorMode(UCameraAsset* InCameraAsset)
 	: FCameraDirectorAssetEditorMode(InCameraAsset)
 {
-	if (InCameraAsset)
-	{
-		InCameraAsset->EventHandlers.Register(EventHandler, this);
-	}
-}
-
-void FBlueprintCameraDirectorAssetEditorMode::OnCameraRigsChanged(UCameraAsset* InCameraAsset, const TCameraArrayChangedEvent<UCameraRigAsset*>& Event)
-{
-	UBlueprintCameraDirector* CameraDirector = Cast<UBlueprintCameraDirector>(InCameraAsset->GetCameraDirector());
-	if (!ensure(CameraDirector))
-	{
-		return;
-	}
-
-	TSet<UCameraRigAsset*> CameraRigs(InCameraAsset->GetCameraRigs());
-
-	for (FCameraRigProxyTableEntry& Entry : CameraDirector->CameraRigProxyTable->Entries)
-	{
-		if (!CameraRigs.Contains(Entry.CameraRig))
-		{
-			CameraDirector->CameraRigProxyTable->Modify();
-			Entry.CameraRig = nullptr;
-		}
-	}
 }
 
 }  // namespace UE::Cameras
diff --git a/Engine/Plugins/Cameras/GameplayCameras/Source/GameplayCamerasEditor/Private/Toolkits/CameraAssetEditorToolkit.cpp b/Engine/Plugins/Cameras/GameplayCameras/Source/GameplayCamerasEditor/Private/Toolkits/CameraAssetEditorToolkit.cpp
--- a/Engine/Plugins/Cameras/GameplayCameras/Source/GameplayCamerasEditor/Private/Toolkits/CameraAssetEditorToolkit.cpp
+++ b/Engine/Plugins/Cameras/GameplayCameras/Source/GameplayCamerasEditor/Private/Toolkits/CameraAssetEditorToolkit.cpp
@@ -2,17 +2,22 @@
 
 #include "Toolkits/CameraAssetEditorToolkit.h"
 
+#include "AssetRegistry/AssetRegistryModule.h"
 #include "AssetTools/CameraAssetEditor.h"
 #include "Commands/CameraAssetEditorCommands.h"
 #include "Core/CameraAsset.h"
 #include "Core/CameraAssetBuilder.h"
 #include "Core/CameraBuildLog.h"
 #include "Core/CameraDirector.h"
+#include "Core/CameraRigAsset.h"
 #include "Editors/ObjectTreeGraphConfig.h"
 #include "Editors/SFindInObjectTreeGraph.h"
+#include "FileHelpers.h"
 #include "Framework/Docking/LayoutExtender.h"
 #include "Framework/Docking/TabManager.h"
 #include "GameplayCamerasEditorSettings.h"
+#include "Helpers/ObjectReferenceFinder.h"
+#include "IAssetTools.h"
 #include "IGameplayCamerasEditorModule.h"
 #include "IGameplayCamerasLiveEditManager.h"
 #include "IGameplayCamerasModule.h"
@@ -22,9 +27,9 @@
 #include "Toolkits/BuildButtonToolkit.h"
 #include "Toolkits/CameraBuildLogToolkit.h"
 #include "Toolkits/CameraDirectorAssetEditorMode.h"
-#include "Toolkits/CameraRigsAssetEditorMode.h"
 #include "Toolkits/CameraSharedTransitionsAssetEditorMode.h"
 #include "Toolkits/StandardToolkitLayout.h"
+#include "UObject/Package.h"
 #include "Widgets/Docking/SDockTab.h"
 
 #include UE_INLINE_GENERATED_CPP_BY_NAME(CameraAssetEditorToolkit)
@@ -180,7 +185,6 @@ void FCameraAssetEditorToolkit::RegisterToolbar()
 
 		FToolMenuSection& ModesSection = ToolbarMenu->AddSection("EditorModes", TAttribute<FText>(), InsertAfterAssetSection);
 		ModesSection.AddEntry(FToolMenuEntry::InitToolBarButton(Commands.ShowCameraDirector));
-		ModesSection.AddEntry(FToolMenuEntry::InitToolBarButton(Commands.ShowCameraRigs));
 		ModesSection.AddEntry(FToolMenuEntry::InitToolBarButton(Commands.ShowSharedTransitions));
 	}
 }
@@ -218,16 +222,12 @@ void FCameraAssetEditorToolkit::PostInitAssetEditor()
 		AddEditorMode(CameraDirectorEditor.ToSharedRef());
 	}
 
-	const FName CameraRigsModeName = FCameraRigsAssetEditorMode::ModeName;
-	AddEditorMode(MakeShared<FCameraRigsAssetEditorMode>(CameraAsset));
-
 	const FName SharedTransitionsModeName = FCameraSharedTransitionsAssetEditorMode::ModeName;
 	AddEditorMode(MakeShared<FCameraSharedTransitionsAssetEditorMode>(CameraAsset));
 
 	const FCameraAssetEditorCommands& Commands = FCameraAssetEditorCommands::Get();
 	TMap<FName, TSharedPtr<FUICommandInfo>> ModeCommands;
 	ModeCommands.Add(CameraDirectorModeName, Commands.ShowCameraDirector);
-	ModeCommands.Add(CameraRigsModeName, Commands.ShowCameraRigs);
 	ModeCommands.Add(SharedTransitionsModeName, Commands.ShowSharedTransitions);
 	for (auto& Pair : ModeCommands)
 	{
@@ -252,8 +252,10 @@ void FCameraAssetEditorToolkit::PostInitAssetEditor()
 	LiveEditManager = GameplayCamerasModule.GetLiveEditManager();
 
 	const FName InitialModeName = !Settings->LastCameraAssetToolkitModeName.IsNone() ? 
-		Settings->LastCameraAssetToolkitModeName : CameraRigsModeName;
+		Settings->LastCameraAssetToolkitModeName : CameraDirectorModeName;
 	SetEditorMode(InitialModeName);
+
+	UpgradeLegacyCameraAssets();
 }
 
 void FCameraAssetEditorToolkit::OnEditorToolkitModeActivated()
@@ -271,23 +273,30 @@ void FCameraAssetEditorToolkit::OnBuild()
 		return;
 	}
 
+	FCameraDirectorRigUsageInfo UsageInfo;
+
 	FCameraBuildLog BuildLog;
 	FCameraAssetBuilder Builder(BuildLog);
 	Builder.BuildCamera(
 			CameraAsset,
 			FCameraAssetBuilder::FCustomBuildStep::CreateLambda(
-				[](UCameraAsset* InCameraAsset, FCameraBuildLog& BuildLog)
+				[&UsageInfo](UCameraAsset* InCameraAsset, FCameraBuildLog& BuildLog)
 				{
 					IGameplayCamerasEditorModule& GameplayCamerasEditorModule = IGameplayCamerasEditorModule::Get();
 					for (const FOnBuildCameraAsset& Builder : GameplayCamerasEditorModule.GetCameraAssetBuilders())
 					{
 						Builder.ExecuteIfBound(InCameraAsset, BuildLog);
 					}
-					for (UCameraRigAsset* CameraRig : InCameraAsset->GetCameraRigs())
+
+					if (UCameraDirector* CameraDirector = InCameraAsset->GetCameraDirector())
 					{
-						for (const FOnBuildCameraRigAsset& Builder : GameplayCamerasEditorModule.GetCameraRigAssetBuilders())
+						CameraDirector->GatherRigUsageInfo(UsageInfo);
+						for (UCameraRigAsset* CameraRig : UsageInfo.CameraRigs)
 						{
-							Builder.ExecuteIfBound(CameraRig, BuildLog);
+							for (const FOnBuildCameraRigAsset& Builder : GameplayCamerasEditorModule.GetCameraRigAssetBuilders())
+							{
+								Builder.ExecuteIfBound(CameraRig, BuildLog);
+							}
 						}
 					}
 				}));
@@ -299,7 +308,7 @@ void FCameraAssetEditorToolkit::OnBuild()
 		TabManager->TryInvokeTab(MessagesTabId);
 	}
 
-	for (UCameraRigAsset* CameraRigAsset : CameraAsset->GetCameraRigs())
+	for (UCameraRigAsset* CameraRigAsset : UsageInfo.CameraRigs)
 	{
 		FCameraRigPackages BuiltPackages;
 		CameraRigAsset->GatherPackages(BuiltPackages);
@@ -319,10 +328,6 @@ void FCameraAssetEditorToolkit::OnFindInCamera()
 
 void FCameraAssetEditorToolkit::OnGetRootObjectsToSearch(TArray<FFindInObjectTreeGraphSource>& OutSources)
 {
-	TSharedPtr<FCameraRigsAssetEditorMode> CameraRigsMode = GetTypedEditorMode<FCameraRigsAssetEditorMode>(
-			FCameraRigsAssetEditorMode::ModeName);
-	CameraRigsMode->OnGetRootObjectsToSearch(OutSources);
-
 	TSharedPtr<FCameraSharedTransitionsAssetEditorMode> SharedTransitionsMode = GetTypedEditorMode<FCameraSharedTransitionsAssetEditorMode>(
 			FCameraSharedTransitionsAssetEditorMode::ModeName);
 	SharedTransitionsMode->OnGetRootObjectsToSearch(OutSources);
@@ -369,15 +374,6 @@ void FCameraAssetEditorToolkit::OnJumpToObject(UObject* Object, FName PropertyNa
 		return;
 	}
 	
-	if (bFindInCameraRig)
-	{
-		TSharedPtr<FCameraRigsAssetEditorMode> CameraRigsMode = GetTypedEditorMode<FCameraRigsAssetEditorMode>(
-				FCameraRigsAssetEditorMode::ModeName);
-		SetEditorMode(FCameraRigsAssetEditorMode::ModeName);
-		CameraRigsMode->JumpToObject(Object, PropertyName);
-		return;
-	}
-
 	if (bFindInSharedTranstions)
 	{
 		TSharedPtr<FCameraSharedTransitionsAssetEditorMode> SharedTransitionsMode = GetTypedEditorMode<FCameraSharedTransitionsAssetEditorMode>(
@@ -388,7 +384,6 @@ void FCameraAssetEditorToolkit::OnJumpToObject(UObject* Object, FName PropertyNa
 	}
 }
 
-
 FText FCameraAssetEditorToolkit::GetBaseToolkitName() const
 {
 	return LOCTEXT("AppLabel", "Camera Asset");
@@ -410,6 +405,143 @@ FLinearColor FCameraAssetEditorToolkit::GetWorldCentricTabColorScale() const
 	return FLinearColor(0.7, 0.0f, 0.0f, 0.5f);
 }
 
+void FCameraAssetEditorToolkit::UpgradeLegacyCameraAssets()
+{
+	if (!CameraAsset)
+	{
+		return;
+	}
+
+	UPackage* CameraAssetPackage = CameraAsset->GetOutermost();
+	if (!CameraAssetPackage || CameraAssetPackage == GetTransientPackage())
+	{
+		return;
+	}
+
+	// Gather all camera rigs found inside the camera package. Camera rigs used to be "owned"
+	// by the camera asset, but now we want them all to be shared assets.
+	TSet<UCameraRigAsset*> KnownCameraRigs;
+	{
+		TArray<UObject*> ObjectsInPackage;
+		GetObjectsWithPackage(CameraAssetPackage, ObjectsInPackage);
+		for (UObject* Object : ObjectsInPackage)
+		{
+			if (UCameraRigAsset* CameraRig = Cast<UCameraRigAsset>(Object))
+			{
+				KnownCameraRigs.Add(CameraRig);
+			}
+
+			// Also clean up any redirectors. They used to be created for a short time when
+			// renaming owned camera rigs wasn't doing the right thing.
+			if (UObjectRedirector* Redirector = Cast<UObjectRedirector>(Object))
+			{
+				if (KnownCameraRigs.Contains(Cast<UCameraRigAsset>(Redirector->DestinationObject)))
+				{
+					CameraAsset->Modify();
+
+					Redirector->ClearFlags(RF_Public | RF_Standalone);
+					Redirector->DestinationObject = nullptr;
+				}
+			}
+		}
+	}
+
+	// No owned camera rigs? We are done.
+	if (KnownCameraRigs.IsEmpty())
+	{
+		return;
+	}
+
+	// Start working!
+	FScopedSlowTask SlowTask(KnownCameraRigs.Num() + 1,
+			LOCTEXT("UpgradeLegacyCameraAssets", "Upgrading legacy camera asset"));
+	SlowTask.MakeDialog(true);
+
+	IAssetRegistry& AssetRegistry = FAssetRegistryModule::GetRegistry();
+	IAssetTools& AssetTools = IAssetTools::Get();
+
+	TArray<UPackage*> PackagesToSave;
+	PackagesToSave.Add(CameraAssetPackage);
+
+	// Look for packages that reference one of our camera rigs. We don't need to patch up their references
+	// since they will still point to the same object, but we will need to re-save those packages because
+	// the serialized soft-object-path will have changed to the new standalone camera rig asset package.
+	{
+		SlowTask.EnterProgressFrame();
+
+		TArray<FName> OnDiskReferencers;
+		AssetRegistry.GetReferencers(CameraAssetPackage->GetFName(), OnDiskReferencers);
+		TArray<UObject*> KnownCameraRigsArray(KnownCameraRigs.Array());
+
+		for (FName OnDiskReferencer : OnDiskReferencers)
+		{
+			UPackage* ReferencerPackage = FindPackage(nullptr, *OnDiskReferencer.ToString());
+			if (!ensure(ReferencerPackage))
+			{
+				continue;
+			}
+
+			UObject* ReferencingAsset = ReferencerPackage->FindAssetInPackage();
+			if (!ensure(ReferencingAsset))
+			{
+				continue;
+			}
+
+			FObjectReferenceFinder ReferenceFinder(ReferencingAsset, KnownCameraRigsArray);
+			ReferenceFinder.CollectReferences();
+
+			if (ReferenceFinder.HasAnyObjectReference())
+			{
+				ReferencingAsset->Modify();
+				ReferencerPackage->MarkPackageDirty();
+				PackagesToSave.Add(ReferencerPackage);
+			}
+		}
+	}
+
+	// Now create individual assets for each camera rig.
+	const FString CameraRigsBaseName = CameraAsset->GetName();
+	const FString CameraRigsBasePath = FPaths::GetPath(CameraAssetPackage->GetPathName());
+	const EObjectFlags CameraRigFlags = RF_Public | RF_Standalone | RF_Transactional;
+
+	for (UCameraRigAsset* CameraRig : KnownCameraRigs)
+	{
+		SlowTask.EnterProgressFrame();
+
+		// Name the new camera rigs like this: "<CameraAsset>_<CameraRig>"
+		const FString CameraRigPackageName = FString::Printf(TEXT("%s_%s"), *CameraRigsBaseName, *CameraRig->GetDisplayName());
+
+		FString CameraRigPackagePath;
+		FString CameraRigAssetName;
+		AssetTools.CreateUniqueAssetName(
+				CameraRigsBasePath / CameraRigPackageName, FString(), CameraRigPackagePath, CameraRigAssetName);
+
+		// Create the new package, and move the camera rig inside it as its main asset.
+		UPackage* CameraRigPackage = CreatePackage(*CameraRigPackagePath);
+		CameraRig->Rename(*CameraRigAssetName, CameraRigPackage, REN_DontCreateRedirectors);
+		CameraRig->SetFlags(CameraRigFlags);
+
+		// Notify the asset registry that a new asset was created.
+		AssetRegistry.AssetCreated(CameraRig);
+
+		CameraRigPackage->MarkPackageDirty();
+
+		PackagesToSave.Add(CameraRigPackage);
+	}
+
+	// Prompt the user to save all assets.
+	FEditorFileUtils::FPromptForCheckoutAndSaveParams SaveParams;
+	SaveParams.bCheckDirty = false;
+	SaveParams.bPromptToSave = true;
+	SaveParams.Title = LOCTEXT("SaveUpgradedAsset_Title", "Save upgraded packages");
+	SaveParams.Message = LOCTEXT(
+			"SaveUpgradedAsset_Message", 
+			"This camera asset had legacy private camera rigs. "
+			"They have been re-created as standalone assets, and referencing packages have been fixed up. "
+			"Please save all new and modified packages.");
+	FEditorFileUtils::PromptForCheckoutAndSave(PackagesToSave, SaveParams);
+}
+
 }  // namespace UE::Cameras
 
 #undef LOCTEXT_NAMESPACE
diff --git a/Engine/Plugins/Cameras/GameplayCameras/Source/GameplayCamerasEditor/Private/Toolkits/CameraRigsAssetEditorMode.cpp b/Engine/Plugins/Cameras/GameplayCameras/Source/GameplayCamerasEditor/Private/Toolkits/CameraRigsAssetEditorMode.cpp
deleted file mode 100644
diff --git a/Engine/Plugins/Cameras/GameplayCameras/Source/GameplayCamerasEditor/Private/Toolkits/SingleCameraDirectorAssetEditorMode.cpp b/Engine/Plugins/Cameras/GameplayCameras/Source/GameplayCamerasEditor/Private/Toolkits/SingleCameraDirectorAssetEditorMode.cpp
--- a/Engine/Plugins/Cameras/GameplayCameras/Source/GameplayCamerasEditor/Private/Toolkits/SingleCameraDirectorAssetEditorMode.cpp
+++ b/Engine/Plugins/Cameras/GameplayCameras/Source/GameplayCamerasEditor/Private/Toolkits/SingleCameraDirectorAssetEditorMode.cpp
@@ -26,20 +26,5 @@ FSingleCameraDirectorAssetEditorMode::FSingleCameraDirectorAssetEditorMode(UCame
 	}
 }
 
-void FSingleCameraDirectorAssetEditorMode::OnCameraRigsChanged(UCameraAsset* InCameraAsset, const TCameraArrayChangedEvent<UCameraRigAsset*>& Event)
-{
-	USingleCameraDirector* CameraDirector = Cast<USingleCameraDirector>(InCameraAsset->GetCameraDirector());
-	if (!ensure(CameraDirector))
-	{
-		return;
-	}
-
-	if (!InCameraAsset->GetCameraRigs().Contains(CameraDirector->CameraRig))
-	{
-		CameraDirector->Modify();
-		CameraDirector->CameraRig = nullptr;
-	}
-}
-
 }  // namespace UE::Cameras
 
diff --git a/Engine/Plugins/Cameras/GameplayCameras/Source/GameplayCamerasEditor/Private/Widgets/SDeleteCameraObjectDialog.cpp b/Engine/Plugins/Cameras/GameplayCameras/Source/GameplayCamerasEditor/Private/Widgets/SDeleteCameraObjectDialog.cpp
--- a/Engine/Plugins/Cameras/GameplayCameras/Source/GameplayCamerasEditor/Private/Widgets/SDeleteCameraObjectDialog.cpp
+++ b/Engine/Plugins/Cameras/GameplayCameras/Source/GameplayCamerasEditor/Private/Widgets/SDeleteCameraObjectDialog.cpp
@@ -6,9 +6,9 @@
 #include "ContentBrowserModule.h"
 #include "Editor.h"
 #include "FileHelpers.h"
+#include "Helpers/ObjectReferenceFinder.h"
 #include "IContentBrowserSingleton.h"
 #include "ObjectTools.h"
-#include "Serialization/FindReferencersArchive.h"
 #include "Subsystems/AssetEditorSubsystem.h"
 #include "UObject/ObjectRedirector.h"
 #include "Widgets/Input/SButton.h"
@@ -422,11 +422,9 @@ void SDeleteCameraObjectDialog::ScanNextReferencingPackage()
 		return;
 	}
 
-	FFindReferencersArchive FindRefsArchive(PackageAsset, ObjectsToDelete);
-
-	TMap<UObject*, int32> NumReferencesMap;
-	TMultiMap<UObject*, FProperty*> ReferencingPropertiesMMap;
-	if (FindRefsArchive.GetReferenceCounts(NumReferencesMap, ReferencingPropertiesMMap) > 0)
+	FObjectReferenceFinder ReferenceFinder(PackageAsset, ObjectsToDelete);
+	ReferenceFinder.CollectReferences();
+	if (ReferenceFinder.HasAnyObjectReference())
 	{
 		// That package does reference one of the objects we want to delete. Keep it.
 		ReferencingPackages.Add(ReferencingPackageName);
diff --git a/Engine/Plugins/Cameras/GameplayCameras/Source/GameplayCamerasEditor/Public/Customizations/CameraProxyTableDetailsCustomization.h b/Engine/Plugins/Cameras/GameplayCameras/Source/GameplayCamerasEditor/Public/Customizations/CameraProxyTableDetailsCustomization.h
--- a/Engine/Plugins/Cameras/GameplayCameras/Source/GameplayCamerasEditor/Public/Customizations/CameraProxyTableDetailsCustomization.h
+++ b/Engine/Plugins/Cameras/GameplayCameras/Source/GameplayCamerasEditor/Public/Customizations/CameraProxyTableDetailsCustomization.h
@@ -37,7 +37,7 @@ private:
 
 	TSharedRef<SWidget> OnBuildCameraRigPicker();
 	FText OnGetComboButtonText() const;
-	void OnCameraRigSelected(UCameraRigAsset* CameraRig);
+	void OnCameraRigSelected(const FAssetData& InSelectedAsset);
 
 private:
 
diff --git a/Engine/Plugins/Cameras/GameplayCameras/Source/GameplayCamerasEditor/Public/Customizations/CameraRigPtrDetailsCustomization.h b/Engine/Plugins/Cameras/GameplayCameras/Source/GameplayCamerasEditor/Public/Customizations/CameraRigPtrDetailsCustomization.h
deleted file mode 100644
diff --git a/Engine/Plugins/Cameras/GameplayCameras/Source/GameplayCamerasEditor/Public/Editors/GameplayCamerasGraphPanelPinFactory.h b/Engine/Plugins/Cameras/GameplayCameras/Source/GameplayCamerasEditor/Public/Editors/GameplayCamerasGraphPanelPinFactory.h
--- a/Engine/Plugins/Cameras/GameplayCameras/Source/GameplayCamerasEditor/Public/Editors/GameplayCamerasGraphPanelPinFactory.h
+++ b/Engine/Plugins/Cameras/GameplayCameras/Source/GameplayCamerasEditor/Public/Editors/GameplayCamerasGraphPanelPinFactory.h
@@ -22,9 +22,6 @@ public:
 
 private:
 
-	TSharedPtr<SGraphPin> CreateFunctionParameterPin(UEdGraphPin* Pin, UK2Node_CallFunction* CallFunctionNode) const;
-
-	TSharedPtr<SGraphPin> CreateCameraRigPickerPin(UEdGraphPin* Pin) const;
 	TSharedPtr<SGraphPin> CreateCameraVariablePickerPin(UEdGraphPin* Pin) const;
 };
 
diff --git a/Engine/Plugins/Cameras/GameplayCameras/Source/GameplayCamerasEditor/Public/Editors/SBlueprintCameraDirectorRigNameGraphPin.h b/Engine/Plugins/Cameras/GameplayCameras/Source/GameplayCamerasEditor/Public/Editors/SBlueprintCameraDirectorRigNameGraphPin.h
deleted file mode 100644
diff --git a/Engine/Plugins/Cameras/GameplayCameras/Source/GameplayCamerasEditor/Public/Editors/SCameraRigList.h b/Engine/Plugins/Cameras/GameplayCameras/Source/GameplayCamerasEditor/Public/Editors/SCameraRigList.h
deleted file mode 100644
diff --git a/Engine/Plugins/Cameras/GameplayCameras/Source/GameplayCamerasEditor/Public/Editors/SCameraRigNameGraphPin.h b/Engine/Plugins/Cameras/GameplayCameras/Source/GameplayCamerasEditor/Public/Editors/SCameraRigNameGraphPin.h
deleted file mode 100644
diff --git a/Engine/Plugins/Cameras/GameplayCameras/Source/GameplayCamerasEditor/Public/Editors/SCameraRigPickerButton.h b/Engine/Plugins/Cameras/GameplayCameras/Source/GameplayCamerasEditor/Public/Editors/SCameraRigPickerButton.h
deleted file mode 100644
diff --git a/Engine/Plugins/Cameras/GameplayCameras/Source/GameplayCamerasEditor/Public/Helpers/CameraAssetReferenceGatherer.h b/Engine/Plugins/Cameras/GameplayCameras/Source/GameplayCamerasEditor/Public/Helpers/CameraAssetReferenceGatherer.h
deleted file mode 100644
diff --git a/Engine/Plugins/Cameras/GameplayCameras/Source/GameplayCamerasEditor/Public/IGameplayCamerasEditorModule.h b/Engine/Plugins/Cameras/GameplayCameras/Source/GameplayCamerasEditor/Public/IGameplayCamerasEditorModule.h
--- a/Engine/Plugins/Cameras/GameplayCameras/Source/GameplayCamerasEditor/Public/IGameplayCamerasEditorModule.h
+++ b/Engine/Plugins/Cameras/GameplayCameras/Source/GameplayCamerasEditor/Public/IGameplayCamerasEditorModule.h
@@ -23,7 +23,6 @@ namespace UE::Cameras
 
 class FCameraBuildLog;
 class FCameraDirectorAssetEditorMode;
-struct FCameraRigPickerConfig;
 struct FCameraVariablePickerConfig;
 
 struct FCameraDebugCategoryInfo
@@ -57,7 +56,6 @@ public:
 
 public:
 
-	using FCameraRigPickerConfig = UE::Cameras::FCameraRigPickerConfig;
 	using FCameraVariablePickerConfig = UE::Cameras::FCameraVariablePickerConfig;
 
 	/** Creates an editor for the given camera asset */
@@ -72,9 +70,6 @@ public:
 	/** Creates an editor for the given variable collection */
 	virtual UCameraVariableCollectionEditor* CreateCameraVariableCollectionEditor(const EToolkitMode::Type Mode, const TSharedPtr<IToolkitHost>& InitToolkitHost, UCameraVariableCollection* VariableCollection) = 0;
 
-	/** Creates a new camera rig asset picker widget */
-	virtual TSharedRef<SWidget> CreateCameraRigPicker(const FCameraRigPickerConfig& InPickerConfig) = 0;
-
 	/** Creates a new camera varable asset picker widget */
 	virtual TSharedRef<SWidget> CreateCameraVariablePicker(const FCameraVariablePickerConfig& InPickerConfig) = 0;
 
diff --git a/Engine/Plugins/Cameras/GameplayCameras/Source/GameplayCamerasEditor/Public/Toolkits/BlueprintCameraDirectorAssetEditorMode.h b/Engine/Plugins/Cameras/GameplayCameras/Source/GameplayCamerasEditor/Public/Toolkits/BlueprintCameraDirectorAssetEditorMode.h
--- a/Engine/Plugins/Cameras/GameplayCameras/Source/GameplayCamerasEditor/Public/Toolkits/BlueprintCameraDirectorAssetEditorMode.h
+++ b/Engine/Plugins/Cameras/GameplayCameras/Source/GameplayCamerasEditor/Public/Toolkits/BlueprintCameraDirectorAssetEditorMode.h
@@ -18,10 +18,6 @@ public:
 
 	FBlueprintCameraDirectorAssetEditorMode(UCameraAsset* InCameraAsset);
 
-protected:
-
-	virtual void OnCameraRigsChanged(UCameraAsset* InCameraAsset, const TCameraArrayChangedEvent<UCameraRigAsset*>& Event) override;
-
 private:
 
 	TCameraEventHandler<ICameraAssetEventHandler> EventHandler;
diff --git a/Engine/Plugins/Cameras/GameplayCameras/Source/GameplayCamerasEditor/Public/Toolkits/CameraAssetEditorToolkit.h b/Engine/Plugins/Cameras/GameplayCameras/Source/GameplayCamerasEditor/Public/Toolkits/CameraAssetEditorToolkit.h
--- a/Engine/Plugins/Cameras/GameplayCameras/Source/GameplayCamerasEditor/Public/Toolkits/CameraAssetEditorToolkit.h
+++ b/Engine/Plugins/Cameras/GameplayCameras/Source/GameplayCamerasEditor/Public/Toolkits/CameraAssetEditorToolkit.h
@@ -72,6 +72,8 @@ private:
 	void OnJumpToObject(UObject* Object);
 	void OnJumpToObject(UObject* Object, FName PropertyName);
 
+	void UpgradeLegacyCameraAssets();
+
 private:
 
 	static const FName SearchTabId;
diff --git a/Engine/Plugins/Cameras/GameplayCameras/Source/GameplayCamerasEditor/Public/Toolkits/CameraRigsAssetEditorMode.h b/Engine/Plugins/Cameras/GameplayCameras/Source/GameplayCamerasEditor/Public/Toolkits/CameraRigsAssetEditorMode.h
deleted file mode 100644
diff --git a/Engine/Plugins/Cameras/GameplayCameras/Source/GameplayCamerasEditor/Public/Toolkits/SingleCameraDirectorAssetEditorMode.h b/Engine/Plugins/Cameras/GameplayCameras/Source/GameplayCamerasEditor/Public/Toolkits/SingleCameraDirectorAssetEditorMode.h
--- a/Engine/Plugins/Cameras/GameplayCameras/Source/GameplayCamerasEditor/Public/Toolkits/SingleCameraDirectorAssetEditorMode.h
+++ b/Engine/Plugins/Cameras/GameplayCameras/Source/GameplayCamerasEditor/Public/Toolkits/SingleCameraDirectorAssetEditorMode.h
@@ -18,10 +18,6 @@ public:
 
 	FSingleCameraDirectorAssetEditorMode(UCameraAsset* InCameraAsset);
 
-protected:
-
-	virtual void OnCameraRigsChanged(UCameraAsset* InCameraAsset, const TCameraArrayChangedEvent<UCameraRigAsset*>& Event) override;
-
 private:
 
 	TCameraEventHandler<ICameraAssetEventHandler> EventHandler;
diff --git a/Engine/Plugins/Cameras/GameplayCameras/Source/GameplayCamerasUncookedOnly/Private/BlueprintGraph/K2Node_SetCameraRigParameters.h b/Engine/Plugins/Cameras/GameplayCameras/Source/GameplayCamerasUncookedOnly/Private/BlueprintGraph/K2Node_SetCameraRigParameters.h
--- a/Engine/Plugins/Cameras/GameplayCameras/Source/GameplayCamerasUncookedOnly/Private/BlueprintGraph/K2Node_SetCameraRigParameters.h
+++ b/Engine/Plugins/Cameras/GameplayCameras/Source/GameplayCamerasUncookedOnly/Private/BlueprintGraph/K2Node_SetCameraRigParameters.h
@@ -17,7 +17,7 @@ class UCameraRigAsset;
  * found on it. On compile, this node gets expanded into the appropriate number of individual setter
  * function calls for each parameter (see UCameraRigParameterInterop).
  */
-UCLASS(MinimalAPI, meta=(UseCameraRigPickerForPins="CameraRig"))
+UCLASS(MinimalAPI)
 class UK2Node_SetCameraRigParameters : public UK2Node
 {
 	GENERATED_BODY()
