Cameras: add asset type links as menu overlays on the top-right corner of the asset editors

diff --git a/Engine/Plugins/Cameras/GameplayCameras/Source/GameplayCamerasEditor/Private/Helpers/AssetTypeMenuOverlayHelper.cpp b/Engine/Plugins/Cameras/GameplayCameras/Source/GameplayCamerasEditor/Private/Helpers/AssetTypeMenuOverlayHelper.cpp
new file mode 100644
--- /dev/null
+++ b/Engine/Plugins/Cameras/GameplayCameras/Source/GameplayCamerasEditor/Private/Helpers/AssetTypeMenuOverlayHelper.cpp
@@ -0,0 +1,40 @@
+// Copyright Epic Games, Inc. All Rights Reserved.
+
+#include "Helpers/AssetTypeMenuOverlayHelper.h"
+
+#include "EditorClassUtils.h"
+#include "Widgets/SBoxPanel.h"
+#include "Widgets/Text/STextBlock.h"
+
+#define LOCTEXT_NAMESPACE "AssetTypeMenuOverlayHelper"
+
+namespace UE::Cameras
+{
+
+TSharedRef<SWidget> FAssetTypeMenuOverlayHelper::CreateMenuOverlay(UClass* InAssetType)
+{
+	TSharedRef<SHorizontalBox> MenuOverlayBox = SNew(SHorizontalBox)
+		+SHorizontalBox::Slot()
+		.AutoWidth()
+		.VAlign(VAlign_Center)
+		[
+			SNew(STextBlock)
+			.ColorAndOpacity(FSlateColor::UseSubduedForeground())
+			.ShadowOffset(FVector2D::UnitVector)
+			.Text(LOCTEXT("CameraAssetType", "Asset Type: "))
+		]
+		+SHorizontalBox::Slot()
+		.AutoWidth()
+		.VAlign(VAlign_Center)
+		.Padding(0.0f, 0.0f, 8.0f, 0.0f)
+		[
+			FEditorClassUtils::GetSourceLink(InAssetType)
+		];
+
+	return MenuOverlayBox;
+}
+
+}  // namespace UE::Cameras
+
+#undef LOCTEXT_NAMESPACE
+
diff --git a/Engine/Plugins/Cameras/GameplayCameras/Source/GameplayCamerasEditor/Private/Helpers/AssetTypeMenuOverlayHelper.h b/Engine/Plugins/Cameras/GameplayCameras/Source/GameplayCamerasEditor/Private/Helpers/AssetTypeMenuOverlayHelper.h
new file mode 100644
--- /dev/null
+++ b/Engine/Plugins/Cameras/GameplayCameras/Source/GameplayCamerasEditor/Private/Helpers/AssetTypeMenuOverlayHelper.h
@@ -0,0 +1,22 @@
+// Copyright Epic Games, Inc. All Rights Reserved.
+
+#pragma once
+
+#include "CoreTypes.h"
+#include "Templates/SharedPointer.h"
+
+class SWidget;
+class UClass;
+
+namespace UE::Cameras
+{
+
+class FAssetTypeMenuOverlayHelper
+{
+public:
+
+	static TSharedRef<SWidget> CreateMenuOverlay(UClass* InAssetType);
+};
+
+}  // namespace UE::Cameras
+
diff --git a/Engine/Plugins/Cameras/GameplayCameras/Source/GameplayCamerasEditor/Private/Toolkits/CameraAssetEditorToolkit.cpp b/Engine/Plugins/Cameras/GameplayCameras/Source/GameplayCamerasEditor/Private/Toolkits/CameraAssetEditorToolkit.cpp
--- a/Engine/Plugins/Cameras/GameplayCameras/Source/GameplayCamerasEditor/Private/Toolkits/CameraAssetEditorToolkit.cpp
+++ b/Engine/Plugins/Cameras/GameplayCameras/Source/GameplayCamerasEditor/Private/Toolkits/CameraAssetEditorToolkit.cpp
@@ -16,6 +16,7 @@
 #include "Framework/Docking/LayoutExtender.h"
 #include "Framework/Docking/TabManager.h"
 #include "GameplayCamerasEditorSettings.h"
+#include "Helpers/AssetTypeMenuOverlayHelper.h"
 #include "Helpers/ObjectReferenceFinder.h"
 #include "IAssetTools.h"
 #include "IGameplayCamerasEditorModule.h"
@@ -258,6 +259,11 @@ void FCameraAssetEditorToolkit::PostInitAssetEditor()
 	UpgradeLegacyCameraAssets();
 }
 
+void FCameraAssetEditorToolkit::PostRegenerateMenusAndToolbars()
+{
+	SetMenuOverlay(FAssetTypeMenuOverlayHelper::CreateMenuOverlay(UCameraAsset::StaticClass()));
+}
+
 void FCameraAssetEditorToolkit::OnEditorToolkitModeActivated()
 {
 	Settings->LastCameraAssetToolkitModeName = GetCurrentEditorModeName();
diff --git a/Engine/Plugins/Cameras/GameplayCameras/Source/GameplayCamerasEditor/Private/Toolkits/CameraRigAssetEditorToolkit.cpp b/Engine/Plugins/Cameras/GameplayCameras/Source/GameplayCamerasEditor/Private/Toolkits/CameraRigAssetEditorToolkit.cpp
--- a/Engine/Plugins/Cameras/GameplayCameras/Source/GameplayCamerasEditor/Private/Toolkits/CameraRigAssetEditorToolkit.cpp
+++ b/Engine/Plugins/Cameras/GameplayCameras/Source/GameplayCamerasEditor/Private/Toolkits/CameraRigAssetEditorToolkit.cpp
@@ -12,6 +12,7 @@
 #include "Editors/SFindInObjectTreeGraph.h"
 #include "Framework/Docking/LayoutExtender.h"
 #include "Framework/Docking/TabManager.h"
+#include "Helpers/AssetTypeMenuOverlayHelper.h"
 #include "IGameplayCamerasEditorModule.h"
 #include "IGameplayCamerasLiveEditManager.h"
 #include "IGameplayCamerasModule.h"
@@ -213,6 +214,13 @@ void FCameraRigAssetEditorToolkit::PostInitAssetEditor()
 
 	IGameplayCamerasModule& GameplayCamerasModule = FModuleManager::GetModuleChecked<IGameplayCamerasModule>("GameplayCameras");
 	LiveEditManager = GameplayCamerasModule.GetLiveEditManager();
+
+	RegenerateMenusAndToolbars();
+}
+
+void FCameraRigAssetEditorToolkit::PostRegenerateMenusAndToolbars()
+{
+	SetMenuOverlay(FAssetTypeMenuOverlayHelper::CreateMenuOverlay(UCameraRigAsset::StaticClass()));
 }
 
 void FCameraRigAssetEditorToolkit::OnBuild()
diff --git a/Engine/Plugins/Cameras/GameplayCameras/Source/GameplayCamerasEditor/Private/Toolkits/CameraRigProxyAssetEditorToolkit.cpp b/Engine/Plugins/Cameras/GameplayCameras/Source/GameplayCamerasEditor/Private/Toolkits/CameraRigProxyAssetEditorToolkit.cpp
--- a/Engine/Plugins/Cameras/GameplayCameras/Source/GameplayCamerasEditor/Private/Toolkits/CameraRigProxyAssetEditorToolkit.cpp
+++ b/Engine/Plugins/Cameras/GameplayCameras/Source/GameplayCamerasEditor/Private/Toolkits/CameraRigProxyAssetEditorToolkit.cpp
@@ -5,6 +5,7 @@
 #include "AssetTools/CameraRigProxyAssetEditor.h"
 #include "Core/CameraRigProxyAsset.h"
 #include "Framework/Docking/LayoutExtender.h"
+#include "Helpers/AssetTypeMenuOverlayHelper.h"
 #include "PropertyEditorModule.h"
 
 #define LOCTEXT_NAMESPACE "CameraRigProxyAssetEditorToolkit"
@@ -78,6 +79,12 @@ void FCameraRigProxyAssetEditorToolkit::CreateWidgets()
 
 void FCameraRigProxyAssetEditorToolkit::PostInitAssetEditor()
 {
+	RegenerateMenusAndToolbars();
+}
+
+void FCameraRigProxyAssetEditorToolkit::PostRegenerateMenusAndToolbars()
+{
+	SetMenuOverlay(FAssetTypeMenuOverlayHelper::CreateMenuOverlay(UCameraRigProxyAsset::StaticClass()));
 }
 
 FText FCameraRigProxyAssetEditorToolkit::GetBaseToolkitName() const
diff --git a/Engine/Plugins/Cameras/GameplayCameras/Source/GameplayCamerasEditor/Private/Toolkits/CameraVariableCollectionEditorToolkit.cpp b/Engine/Plugins/Cameras/GameplayCameras/Source/GameplayCamerasEditor/Private/Toolkits/CameraVariableCollectionEditorToolkit.cpp
--- a/Engine/Plugins/Cameras/GameplayCameras/Source/GameplayCamerasEditor/Private/Toolkits/CameraVariableCollectionEditorToolkit.cpp
+++ b/Engine/Plugins/Cameras/GameplayCameras/Source/GameplayCamerasEditor/Private/Toolkits/CameraVariableCollectionEditorToolkit.cpp
@@ -10,6 +10,7 @@
 #include "Editors/SCameraVariableCollectionEditor.h"
 #include "Framework/Docking/LayoutExtender.h"
 #include "Framework/Docking/TabManager.h"
+#include "Helpers/AssetTypeMenuOverlayHelper.h"
 #include "IContentBrowserSingleton.h"
 #include "Modules/ModuleManager.h"
 #include "ObjectTools.h"
@@ -226,6 +227,13 @@ void FCameraVariableCollectionEditorToolkit::PostInitAssetEditor()
 		Commands.DeleteVariable,
 		FExecuteAction::CreateSP(this, &FCameraVariableCollectionEditorToolkit::OnDeleteVariable),
 		FCanExecuteAction::CreateSP(this, &FCameraVariableCollectionEditorToolkit::CanDeleteVariable));
+
+	RegenerateMenusAndToolbars();
+}
+
+void FCameraVariableCollectionEditorToolkit::PostRegenerateMenusAndToolbars()
+{
+	SetMenuOverlay(FAssetTypeMenuOverlayHelper::CreateMenuOverlay(UCameraVariableCollection::StaticClass()));
 }
 
 FText FCameraVariableCollectionEditorToolkit::GetBaseToolkitName() const
diff --git a/Engine/Plugins/Cameras/GameplayCameras/Source/GameplayCamerasEditor/Public/Toolkits/CameraAssetEditorToolkit.h b/Engine/Plugins/Cameras/GameplayCameras/Source/GameplayCamerasEditor/Public/Toolkits/CameraAssetEditorToolkit.h
--- a/Engine/Plugins/Cameras/GameplayCameras/Source/GameplayCamerasEditor/Public/Toolkits/CameraAssetEditorToolkit.h
+++ b/Engine/Plugins/Cameras/GameplayCameras/Source/GameplayCamerasEditor/Public/Toolkits/CameraAssetEditorToolkit.h
@@ -46,6 +46,7 @@ protected:
 	virtual void RegisterToolbar() override;
 	virtual void InitToolMenuContext(FToolMenuContext& MenuContext) override;
 	virtual void PostInitAssetEditor() override;
+	virtual void PostRegenerateMenusAndToolbars() override;
 
 	// IToolkit interface
 	virtual FText GetBaseToolkitName() const override;
diff --git a/Engine/Plugins/Cameras/GameplayCameras/Source/GameplayCamerasEditor/Public/Toolkits/CameraRigAssetEditorToolkit.h b/Engine/Plugins/Cameras/GameplayCameras/Source/GameplayCamerasEditor/Public/Toolkits/CameraRigAssetEditorToolkit.h
--- a/Engine/Plugins/Cameras/GameplayCameras/Source/GameplayCamerasEditor/Public/Toolkits/CameraRigAssetEditorToolkit.h
+++ b/Engine/Plugins/Cameras/GameplayCameras/Source/GameplayCamerasEditor/Public/Toolkits/CameraRigAssetEditorToolkit.h
@@ -38,6 +38,7 @@ protected:
 	virtual void RegisterToolbar() override;
 	virtual void InitToolMenuContext(FToolMenuContext& MenuContext) override;
 	virtual void PostInitAssetEditor() override;
+	virtual void PostRegenerateMenusAndToolbars() override;
 
 	// IToolkit interface
 	virtual FText GetBaseToolkitName() const override;
diff --git a/Engine/Plugins/Cameras/GameplayCameras/Source/GameplayCamerasEditor/Public/Toolkits/CameraRigProxyAssetEditorToolkit.h b/Engine/Plugins/Cameras/GameplayCameras/Source/GameplayCamerasEditor/Public/Toolkits/CameraRigProxyAssetEditorToolkit.h
--- a/Engine/Plugins/Cameras/GameplayCameras/Source/GameplayCamerasEditor/Public/Toolkits/CameraRigProxyAssetEditorToolkit.h
+++ b/Engine/Plugins/Cameras/GameplayCameras/Source/GameplayCamerasEditor/Public/Toolkits/CameraRigProxyAssetEditorToolkit.h
@@ -33,6 +33,7 @@ protected:
 	virtual void UnregisterTabSpawners(const TSharedRef<FTabManager>& InTabManager) override;
 	virtual void CreateWidgets() override;
 	virtual void PostInitAssetEditor() override;
+	virtual void PostRegenerateMenusAndToolbars() override;
 
 	// IToolkit interface
 	virtual FText GetBaseToolkitName() const override;
diff --git a/Engine/Plugins/Cameras/GameplayCameras/Source/GameplayCamerasEditor/Public/Toolkits/CameraVariableCollectionEditorToolkit.h b/Engine/Plugins/Cameras/GameplayCameras/Source/GameplayCamerasEditor/Public/Toolkits/CameraVariableCollectionEditorToolkit.h
--- a/Engine/Plugins/Cameras/GameplayCameras/Source/GameplayCamerasEditor/Public/Toolkits/CameraVariableCollectionEditorToolkit.h
+++ b/Engine/Plugins/Cameras/GameplayCameras/Source/GameplayCamerasEditor/Public/Toolkits/CameraVariableCollectionEditorToolkit.h
@@ -40,6 +40,7 @@ protected:
 	virtual void RegisterToolbar() override;
 	virtual void InitToolMenuContext(FToolMenuContext& MenuContext) override;
 	virtual void PostInitAssetEditor() override;
+	virtual void PostRegenerateMenusAndToolbars() override;
 
 	// IToolkit interface
 	virtual FText GetBaseToolkitName() const override;
