Cameras: change array properties in ObjectTreeGraphs to be manually managed

This change removes the auto-management of array item pins on array properties in ObjectTreeGraphs. These properties now get an "add pin" button similar to some Blueprint nodes, along with context menus for inserting and removing pins.

In order to create Sequence camera nodes with available pins right away, this change also adds a callback in the graph config so that we can hook up code that creates two empty array items in it. However, that further required improvements on the graph config: we can now chain class-specific configs so that what is set for base classes also applies to the config of children classes unless those configs override that base class config.

diff --git a/Engine/Plugins/Cameras/GameplayCameras/Source/GameplayCamerasEditor/Private/Commands/ObjectTreeGraphEditorCommands.cpp b/Engine/Plugins/Cameras/GameplayCameras/Source/GameplayCamerasEditor/Private/Commands/ObjectTreeGraphEditorCommands.cpp
new file mode 100644
--- /dev/null
+++ b/Engine/Plugins/Cameras/GameplayCameras/Source/GameplayCamerasEditor/Private/Commands/ObjectTreeGraphEditorCommands.cpp
@@ -0,0 +1,38 @@
+// Copyright Epic Games, Inc. All Rights Reserved.
+
+#include "Commands/ObjectTreeGraphEditorCommands.h"
+
+#include "Framework/Commands/Commands.h"
+#include "Framework/Commands/InputChord.h"
+#include "Framework/Commands/UICommandInfo.h"
+#include "Styles/GameplayCamerasEditorStyle.h"
+
+#define LOCTEXT_NAMESPACE "ObjectTreeGraphEditorCommands"
+
+namespace UE::Cameras
+{
+
+FObjectTreeGraphEditorCommands::FObjectTreeGraphEditorCommands()
+	: TCommands<FObjectTreeGraphEditorCommands>(
+			"ObjectTreeGraphEditor",
+			LOCTEXT("ObjectTreeGraphEditor", "Object Tree Graph Editor"),
+			NAME_None,
+			FGameplayCamerasEditorStyle::Get()->GetStyleSetName()
+		)
+{
+}
+
+void FObjectTreeGraphEditorCommands::RegisterCommands()
+{
+	UI_COMMAND(InsertArrayItemPinBefore, "Insert Item Pin Before", "Insert array item pin before this pin",
+			EUserInterfaceActionType::Button, FInputChord());
+	UI_COMMAND(InsertArrayItemPinAfter, "Insert Item Pin After", "Insert array item pin after this pin",
+			EUserInterfaceActionType::Button, FInputChord());
+	UI_COMMAND(RemoveArrayItemPin, "Remove Item Pin", "Remove array item pin",
+			EUserInterfaceActionType::Button, FInputChord());
+}
+
+}  // namespace UE::Cameras
+
+#undef LOCTEXT_NAMESPACE
+
diff --git a/Engine/Plugins/Cameras/GameplayCameras/Source/GameplayCamerasEditor/Private/Editors/CameraNodeGraphSchema.cpp b/Engine/Plugins/Cameras/GameplayCameras/Source/GameplayCamerasEditor/Private/Editors/CameraNodeGraphSchema.cpp
--- a/Engine/Plugins/Cameras/GameplayCameras/Source/GameplayCamerasEditor/Private/Editors/CameraNodeGraphSchema.cpp
+++ b/Engine/Plugins/Cameras/GameplayCameras/Source/GameplayCamerasEditor/Private/Editors/CameraNodeGraphSchema.cpp
@@ -15,6 +15,7 @@
 #include "Editors/ObjectTreeGraphNode.h"
 #include "Framework/Notifications/NotificationManager.h"
 #include "GameplayCamerasEditorSettings.h"
+#include "Nodes/Common/ArrayCameraNode.h"
 #include "Nodes/Common/CameraRigCameraNode.h"
 #include "Widgets/Notifications/SNotificationList.h"
 
@@ -53,6 +54,14 @@ FObjectTreeGraphConfig UCameraNodeGraphSchema::BuildGraphConfig() const
 		.SelfPinName(NAME_None)  // No self pin name, we just want the title
 		.CanCreateNew(false)
 		.GraphNodeClass(UCameraRigInterfaceParameterGraphNode::StaticClass());
+	GraphConfig.ObjectClassConfigs.Emplace(UArrayCameraNode::StaticClass())
+		.OnSetupNewObject(FOnSetupNewObject::CreateLambda([](UObject* NewObject)
+				{
+					// Add two new pins by default.
+					UArrayCameraNode* ArrayNode = CastChecked<UArrayCameraNode>(NewObject);
+					ArrayNode->Children.AddDefaulted();
+					ArrayNode->Children.AddDefaulted();
+				}));
 
 	return GraphConfig;
 }
@@ -205,7 +214,7 @@ const FPinConnectionResponse UCameraNodeGraphSchema::CanCreateConnection(const U
 	return Super::CanCreateConnection(A, B);
 }
 
-bool UCameraNodeGraphSchema::OnApplyConnection(UEdGraphPin* A, UEdGraphPin* B, FDelayedPinActions& Actions) const
+bool UCameraNodeGraphSchema::OnApplyConnection(UEdGraphPin* A, UEdGraphPin* B) const
 {
 	// Try to make a connection between a camera node's parameter pin and a camera rig interface parameter.
 	// First, figure out which is which.
@@ -269,7 +278,7 @@ bool UCameraNodeGraphSchema::OnApplyConnection(UEdGraphPin* A, UEdGraphPin* B, F
 	return true;
 }
 
-bool UCameraNodeGraphSchema::OnApplyDisconnection(UEdGraphPin* TargetPin, FDelayedPinActions& Actions, bool bIsReconnecting) const
+bool UCameraNodeGraphSchema::OnApplyDisconnection(UEdGraphPin* TargetPin) const
 {
 	// See if we have a rig parameter connection to break.
 	if (TargetPin->PinType.PinCategory == PC_Self || TargetPin->PinType.PinCategory == PC_CameraParameter)
@@ -302,17 +311,17 @@ bool UCameraNodeGraphSchema::OnApplyDisconnection(UEdGraphPin* TargetPin, FDelay
 	return false;
 }
 
-bool UCameraNodeGraphSchema::OnApplyDisconnection(UEdGraphPin* SourcePin, UEdGraphPin* TargetPin, FDelayedPinActions& Actions) const
+bool UCameraNodeGraphSchema::OnApplyDisconnection(UEdGraphPin* SourcePin, UEdGraphPin* TargetPin) const
 {
 	UObjectTreeGraphNode* RigParameterNode = nullptr;
 
 	if (SourcePin->PinType.PinCategory == PC_Self && TargetPin->PinType.PinCategory == PC_CameraParameter)
 	{
-		return OnApplyDisconnection(SourcePin, Actions, false);
+		return OnApplyDisconnection(SourcePin);
 	}
 	else if (SourcePin->PinType.PinCategory == PC_CameraParameter && TargetPin->PinType.PinCategory == PC_Self)
 	{
-		return OnApplyDisconnection(TargetPin, Actions, false);
+		return OnApplyDisconnection(TargetPin);
 	}
 
 	return false;
diff --git a/Engine/Plugins/Cameras/GameplayCameras/Source/GameplayCamerasEditor/Private/Editors/ObjectTreeGraphConfig.cpp b/Engine/Plugins/Cameras/GameplayCameras/Source/GameplayCamerasEditor/Private/Editors/ObjectTreeGraphConfig.cpp
--- a/Engine/Plugins/Cameras/GameplayCameras/Source/GameplayCamerasEditor/Private/Editors/ObjectTreeGraphConfig.cpp
+++ b/Engine/Plugins/Cameras/GameplayCameras/Source/GameplayCamerasEditor/Private/Editors/ObjectTreeGraphConfig.cpp
@@ -28,6 +28,38 @@ FObjectTreeGraphClassConfig& FObjectTreeGraphClassConfig::OnlyAsRoot()
 	return *this;
 }
 
+FObjectTreeGraphClassConfig FObjectTreeGraphClassConfigs::DefaultConfig;
+
+FObjectTreeGraphClassConfigs::FObjectTreeGraphClassConfigs()
+{
+}
+
+FObjectTreeGraphClassConfigs::FObjectTreeGraphClassConfigs(TArrayView<const FObjectTreeGraphClassConfig*> InClassConfigs)
+	: InnerConfigs(InClassConfigs)
+{
+}
+
+void FObjectTreeGraphClassConfigs::GetStripDisplayNameSuffixes(TArray<FString>& OutSuffixes) const
+{
+	for (const FObjectTreeGraphClassConfig* InnerConfig : InnerConfigs)
+	{
+		OutSuffixes.Append(InnerConfig->StripDisplayNameSuffixes());
+	}
+}
+
+TOptional<EEdGraphPinDirection> FObjectTreeGraphClassConfigs::GetPropertyPinDirectionOverride(const FName& InPropertyName) const
+{
+	for (const FObjectTreeGraphClassConfig* InnerConfig : InnerConfigs)
+	{
+		TOptional<EEdGraphPinDirection> DirectionOverride = InnerConfig->GetPropertyPinDirectionOverride(InPropertyName);
+		if (DirectionOverride.IsSet())
+		{
+			return DirectionOverride;
+		}
+	}
+	return TOptional<EEdGraphPinDirection>();
+}
+
 FObjectTreeGraphConfig::FObjectTreeGraphConfig()
 	: DefaultGraphNodeTitleColor(FLinearColor(0.549f, 0.745f, 0.698f))
 	, DefaultGraphNodeBodyTintColor(FLinearColor::White)
@@ -124,7 +156,7 @@ void FObjectTreeGraphConfig::GetConnectableClasses(TArray<UClass*>& OutClasses,
 				continue;
 			}
 
-			const FObjectTreeGraphClassConfig& ClassConfig = GetObjectClassConfig(*ClassIt);
+			const FObjectTreeGraphClassConfigs ClassConfig = GetObjectClassConfigs(*ClassIt);
 			if (!ClassConfig.CanCreateNew())
 			{
 				continue;
@@ -135,22 +167,22 @@ void FObjectTreeGraphConfig::GetConnectableClasses(TArray<UClass*>& OutClasses,
 	}
 }
 
-const FObjectTreeGraphClassConfig& FObjectTreeGraphConfig::GetObjectClassConfig(const UClass* InObjectClass) const
+FObjectTreeGraphClassConfigs FObjectTreeGraphConfig::GetObjectClassConfigs(const UClass* InObjectClass) const
 {
-	static const FObjectTreeGraphClassConfig DefaultClassConfig;
+	TArray<const FObjectTreeGraphClassConfig*, TInlineAllocator<2>> ClassConfigs;
 	
 	while (InObjectClass)
 	{
 		const FObjectTreeGraphClassConfig* ClassConfig = ObjectClassConfigs.Find(InObjectClass);
 		if (ClassConfig)
 		{
-			return *ClassConfig;
+			ClassConfigs.Add(ClassConfig);
 		}
 
 		InObjectClass = InObjectClass->GetSuperClass();
 	}
 
-	return DefaultClassConfig;
+	return FObjectTreeGraphClassConfigs(ClassConfigs);
 }
 
 FText FObjectTreeGraphConfig::GetDisplayNameText(const UObject* InObject) const
@@ -158,7 +190,7 @@ FText FObjectTreeGraphConfig::GetDisplayNameText(const UObject* InObject) const
 	if (InObject)
 	{
 		FText DisplayNameText;
-		const FObjectTreeGraphClassConfig& ClassConfig = GetObjectClassConfig(InObject->GetClass());
+		const FObjectTreeGraphClassConfigs ClassConfig = GetObjectClassConfigs(InObject->GetClass());
 
 		const IObjectTreeGraphObject* GraphObject = Cast<IObjectTreeGraphObject>(InObject);
 		if (GraphObject && GraphObject->HasSupportFlags(GraphName, EObjectTreeGraphObjectSupportFlags::CustomRename))
@@ -186,13 +218,13 @@ FText FObjectTreeGraphConfig::GetDisplayNameText(const UClass* InClass) const
 {
 	if (InClass)
 	{
-		const FObjectTreeGraphClassConfig& ClassConfig = GetObjectClassConfig(InClass);
+		const FObjectTreeGraphClassConfigs ClassConfig = GetObjectClassConfigs(InClass);
 		return GetDisplayNameText(InClass, ClassConfig);
 	}
 	return FText::GetEmpty();
 }
 
-FText FObjectTreeGraphConfig::GetDisplayNameText(const UClass* InClass, const FObjectTreeGraphClassConfig& InClassConfig) const
+FText FObjectTreeGraphConfig::GetDisplayNameText(const UClass* InClass, const FObjectTreeGraphClassConfigs& InClassConfig) const
 {
 	check(InClass);
 	
@@ -206,12 +238,14 @@ FText FObjectTreeGraphConfig::GetDisplayNameText(const UClass* InClass, const FO
 	return DisplayNameText;
 }
 
-void FObjectTreeGraphConfig::FormatDisplayNameText(const UObject* InObject, const FObjectTreeGraphClassConfig& InClassConfig, FText& InOutDisplayNameText) const
+void FObjectTreeGraphConfig::FormatDisplayNameText(const UObject* InObject, const FObjectTreeGraphClassConfigs& InClassConfig, FText& InOutDisplayNameText) const
 {
-	if (InClassConfig.StripDisplayNameSuffixes().Num() > 0)
+	TArray<FString> StripSuffixes;
+	InClassConfig.GetStripDisplayNameSuffixes(StripSuffixes);
+	if (StripSuffixes.Num() > 0)
 	{
 		FString DisplayName = InOutDisplayNameText.ToString();
-		for (const FString& StripSuffix : InClassConfig.StripDisplayNameSuffixes())
+		for (const FString& StripSuffix : StripSuffixes)
 		{
 			if (DisplayName.RemoveFromEnd(StripSuffix))
 			{
@@ -228,7 +262,7 @@ void FObjectTreeGraphConfig::FormatDisplayNameText(const UObject* InObject, cons
 
 EEdGraphPinDirection FObjectTreeGraphConfig::GetSelfPinDirection(const UClass* InObjectClass) const
 {
-	const FObjectTreeGraphClassConfig& ClassConfig = GetObjectClassConfig(InObjectClass);
+	const FObjectTreeGraphClassConfigs ClassConfig = GetObjectClassConfigs(InObjectClass);
 	TOptional<EEdGraphPinDirection> PinDirectionOverride = ClassConfig.SelfPinDirectionOverride();
 	if (PinDirectionOverride.IsSet())
 	{
@@ -255,7 +289,7 @@ EEdGraphPinDirection FObjectTreeGraphConfig::GetSelfPinDirection(const UClass* I
 
 EEdGraphPinDirection FObjectTreeGraphConfig::GetPropertyPinDirection(const UClass* InObjectClass, const FName& InPropertyName) const
 {
-	const FObjectTreeGraphClassConfig& ClassConfig = GetObjectClassConfig(InObjectClass);
+	const FObjectTreeGraphClassConfigs ClassConfig = GetObjectClassConfigs(InObjectClass);
 	TOptional<EEdGraphPinDirection> PinDirectionOverride = ClassConfig.GetPropertyPinDirectionOverride(InPropertyName);
 	if (PinDirectionOverride.IsSet())
 	{
diff --git a/Engine/Plugins/Cameras/GameplayCameras/Source/GameplayCamerasEditor/Private/Editors/ObjectTreeGraphNode.cpp b/Engine/Plugins/Cameras/GameplayCameras/Source/GameplayCamerasEditor/Private/Editors/ObjectTreeGraphNode.cpp
--- a/Engine/Plugins/Cameras/GameplayCameras/Source/GameplayCamerasEditor/Private/Editors/ObjectTreeGraphNode.cpp
+++ b/Engine/Plugins/Cameras/GameplayCameras/Source/GameplayCamerasEditor/Private/Editors/ObjectTreeGraphNode.cpp
@@ -53,13 +53,13 @@ TSharedPtr<SGraphNode> UObjectTreeGraphNode::CreateVisualWidget()
 FLinearColor UObjectTreeGraphNode::GetNodeTitleColor() const
 {
 	const FNodeContext NodeContext = GetNodeContext();
-	return NodeContext.ObjectClassConfig.NodeTitleColor().Get(NodeContext.GraphConfig.DefaultGraphNodeTitleColor);
+	return NodeContext.ObjectClassConfigs.NodeTitleColor().Get(NodeContext.GraphConfig.DefaultGraphNodeTitleColor);
 }
 
 FLinearColor UObjectTreeGraphNode::GetNodeBodyTintColor() const
 {
 	const FNodeContext NodeContext = GetNodeContext();
-	return NodeContext.ObjectClassConfig.NodeBodyTintColor().Get(NodeContext.GraphConfig.DefaultGraphNodeBodyTintColor);
+	return NodeContext.ObjectClassConfigs.NodeBodyTintColor().Get(NodeContext.GraphConfig.DefaultGraphNodeBodyTintColor);
 }
 
 FText UObjectTreeGraphNode::GetTooltipText() const
@@ -82,15 +82,15 @@ void UObjectTreeGraphNode::AllocateDefaultPins()
 
 	const FNodeContext NodeContext = GetNodeContext();
 	const FObjectTreeGraphConfig& OuterGraphConfig = NodeContext.GraphConfig;
-	const FObjectTreeGraphClassConfig& ObjectClassConfig = NodeContext.ObjectClassConfig;
+	const FObjectTreeGraphClassConfigs ObjectClassConfigs = NodeContext.ObjectClassConfigs;
 
-	if (ObjectClassConfig.HasSelfPin())
+	if (ObjectClassConfigs.HasSelfPin())
 	{
 		FEdGraphPinType SelfPinType;
 		SelfPinType.PinCategory = UObjectTreeGraphSchema::PC_Self;
-		const FName& SelfPinName = ObjectClassConfig.SelfPinName();
+		const FName& SelfPinName = ObjectClassConfigs.SelfPinName();
 		UEdGraphPin* SelfPin = CreatePin(OuterGraphConfig.GetSelfPinDirection(NodeContext.ObjectClass), SelfPinType, SelfPinName);
-		SelfPin->PinFriendlyName = ObjectClassConfig.SelfPinFriendlyName();
+		SelfPin->PinFriendlyName = ObjectClassConfigs.SelfPinFriendlyName();
 	}
 
 	for (TFieldIterator<FProperty> PropertyIt(NodeContext.ObjectClass); PropertyIt; ++PropertyIt)
@@ -109,6 +109,7 @@ void UObjectTreeGraphNode::AllocateDefaultPins()
 				continue;
 			}
 
+			// Make a new pin for connecting this property to another object node.
 			ChildPinType.PinSubCategory = UObjectTreeGraphSchema::PSC_ObjectProperty;
 			UEdGraphPin* PropertyPin = CreatePin(PinDirection, ChildPinType, PropertyName);
 
@@ -122,6 +123,7 @@ void UObjectTreeGraphNode::AllocateDefaultPins()
 				continue;
 			}
 
+			// Make a new invisible pin that will be the parent pin to each array item's pin.
 			FObjectProperty* InnerProperty = CastFieldChecked<FObjectProperty>(ArrayProperty->Inner);
 
 			ChildPinType.PinSubCategory = UObjectTreeGraphSchema::PSC_ArrayProperty;
@@ -132,7 +134,10 @@ void UObjectTreeGraphNode::AllocateDefaultPins()
 			ArrayPin->PinToolTip = InnerProperty->PropertyClass->GetName();
 			ArrayPin->bHidden = true;  // Always hidden, we only ever show the sub-pins.
 
-			CreateNewItemPin(ArrayPin);
+			// Create pins for each array item.
+			FScriptArrayHelper ArrayHelper(ArrayProperty, ArrayProperty->ContainerPtrToValuePtr<void>(Object));
+			const int32 ArrayNum = ArrayHelper.Num();
+			CreateNewItemPins(ArrayPin, ArrayNum);
 		}
 	}
 }
@@ -162,10 +167,35 @@ void UObjectTreeGraphNode::NodeConnectionListChanged()
 void UObjectTreeGraphNode::OnPinRemoved(UEdGraphPin* InRemovedPin)
 {
 	Super::OnPinRemoved(InRemovedPin);
-	RefreshArrayPropertyPinNames();
+
+	if (InRemovedPin && 
+			InRemovedPin->PinType.PinCategory == UObjectTreeGraphSchema::PC_Property &&
+			InRemovedPin->PinType.PinSubCategory == UObjectTreeGraphSchema::PSC_ArrayPropertyItem)
+	{
+		RefreshArrayPropertyPinNames();
+	}
+}
+
+void UObjectTreeGraphNode::GetArrayProperties(TArray<FArrayProperty*>& OutArrayProperties) const
+{
+	const FNodeContext NodeContext = GetNodeContext();
+
+	for (UEdGraphPin* Pin : Pins)
+	{
+		if (Pin->PinType.PinCategory == UObjectTreeGraphSchema::PC_Property &&
+				Pin->PinType.PinSubCategory == UObjectTreeGraphSchema::PSC_ArrayProperty &&
+				Pin->ParentPin == nullptr)
+		{
+			FProperty* Property = NodeContext.ObjectClass->FindPropertyByName(Pin->GetFName());
+			if (FArrayProperty* ArrayProperty = CastField<FArrayProperty>(Property))
+			{
+				OutArrayProperties.Add(ArrayProperty);
+			}
+		}
+	}
 }
 
-void UObjectTreeGraphNode::CreateNewItemPin(FArrayProperty& InArrayProperty)
+void UObjectTreeGraphNode::CreateNewItemPins(FArrayProperty& InArrayProperty, int32 NumExtraPins)
 {
 	UEdGraphPin** ParentArrayPinPtr = Pins.FindByPredicate([&InArrayProperty](UEdGraphPin* Item)
 			{
@@ -173,16 +203,75 @@ void UObjectTreeGraphNode::CreateNewItemPin(FArrayProperty& InArrayProperty)
 			});
 	if (ensure(ParentArrayPinPtr))
 	{
-		CreateNewItemPin(*ParentArrayPinPtr);
+		CreateNewItemPins(*ParentArrayPinPtr, NumExtraPins);
 	}
 }
 
-void UObjectTreeGraphNode::CreateNewItemPin(UEdGraphPin* InParentArrayPin)
+void UObjectTreeGraphNode::CreateNewItemPins(UEdGraphPin* InParentArrayPin, int32 NumExtraPins)
 {
+	if (!ensure(InParentArrayPin && NumExtraPins >= 0))
+	{
+		return;
+	}
+
+	if (NumExtraPins == 0)
+	{
+		return;
+	}
+
 	const FNodeContext NodeContext = GetNodeContext();
 
 	const FName PropertyName = InParentArrayPin->GetFName();
-	const int32 NewIndex = InParentArrayPin->SubPins.Num();
+	const int32 StartIndex = InParentArrayPin->SubPins.Num();
+
+	const int32 ParentPinIndex = Pins.Find(InParentArrayPin);
+	ensure(ParentPinIndex >= 0);
+
+	FEdGraphPinType ChildPinType;
+	ChildPinType.PinCategory = UObjectTreeGraphSchema::PC_Property;
+	ChildPinType.PinSubCategory = UObjectTreeGraphSchema::PSC_ArrayPropertyItem;
+
+	const EEdGraphPinDirection PinDirection = NodeContext.GraphConfig.GetPropertyPinDirection(NodeContext.ObjectClass, PropertyName);
+
+	InParentArrayPin->Modify();
+
+	for (int32 Index = 0; Index < NumExtraPins; ++Index)
+	{
+		const int32 NewIndex = StartIndex + Index;
+
+		FName ChildPinName = PropertyName;
+		ChildPinName.SetNumber(NewIndex);
+		UEdGraphPin* ChildPin = CreatePin(PinDirection, ChildPinType, ChildPinName);
+		ChildPin->PinFriendlyName = FText::Format(
+				LOCTEXT("ArrayPinFriendlyNameFmt", "{0} {1}"), FText::FromName(PropertyName), NewIndex);
+
+		ChildPin->ParentPin = InParentArrayPin;
+		InParentArrayPin->SubPins.Add(ChildPin);
+
+		// Always re-insert the child pin so that all child pins are just after
+		// the parent array pin.
+		const int32 ChildPinIndex = ParentPinIndex + InParentArrayPin->SubPins.Num();
+		Pins.Pop(EAllowShrinking::No);
+		Pins.Insert(ChildPin, ChildPinIndex);
+	}
+}
+
+void UObjectTreeGraphNode::InsertNewItemPin(UEdGraphPin* InParentArrayPin, int32 Index)
+{
+	if (!ensure(InParentArrayPin))
+	{
+		return;
+	}
+	if (!ensure(Index >= 0 && Index < InParentArrayPin->SubPins.Num()))
+	{
+		return;
+	}
+
+	const FNodeContext NodeContext = GetNodeContext();
+
+	const FName PropertyName = InParentArrayPin->GetFName();
+	const int32 ParentPinIndex = Pins.Find(InParentArrayPin);
+	ensure(ParentPinIndex >= 0);
 
 	FEdGraphPinType ChildPinType;
 	ChildPinType.PinCategory = UObjectTreeGraphSchema::PC_Property;
@@ -193,22 +282,23 @@ void UObjectTreeGraphNode::CreateNewItemPin(UEdGraphPin* InParentArrayPin)
 	InParentArrayPin->Modify();
 
 	FName ChildPinName = PropertyName;
-	ChildPinName.SetNumber(NewIndex);
+	ChildPinName.SetNumber(Index);
 	UEdGraphPin* ChildPin = CreatePin(PinDirection, ChildPinType, ChildPinName);
-	ChildPin->PinFriendlyName = FText::Format(LOCTEXT("ArrayPinFriendlyNameFmt", "{0} {1}"), FText::FromName(PropertyName), NewIndex);
+	ChildPin->PinFriendlyName = FText::Format(
+			LOCTEXT("ArrayPinFriendlyNameFmt", "{0} {1}"), FText::FromName(PropertyName), Index);
 
 	ChildPin->ParentPin = InParentArrayPin;
-	InParentArrayPin->SubPins.Add(ChildPin);
+	InParentArrayPin->SubPins.Insert(ChildPin, Index);
 
 	// Always re-insert the child pin so that all child pins are just after
 	// the parent array pin.
-	const int32 ParentPinIndex = Pins.Find(InParentArrayPin);
-	if (ensure(ParentPinIndex >= 0))
-	{
-		const int32 ChildPinIndex = ParentPinIndex + InParentArrayPin->SubPins.Num();
-		Pins.Pop(EAllowShrinking::No);
-		Pins.Insert(ChildPin, ChildPinIndex);
-	}
+	const int32 ChildPinIndex = ParentPinIndex + Index + 1;
+	Pins.Pop(EAllowShrinking::No);
+	Pins.Insert(ChildPin, ChildPinIndex);
+
+	// Rename all subsequent pins so they display the correct index.
+	// NOTE: this will actually rename *all* array property pins, which is a bit heavy handed.
+	RefreshArrayPropertyPinNames();
 }
 
 void UObjectTreeGraphNode::RemoveItemPin(UEdGraphPin* InItemPin)
@@ -317,7 +407,7 @@ UEdGraphPin* UObjectTreeGraphNode::GetPinForProperty(FObjectProperty* InProperty
 	return nullptr;
 }
 
-UEdGraphPin* UObjectTreeGraphNode::GetPinForProperty(FArrayProperty* InProperty, int32 Index) const
+UEdGraphPin* UObjectTreeGraphNode::GetPinForProperty(FArrayProperty* InProperty) const
 {
 	UEdGraphPin* const* FoundItem = Pins.FindByPredicate(
 			[InProperty](UEdGraphPin* Item)
@@ -326,31 +416,19 @@ UEdGraphPin* UObjectTreeGraphNode::GetPinForProperty(FArrayProperty* InProperty,
 					Item->PinType.PinSubCategory == UObjectTreeGraphSchema::PSC_ArrayProperty &&
 					Item->GetFName() == InProperty->GetFName(); 
 			});
-	if (FoundItem && ensure((*FoundItem)->SubPins.IsValidIndex(Index)))
+	if (FoundItem)
 	{
-		return (*FoundItem)->SubPins[Index];
+		return *FoundItem;
 	}
 	return nullptr;
 }
 
-UEdGraphPin* UObjectTreeGraphNode::GetPinForPropertyNewItem(FArrayProperty* InProperty, bool bCreateNew)
+UEdGraphPin* UObjectTreeGraphNode::GetPinForProperty(FArrayProperty* InProperty, int32 Index) const
 {
-	UEdGraphPin* const* FoundItem = Pins.FindByPredicate(
-			[InProperty](UEdGraphPin* Item)
-			{ 
-				return Item->PinType.PinCategory == UObjectTreeGraphSchema::PC_Property &&
-					Item->PinType.PinSubCategory == UObjectTreeGraphSchema::PSC_ArrayProperty &&
-					Item->GetFName() == InProperty->GetFName(); 
-			});
-	if (FoundItem)
+	UEdGraphPin* ArrayPin = GetPinForProperty(InProperty);
+	if (ArrayPin && ensure(ArrayPin->SubPins.IsValidIndex(Index)))
 	{
-		ensure((*FoundItem)->SubPins.Num() > 0 && (*FoundItem)->SubPins.Last()->LinkedTo.IsEmpty());
-		UEdGraphPin* NewItemPin = (*FoundItem)->SubPins.Last();
-		if (bCreateNew)
-		{
-			CreateNewItemPin(*FoundItem);
-		}
-		return NewItemPin;
+		return ArrayPin->SubPins[Index];
 	}
 	return nullptr;
 }
@@ -536,8 +614,8 @@ void UObjectTreeGraphNode::OnRenameNode(const FString& NewName)
 
 bool UObjectTreeGraphNode::CanDuplicateNode() const
 {
-	const FObjectTreeGraphClassConfig& ObjectClassConfig = GetObjectClassConfig();
-	if (!ObjectClassConfig.CanCreateNew())  // If it can't be created, it shouldn't be worked around by copy/pasting
+	const FObjectTreeGraphClassConfigs ObjectClassConfigs = GetObjectClassConfigs();
+	if (!ObjectClassConfigs.CanCreateNew())  // If it can't be created, it shouldn't be worked around by copy/pasting
 	{
 		return false;
 	}
@@ -547,8 +625,8 @@ bool UObjectTreeGraphNode::CanDuplicateNode() const
 
 bool UObjectTreeGraphNode::CanUserDeleteNode() const
 {
-	const FObjectTreeGraphClassConfig& ObjectClassConfig = GetObjectClassConfig();
-	if (!ObjectClassConfig.CanDelete())
+	const FObjectTreeGraphClassConfigs ObjectClassConfigs = GetObjectClassConfigs();
+	if (!ObjectClassConfigs.CanDelete())
 	{
 		return false;
 	}
@@ -598,20 +676,20 @@ UObjectTreeGraphNode::FNodeContext UObjectTreeGraphNode::GetNodeContext() const
 	if (UObject* Object = WeakObject.Get())
 	{
 		UClass* ObjectClass = Object->GetClass();
-		const FObjectTreeGraphClassConfig& ObjectClassConfig = OuterGraphConfig.GetObjectClassConfig(ObjectClass);
+		const FObjectTreeGraphClassConfigs ObjectClassConfigs = OuterGraphConfig.GetObjectClassConfigs(ObjectClass);
 
-		return FNodeContext{ ObjectClass, OuterGraph, OuterGraphConfig, ObjectClassConfig };
+		return FNodeContext{ ObjectClass, OuterGraph, OuterGraphConfig, ObjectClassConfigs };
 	}
 	else
 	{
-		const FObjectTreeGraphClassConfig& ObjectClassConfig = OuterGraphConfig.GetObjectClassConfig(nullptr);
-		return FNodeContext{ nullptr, OuterGraph, OuterGraphConfig, ObjectClassConfig };
+		const FObjectTreeGraphClassConfigs ObjectClassConfigs = OuterGraphConfig.GetObjectClassConfigs(nullptr);
+		return FNodeContext{ nullptr, OuterGraph, OuterGraphConfig, ObjectClassConfigs };
 	}
 }
 
-const FObjectTreeGraphClassConfig& UObjectTreeGraphNode::GetObjectClassConfig() const
+const FObjectTreeGraphClassConfigs UObjectTreeGraphNode::GetObjectClassConfigs() const
 {
-	return GetNodeContext().ObjectClassConfig;
+	return GetNodeContext().ObjectClassConfigs;
 }
 
 #undef LOCTEXT_NAMESPACE
diff --git a/Engine/Plugins/Cameras/GameplayCameras/Source/GameplayCamerasEditor/Private/Editors/ObjectTreeGraphSchema.cpp b/Engine/Plugins/Cameras/GameplayCameras/Source/GameplayCamerasEditor/Private/Editors/ObjectTreeGraphSchema.cpp
--- a/Engine/Plugins/Cameras/GameplayCameras/Source/GameplayCamerasEditor/Private/Editors/ObjectTreeGraphSchema.cpp
+++ b/Engine/Plugins/Cameras/GameplayCameras/Source/GameplayCamerasEditor/Private/Editors/ObjectTreeGraphSchema.cpp
@@ -2,6 +2,7 @@
 
 #include "Editors/ObjectTreeGraphSchema.h"
 
+#include "Commands/ObjectTreeGraphEditorCommands.h"
 #include "Core/ObjectTreeGraphRootObject.h"
 #include "Editors/ObjectTreeConnectionDrawingPolicy.h"
 #include "Editors/ObjectTreeGraph.h"
@@ -11,6 +12,7 @@
 #include "IGameplayCamerasEditorModule.h"
 #include "ScopedTransaction.h"
 #include "Serialization/ArchiveUObject.h"
+#include "ToolMenu.h"
 #include "UObject/UObjectIterator.h"
 #include "UObject/UnrealType.h"
 #include "UnrealExporter.h"
@@ -259,6 +261,8 @@ void UObjectTreeGraphSchema::CreateConnections(UObjectTreeGraphNode* InGraphNode
 	{
 		if (FObjectProperty* ObjectProperty = CastField<FObjectProperty>(ConnectableProperty))
 		{
+			// Object reference property... if the property value is not null, find the node that corresponds
+			// to the referenced object. If we find it, create a graph connection between the two.
 			UEdGraphPin* Pin = InGraphNode->GetPinForProperty(ObjectProperty);
 			if (!ensure(Pin))
 			{
@@ -285,13 +289,15 @@ void UObjectTreeGraphSchema::CreateConnections(UObjectTreeGraphNode* InGraphNode
 		}
 		else if (FArrayProperty* ArrayProperty = CastField<FArrayProperty>(ConnectableProperty))
 		{
+			// Array of object references... add all the pins needed for the array's size, and connect each
+			// of those pins to a matching object node, similarly to above with object references.
 			FObjectProperty* InnerProperty = CastFieldChecked<FObjectProperty>(ArrayProperty->Inner);
 			FScriptArrayHelper ArrayHelper(ArrayProperty, ArrayProperty->ContainerPtrToValuePtr<void>(Object));
 
 			const int32 ArrayNum = ArrayHelper.Num();
 			for (int32 Index = 0; Index < ArrayNum; ++Index)
 			{
-				UEdGraphPin* Pin = InGraphNode->GetPinForPropertyNewItem(ArrayProperty, true);
+				UEdGraphPin* Pin = InGraphNode->GetPinForProperty(ArrayProperty, Index);
 				if (!ensure(Pin))
 				{
 					continue;
@@ -340,9 +346,9 @@ UObjectTreeGraphNode* UObjectTreeGraphSchema::CreateObjectNode(UObjectTreeGraph*
 UObjectTreeGraphNode* UObjectTreeGraphSchema::OnCreateObjectNode(UObjectTreeGraph* InGraph, UObject* InObject) const
 {
 	const FObjectTreeGraphConfig& Config = InGraph->GetConfig();
-	const FObjectTreeGraphClassConfig& ClassConfig = Config.GetObjectClassConfig(InObject->GetClass());
+	const FObjectTreeGraphClassConfigs ClassConfigs = Config.GetObjectClassConfigs(InObject->GetClass());
 
-	TSubclassOf<UObjectTreeGraphNode> GraphNodeClass = ClassConfig.GraphNodeClass();
+	TSubclassOf<UObjectTreeGraphNode> GraphNodeClass = ClassConfigs.GraphNodeClass();
 	if (!GraphNodeClass.Get())
 	{
 		GraphNodeClass = Config.DefaultGraphNodeClass;
@@ -448,8 +454,8 @@ void UObjectTreeGraphSchema::GetGraphContextActions(FGraphContextMenuBuilder& Co
 			continue;
 		}
 
-		const FObjectTreeGraphClassConfig& ClassConfig = GraphConfig.GetObjectClassConfig(*ClassIt);
-		if (!ClassConfig.CanCreateNew())
+		const FObjectTreeGraphClassConfigs ClassConfigs = GraphConfig.GetObjectClassConfigs(*ClassIt);
+		if (!ClassConfigs.CanCreateNew())
 		{
 			continue;
 		}
@@ -476,7 +482,7 @@ void UObjectTreeGraphSchema::GetGraphContextActions(FGraphContextMenuBuilder& Co
 		const FText DisplayName = GraphConfig.GetDisplayNameText(PossibleObjectClass);
 
 		TArray<FString> CategoryNames;
-		const FName CreateCategoryMetaData = GraphConfig.GetObjectClassConfig(PossibleObjectClass).CreateCategoryMetaData();
+		const FName CreateCategoryMetaData = GraphConfig.GetObjectClassConfigs(PossibleObjectClass).CreateCategoryMetaData();
 		for (UClass* CurClass = PossibleObjectClass; CurClass; CurClass = CurClass->GetSuperClass())
 		{
 			const FString* CategoryNamesMetaData = CurClass->FindMetaData(CreateCategoryMetaData);
@@ -524,7 +530,32 @@ void UObjectTreeGraphSchema::FilterGraphContextPlaceableClasses(TArray<UClass*>&
 
 void UObjectTreeGraphSchema::GetContextMenuActions(UToolMenu* Menu, UGraphNodeContextMenuContext* Context) const
 {
+	using namespace UE::Cameras;
+
 	Super::GetContextMenuActions(Menu, Context);
+
+	const UEdGraph* CurrentGraph = Context->Graph;
+	const UEdGraphNode* CurrentNode = Context->Node;
+	const UEdGraphPin* CurrentPin = Context->Pin;
+
+	const FObjectTreeGraphEditorCommands& Commands = FObjectTreeGraphEditorCommands::Get();
+
+	if (CurrentPin)
+	{
+		FToolMenuSection& Section = Menu->FindOrAddSection("ObjectTreeGraphSchemaPinActions");
+		Section.InitSection(
+				"ObjectTreeGraphSchemaPinActions", 
+				LOCTEXT("ObjectPinActionsMenuHeader", "Object Pin Actions"), 
+				FToolMenuInsert());
+
+		if (CurrentPin->PinType.PinCategory == PC_Property &&
+				CurrentPin->PinType.PinSubCategory == PSC_ArrayPropertyItem)
+		{
+			Section.AddMenuEntry(Commands.InsertArrayItemPinBefore);
+			Section.AddMenuEntry(Commands.InsertArrayItemPinAfter);
+			Section.AddMenuEntry(Commands.RemoveArrayItemPin);
+		}
+	}
 }
 
 FName UObjectTreeGraphSchema::GetParentContextMenuName() const
@@ -627,95 +658,22 @@ const FPinConnectionResponse UObjectTreeGraphSchema::CanCreateConnection(const U
 
 bool UObjectTreeGraphSchema::TryCreateConnection(UEdGraphPin* A, UEdGraphPin* B) const
 {
-	const FPinConnectionResponse Response = CanCreateConnection(A, B);
-
-	if (A->LinkedTo.Contains(B) && B->LinkedTo.Contains(A))
-	{
-		return false;
-	}
-
 	FScopedTransaction Transaction(LOCTEXT("CreateConnection", "Create Connection"));
 
-	bool bModified = true;
-	FDelayedPinActions Actions;
-
-	// We need to reimplement this completely and skip the up-call to the base class.
-	// Yes, it makes me sad too. But this is because we need to interleave graph manipulation (like
-	// making and breaking pin links) with underlying model manipulation (setting and unsetting
-	// references between objects).
-	//
-	// Note that we do a few delicate things here:
-	//
-	// 1. We pass `true` to ApplyDisconnection` to indicate that disconnected pins should not be
-	//    removed (which happens when disconnecting an array property pin). We want those pins kept
-	//    around because they're only disconnected to be immediately reconnected.
-	//
-	// 2. However, in at least one situation we want the pin removed immediately. It's when there's
-	//    an object currently at index X in an array, and the user reconnects it to index Y. In that
-	//    case, we want to get rid of pin X, removing that item in the array, so that pin Y is at
-	//    the correct index. This is why we apply delayed pin actions as usual after graph 
-	//    manipulation is done.
-	//
-	switch (Response.Response)
-	{
-	case CONNECT_RESPONSE_MAKE:
-		A->MakeLinkTo(B);
-		ApplyConnection(A, B, Actions);
-		break;
-
-	case CONNECT_RESPONSE_BREAK_OTHERS_A:
-		ApplyDisconnection(A, Actions, true);
-		A->BreakAllPinLinks(true);
-		A->MakeLinkTo(B);
-		Actions.Apply();
-		ApplyConnection(A, B, Actions);
-		break;
-
-	case CONNECT_RESPONSE_BREAK_OTHERS_B:
-		ApplyDisconnection(B, Actions, true);
-		B->BreakAllPinLinks(true);
-		A->MakeLinkTo(B);
-		Actions.Apply();
-		ApplyConnection(A, B, Actions);
-		break;
-
-	case CONNECT_RESPONSE_BREAK_OTHERS_AB:
-		ApplyDisconnection(A, Actions, true);
-		ApplyDisconnection(B, Actions, true);
-		A->BreakAllPinLinks(true);
-		B->BreakAllPinLinks(true);
-		A->MakeLinkTo(B);
-		Actions.Apply();
-		ApplyConnection(A, B, Actions);
-		break;
-
-	case CONNECT_RESPONSE_MAKE_WITH_CONVERSION_NODE:
-		bModified = CreateAutomaticConversionNodeAndConnections(A, B);
-		break;
-
-	case CONNECT_RESPONSE_MAKE_WITH_PROMOTION:
-		bModified = CreatePromotedConnection(A, B);
-		break;
-
-	case CONNECT_RESPONSE_DISALLOW:
-	default:
-		bModified = false;
-		break;
-	}
+	const bool bModified = Super::TryCreateConnection(A, B);
 
 	if (!bModified)
 	{
-		ensure(Actions.IsEmpty());
 		Transaction.Cancel();
 		return false;
 	}
 
-#if WITH_EDITOR
-	A->GetOwningNode()->PinConnectionListChanged(A);
-	B->GetOwningNode()->PinConnectionListChanged(B);
-#endif
-
-	Actions.Apply();
+	UObjectTreeGraphNode* NodeA = Cast<UObjectTreeGraphNode>(A->GetOwningNode());
+	UObjectTreeGraphNode* NodeB = Cast<UObjectTreeGraphNode>(B->GetOwningNode());
+	if (NodeA && NodeA->GetObject() && NodeB && NodeB->GetObject())
+	{
+		ApplyConnection(A, B);
+	}
 
 	return true;
 }
@@ -724,7 +682,6 @@ void UObjectTreeGraphSchema::BreakNodeLinks(UEdGraphNode& TargetNode) const
 {
 	const FScopedTransaction Transaction(LOCTEXT("BreakNodeLinks", "Break Node Links"));
 
-	FDelayedPinActions Actions;
 	TArray<UEdGraphPin*> CachedPins = TargetNode.Pins;
 
 #if WITH_EDITOR
@@ -762,80 +719,32 @@ void UObjectTreeGraphSchema::BreakNodeLinks(UEdGraphNode& TargetNode) const
 
 void UObjectTreeGraphSchema::BreakPinLinks(UEdGraphPin& TargetPin, bool bSendsNodeNotification) const
 {
-	if (TargetPin.LinkedTo.IsEmpty())
+	FScopedTransaction Transaction(LOCTEXT("BreakPinLinks", "Break Pin Links"));
+
+	UObjectTreeGraphNode* TargetNode = Cast<UObjectTreeGraphNode>(TargetPin.GetOwningNode());
+	if (TargetNode && TargetNode->GetObject())
 	{
-		Super::BreakPinLinks(TargetPin, bSendsNodeNotification);
-		return;
+		ApplyDisconnection(&TargetPin);
 	}
 
-	const FScopedTransaction Transaction(LOCTEXT("BreakPinLinks", "Break Pin Links"));
-
-	FDelayedPinActions Actions;
-	ApplyDisconnection(&TargetPin, Actions, false);
-
 	Super::BreakPinLinks(TargetPin, bSendsNodeNotification);
-
-	Actions.Apply();
 }
 
 void UObjectTreeGraphSchema::BreakSinglePinLink(UEdGraphPin* SourcePin, UEdGraphPin* TargetPin) const
 {
-	const FScopedTransaction Transaction(LOCTEXT("BreakSinglePinLink", "Break Pin Link"));
-
-	FDelayedPinActions Actions;
-	ApplyDisconnection(SourcePin, TargetPin, Actions);
-
-	Super::BreakSinglePinLink(SourcePin, TargetPin);
-
-	Actions.Apply();
-}
+	FScopedTransaction Transaction(LOCTEXT("BreakSinglePinLink", "Break Pin Link"));
 
-void UObjectTreeGraphSchema::FDelayedPinActions::CreateNewItemPin(UObjectTreeGraphNode* Node, FArrayProperty* ArrayProperty)
-{
-	if (ensure(Node && ArrayProperty))
+	UObjectTreeGraphNode* SourceNode = Cast<UObjectTreeGraphNode>(SourcePin->GetOwningNode());
+	UObjectTreeGraphNode* TargetNode = Cast<UObjectTreeGraphNode>(TargetPin->GetOwningNode());
+	if (SourceNode && SourceNode->GetObject() && TargetNode && TargetNode->GetObject())
 	{
-		ItemPinsToCreate.Add({ Node, ArrayProperty });
+		ApplyDisconnection(SourcePin, TargetPin);
 	}
-}
 
-void UObjectTreeGraphSchema::FDelayedPinActions::RemoveItemPin(UEdGraphPin* Pin)
-{
-	if (ensure(Pin))
-	{
-		ItemPinsToRemove.Add(Pin);
-	}
-}
-
-bool UObjectTreeGraphSchema::FDelayedPinActions::IsEmpty() const
-{
-	return ItemPinsToCreate.IsEmpty() && ItemPinsToRemove.IsEmpty();
-}
-
-void UObjectTreeGraphSchema::FDelayedPinActions::Apply()
-{
-	TSet<UEdGraphNode*> NodesToNotify;
-	for (UEdGraphPin* Pin : ItemPinsToRemove)
-	{
-		UObjectTreeGraphNode* OwningNode = Cast<UObjectTreeGraphNode>(Pin->GetOwningNode());
-		check(OwningNode);
-		OwningNode->RemoveItemPin(Pin);
-
-		NodesToNotify.Add(OwningNode);
-	}
-	for (TTuple<UObjectTreeGraphNode*, FArrayProperty*>& Pair : ItemPinsToCreate)
-	{
-		Pair.Key->CreateNewItemPin(*Pair.Value);
-
-		NodesToNotify.Add(Pair.Key);
-	}
-	for (UEdGraphNode* Node : NodesToNotify)
-	{
-		Node->GetGraph()->NotifyNodeChanged(Node);
-	}
-	ItemPinsToRemove.Reset();
+	Super::BreakSinglePinLink(SourcePin, TargetPin);
 }
 
-void UObjectTreeGraphSchema::ApplyConnection(UEdGraphPin* A, UEdGraphPin* B, FDelayedPinActions& Actions) const
+void UObjectTreeGraphSchema::ApplyConnection(UEdGraphPin* A, UEdGraphPin* B) const
 {
 	// Input must have been validated prior to calling this method:
 	//
@@ -851,7 +760,7 @@ void UObjectTreeGraphSchema::ApplyConnection(UEdGraphPin* A, UEdGraphPin* B, FDe
 	check(A && B);
 
 	// See if a sub-class is handling this situation.
-	if (OnApplyConnection(A, B, Actions))
+	if (OnApplyConnection(A, B))
 	{
 		return;
 	}
@@ -882,7 +791,7 @@ void UObjectTreeGraphSchema::ApplyConnection(UEdGraphPin* A, UEdGraphPin* B, FDe
 	check(PropertyObject && ValueObject);
 
 	// If it is a property pin, set the value of the underlying property.
-	// If it is an array property pin, add a new item at the pin's index, and optionally add a new pin.
+	// If it is an array property pin, set the array item value at the pin's index.
 	FProperty* Property = PropertyNode->GetPropertyForPin(PropertyPin);
 
 	if (FObjectProperty* ObjectProperty = CastField<FObjectProperty>(Property))
@@ -906,38 +815,29 @@ void UObjectTreeGraphSchema::ApplyConnection(UEdGraphPin* A, UEdGraphPin* B, FDe
 		ensure(Index != INDEX_NONE);
 
 		FScriptArrayHelper ArrayHelper(ArrayProperty, ArrayProperty->ContainerPtrToValuePtr<void>(PropertyObject));
-		const int32 PreviousNum = ArrayHelper.Num();
-		const bool bExpandedArray = ArrayHelper.ExpandForIndex(Index);
+		ensure(Index < ArrayHelper.Num());
 
 		FObjectProperty* InnerProperty = CastFieldChecked<FObjectProperty>(ArrayProperty->Inner);
 		InnerProperty->SetObjectPropertyValue(ArrayHelper.GetRawPtr(Index), ValueObject);
 
-		if (bExpandedArray)
-		{
-			ensure(Index == PreviousNum);  // We only support adding a single new element at the end, not "inserting"
-										   // into arbitrary indices requiring expanding the array by more than one
-										   // element.
-			Actions.CreateNewItemPin(PropertyNode, ArrayProperty);
-		}
-
 		FPropertyChangedEvent PropertyChangedEvent(Property);
-		PropertyChangedEvent.ChangeType = bExpandedArray ? EPropertyChangeType::ArrayAdd : EPropertyChangeType::ValueSet;
+		PropertyChangedEvent.ChangeType = EPropertyChangeType::ValueSet;
 		PropertyObject->PostEditChangeProperty(PropertyChangedEvent);
 	}
 }
 
-bool UObjectTreeGraphSchema::OnApplyConnection(UEdGraphPin* A, UEdGraphPin* B, FDelayedPinActions& Actions) const
+bool UObjectTreeGraphSchema::OnApplyConnection(UEdGraphPin* A, UEdGraphPin* B) const
 {
 	return false;
 }
 
-void UObjectTreeGraphSchema::ApplyDisconnection(UEdGraphPin* TargetPin, FDelayedPinActions& Actions, bool bIsReconnecting) const
+void UObjectTreeGraphSchema::ApplyDisconnection(UEdGraphPin* TargetPin) const
 {
 	// Input must have been validated prior to calling this method:
 	//
 	// - no null objects
-	// - the pin is the property pin to reset, or the self pin connected to a propert pin
-	// - the pin belongs to a ObjectTreeGraph node
+	// - the pin is the property pin to reset, or the self pin connected to a property pin
+	// - the pin belongs to an ObjectTreeGraph node
 	// - this node has a valid object
 	// - we should have a transaction active
 	//
@@ -954,21 +854,16 @@ void UObjectTreeGraphSchema::ApplyDisconnection(UEdGraphPin* TargetPin, FDelayed
 	}
 
 	// See if a sub-class is handling this situation.
-	if (OnApplyDisconnection(TargetPin, Actions, bIsReconnecting))
+	if (OnApplyDisconnection(TargetPin))
 	{
 		return;
 	}
 
 	// We may either disconnect a self pin, or a property or array property pin. Let's see
 	// what sort of pin we were given: we want the property side of things.
-	bool bRemoveArrayItem = !bIsReconnecting;
 	if (TargetPin->PinType.PinCategory == PC_Self)
 	{
 		TargetPin = TargetPin->LinkedTo[0];
-		// If `bIsReconnecting` is true, it means we need to leave `TargetPin` alone because
-		// we want to reconnect it right away. However, if the property pin is on the other
-		// side the of the link, we are free to remove it if it's an array property pin.
-		bRemoveArrayItem = true;
 	}
 	check(TargetPin->PinType.PinCategory == PC_Property);
 
@@ -979,8 +874,7 @@ void UObjectTreeGraphSchema::ApplyDisconnection(UEdGraphPin* TargetPin, FDelayed
 	check(PropertyObject);
 
 	// If it is a property pin, clear the value of the underlying property.
-	// If it is an array property pin, remove the value at the given index in the underlying array,
-	// or just clear it if we want to reconnect that pin immediately.
+	// If it is an array property pin, clear the value at the given index in the underlying array.
 	FProperty* Property = PropertyNode->GetPropertyForPin(TargetPin);
 
 	if (FObjectProperty* ObjectProperty = CastField<FObjectProperty>(Property))
@@ -1000,32 +894,24 @@ void UObjectTreeGraphSchema::ApplyDisconnection(UEdGraphPin* TargetPin, FDelayed
 
 		PropertyObject->Modify();
 
-		int32 Index = PropertyNode->GetIndexOfArrayPin(TargetPin);
-		ensure(Index != INDEX_NONE);
-
+		const int32 Index = PropertyNode->GetIndexOfArrayPin(TargetPin);
 		FScriptArrayHelper ArrayHelper(ArrayProperty, ArrayProperty->ContainerPtrToValuePtr<void>(PropertyObject));
-		if (bRemoveArrayItem)
-		{
-			ArrayHelper.RemoveValues(Index);
-			Actions.RemoveItemPin(TargetPin);
-		}
-		else
-		{
-			FObjectProperty* InnerProperty = CastFieldChecked<FObjectProperty>(ArrayProperty->Inner);
-			InnerProperty->SetObjectPropertyValue(ArrayHelper.GetRawPtr(Index), nullptr);
-		}
+		ensure(Index >= 0 && Index < ArrayHelper.Num());
 
-		FPropertyChangedEvent PropertyChangedEvent(Property, EPropertyChangeType::ArrayRemove);
+		FObjectProperty* InnerProperty = CastFieldChecked<FObjectProperty>(ArrayProperty->Inner);
+		InnerProperty->SetObjectPropertyValue(ArrayHelper.GetRawPtr(Index), nullptr);
+
+		FPropertyChangedEvent PropertyChangedEvent(Property, EPropertyChangeType::ValueSet);
 		PropertyObject->PostEditChangeProperty(PropertyChangedEvent);
 	}
 }
 
-bool UObjectTreeGraphSchema::OnApplyDisconnection(UEdGraphPin* TargetPin, FDelayedPinActions& Actions, bool bIsReconnecting) const
+bool UObjectTreeGraphSchema::OnApplyDisconnection(UEdGraphPin* TargetPin) const
 {
 	return false;
 }
 
-void UObjectTreeGraphSchema::ApplyDisconnection(UEdGraphPin* SourcePin, UEdGraphPin* TargetPin, FDelayedPinActions& Actions) const
+void UObjectTreeGraphSchema::ApplyDisconnection(UEdGraphPin* SourcePin, UEdGraphPin* TargetPin) const
 {
 	// Input must have been validated prior to calling this method:
 	//
@@ -1035,18 +921,18 @@ void UObjectTreeGraphSchema::ApplyDisconnection(UEdGraphPin* SourcePin, UEdGraph
 	check(SourcePin && TargetPin);
 
 	// See if a sub-class is handling this situation.
-	if (OnApplyDisconnection(SourcePin, TargetPin, Actions))
+	if (OnApplyDisconnection(SourcePin, TargetPin))
 	{
 		return;
 	}
 
 	if (SourcePin->PinType.PinCategory == PC_Self && TargetPin->PinType.PinCategory == PC_Property)
 	{
-		return ApplyDisconnection(TargetPin, Actions, false);
+		return ApplyDisconnection(TargetPin);
 	}
 	else if (SourcePin->PinType.PinCategory == PC_Property && TargetPin->PinType.PinCategory == PC_Self)
 	{
-		return ApplyDisconnection(SourcePin, Actions, false);
+		return ApplyDisconnection(SourcePin);
 	}
 	else
 	{
@@ -1054,11 +940,166 @@ void UObjectTreeGraphSchema::ApplyDisconnection(UEdGraphPin* SourcePin, UEdGraph
 	}
 }
 
-bool UObjectTreeGraphSchema::OnApplyDisconnection(UEdGraphPin* SourcePin, UEdGraphPin* TargetPin, FDelayedPinActions& Actions) const
+bool UObjectTreeGraphSchema::OnApplyDisconnection(UEdGraphPin* SourcePin, UEdGraphPin* TargetPin) const
 {
 	return false;
 }
 
+void UObjectTreeGraphSchema::InsertArrayItemPin(UEdGraphPin* ArrayPin, int32 Index) const
+{
+	if (!ensure(ArrayPin))
+	{
+		return;
+	}
+
+	UObjectTreeGraphNode* ObjectNode = Cast<UObjectTreeGraphNode>(ArrayPin->GetOwningNode());
+	if (!ensure(ObjectNode && ObjectNode->GetObject()))
+	{
+		return;
+	}
+
+	FArrayProperty* ArrayProperty = CastField<FArrayProperty>(ObjectNode->GetPropertyForPin(ArrayPin));
+	if (!ensure(ArrayProperty))
+	{
+		return;
+	}
+
+	FScopedTransaction Transaction(FText::Format(
+				LOCTEXT("InsertArrayItem", "Add {0} Pin"), FText::FromName(ArrayProperty->GetFName())));
+
+	UObject* Object = ObjectNode->GetObject();
+
+	Object->PreEditChange(ArrayProperty);
+
+	Object->Modify();
+
+	FScriptArrayHelper ArrayHelper(ArrayProperty, ArrayProperty->ContainerPtrToValuePtr<void>(Object));
+	FObjectProperty* InnerProperty = CastFieldChecked<FObjectProperty>(ArrayProperty->Inner);
+
+	const bool bIsActualInsert = (Index >= 0 && Index <= ArrayHelper.Num());
+	if (bIsActualInsert)
+	{
+		ArrayHelper.InsertValues(Index, 1);
+		InnerProperty->SetObjectPropertyValue(ArrayHelper.GetRawPtr(Index), nullptr);
+		ObjectNode->InsertNewItemPin(ArrayPin, Index);
+	}
+	else
+	{
+		const int32 NewIndex = ArrayHelper.AddValues(1);
+		InnerProperty->SetObjectPropertyValue(ArrayHelper.GetRawPtr(NewIndex), nullptr);
+		ObjectNode->CreateNewItemPins(ArrayPin, 1);
+	}
+
+	FPropertyChangedEvent PropertyChangedEvent(ArrayProperty, EPropertyChangeType::ArrayAdd);
+	Object->PostEditChangeProperty(PropertyChangedEvent);
+
+	UEdGraph* Graph = ObjectNode->GetGraph();
+	Graph->NotifyNodeChanged(ObjectNode);
+}
+
+void UObjectTreeGraphSchema::InsertArrayItemPinBefore(UEdGraphPin* ArrayItemPin) const
+{
+	if (!ensure(ArrayItemPin))
+	{
+		return;
+	}
+
+	UEdGraphPin* ArrayPin = ArrayItemPin->ParentPin;
+	if (!ensure(ArrayPin))
+	{
+		return;
+	}
+
+	const int32 Index = ArrayPin->SubPins.Find(ArrayItemPin);
+	if (ensure(Index >= 0))
+	{
+		InsertArrayItemPin(ArrayPin, Index);
+	}
+}
+
+void UObjectTreeGraphSchema::InsertArrayItemPinAfter(UEdGraphPin* ArrayItemPin) const
+{
+	if (!ensure(ArrayItemPin))
+	{
+		return;
+	}
+
+	UEdGraphPin* ArrayPin = ArrayItemPin->ParentPin;
+	if (!ensure(ArrayPin))
+	{
+		return;
+	}
+
+	const int32 Index = ArrayPin->SubPins.Find(ArrayItemPin);
+	if (ensure(Index >= 0))
+	{
+		InsertArrayItemPin(ArrayPin, Index + 1);
+	}
+}
+
+void UObjectTreeGraphSchema::RemoveArrayItemPin(UEdGraphPin* ArrayItemPin) const
+{
+	if (!ensure(ArrayItemPin))
+	{
+		return;
+	}
+
+	if (!ensure(ArrayItemPin->PinType.PinCategory == PC_Property && 
+				ArrayItemPin->PinType.PinSubCategory == PSC_ArrayPropertyItem))
+	{
+		return;
+	}
+
+	UEdGraphPin* ArrayPin = ArrayItemPin->ParentPin;
+	if (!ensure(ArrayPin))
+	{
+		return;
+	}
+
+	if (!ensure(ArrayPin->PinType.PinCategory == PC_Property && 
+				ArrayPin->PinType.PinSubCategory == PSC_ArrayProperty))
+	{
+		return;
+	}
+
+	const int32 Index = ArrayItemPin->ParentPin->SubPins.Find(ArrayItemPin);
+	if (!ensure(Index >= 0))
+	{
+		return;
+	}
+
+	UObjectTreeGraphNode* ObjectNode = Cast<UObjectTreeGraphNode>(ArrayPin->GetOwningNode());
+	if (!ensure(ObjectNode && ObjectNode->GetObject()))
+	{
+		return;
+	}
+
+	FArrayProperty* ArrayProperty = CastField<FArrayProperty>(ObjectNode->GetPropertyForPin(ArrayPin));
+	if (!ensure(ArrayProperty))
+	{
+		return;
+	}
+
+	FScopedTransaction Transaction(FText::Format(
+				LOCTEXT("RemoveArrayItem", "Remove {0} Pin"), FText::FromName(ArrayProperty->GetFName())));
+
+	UObject* Object = ObjectNode->GetObject();
+
+	Object->PreEditChange(ArrayProperty);
+
+	Object->Modify();
+
+	FScriptArrayHelper ArrayHelper(ArrayProperty, ArrayProperty->ContainerPtrToValuePtr<void>(Object));
+	ArrayHelper.RemoveValues(Index, 1);
+	ObjectNode->RemoveItemPin(ArrayItemPin);
+
+	FPropertyChangedEvent PropertyChangedEvent(ArrayProperty, EPropertyChangeType::ArrayRemove);
+	Object->PostEditChangeProperty(PropertyChangedEvent);
+
+	UEdGraph* Graph = ObjectNode->GetGraph();
+	Graph->NotifyNodeChanged(ObjectNode);
+}
+
 bool UObjectTreeGraphSchema::SupportsDropPinOnNode(UEdGraphNode* InTargetNode, const FEdGraphPinType& InSourcePinType, EEdGraphPinDirection InSourcePinDirection, FText& OutErrorMessage) const
 {
 	return Super::SupportsDropPinOnNode(InTargetNode, InSourcePinType, InSourcePinDirection, OutErrorMessage);
@@ -1296,15 +1337,15 @@ bool UObjectTreeGraphSchema::CanImportNodesFromText(UObjectTreeGraph* InGraph, c
 	return Factory.CanCreateObjectsFromText(TextToImport);
 }
 
-const FObjectTreeGraphClassConfig& UObjectTreeGraphSchema::GetObjectClassConfig(const UObjectTreeGraphNode* InNode) const
+const FObjectTreeGraphClassConfigs UObjectTreeGraphSchema::GetObjectClassConfigs(const UObjectTreeGraphNode* InNode) const
 {
 	const UObjectTreeGraph* Graph = CastChecked<UObjectTreeGraph>(InNode->GetGraph());
-	return GetObjectClassConfig(Graph, InNode->GetObject()->GetClass());
+	return GetObjectClassConfigs(Graph, InNode->GetObject()->GetClass());
 }
 
-const FObjectTreeGraphClassConfig& UObjectTreeGraphSchema::GetObjectClassConfig(const UObjectTreeGraph* InGraph, UClass* InObjectClass) const
+const FObjectTreeGraphClassConfigs UObjectTreeGraphSchema::GetObjectClassConfigs(const UObjectTreeGraph* InGraph, UClass* InObjectClass) const
 {
-	return InGraph->GetConfig().GetObjectClassConfig(InObjectClass);
+	return InGraph->GetConfig().GetObjectClassConfigs(InObjectClass);
 }
 
 FObjectGraphSchemaAction_NewNode::FObjectGraphSchemaAction_NewNode()
@@ -1350,6 +1391,10 @@ UEdGraphNode* FObjectGraphSchemaAction_NewNode::PerformAction(UEdGraph* ParentGr
 
 	if (NewObject)
 	{
+		const FObjectTreeGraphConfig& GraphConfig = ObjectTreeGraph->GetConfig();
+		const FObjectTreeGraphClassConfigs ObjectClassConfigs = GraphConfig.GetObjectClassConfigs(ObjectClass);
+		ObjectClassConfigs.OnSetupNewObject().ExecuteIfBound(NewObject);
+
 		ObjectTreeGraph->Modify();
 
 		UObjectTreeGraphNode* NewGraphNode = Schema->CreateObjectNode(ObjectTreeGraph, NewObject);
diff --git a/Engine/Plugins/Cameras/GameplayCameras/Source/GameplayCamerasEditor/Private/Editors/ObjectTreeGraphSearch.cpp b/Engine/Plugins/Cameras/GameplayCameras/Source/GameplayCamerasEditor/Private/Editors/ObjectTreeGraphSearch.cpp
--- a/Engine/Plugins/Cameras/GameplayCameras/Source/GameplayCamerasEditor/Private/Editors/ObjectTreeGraphSearch.cpp
+++ b/Engine/Plugins/Cameras/GameplayCameras/Source/GameplayCamerasEditor/Private/Editors/ObjectTreeGraphSearch.cpp
@@ -65,8 +65,6 @@ void FObjectTreeGraphSearch::SearchObject(UObject* InObject, FSearchState& InOut
 	UClass* ObjectClass = InObject->GetClass();
 	const FObjectTreeGraphConfig& GraphConfig(*InOutState.GraphConfig);
 
-	const FObjectTreeGraphClassConfig& ObjectClassConfig = GraphConfig.GetObjectClassConfig(ObjectClass);
-
 	if (MatchObject(InObject, InOutState))
 	{
 		InOutState.Results.Add(FSearchResult{ InOutState.RootObject, InOutState.GraphConfig, InObject });
diff --git a/Engine/Plugins/Cameras/GameplayCameras/Source/GameplayCamerasEditor/Private/Editors/SObjectTreeGraphEditor.cpp b/Engine/Plugins/Cameras/GameplayCameras/Source/GameplayCamerasEditor/Private/Editors/SObjectTreeGraphEditor.cpp
--- a/Engine/Plugins/Cameras/GameplayCameras/Source/GameplayCamerasEditor/Private/Editors/SObjectTreeGraphEditor.cpp
+++ b/Engine/Plugins/Cameras/GameplayCameras/Source/GameplayCamerasEditor/Private/Editors/SObjectTreeGraphEditor.cpp
@@ -3,6 +3,7 @@
 #include "Editors/SObjectTreeGraphEditor.h"
 
 #include "Algo/AnyOf.h"
+#include "Commands/ObjectTreeGraphEditorCommands.h"
 #include "EdGraph/EdGraphPin.h"
 #include "EdGraph/EdGraphSchema.h"
 #include "Editor.h"
@@ -75,6 +76,8 @@ SObjectTreeGraphEditor::~SObjectTreeGraphEditor()
 
 void SObjectTreeGraphEditor::InitializeBuiltInCommands()
 {
+	using namespace UE::Cameras;
+
 	if (BuiltInCommands.IsValid())
 	{
 		return;
@@ -82,6 +85,7 @@ void SObjectTreeGraphEditor::InitializeBuiltInCommands()
 
 	const FGenericCommands& GenericCommands = FGenericCommands::Get();
 	const FGraphEditorCommandsImpl& GraphEditorCommands = FGraphEditorCommands::Get();
+	const FObjectTreeGraphEditorCommands& ObjectTreeGraphEditorCommands = FObjectTreeGraphEditorCommands::Get();
 
 	BuiltInCommands = MakeShared<FUICommandList>();
 
@@ -157,6 +161,17 @@ void SObjectTreeGraphEditor::InitializeBuiltInCommands()
 	BuiltInCommands->MapAction(GraphEditorCommands.DistributeNodesVertically,
 			FExecuteAction::CreateSP(this, &SObjectTreeGraphEditor::OnDistributeNodesVertically)
 			);
+
+	// Custom commands.
+	BuiltInCommands->MapAction(ObjectTreeGraphEditorCommands.InsertArrayItemPinBefore,
+			FExecuteAction::CreateSP(this, &SObjectTreeGraphEditor::OnInsertArrayItemPinBefore)
+			);
+	BuiltInCommands->MapAction(ObjectTreeGraphEditorCommands.InsertArrayItemPinAfter,
+			FExecuteAction::CreateSP(this, &SObjectTreeGraphEditor::OnInsertArrayItemPinAfter)
+			);
+	BuiltInCommands->MapAction(ObjectTreeGraphEditorCommands.RemoveArrayItemPin,
+			FExecuteAction::CreateSP(this, &SObjectTreeGraphEditor::OnRemoveArrayItemPin)
+			);
 }
 
 void SObjectTreeGraphEditor::JumpToNode(UEdGraphNode* InNode)
@@ -593,5 +608,50 @@ void SObjectTreeGraphEditor::OnDistributeNodesVertically()
 	GraphEditor->OnDistributeNodesV();
 }
 
+void SObjectTreeGraphEditor::OnInsertArrayItemPinBefore()
+{
+	if (UEdGraphPin* SelectedPin = GraphEditor->GetGraphPinForMenu())
+	{
+		UEdGraph* CurrentGraph = GraphEditor->GetCurrentGraph();
+		const UObjectTreeGraphSchema* Schema = CastChecked<UObjectTreeGraphSchema>(CurrentGraph->GetSchema());
+
+		Schema->InsertArrayItemPinBefore(SelectedPin);
+
+		GraphEditor->RefreshNode(*SelectedPin->GetOwningNode());
+	}
+}
+
+void SObjectTreeGraphEditor::OnInsertArrayItemPinAfter()
+{
+	if (UEdGraphPin* SelectedPin = GraphEditor->GetGraphPinForMenu())
+	{
+		UEdGraph* CurrentGraph = GraphEditor->GetCurrentGraph();
+		const UObjectTreeGraphSchema* Schema = CastChecked<UObjectTreeGraphSchema>(CurrentGraph->GetSchema());
+
+		Schema->InsertArrayItemPinAfter(SelectedPin);
+
+		GraphEditor->RefreshNode(*SelectedPin->GetOwningNode());
+	}
+}
+
+void SObjectTreeGraphEditor::OnRemoveArrayItemPin()
+{
+	if (UEdGraphPin* SelectedPin = GraphEditor->GetGraphPinForMenu())
+	{
+		UEdGraph* CurrentGraph = GraphEditor->GetCurrentGraph();
+		const UObjectTreeGraphSchema* Schema = CastChecked<UObjectTreeGraphSchema>(CurrentGraph->GetSchema());
+
+		// Get owning node before we remove the pin.
+		UEdGraphNode* OwningNode = SelectedPin->GetOwningNode();
+		
+		Schema->RemoveArrayItemPin(SelectedPin);
+
+		if (ensure(OwningNode))
+		{
+			GraphEditor->RefreshNode(*OwningNode);
+		}
+	}
+}
+
 #undef LOCTEXT_NAMESPACE
 
diff --git a/Engine/Plugins/Cameras/GameplayCameras/Source/GameplayCamerasEditor/Private/Editors/SObjectTreeGraphNode.cpp b/Engine/Plugins/Cameras/GameplayCameras/Source/GameplayCamerasEditor/Private/Editors/SObjectTreeGraphNode.cpp
--- a/Engine/Plugins/Cameras/GameplayCameras/Source/GameplayCamerasEditor/Private/Editors/SObjectTreeGraphNode.cpp
+++ b/Engine/Plugins/Cameras/GameplayCameras/Source/GameplayCamerasEditor/Private/Editors/SObjectTreeGraphNode.cpp
@@ -2,7 +2,15 @@
 
 #include "Editors/SObjectTreeGraphNode.h"
 
+#include "EdGraph/EdGraph.h"
 #include "Editors/ObjectTreeGraphNode.h"
+#include "Editors/ObjectTreeGraphSchema.h"
+#include "GraphEditorSettings.h"
+#include "Widgets/Images/SImage.h"
+#include "Widgets/Input/SButton.h"
+#include "Widgets/SBoxPanel.h"
+
+#define LOCTEXT_NAMESPACE "SObjectTreeGraphNode"
 
 void SObjectTreeGraphNode::Construct(const FArguments& InArgs)
 {
@@ -24,3 +32,81 @@ void SObjectTreeGraphNode::MoveTo(const FVector2D& NewPosition, FNodeSet& NodeFi
 	}
 }
 
+void SObjectTreeGraphNode::CreateOutputSideAddButton(TSharedPtr<SVerticalBox> OutputBox)
+{
+	TArray<FArrayProperty*> ArrayProperties;
+	ObjectGraphNode->GetArrayProperties(ArrayProperties);
+	for (FArrayProperty* ArrayProperty : ArrayProperties)
+	{
+		TSharedRef<SWidget> AddPinButton = MakeAddArrayPropertyPinButton(ArrayProperty);
+
+		FMargin AddPinPadding = Settings->GetOutputPinPadding();
+		AddPinPadding.Top += 6.0f;
+
+		OutputBox->AddSlot()
+		.AutoHeight()
+		.VAlign(VAlign_Center)
+		.HAlign(HAlign_Right)
+		.Padding(AddPinPadding)
+		[
+			AddPinButton
+		];
+	}
+}
+
+TSharedRef<SWidget> SObjectTreeGraphNode::MakeAddArrayPropertyPinButton(FArrayProperty* ArrayProperty)
+{
+	TSharedRef<SWidget> ButtonContent = SNew(SHorizontalBox)
+	+SHorizontalBox::Slot()
+	.AutoWidth()
+	.HAlign(HAlign_Left)
+	[
+		SNew(STextBlock)
+		.Text(FText::Format(
+					LOCTEXT("AddPropertyPinButtonLabelFmt", "Add {0} pin"),
+					FText::FromName(ArrayProperty->GetFName())))
+		.ColorAndOpacity(FLinearColor::White)
+	]
+	+SHorizontalBox::Slot()
+	.AutoWidth()
+	. VAlign(VAlign_Center)
+	. Padding(7, 0, 0, 0)
+	[
+		SNew(SImage)
+		.Image(FAppStyle::GetBrush(TEXT("Icons.PlusCircle")))
+	];
+
+	TSharedRef<SButton> AddPinButton = SNew(SButton)
+	.ContentPadding(0.0f)
+	.ButtonStyle( FAppStyle::Get(), "NoBorder" )
+	.OnClicked(this, &SObjectTreeGraphNode::OnAddArrayPropertyPin, ArrayProperty)
+	.IsEnabled(this, &SGraphNode::IsNodeEditable)
+	.ToolTipText(FText::Format(
+				LOCTEXT("AddPropertyPinButtonTooltipFmt", "Adds a new pin for the '{0}' property on this node"),
+				FText::FromName(ArrayProperty->GetFName())))
+	[
+		ButtonContent
+	];
+
+	AddPinButton->SetCursor( EMouseCursor::Hand );
+
+	return AddPinButton;
+}
+
+FReply SObjectTreeGraphNode::OnAddArrayPropertyPin(FArrayProperty* ArrayProperty)
+{
+	if (ObjectGraphNode)
+	{
+		UEdGraph* Graph = ObjectGraphNode->GetGraph();
+		const UObjectTreeGraphSchema* Schema = Cast<const UObjectTreeGraphSchema>(Graph->GetSchema());
+
+		UEdGraphPin* ArrayPin = ObjectGraphNode->GetPinForProperty(ArrayProperty);
+		Schema->InsertArrayItemPin(ArrayPin, INDEX_NONE);
+
+		UpdateGraphNode();
+	}
+	return FReply::Handled();
+}
+
+#undef LOCTEXT_NAMESPACE
+
diff --git a/Engine/Plugins/Cameras/GameplayCameras/Source/GameplayCamerasEditor/Private/GameplayCamerasEditorModule.cpp b/Engine/Plugins/Cameras/GameplayCameras/Source/GameplayCamerasEditor/Private/GameplayCamerasEditorModule.cpp
--- a/Engine/Plugins/Cameras/GameplayCameras/Source/GameplayCamerasEditor/Private/GameplayCamerasEditorModule.cpp
+++ b/Engine/Plugins/Cameras/GameplayCameras/Source/GameplayCamerasEditor/Private/GameplayCamerasEditorModule.cpp
@@ -11,6 +11,7 @@
 #include "Commands/CameraRigTransitionEditorCommands.h"
 #include "Commands/CameraVariableCollectionEditorCommands.h"
 #include "Commands/GameplayCamerasDebuggerCommands.h"
+#include "Commands/ObjectTreeGraphEditorCommands.h"
 #include "Customizations/CameraParameterDetailsCustomizations.h"
 #include "Customizations/CameraProxyTableDetailsCustomization.h"
 #include "Customizations/CameraRigAssetReferenceDetailsCustomization.h"
@@ -108,6 +109,7 @@ public:
 		FCameraRigTransitionEditorCommands::Unregister();
 		FCameraVariableCollectionEditorCommands::Unregister();
 		FGameplayCamerasDebuggerCommands::Unregister();
+		FObjectTreeGraphEditorCommands::Unregister();
 
 		UnregisterSettings();
 		UnregisterCameraDirectorEditors();
@@ -431,6 +433,7 @@ private:
 		FCameraRigTransitionEditorCommands::Register();
 		FCameraVariableCollectionEditorCommands::Register();
 		FGameplayCamerasDebuggerCommands::Register();
+		FObjectTreeGraphEditorCommands::Register();
 	}
 
 	void RegisterRewindDebuggerFeatures()
diff --git a/Engine/Plugins/Cameras/GameplayCameras/Source/GameplayCamerasEditor/Public/Commands/ObjectTreeGraphEditorCommands.h b/Engine/Plugins/Cameras/GameplayCameras/Source/GameplayCamerasEditor/Public/Commands/ObjectTreeGraphEditorCommands.h
new file mode 100644
--- /dev/null
+++ b/Engine/Plugins/Cameras/GameplayCameras/Source/GameplayCamerasEditor/Public/Commands/ObjectTreeGraphEditorCommands.h
@@ -0,0 +1,31 @@
+// Copyright Epic Games, Inc. All Rights Reserved.
+
+#pragma once
+
+#include "CoreTypes.h"
+#include "Framework/Commands/Commands.h"
+
+class FUICommandInfo;
+class UClass;
+struct FInputChord;
+
+namespace UE::Cameras
+{
+
+class FObjectTreeGraphEditorCommands : public TCommands<FObjectTreeGraphEditorCommands>
+{
+public:
+
+	FObjectTreeGraphEditorCommands();
+
+	virtual void RegisterCommands() override;
+	
+public:
+
+	TSharedPtr<FUICommandInfo> InsertArrayItemPinBefore;
+	TSharedPtr<FUICommandInfo> InsertArrayItemPinAfter;
+	TSharedPtr<FUICommandInfo> RemoveArrayItemPin;
+};
+
+}  // namespace UE::Cameras
+
diff --git a/Engine/Plugins/Cameras/GameplayCameras/Source/GameplayCamerasEditor/Public/Editors/CameraNodeGraphSchema.h b/Engine/Plugins/Cameras/GameplayCameras/Source/GameplayCamerasEditor/Public/Editors/CameraNodeGraphSchema.h
--- a/Engine/Plugins/Cameras/GameplayCameras/Source/GameplayCamerasEditor/Public/Editors/CameraNodeGraphSchema.h
+++ b/Engine/Plugins/Cameras/GameplayCameras/Source/GameplayCamerasEditor/Public/Editors/CameraNodeGraphSchema.h
@@ -34,9 +34,9 @@ protected:
 	virtual void OnCreateAllNodes(UObjectTreeGraph* InGraph, const FCreatedNodes& InCreatedNodes) const override;
 	virtual void OnAddConnectableObject(UObjectTreeGraph* InGraph, UObjectTreeGraphNode* InNewNode) const override;
 	virtual void OnRemoveConnectableObject(UObjectTreeGraph* InGraph, UObjectTreeGraphNode* InRemovedNode) const override;
-	virtual bool OnApplyConnection(UEdGraphPin* A, UEdGraphPin* B, FDelayedPinActions& Actions) const override;
-	virtual bool OnApplyDisconnection(UEdGraphPin* TargetPin, FDelayedPinActions& Actions, bool bIsReconnecting) const override;
-	virtual bool OnApplyDisconnection(UEdGraphPin* SourcePin, UEdGraphPin* TargetPin, FDelayedPinActions& Actions) const override;
+	virtual bool OnApplyConnection(UEdGraphPin* A, UEdGraphPin* B) const override;
+	virtual bool OnApplyDisconnection(UEdGraphPin* TargetPin) const override;
+	virtual bool OnApplyDisconnection(UEdGraphPin* SourcePin, UEdGraphPin* TargetPin) const override;
 };
 
 /**
diff --git a/Engine/Plugins/Cameras/GameplayCameras/Source/GameplayCamerasEditor/Public/Editors/ObjectTreeGraphConfig.h b/Engine/Plugins/Cameras/GameplayCameras/Source/GameplayCamerasEditor/Public/Editors/ObjectTreeGraphConfig.h
--- a/Engine/Plugins/Cameras/GameplayCameras/Source/GameplayCamerasEditor/Public/Editors/ObjectTreeGraphConfig.h
+++ b/Engine/Plugins/Cameras/GameplayCameras/Source/GameplayCamerasEditor/Public/Editors/ObjectTreeGraphConfig.h
@@ -22,19 +22,27 @@ class UObjectTreeGraphNode;
 struct FObjectTreeGraphConfig;
 
 DECLARE_DELEGATE_OneParam(FOnBuildObjectTreeGraphConfig, FObjectTreeGraphConfig& InOutConfig);
+
+DECLARE_DELEGATE_OneParam(FOnSetupNewObject, UObject*);
 DECLARE_DELEGATE_RetVal_OneParam(FText, FOnGetObjectClassDisplayName, const UClass*);
-DECLARE_DELEGATE_TwoParams(FOnFormatObjectDisplayName, const UObject*, FText&);
+
 DECLARE_DELEGATE_TwoParams(FOnGetGraphDisplayInfo, const UObjectTreeGraph*, FGraphDisplayInfo&);
+DECLARE_DELEGATE_TwoParams(FOnFormatObjectDisplayName, const UObject*, FText&);
 
 #define OTGCC_FIELD(FieldType, FieldName)\
 	public:\
 		typename TCallTraits<FieldType>::ConstReference FieldName() const\
 			{ return _##FieldName; }\
 		FObjectTreeGraphClassConfig& FieldName(typename TCallTraits<FieldType>::ParamType InValue)\
-			{ _##FieldName = InValue; return *this; }\
+			{ _##FieldName = InValue; _bOverride##FieldName = true; return *this; }\
+		bool Has##FieldName##Override() const\
+			{ return _bOverride##FieldName; }\
 	private:\
 		FieldType _##FieldName;
 
+#define OTGCC_FIELD_FLAG(FieldName)\
+		bool _bOverride##FieldName;
+
 /**
  * A structure providing optional configuration options for a given object class.
  */
@@ -62,9 +70,12 @@ public:
 	/** Color of the graph node's body. */
 	OTGCC_FIELD(TOptional<FLinearColor>, NodeBodyTintColor)
 
+	/** A custom callback to setup a newly created object added in the graph editor. */
+	OTGCC_FIELD(FOnSetupNewObject, OnSetupNewObject)
+
 	/** Whether the graph node title uses the underlying object's name instead of its class name. */
 	OTGCC_FIELD(bool, NodeTitleUsesObjectName)
-	/** A custom call back to get the object's display name used in the graph node title. */
+	/** A custom callback to get the object's display name used in the graph node title. */
 	OTGCC_FIELD(FOnGetObjectClassDisplayName, OnGetObjectClassDisplayName)
 
 	/** Whether users can create new objects of this class in the graph. */
@@ -75,6 +86,24 @@ public:
 	/** The metadata specifier to look for in order to categorize the 'create node' action for this class. */
 	OTGCC_FIELD(FName, CreateCategoryMetaData)
 
+private:
+
+	// Packed flags.
+	OTGCC_FIELD_FLAG(GraphNodeClass)
+	OTGCC_FIELD_FLAG(SelfPinName)
+	OTGCC_FIELD_FLAG(SelfPinFriendlyName)
+	OTGCC_FIELD_FLAG(SelfPinDirectionOverride)
+	OTGCC_FIELD_FLAG(HasSelfPin)
+	OTGCC_FIELD_FLAG(DefaultPropertyPinDirectionOverride)
+	OTGCC_FIELD_FLAG(NodeTitleColor)
+	OTGCC_FIELD_FLAG(NodeBodyTintColor)
+	OTGCC_FIELD_FLAG(OnSetupNewObject)
+	OTGCC_FIELD_FLAG(NodeTitleUsesObjectName)
+	OTGCC_FIELD_FLAG(OnGetObjectClassDisplayName)
+	OTGCC_FIELD_FLAG(CanCreateNew)
+	OTGCC_FIELD_FLAG(CanDelete)
+	OTGCC_FIELD_FLAG(CreateCategoryMetaData)
+
 public:
 
 	FObjectTreeGraphClassConfig();
@@ -127,6 +156,89 @@ private:
 	TMap<FName, EEdGraphPinDirection> _PropertyPinDirectionOverrides;
 };
 
+#undef OTGCC_FIELD
+#undef OTGCC_FIELD_FLAG
+
+#define OTGCCS_FIELD(FieldType, FieldName)\
+	public:\
+		typename TCallTraits<FieldType>::ConstReference FieldName() const\
+		{\
+			for (const FObjectTreeGraphClassConfig* InnerConfig : InnerConfigs)\
+			{\
+				if (InnerConfig->Has##FieldName##Override())\
+				{\
+					return InnerConfig->FieldName();\
+				}\
+			}\
+			return DefaultConfig.FieldName();\
+		}
+
+/**
+ * A composite of multiple class configurations, for handling configuration options set on
+ * different classes in a class hierarchy.
+ */
+struct FObjectTreeGraphClassConfigs
+{
+public:
+
+	/** The subclass of graph nodes to create. */
+	OTGCCS_FIELD(TSubclassOf<UObjectTreeGraphNode>, GraphNodeClass)
+
+	/** The name of the self pin. */
+	OTGCCS_FIELD(FName, SelfPinName)
+	/** The display name of the self pin. */
+	OTGCCS_FIELD(FText, SelfPinFriendlyName)
+	/** The direction of the self pin. */
+	OTGCCS_FIELD(TOptional<EEdGraphPinDirection>, SelfPinDirectionOverride)
+	/** Whether graph nodes for this class have a self pin. */
+	OTGCCS_FIELD(bool, HasSelfPin)
+
+	/** Default direction of property pins. */
+	OTGCCS_FIELD(TOptional<EEdGraphPinDirection>, DefaultPropertyPinDirectionOverride)
+
+	/** Color of the graph node's title. */
+	OTGCCS_FIELD(TOptional<FLinearColor>, NodeTitleColor)
+	/** Color of the graph node's body. */
+	OTGCCS_FIELD(TOptional<FLinearColor>, NodeBodyTintColor)
+
+	/** A custom callback to setup a newly created object added in the graph editor. */
+	OTGCCS_FIELD(FOnSetupNewObject, OnSetupNewObject)
+
+	/** Whether the graph node title uses the underlying object's name instead of its class name. */
+	OTGCCS_FIELD(bool, NodeTitleUsesObjectName)
+	/** A custom callback to get the object's display name used in the graph node title. */
+	OTGCCS_FIELD(FOnGetObjectClassDisplayName, OnGetObjectClassDisplayName)
+
+	/** Whether users can create new objects of this class in the graph. */
+	OTGCCS_FIELD(bool, CanCreateNew)
+	/** Whether users can duplicate objects of this class in the graph. */
+	OTGCCS_FIELD(bool, CanDelete)
+
+	/** The metadata specifier to look for in order to categorize the 'create node' action for this class. */
+	OTGCCS_FIELD(FName, CreateCategoryMetaData)
+
+public:
+
+	FObjectTreeGraphClassConfigs();
+	FObjectTreeGraphClassConfigs(TArrayView<const FObjectTreeGraphClassConfig*> InClassConfigs);
+
+public:
+
+	/** Gets the name suffixes to strip. */
+	void GetStripDisplayNameSuffixes(TArray<FString>& OutSuffixes) const;
+
+	/** Gets the custom property pin direction for a given named property. */
+	TOptional<EEdGraphPinDirection> GetPropertyPinDirectionOverride(const FName& InPropertyName) const;
+
+private:
+
+	static FObjectTreeGraphClassConfig DefaultConfig;
+
+	TArray<const FObjectTreeGraphClassConfig*, TInlineAllocator<2>> InnerConfigs;
+};
+
+#undef OTGCCS_FIELD
+
 /**
  * A structure that provides all the information needed to build, edit, and maintain an 
  * object tree graph.
@@ -215,7 +327,7 @@ public:
 	void GetConnectableClasses(TArray<UClass*>& OutClasses, bool bPlaceableOnly = false);
 	
 	/** Gets the advanced class-specific configuration for the given class. */
-	const FObjectTreeGraphClassConfig& GetObjectClassConfig(const UClass* InObjectClass) const;
+	FObjectTreeGraphClassConfigs GetObjectClassConfigs(const UClass* InObjectClass) const;
 
 	/** Computes the display name of the given object. */
 	FText GetDisplayNameText(const UObject* InObject) const;
@@ -230,7 +342,7 @@ public:
 
 private:
 
-	FText GetDisplayNameText(const UClass* InClass, const FObjectTreeGraphClassConfig& InClassConfig) const;
-	void FormatDisplayNameText(const UObject* InObject, const FObjectTreeGraphClassConfig& InClassConfig, FText& InOutDisplayNameText) const;
+	FText GetDisplayNameText(const UClass* InClass, const FObjectTreeGraphClassConfigs& InClassConfig) const;
+	void FormatDisplayNameText(const UObject* InObject, const FObjectTreeGraphClassConfigs& InClassConfig, FText& InOutDisplayNameText) const;
 };
 
diff --git a/Engine/Plugins/Cameras/GameplayCameras/Source/GameplayCamerasEditor/Public/Editors/ObjectTreeGraphNode.h b/Engine/Plugins/Cameras/GameplayCameras/Source/GameplayCamerasEditor/Public/Editors/ObjectTreeGraphNode.h
--- a/Engine/Plugins/Cameras/GameplayCameras/Source/GameplayCamerasEditor/Public/Editors/ObjectTreeGraphNode.h
+++ b/Engine/Plugins/Cameras/GameplayCameras/Source/GameplayCamerasEditor/Public/Editors/ObjectTreeGraphNode.h
@@ -4,6 +4,7 @@
 
 #include "Containers/Array.h"
 #include "EdGraph/EdGraphNode.h"
+#include "Editors/ObjectTreeGraphConfig.h"
 #include "UObject/UObjectGlobals.h"
 
 #include "ObjectTreeGraphNode.generated.h"
@@ -11,8 +12,6 @@
 class FObjectProperty;
 class UEdGraphPin;
 class UObjectTreeGraph;
-struct FObjectTreeGraphConfig;
-struct FObjectTreeGraphClassConfig;
 
 /**
  * A graph node that represents an object inside an object tree graph.
@@ -47,10 +46,10 @@ public:
 
 	/** Finds the pin for the given object property. */
 	UEdGraphPin* GetPinForProperty(FObjectProperty* InProperty) const;
+	/** Finds the invisible parent pin for the given array property. */
+	UEdGraphPin* GetPinForProperty(FArrayProperty* InProperty) const;
 	/** Finds the pin for the given item in an array property. */
 	UEdGraphPin* GetPinForProperty(FArrayProperty* InProperty, int32 Index) const;
-	/** Finds the extra free pin used to add new items in an array property. */
-	UEdGraphPin* GetPinForPropertyNewItem(FArrayProperty* InProperty, bool bCreateNew);
 	/** Gets the underlying property represented by the given pin. */
 	FProperty* GetPropertyForPin(const UEdGraphPin* InPin) const;
 	/** Gets the type of object that can connect to the given pin. */
@@ -87,8 +86,10 @@ public:
 public:
 
 	// Internal API.
-	void CreateNewItemPin(FArrayProperty& InArrayProperty);
-	void CreateNewItemPin(UEdGraphPin* InParentArrayPin);
+	void GetArrayProperties(TArray<FArrayProperty*>& OutArrayProperties) const;
+	void CreateNewItemPins(FArrayProperty& InArrayProperty, int32 NumExtraPins);
+	void CreateNewItemPins(UEdGraphPin* InParentArrayPin, int32 NumExtraPins);
+	void InsertNewItemPin(UEdGraphPin* InParentArrayPin, int32 Index);
 	void RemoveItemPin(UEdGraphPin* InItemPin);
 	void RefreshArrayPropertyPinNames();
 
@@ -99,11 +100,11 @@ protected:
 		UClass* ObjectClass;
 		UObjectTreeGraph* Graph;
 		const FObjectTreeGraphConfig& GraphConfig;
-		const FObjectTreeGraphClassConfig& ObjectClassConfig;
+		const FObjectTreeGraphClassConfigs ObjectClassConfigs;
 	};
 
 	FNodeContext GetNodeContext() const;
-	const FObjectTreeGraphClassConfig& GetObjectClassConfig() const;
+	const FObjectTreeGraphClassConfigs GetObjectClassConfigs() const;
 
 private:
 
diff --git a/Engine/Plugins/Cameras/GameplayCameras/Source/GameplayCamerasEditor/Public/Editors/ObjectTreeGraphSchema.h b/Engine/Plugins/Cameras/GameplayCameras/Source/GameplayCamerasEditor/Public/Editors/ObjectTreeGraphSchema.h
--- a/Engine/Plugins/Cameras/GameplayCameras/Source/GameplayCamerasEditor/Public/Editors/ObjectTreeGraphSchema.h
+++ b/Engine/Plugins/Cameras/GameplayCameras/Source/GameplayCamerasEditor/Public/Editors/ObjectTreeGraphSchema.h
@@ -12,7 +12,7 @@ class IObjectTreeGraphRootObject;
 class UEdGraph;
 class UObjectTreeGraph;
 class UObjectTreeGraphNode;
-struct FObjectTreeGraphClassConfig;
+struct FObjectTreeGraphClassConfigs;
 
 /**
  * Schema class for an object tree graph.
@@ -59,6 +59,18 @@ public:
 	/** Checks if the given text is suitable for importing. */
 	bool CanImportNodesFromText(UObjectTreeGraph* InGraph, const FString& TextToImport) const;
 
+	/** Inserts a new array property item pin. */
+	void InsertArrayItemPin(UEdGraphPin* ArrayPin, int32 Index = INDEX_NONE) const;
+
+	/** Inserts a new array property item pin before the given pin. */
+	void InsertArrayItemPinBefore(UEdGraphPin* ArrayItemPin) const;
+
+	/** Inserts a new array property item pin after the given pin. */
+	void InsertArrayItemPinAfter(UEdGraphPin* ArrayItemPin) const;
+
+	/** Removes an array property item pin. */
+	void RemoveArrayItemPin(UEdGraphPin* ArrayItemPin) const;
+
 public:
 
 	// UEdGraphSchema interface.
@@ -85,19 +97,6 @@ protected:
 		TMap<UObject*, UObjectTreeGraphNode*> CreatedNodes;
 	};
 
-	struct FDelayedPinActions
-	{
-		void CreateNewItemPin(UObjectTreeGraphNode* Node, FArrayProperty* ArrayProperty);
-		void RemoveItemPin(UEdGraphPin* Pin);
-
-		bool IsEmpty() const;
-		void Apply();
-
-	private:
-		TArray<TTuple<UObjectTreeGraphNode*, FArrayProperty*>> ItemPinsToCreate;
-		TArray<UEdGraphPin*> ItemPinsToRemove;
-	};
-
 	// UObjectTreeGraphSchema interface.
 	virtual void CollectAllObjects(UObjectTreeGraph* InGraph, TSet<UObject*>& OutAllObjects) const;
 	virtual void OnCreateAllNodes(UObjectTreeGraph* InGraph, const FCreatedNodes& InCreatedNodes) const;
@@ -105,9 +104,9 @@ protected:
 	virtual void OnAddConnectableObject(UObjectTreeGraph* InGraph, UObjectTreeGraphNode* InNewNode) const;
 	virtual void OnRemoveConnectableObject(UObjectTreeGraph* InGraph, UObjectTreeGraphNode* InRemovedNode) const;
 	virtual void CopyNonObjectNodes(TArrayView<UObject*> InObjects, FStringOutputDevice& OutDevice) const;
-	virtual bool OnApplyConnection(UEdGraphPin* A, UEdGraphPin* B, FDelayedPinActions& Actions) const;
-	virtual bool OnApplyDisconnection(UEdGraphPin* TargetPin, FDelayedPinActions& Actions, bool bIsReconnecting) const;
-	virtual bool OnApplyDisconnection(UEdGraphPin* SourcePin, UEdGraphPin* TargetPin, FDelayedPinActions& Actions) const;
+	virtual bool OnApplyConnection(UEdGraphPin* A, UEdGraphPin* B) const;
+	virtual bool OnApplyDisconnection(UEdGraphPin* TargetPin) const;
+	virtual bool OnApplyDisconnection(UEdGraphPin* SourcePin, UEdGraphPin* TargetPin) const;
 	virtual void OnDeleteNodeFromGraph(UObjectTreeGraph* Graph, UEdGraphNode* Node) const;
 	virtual void FilterGraphContextPlaceableClasses(TArray<UClass*>& InOutClasses) const;
 
@@ -116,12 +115,12 @@ protected:
 	static void CollectAllReferencedObjects(UObjectTreeGraph* InGraph, TSet<UObject*>& OutAllObjects);
 	static bool CollectAllConnectableObjectsFromRootInterface(UObjectTreeGraph* InGraph, TSet<UObject*>& OutAllObjects, bool bAllowNoRootInterface);
 
-	const FObjectTreeGraphClassConfig& GetObjectClassConfig(const UObjectTreeGraphNode* InNode) const;
-	const FObjectTreeGraphClassConfig& GetObjectClassConfig(const UObjectTreeGraph* InGraph, UClass* InObjectClass) const;
+	const FObjectTreeGraphClassConfigs GetObjectClassConfigs(const UObjectTreeGraphNode* InNode) const;
+	const FObjectTreeGraphClassConfigs GetObjectClassConfigs(const UObjectTreeGraph* InGraph, UClass* InObjectClass) const;
 
-	void ApplyConnection(UEdGraphPin* A, UEdGraphPin* B, FDelayedPinActions& Actions) const;
-	void ApplyDisconnection(UEdGraphPin* TargetPin, FDelayedPinActions& Actions, bool bIsReconnecting) const;
-	void ApplyDisconnection(UEdGraphPin* SourcePin, UEdGraphPin* TargetPin, FDelayedPinActions& Actions) const;
+	void ApplyConnection(UEdGraphPin* A, UEdGraphPin* B) const;
+	void ApplyDisconnection(UEdGraphPin* TargetPin) const;
+	void ApplyDisconnection(UEdGraphPin* SourcePin, UEdGraphPin* TargetPin) const;
 
 private:
 
diff --git a/Engine/Plugins/Cameras/GameplayCameras/Source/GameplayCamerasEditor/Public/Editors/SObjectTreeGraphEditor.h b/Engine/Plugins/Cameras/GameplayCameras/Source/GameplayCamerasEditor/Public/Editors/SObjectTreeGraphEditor.h
--- a/Engine/Plugins/Cameras/GameplayCameras/Source/GameplayCamerasEditor/Public/Editors/SObjectTreeGraphEditor.h
+++ b/Engine/Plugins/Cameras/GameplayCameras/Source/GameplayCamerasEditor/Public/Editors/SObjectTreeGraphEditor.h
@@ -100,6 +100,10 @@ protected:
 	void OnDistributeNodesHorizontally();
 	void OnDistributeNodesVertically();
 
+	void OnInsertArrayItemPinBefore();
+	void OnInsertArrayItemPinAfter();
+	void OnRemoveArrayItemPin();
+
 	TArray<UClass*> FilterPlaceableObjectClasses(TArrayView<UClass* const> InObjectClasses);
 
 protected:
diff --git a/Engine/Plugins/Cameras/GameplayCameras/Source/GameplayCamerasEditor/Public/Editors/SObjectTreeGraphNode.h b/Engine/Plugins/Cameras/GameplayCameras/Source/GameplayCamerasEditor/Public/Editors/SObjectTreeGraphNode.h
--- a/Engine/Plugins/Cameras/GameplayCameras/Source/GameplayCamerasEditor/Public/Editors/SObjectTreeGraphNode.h
+++ b/Engine/Plugins/Cameras/GameplayCameras/Source/GameplayCamerasEditor/Public/Editors/SObjectTreeGraphNode.h
@@ -5,6 +5,7 @@
 #include "SGraphNode.h"
 #include "SGraphPin.h"
 
+class FArrayProperty;
 class UObjectTreeGraphNode;
 
 /**
@@ -26,11 +27,22 @@ public:
 
 public:
 
+	// SGraphNode interface.
+	virtual void CreateOutputSideAddButton(TSharedPtr<SVerticalBox> OutputBox) override;
+
 	// SNodePanel::SNode interface.
 	virtual void MoveTo(const FVector2D& NewPosition, FNodeSet& NodeFilter, bool bMarkDirty=true) override;
 
+protected:
+
+	TSharedRef<SWidget> MakeAddArrayPropertyPinButton(FArrayProperty* ArrayProperty);
+
+	FReply OnAddArrayPropertyPin(FArrayProperty* ArrayProperty);
+
 protected:
 
 	UObjectTreeGraphNode* ObjectGraphNode;
+
+	bool bHasAddPinButtons = false;
 };
 
