Cameras: add debug renderer flag for when we are outside of the camera system, such as a debug camera.

diff --git a/Engine/Plugins/Cameras/GameplayCameras/Source/GameplayCameras/Private/Core/CameraSystemEvaluator.cpp b/Engine/Plugins/Cameras/GameplayCameras/Source/GameplayCameras/Private/Core/CameraSystemEvaluator.cpp
--- a/Engine/Plugins/Cameras/GameplayCameras/Source/GameplayCameras/Private/Core/CameraSystemEvaluator.cpp
+++ b/Engine/Plugins/Cameras/GameplayCameras/Source/GameplayCameras/Private/Core/CameraSystemEvaluator.cpp
@@ -395,7 +395,7 @@ void FCameraSystemEvaluator::DebugUpdate(const FCameraSystemDebugUpdateParams& P
 	}
 #endif
 	
-	FCameraDebugRenderer Renderer(OwnerWorld, Params.CanvasObject);
+	FCameraDebugRenderer Renderer(OwnerWorld, Params.CanvasObject, Params.bIsDebugCameraEnabled);
 	RootDebugBlock->RootDebugDraw(Renderer);
 }
 
diff --git a/Engine/Plugins/Cameras/GameplayCameras/Source/GameplayCameras/Private/Debug/CameraDebugRenderer.cpp b/Engine/Plugins/Cameras/GameplayCameras/Source/GameplayCameras/Private/Debug/CameraDebugRenderer.cpp
--- a/Engine/Plugins/Cameras/GameplayCameras/Source/GameplayCameras/Private/Debug/CameraDebugRenderer.cpp
+++ b/Engine/Plugins/Cameras/GameplayCameras/Source/GameplayCameras/Private/Debug/CameraDebugRenderer.cpp
@@ -88,9 +88,10 @@ static FAutoConsoleVariableRef CVarGameplayCamerasDebugDefaultCoordinateSystemAx
 	GGameplayCamerasDebugDefaultCoordinateSystemAxesLength,
 	TEXT("Default: 100. The default length of coordinate system axes."));
 
-FCameraDebugRenderer::FCameraDebugRenderer(UWorld* InWorld, UCanvas* InCanvasObject)
+FCameraDebugRenderer::FCameraDebugRenderer(UWorld* InWorld, UCanvas* InCanvasObject, bool bInIsExternalRendering)
 	: World(InWorld)
 	, CanvasObject(InCanvasObject)
+	, bIsExternalRendering(bInIsExternalRendering)
 	, DrawColor(FColor::White)
 {
 	RenderFont = GEngine->GetSmallFont();
diff --git a/Engine/Plugins/Cameras/GameplayCameras/Source/GameplayCameras/Private/GameFramework/GameplayCameraSystemComponent.cpp b/Engine/Plugins/Cameras/GameplayCameras/Source/GameplayCameras/Private/GameFramework/GameplayCameraSystemComponent.cpp
--- a/Engine/Plugins/Cameras/GameplayCameras/Source/GameplayCameras/Private/GameFramework/GameplayCameraSystemComponent.cpp
+++ b/Engine/Plugins/Cameras/GameplayCameras/Source/GameplayCameras/Private/GameFramework/GameplayCameraSystemComponent.cpp
@@ -257,8 +257,14 @@ void UGameplayCameraSystemComponent::DebugDraw(UCanvas* Canvas, APlayerControlle
 	TSharedPtr<FCameraSystemEvaluator> Evaluator = GetCameraSystemEvaluator(false);
 	if (Evaluator.IsValid())
 	{
+		// We're looking from the outside if we are not the view target, or if we don't have a player
+		// anymore (which happens in spectator mode like with the debug camera).
+		APlayerController* ActualPlayerController = WeakPlayerController.Get();
+		const bool bIsDebugCameraEnabled = (ActualPlayerController && (ActualPlayerController->GetViewTarget() != GetOwner() || !ActualPlayerController->Player));
+
 		FCameraSystemDebugUpdateParams DebugUpdateParams;
 		DebugUpdateParams.CanvasObject = Canvas;
+		DebugUpdateParams.bIsDebugCameraEnabled = bIsDebugCameraEnabled;
 		Evaluator->DebugUpdate(DebugUpdateParams);
 	}
 }
diff --git a/Engine/Plugins/Cameras/GameplayCameras/Source/GameplayCameras/Public/Core/CameraSystemEvaluator.h b/Engine/Plugins/Cameras/GameplayCameras/Source/GameplayCameras/Public/Core/CameraSystemEvaluator.h
--- a/Engine/Plugins/Cameras/GameplayCameras/Source/GameplayCameras/Public/Core/CameraSystemEvaluator.h
+++ b/Engine/Plugins/Cameras/GameplayCameras/Source/GameplayCameras/Public/Core/CameraSystemEvaluator.h
@@ -87,6 +87,7 @@ public:
 struct FCameraSystemDebugUpdateParams
 {
 	UCanvas* CanvasObject = nullptr;
+	bool bIsDebugCameraEnabled = false;
 };
 #endif  // UE_GAMEPLAY_CAMERAS_DEBUG
 
diff --git a/Engine/Plugins/Cameras/GameplayCameras/Source/GameplayCameras/Public/Debug/CameraDebugRenderer.h b/Engine/Plugins/Cameras/GameplayCameras/Source/GameplayCameras/Public/Debug/CameraDebugRenderer.h
--- a/Engine/Plugins/Cameras/GameplayCameras/Source/GameplayCameras/Public/Debug/CameraDebugRenderer.h
+++ b/Engine/Plugins/Cameras/GameplayCameras/Source/GameplayCameras/Public/Debug/CameraDebugRenderer.h
@@ -41,10 +41,13 @@ class FCameraDebugRenderer
 public:
 
 	/** Creates a new debug renderer. */
-	GAMEPLAYCAMERAS_API FCameraDebugRenderer(UWorld* InWorld, UCanvas* InCanvasObject);
+	GAMEPLAYCAMERAS_API FCameraDebugRenderer(UWorld* InWorld, UCanvas* InCanvasObject, bool bInIsExternalRendering = false);
 	/** Destroys the debug renderer. */
 	GAMEPLAYCAMERAS_API ~FCameraDebugRenderer();
 
+	/** Whether we are looking at the camera system from the "outside". */
+	bool IsExternalRendering() const { return bIsExternalRendering; }
+
 	/** Adds text to the text wall. */
 	GAMEPLAYCAMERAS_API void AddText(const FString& InString);
 	GAMEPLAYCAMERAS_API void AddText(const TCHAR* Fmt, ...);
@@ -155,6 +158,9 @@ private:
 	UWorld* World;
 	/** The canvas used to draw the text wall. */
 	UCanvas* CanvasObject;
+	/** Whether we are looking from the "outside" of the camera system. */
+	bool bIsExternalRendering = false;
+
 	/** The draw color of the canvas. */
 	FColor DrawColor;
 
