Cameras: add new debug rendering methods, and draw the camera pose when in external mode

diff --git a/Engine/Plugins/Cameras/GameplayCameras/Source/GameplayCameras/Private/Debug/CameraDebugRenderer.cpp b/Engine/Plugins/Cameras/GameplayCameras/Source/GameplayCameras/Private/Debug/CameraDebugRenderer.cpp
--- a/Engine/Plugins/Cameras/GameplayCameras/Source/GameplayCameras/Private/Debug/CameraDebugRenderer.cpp
+++ b/Engine/Plugins/Cameras/GameplayCameras/Source/GameplayCameras/Private/Debug/CameraDebugRenderer.cpp
@@ -6,6 +6,7 @@
 #include "CanvasItem.h"
 #include "CanvasTypes.h"
 #include "Components/LineBatchComponent.h"
+#include "Core/CameraPose.h"
 #include "Debug/CameraDebugClock.h"
 #include "Debug/CameraDebugColors.h"
 #include "Debug/DebugTextRenderer.h"
@@ -82,6 +83,12 @@ static FAutoConsoleVariableRef CVarGameplayCamerasDebugMaxCardColumns(
 	GGameplayCamerasDebugMaxCardColumns,
 	TEXT("Default: 2. The number of columns to layout the debug cards (e.g. graphs, clocks, etc.)"));
 
+float GGameplayCamerasDebugDefaultCameraSize = 50.f;
+static FAutoConsoleVariableRef CVarGameplayCamerasDebugDefaultCameraSize(
+	TEXT("GameplayCameras.Debug.DefaultCameraSize"),
+	GGameplayCamerasDebugDefaultCameraSize,
+	TEXT("Default: 50. The default size of debug cameras."));
+
 float GGameplayCamerasDebugDefaultCoordinateSystemAxesLength = 100.f;
 static FAutoConsoleVariableRef CVarGameplayCamerasDebugDefaultCoordinateSystemAxesLength(
 	TEXT("GameplayCameras.Debug.DefaultCoordinateSystemAxesLength"),
@@ -272,6 +279,15 @@ void FCameraDebugRenderer::DrawClock(FCameraDebugClock& InClock, const FText& In
 	InClock.Draw(GetCanvas(), DrawParams);
 }
 
+void FCameraDebugRenderer::DrawCameraPose(const FCameraPose& InCameraPose, const FLinearColor& LineColor, float CameraSize)
+{
+	const FTransform3d Transform = InCameraPose.GetTransform();
+	const float EffectiveFieldOfView = InCameraPose.GetEffectiveFieldOfView();
+	const float AspectRatio = InCameraPose.GetSensorAspectRatio();
+	const double TargetDistance = InCameraPose.GetTargetDistance();
+	DrawCamera(Transform, EffectiveFieldOfView, AspectRatio, TargetDistance, LineColor, CameraSize, 1.f);
+}
+
 FVector2f FCameraDebugRenderer::GetNextCardPosition()
 {
 	const FVector2f Result(NextCardPosition);
@@ -301,6 +317,24 @@ void FCameraDebugRenderer::GetNextDrawGraphParams(FCameraDebugGraphDrawParams& O
 	OutDrawParams.GraphSize = FVector2f(GGameplayCamerasDebugCardWidth, GGameplayCamerasDebugCardHeight);
 }
 
+void FCameraDebugRenderer::Draw2DPointCross(const FVector2D& Location, float CrossSize, const FLinearColor& LineColor, float LineThickness)
+{
+	if (FCanvas* Canvas = GetCanvas())
+	{
+		const float HalfCrossSize = CrossSize / 2.f;
+
+		FCanvasLineItem Horizontal(Location - FVector2D(HalfCrossSize, 0), Location + FVector2D(HalfCrossSize, 0));
+		Horizontal.SetColor(LineColor);
+		Horizontal.LineThickness = LineThickness;
+		Canvas->DrawItem(Horizontal);
+
+		FCanvasLineItem Vertical(Location - FVector2D(0, HalfCrossSize), Location + FVector2D(0, HalfCrossSize));
+		Vertical.SetColor(LineColor);
+		Vertical.LineThickness = LineThickness;
+		Canvas->DrawItem(Vertical);
+	}
+}
+
 void FCameraDebugRenderer::Draw2DLine(const FVector2D& Start, const FVector2D& End, const FLinearColor& LineColor, float LineThickness)
 {
 	if (FCanvas* Canvas = GetCanvas())
@@ -355,6 +389,14 @@ void FCameraDebugRenderer::Draw2DCircle(const FVector2D& Center, float Radius, c
 	}
 }
 
+void FCameraDebugRenderer::DrawPoint(const FVector3d& Location, float PointSize, const FLinearColor& LineColor, float LineThickness)
+{
+	if (ULineBatchComponent* LineBatcher = GetDebugLineBatcher())
+	{
+		LineBatcher->DrawPoint(Location, LineColor, PointSize, SDPG_Foreground);
+	}
+}
+
 void FCameraDebugRenderer::DrawLine(const FVector3d& Start, const FVector3d& End, const FLinearColor& LineColor, float LineThickness)
 {
 	if (ULineBatchComponent* LineBatcher = GetDebugLineBatcher())
@@ -379,6 +421,70 @@ void FCameraDebugRenderer::DrawDirectionalArrow(const FVector3d& Start, const FV
 	}
 }
 
+void FCameraDebugRenderer::DrawCamera(const FTransform3d& Transform, float HorizontalFieldOfView, float AspectRatio, float TargetDistance, const FLinearColor& LineColor, float CameraSize, float LineThickness)
+{
+	if (AspectRatio <= 0.f)
+	{
+		AspectRatio = 1.f;
+	}
+	if (CameraSize <= 0.f)
+	{
+		CameraSize = GGameplayCamerasDebugDefaultCameraSize;
+	}
+
+	if (ULineBatchComponent* LineBatcher = GetDebugLineBatcher())
+	{
+		// We draw a pyramid representing the camera's FOV and aspect ratio. So we only need the origin
+		// point and the four corner points of the base.
+		const float TanHalfHFOV = FMath::Tan(FMath::DegreesToRadians(HorizontalFieldOfView / 2.f));
+		const float BaseHalfWidth = TanHalfHFOV * CameraSize;
+		const float BaseHalfHeight = BaseHalfWidth / AspectRatio;
+
+		const FVector3d ForwardDir(FVector3d::ForwardVector);
+		const FVector3d UpDir(FVector3d::UpVector);
+		const FVector3d RightDir(FVector3d::RightVector);
+
+		// Upper right, bottom right, bottom left, upper left.
+		FVector3d BaseCorners[4];
+		BaseCorners[0] = (ForwardDir * CameraSize) + (UpDir * BaseHalfHeight) + (RightDir * BaseHalfWidth);
+		BaseCorners[1] = (ForwardDir * CameraSize) - (UpDir * BaseHalfHeight) + (RightDir * BaseHalfWidth);
+		BaseCorners[2] = (ForwardDir * CameraSize) - (UpDir * BaseHalfHeight) - (RightDir * BaseHalfWidth);
+		BaseCorners[3] = (ForwardDir * CameraSize) + (UpDir * BaseHalfHeight) - (RightDir * BaseHalfWidth);
+
+		const FVector3d Location = Transform.GetLocation();
+		for (FVector3d& BaseCorner : BaseCorners)
+		{
+			BaseCorner = Location + Transform.TransformVectorNoScale(BaseCorner);
+		}
+
+		TArray<FBatchedLine> BatchedLines;
+		// Pyramid corners.
+		BatchedLines.Emplace(Location, BaseCorners[0], LineColor, 0.f, LineThickness, SDPG_Foreground);
+		BatchedLines.Emplace(Location, BaseCorners[1], LineColor, 0.f, LineThickness, SDPG_Foreground);
+		BatchedLines.Emplace(Location, BaseCorners[2], LineColor, 0.f, LineThickness, SDPG_Foreground);
+		BatchedLines.Emplace(Location, BaseCorners[3], LineColor, 0.f, LineThickness, SDPG_Foreground);
+		// Base edges.
+		BatchedLines.Emplace(BaseCorners[0], BaseCorners[1], LineColor, 0.f, LineThickness, SDPG_Foreground);
+		BatchedLines.Emplace(BaseCorners[1], BaseCorners[2], LineColor, 0.f, LineThickness, SDPG_Foreground);
+		BatchedLines.Emplace(BaseCorners[2], BaseCorners[3], LineColor, 0.f, LineThickness, SDPG_Foreground);
+		BatchedLines.Emplace(BaseCorners[3], BaseCorners[0], LineColor, 0.f, LineThickness, SDPG_Foreground);
+		// Base cross.
+		BatchedLines.Emplace(BaseCorners[0], BaseCorners[2], LineColor, 0.f, LineThickness, SDPG_Foreground);
+		BatchedLines.Emplace(BaseCorners[1], BaseCorners[3], LineColor, 0.f, LineThickness, SDPG_Foreground);
+
+		// Optional target distance line.
+		if (TargetDistance > 0.f)
+		{
+			const FVector3d AimDir = Transform.GetRotation().GetForwardVector();
+			BatchedLines.Emplace(
+					Location + AimDir * CameraSize, Location + AimDir * TargetDistance, 
+					LineColor, 0.f, LineThickness, SDPG_Foreground);
+		}
+
+		LineBatcher->DrawLines(BatchedLines);
+	}
+}
+
 void FCameraDebugRenderer::DrawCoordinateSystem(const FVector3d& Location, const FRotator3d& Rotation, float AxesLength)
 {
 	if (ULineBatchComponent* LineBatcher = GetDebugLineBatcher())
diff --git a/Engine/Plugins/Cameras/GameplayCameras/Source/GameplayCameras/Private/Debug/CameraPoseDebugBlock.cpp b/Engine/Plugins/Cameras/GameplayCameras/Source/GameplayCameras/Private/Debug/CameraPoseDebugBlock.cpp
--- a/Engine/Plugins/Cameras/GameplayCameras/Source/GameplayCameras/Private/Debug/CameraPoseDebugBlock.cpp
+++ b/Engine/Plugins/Cameras/GameplayCameras/Source/GameplayCameras/Private/Debug/CameraPoseDebugBlock.cpp
@@ -29,44 +29,54 @@ void DebugDrawCameraPoseField(FCameraDebugRenderer& Renderer, const TCHAR* Field
 UE_DEFINE_CAMERA_DEBUG_BLOCK(FCameraPoseDebugBlock)
 
 FCameraPoseDebugBlock::FCameraPoseDebugBlock()
+	: CameraPoseLineColor(FColorList::SlateBlue)
 {
 }
 
 FCameraPoseDebugBlock::FCameraPoseDebugBlock(const FCameraPose& InCameraPose)
 	: CameraPose(InCameraPose)
+	, CameraPoseLineColor(FColorList::SlateBlue)
 {
 }
 
 void FCameraPoseDebugBlock::OnDebugDraw(const FCameraDebugBlockDrawParams& Params, FCameraDebugRenderer& Renderer)
 {
-	bool bShowUnchanged = false;
-	if (!ShowUnchangedCVarName.IsEmpty())
+	if (bDrawText)
 	{
-		IConsoleVariable* ShowUnchangedCVar = IConsoleManager::Get().FindConsoleVariable(*ShowUnchangedCVarName, false);
-		if (ensureMsgf(ShowUnchangedCVar, TEXT("No such console variable: %s"), *ShowUnchangedCVarName))
+		bool bShowUnchanged = false;
+		if (!ShowUnchangedCVarName.IsEmpty())
 		{
-			bShowUnchanged = ShowUnchangedCVar->GetBool();
+			IConsoleVariable* ShowUnchangedCVar = IConsoleManager::Get().FindConsoleVariable(*ShowUnchangedCVarName, false);
+			if (ensureMsgf(ShowUnchangedCVar, TEXT("No such console variable: %s"), *ShowUnchangedCVarName))
+			{
+				bShowUnchanged = ShowUnchangedCVar->GetBool();
+			}
 		}
-	}
 
-	const FCameraDebugColors& Colors = FCameraDebugColors::Get();
-	const FColor ChangedColor = Colors.Default;
-	const FColor UnchangedColor = Colors.Passive;
+		const FCameraDebugColors& Colors = FCameraDebugColors::Get();
+		const FColor ChangedColor = Colors.Default;
+		const FColor UnchangedColor = Colors.Passive;
 
-	const FCameraPoseFlags& ChangedFlags = CameraPose.GetChangedFlags();
+		const FCameraPoseFlags& ChangedFlags = CameraPose.GetChangedFlags();
 
 #define UE_CAMERA_POSE_FOR_PROPERTY(PropType, PropName)\
-	if (bShowUnchanged || ChangedFlags.PropName)\
-	{\
-		const FColor& PropColor = ChangedFlags.PropName ? ChangedColor : UnchangedColor;\
-		Private::DebugDrawCameraPoseField<PropType>(Renderer, TEXT(#PropName), CameraPose.Get##PropName(), PropColor);\
-	}
-	UE_CAMERA_POSE_FOR_ALL_PROPERTIES()
+		if (bShowUnchanged || ChangedFlags.PropName)\
+		{\
+			const FColor& PropColor = ChangedFlags.PropName ? ChangedColor : UnchangedColor;\
+			Private::DebugDrawCameraPoseField<PropType>(Renderer, TEXT(#PropName), CameraPose.Get##PropName(), PropColor);\
+		}
+		UE_CAMERA_POSE_FOR_ALL_PROPERTIES()
 #undef UE_CAMERA_POSE_FOR_PROPERTY
 
-	Renderer.SetTextColor(Colors.Default);
-	Renderer.AddText(TEXT("Effective FOV  : %f\n"), CameraPose.GetEffectiveFieldOfView());
-	Renderer.AddText(TEXT("Effective Aspect Ratio  : %f\n"), CameraPose.GetSensorAspectRatio());
+			Renderer.SetTextColor(Colors.Default);
+		Renderer.AddText(TEXT("Effective FOV  : %f\n"), CameraPose.GetEffectiveFieldOfView());
+		Renderer.AddText(TEXT("Effective Aspect Ratio  : %f\n"), CameraPose.GetSensorAspectRatio());
+	}
+
+	if (bDrawInExternalRendering && Renderer.IsExternalRendering())
+	{
+		Renderer.DrawCameraPose(CameraPose, CameraPoseLineColor, CameraPoseSize);
+	}
 }
 
 void FCameraPoseDebugBlock::OnSerialize(FArchive& Ar)
diff --git a/Engine/Plugins/Cameras/GameplayCameras/Source/GameplayCameras/Private/Debug/RootCameraDebugBlock.cpp b/Engine/Plugins/Cameras/GameplayCameras/Source/GameplayCameras/Private/Debug/RootCameraDebugBlock.cpp
--- a/Engine/Plugins/Cameras/GameplayCameras/Source/GameplayCameras/Private/Debug/RootCameraDebugBlock.cpp
+++ b/Engine/Plugins/Cameras/GameplayCameras/Source/GameplayCameras/Private/Debug/RootCameraDebugBlock.cpp
@@ -98,6 +98,11 @@ void FRootCameraDebugBlock::BuildDebugBlocks(const FCameraSystemEvaluator& Camer
 		{
 			RootNodeEvaluator->BuildDebugBlocks(Params, Builder);
 		}
+
+		// Draw the final camera pose in external rendering, but don't draw text.
+		NodeTreeCategory.AddChild(
+				&Builder.BuildDebugBlock<FCameraPoseDebugBlock>(CameraSystem.GetEvaluatedResult().CameraPose)
+					.ShouldDrawText(false));
 	}
 	Builder.EndChildDebugBlock();
 
diff --git a/Engine/Plugins/Cameras/GameplayCameras/Source/GameplayCameras/Public/Debug/CameraDebugRenderer.h b/Engine/Plugins/Cameras/GameplayCameras/Source/GameplayCameras/Public/Debug/CameraDebugRenderer.h
--- a/Engine/Plugins/Cameras/GameplayCameras/Source/GameplayCameras/Public/Debug/CameraDebugRenderer.h
+++ b/Engine/Plugins/Cameras/GameplayCameras/Source/GameplayCameras/Public/Debug/CameraDebugRenderer.h
@@ -15,6 +15,7 @@ class UCanvas;
 class UFont;
 class ULineBatchComponent;
 class UWorld;
+struct FCameraPose;
 
 #if UE_GAMEPLAY_CAMERAS_DEBUG
 
@@ -84,8 +85,15 @@ public:
 	template<uint8 NumValues>
 	void DrawGraph(TCameraDebugGraph<NumValues>& InGraph, const FText& InGraphName);
 
+	/**
+	 * Draws the given camera pose using the DrawCamera method.
+	 */
+	GAMEPLAYCAMERAS_API void DrawCameraPose(const FCameraPose& InCameraPose, const FLinearColor& LineColor, float CameraSize = 0.f);
+
 public:
 
+	/** Draws a 2D cross at a point. */
+	GAMEPLAYCAMERAS_API void Draw2DPointCross(const FVector2D& Location, float CrossSize, const FLinearColor& LineColor, float LineThickness = 1.f);
 	/** Draws a 2D line. */
 	GAMEPLAYCAMERAS_API void Draw2DLine(const FVector2D& Start, const FVector2D& End, const FLinearColor& LineColor, float LineThickness = 1.f);
 	/** Draws a 2D box. */
@@ -95,12 +103,16 @@ public:
 	/** Draws a 2D circle. */
 	GAMEPLAYCAMERAS_API void Draw2DCircle(const FVector2D& Center, float Radius, const FLinearColor& LineColor, float LineThickness = 1.f, int32 NumSides = 0);
 
+	/** Draws a 3D point. */
+	GAMEPLAYCAMERAS_API void DrawPoint(const FVector3d& Location, float PointSize, const FLinearColor& LineColor, float LineThickness = 1.f);
 	/** Draws a 3D line. */
 	GAMEPLAYCAMERAS_API void DrawLine(const FVector3d& Start, const FVector3d& End, const FLinearColor& LineColor, float LineThickness = 1.f);
 	/** Draws a 3D sphere. */
-	GAMEPLAYCAMERAS_API void DrawSphere(const FVector3d& Center, float Radius, int32 Segments, const FLinearColor& LineColor, float LineThickness);
+	GAMEPLAYCAMERAS_API void DrawSphere(const FVector3d& Center, float Radius, int32 Segments, const FLinearColor& LineColor, float LineThickness = 1.f);
 	/** Draws a 3D arrow. */
-	GAMEPLAYCAMERAS_API void DrawDirectionalArrow(const FVector3d& Start, const FVector3d& End, float ArrowSize, const FLinearColor& LineColor, float LineThickness);
+	GAMEPLAYCAMERAS_API void DrawDirectionalArrow(const FVector3d& Start, const FVector3d& End, float ArrowSize, const FLinearColor& LineColor, float LineThickness = 1.f);
+	/** Draws a camera. */
+	GAMEPLAYCAMERAS_API void DrawCamera(const FTransform3d& Transform, float HorizontalFieldOfView, float AspectRatio, float TargetDistance, const FLinearColor& LineColor, float CameraSize = 0.f, float LineThickness = 1.f);
 	/** Draws a 3D coordinate system. */
 	GAMEPLAYCAMERAS_API void DrawCoordinateSystem(const FVector3d& Location, const FRotator3d& Rotation, float AxesLength = 0.f);
 	/** Draws a 3D coordinate system. */
diff --git a/Engine/Plugins/Cameras/GameplayCameras/Source/GameplayCameras/Public/Debug/CameraPoseDebugBlock.h b/Engine/Plugins/Cameras/GameplayCameras/Source/GameplayCameras/Public/Debug/CameraPoseDebugBlock.h
--- a/Engine/Plugins/Cameras/GameplayCameras/Source/GameplayCameras/Public/Debug/CameraPoseDebugBlock.h
+++ b/Engine/Plugins/Cameras/GameplayCameras/Source/GameplayCameras/Public/Debug/CameraPoseDebugBlock.h
@@ -24,6 +24,34 @@ public:
 	/** Creates a new camera pose debug block. */
 	FCameraPoseDebugBlock(const FCameraPose& InCameraPose);
 
+	/** Sets whether the camera pose values should be printed in the text HUD. */
+	FCameraPoseDebugBlock& ShouldDrawText(bool bShouldDraw)
+	{
+		bDrawText = bShouldDraw;
+		return *this;
+	}
+
+	/** Sets whether the camera pose should be drawn when in external debug rendering. */
+	FCameraPoseDebugBlock& ShouldDrawInExternalRendering(bool bShouldDraw)
+	{
+		bDrawInExternalRendering = bShouldDraw;
+		return *this;
+	}
+
+	/** Sets the external rendering color. */
+	FCameraPoseDebugBlock& SetExternalRenderingLineColor(const FLinearColor& LineColor)
+	{
+		CameraPoseLineColor = LineColor;
+		return *this;
+	}
+
+	/** Sets the external rendering size. */
+	FCameraPoseDebugBlock& SetExternalRenderingSize(float CameraSize)
+	{
+		CameraPoseSize = CameraSize;
+		return *this;
+	}
+
 	/** 
 	 * Specifies the console variable to use to toggle between only showing camera pose
 	 * properties that were written to, or showing all camera pose properties.
@@ -43,6 +71,10 @@ private:
 
 	FCameraPose CameraPose;
 	FString ShowUnchangedCVarName;
+	FLinearColor CameraPoseLineColor;
+	float CameraPoseSize = -1.f;
+	bool bDrawText = true;
+	bool bDrawInExternalRendering = true;
 };
 
 }  // namespace UE::Cameras
