Cameras: add ability for runtime code to react to property changes on camera nodes

This is useful for making live-editing immediately reactive when editing struct properties, although a given camera node evaluator needs to opt-in to it.

diff --git a/Engine/Plugins/Cameras/GameplayCameras/Source/GameplayCameras/Public/IGameplayCamerasLiveEditListener.h b/Engine/Plugins/Cameras/GameplayCameras/Source/GameplayCameras/Public/IGameplayCamerasLiveEditListener.h
--- a/Engine/Plugins/Cameras/GameplayCameras/Source/GameplayCameras/Public/IGameplayCamerasLiveEditListener.h
+++ b/Engine/Plugins/Cameras/GameplayCameras/Source/GameplayCameras/Public/IGameplayCamerasLiveEditListener.h
@@ -6,7 +6,9 @@
 #include "Templates/SharedPointer.h"
 #include "UObject/UnrealType.h"
 
+class UCameraNode;
 class UPackage;
+struct FPropertyChangedEvent;
 
 namespace UE::Cameras
 {
@@ -27,11 +29,17 @@ public:
 
 	virtual ~IGameplayCamerasLiveEditListener() {}
 
+	/** Called when a camera asset has been (re)built. */
 	void PostBuildAsset(const FGameplayCameraAssetBuildEvent& BuildEvent) { OnPostBuildAsset(BuildEvent); }
+	/** Called when a camera node's property has been changed. */
+	void PostEditChangeProperty(const UCameraNode* InCameraNode, const FPropertyChangedEvent& PropertyChangedEvent) { OnPostEditChangeProperty(InCameraNode, PropertyChangedEvent); }
 
 protected:
 
-	virtual void OnPostBuildAsset(const FGameplayCameraAssetBuildEvent& BuildEvent) = 0;
+	/** Called when a camera asset has been (re)built. */
+	virtual void OnPostBuildAsset(const FGameplayCameraAssetBuildEvent& BuildEvent) {}
+	/** Called when a camera node's property has been changed. */
+	virtual void OnPostEditChangeProperty(const UCameraNode* InCameraNode, const FPropertyChangedEvent& PropertyChangedEvent) {}
 };
 
 #endif  // WITH_EDITOR
diff --git a/Engine/Plugins/Cameras/GameplayCameras/Source/GameplayCameras/Public/IGameplayCamerasLiveEditManager.h b/Engine/Plugins/Cameras/GameplayCameras/Source/GameplayCameras/Public/IGameplayCamerasLiveEditManager.h
--- a/Engine/Plugins/Cameras/GameplayCameras/Source/GameplayCameras/Public/IGameplayCamerasLiveEditManager.h
+++ b/Engine/Plugins/Cameras/GameplayCameras/Source/GameplayCameras/Public/IGameplayCamerasLiveEditManager.h
@@ -6,7 +6,9 @@
 #include "Templates/SharedPointer.h"
 #include "UObject/UnrealType.h"
 
+class UCameraNode;
 class UPackage;
+struct FPropertyChangedEvent;
 
 namespace UE::Cameras
 {
@@ -29,8 +31,19 @@ public:
 
 	/** Add a listener for the given package. */
 	virtual void AddListener(const UPackage* InAssetPackage, IGameplayCamerasLiveEditListener* Listener) = 0;
-	/** Removes a listener for the given package. */
+	/** Remove a listener for the given package. */
 	virtual void RemoveListener(const UPackage* InAssetPackage, IGameplayCamerasLiveEditListener* Listener) = 0;
+
+	/** Notify all listeners that a property was changed on the given camera node. */
+	virtual void NotifyPostEditChangeProperty(const UCameraNode* InCameraNode, const FPropertyChangedEvent& PropertyChangedEvent) const = 0;
+
+	/** Add a listener for the given node. */
+	virtual void AddListener(const UCameraNode* InCameraNode, IGameplayCamerasLiveEditListener* Listener) = 0;
+	/** Remove a listener for the given node. */
+	virtual void RemoveListener(const UCameraNode* InCameraNode, IGameplayCamerasLiveEditListener* Listener) = 0;
+
+	/** Remove the given listener from all notifications. */
+	virtual void RemoveListener(IGameplayCamerasLiveEditListener* Listener) = 0;
 };
 
 #endif  // WITH_EDITOR
diff --git a/Engine/Plugins/Cameras/GameplayCameras/Source/GameplayCamerasEditor/Private/GameplayCamerasLiveEditManager.cpp b/Engine/Plugins/Cameras/GameplayCameras/Source/GameplayCamerasEditor/Private/GameplayCamerasLiveEditManager.cpp
--- a/Engine/Plugins/Cameras/GameplayCameras/Source/GameplayCamerasEditor/Private/GameplayCamerasLiveEditManager.cpp
+++ b/Engine/Plugins/Cameras/GameplayCameras/Source/GameplayCamerasEditor/Private/GameplayCamerasLiveEditManager.cpp
@@ -10,6 +10,47 @@
 namespace UE::Cameras
 {
 
+namespace Internal
+{
+
+using FListenerArray = TArray<IGameplayCamerasLiveEditListener*, TInlineAllocator<4>>;
+
+template<typename ObjectType>
+void AddListenerImpl(
+		TMap<TWeakObjectPtr<const ObjectType>, FListenerArray>& ListenerMap,
+		const ObjectType* Object,
+		IGameplayCamerasLiveEditListener* Listener)
+{
+	if (ensure(Object && Listener))
+	{
+		FListenerArray& Listeners = ListenerMap.FindOrAdd(Object);
+		Listeners.Add(Listener);
+	}
+}
+
+template<typename ObjectType>
+void RemoveListenerImpl(
+		TMap<TWeakObjectPtr<const ObjectType>, FListenerArray>& ListenerMap,
+		const ObjectType* Object,
+		IGameplayCamerasLiveEditListener* Listener)
+{
+	if (ensure(Object && Listener))
+	{
+		FListenerArray* Listeners = ListenerMap.Find(Object);
+		if (ensure(Listeners))
+		{
+			const int32 NumRemoved = Listeners->RemoveSwap(Listener);
+			ensure(NumRemoved == 1);
+			if (Listeners->IsEmpty())
+			{
+				ListenerMap.Remove(Object);
+			}
+		}
+	}
+}
+
+}  // namespace Internal
+
 FGameplayCamerasLiveEditManager::FGameplayCamerasLiveEditManager()
 {
 	FCoreUObjectDelegates::GetPostGarbageCollect().AddRaw(this, &FGameplayCamerasLiveEditManager::OnPostGarbageCollection);
@@ -22,7 +63,7 @@ FGameplayCamerasLiveEditManager::~FGameplayCamerasLiveEditManager()
 
 void FGameplayCamerasLiveEditManager::NotifyPostBuildAsset(const UPackage* InAssetPackage) const
 {
-	if (const FListenerArray* Listeners = ListenerMap.Find(InAssetPackage))
+	if (const FListenerArray* Listeners = PackageListenerMap.Find(InAssetPackage))
 	{
 		FGameplayCameraAssetBuildEvent BuildEvent;
 		BuildEvent.AssetPackage = InAssetPackage;
@@ -36,25 +77,53 @@ void FGameplayCamerasLiveEditManager::NotifyPostBuildAsset(const UPackage* InAss
 
 void FGameplayCamerasLiveEditManager::AddListener(const UPackage* InAssetPackage, IGameplayCamerasLiveEditListener* Listener)
 {
-	if (ensure(InAssetPackage && Listener))
+	Internal::AddListenerImpl(PackageListenerMap, InAssetPackage, Listener);
+}
+
+void FGameplayCamerasLiveEditManager::RemoveListener(const UPackage* InAssetPackage, IGameplayCamerasLiveEditListener* Listener)
+{
+	Internal::RemoveListenerImpl(PackageListenerMap, InAssetPackage, Listener);
+}
+
+void FGameplayCamerasLiveEditManager::NotifyPostEditChangeProperty(const UCameraNode* InCameraNode, const FPropertyChangedEvent& PropertyChangedEvent) const
+{
+	if (const FListenerArray* Listeners = NodeListenerMap.Find(InCameraNode))
 	{
-		FListenerArray& Listeners = ListenerMap.FindOrAdd(InAssetPackage);
-		Listeners.Add(Listener);
+		for (IGameplayCamerasLiveEditListener* Listener : *Listeners)
+		{
+			Listener->PostEditChangeProperty(InCameraNode, PropertyChangedEvent);
+		}
 	}
 }
 
-void FGameplayCamerasLiveEditManager::RemoveListener(const UPackage* InAssetPackage, IGameplayCamerasLiveEditListener* Listener)
+void FGameplayCamerasLiveEditManager::AddListener(const UCameraNode* InCameraNode, IGameplayCamerasLiveEditListener* Listener)
+{
+	Internal::AddListenerImpl(NodeListenerMap, InCameraNode, Listener);
+}
+
+void FGameplayCamerasLiveEditManager::RemoveListener(const UCameraNode* InCameraNode, IGameplayCamerasLiveEditListener* Listener)
+{
+	Internal::RemoveListenerImpl(NodeListenerMap, InCameraNode, Listener);
+}
+
+void FGameplayCamerasLiveEditManager::RemoveListener(IGameplayCamerasLiveEditListener* Listener)
 {
-	if (ensure(InAssetPackage && Listener))
+	if (ensure(Listener))
 	{
-		FListenerArray* Listeners = ListenerMap.Find(InAssetPackage);
-		if (ensure(Listeners))
+		for (auto It = PackageListenerMap.CreateIterator(); It; ++It)
 		{
-			const int32 NumRemoved = Listeners->RemoveSwap(Listener);
-			ensure(NumRemoved == 1);
-			if (Listeners->IsEmpty())
+			It.Value().Remove(Listener);
+			if (It.Value().IsEmpty())
+			{
+				It.RemoveCurrent();
+			}
+		}
+		for (auto It = NodeListenerMap.CreateIterator(); It; ++It)
+		{
+			It.Value().Remove(Listener);
+			if (It.Value().IsEmpty())
 			{
-				ListenerMap.Remove(InAssetPackage);
+				It.RemoveCurrent();
 			}
 		}
 	}
@@ -67,7 +136,7 @@ void FGameplayCamerasLiveEditManager::OnPostGarbageCollection()
 
 void FGameplayCamerasLiveEditManager::RemoveGarbage()
 {
-	for (auto It = ListenerMap.CreateIterator(); It; ++It)
+	for (auto It = PackageListenerMap.CreateIterator(); It; ++It)
 	{
 		if (!It.Key().IsValid())
 		{
diff --git a/Engine/Plugins/Cameras/GameplayCameras/Source/GameplayCamerasEditor/Private/GameplayCamerasLiveEditManager.h b/Engine/Plugins/Cameras/GameplayCameras/Source/GameplayCamerasEditor/Private/GameplayCamerasLiveEditManager.h
--- a/Engine/Plugins/Cameras/GameplayCameras/Source/GameplayCamerasEditor/Private/GameplayCamerasLiveEditManager.h
+++ b/Engine/Plugins/Cameras/GameplayCameras/Source/GameplayCamerasEditor/Private/GameplayCamerasLiveEditManager.h
@@ -4,11 +4,11 @@
 
 #include "IGameplayCamerasLiveEditManager.h"
 
+#include "Core/CameraNode.h"
 #include "CoreTypes.h"
+#include "UObject/Package.h"
 #include "UObject/WeakObjectPtr.h"
 
-class UPackage;
-
 namespace UE::Cameras
 {
 
@@ -25,6 +25,10 @@ public:
 	virtual void NotifyPostBuildAsset(const UPackage* InAssetPackage) const override;
 	virtual void AddListener(const UPackage* InAssetPackage, IGameplayCamerasLiveEditListener* Listener) override;
 	virtual void RemoveListener(const UPackage* InAssetPackage, IGameplayCamerasLiveEditListener* Listener) override;
+	virtual void NotifyPostEditChangeProperty(const UCameraNode* InCameraNode, const FPropertyChangedEvent& PropertyChangedEvent) const override;
+	virtual void AddListener(const UCameraNode* InCameraNode, IGameplayCamerasLiveEditListener* Listener) override;
+	virtual void RemoveListener(const UCameraNode* InCameraNode, IGameplayCamerasLiveEditListener* Listener) override;
+	virtual void RemoveListener(IGameplayCamerasLiveEditListener* Listener) override;
 
 private:
 
@@ -34,9 +38,13 @@ private:
 
 private:
 
-	using FListenerArray = TArray<IGameplayCamerasLiveEditListener*>;
-	using FListenerMap = TMap<TWeakObjectPtr<const UPackage>, FListenerArray>;
-	FListenerMap ListenerMap;
+	using FListenerArray = TArray<IGameplayCamerasLiveEditListener*, TInlineAllocator<4>>;
+
+	using FPackageListenerMap = TMap<TWeakObjectPtr<const UPackage>, FListenerArray>;
+	FPackageListenerMap PackageListenerMap;
+
+	using FNodeListenerMap = TMap<TWeakObjectPtr<const UCameraNode>, FListenerArray>;
+	FNodeListenerMap NodeListenerMap;
 };
 
 }  // namespace UE::Cameras
diff --git a/Engine/Plugins/Cameras/GameplayCameras/Source/GameplayCamerasEditor/Private/Toolkits/CameraRigAssetEditorToolkit.cpp b/Engine/Plugins/Cameras/GameplayCameras/Source/GameplayCamerasEditor/Private/Toolkits/CameraRigAssetEditorToolkit.cpp
--- a/Engine/Plugins/Cameras/GameplayCameras/Source/GameplayCamerasEditor/Private/Toolkits/CameraRigAssetEditorToolkit.cpp
+++ b/Engine/Plugins/Cameras/GameplayCameras/Source/GameplayCamerasEditor/Private/Toolkits/CameraRigAssetEditorToolkit.cpp
@@ -215,6 +215,8 @@ void FCameraRigAssetEditorToolkit::PostInitAssetEditor()
 	IGameplayCamerasModule& GameplayCamerasModule = FModuleManager::GetModuleChecked<IGameplayCamerasModule>("GameplayCameras");
 	LiveEditManager = GameplayCamerasModule.GetLiveEditManager();
 
+	Impl->SetLiveEditManager(LiveEditManager);
+
 	RegenerateMenusAndToolbars();
 }
 
diff --git a/Engine/Plugins/Cameras/GameplayCameras/Source/GameplayCamerasEditor/Private/Toolkits/CameraRigAssetEditorToolkitBase.cpp b/Engine/Plugins/Cameras/GameplayCameras/Source/GameplayCamerasEditor/Private/Toolkits/CameraRigAssetEditorToolkitBase.cpp
--- a/Engine/Plugins/Cameras/GameplayCameras/Source/GameplayCamerasEditor/Private/Toolkits/CameraRigAssetEditorToolkitBase.cpp
+++ b/Engine/Plugins/Cameras/GameplayCameras/Source/GameplayCamerasEditor/Private/Toolkits/CameraRigAssetEditorToolkitBase.cpp
@@ -3,11 +3,13 @@
 #include "Toolkits/CameraRigAssetEditorToolkitBase.h"
 
 #include "Commands/CameraRigAssetEditorCommands.h"
+#include "Core/CameraNode.h"
 #include "Core/CameraRigAsset.h"
 #include "Editors/SCameraRigAssetEditor.h"
 #include "Editors/SObjectTreeGraphToolbox.h"
 #include "Framework/Commands/UICommandList.h"
 #include "Framework/Docking/TabManager.h"
+#include "IGameplayCamerasLiveEditManager.h"
 #include "Modules/ModuleManager.h"
 #include "PropertyEditorModule.h"
 #include "Styles/GameplayCamerasEditorStyle.h"
@@ -183,6 +185,15 @@ void FCameraRigAssetEditorToolkitBase::NotifyPostChange(const FPropertyChangedEv
 	{
 		CameraRigAsset->BuildStatus = ECameraBuildStatus::Dirty;
 	}
+
+	if (LiveEditManager && PropertyChangedEvent.GetNumObjectsBeingEdited() > 0)
+	{
+		const UCameraNode* EditedCameraNode = Cast<UCameraNode>(PropertyChangedEvent.GetObjectBeingEdited(0));
+		if (EditedCameraNode)
+		{
+			LiveEditManager->NotifyPostEditChangeProperty(EditedCameraNode, PropertyChangedEvent);
+		}
+	}
 }
 
 void FCameraRigAssetEditorToolkitBase::SetCameraRigEditorMode(ECameraRigAssetEditorMode InEditorMode)
@@ -202,11 +213,14 @@ void FCameraRigAssetEditorToolkitBase::OnAnyGraphChanged(const FEdGraphEditActio
 	if (CameraRigAsset)
 	{
 		CameraRigAsset->BuildStatus = ECameraBuildStatus::Dirty;
-
-		CameraRigBuildStatusDirtied.Broadcast();
 	}
 }
 
+void FCameraRigAssetEditorToolkitBase::SetLiveEditManager(TSharedPtr<IGameplayCamerasLiveEditManager> InLiveEditManager)
+{
+	LiveEditManager = InLiveEditManager;
+}
+
 }  // namespace UE::Cameras
 
 #undef LOCTEXT_NAMESPACE
diff --git a/Engine/Plugins/Cameras/GameplayCameras/Source/GameplayCamerasEditor/Public/Toolkits/CameraRigAssetEditorToolkitBase.h b/Engine/Plugins/Cameras/GameplayCameras/Source/GameplayCamerasEditor/Public/Toolkits/CameraRigAssetEditorToolkitBase.h
--- a/Engine/Plugins/Cameras/GameplayCameras/Source/GameplayCamerasEditor/Public/Toolkits/CameraRigAssetEditorToolkitBase.h
+++ b/Engine/Plugins/Cameras/GameplayCameras/Source/GameplayCamerasEditor/Public/Toolkits/CameraRigAssetEditorToolkitBase.h
@@ -19,6 +19,7 @@ struct FEdGraphEditAction;
 namespace UE::Cameras
 {
 
+class IGameplayCamerasLiveEditManager;
 class FStandardToolkitLayout;
 class SCameraRigAssetEditor;
 enum class ECameraRigAssetEditorMode;
@@ -49,7 +50,7 @@ public:
 	void BuildToolbarMenu(UToolMenu* ToolbarMenu);
 	void BindCommands(TSharedRef<FUICommandList> CommandList);
 
-	FSimpleMulticastDelegate& OnCameraRigBuildStatusDirtied() { return CameraRigBuildStatusDirtied; }
+	void SetLiveEditManager(TSharedPtr<IGameplayCamerasLiveEditManager> InLiveEditManager);
 
 protected:
 
@@ -94,8 +95,8 @@ private:
 	/** Toolbox widget */
 	TSharedPtr<SObjectTreeGraphToolbox> ToolboxWidget;
 
-	/** A multicast delegate that is invoked when the camera rig asset is dirtied */
-	FSimpleMulticastDelegate CameraRigBuildStatusDirtied;
+	/** Live editing manager */
+	TSharedPtr<IGameplayCamerasLiveEditManager> LiveEditManager;
 };
 
 }  // namespace UE::Cameras
