Cameras: add composite curve property types

These types will be used by new camera nodes in later changelists.

diff --git a/Engine/Plugins/Cameras/GameplayCameras/Source/GameplayCameras/Private/Curves/CameraRotatorCurve.cpp b/Engine/Plugins/Cameras/GameplayCameras/Source/GameplayCameras/Private/Curves/CameraRotatorCurve.cpp
new file mode 100644
--- /dev/null
+++ b/Engine/Plugins/Cameras/GameplayCameras/Source/GameplayCameras/Private/Curves/CameraRotatorCurve.cpp
@@ -0,0 +1,15 @@
+// Copyright Epic Games, Inc. All Rights Reserved.
+
+#include "Curves/CameraRotatorCurve.h"
+
+#include "Math/Rotator.h"
+
+FRotator FCameraRotatorCurve::GetValue(float InTime) const
+{
+	FRotator Result;
+	Result.Pitch = Curves[0].Eval(InTime);
+	Result.Yaw = Curves[1].Eval(InTime);
+	Result.Roll = Curves[2].Eval(InTime);
+	return Result;
+}
+
diff --git a/Engine/Plugins/Cameras/GameplayCameras/Source/GameplayCameras/Private/Curves/CameraSingleCurve.cpp b/Engine/Plugins/Cameras/GameplayCameras/Source/GameplayCameras/Private/Curves/CameraSingleCurve.cpp
new file mode 100644
--- /dev/null
+++ b/Engine/Plugins/Cameras/GameplayCameras/Source/GameplayCameras/Private/Curves/CameraSingleCurve.cpp
@@ -0,0 +1,9 @@
+// Copyright Epic Games, Inc. All Rights Reserved.
+
+#include "Curves/CameraSingleCurve.h"
+
+float FCameraSingleCurve::GetValue(float InTime) const
+{
+	return Curve.Eval(InTime);
+}
+
diff --git a/Engine/Plugins/Cameras/GameplayCameras/Source/GameplayCameras/Private/Curves/CameraVectorCurve.cpp b/Engine/Plugins/Cameras/GameplayCameras/Source/GameplayCameras/Private/Curves/CameraVectorCurve.cpp
new file mode 100644
--- /dev/null
+++ b/Engine/Plugins/Cameras/GameplayCameras/Source/GameplayCameras/Private/Curves/CameraVectorCurve.cpp
@@ -0,0 +1,13 @@
+// Copyright Epic Games, Inc. All Rights Reserved.
+
+#include "Curves/CameraVectorCurve.h"
+
+FVector FCameraVectorCurve::GetValue(float InTime) const
+{
+	FVector Result;
+	Result.X = Curves[0].Eval(InTime);
+	Result.Y = Curves[1].Eval(InTime);
+	Result.Z = Curves[2].Eval(InTime);
+	return Result;
+}
+
diff --git a/Engine/Plugins/Cameras/GameplayCameras/Source/GameplayCameras/Public/Curves/CameraRotatorCurve.h b/Engine/Plugins/Cameras/GameplayCameras/Source/GameplayCameras/Public/Curves/CameraRotatorCurve.h
new file mode 100644
--- /dev/null
+++ b/Engine/Plugins/Cameras/GameplayCameras/Source/GameplayCameras/Public/Curves/CameraRotatorCurve.h
@@ -0,0 +1,20 @@
+// Copyright Epic Games, Inc. All Rights Reserved.
+
+#pragma once
+
+#include "Curves/RichCurve.h"
+#include "Math/MathFwd.h"
+
+#include "CameraRotatorCurve.generated.h"
+
+USTRUCT()
+struct FCameraRotatorCurve
+{
+	GENERATED_BODY()
+
+	UPROPERTY()
+	FRichCurve Curves[3];
+
+	FRotator GetValue(float InTime) const;
+};
+
diff --git a/Engine/Plugins/Cameras/GameplayCameras/Source/GameplayCameras/Public/Curves/CameraSingleCurve.h b/Engine/Plugins/Cameras/GameplayCameras/Source/GameplayCameras/Public/Curves/CameraSingleCurve.h
new file mode 100644
--- /dev/null
+++ b/Engine/Plugins/Cameras/GameplayCameras/Source/GameplayCameras/Public/Curves/CameraSingleCurve.h
@@ -0,0 +1,20 @@
+// Copyright Epic Games, Inc. All Rights Reserved.
+
+#pragma once
+
+#include "Curves/RichCurve.h"
+#include "Math/MathFwd.h"
+
+#include "CameraSingleCurve.generated.h"
+
+USTRUCT()
+struct FCameraSingleCurve
+{
+	GENERATED_BODY()
+
+	UPROPERTY()
+	FRichCurve Curve;
+
+	float GetValue(float InTime) const;
+};
+
diff --git a/Engine/Plugins/Cameras/GameplayCameras/Source/GameplayCameras/Public/Curves/CameraVectorCurve.h b/Engine/Plugins/Cameras/GameplayCameras/Source/GameplayCameras/Public/Curves/CameraVectorCurve.h
new file mode 100644
--- /dev/null
+++ b/Engine/Plugins/Cameras/GameplayCameras/Source/GameplayCameras/Public/Curves/CameraVectorCurve.h
@@ -0,0 +1,20 @@
+// Copyright Epic Games, Inc. All Rights Reserved.
+
+#pragma once
+
+#include "Curves/RichCurve.h"
+#include "Math/MathFwd.h"
+
+#include "CameraVectorCurve.generated.h"
+
+USTRUCT()
+struct FCameraVectorCurve
+{
+	GENERATED_BODY()
+
+	UPROPERTY()
+	FRichCurve Curves[3];
+
+	FVector GetValue(float InTime) const;
+};
+
