Cameras: add callbacks for when new nodes are added to a camera rig's graph

diff --git a/Engine/Plugins/Cameras/GameplayCameras/Source/GameplayCameras/Private/Core/CameraRigAsset.cpp b/Engine/Plugins/Cameras/GameplayCameras/Source/GameplayCameras/Private/Core/CameraRigAsset.cpp
--- a/Engine/Plugins/Cameras/GameplayCameras/Source/GameplayCameras/Private/Core/CameraRigAsset.cpp
+++ b/Engine/Plugins/Cameras/GameplayCameras/Source/GameplayCameras/Private/Core/CameraRigAsset.cpp
@@ -334,33 +334,41 @@ void UCameraRigAsset::GetConnectableObjects(FName InGraphName, TSet<UObject*>& O
 
 void UCameraRigAsset::AddConnectableObject(FName InGraphName, UObject* InObject)
 {
+	using namespace UE::Cameras;
+
 	Modify();
 
 	if (InGraphName == NodeTreeGraphName)
 	{
 		const int32 Index = AllNodeTreeObjects.AddUnique(InObject);
 		ensure(Index == AllNodeTreeObjects.Num() - 1);
+		EventHandlers.Notify(&ICameraRigAssetEventHandler::OnObjectAddedToGraph, NodeTreeGraphName, InObject);
 	}
 	else if (InGraphName == TransitionsGraphName)
 	{
 		const int32 Index = AllTransitionsObjects.AddUnique(InObject);
 		ensure(Index == AllTransitionsObjects.Num() - 1);
+		EventHandlers.Notify(&ICameraRigAssetEventHandler::OnObjectAddedToGraph, TransitionsGraphName, InObject);
 	}
 }
 
 void UCameraRigAsset::RemoveConnectableObject(FName InGraphName, UObject* InObject)
 {
+	using namespace UE::Cameras;
+
 	Modify();
 
 	if (InGraphName == NodeTreeGraphName)
 	{
 		const int32 NumRemoved = AllNodeTreeObjects.Remove(InObject);
 		ensure(NumRemoved == 1);
+		EventHandlers.Notify(&ICameraRigAssetEventHandler::OnObjectRemovedFromGraph, NodeTreeGraphName, InObject);
 	}
 	else if (InGraphName == TransitionsGraphName)
 	{
 		const int32 NumRemoved = AllTransitionsObjects.Remove(InObject);
 		ensure(NumRemoved == 1);
+		EventHandlers.Notify(&ICameraRigAssetEventHandler::OnObjectRemovedFromGraph, TransitionsGraphName, InObject);
 	}
 }
 
diff --git a/Engine/Plugins/Cameras/GameplayCameras/Source/GameplayCameras/Public/Core/CameraRigAsset.h b/Engine/Plugins/Cameras/GameplayCameras/Source/GameplayCameras/Public/Core/CameraRigAsset.h
--- a/Engine/Plugins/Cameras/GameplayCameras/Source/GameplayCameras/Public/Core/CameraRigAsset.h
+++ b/Engine/Plugins/Cameras/GameplayCameras/Source/GameplayCameras/Public/Core/CameraRigAsset.h
@@ -35,6 +35,11 @@ namespace UE::Cameras
 
 		/** Called when the camera rig asset has been built. */
 		virtual void OnCameraRigBuilt(const UCameraRigAsset* CameraRigAsset) {}
+
+#if WITH_EDITOR
+		virtual void OnObjectAddedToGraph(const FName GraphName, UObject* Object) {}
+		virtual void OnObjectRemovedFromGraph(const FName GraphName, UObject* Object) {}
+#endif  // WITH_EDITOR
 	};
 }
 
