Cameras: add "curves" tab in the camera rig asset editor

This new tab shows a curve editor populated with any node and their curve properties. As camera nodes are added/removed, the curve editor tree view is kept in sync.

diff --git a/Engine/Plugins/Cameras/GameplayCameras/Source/GameplayCamerasEditor/GameplayCamerasEditor.Build.cs b/Engine/Plugins/Cameras/GameplayCameras/Source/GameplayCamerasEditor/GameplayCamerasEditor.Build.cs
--- a/Engine/Plugins/Cameras/GameplayCameras/Source/GameplayCamerasEditor/GameplayCamerasEditor.Build.cs
+++ b/Engine/Plugins/Cameras/GameplayCameras/Source/GameplayCamerasEditor/GameplayCamerasEditor.Build.cs
@@ -26,6 +26,7 @@ public class GameplayCamerasEditor : ModuleRules
 				"CinematicCamera",
 				"Core",
 				"CoreUObject",
+				"CurveEditor",
 				"DeveloperSettings",
 				"EditorFramework",
 				"EditorSubsystem",
diff --git a/Engine/Plugins/Cameras/GameplayCameras/Source/GameplayCamerasEditor/Private/CurveEditor/CurvePropertyEditorTreeItem.cpp b/Engine/Plugins/Cameras/GameplayCameras/Source/GameplayCamerasEditor/Private/CurveEditor/CurvePropertyEditorTreeItem.cpp
new file mode 100644
--- /dev/null
+++ b/Engine/Plugins/Cameras/GameplayCameras/Source/GameplayCamerasEditor/Private/CurveEditor/CurvePropertyEditorTreeItem.cpp
@@ -0,0 +1,77 @@
+// Copyright Epic Games, Inc. All Rights Reserved.
+
+#include "CurveEditor/CurvePropertyEditorTreeItem.h"
+
+#include "RichCurveEditorModel.h"
+#include "Tree/SCurveEditorTreePin.h"
+#include "Tree/SCurveEditorTreeSelect.h"
+#include "Widgets/SBoxPanel.h"
+#include "Widgets/Text/STextBlock.h"
+
+namespace UE::Cameras
+{
+
+FCurvePropertyEditorTreeItem::FCurvePropertyEditorTreeItem()
+{
+}
+
+FCurvePropertyEditorTreeItem::FCurvePropertyEditorTreeItem(FCurvePropertyInfo&& InInfo)
+	: Info(MoveTemp(InInfo))
+{
+}
+
+FCurvePropertyEditorTreeItem::FCurvePropertyEditorTreeItem(const FText& InDisplayName, TWeakObjectPtr<> InWeakOwner)
+{
+	Info.DisplayName = InDisplayName;
+	Info.WeakOwner = InWeakOwner;
+}
+
+FCurvePropertyEditorTreeItem::FCurvePropertyEditorTreeItem(FRichCurve* InRichCurve, const FText& InCurveName, const FLinearColor& InCurveColor, TWeakObjectPtr<> InWeakOwner)
+{
+	Info.Curve = InRichCurve;
+	Info.DisplayName = InCurveName;
+	Info.Color = InCurveColor;
+	Info.WeakOwner = InWeakOwner;
+}
+
+TSharedPtr<SWidget> FCurvePropertyEditorTreeItem::GenerateCurveEditorTreeWidget(const FName& InColumnName, TWeakPtr<FCurveEditor> InCurveEditor, FCurveEditorTreeItemID InTreeItemID, const TSharedRef<ITableRow>& TableRow)
+{
+	if (InColumnName == ColumnNames.Label)
+	{
+		return SNew(SHorizontalBox)
+			+SHorizontalBox::Slot()
+			.Padding(FMargin(4.f))
+			.VAlign(VAlign_Center)
+			.HAlign(HAlign_Right)
+			.AutoWidth()
+			[
+				SNew(STextBlock)
+				.Text(Info.DisplayName)
+				.ColorAndOpacity(FSlateColor(Info.Color))
+			];
+	}
+	else if (InColumnName == ColumnNames.SelectHeader)
+	{
+		return SNew(SCurveEditorTreeSelect, InCurveEditor, InTreeItemID, TableRow);
+	}
+	else if (InColumnName == ColumnNames.PinHeader)
+	{
+		return SNew(SCurveEditorTreePin, InCurveEditor, InTreeItemID, TableRow);
+	}
+
+	return nullptr;
+}
+
+void FCurvePropertyEditorTreeItem::CreateCurveModels(TArray<TUniquePtr<FCurveModel>>& OutCurveModels)
+{
+	if (Info.Curve)
+	{
+		TUniquePtr<FRichCurveEditorModelRaw> NewCurve = MakeUnique<FRichCurveEditorModelRaw>(Info.Curve, Info.WeakOwner.Get());
+		NewCurve->SetShortDisplayName(Info.DisplayName);
+		NewCurve->SetColor(Info.Color, false);
+		OutCurveModels.Add(MoveTemp(NewCurve));
+	}
+}
+
+}  // namespace UE::Cameras
+
diff --git a/Engine/Plugins/Cameras/GameplayCameras/Source/GameplayCamerasEditor/Private/CurveEditor/CurvePropertyEditorTreeItem.h b/Engine/Plugins/Cameras/GameplayCameras/Source/GameplayCamerasEditor/Private/CurveEditor/CurvePropertyEditorTreeItem.h
new file mode 100644
--- /dev/null
+++ b/Engine/Plugins/Cameras/GameplayCameras/Source/GameplayCamerasEditor/Private/CurveEditor/CurvePropertyEditorTreeItem.h
@@ -0,0 +1,57 @@
+// Copyright Epic Games, Inc. All Rights Reserved.
+
+#pragma once
+
+#include "Internationalization/Text.h"
+#include "Math/Color.h"
+#include "Tree/ICurveEditorTreeItem.h"
+#include "UObject/WeakObjectPtr.h"
+
+struct FRichCurve;
+
+namespace UE::Cameras
+{
+
+/**
+ * Structure for providing information about a curve property, to be shown in
+ * a curve editor tree item (see FCurvePropertyEditorTreeItem).
+ */
+struct FCurvePropertyInfo
+{
+	/** Display name of the tree item. */
+	FText DisplayName;
+	/** Color of the tree item and any associated curve. */
+	FLinearColor Color = FLinearColor::White;
+
+	/** Name of the property on the owning object. */
+	FName PropertyName;
+	/** The object on which the curve property resides. */
+	TWeakObjectPtr<> WeakOwner;
+
+	/** The curve associated with the tree item. */
+	FRichCurve* Curve = nullptr;
+};
+
+/**
+ * Curve editor tree view item for a curve tied to an object's property.
+ */
+struct FCurvePropertyEditorTreeItem : public ICurveEditorTreeItem
+{
+	FCurvePropertyInfo Info;
+
+public:
+
+	FCurvePropertyEditorTreeItem();
+	FCurvePropertyEditorTreeItem(FCurvePropertyInfo&& InInfo);
+	FCurvePropertyEditorTreeItem(const FText& InDisplayName, TWeakObjectPtr<> InWeakOwner);
+	FCurvePropertyEditorTreeItem(FRichCurve* InRichCurve, const FText& InCurveName, const FLinearColor& InCurveColor, TWeakObjectPtr<> InWeakOwner);
+
+	UObject* GetOwner() const { return Info.WeakOwner.Get(); }
+
+	// ICurveEditorTreeItem interface.
+	virtual TSharedPtr<SWidget> GenerateCurveEditorTreeWidget(const FName& InColumnName, TWeakPtr<FCurveEditor> InCurveEditor, FCurveEditorTreeItemID InTreeItemID, const TSharedRef<ITableRow>& TableRow) override;
+	virtual void CreateCurveModels(TArray<TUniquePtr<FCurveModel>>& OutCurveModels) override;
+};
+
+}  // namespace UE::Cameras
+
diff --git a/Engine/Plugins/Cameras/GameplayCameras/Source/GameplayCamerasEditor/Private/Styles/GameplayCamerasEditorStyle.cpp b/Engine/Plugins/Cameras/GameplayCameras/Source/GameplayCamerasEditor/Private/Styles/GameplayCamerasEditorStyle.cpp
--- a/Engine/Plugins/Cameras/GameplayCameras/Source/GameplayCamerasEditor/Private/Styles/GameplayCamerasEditorStyle.cpp
+++ b/Engine/Plugins/Cameras/GameplayCameras/Source/GameplayCamerasEditor/Private/Styles/GameplayCamerasEditorStyle.cpp
@@ -79,6 +79,7 @@ FGameplayCamerasEditorStyle::FGameplayCamerasEditorStyle()
 	Set("CameraRigAssetEditor.Tabs.Messages", new CORE_IMAGE_BRUSH_SVG("Starship/Common/OutputLog", Icon16x16));
 	Set("CameraRigAssetEditor.Tabs.NodeHierarchy", new IMAGE_BRUSH_SVG("Icons/CameraRig-NodeHierarchy", Icon16x16));
 	Set("CameraRigAssetEditor.Tabs.Transitions", new IMAGE_BRUSH_SVG("Icons/CameraRig-Transitions", Icon16x16));
+	Set("CameraRigAssetEditor.Tabs.Curves", new IMAGE_BRUSH_SVG("Icons/CurveEditor", Icon16x16));
 
 	Set("CameraRigAssetEditor.Build", new IMAGE_BRUSH_SVG("Icons/CameraRig-BuildStatus_Background", Icon20x20));
 	Set("CameraRigAssetEditor.BuildStatus.Background", new IMAGE_BRUSH_SVG("Icons/CameraRig-BuildStatus_Background", Icon20x20));
@@ -109,6 +110,10 @@ FGameplayCamerasEditorStyle::FGameplayCamerasEditorStyle()
 	Set("CameraVariableCollectionEditor.Entry.Type", FTextBlockStyle(NormalText).SetFont(DEFAULT_FONT("Italic", 10)));
 	Set("CameraVariableCollectionEditor.Entry.Value", FTextBlockStyle(NormalText).SetFont(DEFAULT_FONT("Regular", 10)));
 
+	// Curve editor icons.
+	Set("CurveEditor.ShowInCurvesTab", new IMAGE_BRUSH_SVG("Icons/CurveEditor", Icon16x16));
+	//Set("NiagaraEditor.CurveDetails.ShowInOverview.Small", new CORE_IMAGE_BRUSH("Common/GoToSource", Icon12x12, FLinearColor(.9f, .9f, .9f, 1.0f)));
+
 	// Debugger tool icons.
 	Set("Debugger.TabIcon", new IMAGE_BRUSH_SVG("Icons/GameplayCameraSystem_16", Icon16x16));
 
diff --git a/Engine/Plugins/Cameras/GameplayCameras/Source/GameplayCamerasEditor/Private/Toolkits/CameraRigAssetEditorToolkit.cpp b/Engine/Plugins/Cameras/GameplayCameras/Source/GameplayCamerasEditor/Private/Toolkits/CameraRigAssetEditorToolkit.cpp
--- a/Engine/Plugins/Cameras/GameplayCameras/Source/GameplayCamerasEditor/Private/Toolkits/CameraRigAssetEditorToolkit.cpp
+++ b/Engine/Plugins/Cameras/GameplayCameras/Source/GameplayCamerasEditor/Private/Toolkits/CameraRigAssetEditorToolkit.cpp
@@ -12,6 +12,7 @@
 #include "Editors/SFindInObjectTreeGraph.h"
 #include "Framework/Docking/LayoutExtender.h"
 #include "Framework/Docking/TabManager.h"
+#include "GraphEditAction.h"
 #include "Helpers/AssetTypeMenuOverlayHelper.h"
 #include "IGameplayCamerasEditorModule.h"
 #include "IGameplayCamerasLiveEditManager.h"
@@ -22,6 +23,7 @@
 #include "Toolkits/BuildButtonToolkit.h"
 #include "Toolkits/CameraBuildLogToolkit.h"
 #include "Toolkits/CameraRigAssetEditorToolkitBase.h"
+#include "Toolkits/CurveEditorToolkit.h"
 #include "Toolkits/StandardToolkitLayout.h"
 #include "Widgets/Docking/SDockTab.h"
 
@@ -34,6 +36,7 @@ namespace UE::Cameras
 
 const FName FCameraRigAssetEditorToolkit::SearchTabId(TEXT("CameraRigAssetEditor_Search"));
 const FName FCameraRigAssetEditorToolkit::MessagesTabId(TEXT("CameraRigAssetEditor_Messages"));
+const FName FCameraRigAssetEditorToolkit::CurvesTabId(TEXT("CameraRigAssetEditor_Curves"));
 
 FCameraRigAssetEditorToolkit::FCameraRigAssetEditorToolkit(UAssetEditor* InOwningAssetEditor)
 	: FBaseAssetToolkit(InOwningAssetEditor)
@@ -41,12 +44,14 @@ FCameraRigAssetEditorToolkit::FCameraRigAssetEditorToolkit(UAssetEditor* InOwnin
 	Impl = MakeShared<FCameraRigAssetEditorToolkitBase>(TEXT("CameraRigAssetEditor_Layout_v6"));
 	BuildButtonToolkit = MakeShared<FBuildButtonToolkit>();
 	BuildLogToolkit = MakeShared<FCameraBuildLogToolkit>();
+	CurveEditorToolkit = MakeShared<FCurveEditorToolkit>();
 
 	// Override base class default layout.
 	TSharedPtr<FStandardToolkitLayout> StandardLayout = Impl->GetStandardLayout();
 	{
 		StandardLayout->AddBottomTab(SearchTabId);
 		StandardLayout->AddBottomTab(MessagesTabId);
+		StandardLayout->AddBottomTab(CurvesTabId);
 	}
 	StandaloneDefaultLayout = StandardLayout->GetLayout();
 
@@ -61,8 +66,15 @@ FCameraRigAssetEditorToolkit::FCameraRigAssetEditorToolkit(UAssetEditor* InOwnin
 
 void FCameraRigAssetEditorToolkit::SetCameraRigAsset(UCameraRigAsset* InCameraRig)
 {
+	EventHandler.Unlink();
+
 	Impl->SetCameraRigAsset(InCameraRig);
 	BuildButtonToolkit->SetTarget(InCameraRig);
+
+	if (InCameraRig)
+	{
+		InCameraRig->EventHandlers.Register(EventHandler, this);
+	}
 }
 
 void FCameraRigAssetEditorToolkit::RegisterTabSpawners(const TSharedRef<class FTabManager>& InTabManager)
@@ -83,6 +95,11 @@ void FCameraRigAssetEditorToolkit::RegisterTabSpawners(const TSharedRef<class FT
 		.SetDisplayName(LOCTEXT("Messages", "Messages"))
 		.SetGroup(AssetEditorTabsCategory.ToSharedRef())
 		.SetIcon(FSlateIcon(CamerasStyleSetName, "CameraRigAssetEditor.Tabs.Messages"));
+
+	InTabManager->RegisterTabSpawner(CurvesTabId, FOnSpawnTab::CreateSP(this, &FCameraRigAssetEditorToolkit::SpawnTab_Curves))
+		.SetDisplayName(LOCTEXT("Curves", "Curves"))
+		.SetGroup(AssetEditorTabsCategory.ToSharedRef())
+		.SetIcon(FSlateIcon(CamerasStyleSetName, "CameraRigAssetEditor.Tabs.Curves"));
 }
 
 TSharedRef<SDockTab> FCameraRigAssetEditorToolkit::SpawnTab_Search(const FSpawnTabArgs& Args)
@@ -107,6 +124,37 @@ TSharedRef<SDockTab> FCameraRigAssetEditorToolkit::SpawnTab_Messages(const FSpaw
 	return MessagesTab.ToSharedRef();
 }
 
+TSharedRef<SDockTab> FCameraRigAssetEditorToolkit::SpawnTab_Curves(const FSpawnTabArgs& Args)
+{
+	if (!CurveEditorToolkit->IsInitialized())
+	{
+		TArray<UObject*> CameraRigObjects;
+		UPackage* CameraRigPackage = Impl->GetCameraRigAsset()->GetPackage();
+		GetObjectsWithPackage(CameraRigPackage, CameraRigObjects);
+		CurveEditorToolkit->Initialize(CameraRigObjects);
+	}
+
+	TSharedPtr<SDockTab> CurvesTab = SNew(SDockTab)
+		.Label(LOCTEXT("CurvesTabTitle", "Curves"))
+		.OnTabClosed(this, &FCameraRigAssetEditorToolkit::OnCurvesTabClosed)
+		[
+			CurveEditorToolkit->GetCurveEditorWidget().ToSharedRef()
+		];
+
+	return CurvesTab.ToSharedRef();
+}
+
+void FCameraRigAssetEditorToolkit::OnCurvesTabClosed(TSharedRef<SDockTab> InTab)
+{
+	if (CurveEditorToolkit->IsInitialized())
+	{
+		InTab->ClearContent();
+
+		// Clear the curve editor when the tab is closed.
+		CurveEditorToolkit->Shutdown();
+	}
+}
+
 void FCameraRigAssetEditorToolkit::UnregisterTabSpawners(const TSharedRef<FTabManager>& InTabManager)
 {
 	// Skip FBaseAssetToolkit here because we don't want a viewport tab.
@@ -116,6 +164,7 @@ void FCameraRigAssetEditorToolkit::UnregisterTabSpawners(const TSharedRef<FTabMa
 
 	InTabManager->UnregisterTabSpawner(SearchTabId);
 	InTabManager->UnregisterTabSpawner(MessagesTabId);
+	InTabManager->UnregisterTabSpawner(CurvesTabId);
 }
 
 void FCameraRigAssetEditorToolkit::CreateWidgets()
@@ -282,6 +331,37 @@ void FCameraRigAssetEditorToolkit::OnJumpToObject(UObject* Object, FName Propert
 	CameraRigEditor->FindAndJumpToObjectNode(Object);
 }
 
+void FCameraRigAssetEditorToolkit::OnInvokeCurveEditor(UObject* Object, FName PropertyName)
+{
+	UPackage* Package = Impl->GetCameraRigAsset()->GetPackage();
+	if (Object->IsIn(Package))
+	{
+		TabManager->TryInvokeTab(CurvesTabId);
+
+		CurveEditorToolkit->SelectCurves(Object, PropertyName);
+	}
+}
+
+void FCameraRigAssetEditorToolkit::OnObjectAddedToGraph(const FName GraphName, UObject* Object)
+{
+	ECameraRigAssetEditorMode CurrentMode = Impl->GetCameraRigEditorMode();
+	if ((GraphName == UCameraRigAsset::NodeTreeGraphName && CurrentMode == ECameraRigAssetEditorMode::NodeGraph) ||
+			(GraphName == UCameraRigAsset::TransitionsGraphName && CurrentMode == ECameraRigAssetEditorMode::TransitionGraph))
+	{
+		CurveEditorToolkit->AddCurveOwner(Object);
+	}
+}
+
+void FCameraRigAssetEditorToolkit::OnObjectRemovedFromGraph(const FName GraphName, UObject* Object)
+{
+	ECameraRigAssetEditorMode CurrentMode = Impl->GetCameraRigEditorMode();
+	if ((GraphName == UCameraRigAsset::NodeTreeGraphName && CurrentMode == ECameraRigAssetEditorMode::NodeGraph) ||
+			(GraphName == UCameraRigAsset::TransitionsGraphName && CurrentMode == ECameraRigAssetEditorMode::TransitionGraph))
+	{
+		CurveEditorToolkit->RemoveCurveOwner(Object);
+	}
+}
+
 FText FCameraRigAssetEditorToolkit::GetBaseToolkitName() const
 {
 	return LOCTEXT("AppLabel", "Camera Rig Asset");
diff --git a/Engine/Plugins/Cameras/GameplayCameras/Source/GameplayCamerasEditor/Private/Toolkits/CameraRigAssetEditorToolkitBase.cpp b/Engine/Plugins/Cameras/GameplayCameras/Source/GameplayCamerasEditor/Private/Toolkits/CameraRigAssetEditorToolkitBase.cpp
--- a/Engine/Plugins/Cameras/GameplayCameras/Source/GameplayCamerasEditor/Private/Toolkits/CameraRigAssetEditorToolkitBase.cpp
+++ b/Engine/Plugins/Cameras/GameplayCameras/Source/GameplayCamerasEditor/Private/Toolkits/CameraRigAssetEditorToolkitBase.cpp
@@ -207,6 +207,11 @@ bool FCameraRigAssetEditorToolkitBase::IsCameraRigEditorMode(ECameraRigAssetEdit
 	return CameraRigEditorWidget->IsEditorMode(InEditorMode);
 }
 
+ECameraRigAssetEditorMode FCameraRigAssetEditorToolkitBase::GetCameraRigEditorMode() const
+{
+	return CameraRigEditorWidget->GetEditorMode();
+}
+
 void FCameraRigAssetEditorToolkitBase::OnAnyGraphChanged(const FEdGraphEditAction& InEditAction)
 {
 	// Called when something is modified in the node graph or transition graph.
diff --git a/Engine/Plugins/Cameras/GameplayCameras/Source/GameplayCamerasEditor/Private/Toolkits/CurveEditorToolkit.cpp b/Engine/Plugins/Cameras/GameplayCameras/Source/GameplayCamerasEditor/Private/Toolkits/CurveEditorToolkit.cpp
new file mode 100644
--- /dev/null
+++ b/Engine/Plugins/Cameras/GameplayCameras/Source/GameplayCamerasEditor/Private/Toolkits/CurveEditorToolkit.cpp
@@ -0,0 +1,272 @@
+// Copyright Epic Games, Inc. All Rights Reserved.
+
+#include "Toolkits/CurveEditorToolkit.h"
+
+#include "CurveEditor.h"
+#include "CurveEditor/CurvePropertyEditorTreeItem.h"
+#include "Curves/CameraRotatorCurve.h"
+#include "Curves/CameraSingleCurve.h"
+#include "Curves/CameraVectorCurve.h"
+#include "Curves/RichCurve.h"
+#include "Framework/MultiBox/MultiBoxBuilder.h"
+#include "SCurveEditorPanel.h"
+#include "Tree/SCurveEditorTree.h"
+#include "UObject/Class.h"
+
+#define LOCTEXT_NAMESPACE "CurveEditorToolkit"
+
+namespace UE::Cameras
+{
+
+FCurveEditorToolkit::FCurveEditorToolkit()
+{
+}
+
+void FCurveEditorToolkit::Initialize()
+{
+	if (!ensureMsgf(!CurveEditor.IsValid(), TEXT("This curve editor toolkit has already been initialized.")))
+	{
+		return;
+	}
+
+	// Create the editor.
+	CurveEditor = MakeShared<FCurveEditor>();
+
+	FCurveEditorInitParams InitParams;
+	CurveEditor->InitCurveEditor(InitParams);
+
+	CurveEditor->GridLineLabelFormatXAttribute = LOCTEXT("GridXLabelFormat", "{0}");
+
+	TUniquePtr<ICurveEditorBounds> EditorBounds = MakeUnique<FStaticCurveEditorBounds>();
+	CurveEditor->SetBounds(MoveTemp(EditorBounds));
+
+	// Create the panel.
+	TSharedRef<SCurveEditorPanel> CurveEditorPanel = SNew(SCurveEditorPanel, CurveEditor.ToSharedRef())
+		.TreeContent()
+		[
+			SNew(SCurveEditorTree, CurveEditor)
+		];
+
+	// Create the toolbar.
+	TSharedPtr<FUICommandList> Commands = CurveEditorPanel->GetCommands();
+	TSharedPtr<FExtender> ToolbarExtender = CurveEditorPanel->GetToolbarExtender();
+
+	FToolBarBuilder ToolBarBuilder(Commands, FMultiBoxCustomization::None, ToolbarExtender, true);
+	ToolBarBuilder.BeginSection("Asset");
+	ToolBarBuilder.EndSection();
+	TSharedRef<SWidget> ToolBarWidget = ToolBarBuilder.MakeWidget();
+
+	TSharedRef<SBorder> CurveEditorPanelWrapper = SNew(SBorder)
+		.BorderImage(FAppStyle::Get().GetBrush("Brushes.Panel"))
+		.Padding(FMargin(16.0f, 16.0f))
+		[
+			SNew(SVerticalBox)
+			+SVerticalBox::Slot()
+			.AutoHeight()
+			[
+				ToolBarWidget
+			]
+			+SVerticalBox::Slot()
+			.FillHeight(1.f)
+			[
+				CurveEditorPanel
+			]
+		];
+
+	CurveEditorWidget = CurveEditorPanelWrapper;
+}
+
+void FCurveEditorToolkit::Initialize(TArrayView<UObject*> InCurveOwners)
+{
+	Initialize();
+
+	for (UObject* CurveOwner : InCurveOwners)
+	{
+		AddCurves(CurveOwner);
+	}
+}
+
+void FCurveEditorToolkit::Shutdown()
+{
+	if (!ensureMsgf(CurveEditor.IsValid(), TEXT("This curve editor toolkit was not initialized.")))
+	{
+		return;
+	}
+
+	CurveEditorWidget.Reset();
+	CurveEditor.Reset();
+}
+
+void FCurveEditorToolkit::AddCurveOwner(UObject* InCurveOwner)
+{
+	AddCurves(InCurveOwner);
+}
+
+void FCurveEditorToolkit::RemoveCurveOwner(UObject* InCurveOwner)
+{
+	TArray<FCurveEditorTreeItemID> RootTreeItemIDs = CurveEditor->GetRootTreeItems();
+	for (FCurveEditorTreeItemID RootTreeItemID : RootTreeItemIDs)
+	{
+		FCurveEditorTreeItem& TreeItem = CurveEditor->GetTreeItem(RootTreeItemID);
+		TSharedPtr<FCurvePropertyEditorTreeItem> SimpleTreeItem = StaticCastSharedPtr<FCurvePropertyEditorTreeItem>(TreeItem.GetItem());
+		if (SimpleTreeItem->GetOwner() == InCurveOwner)
+		{
+			CurveEditor->RemoveTreeItem(RootTreeItemID);
+		}
+	}
+}
+
+void FCurveEditorToolkit::RemoveAllCurveOwners()
+{
+	CurveEditor->RemoveAllTreeItems();
+}
+
+void FCurveEditorToolkit::SelectCurves(UObject* InCurveOwner, FName InPropertyName)
+{
+	TArray<FCurveEditorTreeItemID> Selection;
+
+	// Find all root items that relate to the given object.
+	TArray<FCurveEditorTreeItemID> RootTreeItemIDs = CurveEditor->GetRootTreeItems();
+	for (FCurveEditorTreeItemID RootTreeItemID : RootTreeItemIDs)
+	{
+		FCurveEditorTreeItem& RootTreeItem = CurveEditor->GetTreeItem(RootTreeItemID);
+		TSharedPtr<FCurvePropertyEditorTreeItem> ObjectTreeItem = StaticCastSharedPtr<FCurvePropertyEditorTreeItem>(RootTreeItem.GetItem());
+		if (ObjectTreeItem->GetOwner() == InCurveOwner)
+		{
+			// Find any curves, or curve groups, that match the given property name.
+			for (FCurveEditorTreeItemID ChildID : RootTreeItem.GetChildren())
+			{
+				FCurveEditorTreeItem& ChildTreeItem = CurveEditor->GetTreeItem(ChildID);
+				TSharedPtr<FCurvePropertyEditorTreeItem> PropertyTreeItem = StaticCastSharedPtr<FCurvePropertyEditorTreeItem>(ChildTreeItem.GetItem());
+				if (PropertyTreeItem->Info.PropertyName == InPropertyName)
+				{
+					Selection.Add(ChildID);
+				}
+			}
+		}
+	}
+
+	if (!Selection.IsEmpty())
+	{
+		CurveEditor->SetTreeSelection(MoveTemp(Selection));
+	}
+}
+
+void FCurveEditorToolkit::AddCurves(UObject* InObject)
+{
+	if (!InObject)
+	{
+		return;
+	}
+
+	UClass* ObjectClass = InObject->GetClass();
+
+	FCurveEditorTreeItem* ObjectItem = nullptr;
+	auto GetObjectItemID = [this, InObject, &ObjectItem]() -> FCurveEditorTreeItemID
+	{
+		if (!ObjectItem)
+		{
+			FCurvePropertyInfo ObjectInfo;
+			ObjectInfo.WeakOwner = InObject;
+			ObjectInfo.DisplayName = FText::FromName(InObject->GetFName());
+			ObjectItem = AddTreeItem(FCurveEditorTreeItemID::Invalid(), MoveTemp(ObjectInfo));
+		}
+		return ObjectItem->GetID();
+	};
+
+	const FLinearColor RGBColors[3] = { FLinearColor::Red, FLinearColor::Green, FLinearColor::Blue };
+
+	for (TFieldIterator<FProperty> PropertyIt(ObjectClass); PropertyIt; ++PropertyIt)
+	{
+		FStructProperty* StructProperty = CastField<FStructProperty>(*PropertyIt);
+		if (!StructProperty)
+		{
+			continue;
+		}
+
+		const FName PropertyName = StructProperty->GetFName();
+		const FText PropertyDisplayName = FText::FromName(StructProperty->GetFName());
+
+		if (StructProperty->Struct == FRichCurve::StaticStruct())
+		{
+			FCurvePropertyInfo CurveInfo;
+			CurveInfo.Curve = StructProperty->ContainerPtrToValuePtr<FRichCurve>(InObject);
+			CurveInfo.DisplayName = PropertyDisplayName;
+			CurveInfo.Color = FLinearColor::White;
+			CurveInfo.WeakOwner = InObject;
+			CurveInfo.PropertyName = PropertyName;
+			AddTreeItem(GetObjectItemID(), MoveTemp(CurveInfo));
+		}
+		else if (StructProperty->Struct == FCameraSingleCurve::StaticStruct())
+		{
+			FCameraSingleCurve* SingleCurve = StructProperty->ContainerPtrToValuePtr<FCameraSingleCurve>(InObject);
+
+			FCurvePropertyInfo CurveInfo;
+			CurveInfo.Curve = &SingleCurve->Curve;
+			CurveInfo.DisplayName = PropertyDisplayName;
+			CurveInfo.Color = FLinearColor::White;
+			CurveInfo.WeakOwner = InObject;
+			CurveInfo.PropertyName = PropertyName;
+			AddTreeItem(GetObjectItemID(), MoveTemp(CurveInfo));
+		}
+		else if (StructProperty->Struct == FCameraVectorCurve::StaticStruct())
+		{
+			const FText XYZNames[3] = { LOCTEXT("X", "X"), LOCTEXT("Y", "Y"), LOCTEXT("Z", "Z") };
+			FCameraVectorCurve* VectorCurve = StructProperty->ContainerPtrToValuePtr<FCameraVectorCurve>(InObject);
+
+			FCurvePropertyInfo ParentInfo;
+			ParentInfo.DisplayName = PropertyDisplayName;
+			ParentInfo.WeakOwner = InObject;
+			ParentInfo.PropertyName = PropertyName;
+			FCurveEditorTreeItem* ParentItem = AddTreeItem(GetObjectItemID(), MoveTemp(ParentInfo));
+
+			for (int32 Index = 0; Index < 3; ++Index)
+			{
+				FCurvePropertyInfo CurveInfo;
+				CurveInfo.Curve = &VectorCurve->Curves[Index];
+				CurveInfo.DisplayName = XYZNames[Index];
+				CurveInfo.Color = RGBColors[Index];
+				CurveInfo.WeakOwner = InObject;
+				CurveInfo.PropertyName = PropertyName;
+				AddTreeItem(ParentItem->GetID(), MoveTemp(CurveInfo));
+			}
+		}
+		else if (StructProperty->Struct == FCameraRotatorCurve::StaticStruct())
+		{
+			const FText YPRNames[3] = { LOCTEXT("Yaw", "Yaw"), LOCTEXT("Pitch", "Pitch"), LOCTEXT("Roll", "Roll") };
+			FCameraRotatorCurve* RotatorCurve = StructProperty->ContainerPtrToValuePtr<FCameraRotatorCurve>(InObject);
+
+			FCurvePropertyInfo ParentInfo;
+			ParentInfo.DisplayName = PropertyDisplayName;
+			ParentInfo.WeakOwner = InObject;
+			ParentInfo.PropertyName = PropertyName;
+			FCurveEditorTreeItem* ParentItem = AddTreeItem(GetObjectItemID(), MoveTemp(ParentInfo));
+
+			for (int32 Index = 0; Index < 3; ++Index)
+			{
+				FCurvePropertyInfo CurveInfo;
+				CurveInfo.Curve = &RotatorCurve->Curves[Index];
+				CurveInfo.DisplayName = YPRNames[Index];
+				CurveInfo.Color = RGBColors[Index];
+				CurveInfo.WeakOwner = InObject;
+				CurveInfo.PropertyName = PropertyName;
+				AddTreeItem(ParentItem->GetID(), MoveTemp(CurveInfo));
+			}
+		}
+	}
+}
+
+ FCurveEditorTreeItem* FCurveEditorToolkit::AddTreeItem(FCurveEditorTreeItemID ParentID, FCurvePropertyInfo&& CurveInfo)
+{
+	FCurveEditorTreeItem* TreeItem = CurveEditor->AddTreeItem(ParentID);
+
+	TSharedPtr<FCurvePropertyEditorTreeItem> TreeItemModel = MakeShared<FCurvePropertyEditorTreeItem>(MoveTemp(CurveInfo));
+	TreeItem->SetStrongItem(TreeItemModel);
+
+	return TreeItem;
+}
+
+}  // namespace UE::Cameras
+
+#undef LOCTEXT_NAMESPACE
+
diff --git a/Engine/Plugins/Cameras/GameplayCameras/Source/GameplayCamerasEditor/Public/Toolkits/CameraRigAssetEditorToolkit.h b/Engine/Plugins/Cameras/GameplayCameras/Source/GameplayCamerasEditor/Public/Toolkits/CameraRigAssetEditorToolkit.h
--- a/Engine/Plugins/Cameras/GameplayCameras/Source/GameplayCamerasEditor/Public/Toolkits/CameraRigAssetEditorToolkit.h
+++ b/Engine/Plugins/Cameras/GameplayCameras/Source/GameplayCamerasEditor/Public/Toolkits/CameraRigAssetEditorToolkit.h
@@ -2,6 +2,7 @@
 
 #pragma once
 
+#include "Core/CameraRigAsset.h"
 #include "CoreTypes.h"
 #include "Editors/ObjectTreeGraphConfig.h"
 #include "Tools/BaseAssetToolkit.h"
@@ -10,6 +11,7 @@
 
 class SFindInObjectTreeGraph;
 class UCameraRigAsset;
+struct FEdGraphEditAction;
 struct FFindInObjectTreeGraphSource;
 
 namespace UE::Cameras
@@ -18,10 +20,12 @@ namespace UE::Cameras
 class FBuildButtonToolkit;
 class FCameraBuildLogToolkit;
 class FCameraRigAssetEditorToolkitBase;
+class FCurveEditorToolkit;
 class IGameplayCamerasLiveEditManager;
 
 class FCameraRigAssetEditorToolkit
 	: public FBaseAssetToolkit
+	, public ICameraRigAssetEventHandler
 {
 public:
 
@@ -46,10 +50,17 @@ protected:
 	virtual FString GetWorldCentricTabPrefix() const override;
 	virtual FLinearColor GetWorldCentricTabColorScale() const override;
 
+	// ICameraRigAssetEventHandler interface
+	virtual void OnObjectAddedToGraph(const FName GraphName, UObject* Object) override;
+	virtual void OnObjectRemovedFromGraph(const FName GraphName, UObject* Object) override;
+
 private:
 
 	TSharedRef<SDockTab> SpawnTab_Search(const FSpawnTabArgs& Args);
 	TSharedRef<SDockTab> SpawnTab_Messages(const FSpawnTabArgs& Args);
+	TSharedRef<SDockTab> SpawnTab_Curves(const FSpawnTabArgs& Args);
+
+	void OnCurvesTabClosed(TSharedRef<SDockTab> InTab);
 
 	void OnBuild();
 	void OnFindInCameraRig();
@@ -57,10 +68,13 @@ private:
 	void OnGetRootObjectsToSearch(TArray<FFindInObjectTreeGraphSource>& OutSources);
 	void OnJumpToObject(UObject* Object, FName PropertyName);
 
+	void OnInvokeCurveEditor(UObject* Object, FName PropertyName);
+
 private:
 
 	static const FName SearchTabId;
 	static const FName MessagesTabId;
+	static const FName CurvesTabId;
 
 	/** Base implementation */
 	TSharedPtr<FCameraRigAssetEditorToolkitBase> Impl;
@@ -74,12 +88,17 @@ private:
 	TSharedPtr<FBuildButtonToolkit> BuildButtonToolkit;
 	/** The output log */
 	TSharedPtr<FCameraBuildLogToolkit> BuildLogToolkit;
+	/** The curve editor */
+	TSharedPtr<FCurveEditorToolkit> CurveEditorToolkit;
 
 	/** Search widget */
 	TSharedPtr<SFindInObjectTreeGraph> SearchWidget;
 
 	/** Live edit manager for updating the assets in the runtime */
 	TSharedPtr<IGameplayCamerasLiveEditManager> LiveEditManager;
+
+	/** Event handler */
+	TCameraEventHandler<ICameraRigAssetEventHandler> EventHandler;
 };
 
 }  // namespace UE::Cameras
diff --git a/Engine/Plugins/Cameras/GameplayCameras/Source/GameplayCamerasEditor/Public/Toolkits/CameraRigAssetEditorToolkitBase.h b/Engine/Plugins/Cameras/GameplayCameras/Source/GameplayCamerasEditor/Public/Toolkits/CameraRigAssetEditorToolkitBase.h
--- a/Engine/Plugins/Cameras/GameplayCameras/Source/GameplayCamerasEditor/Public/Toolkits/CameraRigAssetEditorToolkitBase.h
+++ b/Engine/Plugins/Cameras/GameplayCameras/Source/GameplayCamerasEditor/Public/Toolkits/CameraRigAssetEditorToolkitBase.h
@@ -52,6 +52,9 @@ public:
 
 	void SetLiveEditManager(TSharedPtr<IGameplayCamerasLiveEditManager> InLiveEditManager);
 
+	ECameraRigAssetEditorMode GetCameraRigEditorMode() const;
+	bool IsCameraRigEditorMode(ECameraRigAssetEditorMode InEditorMode) const;
+
 protected:
 
 	// FGCObject interface
@@ -64,7 +67,6 @@ protected:
 protected:
 
 	void SetCameraRigEditorMode(ECameraRigAssetEditorMode InEditorMode);
-	bool IsCameraRigEditorMode(ECameraRigAssetEditorMode InEditorMode) const;
 
 private:
 
diff --git a/Engine/Plugins/Cameras/GameplayCameras/Source/GameplayCamerasEditor/Public/Toolkits/CurveEditorToolkit.h b/Engine/Plugins/Cameras/GameplayCameras/Source/GameplayCamerasEditor/Public/Toolkits/CurveEditorToolkit.h
new file mode 100644
--- /dev/null
+++ b/Engine/Plugins/Cameras/GameplayCameras/Source/GameplayCamerasEditor/Public/Toolkits/CurveEditorToolkit.h
@@ -0,0 +1,77 @@
+// Copyright Epic Games, Inc. All Rights Reserved.
+
+#pragma once
+
+#include "CoreTypes.h"
+#include "Internationalization/Text.h"
+#include "Math/Color.h"
+#include "Templates/SharedPointer.h"
+#include "UObject/WeakObjectPtr.h"
+
+class FCurveEditor;
+class FUICommandInfo;
+class SWidget;
+struct FCurveEditorTreeItem;
+struct FCurveEditorTreeItemID;
+struct FRichCurve;
+struct FSlateIcon;
+struct FToolMenuEntry;
+
+namespace UE::Cameras
+{
+
+struct FCurvePropertyInfo;
+
+/**
+ * A utility toolkit that hosts a curve editor showing curves from any object that
+ * has curve properties.
+ */
+class FCurveEditorToolkit : public TSharedFromThis<FCurveEditorToolkit>
+{
+public:
+
+	FCurveEditorToolkit();
+
+	/** Initializes the toolkit and creates the curve editor. */
+	void Initialize();
+	/** Initializes the toolkit and creates the curve editor, plus adds curves from the given objects. */
+	void Initialize(TArrayView<UObject*> InCurveOwners);
+
+	/** Whether this toolkit is initialized and has a valid curve editor. */
+	bool IsInitialized() const { return CurveEditor.IsValid(); }
+
+	/** Removes all curves and destroys the curve editor. */
+	void Shutdown();
+
+	/**
+	 * Adds curves for any curve property on the given object. If the object doesn't have any recognizable
+	 * curve properties, nothing happens.
+	 */
+	void AddCurveOwner(UObject* InCurveOwner);
+	/**
+	 * Removes any existing curves belonging to the given object.
+	 */
+	void RemoveCurveOwner(UObject* InCurveOwner);
+
+	/** Removes all curves from all curve owners. */
+	void RemoveAllCurveOwners();
+
+	/** Selects the curves associated with the given object and property name. */
+	void SelectCurves(UObject* InCurveOwner, FName InPropertyName);
+
+	/** Gets the curve editor widget. */
+	TSharedPtr<SWidget> GetCurveEditorWidget() const { return CurveEditorWidget; }
+
+private:
+
+	void AddCurves(UObject* InObject);
+	FCurveEditorTreeItem* AddTreeItem(FCurveEditorTreeItemID ParentID, FCurvePropertyInfo&& CurveInfo);
+
+private:
+
+	TSharedPtr<FCurveEditor> CurveEditor;
+	TSharedPtr<SWidget> CurveEditorWidget;
+};
+
+}  // namespace UE::Cameras
+
