Cameras: add spline orbit camera node

diff --git a/Engine/Plugins/Cameras/GameplayCameras/Source/GameplayCameras/Private/Curves/CameraRotatorCurve.cpp b/Engine/Plugins/Cameras/GameplayCameras/Source/GameplayCameras/Private/Curves/CameraRotatorCurve.cpp
--- a/Engine/Plugins/Cameras/GameplayCameras/Source/GameplayCameras/Private/Curves/CameraRotatorCurve.cpp
+++ b/Engine/Plugins/Cameras/GameplayCameras/Source/GameplayCameras/Private/Curves/CameraRotatorCurve.cpp
@@ -13,3 +13,8 @@ FRotator FCameraRotatorCurve::GetValue(float InTime) const
 	return Result;
 }
 
+bool FCameraRotatorCurve::HasAnyData() const
+{
+	return Curves[0].HasAnyData() || Curves[1].HasAnyData() || Curves[2].HasAnyData();
+}
+
diff --git a/Engine/Plugins/Cameras/GameplayCameras/Source/GameplayCameras/Private/Curves/CameraSingleCurve.cpp b/Engine/Plugins/Cameras/GameplayCameras/Source/GameplayCameras/Private/Curves/CameraSingleCurve.cpp
--- a/Engine/Plugins/Cameras/GameplayCameras/Source/GameplayCameras/Private/Curves/CameraSingleCurve.cpp
+++ b/Engine/Plugins/Cameras/GameplayCameras/Source/GameplayCameras/Private/Curves/CameraSingleCurve.cpp
@@ -7,3 +7,8 @@ float FCameraSingleCurve::GetValue(float InTime) const
 	return Curve.Eval(InTime);
 }
 
+bool FCameraSingleCurve::HasAnyData() const
+{
+	return Curve.HasAnyData();
+}
+
diff --git a/Engine/Plugins/Cameras/GameplayCameras/Source/GameplayCameras/Private/Curves/CameraVectorCurve.cpp b/Engine/Plugins/Cameras/GameplayCameras/Source/GameplayCameras/Private/Curves/CameraVectorCurve.cpp
--- a/Engine/Plugins/Cameras/GameplayCameras/Source/GameplayCameras/Private/Curves/CameraVectorCurve.cpp
+++ b/Engine/Plugins/Cameras/GameplayCameras/Source/GameplayCameras/Private/Curves/CameraVectorCurve.cpp
@@ -11,3 +11,8 @@ FVector FCameraVectorCurve::GetValue(float InTime) const
 	return Result;
 }
 
+bool FCameraVectorCurve::HasAnyData() const
+{
+	return Curves[0].HasAnyData() || Curves[1].HasAnyData() || Curves[2].HasAnyData();
+}
+
diff --git a/Engine/Plugins/Cameras/GameplayCameras/Source/GameplayCameras/Private/Nodes/Common/SplineOrbitCameraNode.cpp b/Engine/Plugins/Cameras/GameplayCameras/Source/GameplayCameras/Private/Nodes/Common/SplineOrbitCameraNode.cpp
new file mode 100644
--- /dev/null
+++ b/Engine/Plugins/Cameras/GameplayCameras/Source/GameplayCameras/Private/Nodes/Common/SplineOrbitCameraNode.cpp
@@ -0,0 +1,444 @@
+// Copyright Epic Games, Inc. All Rights Reserved.
+
+#include "Nodes/Common/SplineOrbitCameraNode.h"
+
+#include "Core/CameraEvaluationContext.h"
+#include "Core/CameraOperation.h"
+#include "Core/CameraRigJoints.h"
+#include "Debug/CameraDebugBlock.h"
+#include "Debug/CameraDebugBlockBuilder.h"
+#include "Debug/CameraDebugRenderer.h"
+#include "GameFramework/Pawn.h"
+#include "GameFramework/PlayerController.h"
+#include "HAL/IConsoleManager.h"
+#include "IGameplayCamerasModule.h"
+#include "IGameplayCamerasLiveEditListener.h"
+#include "IGameplayCamerasLiveEditManager.h"
+#include "Math/Axis.h"
+#include "Math/CameraNodeSpaceMath.h"
+#include "Math/InterpCurve.h"
+#include "Math/Ray.h"
+#include "Nodes/Input/Input2DCameraNode.h"
+
+#include UE_INLINE_GENERATED_CPP_BY_NAME(SplineOrbitCameraNode)
+
+namespace UE::Cameras
+{
+
+bool GGameplayCamerasSplineOrbitShowLocationOffsetSpline = true;
+static FAutoConsoleVariableRef CVarGameplayCamerasSplineOrbitShowLocationOffsetSpline(
+	TEXT("GameplayCameras.SplineOrbit.ShowLocationOffsetSpline"),
+	GGameplayCamerasSplineOrbitShowLocationOffsetSpline,
+	TEXT("(Default: 1. Whether to show the camera's spline trajectory."));
+
+bool GGameplayCamerasSplineOrbitShowLocationOffsetOrbits = true;
+static FAutoConsoleVariableRef CVarGameplayCamerasSplineOrbitShowLocationOffsetOrbits(
+	TEXT("GameplayCameras.SplineOrbit.ShowLocationOffsetOrbits"),
+	GGameplayCamerasSplineOrbitShowLocationOffsetOrbits,
+	TEXT("Default: 1. Whether to show the control points' orbits."));
+
+class FSplineOrbitCameraNodeEvaluator 
+	: public FCameraNodeEvaluator
+#if WITH_EDITOR
+    , public IGameplayCamerasLiveEditListener
+#endif  // WITH_EDITOR
+{
+	UE_DECLARE_CAMERA_NODE_EVALUATOR(GAMEPLAYCAMERAS_API, FSplineOrbitCameraNodeEvaluator)
+
+public:
+
+	~FSplineOrbitCameraNodeEvaluator();
+
+protected:
+
+	virtual void OnBuild(const FCameraNodeEvaluatorBuildParams& Params) override;
+	virtual void OnInitialize(const FCameraNodeEvaluatorInitializeParams& Params, FCameraNodeEvaluationResult& OutResult) override;
+	virtual FCameraNodeEvaluatorChildrenView OnGetChildren() override;
+	virtual void OnRun(const FCameraNodeEvaluationParams& Params, FCameraNodeEvaluationResult& OutResult) override;
+	virtual void OnExecuteOperation(const FCameraOperationParams& Params, FCameraOperation& Operation) override;
+
+#if WITH_EDITOR
+	virtual void OnPostEditChangeProperty(const UCameraNode* InCameraNode, const FPropertyChangedEvent& PropertyChangedEvent) override;
+#endif  // WITH_EDITOR
+
+#if UE_GAMEPLAY_CAMERAS_DEBUG
+	virtual void OnBuildDebugBlocks(const FCameraDebugBlockBuildParams& Params, FCameraDebugBlockBuilder& Builder) override;
+#endif  // UE_GAMEPLAY_CAMERAS_DEBUG
+
+private:
+
+	APlayerController* GetPlayerController(TSharedPtr<const FCameraEvaluationContext> EvaluationContext) const;
+
+	void RebuildCurves();
+
+private:
+
+	FInput2DCameraNodeEvaluator* InputSlotEvaluator = nullptr;
+
+	FCompressedRichCurve LocationOffsetSpline[3];
+	FCompressedRichCurve TargetOffsetSpline[3];
+	FCompressedRichCurve RotationOffsetSpline[3];
+
+	bool bHasAnyTargetOffset = false;
+	bool bHasAnyRotationOffset = false;
+
+#if UE_GAMEPLAY_CAMERAS_DEBUG
+	FVector2d DebugYawPitch;
+	FTransform DebugPivotTransform;
+	FVector3d DebugLocationOffset;
+	FVector3d DebugWorldTargetOffset;
+	FRotator3d DebugRotationOffset;
+#endif  // UE_GAMEPLAY_CAMERAS_DEBUG
+};
+
+UE_DEFINE_CAMERA_NODE_EVALUATOR(FSplineOrbitCameraNodeEvaluator)
+
+UE_DECLARE_CAMERA_DEBUG_BLOCK_START(GAMEPLAYCAMERAS_API, FSplineOrbitCameraDebugBlock)
+	UE_DECLARE_CAMERA_DEBUG_BLOCK_FIELD(FVector2d, OrbitYawPitch);
+	UE_DECLARE_CAMERA_DEBUG_BLOCK_FIELD(FTransform3d, PivotTransform);
+	UE_DECLARE_CAMERA_DEBUG_BLOCK_FIELD(FVector3d, LocationOffset);
+	UE_DECLARE_CAMERA_DEBUG_BLOCK_FIELD(FVector3d, WorldTargetOffset);
+	UE_DECLARE_CAMERA_DEBUG_BLOCK_FIELD(FRotator3d, RotationOffset);
+#if UE_GAMEPLAY_CAMERAS_DEBUG
+	// Pointer to the spline orbit node for rendering the orbit spline in debug cameras.
+	// TODO: This won't get serialized when recording gameplay (e.g. with RewindDebugger). We need to add
+	//		 support for global data shared between frames, otherwise would have to serialize the curves 
+	//		 each frame.
+	public:
+		TWeakObjectPtr<const USplineOrbitCameraNode> WeakSplineOrbitNode;
+	private:
+		void RenderLocationOffsetSpline(const USplineOrbitCameraNode* SplineOrbitNode, FCameraDebugRenderer& Renderer);
+		void RenderLocationOffsetOrbits(const USplineOrbitCameraNode* SplineOrbitNode, FCameraDebugRenderer& Renderer);
+#endif  // UE_GAMEPLAY_CAMERAS_DEBUG
+UE_DECLARE_CAMERA_DEBUG_BLOCK_END()
+
+UE_DEFINE_CAMERA_DEBUG_BLOCK_WITH_FIELDS(FSplineOrbitCameraDebugBlock)
+
+void FSplineOrbitCameraNodeEvaluator::OnBuild(const FCameraNodeEvaluatorBuildParams& Params)
+{
+	const USplineOrbitCameraNode* SplineOrbitNode = GetCameraNodeAs<USplineOrbitCameraNode>();
+	InputSlotEvaluator = Params.BuildEvaluatorAs<FInput2DCameraNodeEvaluator>(SplineOrbitNode->InputSlot);
+}
+
+void FSplineOrbitCameraNodeEvaluator::OnInitialize(const FCameraNodeEvaluatorInitializeParams& Params, FCameraNodeEvaluationResult& OutResult)
+{
+	AddNodeEvaluatorFlags(ECameraNodeEvaluatorFlags::SupportsOperations);
+
+	RebuildCurves();
+
+#if WITH_EDITOR
+	IGameplayCamerasModule& GameplayCamerasModule = IGameplayCamerasModule::Get();
+	if (TSharedPtr<IGameplayCamerasLiveEditManager> LiveEditManager = GameplayCamerasModule.GetLiveEditManager())
+	{
+		LiveEditManager->AddListener(GetCameraNode(), this);
+	}
+#endif  // WITH_EDITOR
+}
+
+FSplineOrbitCameraNodeEvaluator::~FSplineOrbitCameraNodeEvaluator()
+{
+#if WITH_EDITOR
+	IGameplayCamerasModule& GameplayCamerasModule = IGameplayCamerasModule::Get();
+	if (TSharedPtr<IGameplayCamerasLiveEditManager> LiveEditManager = GameplayCamerasModule.GetLiveEditManager())
+	{
+		LiveEditManager->RemoveListener(this);
+	}
+#endif  // WITH_EDITOR
+}
+
+void FSplineOrbitCameraNodeEvaluator::RebuildCurves()
+{
+	const USplineOrbitCameraNode* SplineOrbitNode = GetCameraNodeAs<USplineOrbitCameraNode>();
+
+	bHasAnyTargetOffset = SplineOrbitNode->TargetOffsetSpline.HasAnyData();
+	bHasAnyRotationOffset = SplineOrbitNode->RotationOffsetSpline.HasAnyData();
+
+	SplineOrbitNode->LocationOffsetSpline.Curves[0].CompressCurve(LocationOffsetSpline[0]);
+	SplineOrbitNode->LocationOffsetSpline.Curves[1].CompressCurve(LocationOffsetSpline[1]);
+	SplineOrbitNode->LocationOffsetSpline.Curves[2].CompressCurve(LocationOffsetSpline[2]);
+
+	SplineOrbitNode->TargetOffsetSpline.Curves[0].CompressCurve(TargetOffsetSpline[0]);
+	SplineOrbitNode->TargetOffsetSpline.Curves[1].CompressCurve(TargetOffsetSpline[1]);
+	SplineOrbitNode->TargetOffsetSpline.Curves[2].CompressCurve(TargetOffsetSpline[2]);
+
+	SplineOrbitNode->RotationOffsetSpline.Curves[0].CompressCurve(RotationOffsetSpline[0]);
+	SplineOrbitNode->RotationOffsetSpline.Curves[1].CompressCurve(RotationOffsetSpline[1]);
+	SplineOrbitNode->RotationOffsetSpline.Curves[2].CompressCurve(RotationOffsetSpline[2]);
+}
+
+FCameraNodeEvaluatorChildrenView FSplineOrbitCameraNodeEvaluator::OnGetChildren()
+{
+	return FCameraNodeEvaluatorChildrenView({ InputSlotEvaluator });
+}
+
+void FSplineOrbitCameraNodeEvaluator::OnRun(const FCameraNodeEvaluationParams& Params, FCameraNodeEvaluationResult& OutResult)
+{
+	FRotator3d OrbitRotation = FRotator3d::ZeroRotator;
+	if (InputSlotEvaluator)
+	{
+		InputSlotEvaluator->Run(Params, OutResult);
+		const FVector2d YawPitch = InputSlotEvaluator->GetInputValue();
+		OrbitRotation = FRotator3d(YawPitch.Y, YawPitch.X, 0);
+	}
+	else if (APlayerController* PlayerController = GetPlayerController(Params.EvaluationContext))
+	{
+		const FRotator3d ControlRotation = PlayerController->GetControlRotation();
+		OrbitRotation = ControlRotation;
+	}
+
+	const USplineOrbitCameraNode* SplineOrbitNode = GetCameraNodeAs<USplineOrbitCameraNode>();
+
+	// Compute the camera transform similarly to the boom arm (see FBoomArmCameraNodeEvaluator).
+	const FTransform3d OrbitPivot(OrbitRotation, OutResult.CameraPose.GetLocation());
+	const FVector3d LocationOffset = FVector3d(
+			LocationOffsetSpline[0].Eval(OrbitRotation.Pitch),
+			LocationOffsetSpline[1].Eval(OrbitRotation.Pitch),
+			LocationOffsetSpline[2].Eval(OrbitRotation.Pitch));
+	FTransform3d OrbitTransform(FTransform3d(LocationOffset) * OrbitPivot);
+
+#if UE_GAMEPLAY_CAMERAS_DEBUG
+	DebugYawPitch = FVector2d(OrbitRotation.Yaw, OrbitRotation.Pitch);
+	DebugPivotTransform = OrbitPivot;
+	DebugLocationOffset = LocationOffset;
+	DebugWorldTargetOffset = FVector3d::ZeroVector;
+	DebugRotationOffset = FRotator3d::ZeroRotator;
+#endif  // UE_GAMEPLAY_CAMERAS_DEBUG
+
+	// Now rotate the camera according to the target offset, if any.
+	if (bHasAnyTargetOffset)
+	{
+		const FVector3d TargetOffset = FVector3d(
+				TargetOffsetSpline[0].Eval(OrbitRotation.Pitch),
+				TargetOffsetSpline[1].Eval(OrbitRotation.Pitch),
+				TargetOffsetSpline[2].Eval(OrbitRotation.Pitch));
+		if (!TargetOffset.IsNearlyZero())
+		{
+			// Project the orbit center onto the line of sight. It would already be on it if there wasn't
+			// any lateral or vertical offset in LocationOffset, but most often there is, so the line of sight
+			// is offset from the center.
+			const FRay LineOfSight(OrbitTransform.GetLocation(), OrbitTransform.GetRotation().GetForwardVector());
+			const FVector3d ProjectedOrbitPivot = LineOfSight.ClosestPoint(OrbitPivot.GetLocation());
+
+			// Now use this projected point as the "target" of the camera for the purposes of orbiting.
+			// Offset that target and make the camera look at the new target.
+			FVector3d NewTarget;
+			FCameraNodeSpaceParams SpaceMathParams(Params, OutResult);
+			FCameraNodeSpaceMath::OffsetCameraNodeSpacePosition(
+					SpaceMathParams, ProjectedOrbitPivot, TargetOffset, SplineOrbitNode->TargetOffsetSpace, NewTarget);
+
+			const FVector3d NewLineOfSight = (NewTarget - LineOfSight.Origin);
+			FRotator3d NewOrientation = NewLineOfSight.ToOrientationRotator();
+			NewOrientation.Roll = OrbitTransform.Rotator().Roll;
+			OrbitTransform.SetRotation(NewOrientation.Quaternion());
+
+#if UE_GAMEPLAY_CAMERAS_DEBUG
+			DebugWorldTargetOffset = (NewTarget - ProjectedOrbitPivot);
+#endif  // UE_GAMEPLAY_CAMERAS_DEBUG
+		}
+	}
+
+	// Add any extra rotation if needed.
+	if (bHasAnyRotationOffset)
+	{
+		const FVector3d RotationOffset = FVector3d(
+				RotationOffsetSpline[0].Eval(OrbitRotation.Pitch),
+				RotationOffsetSpline[1].Eval(OrbitRotation.Pitch),
+				RotationOffsetSpline[2].Eval(OrbitRotation.Pitch));
+		if (!RotationOffset.IsNearlyZero())
+		{
+			const FTransform3d RotationOffsetTransform(FRotator3d::MakeFromEuler(RotationOffset));
+			OrbitTransform = RotationOffsetTransform * OrbitTransform;
+
+#if UE_GAMEPLAY_CAMERAS_DEBUG
+			DebugRotationOffset = FRotator3d::MakeFromEuler(RotationOffset);
+#endif  // UE_GAMEPLAY_CAMERAS_DEBUG
+		}
+	}
+
+	OutResult.CameraPose.SetTransform(OrbitTransform);
+
+	OutResult.CameraRigJoints.AddYawPitchJoint(OrbitPivot);
+}
+
+void FSplineOrbitCameraNodeEvaluator::OnExecuteOperation(const FCameraOperationParams& Params, FCameraOperation& Operation)
+{
+	if (!InputSlotEvaluator)
+	{
+		// If we don't have an input slot, we use the pawn rotation directly in OnRun. So let's handle
+		// some operations by affecting that pawn rotation ourselves.
+		if (FYawPitchCameraOperation* Op = Operation.CastOperation<FYawPitchCameraOperation>())
+		{
+			if (APlayerController* PlayerController = GetPlayerController(Params.EvaluationContext))
+			{
+				FRotator3d ControlRotation = PlayerController->GetControlRotation();
+				ControlRotation.Yaw = Op->Yaw.Apply(ControlRotation.Yaw);
+				ControlRotation.Pitch = Op->Pitch.Apply(ControlRotation.Pitch);
+				PlayerController->SetControlRotation(ControlRotation);
+			}
+		}
+	}
+}
+
+APlayerController* FSplineOrbitCameraNodeEvaluator::GetPlayerController(TSharedPtr<const FCameraEvaluationContext> EvaluationContext) const
+{
+	if (EvaluationContext)
+	{
+		return EvaluationContext->GetPlayerController();
+	}
+	return nullptr;
+}
+
+#if WITH_EDITOR
+
+void FSplineOrbitCameraNodeEvaluator::OnPostEditChangeProperty(const UCameraNode* InCameraNode, const FPropertyChangedEvent& PropertyChangedEvent)
+{
+	const FName PropertyName = PropertyChangedEvent.GetMemberPropertyName();
+	if (PropertyName == GET_MEMBER_NAME_CHECKED(USplineOrbitCameraNode, LocationOffsetSpline) ||
+			PropertyName == GET_MEMBER_NAME_CHECKED(USplineOrbitCameraNode, TargetOffsetSpline) ||
+			PropertyName == GET_MEMBER_NAME_CHECKED(USplineOrbitCameraNode, RotationOffsetSpline))
+	{
+		RebuildCurves();
+	}
+}
+
+#endif  // WITH_EDITOR
+
+#if UE_GAMEPLAY_CAMERAS_DEBUG
+
+void FSplineOrbitCameraNodeEvaluator::OnBuildDebugBlocks(const FCameraDebugBlockBuildParams& Params, FCameraDebugBlockBuilder& Builder)
+{
+	FSplineOrbitCameraDebugBlock& DebugBlock = Builder.AttachDebugBlock<FSplineOrbitCameraDebugBlock>();
+
+	DebugBlock.OrbitYawPitch = DebugYawPitch;
+	DebugBlock.PivotTransform = DebugPivotTransform;
+	DebugBlock.LocationOffset = DebugLocationOffset;
+	DebugBlock.WorldTargetOffset = DebugWorldTargetOffset;
+	DebugBlock.RotationOffset = DebugRotationOffset;
+	DebugBlock.WeakSplineOrbitNode = GetCameraNodeAs<USplineOrbitCameraNode>();
+}
+
+void FSplineOrbitCameraDebugBlock::OnDebugDraw(const FCameraDebugBlockDrawParams& Params, FCameraDebugRenderer& Renderer)
+{
+	Renderer.AddText(TEXT("yaw: %.3f pitch %.3f"), OrbitYawPitch.X, OrbitYawPitch.Y);
+	Renderer.NewLine();
+	Renderer.AddIndent();
+	{
+		Renderer.AddText(TEXT("location offset: %s\n"), *LocationOffset.ToString());
+		Renderer.AddText(TEXT("target offset: %s\n"), *WorldTargetOffset.ToString());
+		Renderer.AddText(TEXT("rotation offset: %s\n"), *RotationOffset.ToString());
+	}
+	Renderer.RemoveIndent();
+
+	if (Renderer.IsExternalRendering())
+	{
+		// TODO: right now we need to keep a pointer to the camera node to read the spline data (see previous comment).
+		if (const USplineOrbitCameraNode* SplineOrbitNode = WeakSplineOrbitNode.Get())
+		{
+			RenderLocationOffsetSpline(SplineOrbitNode, Renderer);	
+			RenderLocationOffsetOrbits(SplineOrbitNode, Renderer);
+
+			// Render the orbit pivot.
+			Renderer.DrawSphere(PivotTransform.GetLocation(), 1.f, 8, FColorList::Brass, 1.f);
+		}
+	}
+}
+
+void FSplineOrbitCameraDebugBlock::RenderLocationOffsetSpline(const USplineOrbitCameraNode* SplineOrbitNode, FCameraDebugRenderer& Renderer)
+{
+	if (!GGameplayCamerasSplineOrbitShowLocationOffsetSpline)
+	{
+		return;
+	}
+
+	// Sample the curve for pitch values inside a sensible range.
+	// We use the uncompressed curves directly from the camera node here because the compressed curves
+	// from the evaluator may not be available if we recorded the debug data and the evaluator is gone.
+	auto ComputeLocationOffsetSample = [this, SplineOrbitNode](float InPitchAngle) -> FVector3d
+	{
+		const FRotator3d CurOrbitRotation(InPitchAngle, OrbitYawPitch.X, 0.f);
+		const FTransform3d CurPivotTransform(CurOrbitRotation, PivotTransform.GetLocation());
+
+		const FVector3d CurLocationOffset = SplineOrbitNode->LocationOffsetSpline.GetValue(InPitchAngle);
+		const FTransform3d CurOrbitTransform(FTransform3d(CurLocationOffset) * CurPivotTransform);
+		return CurOrbitTransform.GetLocation();
+	};
+
+	const FLinearColor SplineColor(FColorList::OrangeRed);
+	FVector3d PrevSamplePoint = ComputeLocationOffsetSample(-89.f);
+
+	for (float CurPitchAngle = -89.f + 2.f; CurPitchAngle <= 89.f; CurPitchAngle += 2.f)
+	{
+		FVector3d NextSamplePoint = ComputeLocationOffsetSample(CurPitchAngle);
+
+		Renderer.DrawLine(PrevSamplePoint, NextSamplePoint, SplineColor);
+
+		PrevSamplePoint = NextSamplePoint;
+	}
+}
+
+void FSplineOrbitCameraDebugBlock::RenderLocationOffsetOrbits(const USplineOrbitCameraNode* SplineOrbitNode, FCameraDebugRenderer& Renderer)
+{
+	if (!GGameplayCamerasSplineOrbitShowLocationOffsetOrbits)
+	{
+		return;
+	}
+
+	// Gather the control points' pitch values. We'll draw an ellipse at each of them.
+	TArray<float, TInlineAllocator<32>> PitchValues;
+	for (int32 CurveIndex = 0; CurveIndex < 3; ++CurveIndex)
+	{
+		for (auto It = SplineOrbitNode->LocationOffsetSpline.Curves[CurveIndex].GetKeyIterator(); It; ++It)
+		{
+			const float PitchValue = It->Time;
+			if (PitchValue >= -89.f && PitchValue <= 89.f)
+			{
+				PitchValues.Add(It->Time);
+			}
+		}
+	}
+
+	// Now sample the location offset all around the pivot (360 degrees) for each pitch value of
+	// the control points.
+	auto ComputeLocationOffsetSample = [this, SplineOrbitNode](
+			float InYawAngle, float InPitchAngle, const FVector3d& InLocationOffset) -> FVector3d
+	{
+		const FRotator3d CurOrbitRotation(InPitchAngle, InYawAngle, 0.f);
+		const FTransform3d CurPivotTransform(CurOrbitRotation, PivotTransform.GetLocation());
+
+		const FTransform3d CurOrbitTransform(FTransform3d(InLocationOffset) * CurPivotTransform);
+		return CurOrbitTransform.GetLocation();
+	};
+
+	const float OrbitAngleStep = 2.f;
+	const FLinearColor EllipseColor(FColorList::Brass);
+	for (float PitchValue : PitchValues)
+	{
+		const FVector3d CurLocationOffset = SplineOrbitNode->LocationOffsetSpline.GetValue(PitchValue);
+
+		FVector3d  PrevSamplePoint = ComputeLocationOffsetSample(0.f, PitchValue, CurLocationOffset);
+		for (float CurEllipseAngle = OrbitAngleStep; CurEllipseAngle < 360.f; CurEllipseAngle += OrbitAngleStep)
+		{
+			FVector3d NextSamplePoint = ComputeLocationOffsetSample(CurEllipseAngle, PitchValue, CurLocationOffset);
+
+			Renderer.DrawLine(PrevSamplePoint, NextSamplePoint, EllipseColor);
+
+			PrevSamplePoint = NextSamplePoint;
+		}
+	}
+}
+
+#endif  // UE_GAMEPLAY_CAMERAS_DEBUG
+
+}  // namespace UE::Cameras
+
+FCameraNodeChildrenView USplineOrbitCameraNode::OnGetChildren()
+{
+	return FCameraNodeChildrenView({ InputSlot });
+}
+
+FCameraNodeEvaluatorPtr USplineOrbitCameraNode::OnBuildEvaluator(FCameraNodeEvaluatorBuilder& Builder) const
+{
+	using namespace UE::Cameras;
+	return Builder.BuildEvaluator<FSplineOrbitCameraNodeEvaluator>();
+}
+
diff --git a/Engine/Plugins/Cameras/GameplayCameras/Source/GameplayCameras/Public/Curves/CameraRotatorCurve.h b/Engine/Plugins/Cameras/GameplayCameras/Source/GameplayCameras/Public/Curves/CameraRotatorCurve.h
--- a/Engine/Plugins/Cameras/GameplayCameras/Source/GameplayCameras/Public/Curves/CameraRotatorCurve.h
+++ b/Engine/Plugins/Cameras/GameplayCameras/Source/GameplayCameras/Public/Curves/CameraRotatorCurve.h
@@ -8,7 +8,7 @@
 #include "CameraRotatorCurve.generated.h"
 
 USTRUCT()
-struct FCameraRotatorCurve
+struct GAMEPLAYCAMERAS_API FCameraRotatorCurve
 {
 	GENERATED_BODY()
 
@@ -16,5 +16,6 @@ struct FCameraRotatorCurve
 	FRichCurve Curves[3];
 
 	FRotator GetValue(float InTime) const;
+	bool HasAnyData() const;
 };
 
diff --git a/Engine/Plugins/Cameras/GameplayCameras/Source/GameplayCameras/Public/Curves/CameraSingleCurve.h b/Engine/Plugins/Cameras/GameplayCameras/Source/GameplayCameras/Public/Curves/CameraSingleCurve.h
--- a/Engine/Plugins/Cameras/GameplayCameras/Source/GameplayCameras/Public/Curves/CameraSingleCurve.h
+++ b/Engine/Plugins/Cameras/GameplayCameras/Source/GameplayCameras/Public/Curves/CameraSingleCurve.h
@@ -8,7 +8,7 @@
 #include "CameraSingleCurve.generated.h"
 
 USTRUCT()
-struct FCameraSingleCurve
+struct GAMEPLAYCAMERAS_API FCameraSingleCurve
 {
 	GENERATED_BODY()
 
@@ -16,5 +16,6 @@ struct FCameraSingleCurve
 	FRichCurve Curve;
 
 	float GetValue(float InTime) const;
+	bool HasAnyData() const;
 };
 
diff --git a/Engine/Plugins/Cameras/GameplayCameras/Source/GameplayCameras/Public/Curves/CameraVectorCurve.h b/Engine/Plugins/Cameras/GameplayCameras/Source/GameplayCameras/Public/Curves/CameraVectorCurve.h
--- a/Engine/Plugins/Cameras/GameplayCameras/Source/GameplayCameras/Public/Curves/CameraVectorCurve.h
+++ b/Engine/Plugins/Cameras/GameplayCameras/Source/GameplayCameras/Public/Curves/CameraVectorCurve.h
@@ -8,7 +8,7 @@
 #include "CameraVectorCurve.generated.h"
 
 USTRUCT()
-struct FCameraVectorCurve
+struct GAMEPLAYCAMERAS_API FCameraVectorCurve
 {
 	GENERATED_BODY()
 
@@ -16,5 +16,6 @@ struct FCameraVectorCurve
 	FRichCurve Curves[3];
 
 	FVector GetValue(float InTime) const;
+	bool HasAnyData() const;
 };
 
diff --git a/Engine/Plugins/Cameras/GameplayCameras/Source/GameplayCameras/Public/Nodes/Common/SplineOrbitCameraNode.h b/Engine/Plugins/Cameras/GameplayCameras/Source/GameplayCameras/Public/Nodes/Common/SplineOrbitCameraNode.h
new file mode 100644
--- /dev/null
+++ b/Engine/Plugins/Cameras/GameplayCameras/Source/GameplayCameras/Public/Nodes/Common/SplineOrbitCameraNode.h
@@ -0,0 +1,96 @@
+// Copyright Epic Games, Inc. All Rights Reserved.
+
+#pragma once
+
+#include "Core/CameraNode.h"
+#include "Core/CameraParameters.h"
+#include "Curves/CameraRotatorCurve.h"
+#include "Curves/CameraVectorCurve.h"
+#include "Nodes/CameraNodeTypes.h"
+
+#include "SplineOrbitCameraNode.generated.h"
+
+class UCameraValueInterpolator;
+class UInput2DCameraNode;
+
+/**
+ * Control point for the spline orbit camera node.
+ */
+USTRUCT()
+struct FSplineOrbitControlPoint
+{
+	GENERATED_BODY()
+
+	/**
+	 * The offset of the camera position from the orbit pivot.
+	 * This defines the radius of the orbit at the given control point, along with any lateral or 
+	 * vertical position offsets.
+	 */
+	UPROPERTY(EditAnywhere, Category="Spline Orbit")
+	FVector3d LocationOffset = FVector3d::ZeroVector;
+
+	/**
+	 * The offset of the camera target as defined by projecting the orbit pivot on the line of sight.
+	 * This adds rotation to the camera by making it look higher/lower/etc at the given control
+	 * point.
+	 */
+	UPROPERTY(EditAnywhere, Category="Target")
+	FVector3d TargetOffset = FVector3d::ZeroVector;
+
+	/**
+	 * A rotation offset applied to the camera.
+	 * This adds rotation to the camera, in local space, applied after TargetOffset.
+	 */
+	UPROPERTY(EditAnywhere, Category="Target")
+	FRotator3d RotationOffset = FRotator3d::ZeroRotator;
+
+	/**
+	 * The pitch angle for this control point.
+	 */
+	UPROPERTY(EditAnywhere, Category="Spline Orbit")
+	float PitchAngle = 0.f;
+};
+
+/**
+ * A camera node that can orbit around a pivot point, and the shape of the orbit is defined
+ * by pitch-based parameters.
+ */
+UCLASS(MinimalAPI, meta=(CameraNodeCategories="Common,Transform"))
+class USplineOrbitCameraNode : public UCameraNode
+{
+	GENERATED_BODY()
+
+protected:
+
+	// UCameraNode interface.
+	virtual FCameraNodeChildrenView OnGetChildren() override;
+	virtual FCameraNodeEvaluatorPtr OnBuildEvaluator(FCameraNodeEvaluatorBuilder& Builder) const override;
+
+public:
+
+	/** Spline that defines the camera location's offset for a given pitch angle. */
+	UPROPERTY(EditAnywhere, Category="Spline Orbit", meta=(CurveEditorXLabelFormat="{0}deg"))
+	FCameraVectorCurve LocationOffsetSpline;
+
+	/** Spline that defines an additive camera target offset for a given pitch angle. */
+	UPROPERTY(EditAnywhere, Category="Spline Orbit", meta=(CurveEditorXLabelFormat="{0}deg"))
+	FCameraVectorCurve TargetOffsetSpline;
+
+	/** Spline that defines an additive camera rotation offset for a given pitch angle. */
+	UPROPERTY(EditAnywhere, Category="Spline Orbit", meta=(CurveEditorXLabelFormat="{0}deg"))
+	FCameraRotatorCurve RotationOffsetSpline;
+
+	/**
+	 * The space in which the control points' TargetOffset is applied.
+	 */
+	UPROPERTY(EditAnywhere, Category="Target")
+	ECameraNodeSpace TargetOffsetSpace = ECameraNodeSpace::ActiveContext;
+
+	/**
+	 * The input slot for controlling the orbiting.
+	 * If no input slot is specified, this node will use the player controller view rotation.
+	 */
+	UPROPERTY(meta=(ObjectTreeGraphPinDirection=Input))
+	TObjectPtr<UInput2DCameraNode> InputSlot;
+};
+
