Cameras: transition nodes are now all coming out from the output side of the rig node, as normal

This removes the "enter transitions as inputs, exit transitions as outputs" UX that was previously in the transition graph editor.

diff --git a/Engine/Plugins/Cameras/GameplayCameras/Source/GameplayCameras/Public/Core/CameraAsset.h b/Engine/Plugins/Cameras/GameplayCameras/Source/GameplayCameras/Public/Core/CameraAsset.h
--- a/Engine/Plugins/Cameras/GameplayCameras/Source/GameplayCameras/Public/Core/CameraAsset.h
+++ b/Engine/Plugins/Cameras/GameplayCameras/Source/GameplayCameras/Public/Core/CameraAsset.h
@@ -137,11 +137,11 @@ private:
 	TArray<TObjectPtr<UCameraRigAsset>> CameraRigs_DEPRECATED;
 
 	/** A list of default enter transitions for all the camera rigs in this asset. */
-	UPROPERTY(meta=(ObjectTreeGraphPinDirection=Input))
+	UPROPERTY(Instanced)
 	TArray<TObjectPtr<UCameraRigTransition>> EnterTransitions;
 
 	/** A list of default exit transitions for all the camera rigs in this asset. */
-	UPROPERTY(meta=(ObjectTreeGraphPinDirection=Output))
+	UPROPERTY(Instanced)
 	TArray<TObjectPtr<UCameraRigTransition>> ExitTransitions;
 
 	/** The current build state of this camera asset. */
diff --git a/Engine/Plugins/Cameras/GameplayCameras/Source/GameplayCameras/Public/Core/CameraRigAsset.h b/Engine/Plugins/Cameras/GameplayCameras/Source/GameplayCameras/Public/Core/CameraRigAsset.h
--- a/Engine/Plugins/Cameras/GameplayCameras/Source/GameplayCameras/Public/Core/CameraRigAsset.h
+++ b/Engine/Plugins/Cameras/GameplayCameras/Source/GameplayCameras/Public/Core/CameraRigAsset.h
@@ -208,11 +208,11 @@ public:
 	FCameraRigInterface Interface;
 
 	/** List of enter transitions for this camera rig. */
-	UPROPERTY(Instanced, meta=(ObjectTreeGraphPinDirection=Input))
+	UPROPERTY(Instanced)
 	TArray<TObjectPtr<UCameraRigTransition>> EnterTransitions;
 
 	/** List of exist transitions for this camera rig. */
-	UPROPERTY(Instanced, meta=(ObjectTreeGraphPinDirection=Output))
+	UPROPERTY(Instanced)
 	TArray<TObjectPtr<UCameraRigTransition>> ExitTransitions;
 	
 	/** Default orientation initialization when this camera rig is activated. */
diff --git a/Engine/Plugins/Cameras/GameplayCameras/Source/GameplayCamerasEditor/Private/Editors/CameraRigTransitionGraphSchema.cpp b/Engine/Plugins/Cameras/GameplayCameras/Source/GameplayCamerasEditor/Private/Editors/CameraRigTransitionGraphSchema.cpp
--- a/Engine/Plugins/Cameras/GameplayCameras/Source/GameplayCamerasEditor/Private/Editors/CameraRigTransitionGraphSchema.cpp
+++ b/Engine/Plugins/Cameras/GameplayCameras/Source/GameplayCamerasEditor/Private/Editors/CameraRigTransitionGraphSchema.cpp
@@ -26,34 +26,5 @@ void UCameraRigTransitionGraphSchema::OnBuildGraphConfig(FObjectTreeGraphConfig&
 		.NodeTitleColor(Settings->CameraRigAssetTitleColor);
 }
 
-ETransitionGraphContextActions UCameraRigTransitionGraphSchema::GetTransitionGraphContextActions(FGraphContextMenuBuilder& ContextMenuBuilder) const
-{
-	ETransitionGraphContextActions PossibleActions = ETransitionGraphContextActions::None;
-
-	if (const UEdGraphPin* DraggedPin = ContextMenuBuilder.FromPin)
-	{
-		// If we are creating a node from dragging a pin into an empty space, figure out which transition
-		// we can create based on the direction of the dragged pin.
-		UObjectTreeGraphNode* OwningNode = Cast<UObjectTreeGraphNode>(DraggedPin->GetOwningNode());
-		if (OwningNode)
-		{
-			FProperty* DraggedPinProperty = OwningNode->GetPropertyForPin(DraggedPin);
-			if (DraggedPinProperty && DraggedPinProperty->GetOwnerClass()->IsChildOf<UCameraRigAsset>())
-			{
-				if (DraggedPinProperty->GetFName() == GET_MEMBER_NAME_CHECKED(UCameraRigAsset, EnterTransitions))
-				{
-					PossibleActions |= ETransitionGraphContextActions::CreateEnterTransition;
-				}
-				if (DraggedPinProperty->GetFName() == GET_MEMBER_NAME_CHECKED(UCameraRigAsset, ExitTransitions))
-				{
-					PossibleActions |= ETransitionGraphContextActions::CreateExitTransition;
-				}
-			}
-		}
-	}
-
-	return PossibleActions;
-}
-
 #undef LOCTEXT_NAMESPACE
 
diff --git a/Engine/Plugins/Cameras/GameplayCameras/Source/GameplayCamerasEditor/Private/Editors/CameraRigTransitionGraphSchemaBase.cpp b/Engine/Plugins/Cameras/GameplayCameras/Source/GameplayCamerasEditor/Private/Editors/CameraRigTransitionGraphSchemaBase.cpp
--- a/Engine/Plugins/Cameras/GameplayCameras/Source/GameplayCamerasEditor/Private/Editors/CameraRigTransitionGraphSchemaBase.cpp
+++ b/Engine/Plugins/Cameras/GameplayCameras/Source/GameplayCamerasEditor/Private/Editors/CameraRigTransitionGraphSchemaBase.cpp
@@ -132,59 +132,5 @@ void UCameraRigTransitionGraphSchemaBase::CollectAllObjects(UObjectTreeGraph* In
 	}
 }
 
-void UCameraRigTransitionGraphSchemaBase::GetGraphContextActions(FGraphContextMenuBuilder& ContextMenuBuilder) const
-{
-	ETransitionGraphContextActions PossibleActions = GetTransitionGraphContextActions(ContextMenuBuilder);
-
-	if (EnumHasAnyFlags(PossibleActions, ETransitionGraphContextActions::CreateEnterTransition))
-	{
-		TSharedRef<FCameraRigTransitionGraphSchemaAction_NewTransitionNode> EnterAction = MakeShared<FCameraRigTransitionGraphSchemaAction_NewTransitionNode>(
-				LOCTEXT("TransitionsCategory", "Transitions"),
-				LOCTEXT("EnterTransition", "Enter Transition"),
-				LOCTEXT("EnterTransitionToolTip", "Creates a new enter transition"));
-		EnterAction->TransitionType = FCameraRigTransitionGraphSchemaAction_NewTransitionNode::ETransitionType::Enter;
-		ContextMenuBuilder.AddAction(StaticCastSharedPtr<FEdGraphSchemaAction>(EnterAction.ToSharedPtr()));
-	}
-	
-	if (EnumHasAnyFlags(PossibleActions, ETransitionGraphContextActions::CreateExitTransition))
-	{
-		TSharedRef<FCameraRigTransitionGraphSchemaAction_NewTransitionNode> ExitAction = MakeShared<FCameraRigTransitionGraphSchemaAction_NewTransitionNode>(
-				LOCTEXT("TransitionsCategory", "Transitions"),
-				LOCTEXT("ExitTransition", "Exit Transition"),
-				LOCTEXT("ExitTransitionToolTip", "Creates a new exit transition"),
-				0);
-		ExitAction->TransitionType = FCameraRigTransitionGraphSchemaAction_NewTransitionNode::ETransitionType::Exit;
-		ContextMenuBuilder.AddAction(StaticCastSharedPtr<FEdGraphSchemaAction>(ExitAction.ToSharedPtr()));
-	}
-
-	Super::GetGraphContextActions(ContextMenuBuilder);
-}
-
-void UCameraRigTransitionGraphSchemaBase::FilterGraphContextPlaceableClasses(TArray<UClass*>& InOutClasses) const
-{
-	InOutClasses.Remove(UCameraRigTransition::StaticClass());
-}
-
-FCameraRigTransitionGraphSchemaAction_NewTransitionNode::FCameraRigTransitionGraphSchemaAction_NewTransitionNode()
-{
-	ObjectClass = UCameraRigTransition::StaticClass();
-}
-
-FCameraRigTransitionGraphSchemaAction_NewTransitionNode::FCameraRigTransitionGraphSchemaAction_NewTransitionNode(FText InNodeCategory, FText InMenuDesc, FText InToolTip, const int32 InGrouping, FText InKeywords)
-	: FObjectGraphSchemaAction_NewNode(InNodeCategory, InMenuDesc, InToolTip, InGrouping, InKeywords)
-{
-	ObjectClass = UCameraRigTransition::StaticClass();
-}
-
-void FCameraRigTransitionGraphSchemaAction_NewTransitionNode::AutoSetupNewNode(UObjectTreeGraphNode* NewNode, UEdGraphPin* FromPin)
-{
-	if (TransitionType == ETransitionType::Enter)
-	{
-		NewNode->OverrideSelfPinDirection(EGPD_Output);
-	}
-
-	FObjectGraphSchemaAction_NewNode::AutoSetupNewNode(NewNode, FromPin);
-}
-
 #undef LOCTEXT_NAMESPACE
 
diff --git a/Engine/Plugins/Cameras/GameplayCameras/Source/GameplayCamerasEditor/Private/Editors/CameraSharedTransitionGraphSchema.cpp b/Engine/Plugins/Cameras/GameplayCameras/Source/GameplayCamerasEditor/Private/Editors/CameraSharedTransitionGraphSchema.cpp
--- a/Engine/Plugins/Cameras/GameplayCameras/Source/GameplayCamerasEditor/Private/Editors/CameraSharedTransitionGraphSchema.cpp
+++ b/Engine/Plugins/Cameras/GameplayCameras/Source/GameplayCamerasEditor/Private/Editors/CameraSharedTransitionGraphSchema.cpp
@@ -3,9 +3,7 @@
 #include "Editors/CameraSharedTransitionGraphSchema.h"
 
 #include "Core/CameraAsset.h"
-#include "Core/CameraRigAsset.h"
 #include "Editors/ObjectTreeGraphConfig.h"
-#include "Editors/ObjectTreeGraphNode.h"
 #include "GameplayCamerasEditorSettings.h"
 
 #include UE_INLINE_GENERATED_CPP_BY_NAME(CameraSharedTransitionGraphSchema)
@@ -27,34 +25,5 @@ void UCameraSharedTransitionGraphSchema::OnBuildGraphConfig(FObjectTreeGraphConf
 		.NodeTitleColor(Settings->CameraAssetTitleColor);
 }
 
-ETransitionGraphContextActions UCameraSharedTransitionGraphSchema::GetTransitionGraphContextActions(FGraphContextMenuBuilder& ContextMenuBuilder) const
-{
-	ETransitionGraphContextActions PossibleActions = ETransitionGraphContextActions::None;
-
-	if (const UEdGraphPin* DraggedPin = ContextMenuBuilder.FromPin)
-	{
-		// If we are creating a node from dragging a pin into an empty space, figure out which transition
-		// we can create based on the direction of the dragged pin.
-		UObjectTreeGraphNode* OwningNode = Cast<UObjectTreeGraphNode>(DraggedPin->GetOwningNode());
-		if (OwningNode)
-		{
-			FProperty* DraggedPinProperty = OwningNode->GetPropertyForPin(DraggedPin);
-			if (DraggedPinProperty && DraggedPinProperty->GetOwnerClass()->IsChildOf<UCameraAsset>())
-			{
-				if (DraggedPinProperty->GetFName() == GET_MEMBER_NAME_CHECKED(UCameraAsset, EnterTransitions))
-				{
-					PossibleActions |= ETransitionGraphContextActions::CreateEnterTransition;
-				}
-				if (DraggedPinProperty->GetFName() == GET_MEMBER_NAME_CHECKED(UCameraAsset, ExitTransitions))
-				{
-					PossibleActions |= ETransitionGraphContextActions::CreateExitTransition;
-				}
-			}
-		}
-	}
-
-	return PossibleActions;
-}
-
 #undef LOCTEXT_NAMESPACE
 
diff --git a/Engine/Plugins/Cameras/GameplayCameras/Source/GameplayCamerasEditor/Public/Editors/CameraRigTransitionGraphSchema.h b/Engine/Plugins/Cameras/GameplayCameras/Source/GameplayCamerasEditor/Public/Editors/CameraRigTransitionGraphSchema.h
--- a/Engine/Plugins/Cameras/GameplayCameras/Source/GameplayCamerasEditor/Public/Editors/CameraRigTransitionGraphSchema.h
+++ b/Engine/Plugins/Cameras/GameplayCameras/Source/GameplayCamerasEditor/Public/Editors/CameraRigTransitionGraphSchema.h
@@ -18,6 +18,5 @@ protected:
 
 	// UCameraRigTransitionGraphSchemaBase interface.
 	virtual void OnBuildGraphConfig(FObjectTreeGraphConfig& InOutGraphConfig) const override;
-	virtual ETransitionGraphContextActions GetTransitionGraphContextActions(FGraphContextMenuBuilder& ContextMenuBuilder) const override;
 };
 
diff --git a/Engine/Plugins/Cameras/GameplayCameras/Source/GameplayCamerasEditor/Public/Editors/CameraRigTransitionGraphSchemaBase.h b/Engine/Plugins/Cameras/GameplayCameras/Source/GameplayCamerasEditor/Public/Editors/CameraRigTransitionGraphSchemaBase.h
--- a/Engine/Plugins/Cameras/GameplayCameras/Source/GameplayCamerasEditor/Public/Editors/CameraRigTransitionGraphSchemaBase.h
+++ b/Engine/Plugins/Cameras/GameplayCameras/Source/GameplayCamerasEditor/Public/Editors/CameraRigTransitionGraphSchemaBase.h
@@ -8,17 +8,6 @@
 
 struct FObjectTreeGraphConfig;
 
-/**
- * The list of transition-specific actions that can be performed on a transition graph.
- */
-enum class ETransitionGraphContextActions
-{
-	None = 0,
-	CreateEnterTransition = 1 << 0,
-	CreateExitTransition = 1 << 1
-};
-ENUM_CLASS_FLAGS(ETransitionGraphContextActions);
-
 /**
  * Base schema class for camera transition graph.
  */
@@ -33,55 +22,12 @@ public:
 
 protected:
 
-	// UEdGraphSchema interface.
-	virtual void GetGraphContextActions(FGraphContextMenuBuilder& ContextMenuBuilder) const override;
-
 	// UObjectTreeGraphSchema interface.
 	virtual void CollectAllObjects(UObjectTreeGraph* InGraph, TSet<UObject*>& OutAllObjects) const override;
-	virtual void FilterGraphContextPlaceableClasses(TArray<UClass*>& InOutClasses) const override;
 
 protected:
 
 	// UCameraRigTransitionGraphSchemaBase interface.
 	virtual void OnBuildGraphConfig(FObjectTreeGraphConfig& InOutGraphConfig) const {}
-	virtual ETransitionGraphContextActions GetTransitionGraphContextActions(FGraphContextMenuBuilder& ContextMenuBuilder) const { return ETransitionGraphContextActions::None; }
-};
-
-/**
- * Graph action to create a new transition node.
- *
- * We need a custom action for this because we need to switch the "self" pin according to whether
- * we want an enter or exit transition.
- */
-USTRUCT()
-struct FCameraRigTransitionGraphSchemaAction_NewTransitionNode : public FObjectGraphSchemaAction_NewNode
-{
-	GENERATED_BODY()
-
-public:
-
-	enum class ETransitionType
-	{
-		Enter,
-		Exit
-	};
-
-	/** The transition type to create. */
-	ETransitionType TransitionType;
-
-public:
-
-	FCameraRigTransitionGraphSchemaAction_NewTransitionNode();
-	FCameraRigTransitionGraphSchemaAction_NewTransitionNode(FText InNodeCategory, FText InMenuDesc, FText InToolTip, const int32 InGrouping = 0, FText InKeywords = FText());
-
-public:
-
-	// FEdGraphSchemaAction interface.
-	static FName StaticGetTypeId() { static FName Type("FCameraRigTransitionGraphSchemaAction_NewTransitionNode"); return Type; }
-	virtual FName GetTypeId() const override { return StaticGetTypeId(); } 
-
-protected:
-
-	virtual void AutoSetupNewNode(UObjectTreeGraphNode* NewNode, UEdGraphPin* FromPin) override;
 };
 
diff --git a/Engine/Plugins/Cameras/GameplayCameras/Source/GameplayCamerasEditor/Public/Editors/CameraSharedTransitionGraphSchema.h b/Engine/Plugins/Cameras/GameplayCameras/Source/GameplayCamerasEditor/Public/Editors/CameraSharedTransitionGraphSchema.h
--- a/Engine/Plugins/Cameras/GameplayCameras/Source/GameplayCamerasEditor/Public/Editors/CameraSharedTransitionGraphSchema.h
+++ b/Engine/Plugins/Cameras/GameplayCameras/Source/GameplayCamerasEditor/Public/Editors/CameraSharedTransitionGraphSchema.h
@@ -18,6 +18,5 @@ protected:
 
 	// UCameraRigTransitionGraphSchemaBase interface.
 	virtual void OnBuildGraphConfig(FObjectTreeGraphConfig& InOutGraphConfig) const override;
-	virtual ETransitionGraphContextActions GetTransitionGraphContextActions(FGraphContextMenuBuilder& ContextMenuBuilder) const override;
 };
 
