Cameras: don't crash when removing nodes while the curve editor doesn't exist yet

diff --git a/Engine/Plugins/Cameras/GameplayCameras/Source/GameplayCamerasEditor/Private/Toolkits/CurveEditorToolkit.cpp b/Engine/Plugins/Cameras/GameplayCameras/Source/GameplayCamerasEditor/Private/Toolkits/CurveEditorToolkit.cpp
--- a/Engine/Plugins/Cameras/GameplayCameras/Source/GameplayCamerasEditor/Private/Toolkits/CurveEditorToolkit.cpp
+++ b/Engine/Plugins/Cameras/GameplayCameras/Source/GameplayCamerasEditor/Private/Toolkits/CurveEditorToolkit.cpp
@@ -99,11 +99,19 @@ void FCurveEditorToolkit::Shutdown()
 
 void FCurveEditorToolkit::AddCurveOwner(UObject* InCurveOwner)
 {
-	AddCurves(InCurveOwner);
+	if (IsInitialized())
+	{
+		AddCurves(InCurveOwner);
+	}
 }
 
 void FCurveEditorToolkit::RemoveCurveOwner(UObject* InCurveOwner)
 {
+	if (!IsInitialized())
+	{
+		return;
+	}
+
 	TArray<FCurveEditorTreeItemID> RootTreeItemIDs = CurveEditor->GetRootTreeItems();
 	for (FCurveEditorTreeItemID RootTreeItemID : RootTreeItemIDs)
 	{
