Cameras: cosmetic changes for the object-tree-graphs and camera graphs

- Change the look of object-tree-graph nodes, including their colors and pin icons.
- Groups of array pins are separated from any next pins by a spacer.
- Only the first pin in a group of array pins has a full label, the other only have an index.
- Remove the "self" label for camera graph nodes.

diff --git a/Engine/Plugins/Cameras/GameplayCameras/Source/GameplayCamerasEditor/Private/Editors/CameraNodeGraphNode.cpp b/Engine/Plugins/Cameras/GameplayCameras/Source/GameplayCamerasEditor/Private/Editors/CameraNodeGraphNode.cpp
--- a/Engine/Plugins/Cameras/GameplayCameras/Source/GameplayCamerasEditor/Private/Editors/CameraNodeGraphNode.cpp
+++ b/Engine/Plugins/Cameras/GameplayCameras/Source/GameplayCamerasEditor/Private/Editors/CameraNodeGraphNode.cpp
@@ -5,6 +5,7 @@
 #include "Core/CameraParameters.h"
 #include "EdGraph/EdGraphPin.h"
 #include "Editors/CameraNodeGraphSchema.h"
+#include "Editors/SCameraNodeGraphNode.h"
 #include "ToolMenus.h"
 
 UCameraNodeGraphNode::UCameraNodeGraphNode(const FObjectInitializer& ObjInit)
@@ -43,6 +44,11 @@ UE_CAMERA_VARIABLE_FOR_ALL_TYPES()
 	}
 }
 
+TSharedPtr<SGraphNode> UCameraNodeGraphNode::CreateVisualWidget()
+{
+	return SNew(SCameraNodeGraphNode).GraphNode(this);
+}
+
 UEdGraphPin* UCameraNodeGraphNode::GetPinForCameraParameterProperty(const FName& InPropertyName) const
 {
 	UEdGraphPin* const* FoundItem = Pins.FindByPredicate(
diff --git a/Engine/Plugins/Cameras/GameplayCameras/Source/GameplayCamerasEditor/Private/Editors/CameraNodeGraphSchema.cpp b/Engine/Plugins/Cameras/GameplayCameras/Source/GameplayCamerasEditor/Private/Editors/CameraNodeGraphSchema.cpp
--- a/Engine/Plugins/Cameras/GameplayCameras/Source/GameplayCamerasEditor/Private/Editors/CameraNodeGraphSchema.cpp
+++ b/Engine/Plugins/Cameras/GameplayCameras/Source/GameplayCamerasEditor/Private/Editors/CameraNodeGraphSchema.cpp
@@ -39,6 +39,7 @@ FObjectTreeGraphConfig UCameraNodeGraphSchema::BuildGraphConfig() const
 	GraphConfig.NonConnectableObjectClasses.Add(UBlendCameraNode::StaticClass());
 	GraphConfig.GraphDisplayInfo.PlainName = LOCTEXT("NodeGraphPlainName", "CameraNodes");
 	GraphConfig.GraphDisplayInfo.DisplayName = LOCTEXT("NodeGraphDisplayName", "Camera Nodes");
+	GraphConfig.DefaultSelfPinName = NAME_None;
 	GraphConfig.ObjectClassConfigs.Emplace(UCameraRigAsset::StaticClass())
 		.OnlyAsRoot()
 		.HasSelfPin(false)
@@ -47,11 +48,11 @@ FObjectTreeGraphConfig UCameraNodeGraphSchema::BuildGraphConfig() const
 	GraphConfig.ObjectClassConfigs.Emplace(UCameraNode::StaticClass())
 		.StripDisplayNameSuffix(TEXT("Camera Node"))
 		.CreateCategoryMetaData(TEXT("CameraNodeCategories"))
+		.NodeTitleColor(Settings->CameraNodeTitleColor)
 		.GraphNodeClass(UCameraNodeGraphNode::StaticClass());
 	GraphConfig.ObjectClassConfigs.Emplace(UCameraRigCameraNode::StaticClass())
 		.GraphNodeClass(UCameraRigNodeGraphNode::StaticClass());
 	GraphConfig.ObjectClassConfigs.Emplace(UCameraRigInterfaceParameter::StaticClass())
-		.SelfPinName(NAME_None)  // No self pin name, we just want the title
 		.CanCreateNew(false)
 		.GraphNodeClass(UCameraRigInterfaceParameterGraphNode::StaticClass());
 	GraphConfig.ObjectClassConfigs.Emplace(UArrayCameraNode::StaticClass())
diff --git a/Engine/Plugins/Cameras/GameplayCameras/Source/GameplayCamerasEditor/Private/Editors/CameraRigTransitionGraphSchemaBase.cpp b/Engine/Plugins/Cameras/GameplayCameras/Source/GameplayCamerasEditor/Private/Editors/CameraRigTransitionGraphSchemaBase.cpp
--- a/Engine/Plugins/Cameras/GameplayCameras/Source/GameplayCamerasEditor/Private/Editors/CameraRigTransitionGraphSchemaBase.cpp
+++ b/Engine/Plugins/Cameras/GameplayCameras/Source/GameplayCamerasEditor/Private/Editors/CameraRigTransitionGraphSchemaBase.cpp
@@ -83,6 +83,7 @@ FObjectTreeGraphConfig UCameraRigTransitionGraphSchemaBase::BuildGraphConfig() c
 	GraphConfig.ConnectableObjectClasses.Add(UCameraRigTransition::StaticClass());
 	GraphConfig.ConnectableObjectClasses.Add(UCameraRigTransitionCondition::StaticClass());
 	GraphConfig.ConnectableObjectClasses.Add(UBlendCameraNode::StaticClass());
+	GraphConfig.DefaultSelfPinName = NAME_None;
 	GraphConfig.ObjectClassConfigs.Emplace(UCameraRigTransition::StaticClass())
 		.NodeTitleColor(Settings->CameraRigTransitionTitleColor);
 	GraphConfig.ObjectClassConfigs.Emplace(UCameraRigTransitionCondition::StaticClass())
@@ -90,7 +91,8 @@ FObjectTreeGraphConfig UCameraRigTransitionGraphSchemaBase::BuildGraphConfig() c
 		.NodeTitleColor(Settings->CameraRigTransitionConditionTitleColor);
 	GraphConfig.ObjectClassConfigs.Emplace(UBlendCameraNode::StaticClass())
 		.StripDisplayNameSuffix(TEXT("Camera Node"))
-		.CreateCategoryMetaData(TEXT("CameraNodeCategories"));
+		.CreateCategoryMetaData(TEXT("CameraNodeCategories"))
+		.NodeTitleColor(Settings->CameraBlendNodeTitleColor);
 
 	OnBuildGraphConfig(GraphConfig);
 
diff --git a/Engine/Plugins/Cameras/GameplayCameras/Source/GameplayCamerasEditor/Private/Editors/ObjectTreeGraphConfig.cpp b/Engine/Plugins/Cameras/GameplayCameras/Source/GameplayCamerasEditor/Private/Editors/ObjectTreeGraphConfig.cpp
--- a/Engine/Plugins/Cameras/GameplayCameras/Source/GameplayCamerasEditor/Private/Editors/ObjectTreeGraphConfig.cpp
+++ b/Engine/Plugins/Cameras/GameplayCameras/Source/GameplayCamerasEditor/Private/Editors/ObjectTreeGraphConfig.cpp
@@ -11,8 +11,7 @@
 #define LOCTEXT_NAMESPACE "ObjectTreeGraphConfig"
 
 FObjectTreeGraphClassConfig::FObjectTreeGraphClassConfig()
-	: _SelfPinName(NAME_Self)
-	, _SelfPinFriendlyName(FText::GetEmpty())
+	: _SelfPinName(NAME_None)
 	, _HasSelfPin(true)
 	, _NodeTitleUsesObjectName(false)
 	, _CanCreateNew(true)
@@ -61,7 +60,9 @@ TOptional<EEdGraphPinDirection> FObjectTreeGraphClassConfigs::GetPropertyPinDire
 }
 
 FObjectTreeGraphConfig::FObjectTreeGraphConfig()
-	: DefaultGraphNodeTitleColor(FLinearColor(0.549f, 0.745f, 0.698f))
+	: DefaultSelfPinName(NAME_Self)
+	, DefaultGraphNodeTitleColor(FLinearColor(0.549f, 0.745f, 0.698f))
+	, DefaultGraphNodeTitleTextColor(FLinearColor::White)
 	, DefaultGraphNodeBodyTintColor(FLinearColor::White)
 {
 }
diff --git a/Engine/Plugins/Cameras/GameplayCameras/Source/GameplayCamerasEditor/Private/Editors/ObjectTreeGraphNode.cpp b/Engine/Plugins/Cameras/GameplayCameras/Source/GameplayCamerasEditor/Private/Editors/ObjectTreeGraphNode.cpp
--- a/Engine/Plugins/Cameras/GameplayCameras/Source/GameplayCamerasEditor/Private/Editors/ObjectTreeGraphNode.cpp
+++ b/Engine/Plugins/Cameras/GameplayCameras/Source/GameplayCamerasEditor/Private/Editors/ObjectTreeGraphNode.cpp
@@ -56,6 +56,12 @@ FLinearColor UObjectTreeGraphNode::GetNodeTitleColor() const
 	return NodeContext.ObjectClassConfigs.NodeTitleColor().Get(NodeContext.GraphConfig.DefaultGraphNodeTitleColor);
 }
 
+FLinearColor UObjectTreeGraphNode::GetNodeTitleTextColor() const
+{
+	const FNodeContext NodeContext = GetNodeContext();
+	return NodeContext.ObjectClassConfigs.NodeTitleTextColor().Get(NodeContext.GraphConfig.DefaultGraphNodeTitleTextColor);
+}
+
 FLinearColor UObjectTreeGraphNode::GetNodeBodyTintColor() const
 {
 	const FNodeContext NodeContext = GetNodeContext();
@@ -88,9 +94,9 @@ void UObjectTreeGraphNode::AllocateDefaultPins()
 	{
 		FEdGraphPinType SelfPinType;
 		SelfPinType.PinCategory = UObjectTreeGraphSchema::PC_Self;
-		const FName& SelfPinName = ObjectClassConfigs.SelfPinName();
+		const FName& SelfPinName = ObjectClassConfigs.SelfPinName(OuterGraphConfig.DefaultSelfPinName);
 		UEdGraphPin* SelfPin = CreatePin(OuterGraphConfig.GetSelfPinDirection(NodeContext.ObjectClass), SelfPinType, SelfPinName);
-		SelfPin->PinFriendlyName = ObjectClassConfigs.SelfPinFriendlyName();
+		SelfPin->PinFriendlyName = ObjectClassConfigs.SelfPinFriendlyName(OuterGraphConfig.DefaultSelfPinFriendlyName);
 	}
 
 	for (TFieldIterator<FProperty> PropertyIt(NodeContext.ObjectClass); PropertyIt; ++PropertyIt)
@@ -176,13 +182,14 @@ void UObjectTreeGraphNode::OnPinRemoved(UEdGraphPin* InRemovedPin)
 	}
 }
 
-void UObjectTreeGraphNode::GetArrayProperties(TArray<FArrayProperty*>& OutArrayProperties) const
+void UObjectTreeGraphNode::GetArrayProperties(TArray<FArrayProperty*>& OutArrayProperties, EEdGraphPinDirection Direction) const
 {
 	const FNodeContext NodeContext = GetNodeContext();
 
 	for (UEdGraphPin* Pin : Pins)
 	{
-		if (Pin->PinType.PinCategory == UObjectTreeGraphSchema::PC_Property &&
+		if (Pin->Direction == Direction &&
+				Pin->PinType.PinCategory == UObjectTreeGraphSchema::PC_Property &&
 				Pin->PinType.PinSubCategory == UObjectTreeGraphSchema::PSC_ArrayProperty &&
 				Pin->ParentPin == nullptr)
 		{
@@ -242,8 +249,15 @@ void UObjectTreeGraphNode::CreateNewItemPins(UEdGraphPin* InParentArrayPin, int3
 		FName ChildPinName = PropertyName;
 		ChildPinName.SetNumber(NewIndex);
 		UEdGraphPin* ChildPin = CreatePin(PinDirection, ChildPinType, ChildPinName);
-		ChildPin->PinFriendlyName = FText::Format(
-				LOCTEXT("ArrayPinFriendlyNameFmt", "{0} {1}"), FText::FromName(PropertyName), NewIndex);
+		if (NewIndex == 0)
+		{
+			ChildPin->PinFriendlyName = FText::Format(
+					LOCTEXT("ArrayPinFriendlyNameFmt", "{0} {1}"), FText::FromName(PropertyName), NewIndex);
+		}
+		else
+		{
+			ChildPin->PinFriendlyName = FText::AsNumber(NewIndex);
+		}
 
 		ChildPin->ParentPin = InParentArrayPin;
 		InParentArrayPin->SubPins.Add(ChildPin);
@@ -284,8 +298,15 @@ void UObjectTreeGraphNode::InsertNewItemPin(UEdGraphPin* InParentArrayPin, int32
 	FName ChildPinName = PropertyName;
 	ChildPinName.SetNumber(Index);
 	UEdGraphPin* ChildPin = CreatePin(PinDirection, ChildPinType, ChildPinName);
-	ChildPin->PinFriendlyName = FText::Format(
-			LOCTEXT("ArrayPinFriendlyNameFmt", "{0} {1}"), FText::FromName(PropertyName), Index);
+	if (Index == 0)
+	{
+		ChildPin->PinFriendlyName = FText::Format(
+				LOCTEXT("ArrayPinFriendlyNameFmt", "{0} {1}"), FText::FromName(PropertyName), Index);
+	}
+	else
+	{
+		ChildPin->PinFriendlyName = FText::AsNumber(Index);
+	}
 
 	ChildPin->ParentPin = InParentArrayPin;
 	InParentArrayPin->SubPins.Insert(ChildPin, Index);
@@ -336,7 +357,16 @@ void UObjectTreeGraphNode::RefreshArrayPropertyPinNames()
 			{
 				UEdGraphPin* ChildPin = Pin->SubPins[PinIndex];
 				ChildPin->PinName.SetNumber(PinIndex);
-				ChildPin->PinFriendlyName = FText::Format(LOCTEXT("ArrayPinFriendlyNameFmt", "{0} {1}"), FText::FromName(PropertyName), PinIndex);
+
+				if (PinIndex == 0)
+				{
+					ChildPin->PinFriendlyName = FText::Format(
+							LOCTEXT("ArrayPinFriendlyNameFmt", "{0} {1}"), FText::FromName(PropertyName), PinIndex);
+				}
+				else
+				{
+					ChildPin->PinFriendlyName = FText::AsNumber(PinIndex);
+				}
 			}
 		}
 	}
diff --git a/Engine/Plugins/Cameras/GameplayCameras/Source/GameplayCamerasEditor/Private/Editors/SCameraNodeGraphNode.cpp b/Engine/Plugins/Cameras/GameplayCameras/Source/GameplayCamerasEditor/Private/Editors/SCameraNodeGraphNode.cpp
new file mode 100644
--- /dev/null
+++ b/Engine/Plugins/Cameras/GameplayCameras/Source/GameplayCamerasEditor/Private/Editors/SCameraNodeGraphNode.cpp
@@ -0,0 +1,35 @@
+// Copyright Epic Games, Inc. All Rights Reserved.
+
+#include "Editors/SCameraNodeGraphNode.h"
+
+#include "Editors/CameraNodeGraphSchema.h"
+#include "Styles/GameplayCamerasEditorStyle.h"
+
+void SCameraNodeGraphNode::Construct(const FArguments& InArgs)
+{
+	SObjectTreeGraphNode::Construct(
+			SObjectTreeGraphNode::FArguments()
+			.GraphNode(InArgs._GraphNode));
+}
+
+TSharedPtr<SGraphPin> SCameraNodeGraphNode::CreatePinWidget(UEdGraphPin* InPin) const
+{
+	using namespace UE::Cameras;
+
+	TSharedPtr<SGraphPin> PinWidget = SObjectTreeGraphNode::CreatePinWidget(InPin);
+
+	if (const UCameraNodeGraphSchema* GraphSchema = Cast<const UCameraNodeGraphSchema>(InPin->GetSchema()))
+	{
+		TSharedRef<FGameplayCamerasEditorStyle> GraphStyle = FGameplayCamerasEditorStyle::Get();
+
+		if (InPin->PinType.PinCategory == UCameraNodeGraphSchema::PC_CameraParameter)
+		{
+			const FSlateBrush* ConnectedBrush = GraphStyle->GetBrush("Graph.CameraRigParameterPin.Connected");
+			const FSlateBrush* DisconnectedBrush = GraphStyle->GetBrush("Graph.CameraRigParameterPin.Disconnected");
+			PinWidget->SetCustomPinIcon(ConnectedBrush, DisconnectedBrush);
+		}
+	}
+
+	return PinWidget;
+}
+
diff --git a/Engine/Plugins/Cameras/GameplayCameras/Source/GameplayCamerasEditor/Private/Editors/SCameraRigInterfaceParameterGraphNode.cpp b/Engine/Plugins/Cameras/GameplayCameras/Source/GameplayCamerasEditor/Private/Editors/SCameraRigInterfaceParameterGraphNode.cpp
--- a/Engine/Plugins/Cameras/GameplayCameras/Source/GameplayCamerasEditor/Private/Editors/SCameraRigInterfaceParameterGraphNode.cpp
+++ b/Engine/Plugins/Cameras/GameplayCameras/Source/GameplayCamerasEditor/Private/Editors/SCameraRigInterfaceParameterGraphNode.cpp
@@ -4,17 +4,12 @@
 
 #include "Core/CameraRigAsset.h"
 #include "Editors/ObjectTreeGraphNode.h"
-#include "SCommentBubble.h"
 #include "SGraphNode.h"
 #include "SNodePanel.h"
 #include "Styles/GameplayCamerasEditorStyle.h"
-#include "TutorialMetaData.h"
 #include "Widgets/Images/SImage.h"
-#include "Widgets/Layout/SBorder.h"
-#include "Widgets/Layout/SSpacer.h"
 #include "Widgets/Notifications/SErrorText.h"
 #include "Widgets/SBoxPanel.h"
-#include "Widgets/SNullWidget.h"
 #include "Widgets/SOverlay.h"
 #include "Widgets/Text/STextBlock.h"
 
@@ -124,12 +119,15 @@ const FSlateBrush* SCameraRigInterfaceParameterGraphNode::GetShadowBrush(bool bS
 		CamerasEditorStyle->GetBrush(TEXT("Graph.CameraRigParameterNode.Shadow"));
 }
 
-void SCameraRigInterfaceParameterGraphNode::GetDiffHighlightBrushes(const FSlateBrush*& BackgroundOut, const FSlateBrush*& ForegroundOut) const
+TSharedPtr<SGraphPin> SCameraRigInterfaceParameterGraphNode::CreatePinWidget(UEdGraphPin* InPin) const
 {
 	using namespace UE::Cameras;
+	TSharedPtr<SGraphPin> PinWidget = SObjectTreeGraphNode::CreatePinWidget(InPin);
 	TSharedRef<FGameplayCamerasEditorStyle> CamerasEditorStyle = FGameplayCamerasEditorStyle::Get();
-	BackgroundOut = CamerasEditorStyle->GetBrush(TEXT("Graph.CameraRigParameterNode.DiffHighlight"));
-	ForegroundOut = CamerasEditorStyle->GetBrush(TEXT("Graph.CameraRigParameterNode.DiffHighlightShading"));
+	const FSlateBrush* ConnectedBrush = CamerasEditorStyle->GetBrush("Graph.CameraRigParameterPin.Connected");
+	const FSlateBrush* DisconnectedBrush = CamerasEditorStyle->GetBrush("Graph.CameraRigParameterPin.Disconnected");
+	PinWidget->SetCustomPinIcon(ConnectedBrush, DisconnectedBrush);
+	return PinWidget;
 }
 
 FText SCameraRigInterfaceParameterGraphNode::GetInterfaceParameterName() const
diff --git a/Engine/Plugins/Cameras/GameplayCameras/Source/GameplayCamerasEditor/Private/Editors/SObjectTreeGraphNode.cpp b/Engine/Plugins/Cameras/GameplayCameras/Source/GameplayCamerasEditor/Private/Editors/SObjectTreeGraphNode.cpp
--- a/Engine/Plugins/Cameras/GameplayCameras/Source/GameplayCamerasEditor/Private/Editors/SObjectTreeGraphNode.cpp
+++ b/Engine/Plugins/Cameras/GameplayCameras/Source/GameplayCamerasEditor/Private/Editors/SObjectTreeGraphNode.cpp
@@ -6,9 +6,12 @@
 #include "Editors/ObjectTreeGraphNode.h"
 #include "Editors/ObjectTreeGraphSchema.h"
 #include "GraphEditorSettings.h"
+#include "Styles/ObjectTreeGraphEditorStyle.h"
 #include "Widgets/Images/SImage.h"
 #include "Widgets/Input/SButton.h"
+#include "Widgets/Layout/SSpacer.h"
 #include "Widgets/SBoxPanel.h"
+#include "Widgets/Text/SInlineEditableTextBlock.h"
 
 #define LOCTEXT_NAMESPACE "SObjectTreeGraphNode"
 
@@ -17,6 +20,8 @@ void SObjectTreeGraphNode::Construct(const FArguments& InArgs)
 	GraphNode = InArgs._GraphNode;
 	ObjectGraphNode = InArgs._GraphNode;
 
+	TitleBorderMargin = FMargin(12, 6, 6, 6);
+
 	SetCursor(EMouseCursor::CardinalCross);
 
 	UpdateGraphNode();
@@ -32,10 +37,176 @@ void SObjectTreeGraphNode::MoveTo(const FVector2D& NewPosition, FNodeSet& NodeFi
 	}
 }
 
+const FSlateBrush* SObjectTreeGraphNode::GetShadowBrush(bool bSelected) const
+{
+	TSharedRef<FObjectTreeGraphEditorStyle> GraphStyle = FObjectTreeGraphEditorStyle::Get();
+	return bSelected ? 
+		GraphStyle->GetBrush(TEXT("ObjectTreeGraphNode.ShadowSelected")) :
+		GraphStyle->GetBrush(TEXT("ObjectTreeGraphNode.Shadow"));
+}
+
+void SObjectTreeGraphNode::CreateInputSideAddButton(TSharedPtr<SVerticalBox> InputBox)
+{
+	MakeAllAddArrayPropertyPinButtons(InputBox, EEdGraphPinDirection::EGPD_Input);
+}
+
 void SObjectTreeGraphNode::CreateOutputSideAddButton(TSharedPtr<SVerticalBox> OutputBox)
+{
+	MakeAllAddArrayPropertyPinButtons(OutputBox, EEdGraphPinDirection::EGPD_Output);
+}
+
+TSharedPtr<SGraphPin> SObjectTreeGraphNode::CreatePinWidget(UEdGraphPin* InPin) const
+{
+	TSharedPtr<SGraphPin> PinWidget = SGraphNode::CreatePinWidget(InPin);
+
+	if (const UObjectTreeGraphSchema* GraphSchema = Cast<const UObjectTreeGraphSchema>(InPin->GetSchema()))
+	{
+		TSharedRef<FObjectTreeGraphEditorStyle> GraphStyle = FObjectTreeGraphEditorStyle::Get();
+
+		if (InPin->PinType.PinCategory == UObjectTreeGraphSchema::PC_Self)
+		{
+			const FSlateBrush* ConnectedBrush = GraphStyle->GetBrush("ObjectTreeGraphNode.SelfPin.Connected");
+			const FSlateBrush* DisconnectedBrush = GraphStyle->GetBrush("ObjectTreeGraphNode.SelfPin.Disconnected");
+			PinWidget->SetCustomPinIcon(ConnectedBrush, DisconnectedBrush);
+		}
+		else if (InPin->PinType.PinCategory == UObjectTreeGraphSchema::PC_Property ||
+				InPin->PinType.PinCategory == UObjectTreeGraphSchema::PSC_ArrayProperty)
+		{
+			const FSlateBrush* ConnectedBrush = GraphStyle->GetBrush("ObjectTreeGraphNode.ObjectPin.Connected");
+			const FSlateBrush* DisconnectedBrush = GraphStyle->GetBrush("ObjectTreeGraphNode.ObjectPin.Disconnected");
+			PinWidget->SetCustomPinIcon(ConnectedBrush, DisconnectedBrush);
+		}
+	}
+
+	return PinWidget;
+}
+
+void SObjectTreeGraphNode::AddPin(const TSharedRef<SGraphPin>& PinToAdd)
+{
+	SGraphNode::AddPin(PinToAdd);
+
+	if (const UEdGraphPin* PinObj = PinToAdd->GetPinObj())
+	{
+		// If this is the last of a list of array pins, add a spacer below.
+		const UEdGraphPin* ArrayPinObj = PinObj->ParentPin;
+		if (PinObj->PinType.PinCategory == UObjectTreeGraphSchema::PC_Property &&
+				PinObj->PinType.PinSubCategory == UObjectTreeGraphSchema::PSC_ArrayPropertyItem &&
+				ArrayPinObj && ArrayPinObj->SubPins.Num() > 0 && ArrayPinObj->SubPins.Last() == PinObj)
+		{
+			if (PinToAdd->GetDirection() == EEdGraphPinDirection::EGPD_Input)
+			{
+				LeftNodeBox->AddSlot()
+				.AutoHeight()
+				.HAlign(HAlign_Left)
+				.VAlign(VAlign_Center)
+				.Padding(Settings->GetInputPinPadding())
+				[
+					SNew(SSpacer)
+					.Size(FVector2D(12, 12))
+				];
+			}
+			else // Direction == EEdGraphPinDirection::EGPD_Output
+			{
+				RightNodeBox->AddSlot()
+				.AutoHeight()
+				.HAlign(HAlign_Right)
+				.VAlign(VAlign_Center)
+				.Padding(Settings->GetOutputPinPadding())
+				[
+					SNew(SSpacer)
+					.Size(FVector2D(12, 12))
+				];
+			}
+		}
+	}
+}
+
+void SObjectTreeGraphNode::SetDefaultTitleAreaWidget(TSharedRef<SOverlay> DefaultTitleAreaWidget)
+{
+	SGraphNode::SetDefaultTitleAreaWidget(DefaultTitleAreaWidget);
+
+	// Somewhat of a hack, but clear up the title area and rebuild it the way we want.
+	DefaultTitleAreaWidget->ClearChildren();
+
+	// Node title.
+	TSharedPtr<SNodeTitle> NodeTitle = SNew(SNodeTitle, GraphNode);
+
+	// Node icon.
+	IconColor = FLinearColor::White;
+	const FSlateBrush* IconBrush = nullptr;
+	if (GraphNode != NULL && GraphNode->ShowPaletteIconOnNode())
+	{
+		IconBrush = GraphNode->GetIconAndTint(IconColor).GetOptionalIcon();
+	}
+
+	TSharedRef<FObjectTreeGraphEditorStyle> GraphStyle = FObjectTreeGraphEditorStyle::Get();
+
+	DefaultTitleAreaWidget->AddSlot()
+		.HAlign(HAlign_Fill)
+		.VAlign(VAlign_Center)
+		[
+			SNew(SBorder)
+			.Padding(0)
+			.BorderImage(GraphStyle->GetBrush("ObjectTreeGraphNode.TitleBackground"))
+			.BorderBackgroundColor(this, &SGraphNode::GetNodeTitleColor)
+			[
+				SNew(SHorizontalBox)
+				+SHorizontalBox::Slot()
+				.HAlign(HAlign_Fill)
+				[
+					SNew(SBorder)
+					.BorderImage(FAppStyle::GetBrush("NoBorder"))
+					.Padding(TitleBorderMargin)
+					[
+						SNew(SHorizontalBox)
+						+SHorizontalBox::Slot()
+						.VAlign(VAlign_Top)
+						.AutoWidth()
+						.Padding(FMargin(0, 0, 4, 0))
+						[
+							SNew(SImage)
+							.Image(IconBrush)
+							.ColorAndOpacity(this, &SGraphNode::GetNodeTitleIconColor)
+						]
+						+SHorizontalBox::Slot()
+						[
+							SNew(SVerticalBox)
+							+SVerticalBox::Slot()
+							.AutoHeight()
+							.HAlign(HAlign_Center)
+							[
+								CreateTitleWidget(NodeTitle)
+							]
+							+SVerticalBox::Slot()
+							.AutoHeight()
+							[
+								NodeTitle.ToSharedRef()
+							]
+						]
+					]
+				]
+				+SHorizontalBox::Slot()
+				.HAlign(HAlign_Right)
+				.VAlign(VAlign_Center)
+				.Padding(6, 6, 12, 6)
+				.AutoWidth()
+				[
+					CreateTitleRightWidget()
+				]
+			]
+		];
+}
+
+const FSlateBrush* SObjectTreeGraphNode::GetNodeBodyBrush() const
+{
+	TSharedRef<FObjectTreeGraphEditorStyle> GraphStyle = FObjectTreeGraphEditorStyle::Get();
+	return GraphStyle->GetBrush("ObjectTreeGraphNode.Body");
+}
+
+void SObjectTreeGraphNode::MakeAllAddArrayPropertyPinButtons(TSharedPtr<SVerticalBox> Box, EEdGraphPinDirection Direction)
 {
 	TArray<FArrayProperty*> ArrayProperties;
-	ObjectGraphNode->GetArrayProperties(ArrayProperties);
+	ObjectGraphNode->GetArrayProperties(ArrayProperties, Direction);
 	for (FArrayProperty* ArrayProperty : ArrayProperties)
 	{
 		TSharedRef<SWidget> AddPinButton = MakeAddArrayPropertyPinButton(ArrayProperty);
@@ -43,7 +214,7 @@ void SObjectTreeGraphNode::CreateOutputSideAddButton(TSharedPtr<SVerticalBox> Ou
 		FMargin AddPinPadding = Settings->GetOutputPinPadding();
 		AddPinPadding.Top += 6.0f;
 
-		OutputBox->AddSlot()
+		Box->AddSlot()
 		.AutoHeight()
 		.VAlign(VAlign_Center)
 		.HAlign(HAlign_Right)
diff --git a/Engine/Plugins/Cameras/GameplayCameras/Source/GameplayCamerasEditor/Private/Styles/GameplayCamerasEditorStyle.cpp b/Engine/Plugins/Cameras/GameplayCameras/Source/GameplayCamerasEditor/Private/Styles/GameplayCamerasEditorStyle.cpp
--- a/Engine/Plugins/Cameras/GameplayCameras/Source/GameplayCamerasEditor/Private/Styles/GameplayCamerasEditorStyle.cpp
+++ b/Engine/Plugins/Cameras/GameplayCameras/Source/GameplayCamerasEditor/Private/Styles/GameplayCamerasEditorStyle.cpp
@@ -136,6 +136,9 @@ FGameplayCamerasEditorStyle::FGameplayCamerasEditorStyle()
 	Set("Graph.CameraRigParameterNode.DiffHighlight", new BOX_BRUSH("Graph/CameraRigParameterNode_DiffHighlight", FMargin(18.0f/64.0f)));
 	Set("Graph.CameraRigParameterNode.DiffHighlightShading", new BOX_BRUSH("Graph/CameraRigParameterNode_DiffHighlightShading", FMargin(18.0f/64.0f)));
 
+	Set("Graph.CameraRigParameterPin.Connected", new IMAGE_BRUSH("Graph/ObjectTreeGraphNode_DiamondPin_Connected", Icon16x16));
+	Set("Graph.CameraRigParameterPin.Disconnected", new IMAGE_BRUSH("Graph/ObjectTreeGraphNode_DiamondPin_Disconnected", Icon16x16));
+
 	FSlateStyleRegistry::RegisterSlateStyle(*this);
 }
 
diff --git a/Engine/Plugins/Cameras/GameplayCameras/Source/GameplayCamerasEditor/Private/Styles/ObjectTreeGraphEditorStyle.cpp b/Engine/Plugins/Cameras/GameplayCameras/Source/GameplayCamerasEditor/Private/Styles/ObjectTreeGraphEditorStyle.cpp
--- a/Engine/Plugins/Cameras/GameplayCameras/Source/GameplayCamerasEditor/Private/Styles/ObjectTreeGraphEditorStyle.cpp
+++ b/Engine/Plugins/Cameras/GameplayCameras/Source/GameplayCamerasEditor/Private/Styles/ObjectTreeGraphEditorStyle.cpp
@@ -17,6 +17,8 @@ TSharedPtr<FObjectTreeGraphEditorStyle> FObjectTreeGraphEditorStyle::Singleton;
 FObjectTreeGraphEditorStyle::FObjectTreeGraphEditorStyle()
 	: FSlateStyleSet("ObjectTreeGraphEditorStyle")
 {
+	const FVector2D Icon12x16(12.0f, 16.0f);
+	const FVector2D Icon15x11(15.0f, 11.0f);
 	const FVector2D Icon16x16(16.0f, 16.0f);
 	const FVector2D Icon20x20(20.0f, 20.0f);
 	const FVector2D Icon24x24(24.0f, 24.0f);
@@ -31,6 +33,22 @@ FObjectTreeGraphEditorStyle::FObjectTreeGraphEditorStyle()
 	const FButtonStyle& DefaultButton = FAppStyle::Get().GetWidgetStyle<FButtonStyle>("Button");
 	const FTextBlockStyle& NormalText = FAppStyle::Get().GetWidgetStyle<FTextBlockStyle>("NormalText");
 
+	// Object tree graph node styles.
+	Set("ObjectTreeGraphNode.Body", new BOX_BRUSH("Graph/ObjectTreeGraphNode_Body", FMargin(16.f/64.f, 25.f/64.f, 16.f/64.f, 16.f/64.f)));
+	Set("ObjectTreeGraphNode.Shadow", new BOX_BRUSH("Graph/ObjectTreeGraphNode_Shadow", FMargin(18.0f/64.0f)));
+	Set("ObjectTreeGraphNode.ShadowSelected", new BOX_BRUSH("Graph/ObjectTreeGraphNode_Shadow_Selected", FMargin(18.0f/64.0f)));
+	Set("ObjectTreeGraphNode.TitleHighlight", new BOX_BRUSH("Graph/ObjectTreeGraphNode_TitleHighlight", FMargin(16.f/64.f, 1.f, 16.f/64.f, 0.f)));
+	Set("ObjectTreeGraphNode.TitleBackground", new BOX_BRUSH("Graph/ObjectTreeGraphNode_TitleBackground", FMargin(16.f/64.f, 25.f/64.f, 16.f/64.f, 16.f/64.f)));
+	Set("ObjectTreeGraphNode.DiffHighlight", new BOX_BRUSH("Graph/ObjectTreeGraphNode_DiffHighlight", FMargin(18.0f/64.0f)));
+	Set("ObjectTreeGraphNode.DiffHighlightShading", new BOX_BRUSH("Graph/ObjectTreeGraphNode_DiffHighlightShading", FMargin(18.0f/64.0f)));
+
+	// Object tree graph pin styles.
+	Set("ObjectTreeGraphNode.SelfPin.Connected", new IMAGE_BRUSH("Graph/ObjectTreeGraphNode_SelfPin_Connected", Icon12x16));
+	Set("ObjectTreeGraphNode.SelfPin.Disconnected", new IMAGE_BRUSH("Graph/ObjectTreeGraphNode_SelfPin_Disconnected", Icon12x16));
+
+	Set("ObjectTreeGraphNode.ObjectPin.Connected", new IMAGE_BRUSH("Graph/ObjectTreeGraphNode_ArrowPin_Connected", Icon15x11));
+	Set("ObjectTreeGraphNode.ObjectPin.Disconnected", new IMAGE_BRUSH("Graph/ObjectTreeGraphNode_ArrowPin_Disconnected", Icon15x11));
+
 	// Object tree graph toolbox styles.
 	Set("ObjectTreeGraphToolbox.Entry", FButtonStyle(DefaultButton)
 		.SetNormal(FSlateRoundedBoxBrush(FLinearColor::Transparent, 6.0f, FStyleColors::Dropdown, 1.0f))
diff --git a/Engine/Plugins/Cameras/GameplayCameras/Source/GameplayCamerasEditor/Public/Editors/CameraNodeGraphNode.h b/Engine/Plugins/Cameras/GameplayCameras/Source/GameplayCamerasEditor/Public/Editors/CameraNodeGraphNode.h
--- a/Engine/Plugins/Cameras/GameplayCameras/Source/GameplayCamerasEditor/Public/Editors/CameraNodeGraphNode.h
+++ b/Engine/Plugins/Cameras/GameplayCameras/Source/GameplayCamerasEditor/Public/Editors/CameraNodeGraphNode.h
@@ -29,5 +29,6 @@ public:
 
 	// UEdGraphNode interface
 	virtual void AllocateDefaultPins() override;
+	virtual TSharedPtr<SGraphNode> CreateVisualWidget() override;
 };
 
diff --git a/Engine/Plugins/Cameras/GameplayCameras/Source/GameplayCamerasEditor/Public/Editors/ObjectTreeGraphConfig.h b/Engine/Plugins/Cameras/GameplayCameras/Source/GameplayCamerasEditor/Public/Editors/ObjectTreeGraphConfig.h
--- a/Engine/Plugins/Cameras/GameplayCameras/Source/GameplayCamerasEditor/Public/Editors/ObjectTreeGraphConfig.h
+++ b/Engine/Plugins/Cameras/GameplayCameras/Source/GameplayCamerasEditor/Public/Editors/ObjectTreeGraphConfig.h
@@ -65,8 +65,10 @@ public:
 	/** Default direction of property pins. */
 	OTGCC_FIELD(TOptional<EEdGraphPinDirection>, DefaultPropertyPinDirectionOverride)
 
-	/** Color of the graph node's title. */
+	/** Color of the graph node's title background. */
 	OTGCC_FIELD(TOptional<FLinearColor>, NodeTitleColor)
+	/** Color of the graph node's title text. */
+	OTGCC_FIELD(TOptional<FLinearColor>, NodeTitleTextColor)
 	/** Color of the graph node's body. */
 	OTGCC_FIELD(TOptional<FLinearColor>, NodeBodyTintColor)
 
@@ -96,6 +98,7 @@ private:
 	OTGCC_FIELD_FLAG(HasSelfPin)
 	OTGCC_FIELD_FLAG(DefaultPropertyPinDirectionOverride)
 	OTGCC_FIELD_FLAG(NodeTitleColor)
+	OTGCC_FIELD_FLAG(NodeTitleTextColor)
 	OTGCC_FIELD_FLAG(NodeBodyTintColor)
 	OTGCC_FIELD_FLAG(OnSetupNewObject)
 	OTGCC_FIELD_FLAG(NodeTitleUsesObjectName)
@@ -161,7 +164,7 @@ private:
 
 #define OTGCCS_FIELD(FieldType, FieldName)\
 	public:\
-		typename TCallTraits<FieldType>::ConstReference FieldName() const\
+		typename TCallTraits<FieldType>::ConstReference FieldName(typename TCallTraits<FieldType>::ConstReference DefaultValue) const\
 		{\
 			for (const FObjectTreeGraphClassConfig* InnerConfig : InnerConfigs)\
 			{\
@@ -170,7 +173,11 @@ private:
 					return InnerConfig->FieldName();\
 				}\
 			}\
-			return DefaultConfig.FieldName();\
+			return DefaultValue;\
+		}\
+		typename TCallTraits<FieldType>::ConstReference FieldName() const\
+		{\
+			return FieldName(DefaultConfig.FieldName());\
 		}
 
 /**
@@ -196,8 +203,10 @@ public:
 	/** Default direction of property pins. */
 	OTGCCS_FIELD(TOptional<EEdGraphPinDirection>, DefaultPropertyPinDirectionOverride)
 
-	/** Color of the graph node's title. */
+	/** Color of the graph node's title background. */
 	OTGCCS_FIELD(TOptional<FLinearColor>, NodeTitleColor)
+	/** Color of the graph node's title text. */
+	OTGCCS_FIELD(TOptional<FLinearColor>, NodeTitleTextColor)
 	/** Color of the graph node's body. */
 	OTGCCS_FIELD(TOptional<FLinearColor>, NodeBodyTintColor)
 
@@ -273,8 +282,15 @@ public:
 	 */
 	TSubclassOf<UObjectTreeGraphNode> DefaultGraphNodeClass;
 
-	/** The default title color for an object's graph node. */
+	/** The default name for a node's self pin. */
+	FName DefaultSelfPinName;
+	/** The default friendly name for a node's self pin. */
+	FText DefaultSelfPinFriendlyName;
+
+	/** The default title background color for an object's graph node. */
 	FLinearColor DefaultGraphNodeTitleColor;
+	/** The default title text color for an object's graph node. */
+	FLinearColor DefaultGraphNodeTitleTextColor;
 	/** The default body color for an object's graph node. */
 	FLinearColor DefaultGraphNodeBodyTintColor;
 
diff --git a/Engine/Plugins/Cameras/GameplayCameras/Source/GameplayCamerasEditor/Public/Editors/ObjectTreeGraphNode.h b/Engine/Plugins/Cameras/GameplayCameras/Source/GameplayCamerasEditor/Public/Editors/ObjectTreeGraphNode.h
--- a/Engine/Plugins/Cameras/GameplayCameras/Source/GameplayCamerasEditor/Public/Editors/ObjectTreeGraphNode.h
+++ b/Engine/Plugins/Cameras/GameplayCameras/Source/GameplayCamerasEditor/Public/Editors/ObjectTreeGraphNode.h
@@ -63,6 +63,7 @@ public:
 	virtual FText GetNodeTitle(ENodeTitleType::Type TitleType) const override;
 	virtual TSharedPtr<SGraphNode> CreateVisualWidget() override;
 	virtual FLinearColor GetNodeTitleColor() const override;
+	virtual FLinearColor GetNodeTitleTextColor() const override;
 	virtual FLinearColor GetNodeBodyTintColor() const override;
 	virtual FText GetTooltipText() const override;
 	virtual void AllocateDefaultPins() override;
@@ -86,7 +87,7 @@ public:
 public:
 
 	// Internal API.
-	void GetArrayProperties(TArray<FArrayProperty*>& OutArrayProperties) const;
+	void GetArrayProperties(TArray<FArrayProperty*>& OutArrayProperties, EEdGraphPinDirection Direction) const;
 	void CreateNewItemPins(FArrayProperty& InArrayProperty, int32 NumExtraPins);
 	void CreateNewItemPins(UEdGraphPin* InParentArrayPin, int32 NumExtraPins);
 	void InsertNewItemPin(UEdGraphPin* InParentArrayPin, int32 Index);
diff --git a/Engine/Plugins/Cameras/GameplayCameras/Source/GameplayCamerasEditor/Public/Editors/SCameraNodeGraphNode.h b/Engine/Plugins/Cameras/GameplayCameras/Source/GameplayCamerasEditor/Public/Editors/SCameraNodeGraphNode.h
new file mode 100644
--- /dev/null
+++ b/Engine/Plugins/Cameras/GameplayCameras/Source/GameplayCamerasEditor/Public/Editors/SCameraNodeGraphNode.h
@@ -0,0 +1,24 @@
+// Copyright Epic Games, Inc. All Rights Reserved.
+
+#pragma once
+
+#include "Editors/SObjectTreeGraphNode.h"
+
+class SCameraNodeGraphNode : public SObjectTreeGraphNode
+{
+public:
+
+	SLATE_BEGIN_ARGS(SCameraNodeGraphNode)
+		: _GraphNode(nullptr)
+	{}
+		SLATE_ARGUMENT(UObjectTreeGraphNode*, GraphNode)
+	SLATE_END_ARGS()
+
+	void Construct(const FArguments& InArgs);
+
+public:
+
+	// SGraphNode interface.
+	virtual TSharedPtr<SGraphPin> CreatePinWidget(UEdGraphPin* InPin) const override;
+};
+
diff --git a/Engine/Plugins/Cameras/GameplayCameras/Source/GameplayCamerasEditor/Public/Editors/SCameraRigInterfaceParameterGraphNode.h b/Engine/Plugins/Cameras/GameplayCameras/Source/GameplayCamerasEditor/Public/Editors/SCameraRigInterfaceParameterGraphNode.h
--- a/Engine/Plugins/Cameras/GameplayCameras/Source/GameplayCamerasEditor/Public/Editors/SCameraRigInterfaceParameterGraphNode.h
+++ b/Engine/Plugins/Cameras/GameplayCameras/Source/GameplayCamerasEditor/Public/Editors/SCameraRigInterfaceParameterGraphNode.h
@@ -24,7 +24,7 @@ protected:
 	// SGraphNode interface.
 	virtual void UpdateGraphNode() override;
 	virtual const FSlateBrush* GetShadowBrush(bool bSelected) const override;
-	virtual void GetDiffHighlightBrushes(const FSlateBrush*& BackgroundOut, const FSlateBrush*& ForegroundOut) const override;
+	virtual TSharedPtr<SGraphPin> CreatePinWidget(UEdGraphPin* InPin) const override;
 
 private:
 	
diff --git a/Engine/Plugins/Cameras/GameplayCameras/Source/GameplayCamerasEditor/Public/Editors/SObjectTreeGraphNode.h b/Engine/Plugins/Cameras/GameplayCameras/Source/GameplayCamerasEditor/Public/Editors/SObjectTreeGraphNode.h
--- a/Engine/Plugins/Cameras/GameplayCameras/Source/GameplayCamerasEditor/Public/Editors/SObjectTreeGraphNode.h
+++ b/Engine/Plugins/Cameras/GameplayCameras/Source/GameplayCamerasEditor/Public/Editors/SObjectTreeGraphNode.h
@@ -28,13 +28,20 @@ public:
 public:
 
 	// SGraphNode interface.
+	virtual const FSlateBrush* GetShadowBrush(bool bSelected) const override;
+	virtual void CreateInputSideAddButton(TSharedPtr<SVerticalBox> InputBox) override;
 	virtual void CreateOutputSideAddButton(TSharedPtr<SVerticalBox> OutputBox) override;
+	virtual TSharedPtr<SGraphPin> CreatePinWidget(UEdGraphPin* InPin) const override;
+	virtual void AddPin(const TSharedRef<SGraphPin>& PinToAdd) override;
+	virtual void SetDefaultTitleAreaWidget(TSharedRef<SOverlay> DefaultTitleAreaWidget) override;
+	virtual const FSlateBrush* GetNodeBodyBrush() const override;
 
 	// SNodePanel::SNode interface.
 	virtual void MoveTo(const FVector2D& NewPosition, FNodeSet& NodeFilter, bool bMarkDirty=true) override;
 
 protected:
 
+	void MakeAllAddArrayPropertyPinButtons(TSharedPtr<SVerticalBox> Box, EEdGraphPinDirection Direction);
 	TSharedRef<SWidget> MakeAddArrayPropertyPinButton(FArrayProperty* ArrayProperty);
 
 	FReply OnAddArrayPropertyPin(FArrayProperty* ArrayProperty);
