Cameras: move post-process blending utility functions to the engine module.

diff --git a/Engine/Plugins/Cameras/GameplayCameras/Source/GameplayCameras/Private/Core/PostProcessSettingsCollection.cpp b/Engine/Plugins/Cameras/GameplayCameras/Source/GameplayCameras/Private/Core/PostProcessSettingsCollection.cpp
--- a/Engine/Plugins/Cameras/GameplayCameras/Source/GameplayCameras/Private/Core/PostProcessSettingsCollection.cpp
+++ b/Engine/Plugins/Cameras/GameplayCameras/Source/GameplayCameras/Private/Core/PostProcessSettingsCollection.cpp
@@ -2,6 +2,8 @@
 
 #include "Core/PostProcessSettingsCollection.h"
 
+#include "Engine/PostProcessUtils.h"
+
 namespace UE::Cameras
 {
 
@@ -29,259 +31,8 @@ void FPostProcessSettingsCollection::OverrideChanged(const FPostProcessSettings&
 	FPostProcessSettings& ThisPP = PostProcessSettings;
 	const FPostProcessSettings& OtherPP = OtherPostProcessSettings;
 
-#define UE_SET_PP(Name)\
-	if (OtherPP.bOverride_##Name)\
-	{\
-		ThisPP.bOverride_##Name = true;\
-		ThisPP.Name = OtherPP.Name;\
-		bHasAnySetting = true;\
-	}
-
-	{
-		UE_SET_PP(TemperatureType);
-		UE_SET_PP(WhiteTemp);
-		UE_SET_PP(WhiteTint);
-
-		UE_SET_PP(ColorSaturation);
-		UE_SET_PP(ColorContrast);
-		UE_SET_PP(ColorGamma);
-		UE_SET_PP(ColorGain);
-		UE_SET_PP(ColorOffset);
-
-		UE_SET_PP(ColorSaturationShadows);
-		UE_SET_PP(ColorContrastShadows);
-		UE_SET_PP(ColorGammaShadows);
-		UE_SET_PP(ColorGainShadows);
-		UE_SET_PP(ColorOffsetShadows);
-
-		UE_SET_PP(ColorSaturationMidtones);
-		UE_SET_PP(ColorContrastMidtones);
-		UE_SET_PP(ColorGammaMidtones);
-		UE_SET_PP(ColorGainMidtones);
-		UE_SET_PP(ColorOffsetMidtones);
-
-		UE_SET_PP(ColorSaturationHighlights);
-		UE_SET_PP(ColorContrastHighlights);
-		UE_SET_PP(ColorGammaHighlights);
-		UE_SET_PP(ColorGainHighlights);
-		UE_SET_PP(ColorOffsetHighlights);
-
-		UE_SET_PP(ColorCorrectionShadowsMax);
-		UE_SET_PP(ColorCorrectionHighlightsMin);
-		UE_SET_PP(ColorCorrectionHighlightsMax);
-
-		UE_SET_PP(BlueCorrection);
-		UE_SET_PP(ExpandGamut);
-		UE_SET_PP(ToneCurveAmount);
-
-		UE_SET_PP(FilmSlope);
-		UE_SET_PP(FilmToe);
-		UE_SET_PP(FilmShoulder);
-		UE_SET_PP(FilmBlackClip);
-		UE_SET_PP(FilmWhiteClip);
-
-		UE_SET_PP(SceneColorTint);
-		UE_SET_PP(SceneFringeIntensity);
-		UE_SET_PP(ChromaticAberrationStartOffset);
-		UE_SET_PP(BloomIntensity);
-		UE_SET_PP(BloomThreshold);
-		UE_SET_PP(Bloom1Tint);
-		UE_SET_PP(BloomSizeScale);
-		UE_SET_PP(Bloom1Size);
-		UE_SET_PP(Bloom2Tint);
-		UE_SET_PP(Bloom2Size);
-		UE_SET_PP(Bloom3Tint);
-		UE_SET_PP(Bloom3Size);
-		UE_SET_PP(Bloom4Tint);
-		UE_SET_PP(Bloom4Size);
-		UE_SET_PP(Bloom5Tint);
-		UE_SET_PP(Bloom5Size);
-		UE_SET_PP(Bloom6Tint);
-		UE_SET_PP(Bloom6Size);
-		UE_SET_PP(BloomDirtMaskIntensity);
-		UE_SET_PP(BloomDirtMaskTint);
-		UE_SET_PP(BloomConvolutionScatterDispersion);
-		UE_SET_PP(BloomConvolutionSize);
-		UE_SET_PP(BloomConvolutionCenterUV);
-		UE_SET_PP(BloomConvolutionPreFilterMin);
-		UE_SET_PP(BloomConvolutionPreFilterMax);
-		UE_SET_PP(BloomConvolutionPreFilterMult);
-		UE_SET_PP(AmbientCubemapIntensity);
-		UE_SET_PP(AmbientCubemapTint);
-		UE_SET_PP(CameraShutterSpeed);
-		UE_SET_PP(CameraISO);
-		UE_SET_PP(AutoExposureLowPercent);
-		UE_SET_PP(AutoExposureHighPercent);
-		UE_SET_PP(AutoExposureMinBrightness);
-		UE_SET_PP(AutoExposureMaxBrightness);
-		UE_SET_PP(AutoExposureSpeedUp);
-		UE_SET_PP(AutoExposureSpeedDown);
-		UE_SET_PP(AutoExposureBias);
-		UE_SET_PP(HistogramLogMin);
-		UE_SET_PP(HistogramLogMax);
-		UE_SET_PP(LocalExposureMethod);
-		UE_SET_PP(LocalExposureContrastScale_DEPRECATED);
-		UE_SET_PP(LocalExposureHighlightContrastScale);
-		UE_SET_PP(LocalExposureShadowContrastScale);
-		UE_SET_PP(LocalExposureHighlightThreshold);
-		UE_SET_PP(LocalExposureShadowThreshold);
-		UE_SET_PP(LocalExposureDetailStrength);
-		UE_SET_PP(LocalExposureBlurredLuminanceBlend);
-		UE_SET_PP(LocalExposureBlurredLuminanceKernelSizePercent);
-		UE_SET_PP(LocalExposureHighlightThresholdStrength);
-		UE_SET_PP(LocalExposureShadowThresholdStrength);
-		UE_SET_PP(LocalExposureMiddleGreyBias);
-		UE_SET_PP(LensFlareIntensity);
-		UE_SET_PP(LensFlareTint);
-		UE_SET_PP(LensFlareBokehSize);
-		UE_SET_PP(LensFlareThreshold);
-		UE_SET_PP(VignetteIntensity);
-		UE_SET_PP(Sharpen);
-		UE_SET_PP(FilmGrainIntensity);
-		UE_SET_PP(FilmGrainIntensityShadows);
-		UE_SET_PP(FilmGrainIntensityMidtones);
-		UE_SET_PP(FilmGrainIntensityHighlights);
-		UE_SET_PP(FilmGrainShadowsMax);
-		UE_SET_PP(FilmGrainHighlightsMin);
-		UE_SET_PP(FilmGrainHighlightsMax);
-		UE_SET_PP(FilmGrainTexelSize);
-		UE_SET_PP(AmbientOcclusionIntensity);
-		UE_SET_PP(AmbientOcclusionStaticFraction);
-		UE_SET_PP(AmbientOcclusionRadius);
-		UE_SET_PP(AmbientOcclusionFadeDistance);
-		UE_SET_PP(AmbientOcclusionFadeRadius);
-		UE_SET_PP(AmbientOcclusionDistance_DEPRECATED);
-		UE_SET_PP(AmbientOcclusionPower);
-		UE_SET_PP(AmbientOcclusionBias);
-		UE_SET_PP(AmbientOcclusionQuality);
-		UE_SET_PP(AmbientOcclusionMipBlend);
-		UE_SET_PP(AmbientOcclusionMipScale);
-		UE_SET_PP(AmbientOcclusionMipThreshold);
-		UE_SET_PP(AmbientOcclusionTemporalBlendWeight);
-		UE_SET_PP(IndirectLightingColor);
-		UE_SET_PP(IndirectLightingIntensity);
-
-		UE_SET_PP(DepthOfFieldFocalDistance);
-
-		UE_SET_PP(DepthOfFieldFstop);
-		UE_SET_PP(DepthOfFieldMinFstop);
-		UE_SET_PP(DepthOfFieldSensorWidth);
-		UE_SET_PP(DepthOfFieldSqueezeFactor);
-		UE_SET_PP(DepthOfFieldDepthBlurRadius);
-		UE_SET_PP(DepthOfFieldUseHairDepth)
-		UE_SET_PP(DepthOfFieldDepthBlurAmount);
-		UE_SET_PP(DepthOfFieldFocalRegion);
-		UE_SET_PP(DepthOfFieldNearTransitionRegion);
-		UE_SET_PP(DepthOfFieldFarTransitionRegion);
-		UE_SET_PP(DepthOfFieldScale);
-		UE_SET_PP(DepthOfFieldNearBlurSize);
-		UE_SET_PP(DepthOfFieldFarBlurSize);
-		UE_SET_PP(DepthOfFieldOcclusion);
-		UE_SET_PP(DepthOfFieldSkyFocusDistance);
-		UE_SET_PP(DepthOfFieldVignetteSize);
-		UE_SET_PP(MotionBlurAmount);
-		UE_SET_PP(MotionBlurMax);
-		UE_SET_PP(MotionBlurPerObjectSize);
-		UE_SET_PP(ScreenSpaceReflectionQuality);
-		UE_SET_PP(ScreenSpaceReflectionIntensity);
-		UE_SET_PP(ScreenSpaceReflectionMaxRoughness);
-
-		UE_SET_PP(TranslucencyType);
-		UE_SET_PP(RayTracingTranslucencyMaxRoughness);
-		UE_SET_PP(RayTracingTranslucencyRefractionRays);
-		UE_SET_PP(RayTracingTranslucencySamplesPerPixel);
-		UE_SET_PP(RayTracingTranslucencyShadows);
-		UE_SET_PP(RayTracingTranslucencyRefraction);
-
-		UE_SET_PP(DynamicGlobalIlluminationMethod);
-		UE_SET_PP(LumenSurfaceCacheResolution);
-		UE_SET_PP(LumenSceneLightingQuality);
-		UE_SET_PP(LumenSceneDetail);
-		UE_SET_PP(LumenSceneViewDistance);
-		UE_SET_PP(LumenSceneLightingUpdateSpeed);
-		UE_SET_PP(LumenFinalGatherQuality);
-		UE_SET_PP(LumenFinalGatherLightingUpdateSpeed);
-		UE_SET_PP(LumenFinalGatherScreenTraces);
-		UE_SET_PP(LumenMaxTraceDistance);
-
-		UE_SET_PP(LumenDiffuseColorBoost);
-		UE_SET_PP(LumenSkylightLeaking);
-		UE_SET_PP(LumenSkylightLeakingTint);
-		UE_SET_PP(LumenFullSkylightLeakingDistance);
-
-		UE_SET_PP(LumenRayLightingMode);
-		UE_SET_PP(LumenReflectionsScreenTraces);
-		UE_SET_PP(LumenFrontLayerTranslucencyReflections);
-		UE_SET_PP(LumenMaxRoughnessToTraceReflections);
-		UE_SET_PP(LumenMaxReflectionBounces);
-		UE_SET_PP(LumenMaxRefractionBounces);
-		UE_SET_PP(ReflectionMethod);
-		UE_SET_PP(LumenReflectionQuality);
-		UE_SET_PP(RayTracingAO);
-		UE_SET_PP(RayTracingAOSamplesPerPixel);
-		UE_SET_PP(RayTracingAOIntensity);
-		UE_SET_PP(RayTracingAORadius);
-
-		UE_SET_PP(PathTracingMaxBounces);
-		UE_SET_PP(PathTracingSamplesPerPixel);
-		UE_SET_PP(PathTracingMaxPathIntensity);
-		UE_SET_PP(PathTracingEnableEmissiveMaterials);
-		UE_SET_PP(PathTracingEnableReferenceDOF);
-		UE_SET_PP(PathTracingEnableReferenceAtmosphere);
-		UE_SET_PP(PathTracingEnableDenoiser);
-		UE_SET_PP(PathTracingIncludeEmissive);
-		UE_SET_PP(PathTracingIncludeDiffuse);
-		UE_SET_PP(PathTracingIncludeIndirectDiffuse);
-		UE_SET_PP(PathTracingIncludeSpecular);
-		UE_SET_PP(PathTracingIncludeIndirectSpecular);
-		UE_SET_PP(PathTracingIncludeVolume);
-		UE_SET_PP(PathTracingIncludeIndirectVolume);
-
-		UE_SET_PP(DepthOfFieldBladeCount);
-
-		// This is no bOverride_AmbientCubemap so just see if it is set.
-		if (OtherPP.AmbientCubemap)
-		{
-			ThisPP.AmbientCubemap = OtherPP.AmbientCubemap;
-		}
-
-		UE_SET_PP(ColorGradingIntensity);
-		UE_SET_PP(ColorGradingLUT);
-
-		UE_SET_PP(BloomDirtMask);
-		UE_SET_PP(BloomMethod);
-		UE_SET_PP(BloomConvolutionTexture)
-		UE_SET_PP(FilmGrainTexture)
-
-		UE_SET_PP(BloomConvolutionBufferScale);
-
-		UE_SET_PP(AutoExposureBiasCurve);
-		UE_SET_PP(AutoExposureMeterMask);
-		UE_SET_PP(LocalExposureHighlightContrastCurve);
-		UE_SET_PP(LocalExposureShadowContrastCurve);
-		UE_SET_PP(LensFlareBokehShape);
-
-		if (OtherPP.bOverride_LensFlareTints)
-		{
-			for (uint32 i = 0; i < 8; ++i)
-			{
-				ThisPP.LensFlareTints[i] = OtherPP.LensFlareTints[i];
-			}
-		}
-
-		if (OtherPP.bOverride_MobileHQGaussian)
-		{
-			ThisPP.bMobileHQGaussian = OtherPP.bMobileHQGaussian;
-		}
-
-		UE_SET_PP(AutoExposureMethod);
-		UE_SET_PP(AmbientOcclusionRadiusInWS);
-		UE_SET_PP(MotionBlurTargetFPS);
-		UE_SET_PP(AutoExposureApplyPhysicalCameraExposure);
-		UE_SET_PP(UserFlags);
-	}
-
-#undef UE_SET_PP
+	const bool bAnyOverwritten = FPostProcessUtils::OverridePostProcessSettings(ThisPP, OtherPP);
+	bHasAnySetting |= bAnyOverwritten;
 }
 
 void FPostProcessSettingsCollection::LerpAll(const FPostProcessSettingsCollection& ToCollection, float BlendFactor)
@@ -291,316 +42,16 @@ void FPostProcessSettingsCollection::LerpAll(const FPostProcessSettingsCollectio
 
 void FPostProcessSettingsCollection::LerpAll(const FPostProcessSettings& ToPostProcessSettings, float BlendFactor)
 {
-	InternalLerpChanged(ToPostProcessSettings, BlendFactor, false);
-}
-
-void FPostProcessSettingsCollection::LerpChanged(const FPostProcessSettingsCollection& ToCollection, float BlendFactor)
-{
-	LerpChanged(ToCollection.PostProcessSettings, BlendFactor);
+	InternalLerp(ToPostProcessSettings, BlendFactor);
 }
 
-void FPostProcessSettingsCollection::LerpChanged(const FPostProcessSettings& ToPostProcessSettings, float BlendFactor)
+void FPostProcessSettingsCollection::InternalLerp(const FPostProcessSettings& ToPostProcessSettings, float BlendFactor)
 {
-	InternalLerpChanged(ToPostProcessSettings, BlendFactor, true);
-}
-
-void FPostProcessSettingsCollection::InternalLerpChanged(const FPostProcessSettings& ToPostProcessSettings, float BlendFactor, bool bChangedOnly)
-{
-	if (BlendFactor <= 0.f)
-	{
-		return;
-	}
-
-	BlendFactor = FMath::Clamp(BlendFactor, 0.f, 1.f);
-	const bool bShouldFlip = (BlendFactor >= 0.5f);
-
 	FPostProcessSettings& ThisPP = PostProcessSettings;
 	const FPostProcessSettings& ToPP = ToPostProcessSettings;
 
-	// We need to an equivalent of FSceneView::OverridePostProcessSettings... differences include:
-	//
-	// 1) Flipping non-interpolatable properties at 50% blend, instead of always overwriting them.
-	// 2) A few things not being supported, such as accumulating ambient cubmaps.
-	// 3) No support for blendable objects.
-	// 4) Ability to blend _away_ from the current values, towards possibly default (non-overriden) values.
-	//
-	// TODO: refactor post-process settings to make blending them possible by default.
-
-#define UE_SET_PP(Name)\
-	if ((ToPP.bOverride_##Name || (!bChangedOnly && ThisPP.bOverride_##Name)) && bShouldFlip)\
-	{\
-		ThisPP.bOverride_##Name = true;\
-		ThisPP.Name = ToPP.Name;\
-		bHasAnySetting = true;\
-	}
-
-#define UE_LERP_PP(Name)\
-	if (ToPP.bOverride_##Name || (!bChangedOnly && ThisPP.bOverride_##Name))\
-	{\
-		ThisPP.bOverride_##Name = true;\
-		ThisPP.Name = FMath::Lerp(ThisPP.Name, ToPP.Name, BlendFactor);\
-		bHasAnySetting = true;\
-	}
-
-	{
-		UE_SET_PP(TemperatureType);
-		UE_LERP_PP(WhiteTemp);
-		UE_LERP_PP(WhiteTint);
-
-		UE_LERP_PP(ColorSaturation);
-		UE_LERP_PP(ColorContrast);
-		UE_LERP_PP(ColorGamma);
-		UE_LERP_PP(ColorGain);
-		UE_LERP_PP(ColorOffset);
-
-		UE_LERP_PP(ColorSaturationShadows);
-		UE_LERP_PP(ColorContrastShadows);
-		UE_LERP_PP(ColorGammaShadows);
-		UE_LERP_PP(ColorGainShadows);
-		UE_LERP_PP(ColorOffsetShadows);
-
-		UE_LERP_PP(ColorSaturationMidtones);
-		UE_LERP_PP(ColorContrastMidtones);
-		UE_LERP_PP(ColorGammaMidtones);
-		UE_LERP_PP(ColorGainMidtones);
-		UE_LERP_PP(ColorOffsetMidtones);
-
-		UE_LERP_PP(ColorSaturationHighlights);
-		UE_LERP_PP(ColorContrastHighlights);
-		UE_LERP_PP(ColorGammaHighlights);
-		UE_LERP_PP(ColorGainHighlights);
-		UE_LERP_PP(ColorOffsetHighlights);
-
-		UE_LERP_PP(ColorCorrectionShadowsMax);
-		UE_LERP_PP(ColorCorrectionHighlightsMin);
-		UE_LERP_PP(ColorCorrectionHighlightsMax);
-
-		UE_LERP_PP(BlueCorrection);
-		UE_LERP_PP(ExpandGamut);
-		UE_LERP_PP(ToneCurveAmount);
-
-		UE_LERP_PP(FilmSlope);
-		UE_LERP_PP(FilmToe);
-		UE_LERP_PP(FilmShoulder);
-		UE_LERP_PP(FilmBlackClip);
-		UE_LERP_PP(FilmWhiteClip);
-
-		UE_LERP_PP(SceneColorTint);
-		UE_LERP_PP(SceneFringeIntensity);
-		UE_LERP_PP(ChromaticAberrationStartOffset);
-		UE_LERP_PP(BloomIntensity);
-		UE_LERP_PP(BloomThreshold);
-		UE_LERP_PP(Bloom1Tint);
-		UE_LERP_PP(BloomSizeScale);
-		UE_LERP_PP(Bloom1Size);
-		UE_LERP_PP(Bloom2Tint);
-		UE_LERP_PP(Bloom2Size);
-		UE_LERP_PP(Bloom3Tint);
-		UE_LERP_PP(Bloom3Size);
-		UE_LERP_PP(Bloom4Tint);
-		UE_LERP_PP(Bloom4Size);
-		UE_LERP_PP(Bloom5Tint);
-		UE_LERP_PP(Bloom5Size);
-		UE_LERP_PP(Bloom6Tint);
-		UE_LERP_PP(Bloom6Size);
-		UE_LERP_PP(BloomDirtMaskIntensity);
-		UE_LERP_PP(BloomDirtMaskTint);
-		UE_LERP_PP(BloomConvolutionScatterDispersion);
-		UE_LERP_PP(BloomConvolutionSize);
-		UE_LERP_PP(BloomConvolutionCenterUV);
-		UE_LERP_PP(BloomConvolutionPreFilterMin);
-		UE_LERP_PP(BloomConvolutionPreFilterMax);
-		UE_LERP_PP(BloomConvolutionPreFilterMult);
-		UE_LERP_PP(AmbientCubemapIntensity);
-		UE_LERP_PP(AmbientCubemapTint);
-		UE_LERP_PP(CameraShutterSpeed);
-		UE_LERP_PP(CameraISO);
-		UE_LERP_PP(AutoExposureLowPercent);
-		UE_LERP_PP(AutoExposureHighPercent);
-		UE_LERP_PP(AutoExposureMinBrightness);
-		UE_LERP_PP(AutoExposureMaxBrightness);
-		UE_LERP_PP(AutoExposureSpeedUp);
-		UE_LERP_PP(AutoExposureSpeedDown);
-		UE_LERP_PP(AutoExposureBias);
-		UE_LERP_PP(HistogramLogMin);
-		UE_LERP_PP(HistogramLogMax);
-		UE_SET_PP(LocalExposureMethod);
-		UE_LERP_PP(LocalExposureContrastScale_DEPRECATED);
-		UE_LERP_PP(LocalExposureHighlightContrastScale);
-		UE_LERP_PP(LocalExposureShadowContrastScale);
-		UE_LERP_PP(LocalExposureHighlightThreshold);
-		UE_LERP_PP(LocalExposureShadowThreshold);
-		UE_LERP_PP(LocalExposureDetailStrength);
-		UE_LERP_PP(LocalExposureBlurredLuminanceBlend);
-		UE_LERP_PP(LocalExposureBlurredLuminanceKernelSizePercent);
-		UE_LERP_PP(LocalExposureHighlightThresholdStrength);
-		UE_LERP_PP(LocalExposureShadowThresholdStrength);
-		UE_LERP_PP(LocalExposureMiddleGreyBias);
-		UE_LERP_PP(LensFlareIntensity);
-		UE_LERP_PP(LensFlareTint);
-		UE_LERP_PP(LensFlareBokehSize);
-		UE_LERP_PP(LensFlareThreshold);
-		UE_LERP_PP(VignetteIntensity);
-		UE_LERP_PP(Sharpen);
-		UE_LERP_PP(FilmGrainIntensity);
-		UE_LERP_PP(FilmGrainIntensityShadows);
-		UE_LERP_PP(FilmGrainIntensityMidtones);
-		UE_LERP_PP(FilmGrainIntensityHighlights);
-		UE_LERP_PP(FilmGrainShadowsMax);
-		UE_LERP_PP(FilmGrainHighlightsMin);
-		UE_LERP_PP(FilmGrainHighlightsMax);
-		UE_LERP_PP(FilmGrainTexelSize);
-		UE_LERP_PP(AmbientOcclusionIntensity);
-		UE_LERP_PP(AmbientOcclusionStaticFraction);
-		UE_LERP_PP(AmbientOcclusionRadius);
-		UE_LERP_PP(AmbientOcclusionFadeDistance);
-		UE_LERP_PP(AmbientOcclusionFadeRadius);
-		UE_LERP_PP(AmbientOcclusionDistance_DEPRECATED);
-		UE_LERP_PP(AmbientOcclusionPower);
-		UE_LERP_PP(AmbientOcclusionBias);
-		UE_LERP_PP(AmbientOcclusionQuality);
-		UE_LERP_PP(AmbientOcclusionMipBlend);
-		UE_LERP_PP(AmbientOcclusionMipScale);
-		UE_LERP_PP(AmbientOcclusionMipThreshold);
-		UE_LERP_PP(AmbientOcclusionTemporalBlendWeight);
-		UE_LERP_PP(IndirectLightingColor);
-		UE_LERP_PP(IndirectLightingIntensity);
-
-		if (ToPP.bOverride_DepthOfFieldFocalDistance)
-		{
-			if (ThisPP.DepthOfFieldFocalDistance == 0.0f || ToPP.DepthOfFieldFocalDistance == 0.0f)
-			{
-				ThisPP.DepthOfFieldFocalDistance = ToPP.DepthOfFieldFocalDistance;
-			}
-			else
-			{
-				ThisPP.DepthOfFieldFocalDistance = FMath::Lerp(ThisPP.DepthOfFieldFocalDistance, ToPP.DepthOfFieldFocalDistance, BlendFactor);
-			}
-		}
-		UE_LERP_PP(DepthOfFieldFstop);
-		UE_LERP_PP(DepthOfFieldMinFstop);
-		UE_LERP_PP(DepthOfFieldSensorWidth);
-		UE_LERP_PP(DepthOfFieldSqueezeFactor);
-		UE_LERP_PP(DepthOfFieldDepthBlurRadius);
-		UE_SET_PP(DepthOfFieldUseHairDepth)
-		UE_LERP_PP(DepthOfFieldDepthBlurAmount);
-		UE_LERP_PP(DepthOfFieldFocalRegion);
-		UE_LERP_PP(DepthOfFieldNearTransitionRegion);
-		UE_LERP_PP(DepthOfFieldFarTransitionRegion);
-		UE_LERP_PP(DepthOfFieldScale);
-		UE_LERP_PP(DepthOfFieldNearBlurSize);
-		UE_LERP_PP(DepthOfFieldFarBlurSize);
-		UE_LERP_PP(DepthOfFieldOcclusion);
-		UE_LERP_PP(DepthOfFieldSkyFocusDistance);
-		UE_LERP_PP(DepthOfFieldVignetteSize);
-		UE_LERP_PP(MotionBlurAmount);
-		UE_LERP_PP(MotionBlurMax);
-		UE_LERP_PP(MotionBlurPerObjectSize);
-		UE_LERP_PP(ScreenSpaceReflectionQuality);
-		UE_LERP_PP(ScreenSpaceReflectionIntensity);
-		UE_LERP_PP(ScreenSpaceReflectionMaxRoughness);
-
-		UE_SET_PP(TranslucencyType);
-		UE_SET_PP(RayTracingTranslucencyMaxRoughness);
-		UE_SET_PP(RayTracingTranslucencyRefractionRays);
-		UE_SET_PP(RayTracingTranslucencySamplesPerPixel);
-		UE_SET_PP(RayTracingTranslucencyShadows);
-		UE_SET_PP(RayTracingTranslucencyRefraction);
-
-		UE_SET_PP(DynamicGlobalIlluminationMethod);
-		UE_SET_PP(LumenSurfaceCacheResolution);
-		UE_SET_PP(LumenSceneLightingQuality);
-		UE_SET_PP(LumenSceneDetail);
-		UE_SET_PP(LumenSceneViewDistance);
-		UE_SET_PP(LumenSceneLightingUpdateSpeed);
-		UE_SET_PP(LumenFinalGatherQuality);
-		UE_SET_PP(LumenFinalGatherLightingUpdateSpeed);
-		UE_SET_PP(LumenFinalGatherScreenTraces);
-		UE_SET_PP(LumenMaxTraceDistance);
-
-		UE_LERP_PP(LumenDiffuseColorBoost);
-		UE_LERP_PP(LumenSkylightLeaking);
-		UE_LERP_PP(LumenSkylightLeakingTint);
-		UE_LERP_PP(LumenFullSkylightLeakingDistance);
-
-		UE_SET_PP(LumenRayLightingMode);
-		UE_SET_PP(LumenReflectionsScreenTraces);
-		UE_SET_PP(LumenFrontLayerTranslucencyReflections);
-		UE_SET_PP(LumenMaxRoughnessToTraceReflections);
-		UE_SET_PP(LumenMaxReflectionBounces);
-		UE_SET_PP(LumenMaxRefractionBounces);
-		UE_SET_PP(ReflectionMethod);
-		UE_SET_PP(LumenReflectionQuality);
-		UE_SET_PP(RayTracingAO);
-		UE_SET_PP(RayTracingAOSamplesPerPixel);
-		UE_SET_PP(RayTracingAOIntensity);
-		UE_SET_PP(RayTracingAORadius);
-
-		UE_SET_PP(PathTracingMaxBounces);
-		UE_SET_PP(PathTracingSamplesPerPixel);
-		UE_LERP_PP(PathTracingMaxPathIntensity);
-		UE_SET_PP(PathTracingEnableEmissiveMaterials);
-		UE_SET_PP(PathTracingEnableReferenceDOF);
-		UE_SET_PP(PathTracingEnableReferenceAtmosphere);
-		UE_SET_PP(PathTracingEnableDenoiser);
-		UE_SET_PP(PathTracingIncludeEmissive);
-		UE_SET_PP(PathTracingIncludeDiffuse);
-		UE_SET_PP(PathTracingIncludeIndirectDiffuse);
-		UE_SET_PP(PathTracingIncludeSpecular);
-		UE_SET_PP(PathTracingIncludeIndirectSpecular);
-		UE_SET_PP(PathTracingIncludeVolume);
-		UE_SET_PP(PathTracingIncludeIndirectVolume);
-
-		UE_SET_PP(DepthOfFieldBladeCount);
-
-		// No cubemap blending (only supported for FFinalPostProcessSettings)
-		if (bShouldFlip)
-		{
-			ThisPP.AmbientCubemap = ToPP.AmbientCubemap;
-		}
-
-		// No color grading texture blending (only supported for FFinalPostProcessSettings)
-		UE_LERP_PP(ColorGradingIntensity);
-		UE_SET_PP(ColorGradingLUT);
-
-		UE_SET_PP(BloomDirtMask);
-		UE_SET_PP(BloomMethod);
-		UE_SET_PP(BloomConvolutionTexture)
-		UE_SET_PP(FilmGrainTexture)
-
-		// Flipping this instead of blending, as per the comment in SceneView.cpp
-		UE_SET_PP(BloomConvolutionBufferScale);
-
-		UE_SET_PP(AutoExposureBiasCurve);
-		UE_SET_PP(AutoExposureMeterMask);
-		UE_SET_PP(LocalExposureHighlightContrastCurve);
-		UE_SET_PP(LocalExposureShadowContrastCurve);
-		UE_SET_PP(LensFlareBokehShape);
-
-		if (ToPP.bOverride_LensFlareTints)
-		{
-			for (uint32 i = 0; i < 8; ++i)
-			{
-				ThisPP.LensFlareTints[i] = FMath::Lerp(ThisPP.LensFlareTints[i], ToPP.LensFlareTints[i], BlendFactor);
-			}
-		}
-
-		if (ToPP.bOverride_MobileHQGaussian && bShouldFlip)
-		{
-			ThisPP.bMobileHQGaussian = ToPP.bMobileHQGaussian;
-		}
-
-		UE_SET_PP(AutoExposureMethod);
-		UE_SET_PP(AmbientOcclusionRadiusInWS);
-		UE_SET_PP(MotionBlurTargetFPS);
-		UE_SET_PP(AutoExposureApplyPhysicalCameraExposure);
-		UE_SET_PP(UserFlags);
-	}
-
-	// No support for blendable objects for now.
-
-#undef UE_SET_PP
-#undef UE_LERP_PP
+	const bool bAnyOverwritten = FPostProcessUtils::BlendPostProcessSettings(ThisPP, ToPP, BlendFactor);
+	bHasAnySetting |= bAnyOverwritten;
 }
 
 void FPostProcessSettingsCollection::Serialize(FArchive& Ar)
diff --git a/Engine/Plugins/Cameras/GameplayCameras/Source/GameplayCameras/Public/Core/PostProcessSettingsCollection.h b/Engine/Plugins/Cameras/GameplayCameras/Source/GameplayCameras/Public/Core/PostProcessSettingsCollection.h
--- a/Engine/Plugins/Cameras/GameplayCameras/Source/GameplayCameras/Public/Core/PostProcessSettingsCollection.h
+++ b/Engine/Plugins/Cameras/GameplayCameras/Source/GameplayCameras/Public/Core/PostProcessSettingsCollection.h
@@ -46,21 +46,12 @@ struct GAMEPLAYCAMERAS_API FPostProcessSettingsCollection
 	void LerpAll(const FPostProcessSettingsCollection& ToCollection, float BlendFactor);
 	void LerpAll(const FPostProcessSettings& ToPostProcessSettings, float BlendFactor);
 
-	/**
-	 * Interpolates the post-process settings towards the values in the given other collection. Only
-	 * values from the other collection that have their bOverride_Xxx set to true will be interpolated.
-	 * If the current value isn't overriden, it will interpolate from the default value. All affected
-	 * values will get their bOverride_Xxx flag set.
-	 */
-	void LerpChanged(const FPostProcessSettingsCollection& ToCollection, float BlendFactor);
-	void LerpChanged(const FPostProcessSettings& ToPostProcessSettings, float BlendFactor);
-
 	/** Serializes this collection into the given archive. */
 	void Serialize(FArchive& Ar);
 
 private:
 
-	void InternalLerpChanged(const FPostProcessSettings& ToPostProcessSettings, float BlendFactor, bool bChangedOnly);
+	void InternalLerp(const FPostProcessSettings& ToPostProcessSettings, float BlendFactor);
 
 private:
 
diff --git a/Engine/Source/Runtime/Engine/Classes/Engine/PostProcessUtils.h b/Engine/Source/Runtime/Engine/Classes/Engine/PostProcessUtils.h
new file mode 100644
--- /dev/null
+++ b/Engine/Source/Runtime/Engine/Classes/Engine/PostProcessUtils.h
@@ -0,0 +1,44 @@
+// Copyright Epic Games, Inc. All Rights Reserved.
+
+#pragma once
+
+#include "CoreTypes.h"
+
+struct FPostProcessSettings;
+
+/**
+ * A utility class for blending FPostProcessSettings together without creating a FFinalPostProcessSettings
+ * instance. This is useful for building up post-process settings in a modular way before handing it
+ * off to the player camera manager, or other engine class that takes an FPostProcessSettings.
+ */
+struct ENGINE_API FPostProcessUtils
+{
+	/**
+	 * Overwrites post-process settings with another set of post-process settings. Only settings that are 
+	 * overriden in the other set (OtherTo) are overwritten. This effectively "overlays" one set of settings
+	 * on top of another.
+	 *
+	 * @param ThisFrom  The settings to write over.
+	 * @param OtherTo   The settings to overlay on top.
+	 * @return          Whether any setting was overwritten.
+	 */
+	static bool OverridePostProcessSettings(FPostProcessSettings& ThisFrom, const FPostProcessSettings& OtherTo);
+
+	/**
+	 * Blends values from one set of post-process settings to another, storing the result in place of 
+	 * that first set. Settings that are overriden in *either* set are blended, so that settings may blend
+	 * between different values, from a default value to a custom value, or from a custom value back to
+	 * a default value.
+	 *
+	 * Some non-interpolable properties, like enums, get "flipped" over 50% blend. A couple others, like 
+	 * ambient cubemaps, don't get accumulated like with FFinalPostProcessSettings and are instead also
+	 * "flipped" over 50% blend. Blendable objects are not supported at this point.
+	 *
+	 * @param ThisFrom  The settings to blend from. Will receive the blended values.
+	 * @param OtherTo   The settings to blend towards.
+	 * @param BlendFactor  The blend factor.
+	 * @return          Whether any setting was blended.
+	 */
+	static bool BlendPostProcessSettings(FPostProcessSettings& ThisFrom, const FPostProcessSettings& OtherTo, float BlendFactor);
+};
+
diff --git a/Engine/Source/Runtime/Engine/Private/PostProcessUtils.cpp b/Engine/Source/Runtime/Engine/Private/PostProcessUtils.cpp
new file mode 100644
--- /dev/null
+++ b/Engine/Source/Runtime/Engine/Private/PostProcessUtils.cpp
@@ -0,0 +1,565 @@
+// Copyright Epic Games, Inc. All Rights Reserved.
+
+#include "Engine/PostProcessUtils.h"
+
+#include "Engine/Scene.h"
+
+bool FPostProcessUtils::OverridePostProcessSettings(FPostProcessSettings& ThisFrom, const FPostProcessSettings& OtherTo)
+{
+	bool bAnyOverwritten = false;
+
+#define UE_SET_PP(Name)\
+	if (OtherTo.bOverride_##Name)\
+	{\
+		ThisFrom.bOverride_##Name = true;\
+		ThisFrom.Name = OtherTo.Name;\
+		bAnyOverwritten = true;\
+	}
+
+	{
+		UE_SET_PP(TemperatureType);
+		UE_SET_PP(WhiteTemp);
+		UE_SET_PP(WhiteTint);
+
+		UE_SET_PP(ColorSaturation);
+		UE_SET_PP(ColorContrast);
+		UE_SET_PP(ColorGamma);
+		UE_SET_PP(ColorGain);
+		UE_SET_PP(ColorOffset);
+
+		UE_SET_PP(ColorSaturationShadows);
+		UE_SET_PP(ColorContrastShadows);
+		UE_SET_PP(ColorGammaShadows);
+		UE_SET_PP(ColorGainShadows);
+		UE_SET_PP(ColorOffsetShadows);
+
+		UE_SET_PP(ColorSaturationMidtones);
+		UE_SET_PP(ColorContrastMidtones);
+		UE_SET_PP(ColorGammaMidtones);
+		UE_SET_PP(ColorGainMidtones);
+		UE_SET_PP(ColorOffsetMidtones);
+
+		UE_SET_PP(ColorSaturationHighlights);
+		UE_SET_PP(ColorContrastHighlights);
+		UE_SET_PP(ColorGammaHighlights);
+		UE_SET_PP(ColorGainHighlights);
+		UE_SET_PP(ColorOffsetHighlights);
+
+		UE_SET_PP(ColorCorrectionShadowsMax);
+		UE_SET_PP(ColorCorrectionHighlightsMin);
+		UE_SET_PP(ColorCorrectionHighlightsMax);
+
+		UE_SET_PP(BlueCorrection);
+		UE_SET_PP(ExpandGamut);
+		UE_SET_PP(ToneCurveAmount);
+
+		UE_SET_PP(FilmSlope);
+		UE_SET_PP(FilmToe);
+		UE_SET_PP(FilmShoulder);
+		UE_SET_PP(FilmBlackClip);
+		UE_SET_PP(FilmWhiteClip);
+
+		UE_SET_PP(SceneColorTint);
+		UE_SET_PP(SceneFringeIntensity);
+		UE_SET_PP(ChromaticAberrationStartOffset);
+		UE_SET_PP(BloomIntensity);
+		UE_SET_PP(BloomThreshold);
+		UE_SET_PP(Bloom1Tint);
+		UE_SET_PP(BloomSizeScale);
+		UE_SET_PP(Bloom1Size);
+		UE_SET_PP(Bloom2Tint);
+		UE_SET_PP(Bloom2Size);
+		UE_SET_PP(Bloom3Tint);
+		UE_SET_PP(Bloom3Size);
+		UE_SET_PP(Bloom4Tint);
+		UE_SET_PP(Bloom4Size);
+		UE_SET_PP(Bloom5Tint);
+		UE_SET_PP(Bloom5Size);
+		UE_SET_PP(Bloom6Tint);
+		UE_SET_PP(Bloom6Size);
+		UE_SET_PP(BloomDirtMaskIntensity);
+		UE_SET_PP(BloomDirtMaskTint);
+		UE_SET_PP(BloomConvolutionScatterDispersion);
+		UE_SET_PP(BloomConvolutionSize);
+		UE_SET_PP(BloomConvolutionCenterUV);
+		UE_SET_PP(BloomConvolutionPreFilterMin);
+		UE_SET_PP(BloomConvolutionPreFilterMax);
+		UE_SET_PP(BloomConvolutionPreFilterMult);
+		UE_SET_PP(AmbientCubemapIntensity);
+		UE_SET_PP(AmbientCubemapTint);
+		UE_SET_PP(CameraShutterSpeed);
+		UE_SET_PP(CameraISO);
+		UE_SET_PP(AutoExposureLowPercent);
+		UE_SET_PP(AutoExposureHighPercent);
+		UE_SET_PP(AutoExposureMinBrightness);
+		UE_SET_PP(AutoExposureMaxBrightness);
+		UE_SET_PP(AutoExposureSpeedUp);
+		UE_SET_PP(AutoExposureSpeedDown);
+		UE_SET_PP(AutoExposureBias);
+		UE_SET_PP(HistogramLogMin);
+		UE_SET_PP(HistogramLogMax);
+		UE_SET_PP(LocalExposureMethod);
+		UE_SET_PP(LocalExposureContrastScale_DEPRECATED);
+		UE_SET_PP(LocalExposureHighlightContrastScale);
+		UE_SET_PP(LocalExposureShadowContrastScale);
+		UE_SET_PP(LocalExposureHighlightThreshold);
+		UE_SET_PP(LocalExposureShadowThreshold);
+		UE_SET_PP(LocalExposureDetailStrength);
+		UE_SET_PP(LocalExposureBlurredLuminanceBlend);
+		UE_SET_PP(LocalExposureBlurredLuminanceKernelSizePercent);
+		UE_SET_PP(LocalExposureHighlightThresholdStrength);
+		UE_SET_PP(LocalExposureShadowThresholdStrength);
+		UE_SET_PP(LocalExposureMiddleGreyBias);
+		UE_SET_PP(LensFlareIntensity);
+		UE_SET_PP(LensFlareTint);
+		UE_SET_PP(LensFlareBokehSize);
+		UE_SET_PP(LensFlareThreshold);
+		UE_SET_PP(VignetteIntensity);
+		UE_SET_PP(Sharpen);
+		UE_SET_PP(FilmGrainIntensity);
+		UE_SET_PP(FilmGrainIntensityShadows);
+		UE_SET_PP(FilmGrainIntensityMidtones);
+		UE_SET_PP(FilmGrainIntensityHighlights);
+		UE_SET_PP(FilmGrainShadowsMax);
+		UE_SET_PP(FilmGrainHighlightsMin);
+		UE_SET_PP(FilmGrainHighlightsMax);
+		UE_SET_PP(FilmGrainTexelSize);
+		UE_SET_PP(AmbientOcclusionIntensity);
+		UE_SET_PP(AmbientOcclusionStaticFraction);
+		UE_SET_PP(AmbientOcclusionRadius);
+		UE_SET_PP(AmbientOcclusionFadeDistance);
+		UE_SET_PP(AmbientOcclusionFadeRadius);
+		UE_SET_PP(AmbientOcclusionDistance_DEPRECATED);
+		UE_SET_PP(AmbientOcclusionPower);
+		UE_SET_PP(AmbientOcclusionBias);
+		UE_SET_PP(AmbientOcclusionQuality);
+		UE_SET_PP(AmbientOcclusionMipBlend);
+		UE_SET_PP(AmbientOcclusionMipScale);
+		UE_SET_PP(AmbientOcclusionMipThreshold);
+		UE_SET_PP(AmbientOcclusionTemporalBlendWeight);
+		UE_SET_PP(IndirectLightingColor);
+		UE_SET_PP(IndirectLightingIntensity);
+
+		UE_SET_PP(DepthOfFieldFocalDistance);
+
+		UE_SET_PP(DepthOfFieldFstop);
+		UE_SET_PP(DepthOfFieldMinFstop);
+		UE_SET_PP(DepthOfFieldSensorWidth);
+		UE_SET_PP(DepthOfFieldSqueezeFactor);
+		UE_SET_PP(DepthOfFieldDepthBlurRadius);
+		UE_SET_PP(DepthOfFieldUseHairDepth)
+		UE_SET_PP(DepthOfFieldDepthBlurAmount);
+		UE_SET_PP(DepthOfFieldFocalRegion);
+		UE_SET_PP(DepthOfFieldNearTransitionRegion);
+		UE_SET_PP(DepthOfFieldFarTransitionRegion);
+		UE_SET_PP(DepthOfFieldScale);
+		UE_SET_PP(DepthOfFieldNearBlurSize);
+		UE_SET_PP(DepthOfFieldFarBlurSize);
+		UE_SET_PP(DepthOfFieldOcclusion);
+		UE_SET_PP(DepthOfFieldSkyFocusDistance);
+		UE_SET_PP(DepthOfFieldVignetteSize);
+		UE_SET_PP(MotionBlurAmount);
+		UE_SET_PP(MotionBlurMax);
+		UE_SET_PP(MotionBlurPerObjectSize);
+		UE_SET_PP(ScreenSpaceReflectionQuality);
+		UE_SET_PP(ScreenSpaceReflectionIntensity);
+		UE_SET_PP(ScreenSpaceReflectionMaxRoughness);
+
+		UE_SET_PP(TranslucencyType);
+		UE_SET_PP(RayTracingTranslucencyMaxRoughness);
+		UE_SET_PP(RayTracingTranslucencyRefractionRays);
+		UE_SET_PP(RayTracingTranslucencySamplesPerPixel);
+		UE_SET_PP(RayTracingTranslucencyShadows);
+		UE_SET_PP(RayTracingTranslucencyRefraction);
+
+		UE_SET_PP(DynamicGlobalIlluminationMethod);
+		UE_SET_PP(LumenSurfaceCacheResolution);
+		UE_SET_PP(LumenSceneLightingQuality);
+		UE_SET_PP(LumenSceneDetail);
+		UE_SET_PP(LumenSceneViewDistance);
+		UE_SET_PP(LumenSceneLightingUpdateSpeed);
+		UE_SET_PP(LumenFinalGatherQuality);
+		UE_SET_PP(LumenFinalGatherLightingUpdateSpeed);
+		UE_SET_PP(LumenFinalGatherScreenTraces);
+		UE_SET_PP(LumenMaxTraceDistance);
+
+		UE_SET_PP(LumenDiffuseColorBoost);
+		UE_SET_PP(LumenSkylightLeaking);
+		UE_SET_PP(LumenSkylightLeakingTint);
+		UE_SET_PP(LumenFullSkylightLeakingDistance);
+
+		UE_SET_PP(LumenRayLightingMode);
+		UE_SET_PP(LumenReflectionsScreenTraces);
+		UE_SET_PP(LumenFrontLayerTranslucencyReflections);
+		UE_SET_PP(LumenMaxRoughnessToTraceReflections);
+		UE_SET_PP(LumenMaxReflectionBounces);
+		UE_SET_PP(LumenMaxRefractionBounces);
+		UE_SET_PP(ReflectionMethod);
+		UE_SET_PP(LumenReflectionQuality);
+		UE_SET_PP(RayTracingAO);
+		UE_SET_PP(RayTracingAOSamplesPerPixel);
+		UE_SET_PP(RayTracingAOIntensity);
+		UE_SET_PP(RayTracingAORadius);
+
+		UE_SET_PP(PathTracingMaxBounces);
+		UE_SET_PP(PathTracingSamplesPerPixel);
+		UE_SET_PP(PathTracingMaxPathIntensity);
+		UE_SET_PP(PathTracingEnableEmissiveMaterials);
+		UE_SET_PP(PathTracingEnableReferenceDOF);
+		UE_SET_PP(PathTracingEnableReferenceAtmosphere);
+		UE_SET_PP(PathTracingEnableDenoiser);
+		UE_SET_PP(PathTracingIncludeEmissive);
+		UE_SET_PP(PathTracingIncludeDiffuse);
+		UE_SET_PP(PathTracingIncludeIndirectDiffuse);
+		UE_SET_PP(PathTracingIncludeSpecular);
+		UE_SET_PP(PathTracingIncludeIndirectSpecular);
+		UE_SET_PP(PathTracingIncludeVolume);
+		UE_SET_PP(PathTracingIncludeIndirectVolume);
+
+		UE_SET_PP(DepthOfFieldBladeCount);
+
+		// This is no bOverride_AmbientCubemap so just see if it is set.
+		if (OtherTo.AmbientCubemap)
+		{
+			ThisFrom.AmbientCubemap = OtherTo.AmbientCubemap;
+			bAnyOverwritten = true;
+		}
+
+		UE_SET_PP(ColorGradingIntensity);
+		UE_SET_PP(ColorGradingLUT);
+
+		UE_SET_PP(BloomDirtMask);
+		UE_SET_PP(BloomMethod);
+		UE_SET_PP(BloomConvolutionTexture)
+		UE_SET_PP(FilmGrainTexture)
+
+		UE_SET_PP(BloomConvolutionBufferScale);
+
+		UE_SET_PP(AutoExposureBiasCurve);
+		UE_SET_PP(AutoExposureMeterMask);
+		UE_SET_PP(LocalExposureHighlightContrastCurve);
+		UE_SET_PP(LocalExposureShadowContrastCurve);
+		UE_SET_PP(LensFlareBokehShape);
+
+		if (OtherTo.bOverride_LensFlareTints)
+		{
+			for (uint32 i = 0; i < 8; ++i)
+			{
+				ThisFrom.LensFlareTints[i] = OtherTo.LensFlareTints[i];
+			}
+			bAnyOverwritten = true;
+		}
+
+		if (OtherTo.bOverride_MobileHQGaussian)
+		{
+			ThisFrom.bMobileHQGaussian = OtherTo.bMobileHQGaussian;
+			bAnyOverwritten = true;
+		}
+
+		UE_SET_PP(AutoExposureMethod);
+		UE_SET_PP(AmbientOcclusionRadiusInWS);
+		UE_SET_PP(MotionBlurTargetFPS);
+		UE_SET_PP(AutoExposureApplyPhysicalCameraExposure);
+		UE_SET_PP(UserFlags);
+	}
+
+#undef UE_SET_PP
+
+	return bAnyOverwritten;
+}
+
+bool FPostProcessUtils::BlendPostProcessSettings(FPostProcessSettings& ThisFrom, const FPostProcessSettings& OtherTo, float BlendFactor)
+{
+	if (UNLIKELY(BlendFactor <= 0.f))
+	{
+		return false;
+	}
+
+	BlendFactor = FMath::Clamp(BlendFactor, 0.f, 1.f);
+	const bool bShouldFlip = (BlendFactor >= 0.5f);
+
+	bool bAnyOverwritten = false;
+
+#define UE_SET_PP(Name)\
+	if ((OtherTo.bOverride_##Name || ThisFrom.bOverride_##Name) && bShouldFlip)\
+	{\
+		ThisFrom.bOverride_##Name = true;\
+		ThisFrom.Name = OtherTo.Name;\
+		bAnyOverwritten = true;\
+	}
+
+#define UE_LERP_PP(Name)\
+	if (OtherTo.bOverride_##Name || ThisFrom.bOverride_##Name)\
+	{\
+		ThisFrom.bOverride_##Name = true;\
+		ThisFrom.Name = FMath::Lerp(ThisFrom.Name, OtherTo.Name, BlendFactor);\
+		bAnyOverwritten = true;\
+	}
+
+	{
+		UE_SET_PP(TemperatureType);
+		UE_LERP_PP(WhiteTemp);
+		UE_LERP_PP(WhiteTint);
+
+		UE_LERP_PP(ColorSaturation);
+		UE_LERP_PP(ColorContrast);
+		UE_LERP_PP(ColorGamma);
+		UE_LERP_PP(ColorGain);
+		UE_LERP_PP(ColorOffset);
+
+		UE_LERP_PP(ColorSaturationShadows);
+		UE_LERP_PP(ColorContrastShadows);
+		UE_LERP_PP(ColorGammaShadows);
+		UE_LERP_PP(ColorGainShadows);
+		UE_LERP_PP(ColorOffsetShadows);
+
+		UE_LERP_PP(ColorSaturationMidtones);
+		UE_LERP_PP(ColorContrastMidtones);
+		UE_LERP_PP(ColorGammaMidtones);
+		UE_LERP_PP(ColorGainMidtones);
+		UE_LERP_PP(ColorOffsetMidtones);
+
+		UE_LERP_PP(ColorSaturationHighlights);
+		UE_LERP_PP(ColorContrastHighlights);
+		UE_LERP_PP(ColorGammaHighlights);
+		UE_LERP_PP(ColorGainHighlights);
+		UE_LERP_PP(ColorOffsetHighlights);
+
+		UE_LERP_PP(ColorCorrectionShadowsMax);
+		UE_LERP_PP(ColorCorrectionHighlightsMin);
+		UE_LERP_PP(ColorCorrectionHighlightsMax);
+
+		UE_LERP_PP(BlueCorrection);
+		UE_LERP_PP(ExpandGamut);
+		UE_LERP_PP(ToneCurveAmount);
+
+		UE_LERP_PP(FilmSlope);
+		UE_LERP_PP(FilmToe);
+		UE_LERP_PP(FilmShoulder);
+		UE_LERP_PP(FilmBlackClip);
+		UE_LERP_PP(FilmWhiteClip);
+
+		UE_LERP_PP(SceneColorTint);
+		UE_LERP_PP(SceneFringeIntensity);
+		UE_LERP_PP(ChromaticAberrationStartOffset);
+		UE_LERP_PP(BloomIntensity);
+		UE_LERP_PP(BloomThreshold);
+		UE_LERP_PP(Bloom1Tint);
+		UE_LERP_PP(BloomSizeScale);
+		UE_LERP_PP(Bloom1Size);
+		UE_LERP_PP(Bloom2Tint);
+		UE_LERP_PP(Bloom2Size);
+		UE_LERP_PP(Bloom3Tint);
+		UE_LERP_PP(Bloom3Size);
+		UE_LERP_PP(Bloom4Tint);
+		UE_LERP_PP(Bloom4Size);
+		UE_LERP_PP(Bloom5Tint);
+		UE_LERP_PP(Bloom5Size);
+		UE_LERP_PP(Bloom6Tint);
+		UE_LERP_PP(Bloom6Size);
+		UE_LERP_PP(BloomDirtMaskIntensity);
+		UE_LERP_PP(BloomDirtMaskTint);
+		UE_LERP_PP(BloomConvolutionScatterDispersion);
+		UE_LERP_PP(BloomConvolutionSize);
+		UE_LERP_PP(BloomConvolutionCenterUV);
+		UE_LERP_PP(BloomConvolutionPreFilterMin);
+		UE_LERP_PP(BloomConvolutionPreFilterMax);
+		UE_LERP_PP(BloomConvolutionPreFilterMult);
+		UE_LERP_PP(AmbientCubemapIntensity);
+		UE_LERP_PP(AmbientCubemapTint);
+		UE_LERP_PP(CameraShutterSpeed);
+		UE_LERP_PP(CameraISO);
+		UE_LERP_PP(AutoExposureLowPercent);
+		UE_LERP_PP(AutoExposureHighPercent);
+		UE_LERP_PP(AutoExposureMinBrightness);
+		UE_LERP_PP(AutoExposureMaxBrightness);
+		UE_LERP_PP(AutoExposureSpeedUp);
+		UE_LERP_PP(AutoExposureSpeedDown);
+		UE_LERP_PP(AutoExposureBias);
+		UE_LERP_PP(HistogramLogMin);
+		UE_LERP_PP(HistogramLogMax);
+		UE_SET_PP(LocalExposureMethod);
+		UE_LERP_PP(LocalExposureContrastScale_DEPRECATED);
+		UE_LERP_PP(LocalExposureHighlightContrastScale);
+		UE_LERP_PP(LocalExposureShadowContrastScale);
+		UE_LERP_PP(LocalExposureHighlightThreshold);
+		UE_LERP_PP(LocalExposureShadowThreshold);
+		UE_LERP_PP(LocalExposureDetailStrength);
+		UE_LERP_PP(LocalExposureBlurredLuminanceBlend);
+		UE_LERP_PP(LocalExposureBlurredLuminanceKernelSizePercent);
+		UE_LERP_PP(LocalExposureHighlightThresholdStrength);
+		UE_LERP_PP(LocalExposureShadowThresholdStrength);
+		UE_LERP_PP(LocalExposureMiddleGreyBias);
+		UE_LERP_PP(LensFlareIntensity);
+		UE_LERP_PP(LensFlareTint);
+		UE_LERP_PP(LensFlareBokehSize);
+		UE_LERP_PP(LensFlareThreshold);
+		UE_LERP_PP(VignetteIntensity);
+		UE_LERP_PP(Sharpen);
+		UE_LERP_PP(FilmGrainIntensity);
+		UE_LERP_PP(FilmGrainIntensityShadows);
+		UE_LERP_PP(FilmGrainIntensityMidtones);
+		UE_LERP_PP(FilmGrainIntensityHighlights);
+		UE_LERP_PP(FilmGrainShadowsMax);
+		UE_LERP_PP(FilmGrainHighlightsMin);
+		UE_LERP_PP(FilmGrainHighlightsMax);
+		UE_LERP_PP(FilmGrainTexelSize);
+		UE_LERP_PP(AmbientOcclusionIntensity);
+		UE_LERP_PP(AmbientOcclusionStaticFraction);
+		UE_LERP_PP(AmbientOcclusionRadius);
+		UE_LERP_PP(AmbientOcclusionFadeDistance);
+		UE_LERP_PP(AmbientOcclusionFadeRadius);
+		UE_LERP_PP(AmbientOcclusionDistance_DEPRECATED);
+		UE_LERP_PP(AmbientOcclusionPower);
+		UE_LERP_PP(AmbientOcclusionBias);
+		UE_LERP_PP(AmbientOcclusionQuality);
+		UE_LERP_PP(AmbientOcclusionMipBlend);
+		UE_LERP_PP(AmbientOcclusionMipScale);
+		UE_LERP_PP(AmbientOcclusionMipThreshold);
+		UE_LERP_PP(AmbientOcclusionTemporalBlendWeight);
+		UE_LERP_PP(IndirectLightingColor);
+		UE_LERP_PP(IndirectLightingIntensity);
+
+		if (OtherTo.bOverride_DepthOfFieldFocalDistance)
+		{
+			if (ThisFrom.DepthOfFieldFocalDistance == 0.0f || OtherTo.DepthOfFieldFocalDistance == 0.0f)
+			{
+				ThisFrom.DepthOfFieldFocalDistance = OtherTo.DepthOfFieldFocalDistance;
+			}
+			else
+			{
+				ThisFrom.DepthOfFieldFocalDistance = FMath::Lerp(ThisFrom.DepthOfFieldFocalDistance, OtherTo.DepthOfFieldFocalDistance, BlendFactor);
+			}
+			bAnyOverwritten = true;
+		}
+		UE_LERP_PP(DepthOfFieldFstop);
+		UE_LERP_PP(DepthOfFieldMinFstop);
+		UE_LERP_PP(DepthOfFieldSensorWidth);
+		UE_LERP_PP(DepthOfFieldSqueezeFactor);
+		UE_LERP_PP(DepthOfFieldDepthBlurRadius);
+		UE_SET_PP(DepthOfFieldUseHairDepth)
+		UE_LERP_PP(DepthOfFieldDepthBlurAmount);
+		UE_LERP_PP(DepthOfFieldFocalRegion);
+		UE_LERP_PP(DepthOfFieldNearTransitionRegion);
+		UE_LERP_PP(DepthOfFieldFarTransitionRegion);
+		UE_LERP_PP(DepthOfFieldScale);
+		UE_LERP_PP(DepthOfFieldNearBlurSize);
+		UE_LERP_PP(DepthOfFieldFarBlurSize);
+		UE_LERP_PP(DepthOfFieldOcclusion);
+		UE_LERP_PP(DepthOfFieldSkyFocusDistance);
+		UE_LERP_PP(DepthOfFieldVignetteSize);
+		UE_LERP_PP(MotionBlurAmount);
+		UE_LERP_PP(MotionBlurMax);
+		UE_LERP_PP(MotionBlurPerObjectSize);
+		UE_LERP_PP(ScreenSpaceReflectionQuality);
+		UE_LERP_PP(ScreenSpaceReflectionIntensity);
+		UE_LERP_PP(ScreenSpaceReflectionMaxRoughness);
+
+		UE_SET_PP(TranslucencyType);
+		UE_SET_PP(RayTracingTranslucencyMaxRoughness);
+		UE_SET_PP(RayTracingTranslucencyRefractionRays);
+		UE_SET_PP(RayTracingTranslucencySamplesPerPixel);
+		UE_SET_PP(RayTracingTranslucencyShadows);
+		UE_SET_PP(RayTracingTranslucencyRefraction);
+
+		UE_SET_PP(DynamicGlobalIlluminationMethod);
+		UE_SET_PP(LumenSurfaceCacheResolution);
+		UE_SET_PP(LumenSceneLightingQuality);
+		UE_SET_PP(LumenSceneDetail);
+		UE_SET_PP(LumenSceneViewDistance);
+		UE_SET_PP(LumenSceneLightingUpdateSpeed);
+		UE_SET_PP(LumenFinalGatherQuality);
+		UE_SET_PP(LumenFinalGatherLightingUpdateSpeed);
+		UE_SET_PP(LumenFinalGatherScreenTraces);
+		UE_SET_PP(LumenMaxTraceDistance);
+
+		UE_LERP_PP(LumenDiffuseColorBoost);
+		UE_LERP_PP(LumenSkylightLeaking);
+		UE_LERP_PP(LumenSkylightLeakingTint);
+		UE_LERP_PP(LumenFullSkylightLeakingDistance);
+
+		UE_SET_PP(LumenRayLightingMode);
+		UE_SET_PP(LumenReflectionsScreenTraces);
+		UE_SET_PP(LumenFrontLayerTranslucencyReflections);
+		UE_SET_PP(LumenMaxRoughnessToTraceReflections);
+		UE_SET_PP(LumenMaxReflectionBounces);
+		UE_SET_PP(LumenMaxRefractionBounces);
+		UE_SET_PP(ReflectionMethod);
+		UE_SET_PP(LumenReflectionQuality);
+		UE_SET_PP(RayTracingAO);
+		UE_SET_PP(RayTracingAOSamplesPerPixel);
+		UE_SET_PP(RayTracingAOIntensity);
+		UE_SET_PP(RayTracingAORadius);
+
+		UE_SET_PP(PathTracingMaxBounces);
+		UE_SET_PP(PathTracingSamplesPerPixel);
+		UE_LERP_PP(PathTracingMaxPathIntensity);
+		UE_SET_PP(PathTracingEnableEmissiveMaterials);
+		UE_SET_PP(PathTracingEnableReferenceDOF);
+		UE_SET_PP(PathTracingEnableReferenceAtmosphere);
+		UE_SET_PP(PathTracingEnableDenoiser);
+		UE_SET_PP(PathTracingIncludeEmissive);
+		UE_SET_PP(PathTracingIncludeDiffuse);
+		UE_SET_PP(PathTracingIncludeIndirectDiffuse);
+		UE_SET_PP(PathTracingIncludeSpecular);
+		UE_SET_PP(PathTracingIncludeIndirectSpecular);
+		UE_SET_PP(PathTracingIncludeVolume);
+		UE_SET_PP(PathTracingIncludeIndirectVolume);
+
+		UE_SET_PP(DepthOfFieldBladeCount);
+
+		// No cubemap blending (only supported for FFinalPostProcessSettings)
+		if (bShouldFlip)
+		{
+			ThisFrom.AmbientCubemap = OtherTo.AmbientCubemap;
+			bAnyOverwritten = true;
+		}
+
+		// No color grading texture blending (only supported for FFinalPostProcessSettings)
+		UE_LERP_PP(ColorGradingIntensity);
+		UE_SET_PP(ColorGradingLUT);
+
+		UE_SET_PP(BloomDirtMask);
+		UE_SET_PP(BloomMethod);
+		UE_SET_PP(BloomConvolutionTexture)
+		UE_SET_PP(FilmGrainTexture)
+
+		// Flipping this instead of blending, as per the comment in SceneView.cpp
+		UE_SET_PP(BloomConvolutionBufferScale);
+
+		UE_SET_PP(AutoExposureBiasCurve);
+		UE_SET_PP(AutoExposureMeterMask);
+		UE_SET_PP(LocalExposureHighlightContrastCurve);
+		UE_SET_PP(LocalExposureShadowContrastCurve);
+		UE_SET_PP(LensFlareBokehShape);
+
+		if (OtherTo.bOverride_LensFlareTints)
+		{
+			for (uint32 i = 0; i < 8; ++i)
+			{
+				ThisFrom.LensFlareTints[i] = FMath::Lerp(ThisFrom.LensFlareTints[i], OtherTo.LensFlareTints[i], BlendFactor);
+			}
+			bAnyOverwritten = true;
+		}
+
+		if (OtherTo.bOverride_MobileHQGaussian && bShouldFlip)
+		{
+			ThisFrom.bMobileHQGaussian = OtherTo.bMobileHQGaussian;
+			bAnyOverwritten = true;
+		}
+
+		UE_SET_PP(AutoExposureMethod);
+		UE_SET_PP(AmbientOcclusionRadiusInWS);
+		UE_SET_PP(MotionBlurTargetFPS);
+		UE_SET_PP(AutoExposureApplyPhysicalCameraExposure);
+		UE_SET_PP(UserFlags);
+	}
+
+	// No support for blendable objects for now.
+
+#undef UE_SET_PP
+#undef UE_LERP_PP
+
+	return bAnyOverwritten;
+}
+
