Cameras: add world context to the Blueprint camera node evaluator

This gives access to all the Blueprint functions that need a world to operate.

diff --git a/Engine/Plugins/Cameras/GameplayCameras/Source/GameplayCameras/Private/Nodes/Utility/BlueprintCameraNode.cpp b/Engine/Plugins/Cameras/GameplayCameras/Source/GameplayCameras/Private/Nodes/Utility/BlueprintCameraNode.cpp
--- a/Engine/Plugins/Cameras/GameplayCameras/Source/GameplayCameras/Private/Nodes/Utility/BlueprintCameraNode.cpp
+++ b/Engine/Plugins/Cameras/GameplayCameras/Source/GameplayCameras/Private/Nodes/Utility/BlueprintCameraNode.cpp
@@ -8,6 +8,7 @@
 #include "Core/CameraNodeEvaluator.h"
 #include "Core/CameraRigBuildContext.h"
 #include "Core/CameraVariableTable.h"
+#include "Engine/World.h"
 #include "GameFramework/Actor.h"
 #include "GameplayCameras.h"
 #include "Templates/UnrealTemplate.h"
@@ -116,6 +117,34 @@ AActor* UBlueprintCameraNodeEvaluator::FindEvaluationContextOwnerActor(TSubclass
 	}
 }
 
+UWorld* UBlueprintCameraNodeEvaluator::GetWorld() const
+{
+	if (UWorld* CachedWorld = WeakCachedWorld.Get())
+	{
+		return CachedWorld;
+	}
+
+	if (HasAllFlags(RF_ClassDefaultObject))
+	{
+		return nullptr;
+	}
+
+	UObject* Outer = GetOuter();
+	while (Outer)
+	{
+		UWorld* World = Outer->GetWorld();
+		if (World)
+		{
+			WeakCachedWorld = World;
+			return World;
+		}
+
+		Outer = Outer->GetOuter();
+	}
+
+	return nullptr;
+}
+
 void UBlueprintCameraNode::OnBuild(FCameraRigBuildContext& BuildContext)
 {
 	if (!CameraNodeEvaluatorClass)
diff --git a/Engine/Plugins/Cameras/GameplayCameras/Source/GameplayCameras/Public/Nodes/Utility/BlueprintCameraNode.h b/Engine/Plugins/Cameras/GameplayCameras/Source/GameplayCameras/Public/Nodes/Utility/BlueprintCameraNode.h
--- a/Engine/Plugins/Cameras/GameplayCameras/Source/GameplayCameras/Public/Nodes/Utility/BlueprintCameraNode.h
+++ b/Engine/Plugins/Cameras/GameplayCameras/Source/GameplayCameras/Public/Nodes/Utility/BlueprintCameraNode.h
@@ -34,6 +34,14 @@ public:
 	UFUNCTION(BlueprintImplementableEvent, Category="Evaluation")
 	void TickCameraNode(float DeltaTime);
 
+	/**
+	 * A utility function that tries to find if an actor owns the evaluation context.
+	 * Handles the situation where the evaluation context is an actor component (like a
+	 * UGameplayCameraComponent) or an actor itself.
+	 */
+	UFUNCTION(BlueprintPure, Category="Evaluation", meta=(DeterminesOutputType="ActorClass"))
+	AActor* FindEvaluationContextOwnerActor(TSubclassOf<AActor> ActorClass) const;
+
 public:
 
 	using FCameraNodeEvaluationParams = UE::Cameras::FCameraNodeEvaluationParams;
@@ -44,13 +52,11 @@ public:
 
 public:
 
-	/**
-	 * A utility function that tries to find if an actor owns the evaluation context.
-	 * Handles the situation where the evaluation context is an actor component (like a
-	 * UGameplayCameraComponent) or an actor itself.
-	 */
-	UFUNCTION(BlueprintPure, Category="Evaluation", meta=(DeterminesOutputType="ActorClass"))
-	AActor* FindEvaluationContextOwnerActor(TSubclassOf<AActor> ActorClass) const;
+	// UObject interface.
+	virtual UWorld* GetWorld() const override;
+#if WITH_EDITOR
+	virtual bool ImplementsGetWorld() const override { return true; }
+#endif  // WITH_EDITOR
 
 protected:
 
@@ -75,6 +81,8 @@ private:
 	TSharedPtr<const UE::Cameras::FCameraEvaluationContext> CurrentContext;
 
 	FCameraNodeEvaluationResult* CurrentResult = nullptr;
+
+	mutable TWeakObjectPtr<UWorld> WeakCachedWorld;
 };
 
 /**
