Cameras: make the director rig usage available to non-editor builds

It might be useful when dynamically generating camera data at runtime.

diff --git a/Engine/Plugins/Cameras/GameplayCameras/Source/GameplayCameras/Private/Core/CameraAssetBuilder.cpp b/Engine/Plugins/Cameras/GameplayCameras/Source/GameplayCameras/Private/Core/CameraAssetBuilder.cpp
--- a/Engine/Plugins/Cameras/GameplayCameras/Source/GameplayCameras/Private/Core/CameraAssetBuilder.cpp
+++ b/Engine/Plugins/Cameras/GameplayCameras/Source/GameplayCameras/Private/Core/CameraAssetBuilder.cpp
@@ -50,18 +50,15 @@ void FCameraAssetBuilder::BuildCameraImpl()
 	{
 		CameraDirector->BuildCameraDirector(BuildLog);
 
-#if WITH_EDITOR
 		FCameraDirectorRigUsageInfo UsageInfo;
 		CameraDirector->GatherRigUsageInfo(UsageInfo);
 		CameraRigs = UsageInfo.CameraRigs;
-#endif
 	}
 	else
 	{
 		BuildLog.AddMessage(EMessageSeverity::Error, LOCTEXT("MissingDirector", "Camera has no director set."));
 	}
 
-#if WITH_EDITOR
 	if (CameraRigs.IsEmpty())
 	{
 		BuildLog.AddMessage(EMessageSeverity::Warning, LOCTEXT("MissingRigs", "Camera isn't using any camera rigs."));
@@ -72,7 +69,6 @@ void FCameraAssetBuilder::BuildCameraImpl()
 		FCameraRigAssetBuilder CameraRigBuilder(BuildLog);
 		CameraRigBuilder.BuildCameraRig(CameraRig);
 	}
-#endif
 }
 
 void FCameraAssetBuilder::UpdateBuildStatus()
diff --git a/Engine/Plugins/Cameras/GameplayCameras/Source/GameplayCameras/Private/Core/CameraDirector.cpp b/Engine/Plugins/Cameras/GameplayCameras/Source/GameplayCameras/Private/Core/CameraDirector.cpp
--- a/Engine/Plugins/Cameras/GameplayCameras/Source/GameplayCameras/Private/Core/CameraDirector.cpp
+++ b/Engine/Plugins/Cameras/GameplayCameras/Source/GameplayCameras/Private/Core/CameraDirector.cpp
@@ -17,6 +17,11 @@ void UCameraDirector::BuildCameraDirector(UE::Cameras::FCameraBuildLog& BuildLog
 	OnBuildCameraDirector(BuildLog);
 }
 
+void UCameraDirector::GatherRigUsageInfo(FCameraDirectorRigUsageInfo& UsageInfo)
+{
+	OnGatherRigUsageInfo(UsageInfo);
+}
+
 #if WITH_EDITOR
 
 void UCameraDirector::FactoryCreateAsset(const FCameraDirectorFactoryCreateParams& InParams)
@@ -24,10 +29,5 @@ void UCameraDirector::FactoryCreateAsset(const FCameraDirectorFactoryCreateParam
 	OnFactoryCreateAsset(InParams);
 }
 
-void UCameraDirector::GatherRigUsageInfo(FCameraDirectorRigUsageInfo& UsageInfo)
-{
-	OnGatherRigUsageInfo(UsageInfo);
-}
-
-#endif
+#endif  // WITH_EDITOR
 
diff --git a/Engine/Plugins/Cameras/GameplayCameras/Source/GameplayCameras/Private/Directors/BlueprintCameraDirector.cpp b/Engine/Plugins/Cameras/GameplayCameras/Source/GameplayCameras/Private/Directors/BlueprintCameraDirector.cpp
--- a/Engine/Plugins/Cameras/GameplayCameras/Source/GameplayCameras/Private/Directors/BlueprintCameraDirector.cpp
+++ b/Engine/Plugins/Cameras/GameplayCameras/Source/GameplayCameras/Private/Directors/BlueprintCameraDirector.cpp
@@ -436,16 +436,6 @@ void UBlueprintCameraDirector::OnBuildCameraDirector(UE::Cameras::FCameraBuildLo
 	}
 }
 
-#if WITH_EDITOR
-
-void UBlueprintCameraDirector::OnFactoryCreateAsset(const FCameraDirectorFactoryCreateParams& InParams)
-{
-	if (!CameraRigProxyTable)
-	{
-		CameraRigProxyTable = NewObject<UCameraRigProxyTable>(this);
-	}
-}
-
 void UBlueprintCameraDirector::OnGatherRigUsageInfo(FCameraDirectorRigUsageInfo& UsageInfo)
 {
 	using namespace UE::Cameras;
@@ -455,6 +445,8 @@ void UBlueprintCameraDirector::OnGatherRigUsageInfo(FCameraDirectorRigUsageInfo&
 		return;
 	}
 
+#if WITH_EDITORONLY_DATA
+
 	UBlueprint* EvaluatorBlueprint = Cast<UBlueprint>(CameraDirectorEvaluatorClass->ClassGeneratedBy);
 	if (!ensure(EvaluatorBlueprint))
 	{
@@ -464,9 +456,21 @@ void UBlueprintCameraDirector::OnGatherRigUsageInfo(FCameraDirectorRigUsageInfo&
 	FOutgoingReferenceFinder ReferenceFinder(EvaluatorBlueprint, UCameraRigAsset::StaticClass());
 	ReferenceFinder.CollectReferences();
 	ReferenceFinder.GetReferencesOfClass<UCameraRigAsset>(UsageInfo.CameraRigs);
+
+#endif  // WITH_EDITORONLY_DATA
+}
+
+#if WITH_EDITOR
+
+void UBlueprintCameraDirector::OnFactoryCreateAsset(const FCameraDirectorFactoryCreateParams& InParams)
+{
+	if (!CameraRigProxyTable)
+	{
+		CameraRigProxyTable = NewObject<UCameraRigProxyTable>(this);
+	}
 }
 
-#endif
+#endif  // WITH_EDITOR
 
 #undef LOCTEXT_NAMESPACE
 
diff --git a/Engine/Plugins/Cameras/GameplayCameras/Source/GameplayCameras/Private/Directors/SingleCameraDirector.cpp b/Engine/Plugins/Cameras/GameplayCameras/Source/GameplayCameras/Private/Directors/SingleCameraDirector.cpp
--- a/Engine/Plugins/Cameras/GameplayCameras/Source/GameplayCameras/Private/Directors/SingleCameraDirector.cpp
+++ b/Engine/Plugins/Cameras/GameplayCameras/Source/GameplayCameras/Private/Directors/SingleCameraDirector.cpp
@@ -50,8 +50,6 @@ void USingleCameraDirector::OnBuildCameraDirector(UE::Cameras::FCameraBuildLog&
 	}
 }
 
-#if WITH_EDITOR
-
 void USingleCameraDirector::OnGatherRigUsageInfo(FCameraDirectorRigUsageInfo& UsageInfo)
 {
 	if (CameraRig)
@@ -60,7 +58,5 @@ void USingleCameraDirector::OnGatherRigUsageInfo(FCameraDirectorRigUsageInfo& Us
 	}
 }
 
-#endif
-
 #undef LOCTEXT_NAMESPACE
 
diff --git a/Engine/Plugins/Cameras/GameplayCameras/Source/GameplayCameras/Private/Directors/StateTreeCameraDirector.cpp b/Engine/Plugins/Cameras/GameplayCameras/Source/GameplayCameras/Private/Directors/StateTreeCameraDirector.cpp
--- a/Engine/Plugins/Cameras/GameplayCameras/Source/GameplayCameras/Private/Directors/StateTreeCameraDirector.cpp
+++ b/Engine/Plugins/Cameras/GameplayCameras/Source/GameplayCameras/Private/Directors/StateTreeCameraDirector.cpp
@@ -265,16 +265,6 @@ void UStateTreeCameraDirector::OnBuildCameraDirector(UE::Cameras::FCameraBuildLo
 	}
 }
 
-#if WITH_EDITOR
-
-void UStateTreeCameraDirector::OnFactoryCreateAsset(const FCameraDirectorFactoryCreateParams& InParams)
-{
-	if (!CameraRigProxyTable)
-	{
-		CameraRigProxyTable = NewObject<UCameraRigProxyTable>(this);
-	}
-}
-
 void UStateTreeCameraDirector::OnGatherRigUsageInfo(FCameraDirectorRigUsageInfo& UsageInfo)
 {
 	using namespace UE::Cameras;
@@ -290,7 +280,17 @@ void UStateTreeCameraDirector::OnGatherRigUsageInfo(FCameraDirectorRigUsageInfo&
 	ReferenceFinder.GetReferencesOfClass<UCameraRigAsset>(UsageInfo.CameraRigs);
 }
 
-#endif
+#if WITH_EDITOR
+
+void UStateTreeCameraDirector::OnFactoryCreateAsset(const FCameraDirectorFactoryCreateParams& InParams)
+{
+	if (!CameraRigProxyTable)
+	{
+		CameraRigProxyTable = NewObject<UCameraRigProxyTable>(this);
+	}
+}
+
+#endif  // WITH_EDITOR
 
 #undef LOCTEXT_NAMESPACE
 
diff --git a/Engine/Plugins/Cameras/GameplayCameras/Source/GameplayCameras/Public/Core/CameraDirector.h b/Engine/Plugins/Cameras/GameplayCameras/Source/GameplayCameras/Public/Core/CameraDirector.h
--- a/Engine/Plugins/Cameras/GameplayCameras/Source/GameplayCameras/Public/Core/CameraDirector.h
+++ b/Engine/Plugins/Cameras/GameplayCameras/Source/GameplayCameras/Public/Core/CameraDirector.h
@@ -20,6 +20,8 @@ struct FCameraDirectorFactoryCreateParams
 {
 };
 
+#endif  // WITH_EDITOR
+
 /**
  * Parameter struct for gathering camera rigs used by a director.
  */
@@ -28,8 +30,6 @@ struct FCameraDirectorRigUsageInfo
 	TArray<UCameraRigAsset*> CameraRigs;
 };
 
-#endif
-
 /**
  * Base class for a camera director.
  */
@@ -48,12 +48,12 @@ public:
 	/** Builds and validates this camera director. */
 	void BuildCameraDirector(UE::Cameras::FCameraBuildLog& BuildLog);
 
+	/** Gets the list of camera rigs used by this camera director. */
+	void GatherRigUsageInfo(FCameraDirectorRigUsageInfo& UsageInfo);
+
 #if WITH_EDITOR
 	/** Called by the asset factories to setup new data before the editor opens. */
 	void FactoryCreateAsset(const FCameraDirectorFactoryCreateParams& InParams);
-
-	/** Gets the list of camera rigs used by this camera director. */
-	void GatherRigUsageInfo(FCameraDirectorRigUsageInfo& UsageInfo);
 #endif
 
 protected:
@@ -64,12 +64,12 @@ protected:
 	/** Builds and validates this camera director. */
 	virtual void OnBuildCameraDirector(UE::Cameras::FCameraBuildLog& BuildLog) {}
 
+	/** Gets the list of camera rigs used by this camera director. */
+	virtual void OnGatherRigUsageInfo(FCameraDirectorRigUsageInfo& UsageInfo) {}
+
 #if WITH_EDITOR
 	/** Called by the asset factories to setup new data before the editor opens. */
 	virtual void OnFactoryCreateAsset(const FCameraDirectorFactoryCreateParams& InParams) {}
-
-	/** Gets the list of camera rigs used by this camera director. */
-	virtual void OnGatherRigUsageInfo(FCameraDirectorRigUsageInfo& UsageInfo) {}
 #endif
 };
 
diff --git a/Engine/Plugins/Cameras/GameplayCameras/Source/GameplayCameras/Public/Directors/SingleCameraDirector.h b/Engine/Plugins/Cameras/GameplayCameras/Source/GameplayCameras/Public/Directors/SingleCameraDirector.h
--- a/Engine/Plugins/Cameras/GameplayCameras/Source/GameplayCameras/Public/Directors/SingleCameraDirector.h
+++ b/Engine/Plugins/Cameras/GameplayCameras/Source/GameplayCameras/Public/Directors/SingleCameraDirector.h
@@ -23,9 +23,7 @@ protected:
 	// UCameraDirector interface.
 	virtual FCameraDirectorEvaluatorPtr OnBuildEvaluator(FCameraDirectorEvaluatorBuilder& Builder) const override;
 	virtual void OnBuildCameraDirector(UE::Cameras::FCameraBuildLog& BuildLog) override;
-#if WITH_EDITOR
 	virtual void OnGatherRigUsageInfo(FCameraDirectorRigUsageInfo& UsageInfo) override;
-#endif
 
 public:
 
