Cameras: offer a default implementation of GetChildren() for camera nodes

Implementing OnGetChildren() is only needed for optimization or for custom cases, and requires setting a flag.
Default nodes implement the custom version.

diff --git a/Engine/Plugins/Cameras/GameplayCameras/Source/GameplayCameras/Private/Core/CameraNode.cpp b/Engine/Plugins/Cameras/GameplayCameras/Source/GameplayCameras/Private/Core/CameraNode.cpp
--- a/Engine/Plugins/Cameras/GameplayCameras/Source/GameplayCameras/Private/Core/CameraNode.cpp
+++ b/Engine/Plugins/Cameras/GameplayCameras/Source/GameplayCameras/Private/Core/CameraNode.cpp
@@ -3,6 +3,7 @@
 #include "Core/CameraNode.h"
 
 #include "Core/CameraNodeEvaluator.h"
+#include "UObject/Object.h"
 
 #include UE_INLINE_GENERATED_CPP_BY_NAME(CameraNode)
 
@@ -25,7 +26,38 @@ void UCameraNode::PostLoad()
 
 FCameraNodeChildrenView UCameraNode::GetChildren()
 {
-	return OnGetChildren();
+	if (EnumHasAnyFlags(PrivateFlags, ECameraNodeFlags::CustomGetChildren))
+	{
+		return OnGetChildren();
+	}
+
+	const UClass* ThisClass = GetClass();
+	FCameraNodeChildrenView ChildrenView;
+	for (TFieldIterator<FProperty> PropertyIt(ThisClass); PropertyIt; ++PropertyIt)
+	{
+		if (FObjectProperty* ObjectProperty = CastField<FObjectProperty>(*PropertyIt))
+		{
+			if (ObjectProperty->PropertyClass->IsChildOf<UCameraNode>())
+			{
+				UObject* Child = ObjectProperty->GetObjectPropertyValue_InContainer(this);
+				ChildrenView.Add(CastChecked<UCameraNode>(Child, ECastCheckedType::NullAllowed));
+			}
+		}
+		else if (FArrayProperty* ArrayProperty = CastField<FArrayProperty>(*PropertyIt))
+		{
+			FObjectProperty* InnerObjectProperty = CastField<FObjectProperty>(ArrayProperty->Inner);
+			if (InnerObjectProperty && InnerObjectProperty->PropertyClass->IsChildOf<UCameraNode>())
+			{
+				FScriptArrayHelper ArrayHelper(ArrayProperty, ArrayProperty->ContainerPtrToValuePtr<void>(this));
+				for (int32 Index = 0; Index < ArrayHelper.Num(); ++Index)
+				{
+					UObject* Child = InnerObjectProperty->GetObjectPropertyValue(ArrayHelper.GetRawPtr(Index));
+					ChildrenView.Add(CastChecked<UCameraNode>(Child, ECastCheckedType::NullAllowed));
+				}
+			}
+		}
+	}
+	return ChildrenView;
 }
 
 void UCameraNode::PreBuild(FCameraBuildLog& BuildLog)
diff --git a/Engine/Plugins/Cameras/GameplayCameras/Source/GameplayCameras/Private/Nodes/Blends/OrbitBlendCameraNode.cpp b/Engine/Plugins/Cameras/GameplayCameras/Source/GameplayCameras/Private/Nodes/Blends/OrbitBlendCameraNode.cpp
--- a/Engine/Plugins/Cameras/GameplayCameras/Source/GameplayCameras/Private/Nodes/Blends/OrbitBlendCameraNode.cpp
+++ b/Engine/Plugins/Cameras/GameplayCameras/Source/GameplayCameras/Private/Nodes/Blends/OrbitBlendCameraNode.cpp
@@ -339,6 +339,12 @@ bool FOrbitBlendCameraNodeEvaluator::ClosestPoints(const FRay3d& A, const FRay3d
 
 }  // namespace UE::Cameras
 
+UOrbitBlendCameraNode::UOrbitBlendCameraNode(const FObjectInitializer& ObjInit)
+	: Super(ObjInit)
+{
+	AddNodeFlags(ECameraNodeFlags::CustomGetChildren);
+}
+
 FCameraNodeChildrenView UOrbitBlendCameraNode::OnGetChildren()
 {
 	return FCameraNodeChildrenView({ DrivingBlend });
diff --git a/Engine/Plugins/Cameras/GameplayCameras/Source/GameplayCameras/Private/Nodes/Common/ArrayCameraNode.cpp b/Engine/Plugins/Cameras/GameplayCameras/Source/GameplayCameras/Private/Nodes/Common/ArrayCameraNode.cpp
--- a/Engine/Plugins/Cameras/GameplayCameras/Source/GameplayCameras/Private/Nodes/Common/ArrayCameraNode.cpp
+++ b/Engine/Plugins/Cameras/GameplayCameras/Source/GameplayCameras/Private/Nodes/Common/ArrayCameraNode.cpp
@@ -57,6 +57,12 @@ void FArrayCameraNodeEvaluator::OnRun(const FCameraNodeEvaluationParams& Params,
 
 }  // namespace UE::Cameras
 
+UArrayCameraNode::UArrayCameraNode(const FObjectInitializer& ObjInit)
+	: Super(ObjInit)
+{
+	AddNodeFlags(ECameraNodeFlags::CustomGetChildren);
+}
+
 FCameraNodeChildrenView UArrayCameraNode::OnGetChildren()
 {
 	return FCameraNodeChildrenView(Children);
diff --git a/Engine/Plugins/Cameras/GameplayCameras/Source/GameplayCameras/Private/Nodes/Common/BoomArmCameraNode.cpp b/Engine/Plugins/Cameras/GameplayCameras/Source/GameplayCameras/Private/Nodes/Common/BoomArmCameraNode.cpp
--- a/Engine/Plugins/Cameras/GameplayCameras/Source/GameplayCameras/Private/Nodes/Common/BoomArmCameraNode.cpp
+++ b/Engine/Plugins/Cameras/GameplayCameras/Source/GameplayCameras/Private/Nodes/Common/BoomArmCameraNode.cpp
@@ -266,6 +266,12 @@ void FBoomArmCameraDebugBlock::OnDebugDraw(const FCameraDebugBlockDrawParams& Pa
 
 }  // namespace UE::Cameras
 
+UBoomArmCameraNode::UBoomArmCameraNode(const FObjectInitializer& ObjInit)
+	: Super(ObjInit)
+{
+	AddNodeFlags(ECameraNodeFlags::CustomGetChildren);
+}
+
 FCameraNodeChildrenView UBoomArmCameraNode::OnGetChildren()
 {
 	return FCameraNodeChildrenView({ InputSlot });
diff --git a/Engine/Plugins/Cameras/GameplayCameras/Source/GameplayCameras/Private/Nodes/Common/SplineOrbitCameraNode.cpp b/Engine/Plugins/Cameras/GameplayCameras/Source/GameplayCameras/Private/Nodes/Common/SplineOrbitCameraNode.cpp
--- a/Engine/Plugins/Cameras/GameplayCameras/Source/GameplayCameras/Private/Nodes/Common/SplineOrbitCameraNode.cpp
+++ b/Engine/Plugins/Cameras/GameplayCameras/Source/GameplayCameras/Private/Nodes/Common/SplineOrbitCameraNode.cpp
@@ -431,6 +431,12 @@ void FSplineOrbitCameraDebugBlock::RenderLocationOffsetOrbits(const USplineOrbit
 
 }  // namespace UE::Cameras
 
+USplineOrbitCameraNode::USplineOrbitCameraNode(const FObjectInitializer& ObjInit)
+	: Super(ObjInit)
+{
+	AddNodeFlags(ECameraNodeFlags::CustomGetChildren);
+}
+
 FCameraNodeChildrenView USplineOrbitCameraNode::OnGetChildren()
 {
 	return FCameraNodeChildrenView({ InputSlot });
diff --git a/Engine/Plugins/Cameras/GameplayCameras/Source/GameplayCameras/Private/Nodes/Input/AutoRotateInput2DCameraNode.cpp b/Engine/Plugins/Cameras/GameplayCameras/Source/GameplayCameras/Private/Nodes/Input/AutoRotateInput2DCameraNode.cpp
--- a/Engine/Plugins/Cameras/GameplayCameras/Source/GameplayCameras/Private/Nodes/Input/AutoRotateInput2DCameraNode.cpp
+++ b/Engine/Plugins/Cameras/GameplayCameras/Source/GameplayCameras/Private/Nodes/Input/AutoRotateInput2DCameraNode.cpp
@@ -314,6 +314,7 @@ void FAutoRotateInput2DCameraDebugBlock::OnDebugDraw(const FCameraDebugBlockDraw
 UAutoRotateInput2DCameraNode::UAutoRotateInput2DCameraNode(const FObjectInitializer& ObjInit)
 	: Super(ObjInit)
 {
+	AddNodeFlags(ECameraNodeFlags::CustomGetChildren);
 }
 
 FCameraNodeChildrenView UAutoRotateInput2DCameraNode::OnGetChildren()
diff --git a/Engine/Plugins/Cameras/GameplayCameras/Source/GameplayCameras/Public/Core/CameraNode.h b/Engine/Plugins/Cameras/GameplayCameras/Source/GameplayCameras/Public/Core/CameraNode.h
--- a/Engine/Plugins/Cameras/GameplayCameras/Source/GameplayCameras/Public/Core/CameraNode.h
+++ b/Engine/Plugins/Cameras/GameplayCameras/Source/GameplayCameras/Public/Core/CameraNode.h
@@ -19,6 +19,16 @@ namespace UE::Cameras
 /** View on a camera node's children. */
 using FCameraNodeChildrenView = UE::Cameras::TObjectChildrenView<TObjectPtr<UCameraNode>>;
 
+/**
+ * Flags describing the needs of a camera node.
+ */
+enum class ECameraNodeFlags
+{
+	None = 0,
+	CustomGetChildren = 1 << 0
+};
+ENUM_CLASS_FLAGS(ECameraNodeFlags)
+
 /**
  * The base class for a camera node.
  */
@@ -47,6 +57,14 @@ public:
 	/** Builds the evaluator for this node. */
 	FCameraNodeEvaluatorPtr BuildEvaluator(FCameraNodeEvaluatorBuilder& Builder) const;
 
+protected:
+
+	/** Adds the given flags to this node. */
+	void AddNodeFlags(ECameraNodeFlags InFlags) { PrivateFlags |= InFlags; }
+
+	/** Sets the flags for this node. */
+	void SetNodeFlags(ECameraNodeFlags InFlags) { PrivateFlags = InFlags; }
+
 protected:
 
 	/** Get the list of children under this node. */
@@ -83,6 +101,8 @@ public:
 
 #if WITH_EDITORONLY_DATA
 
+public:
+
 	/** Position of the camera node in the node graph editor. */
 	UPROPERTY()
 	FIntVector2 GraphNodePos = FIntVector2::ZeroValue;
@@ -91,6 +111,7 @@ public:
 	UPROPERTY()
 	FString GraphNodeComment;
 
+private:
 
 	// Deprecated properties.
 
@@ -100,5 +121,10 @@ public:
 	int32 GraphNodePosY_DEPRECATED = 0;
 
 #endif  // WITH_EDITORONLY_DATA
+
+private:
+
+	/** Flags on this camera node. */
+	ECameraNodeFlags PrivateFlags = ECameraNodeFlags::None;
 };
 
diff --git a/Engine/Plugins/Cameras/GameplayCameras/Source/GameplayCameras/Public/Nodes/Blends/OrbitBlendCameraNode.h b/Engine/Plugins/Cameras/GameplayCameras/Source/GameplayCameras/Public/Nodes/Blends/OrbitBlendCameraNode.h
--- a/Engine/Plugins/Cameras/GameplayCameras/Source/GameplayCameras/Public/Nodes/Blends/OrbitBlendCameraNode.h
+++ b/Engine/Plugins/Cameras/GameplayCameras/Source/GameplayCameras/Public/Nodes/Blends/OrbitBlendCameraNode.h
@@ -15,6 +15,10 @@ class UOrbitBlendCameraNode : public UBlendCameraNode
 {
 	GENERATED_BODY()
 
+public:
+
+	UOrbitBlendCameraNode(const FObjectInitializer& ObjInit);
+
 public:
 
 	UPROPERTY()
diff --git a/Engine/Plugins/Cameras/GameplayCameras/Source/GameplayCameras/Public/Nodes/Common/ArrayCameraNode.h b/Engine/Plugins/Cameras/GameplayCameras/Source/GameplayCameras/Public/Nodes/Common/ArrayCameraNode.h
--- a/Engine/Plugins/Cameras/GameplayCameras/Source/GameplayCameras/Public/Nodes/Common/ArrayCameraNode.h
+++ b/Engine/Plugins/Cameras/GameplayCameras/Source/GameplayCameras/Public/Nodes/Common/ArrayCameraNode.h
@@ -15,6 +15,10 @@ class UArrayCameraNode : public UCameraNode
 {
 	GENERATED_BODY()
 
+public:
+
+	UArrayCameraNode(const FObjectInitializer& ObjInit);
+
 protected:
 
 	// UCameraNode interface.
diff --git a/Engine/Plugins/Cameras/GameplayCameras/Source/GameplayCameras/Public/Nodes/Common/BoomArmCameraNode.h b/Engine/Plugins/Cameras/GameplayCameras/Source/GameplayCameras/Public/Nodes/Common/BoomArmCameraNode.h
--- a/Engine/Plugins/Cameras/GameplayCameras/Source/GameplayCameras/Public/Nodes/Common/BoomArmCameraNode.h
+++ b/Engine/Plugins/Cameras/GameplayCameras/Source/GameplayCameras/Public/Nodes/Common/BoomArmCameraNode.h
@@ -18,6 +18,10 @@ class UBoomArmCameraNode : public UCameraNode
 {
 	GENERATED_BODY()
 
+public:
+
+	UBoomArmCameraNode(const FObjectInitializer& ObjInit);
+
 protected:
 
 	// UCameraNode interface.
diff --git a/Engine/Plugins/Cameras/GameplayCameras/Source/GameplayCameras/Public/Nodes/Common/SplineOrbitCameraNode.h b/Engine/Plugins/Cameras/GameplayCameras/Source/GameplayCameras/Public/Nodes/Common/SplineOrbitCameraNode.h
--- a/Engine/Plugins/Cameras/GameplayCameras/Source/GameplayCameras/Public/Nodes/Common/SplineOrbitCameraNode.h
+++ b/Engine/Plugins/Cameras/GameplayCameras/Source/GameplayCameras/Public/Nodes/Common/SplineOrbitCameraNode.h
@@ -60,6 +60,10 @@ class USplineOrbitCameraNode : public UCameraNode
 {
 	GENERATED_BODY()
 
+public:
+
+	USplineOrbitCameraNode(const FObjectInitializer& ObjInit);
+
 protected:
 
 	// UCameraNode interface.
