Cameras: move object-tree-graph drag/drop operations into separate files

diff --git a/Engine/Plugins/Cameras/GameplayCameras/Source/GameplayCamerasEditor/Private/Editors/ObjectTreeDragDropOp.cpp b/Engine/Plugins/Cameras/GameplayCameras/Source/GameplayCamerasEditor/Private/Editors/ObjectTreeDragDropOp.cpp
new file mode 100644
--- /dev/null
+++ b/Engine/Plugins/Cameras/GameplayCameras/Source/GameplayCamerasEditor/Private/Editors/ObjectTreeDragDropOp.cpp
@@ -0,0 +1,98 @@
+// Copyright Epic Games, Inc. All Rights Reserved.
+
+#include "Editors/ObjectTreeDragDropOp.h"
+
+#include "Editors/ObjectTreeGraph.h"
+#include "Editors/ObjectTreeGraphConfig.h"
+#include "Editors/ObjectTreeGraphSchema.h"
+#include "GraphEditor.h"
+#include "SGraphPanel.h"
+#include "ScopedTransaction.h"
+
+#define LOCTEXT_NAMESPACE "ObjectTreeGraphDragDropOp"
+
+TSharedRef<FObjectTreeClassDragDropOp> FObjectTreeClassDragDropOp::New(UClass* ObjectClass)
+{
+	TArray<UClass*> ObjectClasses;
+	ObjectClasses.Add(ObjectClass);
+	return New(ObjectClasses);
+}
+
+TSharedRef<FObjectTreeClassDragDropOp> FObjectTreeClassDragDropOp::New(TArrayView<UClass*> ObjectClasses)
+{
+	TSharedRef<FObjectTreeClassDragDropOp> Operation = MakeShared<FObjectTreeClassDragDropOp>();
+	Operation->ObjectClasses = ObjectClasses;
+	Operation->Construct();
+	return Operation;
+}
+
+FReply FObjectTreeClassDragDropOp::ExecuteDragOver(TSharedPtr<SGraphEditor> GraphEditor)
+{
+	UObjectTreeGraph* Graph = CastChecked<UObjectTreeGraph>(GraphEditor->GetCurrentGraph());
+	TArray<UClass*> PlaceableClasses = FilterPlaceableObjectClasses(Graph);
+
+	if (PlaceableClasses.Num() == ObjectClasses.Num())
+	{
+		const FSlateBrush* OKIcon = FAppStyle::GetBrush(TEXT("Graph.ConnectorFeedback.OK"));
+		SetToolTip(
+				FText::Format(
+					LOCTEXT("OnDragOver_Success", "Create {0} node(s) from the dragged object classes"),
+					ObjectClasses.Num()),
+				OKIcon);
+	}
+	else if (PlaceableClasses.Num() > 0)
+	{
+		const FSlateBrush* WarnIcon = FAppStyle::GetBrush(TEXT("Graph.ConnectorFeedback.OKWarn"));
+		SetToolTip(
+				FText::Format(
+					LOCTEXT("OnDragOver_Warning", "Create {0} node(s) from the dragged object classes, ignoring {1} that can't be created in this graph"),
+					PlaceableClasses.Num(), (ObjectClasses.Num() - PlaceableClasses.Num())),
+				WarnIcon);
+	}
+	else
+	{
+		const FSlateBrush* ErrorIcon = FAppStyle::GetBrush(TEXT("Graph.ConnectorFeedback.Error"));
+		SetToolTip(
+				LOCTEXT("OnDragOver_Error", "The dragged object classes can't be created in this graph"),
+				ErrorIcon);
+	}
+
+	return FReply::Handled();
+}
+
+FReply FObjectTreeClassDragDropOp::ExecuteDrop(TSharedPtr<SGraphEditor> GraphEditor, const FVector2D& NewLocation)
+{
+	const FScopedTransaction Transaction(LOCTEXT("DropObjectClasses", "Drop New Nodes"));
+
+	UObjectTreeGraph* Graph = CastChecked<UObjectTreeGraph>(GraphEditor->GetCurrentGraph());
+	TArray<UClass*> PlaceableClasses = FilterPlaceableObjectClasses(Graph);
+
+	GraphEditor->ClearSelectionSet();
+
+	FVector2D CurLocation = NewLocation;
+	for (UClass* PlaceableClass : PlaceableClasses)
+	{
+		FObjectGraphSchemaAction_NewNode Action;
+		Action.ObjectClass = PlaceableClass;
+		UEdGraphNode* NewNode = Action.PerformAction(Graph, nullptr, CurLocation, false);
+		GraphEditor->SetNodeSelection(NewNode, true);
+
+		CurLocation += FVector2D(20, 20);
+	}
+
+	return FReply::Handled();
+}
+
+TArray<UClass*> FObjectTreeClassDragDropOp::FilterPlaceableObjectClasses(UObjectTreeGraph* InGraph)
+{
+	const FObjectTreeGraphConfig& GraphConfig = InGraph->GetConfig();
+	TArray<UClass*> PlaceableClasses = ObjectClasses.FilterByPredicate(
+			[&GraphConfig](UClass* ObjectClass)
+			{
+				return GraphConfig.IsConnectable(ObjectClass);
+			});
+	return PlaceableClasses;
+}
+
+#undef LOCTEXT_NAMESPACE
+
diff --git a/Engine/Plugins/Cameras/GameplayCameras/Source/GameplayCamerasEditor/Private/Editors/SObjectTreeGraphEditor.cpp b/Engine/Plugins/Cameras/GameplayCameras/Source/GameplayCamerasEditor/Private/Editors/SObjectTreeGraphEditor.cpp
--- a/Engine/Plugins/Cameras/GameplayCameras/Source/GameplayCamerasEditor/Private/Editors/SObjectTreeGraphEditor.cpp
+++ b/Engine/Plugins/Cameras/GameplayCameras/Source/GameplayCamerasEditor/Private/Editors/SObjectTreeGraphEditor.cpp
@@ -7,6 +7,7 @@
 #include "EdGraph/EdGraphPin.h"
 #include "EdGraph/EdGraphSchema.h"
 #include "Editor.h"
+#include "Editors/ObjectTreeDragDropOp.h"
 #include "Editors/ObjectTreeGraph.h"
 #include "Editors/ObjectTreeGraphNode.h"
 #include "Editors/ObjectTreeGraphSchema.h"
@@ -189,34 +190,7 @@ FReply SObjectTreeGraphEditor::OnDragOver(const FGeometry& MyGeometry, const FDr
 	TSharedPtr<FObjectTreeClassDragDropOp> ObjectClassOp = DragDropEvent.GetOperationAs<FObjectTreeClassDragDropOp>();
 	if (ObjectClassOp)
 	{
-		TArray<UClass*> PlaceableClasses = FilterPlaceableObjectClasses(ObjectClassOp->GetObjectClasses());
-		if (PlaceableClasses.Num() == ObjectClassOp->GetObjectClasses().Num())
-		{
-			const FSlateBrush* OKIcon = FAppStyle::GetBrush(TEXT("Graph.ConnectorFeedback.OK"));
-			ObjectClassOp->SetToolTip(
-					FText::Format(
-						LOCTEXT("OnDragOver_Success", "Create {0} node(s) from the dragged object classes"),
-						ObjectClassOp->GetObjectClasses().Num()),
-					OKIcon);
-		}
-		else if (PlaceableClasses.Num() > 0)
-		{
-			const FSlateBrush* WarnIcon = FAppStyle::GetBrush(TEXT("Graph.ConnectorFeedback.OKWarn"));
-			ObjectClassOp->SetToolTip(
-					FText::Format(
-						LOCTEXT("OnDragOver_Warning", "Create {0} node(s) from the dragged object classes, ignoring {1} that can't be created in this graph"),
-						PlaceableClasses.Num(), (ObjectClassOp->GetObjectClasses().Num() - PlaceableClasses.Num())),
-					WarnIcon);
-		}
-		else
-		{
-			const FSlateBrush* ErrorIcon = FAppStyle::GetBrush(TEXT("Graph.ConnectorFeedback.Error"));
-			ObjectClassOp->SetToolTip(
-					LOCTEXT("OnDragOver_Error", "The dragged object classes can't be created in this graph"),
-					ErrorIcon);
-		}
-
-		return FReply::Handled();
+		return ObjectClassOp->ExecuteDragOver(GraphEditor);
 	}
 
 	return SCompoundWidget::OnDragOver(MyGeometry, DragDropEvent);
@@ -227,42 +201,15 @@ FReply SObjectTreeGraphEditor::OnDrop(const FGeometry& MyGeometry, const FDragDr
 	TSharedPtr<FObjectTreeClassDragDropOp> ObjectClassOp = DragDropEvent.GetOperationAs<FObjectTreeClassDragDropOp>();
 	if (ObjectClassOp)
 	{
-		const FScopedTransaction Transaction(LOCTEXT("DropObjectClasses", "Drop New Nodes"));
-
-		TArray<UClass*> PlaceableClasses = FilterPlaceableObjectClasses(ObjectClassOp->GetObjectClasses());
-		UEdGraph* Graph = GraphEditor->GetCurrentGraph();
-
-		GraphEditor->ClearSelectionSet();
-
 		SGraphPanel* GraphPanel = GraphEditor->GetGraphPanel();
 		FVector2D NewLocation = GraphPanel->PanelCoordToGraphCoord(MyGeometry.AbsoluteToLocal(DragDropEvent.GetScreenSpacePosition()));
 
-		for (UClass* PlaceableClass : PlaceableClasses)
-		{
-			FObjectGraphSchemaAction_NewNode Action;
-			Action.ObjectClass = PlaceableClass;
-			UEdGraphNode* NewNode = Action.PerformAction(Graph, nullptr, NewLocation, false);
-			GraphEditor->SetNodeSelection(NewNode, true);
-
-			NewLocation += FVector2D(20, 20);
-		}
+		return ObjectClassOp->ExecuteDrop(GraphEditor, NewLocation);
 	}
 
 	return SCompoundWidget::OnDrop(MyGeometry, DragDropEvent);
 }
 
-TArray<UClass*> SObjectTreeGraphEditor::FilterPlaceableObjectClasses(TArrayView<UClass* const> InObjectClasses)
-{
-	UObjectTreeGraph* Graph = CastChecked<UObjectTreeGraph>(GraphEditor->GetCurrentGraph());
-	const FObjectTreeGraphConfig& GraphConfig = Graph->GetConfig();
-	TArray<UClass*> PlaceableClasses = InObjectClasses.FilterByPredicate(
-			[&GraphConfig](UClass* ObjectClass)
-			{
-				return GraphConfig.IsConnectable(ObjectClass);
-			});
-	return PlaceableClasses;
-}
-
 void SObjectTreeGraphEditor::PostUndo(bool bSuccess)
 {
 	if (bSuccess)
diff --git a/Engine/Plugins/Cameras/GameplayCameras/Source/GameplayCamerasEditor/Private/Editors/SObjectTreeGraphToolbox.cpp b/Engine/Plugins/Cameras/GameplayCameras/Source/GameplayCamerasEditor/Private/Editors/SObjectTreeGraphToolbox.cpp
--- a/Engine/Plugins/Cameras/GameplayCameras/Source/GameplayCamerasEditor/Private/Editors/SObjectTreeGraphToolbox.cpp
+++ b/Engine/Plugins/Cameras/GameplayCameras/Source/GameplayCamerasEditor/Private/Editors/SObjectTreeGraphToolbox.cpp
@@ -3,8 +3,8 @@
 #include "Editors/SObjectTreeGraphToolbox.h"
 
 #include "Algo/AnyOf.h"
-#include "DragAndDrop/AssetDragDropOp.h"
 #include "EditorClassUtils.h"
+#include "Editors/ObjectTreeDragDropOp.h"
 #include "Editors/ObjectTreeGraphConfig.h"
 #include "Styles/ObjectTreeGraphEditorStyle.h"
 #include "UObject/UObjectIterator.h"
@@ -273,20 +273,5 @@ FText SObjectTreeGraphToolbox::GetHighlightText() const
 	return SearchTextFilter->GetRawFilterText();
 }
 
-TSharedRef<FObjectTreeClassDragDropOp> FObjectTreeClassDragDropOp::New(UClass* ObjectClass)
-{
-	TArray<UClass*> ObjectClasses;
-	ObjectClasses.Add(ObjectClass);
-	return New(ObjectClasses);
-}
-
-TSharedRef<FObjectTreeClassDragDropOp> FObjectTreeClassDragDropOp::New(TArrayView<UClass*> ObjectClasses)
-{
-	TSharedRef<FObjectTreeClassDragDropOp> Operation = MakeShared<FObjectTreeClassDragDropOp>();
-	Operation->ObjectClasses = ObjectClasses;
-	Operation->Construct();
-	return Operation;
-}
-
 #undef LOCTEXT_NAMESPACE
 
diff --git a/Engine/Plugins/Cameras/GameplayCameras/Source/GameplayCamerasEditor/Public/Editors/ObjectTreeDragDropOp.h b/Engine/Plugins/Cameras/GameplayCameras/Source/GameplayCamerasEditor/Public/Editors/ObjectTreeDragDropOp.h
new file mode 100644
--- /dev/null
+++ b/Engine/Plugins/Cameras/GameplayCameras/Source/GameplayCamerasEditor/Public/Editors/ObjectTreeDragDropOp.h
@@ -0,0 +1,36 @@
+// Copyright Epic Games, Inc. All Rights Reserved.
+
+#pragma once
+
+#include "DragAndDrop/DecoratedDragDropOp.h"
+
+class UObjectTreeGraph;
+class SGraphEditor;
+
+/**
+ * Drag-drop operation for creating a new object (and corresponding graph node) in an object tree graph
+ * by dragging one of the entries from the toolbox widget.
+ */
+class FObjectTreeClassDragDropOp : public FDecoratedDragDropOp
+{
+public:
+	
+	DRAG_DROP_OPERATOR_TYPE(FObjectTreeClassDragDropOp, FDecoratedDragDropOp)
+
+	static TSharedRef<FObjectTreeClassDragDropOp> New(UClass* InObjectClass);
+	static TSharedRef<FObjectTreeClassDragDropOp> New(TArrayView<UClass*> InObjectClasses);
+
+	TArrayView<UClass* const> GetObjectClasses() const { return ObjectClasses; }
+
+	FReply ExecuteDragOver(TSharedPtr<SGraphEditor> GraphEditor);
+	FReply ExecuteDrop(TSharedPtr<SGraphEditor> GraphEditor, const FVector2D& NewLocation);
+
+private:
+
+	TArray<UClass*> FilterPlaceableObjectClasses(UObjectTreeGraph* InGraph);
+
+private:
+
+	TArray<UClass*> ObjectClasses;
+};
+
diff --git a/Engine/Plugins/Cameras/GameplayCameras/Source/GameplayCamerasEditor/Public/Editors/SObjectTreeGraphEditor.h b/Engine/Plugins/Cameras/GameplayCameras/Source/GameplayCamerasEditor/Public/Editors/SObjectTreeGraphEditor.h
--- a/Engine/Plugins/Cameras/GameplayCameras/Source/GameplayCamerasEditor/Public/Editors/SObjectTreeGraphEditor.h
+++ b/Engine/Plugins/Cameras/GameplayCameras/Source/GameplayCamerasEditor/Public/Editors/SObjectTreeGraphEditor.h
@@ -104,8 +104,6 @@ protected:
 	void OnInsertArrayItemPinAfter();
 	void OnRemoveArrayItemPin();
 
-	TArray<UClass*> FilterPlaceableObjectClasses(TArrayView<UClass* const> InObjectClasses);
-
 protected:
 
 	TSharedPtr<SGraphEditor> GraphEditor;
diff --git a/Engine/Plugins/Cameras/GameplayCameras/Source/GameplayCamerasEditor/Public/Editors/SObjectTreeGraphToolbox.h b/Engine/Plugins/Cameras/GameplayCameras/Source/GameplayCamerasEditor/Public/Editors/SObjectTreeGraphToolbox.h
--- a/Engine/Plugins/Cameras/GameplayCameras/Source/GameplayCamerasEditor/Public/Editors/SObjectTreeGraphToolbox.h
+++ b/Engine/Plugins/Cameras/GameplayCameras/Source/GameplayCamerasEditor/Public/Editors/SObjectTreeGraphToolbox.h
@@ -3,7 +3,6 @@
 #pragma once
 
 #include "CoreTypes.h"
-#include "DragAndDrop/DecoratedDragDropOp.h"
 #include "Editors/ObjectTreeGraphConfig.h"
 #include "Misc/TextFilter.h"
 #include "Widgets/SCompoundWidget.h"
@@ -115,23 +114,3 @@ private:
 	bool bUpdateFilteredItemSource = false;
 };
 
-/**
- * Drag-drop operation for creating a new object (and corresponding graph node) in an object tree graph
- * by dragging one of the entries from the toolbox widget.
- */
-class FObjectTreeClassDragDropOp : public FDecoratedDragDropOp
-{
-public:
-	
-	DRAG_DROP_OPERATOR_TYPE(FObjectTreeClassDragDropOp, FDecoratedDragDropOp)
-
-	static TSharedRef<FObjectTreeClassDragDropOp> New(UClass* InObjectClass);
-	static TSharedRef<FObjectTreeClassDragDropOp> New(TArrayView<UClass*> InObjectClasses);
-
-	TArrayView<UClass* const> GetObjectClasses() const { return ObjectClasses; }
-
-private:
-
-	TArray<UClass*> ObjectClasses;
-};
-
