Cameras: add extensibility points to object-tree-graph schema class

Make it easier to store "unmanaged" object-tree-nodes in the graph, i.e. nodes that aren't getting properties set automatically by the connections in the graph.

diff --git a/Engine/Plugins/Cameras/GameplayCameras/Source/GameplayCamerasEditor/Private/Editors/ObjectTreeGraphSchema.cpp b/Engine/Plugins/Cameras/GameplayCameras/Source/GameplayCamerasEditor/Private/Editors/ObjectTreeGraphSchema.cpp
--- a/Engine/Plugins/Cameras/GameplayCameras/Source/GameplayCamerasEditor/Private/Editors/ObjectTreeGraphSchema.cpp
+++ b/Engine/Plugins/Cameras/GameplayCameras/Source/GameplayCamerasEditor/Private/Editors/ObjectTreeGraphSchema.cpp
@@ -365,12 +365,24 @@ UObjectTreeGraphNode* UObjectTreeGraphSchema::OnCreateObjectNode(UObjectTreeGrap
 
 void UObjectTreeGraphSchema::AddConnectableObject(UObjectTreeGraph* InGraph, UObjectTreeGraphNode* InNewNode) const
 {
+	UObject* Object = InNewNode->GetObject();
+	if (!ensure(Object))
+	{
+		return;
+	}
+
+	const FObjectTreeGraphConfig& GraphConfig = InGraph->GetConfig();
+	if (!GraphConfig.IsConnectable(Object->GetClass()))
+	{
+		return;
+	}
+
 	UObjectTreeGraphNode* RootObjectNode = InGraph->GetRootObjectNode();
 	IObjectTreeGraphRootObject* RootObjectInterface = Cast<IObjectTreeGraphRootObject>(RootObjectNode->GetObject());
 	if (RootObjectInterface)
 	{
-		const FName GraphName = InGraph->GetConfig().GraphName;
-		RootObjectInterface->AddConnectableObject(GraphName, InNewNode->GetObject());
+		const FName GraphName = GraphConfig.GraphName;
+		RootObjectInterface->AddConnectableObject(GraphName, Object);
 	}
 
 	OnAddConnectableObject(InGraph, InNewNode);
@@ -382,11 +394,23 @@ void UObjectTreeGraphSchema::OnAddConnectableObject(UObjectTreeGraph* InGraph, U
 
 void UObjectTreeGraphSchema::RemoveConnectableObject(UObjectTreeGraph* InGraph, UObjectTreeGraphNode* InRemovedNode) const
 {
-	const FName GraphName = InGraph->GetConfig().GraphName;
+	UObject* Object = InRemovedNode->GetObject();
+	if (!ensure(Object))
+	{
+		return;
+	}
+
+	const FObjectTreeGraphConfig& GraphConfig = InGraph->GetConfig();
+	if (!GraphConfig.IsConnectable(Object->GetClass()))
+	{
+		return;
+	}
+
 	IObjectTreeGraphRootObject* RootObjectInterface = Cast<IObjectTreeGraphRootObject>(InGraph->GetRootObject());
 	if (RootObjectInterface)
 	{
-		RootObjectInterface->RemoveConnectableObject(GraphName, InRemovedNode->GetObject());
+		const FName GraphName = GraphConfig.GraphName;
+		RootObjectInterface->RemoveConnectableObject(GraphName, Object);
 	}
 
 	OnRemoveConnectableObject(InGraph, InRemovedNode);
@@ -668,6 +692,11 @@ bool UObjectTreeGraphSchema::TryCreateConnection(UEdGraphPin* A, UEdGraphPin* B)
 		return false;
 	}
 
+	if (OnTryCreateCustomConnection(A, B))
+	{
+		return true;
+	}
+
 	UObjectTreeGraphNode* NodeA = Cast<UObjectTreeGraphNode>(A->GetOwningNode());
 	UObjectTreeGraphNode* NodeB = Cast<UObjectTreeGraphNode>(B->GetOwningNode());
 	if (NodeA && NodeA->GetObject() && NodeB && NodeB->GetObject())
@@ -678,6 +707,11 @@ bool UObjectTreeGraphSchema::TryCreateConnection(UEdGraphPin* A, UEdGraphPin* B)
 	return true;
 }
 
+bool UObjectTreeGraphSchema::OnTryCreateCustomConnection(UEdGraphPin* A, UEdGraphPin* B) const
+{
+	return false;
+}
+
 void UObjectTreeGraphSchema::BreakNodeLinks(UEdGraphNode& TargetNode) const
 {
 	const FScopedTransaction Transaction(LOCTEXT("BreakNodeLinks", "Break Node Links"));
@@ -721,29 +755,45 @@ void UObjectTreeGraphSchema::BreakPinLinks(UEdGraphPin& TargetPin, bool bSendsNo
 {
 	FScopedTransaction Transaction(LOCTEXT("BreakPinLinks", "Break Pin Links"));
 
-	UObjectTreeGraphNode* TargetNode = Cast<UObjectTreeGraphNode>(TargetPin.GetOwningNode());
-	if (TargetNode && TargetNode->GetObject())
+	if (!OnBreakCustomPinLinks(TargetPin))
 	{
-		ApplyDisconnection(&TargetPin);
+		UObjectTreeGraphNode* TargetNode = Cast<UObjectTreeGraphNode>(TargetPin.GetOwningNode());
+		if (TargetNode && TargetNode->GetObject())
+		{
+			ApplyDisconnection(&TargetPin);
+		}
 	}
 
 	Super::BreakPinLinks(TargetPin, bSendsNodeNotification);
 }
 
+bool UObjectTreeGraphSchema::OnBreakCustomPinLinks(UEdGraphPin& TargetPin) const
+{
+	return false;
+}
+
 void UObjectTreeGraphSchema::BreakSinglePinLink(UEdGraphPin* SourcePin, UEdGraphPin* TargetPin) const
 {
 	FScopedTransaction Transaction(LOCTEXT("BreakSinglePinLink", "Break Pin Link"));
 
-	UObjectTreeGraphNode* SourceNode = Cast<UObjectTreeGraphNode>(SourcePin->GetOwningNode());
-	UObjectTreeGraphNode* TargetNode = Cast<UObjectTreeGraphNode>(TargetPin->GetOwningNode());
-	if (SourceNode && SourceNode->GetObject() && TargetNode && TargetNode->GetObject())
+	if (!OnBreakSingleCustomPinLink(SourcePin, TargetPin))
 	{
-		ApplyDisconnection(SourcePin, TargetPin);
+		UObjectTreeGraphNode* SourceNode = Cast<UObjectTreeGraphNode>(SourcePin->GetOwningNode());
+		UObjectTreeGraphNode* TargetNode = Cast<UObjectTreeGraphNode>(TargetPin->GetOwningNode());
+		if (SourceNode && SourceNode->GetObject() && TargetNode && TargetNode->GetObject())
+		{
+			ApplyDisconnection(SourcePin, TargetPin);
+		}
 	}
 
 	Super::BreakSinglePinLink(SourcePin, TargetPin);
 }
 
+bool UObjectTreeGraphSchema::OnBreakSingleCustomPinLink(UEdGraphPin* SourcePin, UEdGraphPin* TargetPin) const
+{
+	return false;
+}
+
 void UObjectTreeGraphSchema::ApplyConnection(UEdGraphPin* A, UEdGraphPin* B) const
 {
 	// Input must have been validated prior to calling this method:
diff --git a/Engine/Plugins/Cameras/GameplayCameras/Source/GameplayCamerasEditor/Public/Editors/ObjectTreeGraphSchema.h b/Engine/Plugins/Cameras/GameplayCameras/Source/GameplayCamerasEditor/Public/Editors/ObjectTreeGraphSchema.h
--- a/Engine/Plugins/Cameras/GameplayCameras/Source/GameplayCamerasEditor/Public/Editors/ObjectTreeGraphSchema.h
+++ b/Engine/Plugins/Cameras/GameplayCameras/Source/GameplayCamerasEditor/Public/Editors/ObjectTreeGraphSchema.h
@@ -104,6 +104,9 @@ protected:
 	virtual void OnAddConnectableObject(UObjectTreeGraph* InGraph, UObjectTreeGraphNode* InNewNode) const;
 	virtual void OnRemoveConnectableObject(UObjectTreeGraph* InGraph, UObjectTreeGraphNode* InRemovedNode) const;
 	virtual void CopyNonObjectNodes(TArrayView<UObject*> InObjects, FStringOutputDevice& OutDevice) const;
+	virtual bool OnTryCreateCustomConnection(UEdGraphPin* A, UEdGraphPin* B) const;
+	virtual bool OnBreakCustomPinLinks(UEdGraphPin& TargetPin) const;
+	virtual bool OnBreakSingleCustomPinLink(UEdGraphPin* SourcePin, UEdGraphPin* TargetPin) const;
 	virtual bool OnApplyConnection(UEdGraphPin* A, UEdGraphPin* B) const;
 	virtual bool OnApplyDisconnection(UEdGraphPin* TargetPin) const;
 	virtual bool OnApplyDisconnection(UEdGraphPin* SourcePin, UEdGraphPin* TargetPin) const;
