Cameras: add option to force-activate camera rigs on the main layer.

diff --git a/Engine/Plugins/Cameras/GameplayCameras/Source/GameplayCameras/Private/Core/BlendStackCameraNode.cpp b/Engine/Plugins/Cameras/GameplayCameras/Source/GameplayCameras/Private/Core/BlendStackCameraNode.cpp
--- a/Engine/Plugins/Cameras/GameplayCameras/Source/GameplayCameras/Private/Core/BlendStackCameraNode.cpp
+++ b/Engine/Plugins/Cameras/GameplayCameras/Source/GameplayCameras/Private/Core/BlendStackCameraNode.cpp
@@ -462,9 +462,9 @@ void FTransientBlendStackCameraNodeEvaluator::Push(const FBlendStackCameraPushPa
 		if (!TopEntry.bIsFrozen 
 				&& TopEntry.EvaluationContext == Params.EvaluationContext)
 		{
-			// Don't push anything is what is being requested is already the active 
+			// Don't push anything if what is being requested is already the active 
 			// camera rig.
-			if (TopEntry.CameraRig == Params.CameraRig)
+			if (!Params.bForcePush && TopEntry.CameraRig == Params.CameraRig)
 			{
 				return;
 			}
@@ -472,7 +472,7 @@ void FTransientBlendStackCameraNodeEvaluator::Push(const FBlendStackCameraPushPa
 			// See if we can merge the new camera rig onto the active camera rig.
 			const EBlendStackEntryComparison Comparison = TopEntry.RootEvaluator->Compare(Params.CameraRig);
 
-			if (Comparison == EBlendStackEntryComparison::Active)
+			if (!Params.bForcePush && Comparison == EBlendStackEntryComparison::Active)
 			{
 				// This camera rig is already the active one on the merged stack.
 				return;
@@ -947,13 +947,16 @@ const UCameraRigTransition* FTransientBlendStackCameraNodeEvaluator::FindTransit
 void FPersistentBlendStackCameraNodeEvaluator::Insert(const FBlendStackCameraInsertParams& Params)
 {
 	// See if we already have this camera rig and evaluation context in the stack.
-	for (const FCameraRigEntry& Entry : Entries)
+	if (!Params.bForceInsert)
 	{
-		if (!Entry.bIsFrozen &&
-				Entry.CameraRig == Params.CameraRig &&
-				Entry.EvaluationContext == Params.EvaluationContext)
+		for (const FCameraRigEntry& Entry : Entries)
 		{
-			return;
+			if (!Entry.bIsFrozen &&
+					Entry.CameraRig == Params.CameraRig &&
+					Entry.EvaluationContext == Params.EvaluationContext)
+			{
+				return;
+			}
 		}
 	}
 
diff --git a/Engine/Plugins/Cameras/GameplayCameras/Source/GameplayCameras/Private/Core/CameraSystemEvaluator.cpp b/Engine/Plugins/Cameras/GameplayCameras/Source/GameplayCameras/Private/Core/CameraSystemEvaluator.cpp
--- a/Engine/Plugins/Cameras/GameplayCameras/Source/GameplayCameras/Private/Core/CameraSystemEvaluator.cpp
+++ b/Engine/Plugins/Cameras/GameplayCameras/Source/GameplayCameras/Private/Core/CameraSystemEvaluator.cpp
@@ -239,6 +239,7 @@ void FCameraSystemEvaluator::Update(const FCameraSystemEvaluationParams& Params)
 			FActivateCameraRigParams CameraRigParams;
 			CameraRigParams.EvaluationContext = ActiveCameraRig.EvaluationContext;
 			CameraRigParams.CameraRig = ActiveCameraRig.CameraRig;
+			CameraRigParams.bForceActivate = DirectorResult.bForceActivateCameraRigs;
 			RootEvaluator->ActivateCameraRig(CameraRigParams);
 		}
 		else if (DirectorResult.ActiveCameraRigs.Num() > 1)
@@ -270,6 +271,7 @@ void FCameraSystemEvaluator::Update(const FCameraSystemEvaluationParams& Params)
 			FActivateCameraRigParams CameraRigParams;
 			CameraRigParams.EvaluationContext = CommonContext;
 			CameraRigParams.CameraRig = CombinedCameraRig;
+			CameraRigParams.bForceActivate = DirectorResult.bForceActivateCameraRigs;
 			RootEvaluator->ActivateCameraRig(CameraRigParams);
 		}
 	}
diff --git a/Engine/Plugins/Cameras/GameplayCameras/Source/GameplayCameras/Private/Core/DefaultRootCameraNode.cpp b/Engine/Plugins/Cameras/GameplayCameras/Source/GameplayCameras/Private/Core/DefaultRootCameraNode.cpp
--- a/Engine/Plugins/Cameras/GameplayCameras/Source/GameplayCameras/Private/Core/DefaultRootCameraNode.cpp
+++ b/Engine/Plugins/Cameras/GameplayCameras/Source/GameplayCameras/Private/Core/DefaultRootCameraNode.cpp
@@ -86,6 +86,7 @@ void FDefaultRootCameraNodeEvaluator::OnActivateCameraRig(const FActivateCameraR
 		FBlendStackCameraPushParams PushParams;
 		PushParams.EvaluationContext = Params.EvaluationContext;
 		PushParams.CameraRig = Params.CameraRig;
+		PushParams.bForcePush = Params.bForceActivate;
 		MainLayer->Push(PushParams);
 	}
 	else
@@ -108,6 +109,7 @@ void FDefaultRootCameraNodeEvaluator::OnActivateCameraRig(const FActivateCameraR
 			FBlendStackCameraInsertParams InsertParams;
 			InsertParams.EvaluationContext = Params.EvaluationContext;
 			InsertParams.CameraRig = Params.CameraRig;
+			InsertParams.bForceInsert = Params.bForceActivate;
 			TargetLayer->Insert(InsertParams);
 		}
 	}
diff --git a/Engine/Plugins/Cameras/GameplayCameras/Source/GameplayCameras/Private/Directors/BlueprintCameraDirector.cpp b/Engine/Plugins/Cameras/GameplayCameras/Source/GameplayCameras/Private/Directors/BlueprintCameraDirector.cpp
--- a/Engine/Plugins/Cameras/GameplayCameras/Source/GameplayCameras/Private/Directors/BlueprintCameraDirector.cpp
+++ b/Engine/Plugins/Cameras/GameplayCameras/Source/GameplayCameras/Private/Directors/BlueprintCameraDirector.cpp
@@ -32,6 +32,7 @@ void FBlueprintCameraDirectorEvaluationResult::Reset()
 	ActiveCameraRigs.Reset();
 	ActivePersistentCameraRigs.Reset();
 	InactivePersistentCameraRigs.Reset();
+	bForceActivateCameraRigs = false;
 }
 
 class FBlueprintCameraDirectorEvaluator : public FCameraDirectorEvaluator
@@ -183,6 +184,8 @@ void FBlueprintCameraDirectorEvaluator::ActivateTransientCameraRigs(
 	{
 		OutResult.Add(Params.OwnerContext, ActiveCameraRig);
 	}
+
+	OutResult.bForceActivateCameraRigs = BlueprintResult.bForceActivateCameraRigs;
 }
 
 void FBlueprintCameraDirectorEvaluator::ActivateDeactivePersistentCameraRigs(
@@ -283,9 +286,10 @@ void UBlueprintCameraDirectorEvaluator::DeactivatePersistentVisualCameraRig(UCam
 	EvaluationResult.InactivePersistentCameraRigs.Add({ CameraRigPrefab, ECameraRigLayer::Visual });
 }
 
-void UBlueprintCameraDirectorEvaluator::ActivateCameraRig(UCameraRigAsset* CameraRig)
+void UBlueprintCameraDirectorEvaluator::ActivateCameraRig(UCameraRigAsset* CameraRig, bool bForceNewInstance)
 {
 	EvaluationResult.ActiveCameraRigs.Add(CameraRig);
+	EvaluationResult.bForceActivateCameraRigs |= bForceNewInstance;
 }
 
 void UBlueprintCameraDirectorEvaluator::ActivateCameraRigViaProxy(UCameraRigProxyAsset* CameraRigProxy)
diff --git a/Engine/Plugins/Cameras/GameplayCameras/Source/GameplayCameras/Public/Core/BlendStackCameraNode.h b/Engine/Plugins/Cameras/GameplayCameras/Source/GameplayCameras/Public/Core/BlendStackCameraNode.h
--- a/Engine/Plugins/Cameras/GameplayCameras/Source/GameplayCameras/Public/Core/BlendStackCameraNode.h
+++ b/Engine/Plugins/Cameras/GameplayCameras/Source/GameplayCameras/Public/Core/BlendStackCameraNode.h
@@ -233,6 +233,9 @@ struct FBlendStackCameraPushParams
 
 	/** The source camera rig asset to instantiate and push on the blend stack. */
 	TObjectPtr<const UCameraRigAsset> CameraRig;
+
+	/** Whether to force pushing a new instance of the camera rig, even if it is currently active. */
+	bool bForcePush = false;
 };
 
 /**
@@ -302,6 +305,9 @@ struct FBlendStackCameraInsertParams
 
 	/** The source camera rig asset to instantiate and push on the blend stack. */
 	TObjectPtr<const UCameraRigAsset> CameraRig;
+
+	/** Whether to force insert a new instance of the camera rig, even if there is already one in the stack. */
+	bool bForceInsert = false;
 };
 
 /**
diff --git a/Engine/Plugins/Cameras/GameplayCameras/Source/GameplayCameras/Public/Core/CameraDirectorEvaluator.h b/Engine/Plugins/Cameras/GameplayCameras/Source/GameplayCameras/Public/Core/CameraDirectorEvaluator.h
--- a/Engine/Plugins/Cameras/GameplayCameras/Source/GameplayCameras/Public/Core/CameraDirectorEvaluator.h
+++ b/Engine/Plugins/Cameras/GameplayCameras/Source/GameplayCameras/Public/Core/CameraDirectorEvaluator.h
@@ -85,6 +85,10 @@ struct FCameraDirectorEvaluationResult
 	/** The camera rig(s) that the director says should be active this frame in the main layer. */
 	FActiveCameraRigInfos ActiveCameraRigs;
 
+	/** Whether to force creating a new instance of the camera rigs, even if they are already active in the main layer. */
+	bool bForceActivateCameraRigs = false;
+
+	/** Adds a given camera rig with the given evaluation context to the activation list. */
 	void Add(TSharedPtr<const FCameraEvaluationContext> InContext, TObjectPtr<const UCameraRigAsset> InCameraRig)
 	{
 		ActiveCameraRigs.Add(FActiveCameraRigInfo{ InContext, InCameraRig });
diff --git a/Engine/Plugins/Cameras/GameplayCameras/Source/GameplayCameras/Public/Core/RootCameraNode.h b/Engine/Plugins/Cameras/GameplayCameras/Source/GameplayCameras/Public/Core/RootCameraNode.h
--- a/Engine/Plugins/Cameras/GameplayCameras/Source/GameplayCameras/Public/Core/RootCameraNode.h
+++ b/Engine/Plugins/Cameras/GameplayCameras/Source/GameplayCameras/Public/Core/RootCameraNode.h
@@ -53,6 +53,9 @@ struct FActivateCameraRigParams
 
 	/** The evaluation layer on which to instantiate the camera rig. */
 	ECameraRigLayer Layer = ECameraRigLayer::Main;
+
+	/** Whether a new instance of the camera rig should be created, even if it is already active in the given layer. */
+	bool bForceActivate = false;
 };
 
 /**
diff --git a/Engine/Plugins/Cameras/GameplayCameras/Source/GameplayCameras/Public/Directors/BlueprintCameraDirector.h b/Engine/Plugins/Cameras/GameplayCameras/Source/GameplayCameras/Public/Directors/BlueprintCameraDirector.h
--- a/Engine/Plugins/Cameras/GameplayCameras/Source/GameplayCameras/Public/Directors/BlueprintCameraDirector.h
+++ b/Engine/Plugins/Cameras/GameplayCameras/Source/GameplayCameras/Public/Directors/BlueprintCameraDirector.h
@@ -44,6 +44,9 @@ struct FBlueprintCameraDirectorEvaluationResult
 	/** The list of persistent camera rigs to deactivate. */
 	TArray<FBlueprintPersistentCameraRigInfo> InactivePersistentCameraRigs;
 
+	/** Whether to force a new instance of the main layer camera rigs. */
+	bool bForceActivateCameraRigs = false;
+
 	/** Reset this result for a new evaluation. */
 	void Reset();
 };
@@ -154,7 +157,7 @@ public:
 
 	/** Specifies a camera rig to be active this frame. */
 	UFUNCTION(BlueprintCallable, Category="Evaluation")
-	void ActivateCameraRig(UCameraRigAsset* CameraRig);
+	void ActivateCameraRig(UCameraRigAsset* CameraRig, bool bForceNewInstance = false);
 
 	/**
 	 * Specifies a camera rig to be active this frame, via a proxy which is later resolved
