Cameras: don't auto-build camera assets and camera rig assets when autosave is triggered.

Also remove the check for standalone camera rigs since they are now all standalone.

diff --git a/Engine/Plugins/Cameras/GameplayCameras/Source/GameplayCameras/Private/Core/CameraAsset.cpp b/Engine/Plugins/Cameras/GameplayCameras/Source/GameplayCameras/Private/Core/CameraAsset.cpp
--- a/Engine/Plugins/Cameras/GameplayCameras/Source/GameplayCameras/Private/Core/CameraAsset.cpp
+++ b/Engine/Plugins/Cameras/GameplayCameras/Source/GameplayCameras/Private/Core/CameraAsset.cpp
@@ -156,7 +156,9 @@ void UCameraAsset::PreSave(FObjectPreSaveContext ObjectSaveContext)
 {
 #if WITH_EDITOR
 
-	if (!HasAnyFlags(RF_ClassDefaultObject | RF_ArchetypeObject))
+	const bool bIsUserObject = !HasAnyFlags(RF_ClassDefaultObject | RF_ArchetypeObject);
+	const bool bIsEditorAutoSave = !(ObjectSaveContext.GetSaveFlags() & SAVE_FromAutosave);
+	if (bIsUserObject && !bIsEditorAutoSave)
 	{
 		// Build when saving/cooking.
 		BuildCamera();
diff --git a/Engine/Plugins/Cameras/GameplayCameras/Source/GameplayCameras/Private/Core/CameraRigAsset.cpp b/Engine/Plugins/Cameras/GameplayCameras/Source/GameplayCameras/Private/Core/CameraRigAsset.cpp
--- a/Engine/Plugins/Cameras/GameplayCameras/Source/GameplayCameras/Private/Core/CameraRigAsset.cpp
+++ b/Engine/Plugins/Cameras/GameplayCameras/Source/GameplayCameras/Private/Core/CameraRigAsset.cpp
@@ -199,11 +199,11 @@ void UCameraRigAsset::PreSave(FObjectPreSaveContext ObjectSaveContext)
 {
 #if WITH_EDITOR
 
-	if (GetOutermostObject() == this &&
-			!HasAnyFlags(RF_ClassDefaultObject | RF_ArchetypeObject))
+	const bool bIsUserObject = !HasAnyFlags(RF_ClassDefaultObject | RF_ArchetypeObject);
+	const bool bIsEditorAutoSave = !(ObjectSaveContext.GetSaveFlags() & SAVE_FromAutosave);
+	if (bIsUserObject && !bIsEditorAutoSave)
 	{
-		// Build when saving/cooking, if we are a standalone camera rig (i.e. not a camera rig inside a 
-		// camera asset, since those are built along with the camera asset).
+		// Build when saving/cooking.
 		BuildCameraRig();
 	}
 
