Cameras: add low-level, type-erased getter/setter methods on the variable table.

diff --git a/Engine/Plugins/Cameras/GameplayCameras/Source/GameplayCameras/Private/Core/CameraVariableTable.cpp b/Engine/Plugins/Cameras/GameplayCameras/Source/GameplayCameras/Private/Core/CameraVariableTable.cpp
--- a/Engine/Plugins/Cameras/GameplayCameras/Source/GameplayCameras/Private/Core/CameraVariableTable.cpp
+++ b/Engine/Plugins/Cameras/GameplayCameras/Source/GameplayCameras/Private/Core/CameraVariableTable.cpp
@@ -227,6 +227,25 @@ bool FCameraVariableTable::ContainsValue(FCameraVariableID VariableID) const
 	return EntryLookup.Contains(VariableID);
 }
 
+const uint8* FCameraVariableTable::GetValue(FCameraVariableID VariableID, ECameraVariableType ExpectedVariableType) const
+{
+	const uint8* Value = TryGetValue(VariableID, ExpectedVariableType);
+	ensureMsgf(Value, TEXT("Can't get camera variable (ID '%d') because it doesn't exist in the table."), VariableID.GetValue());
+	return Value;
+}
+
+const uint8* FCameraVariableTable::TryGetValue(FCameraVariableID VariableID, ECameraVariableType ExpectedVariableType) const
+{
+	const FEntry* Entry = FindEntry(VariableID);
+	if (Entry)
+	{
+		ensure(Entry->Type == ExpectedVariableType);
+		return Memory + Entry->Offset;
+	}
+
+	return nullptr;
+}
+
 void FCameraVariableTable::SetValue(
 		FCameraVariableID VariableID, 
 		ECameraVariableType ExpectedVariableType, 
diff --git a/Engine/Plugins/Cameras/GameplayCameras/Source/GameplayCameras/Public/Core/CameraVariableTable.h b/Engine/Plugins/Cameras/GameplayCameras/Source/GameplayCameras/Public/Core/CameraVariableTable.h
--- a/Engine/Plugins/Cameras/GameplayCameras/Source/GameplayCameras/Public/Core/CameraVariableTable.h
+++ b/Engine/Plugins/Cameras/GameplayCameras/Source/GameplayCameras/Public/Core/CameraVariableTable.h
@@ -152,6 +152,14 @@ public:
 
 	// Lower level API.
 
+	const uint8* GetValue(
+			FCameraVariableID VariableID,
+			ECameraVariableType ExpectedVariableType) const;
+
+	const uint8* TryGetValue(
+			FCameraVariableID VariableID,
+			ECameraVariableType ExpectedVariableType) const;
+
 	void SetValue(
 			FCameraVariableID VariableID, 
 			ECameraVariableType ExpectedVariableType, 
