Cameras: remove display name on camera rig assets

Now that they're all standalone assets, their name is always the name of their package.

diff --git a/Engine/Plugins/Cameras/GameplayCameras/Source/GameplayCameras/Private/Core/BlendStackCameraNode.cpp b/Engine/Plugins/Cameras/GameplayCameras/Source/GameplayCameras/Private/Core/BlendStackCameraNode.cpp
--- a/Engine/Plugins/Cameras/GameplayCameras/Source/GameplayCameras/Private/Core/BlendStackCameraNode.cpp
+++ b/Engine/Plugins/Cameras/GameplayCameras/Source/GameplayCameras/Private/Core/BlendStackCameraNode.cpp
@@ -1191,7 +1191,7 @@ FBlendStackCameraDebugBlock::FBlendStackCameraDebugBlock(const FBlendStackCamera
 	for (const FBlendStackCameraNodeEvaluator::FCameraRigEntry& Entry : InEvaluator.Entries)
 	{
 		FEntryDebugInfo EntryDebugInfo;
-		EntryDebugInfo.CameraRigName = Entry.CameraRig ? Entry.CameraRig->GetDisplayName() : FString("<None>");
+		EntryDebugInfo.CameraRigName = Entry.CameraRig ? Entry.CameraRig->GetName() : FString("<None>");
 		Entries.Add(EntryDebugInfo);
 	}
 }
diff --git a/Engine/Plugins/Cameras/GameplayCameras/Source/GameplayCameras/Private/Core/BlendStackRootCameraNode.cpp b/Engine/Plugins/Cameras/GameplayCameras/Source/GameplayCameras/Private/Core/BlendStackRootCameraNode.cpp
--- a/Engine/Plugins/Cameras/GameplayCameras/Source/GameplayCameras/Private/Core/BlendStackRootCameraNode.cpp
+++ b/Engine/Plugins/Cameras/GameplayCameras/Source/GameplayCameras/Private/Core/BlendStackRootCameraNode.cpp
@@ -91,7 +91,7 @@ void FBlendStackRootCameraNodeEvaluator::OnInitialize(const FCameraNodeEvaluator
 
 #if UE_GAMEPLAY_CAMERAS_DEBUG
 		const UCameraRigAsset* CameraRig = RootNode->RootNode->GetTypedOuter<UCameraRigAsset>();
-		CameraRigAssetName = CameraRig ? CameraRig->GetDisplayName() : TEXT("<None>");
+		CameraRigAssetName = GetNameSafe(CameraRig);
 #endif  // UE_GAMEPLAY_CAMERAS_DEBUG
 	}
 }
@@ -327,7 +327,7 @@ void FBlendStackRootCameraNodeEvaluator::OnBuildDebugBlocks(const FCameraDebugBl
 					const UCameraRigAsset* OuterCameraRig = Item.PrefabNode->GetTypedOuter<const UCameraRigAsset>();
 					if (OuterCameraRig)
 					{
-						return OuterCameraRig->GetDisplayName();
+						return GetNameSafe(OuterCameraRig);
 					}
 				}
 				return FString(TEXT("<invalid camera rig>"));
diff --git a/Engine/Plugins/Cameras/GameplayCameras/Source/GameplayCameras/Private/Core/CameraRigAsset.cpp b/Engine/Plugins/Cameras/GameplayCameras/Source/GameplayCameras/Private/Core/CameraRigAsset.cpp
--- a/Engine/Plugins/Cameras/GameplayCameras/Source/GameplayCameras/Private/Core/CameraRigAsset.cpp
+++ b/Engine/Plugins/Cameras/GameplayCameras/Source/GameplayCameras/Private/Core/CameraRigAsset.cpp
@@ -164,15 +164,6 @@ void UCameraRigAsset::GetOwnedGameplayTags(FGameplayTagContainer& TagContainer)
 	TagContainer.AppendTags(GameplayTags);
 }
 
-FString UCameraRigAsset::GetDisplayName() const
-{
-	if (!Interface.DisplayName.IsEmpty())
-	{
-		return Interface.DisplayName;
-	}
-	return GetName();
-}
-
 void UCameraRigAsset::BuildCameraRig()
 {
 	using namespace UE::Cameras;
@@ -312,12 +303,7 @@ void UCameraRigAsset::OnUpdateGraphNodeCommentText(FName InGraphName, const FStr
 
 void UCameraRigAsset::GetGraphNodeName(FName InGraphName, FText& OutName) const
 {
-	OutName = FText::FromString(GetDisplayName());
-}
-
-void UCameraRigAsset::OnRenameGraphNode(FName InGraphName, const FString& NewName)
-{
-	Interface.DisplayName = NewName;
+	OutName = FText::FromString(GetName());
 }
 
 void UCameraRigAsset::GetConnectableObjects(FName InGraphName, TSet<UObject*>& OutObjects) const
diff --git a/Engine/Plugins/Cameras/GameplayCameras/Source/GameplayCameras/Public/Core/CameraRigAsset.h b/Engine/Plugins/Cameras/GameplayCameras/Source/GameplayCameras/Public/Core/CameraRigAsset.h
--- a/Engine/Plugins/Cameras/GameplayCameras/Source/GameplayCameras/Public/Core/CameraRigAsset.h
+++ b/Engine/Plugins/Cameras/GameplayCameras/Source/GameplayCameras/Public/Core/CameraRigAsset.h
@@ -148,9 +148,6 @@ struct FCameraRigInterface
 
 public:
 
-	UPROPERTY()
-	FString DisplayName;
-
 	/** The list of exposed parameters on the camera rig. */
 	UPROPERTY(Instanced)
 	TArray<TObjectPtr<UCameraRigInterfaceParameter>> InterfaceParameters;
@@ -226,11 +223,6 @@ public:
 	/** Gets the camera rig's unique ID. */
 	const FGuid& GetGuid() const { return Guid; }
 
-	/**
-	 * Gets the display name of this camera rig.
-	 * This is either the display name set on the interface object, or its internal name.
-	 */
-	GAMEPLAYCAMERAS_API FString GetDisplayName() const;
 
 public:
 
@@ -278,7 +270,6 @@ protected:
 	virtual const FString& GetGraphNodeCommentText(FName InGraphName) const override;
 	virtual void OnUpdateGraphNodeCommentText(FName InGraphName, const FString& NewComment) override;
 	virtual void GetGraphNodeName(FName InGraphName, FText& OutName) const override;
-	virtual void OnRenameGraphNode(FName InGraphName, const FString& NewName) override;
 #endif
 
 	// IObjectTreeGraphRootObject interface.
diff --git a/Engine/Plugins/Cameras/GameplayCameras/Source/GameplayCamerasEditor/Private/Customizations/CameraProxyTableDetailsCustomization.cpp b/Engine/Plugins/Cameras/GameplayCameras/Source/GameplayCamerasEditor/Private/Customizations/CameraProxyTableDetailsCustomization.cpp
--- a/Engine/Plugins/Cameras/GameplayCameras/Source/GameplayCamerasEditor/Private/Customizations/CameraProxyTableDetailsCustomization.cpp
+++ b/Engine/Plugins/Cameras/GameplayCameras/Source/GameplayCamerasEditor/Private/Customizations/CameraProxyTableDetailsCustomization.cpp
@@ -93,7 +93,7 @@ FText FCameraProxyTableEntryDetailsCustomization::OnGetComboButtonText() const
 		TObjectPtr<UCameraRigAsset>* CameraRigPtr = (TObjectPtr<UCameraRigAsset>*)RawData[0];
 		if (*CameraRigPtr)
 		{
-			DisplayText = FText::FromString((*CameraRigPtr)->GetDisplayName());
+			DisplayText = FText::FromString((*CameraRigPtr)->GetName());
 		}
 		return DisplayText;
 	}
diff --git a/Engine/Plugins/Cameras/GameplayCameras/Source/GameplayCamerasEditor/Private/Toolkits/CameraAssetEditorToolkit.cpp b/Engine/Plugins/Cameras/GameplayCameras/Source/GameplayCamerasEditor/Private/Toolkits/CameraAssetEditorToolkit.cpp
--- a/Engine/Plugins/Cameras/GameplayCameras/Source/GameplayCamerasEditor/Private/Toolkits/CameraAssetEditorToolkit.cpp
+++ b/Engine/Plugins/Cameras/GameplayCameras/Source/GameplayCamerasEditor/Private/Toolkits/CameraAssetEditorToolkit.cpp
@@ -515,7 +515,7 @@ void FCameraAssetEditorToolkit::UpgradeLegacyCameraAssets()
 		SlowTask.EnterProgressFrame();
 
 		// Name the new camera rigs like this: "<CameraAsset>_<CameraRig>"
-		const FString CameraRigPackageName = FString::Printf(TEXT("%s_%s"), *CameraRigsBaseName, *CameraRig->GetDisplayName());
+		const FString CameraRigPackageName = FString::Printf(TEXT("%s_%s"), *CameraRigsBaseName, *CameraRig->GetName());
 
 		FString CameraRigPackagePath;
 		FString CameraRigAssetName;
