Cameras: remove old upgrade code for owned camera rig assets, now unused.

diff --git a/Engine/Plugins/Cameras/GameplayCameras/Source/GameplayCameras/Private/Core/CameraRigAsset.cpp b/Engine/Plugins/Cameras/GameplayCameras/Source/GameplayCameras/Private/Core/CameraRigAsset.cpp
--- a/Engine/Plugins/Cameras/GameplayCameras/Source/GameplayCameras/Private/Core/CameraRigAsset.cpp
+++ b/Engine/Plugins/Cameras/GameplayCameras/Source/GameplayCameras/Private/Core/CameraRigAsset.cpp
@@ -106,18 +106,6 @@ const FName UCameraRigAsset::TransitionsGraphName(TEXT("Transitions"));
 
 void UCameraRigAsset::PostLoad()
 {
-#if WITH_EDITOR
-
-	UCameraAsset* OuterCameraAsset = GetTypedOuter<UCameraAsset>();
-	if (OuterCameraAsset && !HasAllFlags(RF_Public | RF_Transactional))
-	{
-		Modify();
-
-		SetFlags(RF_Public | RF_Transactional);
-	}
-
-#endif
-
 #if WITH_EDITORONLY_DATA
 
 	if (GraphNodePosX_DEPRECATED != 0 || GraphNodePosY_DEPRECATED != 0)
