Cameras: big pass on parameterization of camera rigs

- Camera rig interface parameters are now split into two types. The old type is now called "blendable parameters", since they are blended in the variable table and will later gain the ability for pre-blending. The other new type is called "data parameters" and are for any other arbitrary type of (non-blendable) data, like actor pointers and structures.
- Blendable parameters support camera parameters like before, but also camera variable references now.
- A new allocation info stores the types of data parameters needed by a camera rig.
- Added a camera asset reference structure, similar to the camera rig reference structure, so that the camera components can parameterize their camera too.
- Replaced the hard-coded lists of parameter overrides with a more generic property bag. This makes the code more generic, but also allows us to support the new data parameters.
- Also replaced the hard-coded handling of camera rig prefab nodes with a interface for handling custom parameters exposed on a camera node. This will let us support similar parameters on the Blueprint node, for instance.
- Detail customizations for camera asset and camera rig references now use the property bag detail customization.
- Blendable parameters aren't auto-generated/replaced anymore. The user must add them with the correct type and connect them to a compatible pin.
- To this effect, a new "parameters" panel in the camera rig editor has been added. The user can manage blendable and data parameters in there.
- Graph nodes for blendable and data parameters are handled separately from the base object-tree-graph nodes.
- Added pin colors on graph nodes to help differentiate all the parameter types.
- Removed custom builder registration on the editor module since it wasn't used after all.

diff --git a/Engine/Config/BaseEngine.ini b/Engine/Config/BaseEngine.ini
--- a/Engine/Config/BaseEngine.ini
+++ b/Engine/Config/BaseEngine.ini
@@ -1339,8 +1339,6 @@ ClassName=/Script/OnlineSubsystemUtils.OnlineEngineInterfaceImpl
 +FunctionRedirects=(OldName="SkeletalMeshComponent.SetAnimClass",NewName="SetAnimInstanceClass")
 +FunctionRedirects=(OldName="SkeletalMeshComponent.K2_SetAnimInstanceClass",NewName="SetAnimInstanceClass")
 
-; 5.6
-+FunctionRedirects=(OldName="BlueprintCameraDirectorEvaluator.ActivateCameraRigPrefab", NewName="ActivateCameraRig")
 
 ; End of releases redirects
 
diff --git a/Engine/Plugins/Cameras/GameplayCameras/Config/BaseGameplayCameras.ini b/Engine/Plugins/Cameras/GameplayCameras/Config/BaseGameplayCameras.ini
--- a/Engine/Plugins/Cameras/GameplayCameras/Config/BaseGameplayCameras.ini
+++ b/Engine/Plugins/Cameras/GameplayCameras/Config/BaseGameplayCameras.ini
@@ -5,3 +5,7 @@
 +FunctionRedirects=(OldName="ActivateCameraRigFunctions.ActivateGlobalCameraRig", NewName="ActivatePersistentGlobalCameraRig")
 +FunctionRedirects=(OldName="ActivateCameraRigFunctions.ActivateVisualCameraRig", NewName="ActivatePersistentVisualCameraRig")
 
+; UE 5.6 redirects.
++ClassRedirects=(OldName="/Script/GameplayCameras.CameraRigInterfaceParameter", NewName="CameraRigBlendableParameter")
++FunctionRedirects=(OldName="BlueprintCameraDirectorEvaluator.ActivateCameraRigPrefab", NewName="ActivateCameraRig")
++PropertyRedirects=(OldName="/Script/GameplayCameras.CameraRigInterface.InterfaceParameters", NewName="BlendableParameters")
diff --git a/Engine/Plugins/Cameras/GameplayCameras/Source/GameplayCameras/Private/Core/BlendStackCameraNode.cpp b/Engine/Plugins/Cameras/GameplayCameras/Source/GameplayCameras/Private/Core/BlendStackCameraNode.cpp
--- a/Engine/Plugins/Cameras/GameplayCameras/Source/GameplayCameras/Private/Core/BlendStackCameraNode.cpp
+++ b/Engine/Plugins/Cameras/GameplayCameras/Source/GameplayCameras/Private/Core/BlendStackCameraNode.cpp
@@ -89,6 +89,9 @@ bool FBlendStackCameraNodeEvaluator::InitializeEntry(
 	// Allocate variables in the variable table.
 	NewEntry.Result.VariableTable.Initialize(CameraRig->AllocationInfo.VariableTableInfo);
 
+	// Allocate context data in the data table.
+	NewEntry.Result.ContextDataTable.Initialize(CameraRig->AllocationInfo.ContextDataTableInfo);
+
 	// Initialize the node evaluators.
 	FCameraNodeEvaluatorInitializeParams InitParams(&NewEntry.EvaluatorHierarchy);
 	InitParams.Evaluator = OwningEvaluator;
@@ -220,6 +223,7 @@ void FBlendStackCameraNodeEvaluator::ResolveEntries(TArray<FResolvedEntry>& OutR
 			FCameraNodeEvaluationResult& CurResult = Entry.Result;
 			CurResult.CameraPose.ClearAllChangedFlags();
 			CurResult.VariableTable.ClearAllWrittenThisFrameFlags();
+			CurResult.ContextDataTable.ClearAllWrittenThisFrameFlags();
 		}
 		// else: frozen entries may have null contexts or invalid initial results
 		//       because we're not going to update them anyway. We will however blend
@@ -384,6 +388,7 @@ void FBlendStackCameraNodeEvaluator::OnAddReferencedObjects(FReferenceCollector&
 	{
 		Collector.AddReferencedObject(Entry.CameraRig);
 		Collector.AddReferencedObject(Entry.RootNode);
+		Entry.Result.AddReferencedObjects(Collector);
 	}
 }
 
@@ -655,6 +660,7 @@ void FTransientBlendStackCameraNodeEvaluator::InternalPreBlendPrepare(TArrayView
 		// overrides have been set on them.
 		const FCameraNodeEvaluationResult& ContextResult(ResolvedEntry.Context->GetInitialResult());
 		CurResult.VariableTable.Override(ContextResult.VariableTable, ECameraVariableTableFilter::AllPublic | ECameraVariableTableFilter::Private);
+		CurResult.ContextDataTable.OverrideAll(ContextResult.ContextDataTable);
 
 		// Gather input parameters if needed (and remember if it was indeed needed).
 		if (!Entry.bInputRunThisFrame)
@@ -1044,6 +1050,7 @@ void FPersistentBlendStackCameraNodeEvaluator::InternalUpdate(TArrayView<FResolv
 				const FCameraNodeEvaluationResult& ContextResult(ResolvedEntry.Context->GetInitialResult());
 				CurResult.CameraPose.OverrideChanged(ContextResult.CameraPose);
 				CurResult.VariableTable.Override(ContextResult.VariableTable, ECameraVariableTableFilter::AllPublic | ECameraVariableTableFilter::Private);
+				CurResult.ContextDataTable.OverrideAll(ContextResult.ContextDataTable);
 
 				// Setup flags.
 				CurResult.bIsCameraCut = OutResult.bIsCameraCut || ContextResult.bIsCameraCut || Entry.bForceCameraCut;
diff --git a/Engine/Plugins/Cameras/GameplayCameras/Source/GameplayCameras/Private/Core/BlendStackRootCameraNode.cpp b/Engine/Plugins/Cameras/GameplayCameras/Source/GameplayCameras/Private/Core/BlendStackRootCameraNode.cpp
--- a/Engine/Plugins/Cameras/GameplayCameras/Source/GameplayCameras/Private/Core/BlendStackRootCameraNode.cpp
+++ b/Engine/Plugins/Cameras/GameplayCameras/Source/GameplayCameras/Private/Core/BlendStackRootCameraNode.cpp
@@ -6,10 +6,10 @@
 #include "Core/BlendCameraNode.h"
 #include "Core/CameraRigAsset.h"
 #include "Core/CameraRigAssetReference.h"
-#include "Core/CameraRigParameterOverrideEvaluator.h"
 #include "Debug/CameraDebugBlock.h"
 #include "Debug/CameraDebugBlockBuilder.h"
 #include "Debug/CameraDebugRenderer.h"
+#include "Helpers/CameraRigParameterOverrideEvaluator.h"
 #include "Nodes/Common/CameraRigCameraNode.h"
 
 #include UE_INLINE_GENERATED_CPP_BY_NAME(BlendStackRootCameraNode)
@@ -181,19 +181,19 @@ void FBlendStackRootCameraNodeEvaluator::InitializeBlendedParameterOverridesStac
 	}
 
 	// Build the allocation info for the variable tables we keep with each set of parameter overrides.
-	for (const UCameraRigInterfaceParameter* InterfaceParameter : BlendablePrefabCameraRig->Interface.InterfaceParameters)
+	for (const UCameraRigBlendableParameter* BlendableParameter : BlendablePrefabCameraRig->Interface.BlendableParameters)
 	{
-		if (!ensure(InterfaceParameter))
+		if (!ensure(BlendableParameter))
 		{
 			continue;
 		}
-		if (!InterfaceParameter->PrivateVariable)
+		if (!BlendableParameter->PrivateVariable)
 		{
 			continue;
 		}
 
 		BlendedParameterOverridesTableAllocationInfo.VariableDefinitions.Add(
-				InterfaceParameter->PrivateVariable->GetVariableDefinition());;
+				BlendableParameter->PrivateVariable->GetVariableDefinition());
 	}
 
 	FCameraRigCameraNodeEvaluator* RootPrefabNodeEvaluator = RootEvaluator->CastThisChecked<FCameraRigCameraNodeEvaluator>();
@@ -223,16 +223,16 @@ void FBlendStackRootCameraNodeEvaluator::OnRun(const FCameraNodeEvaluationParams
 	}
 }
 
-void FBlendStackRootCameraNodeEvaluator::SetDefaultInterfaceParameterValues(FCameraVariableTable& OutVariableTable)
+void FBlendStackRootCameraNodeEvaluator::SetDefaultBlendableParameterValues(FCameraVariableTable& OutVariableTable)
 {
-	for (const UCameraRigInterfaceParameter* InterfaceParameter : BlendablePrefabCameraRig->Interface.InterfaceParameters)
+	for (const UCameraRigBlendableParameter* BlendableParameter : BlendablePrefabCameraRig->Interface.BlendableParameters)
 	{
-		if (!ensure(InterfaceParameter))
+		if (!ensure(BlendableParameter))
 		{
 			continue;
 		}
 
-		const UCameraVariableAsset* PrivateVariable = InterfaceParameter->PrivateVariable;
+		const UCameraVariableAsset* PrivateVariable = BlendableParameter->PrivateVariable;
 		if (!PrivateVariable)
 		{
 			continue;
@@ -265,7 +265,7 @@ void FBlendStackRootCameraNodeEvaluator::RunBlendedParameterOverridesStack(const
 
 		// Start by setting the default values of all parameters. If we don't do this, parameter overrides
 		// wouldn't have a base value to blend from.
-		SetDefaultInterfaceParameterValues(CurResult.VariableTable);
+		SetDefaultBlendableParameterValues(CurResult.VariableTable);
 
 		// Next, override the defaults with the specific values of this entry.
 		FCameraRigParameterOverrideEvaluator OverrideEvaluator(BlendedParameterOverrides.PrefabNode->CameraRigReference);
diff --git a/Engine/Plugins/Cameras/GameplayCameras/Source/GameplayCameras/Private/Core/BlendStackRootCameraNode.h b/Engine/Plugins/Cameras/GameplayCameras/Source/GameplayCameras/Private/Core/BlendStackRootCameraNode.h
--- a/Engine/Plugins/Cameras/GameplayCameras/Source/GameplayCameras/Private/Core/BlendStackRootCameraNode.h
+++ b/Engine/Plugins/Cameras/GameplayCameras/Source/GameplayCameras/Private/Core/BlendStackRootCameraNode.h
@@ -84,7 +84,7 @@ protected:
 private:
 
 	void InitializeBlendedParameterOverridesStack();
-	void SetDefaultInterfaceParameterValues(FCameraVariableTable& OutVariableTable);
+	void SetDefaultBlendableParameterValues(FCameraVariableTable& OutVariableTable);
 	void RunBlendedParameterOverridesStack(const FCameraBlendedParameterUpdateParams& Params, FCameraBlendedParameterUpdateResult& OutResult);
 
 private:
diff --git a/Engine/Plugins/Cameras/GameplayCameras/Source/GameplayCameras/Private/Core/CameraAsset.cpp b/Engine/Plugins/Cameras/GameplayCameras/Source/GameplayCameras/Private/Core/CameraAsset.cpp
--- a/Engine/Plugins/Cameras/GameplayCameras/Source/GameplayCameras/Private/Core/CameraAsset.cpp
+++ b/Engine/Plugins/Cameras/GameplayCameras/Source/GameplayCameras/Private/Core/CameraAsset.cpp
@@ -11,6 +11,12 @@
 
 #include UE_INLINE_GENERATED_CPP_BY_NAME(CameraAsset)
 
+bool operator==(const FCameraAssetAllocationInfo& A, const FCameraAssetAllocationInfo& B)
+{
+	return A.VariableTableInfo == B.VariableTableInfo &&
+		A.ContextDataTableInfo == B.ContextDataTableInfo;
+}
+
 const FName UCameraAsset::SharedTransitionsGraphName("SharedTransitions");
 
 void UCameraAsset::SetCameraDirector(UCameraDirector* InCameraDirector)
diff --git a/Engine/Plugins/Cameras/GameplayCameras/Source/GameplayCameras/Private/Core/CameraAssetBuilder.cpp b/Engine/Plugins/Cameras/GameplayCameras/Source/GameplayCameras/Private/Core/CameraAssetBuilder.cpp
--- a/Engine/Plugins/Cameras/GameplayCameras/Source/GameplayCameras/Private/Core/CameraAssetBuilder.cpp
+++ b/Engine/Plugins/Cameras/GameplayCameras/Source/GameplayCameras/Private/Core/CameraAssetBuilder.cpp
@@ -2,11 +2,12 @@
 
 #include "Core/CameraAssetBuilder.h"
 
-#include "Core/CameraDirector.h"
-#include "Core/CameraNodeEvaluatorStorage.h"
 #include "Core/CameraAsset.h"
+#include "Core/CameraDirector.h"
+#include "Core/CameraParameters.h"
 #include "Core/CameraRigAsset.h"
 #include "Core/CameraRigAssetBuilder.h"
+#include "GameplayCamerasDelegates.h"
 #include "Logging/TokenizedMessage.h"
 
 #define LOCTEXT_NAMESPACE "CameraAssetBuilder"
@@ -40,12 +41,15 @@ void FCameraAssetBuilder::BuildCamera(UCameraAsset* InCameraAsset, FCustomBuildS
 	}
 	BuildLog.SetLoggingPrefix(FString());
 	UpdateBuildStatus();
+
+	FGameplayCamerasDelegates::OnCameraAssetBuilt().Broadcast(CameraAsset);
 }
 
 void FCameraAssetBuilder::BuildCameraImpl()
 {
 	TArray<UCameraRigAsset*> CameraRigs;
 
+	// Build the camera director and get the list of camera rigs it references.
 	if (UCameraDirector* CameraDirector = CameraAsset->GetCameraDirector())
 	{
 		CameraDirector->BuildCameraDirector(BuildLog);
@@ -64,11 +68,55 @@ void FCameraAssetBuilder::BuildCameraImpl()
 		BuildLog.AddMessage(EMessageSeverity::Warning, LOCTEXT("MissingRigs", "Camera isn't using any camera rigs."));
 	}
 
+	// Build each of the camera rigs.
 	for (UCameraRigAsset* CameraRig : CameraRigs)
 	{
 		FCameraRigAssetBuilder CameraRigBuilder(BuildLog);
 		CameraRigBuilder.BuildCameraRig(CameraRig);
 	}
+
+	// Get the list of all the camera rigs' interface parameters, and rebuild our
+	// parameters property bag.
+	int32 NextParameterPropertyIndex = 0;
+	TArray<FPropertyBagPropertyDesc> DefaultParameterProperties;
+	TArray<TWeakObjectPtr<const UCameraRigAsset>> ParameterOwners;
+
+	for (const UCameraRigAsset* CameraRig : CameraRigs)
+	{
+		FCameraRigAssetBuilder::AppendDefaultParameters(CameraRig->Interface, DefaultParameterProperties);
+
+		for (int32 Index = NextParameterPropertyIndex; Index < DefaultParameterProperties.Num(); ++Index)
+		{
+			ParameterOwners.Add(CameraRig);
+		}
+		NextParameterPropertyIndex = DefaultParameterProperties.Num();
+	}
+
+	FInstancedPropertyBag DefaultParameters;
+	DefaultParameters.AddProperties(DefaultParameterProperties);
+
+	if (DefaultParameters.GetPropertyBagStruct() != CameraAsset->DefaultParameters.GetPropertyBagStruct() ||
+			ParameterOwners != CameraAsset->ParameterOwners)
+	{
+		CameraAsset->Modify();
+		CameraAsset->DefaultParameters = DefaultParameters;
+		CameraAsset->ParameterOwners = ParameterOwners;
+	}
+
+	// Accumulate all the camera rigs' allocation infos and store that on the asset.
+	FCameraAssetAllocationInfo AllocationInfo;
+
+	for (const UCameraRigAsset* CameraRig : CameraRigs)
+	{
+		AllocationInfo.VariableTableInfo.Combine(CameraRig->AllocationInfo.VariableTableInfo);
+		AllocationInfo.ContextDataTableInfo.Combine(CameraRig->AllocationInfo.ContextDataTableInfo);
+	}
+
+	if (AllocationInfo != CameraAsset->AllocationInfo)
+	{
+		CameraAsset->Modify();
+		CameraAsset->AllocationInfo = AllocationInfo;
+	}
 }
 
 void FCameraAssetBuilder::UpdateBuildStatus()
diff --git a/Engine/Plugins/Cameras/GameplayCameras/Source/GameplayCameras/Private/Core/CameraAssetReference.cpp b/Engine/Plugins/Cameras/GameplayCameras/Source/GameplayCameras/Private/Core/CameraAssetReference.cpp
new file mode 100644
--- /dev/null
+++ b/Engine/Plugins/Cameras/GameplayCameras/Source/GameplayCameras/Private/Core/CameraAssetReference.cpp
@@ -0,0 +1,109 @@
+// Copyright Epic Games, Inc. All Rights Reserved.
+
+#include "Core/CameraAssetReference.h"
+
+#include "Core/CameraAsset.h"
+#include "Core/CameraNodeEvaluator.h"
+#include "Helpers/CameraAssetParameterOverrideEvaluator.h"
+
+#include UE_INLINE_GENERATED_CPP_BY_NAME(CameraAssetReference)
+
+FCameraAssetReference::FCameraAssetReference()
+{
+}
+
+FCameraAssetReference::FCameraAssetReference(UCameraAsset* InCameraAsset)
+	: CameraAsset(InCameraAsset)
+{
+}
+
+void FCameraAssetReference::ApplyParameterOverrides(UE::Cameras::FCameraNodeEvaluationResult& OutResult, bool bDrivenOverridesOnly)
+{
+	using namespace UE::Cameras;
+	FCameraAssetParameterOverrideEvaluator OverrideEvaluator(*this);
+	OverrideEvaluator.ApplyParameterOverrides(OutResult.VariableTable, OutResult.ContextDataTable, bDrivenOverridesOnly);
+}
+
+bool FCameraAssetReference::IsParameterOverriden(const FGuid PropertyID) const
+{
+	return ParameterOverrideGuids.Contains(PropertyID);
+}
+
+void FCameraAssetReference::SetParameterOverriden(const FGuid PropertyID, bool bIsOverridden)
+{
+	if (bIsOverridden)
+	{
+		ParameterOverrideGuids.AddUnique(PropertyID);
+	}
+	else
+	{
+		ParameterOverrideGuids.Remove(PropertyID);
+	}
+}
+
+bool FCameraAssetReference::NeedsRebuildParameters() const
+{
+	if ((!CameraAsset && Parameters.IsValid()) || (CameraAsset && !Parameters.IsValid()))
+	{
+		return true;
+	}
+
+	if (CameraAsset)
+	{
+		const UPropertyBag* AssetParametersType = CameraAsset->GetDefaultParameters().GetPropertyBagStruct();
+		const UPropertyBag* ReferenceParametersType = Parameters.GetPropertyBagStruct();
+		if (AssetParametersType != ReferenceParametersType)
+		{
+			return true;
+		}
+	}
+
+	return false;
+}
+
+bool FCameraAssetReference::RebuildParametersIfNeeded()
+{
+	if (NeedsRebuildParameters())
+	{
+		RebuildParameters();
+		return true;
+	}
+	return false;
+}
+
+void FCameraAssetReference::RebuildParameters()
+{
+	if (CameraAsset)
+	{
+		Parameters.MigrateToNewBagInstanceWithOverrides(CameraAsset->GetDefaultParameters(), ParameterOverrideGuids);
+		
+		// Remove overrides for parameters that don't exist anymore.
+		if (const UPropertyBag* ParametersType = Parameters.GetPropertyBagStruct())
+		{
+			for (TArray<FGuid>::TIterator It = ParameterOverrideGuids.CreateIterator(); It; ++It)
+			{
+				if (!ParametersType->FindPropertyDescByID(*It))
+				{
+					It.RemoveCurrentSwap();
+				}
+			}
+		}
+	}
+	else
+	{
+		Parameters.Reset();
+	}
+}
+
+bool FCameraAssetReference::SerializeFromMismatchedTag(FPropertyTag const& Tag, FStructuredArchive::FSlot Slot)
+{
+	if (Tag.Type == NAME_SoftObjectProperty)
+	{
+		FSoftObjectPtr CameraAssetPath;
+		Slot << CameraAssetPath;
+		CameraAsset = Cast<UCameraAsset>(CameraAssetPath.Get());
+		return true;
+	}
+	return false;
+}
+
diff --git a/Engine/Plugins/Cameras/GameplayCameras/Source/GameplayCameras/Private/Core/CameraContextDataAllocationInfo.cpp b/Engine/Plugins/Cameras/GameplayCameras/Source/GameplayCameras/Private/Core/CameraContextDataAllocationInfo.cpp
new file mode 100644
--- /dev/null
+++ b/Engine/Plugins/Cameras/GameplayCameras/Source/GameplayCameras/Private/Core/CameraContextDataAllocationInfo.cpp
@@ -0,0 +1,54 @@
+// Copyright Epic Games, Inc. All Rights Reserved.
+
+#include "Core/CameraContextDataAllocationInfo.h"
+
+#include "Core/CameraNode.h"
+
+void FCameraContextDataAllocationInfo::Add(const UCameraNode* Owner, FName DataName, ECameraContextDataType DataType, const UObject* DataTypeObject)
+{
+	FName FullDataName(FString::Format(TEXT("{0}_{1}"), { Owner->GetName(), *DataName.ToString() }));
+	FCameraContextDataID DataID = FCameraContextDataID::FromName(FullDataName);
+
+	FCameraContextDataDefinition NewDefinition;
+	NewDefinition.DataID = DataID;
+	NewDefinition.DataType = DataType;
+	NewDefinition.DataTypeObject = DataTypeObject;
+	DataDefinitions.Add(NewDefinition);
+}
+
+void FCameraContextDataAllocationInfo::Combine(const FCameraContextDataAllocationInfo& OtherInfo)
+{
+	TMap<FCameraContextDataID, int32> KnownNames;
+	for (auto It = DataDefinitions.CreateConstIterator(); It; ++It)
+	{
+		const FCameraContextDataDefinition& DataDefinition(*It);
+		KnownNames.Add(DataDefinition.DataID, It.GetIndex());
+	}
+
+	for (const FCameraContextDataDefinition& OtherDataDefinition : OtherInfo.DataDefinitions)
+	{
+		const int32 KnownIndex = KnownNames.FindRef(OtherDataDefinition.DataID, INDEX_NONE);
+		if (KnownIndex == INDEX_NONE)
+		{
+			DataDefinitions.Add(OtherDataDefinition);
+		}
+		else
+		{
+			const FCameraContextDataDefinition& KnownDataDefinition(DataDefinitions[KnownIndex]);
+			ensure(KnownDataDefinition == OtherDataDefinition);
+		}
+	}
+}
+
+bool operator==(const FCameraContextDataDefinition& A, const FCameraContextDataDefinition& B)
+{
+	return A.DataID == B.DataID &&
+		A.DataType == B.DataType &&
+		A.DataTypeObject == B.DataTypeObject;
+}
+
+bool operator==(const FCameraContextDataAllocationInfo& A, const FCameraContextDataAllocationInfo& B)
+{
+	return A.DataDefinitions == B.DataDefinitions;
+}
+
diff --git a/Engine/Plugins/Cameras/GameplayCameras/Source/GameplayCameras/Private/Core/CameraContextDataTable.cpp b/Engine/Plugins/Cameras/GameplayCameras/Source/GameplayCameras/Private/Core/CameraContextDataTable.cpp
new file mode 100644
--- /dev/null
+++ b/Engine/Plugins/Cameras/GameplayCameras/Source/GameplayCameras/Private/Core/CameraContextDataTable.cpp
@@ -0,0 +1,493 @@
+// Copyright Epic Games, Inc. All Rights Reserved.
+
+#include "Core/CameraContextDataTable.h"
+
+#include "Core/CameraContextDataAllocationInfo.h"
+
+namespace UE::Cameras
+{
+
+FCameraContextDataTable::FCameraContextDataTable()
+{
+}
+
+FCameraContextDataTable::~FCameraContextDataTable()
+{
+}
+
+void FCameraContextDataTable::AddReferencedObjects(FReferenceCollector& ReferenceCollector)
+{
+	for (FEntry& Entry : Entries)
+	{
+		ReferenceCollector.AddReferencedObject(Entry.TypeObject);
+
+		uint8* RawData = Memory + Entry.Offset;
+		switch (Entry.Type)
+		{
+			case ECameraContextDataType::Struct:
+				{
+					const UScriptStruct* StructType = CastChecked<const UScriptStruct>(Entry.TypeObject);
+					ReferenceCollector.AddPropertyReferencesWithStructARO(StructType, RawData);
+				}
+				break;
+			case ECameraContextDataType::Object:
+				{
+					TObjectPtr<UObject>* TypedData = reinterpret_cast<TObjectPtr<UObject>*>(RawData);
+					ReferenceCollector.AddReferencedObject(*TypedData);
+				}
+				break;
+			case ECameraContextDataType::Class:
+				{
+					TObjectPtr<UClass>* TypedData = reinterpret_cast<TObjectPtr<UClass>*>(RawData);
+					ReferenceCollector.AddReferencedObject(*TypedData);
+				}
+				break;
+		}
+	}
+}
+
+void FCameraContextDataTable::Initialize(const FCameraContextDataAllocationInfo& AllocationInfo)
+{
+	// Reset any previous state.
+	DestroyBuffer();
+	Entries.Reset();
+	EntryLookup.Reset();
+
+	// Compute the total buffer size we need, and create our entries as we go.
+	const uint32 FirstAlignOf = 32u;
+	uint32 TotalSizeOf = 0;
+	uint32 CurSizeOf, CurAlignOf;
+
+	for (const FCameraContextDataDefinition& DataDefinition : AllocationInfo.DataDefinitions)
+	{
+		GetDataTypeAllocationInfo(DataDefinition.DataType, DataDefinition.DataTypeObject, CurSizeOf, CurAlignOf);
+		const uint32 NewEntryOffset = Align(TotalSizeOf, CurAlignOf);
+		TotalSizeOf = NewEntryOffset + CurSizeOf;
+
+		FEntry NewEntry;
+		NewEntry.ID = DataDefinition.DataID;
+		NewEntry.Type = DataDefinition.DataType;
+		NewEntry.TypeObject = DataDefinition.DataTypeObject;
+		NewEntry.Offset = NewEntryOffset;
+		NewEntry.Flags = EEntryFlags::None;
+#if WITH_EDITORONLY_DATA
+		NewEntry.DebugName = DataDefinition.DataID.DataName;
+#endif
+
+		Entries.Add(NewEntry);
+		EntryLookup.Add(NewEntry.ID, Entries.Num() - 1);
+	}
+
+	// Allocate the memory buffer.
+	Memory = reinterpret_cast<uint8*>(FMemory::Malloc(TotalSizeOf, FirstAlignOf));
+	Capacity = TotalSizeOf;
+	Used = TotalSizeOf;
+
+	// Go back to our entries and initialize each entry to the default value for that data type.
+	for (const FEntry& Entry : Entries)
+	{
+		uint8* DataPtr = Memory + Entry.Offset;
+		InitializeDefaultDataValue(Entry.Type, Entry.TypeObject, DataPtr);
+	}
+}
+
+void FCameraContextDataTable::AddData(const FCameraContextDataDefinition& DataDefinition)
+{
+	if (!ensure(!EntryLookup.Contains(DataDefinition.DataID)))
+	{
+		return;
+	}
+
+	uint32 SizeOf, AlignOf;
+	GetDataTypeAllocationInfo(DataDefinition.DataType, DataDefinition.DataTypeObject, SizeOf, AlignOf);
+
+	uint8* DataPtr = Align(Memory + Used, AlignOf);
+	uint32 NewUsed = (DataPtr + SizeOf) - Memory;
+
+	if (NewUsed > Capacity)
+	{
+		ReallocateBuffer(NewUsed);
+
+		DataPtr = Align(Memory + Used, AlignOf);
+	}
+
+	Used = NewUsed;
+
+	FEntry NewEntry;
+	NewEntry.ID = DataDefinition.DataID;
+	NewEntry.Type = DataDefinition.DataType;
+	NewEntry.TypeObject = DataDefinition.DataTypeObject;
+	NewEntry.Offset = DataPtr - Memory;
+	NewEntry.Flags = EEntryFlags::None;
+#if WITH_EDITORONLY_DATA
+	NewEntry.DebugName = DataDefinition.DataID.DataName;
+#endif
+	
+	Entries.Add(NewEntry);
+	EntryLookup.Add(DataDefinition.DataID, Entries.Num() - 1);
+
+	InitializeDefaultDataValue(NewEntry.Type, NewEntry.TypeObject, Memory + NewEntry.Offset);
+}
+
+bool FCameraContextDataTable::GetDataTypeAllocationInfo(ECameraContextDataType DataType, const UObject* DataTypeObject, uint32& OutSizeOf, uint32& OutAlignOf)
+{
+	switch (DataType)
+	{
+		case ECameraContextDataType::Name:
+			OutSizeOf = sizeof(FName);
+			OutAlignOf = alignof(FName);
+			break;
+		case ECameraContextDataType::String:
+			OutSizeOf = sizeof(FString);
+			OutAlignOf = alignof(FString);
+			break;
+		case ECameraContextDataType::Enum:
+			OutSizeOf = sizeof(uint8);
+			OutAlignOf = alignof(uint8);
+			break;
+		case ECameraContextDataType::Struct:
+			{
+				const UScriptStruct* StructType = CastChecked<const UScriptStruct>(DataTypeObject);
+				if (StructType)
+				{
+					OutSizeOf = StructType->PropertiesSize;
+					OutAlignOf = StructType->MinAlignment;
+				}
+			}
+			break;
+		case ECameraContextDataType::Object:
+			OutSizeOf = sizeof(TObjectPtr<UObject>);
+			OutAlignOf = alignof(TObjectPtr<UObject>);
+			break;
+		case ECameraContextDataType::Class:
+			OutSizeOf = sizeof(TObjectPtr<UClass>);
+			OutAlignOf = alignof(TObjectPtr<UClass>);
+			break;
+		default:
+			ensure(false);
+			return false;
+	}
+	return true;
+}
+
+bool FCameraContextDataTable::InitializeDefaultDataValue(ECameraContextDataType DataType, const UObject* DataTypeObject, uint8* DataPtr)
+{
+	switch (DataType)
+	{
+		case ECameraContextDataType::Name:
+			new (DataPtr) FName();
+			break;
+		case ECameraContextDataType::String:
+			new (DataPtr) FString();
+			break;
+		case ECameraContextDataType::Enum:
+			*reinterpret_cast<uint8*>(DataPtr) = (uint8)CastChecked<const UEnum>(DataTypeObject)->GetValueByIndex(0);
+			break;
+		case ECameraContextDataType::Struct:
+			CastChecked<const UScriptStruct>(DataTypeObject)->InitializeDefaultValue(DataPtr);
+			break;
+		case ECameraContextDataType::Object:
+			new (DataPtr) TObjectPtr<UObject>();
+			break;
+		case ECameraContextDataType::Class:
+			new (DataPtr) TObjectPtr<UClass>();
+			break;
+		default:
+			ensure(false);
+			return false;
+	}
+	return true;
+}
+
+void FCameraContextDataTable::ReallocateBuffer(uint32 MinRequired)
+{
+	static const uint32 DefaultCapacity = 64u;
+	static const uint32 DefaultAlignment = 32u;
+
+	uint32 NewCapacity = Capacity <= 0 ? DefaultCapacity : Capacity * 2;
+	if (MinRequired > 0)
+	{
+		NewCapacity = FMath::Max(NewCapacity, MinRequired);
+	}
+
+	uint8* OldMemory = Memory;
+	uint8* NewMemory = reinterpret_cast<uint8*>(FMemory::Malloc(NewCapacity, DefaultAlignment));
+
+	if (OldMemory)
+	{
+		// UE implements names, strings, and UStructs as bitwise relocatable so we can just
+		// move the memory around.
+		FMemory::Memmove(NewMemory, OldMemory, Capacity);
+		FMemory::Free(OldMemory);
+	}
+
+	Memory = NewMemory;
+	Capacity = NewCapacity;
+}
+
+void FCameraContextDataTable::DestroyBuffer()
+{
+	if (!Memory)
+	{
+		return;
+	}
+
+	for (const FEntry& Entry : Entries)
+	{
+		uint8* DataPtr = Memory + Entry.Offset;
+
+		switch (Entry.Type)
+		{
+			case ECameraContextDataType::Name:
+				reinterpret_cast<FName*>(DataPtr)->~FName();
+				break;
+			case ECameraContextDataType::String:
+				reinterpret_cast<FString*>(DataPtr)->~FString();
+				break;
+			case ECameraContextDataType::Enum:
+				// Nothing to do.
+				break;
+			case ECameraContextDataType::Struct:
+				CastChecked<const UScriptStruct>(Entry.TypeObject)->DestroyStruct(DataPtr);
+				break;
+			case ECameraContextDataType::Object:
+				reinterpret_cast<TObjectPtr<UObject>*>(DataPtr)->~TObjectPtr<UObject>();
+				break;
+			case ECameraContextDataType::Class:
+				reinterpret_cast<TObjectPtr<UClass>*>(DataPtr)->~TObjectPtr<UClass>();
+				break;
+			default:
+				ensure(false);
+				break;
+		}
+	}
+
+	FMemory::Free(Memory);
+
+	Memory = nullptr;
+	Capacity = 0;
+	Used = 0;
+}
+
+void FCameraContextDataTable::SetData(FCameraContextDataID InID, const FStructView& InData)
+{
+	FEntry* Entry = FindEntry(InID);
+	if (ensure(Entry && Entry->Type == ECameraContextDataType::Struct && InData.GetScriptStruct() == Entry->TypeObject))
+	{
+		uint8* DataPtr = Memory + Entry->Offset;
+		const UScriptStruct* StructType = CastChecked<const UScriptStruct>(Entry->TypeObject);
+		StructType->CopyScriptStruct(DataPtr, InData.GetMemory());
+		EnumAddFlags(Entry->Flags, EEntryFlags::Written | EEntryFlags::WrittenThisFrame);
+	}
+}
+
+void FCameraContextDataTable::SetData(FCameraContextDataID InID, const FInstancedStruct& InData)
+{
+	FEntry* Entry = FindEntry(InID);
+	if (ensure(Entry && Entry->Type == ECameraContextDataType::Struct && InData.GetScriptStruct() == Entry->TypeObject))
+	{
+		uint8* DataPtr = Memory + Entry->Offset;
+		const UScriptStruct* StructType = CastChecked<const UScriptStruct>(Entry->TypeObject);
+		StructType->CopyScriptStruct(DataPtr, InData.GetMemory());
+		EnumAddFlags(Entry->Flags, EEntryFlags::Written | EEntryFlags::WrittenThisFrame);
+	}
+}
+
+const FCameraContextDataTable::FEntry* FCameraContextDataTable::FindEntry(FCameraContextDataID InID) const
+{
+	const uint32 EntryIndex = EntryLookup.FindRef(InID, INDEX_NONE);
+	if (ensure(Entries.IsValidIndex(EntryIndex)))
+	{
+		return &Entries[EntryIndex];
+	}
+	return nullptr;
+}
+
+FCameraContextDataTable::FEntry* FCameraContextDataTable::FindEntry(FCameraContextDataID InID)
+{
+	const uint32 EntryIndex = EntryLookup.FindRef(InID, INDEX_NONE);
+	if (ensure(Entries.IsValidIndex(EntryIndex)))
+	{
+		return &Entries[EntryIndex];
+	}
+	return nullptr;
+}
+
+const uint8* FCameraContextDataTable::GetData(FCameraContextDataID DataID, ECameraContextDataType ExpectedDataType, const UObject* ExpectedDataTypeObject) const
+{
+	const uint8* Data = TryGetData(DataID, ExpectedDataType, ExpectedDataTypeObject);
+	ensureMsgf(
+			Data, 
+			TEXT("Can't get camera context data (ID '%s') because it doesn't exist in the table."), 
+			*DataID.DataName.ToString());
+	return Data;
+}
+
+const uint8* FCameraContextDataTable::TryGetData(FCameraContextDataID DataID, ECameraContextDataType ExpectedDataType, const UObject* ExpectedDataTypeObject) const
+{
+	const FEntry* Entry = FindEntry(DataID);
+	if (Entry)
+	{
+		if (ensure(Entry->Type == ExpectedDataType && Entry->TypeObject == ExpectedDataTypeObject))
+		{
+			return Memory + Entry->Offset;
+		}
+	}
+
+	return nullptr;
+}
+
+void FCameraContextDataTable::SetData(FCameraContextDataID DataID, ECameraContextDataType ExpectedDataType, const UObject* ExpectedDataTypeObject, const uint8* InRawDataPtr, bool bMarkAsWrittenThisFrame)
+{
+	const bool bDidSet = TrySetData(DataID, ExpectedDataType, ExpectedDataTypeObject, InRawDataPtr, bMarkAsWrittenThisFrame);
+	ensureMsgf(bDidSet, TEXT("Can't set camera context data (ID '%s') beacuse it doesn't exist in the table."), *DataID.DataName.ToString());
+}
+
+bool FCameraContextDataTable::TrySetData(FCameraContextDataID DataID, ECameraContextDataType ExpectedDataType, const UObject* ExpectedDataTypeObject, const uint8* InRawDataPtr, bool bMarkAsWrittenThisFrame)
+{
+	FEntry* Entry = FindEntry(DataID);
+	if (!Entry)
+	{
+		return false;
+	}
+
+	if (!ensure(ExpectedDataType == Entry->Type && ExpectedDataTypeObject == Entry->TypeObject))
+	{
+		return false;
+	}
+
+	uint8* DataPtr = Memory + Entry->Offset;
+	SetDataValue(Entry->Type, Entry->TypeObject, DataPtr, InRawDataPtr);
+
+	Entry->Flags |= EEntryFlags::Written;
+	if (bMarkAsWrittenThisFrame)
+	{
+		Entry->Flags |= EEntryFlags::WrittenThisFrame;
+	}
+	
+	return true;
+}
+
+bool FCameraContextDataTable::SetDataValue(ECameraContextDataType DataType, const UObject* DataTypeObject, uint8* DestDataPtr, const uint8* SrcDataPtr)
+{
+	switch (DataType)
+	{
+		case ECameraContextDataType::Name:
+			*reinterpret_cast<FName*>(DestDataPtr) = *reinterpret_cast<const FName*>(SrcDataPtr);
+			break;
+		case ECameraContextDataType::String:
+			*reinterpret_cast<FString*>(DestDataPtr) = *reinterpret_cast<const FString*>(SrcDataPtr);
+			break;
+		case ECameraContextDataType::Enum:
+			*reinterpret_cast<uint8*>(DestDataPtr) = *reinterpret_cast<const uint8*>(SrcDataPtr);
+			break;
+		case ECameraContextDataType::Struct:
+			CastChecked<const UScriptStruct>(DataTypeObject)->CopyScriptStruct(DestDataPtr, SrcDataPtr);
+			break;
+		case ECameraContextDataType::Object:
+			*reinterpret_cast<TObjectPtr<UObject>*>(DestDataPtr) = *reinterpret_cast<const TObjectPtr<UObject>*>(SrcDataPtr);
+			break;
+		case ECameraContextDataType::Class:
+			*reinterpret_cast<TObjectPtr<UClass>*>(DestDataPtr) = *reinterpret_cast<const TObjectPtr<UClass>*>(SrcDataPtr);
+			break;
+		default:
+			ensure(false);
+			return false;
+	}
+	return true;
+}
+
+bool FCameraContextDataTable::IsValueWritten(FCameraContextDataID InID) const
+{
+	if (const FEntry* Entry = FindEntry(InID))
+	{
+		return EnumHasAnyFlags(Entry->Flags, EEntryFlags::Written);
+	}
+	return false;
+}
+
+void FCameraContextDataTable::UnsetValue(FCameraContextDataID InID)
+{
+	if (FEntry* Entry = FindEntry(InID))
+	{
+		return EnumRemoveFlags(Entry->Flags, EEntryFlags::Written);
+	}
+}
+
+void FCameraContextDataTable::UnsetAllValues()
+{
+	for (FEntry& Entry : Entries)
+	{
+		EnumRemoveFlags(Entry.Flags, EEntryFlags::Written);
+	}
+}
+
+bool FCameraContextDataTable::IsValueWrittenThisFrame(FCameraContextDataID InID) const
+{
+	if (const FEntry* Entry = FindEntry(InID))
+	{
+		return EnumHasAnyFlags(Entry->Flags, EEntryFlags::WrittenThisFrame);
+	}
+	return false;
+}
+
+void FCameraContextDataTable::ClearAllWrittenThisFrameFlags()
+{
+	for (FEntry& Entry : Entries)
+	{
+		EnumRemoveFlags(Entry.Flags, EEntryFlags::WrittenThisFrame);
+	}
+}
+
+void FCameraContextDataTable::OverrideAll(const FCameraContextDataTable& OtherTable)
+{
+	Override(OtherTable, false);
+}
+
+void FCameraContextDataTable::OverrideKnown(const FCameraContextDataTable& OtherTable)
+{
+	Override(OtherTable, true);
+}
+
+void FCameraContextDataTable::Override(const FCameraContextDataTable& OtherTable, bool bKnownOnly)
+{
+	for (const FEntry& OtherEntry : OtherTable.Entries)
+	{
+		const EEntryFlags OtherFlags = OtherEntry.Flags;
+		if (EnumHasAnyFlags(OtherFlags, EEntryFlags::Written))
+		{
+			int32 ThisIndex = EntryLookup.FindRef(OtherEntry.ID, INDEX_NONE);
+			if (ThisIndex == INDEX_NONE)
+			{
+				if (bKnownOnly)
+				{
+					continue;
+				}
+
+				FCameraContextDataDefinition OtherEntryDefinition;
+				OtherEntryDefinition.DataID = OtherEntry.ID;
+				OtherEntryDefinition.DataType = OtherEntry.Type;
+				OtherEntryDefinition.DataTypeObject = OtherEntry.TypeObject;
+				AddData(OtherEntryDefinition);
+				ThisIndex = Entries.Num() - 1;
+			}
+
+			ensure(ThisIndex != INDEX_NONE);
+			
+			FEntry& ThisEntry = Entries[ThisIndex];
+			if (!ensure(ThisEntry.Type == OtherEntry.Type && ThisEntry.TypeObject == OtherEntry.TypeObject))
+			{
+				continue;
+			}
+
+			uint8* ThisDataPtr = Memory + ThisEntry.Offset;
+			uint8* OtherDataPtr = OtherTable.Memory + OtherEntry.Offset;
+			SetDataValue(ThisEntry.Type, ThisEntry.TypeObject, ThisDataPtr, OtherDataPtr);
+
+			EnumAddFlags(ThisEntry.Flags, EEntryFlags::Written | (OtherFlags & EEntryFlags::WrittenThisFrame));
+		}
+	}
+}
+
+}  // namespace UE::Cameras
+
diff --git a/Engine/Plugins/Cameras/GameplayCameras/Source/GameplayCameras/Private/Core/CameraNodeEvaluator.cpp b/Engine/Plugins/Cameras/GameplayCameras/Source/GameplayCameras/Private/Core/CameraNodeEvaluator.cpp
--- a/Engine/Plugins/Cameras/GameplayCameras/Source/GameplayCameras/Private/Core/CameraNodeEvaluator.cpp
+++ b/Engine/Plugins/Cameras/GameplayCameras/Source/GameplayCameras/Private/Core/CameraNodeEvaluator.cpp
@@ -71,6 +71,11 @@ void FCameraNodeEvaluationResult::Serialize(FArchive& Ar)
 	Ar << bIsValid;
 }
 
+void FCameraNodeEvaluationResult::AddReferencedObjects(FReferenceCollector& Collector)
+{
+	ContextDataTable.AddReferencedObjects(Collector);
+}
+
 FCameraNodeEvaluator* FCameraNodeEvaluatorBuildParams::BuildEvaluator(const UCameraNode* InNode) const
 {
 	if (InNode)
diff --git a/Engine/Plugins/Cameras/GameplayCameras/Source/GameplayCameras/Private/Core/CameraRigAsset.cpp b/Engine/Plugins/Cameras/GameplayCameras/Source/GameplayCameras/Private/Core/CameraRigAsset.cpp
--- a/Engine/Plugins/Cameras/GameplayCameras/Source/GameplayCameras/Private/Core/CameraRigAsset.cpp
+++ b/Engine/Plugins/Cameras/GameplayCameras/Source/GameplayCameras/Private/Core/CameraRigAsset.cpp
@@ -2,10 +2,10 @@
 
 #include "Core/CameraRigAsset.h"
 
-#include "Core/CameraAsset.h"
 #include "Core/CameraBuildLog.h"
 #include "Core/CameraNode.h"
 #include "Core/CameraRigAssetBuilder.h"
+#include "Core/CameraVariableAssets.h"
 #include "UObject/ObjectSaveContext.h"
 
 #include UE_INLINE_GENERATED_CPP_BY_NAME(CameraRigAsset)
@@ -19,23 +19,27 @@ void FCameraRigAllocationInfo::Append(const FCameraRigAllocationInfo& OtherAlloc
 	const FCameraVariableTableAllocationInfo& OtherVariableTableInfo(OtherAllocationInfo.VariableTableInfo);
 	VariableTableInfo.AutoResetVariables.Append(OtherVariableTableInfo.AutoResetVariables);
 	VariableTableInfo.VariableDefinitions.Append(OtherVariableTableInfo.VariableDefinitions);
+
+	const FCameraContextDataAllocationInfo& OtherContextDataTableInfo(OtherAllocationInfo.ContextDataTableInfo);
+	ContextDataTableInfo.DataDefinitions.Append(OtherContextDataTableInfo.DataDefinitions);
 }
 
 bool operator==(const FCameraRigAllocationInfo& A, const FCameraRigAllocationInfo& B)
 {
 	return A.EvaluatorInfo == B.EvaluatorInfo
-		&& A.VariableTableInfo == B.VariableTableInfo;
+		&& A.VariableTableInfo == B.VariableTableInfo
+		&& A.ContextDataTableInfo == B.ContextDataTableInfo;
 }
 
 #if WITH_EDITOR
 
-void UCameraRigInterfaceParameter::GetGraphNodePosition(FName InGraphName, int32& NodePosX, int32& NodePosY) const
+void UCameraRigInterfaceParameterBase::GetGraphNodePosition(FName InGraphName, int32& NodePosX, int32& NodePosY) const
 {
 	NodePosX = GraphNodePos.X;
 	NodePosY = GraphNodePos.Y;
 }
 
-void UCameraRigInterfaceParameter::OnGraphNodeMoved(FName InGraphName, int32 NodePosX, int32 NodePosY, bool bMarkDirty)
+void UCameraRigInterfaceParameterBase::OnGraphNodeMoved(FName InGraphName, int32 NodePosX, int32 NodePosY, bool bMarkDirty)
 {
 	Modify(bMarkDirty);
 
@@ -45,7 +49,7 @@ void UCameraRigInterfaceParameter::OnGraphNodeMoved(FName InGraphName, int32 Nod
 
 #endif
 
-void UCameraRigInterfaceParameter::PostLoad()
+void UCameraRigInterfaceParameterBase::PostLoad()
 {
 	if (!Guid.IsValid())
 	{
@@ -55,7 +59,7 @@ void UCameraRigInterfaceParameter::PostLoad()
 	Super::PostLoad();
 }
 
-void UCameraRigInterfaceParameter::PostInitProperties()
+void UCameraRigInterfaceParameterBase::PostInitProperties()
 {
 	Super::PostInitProperties();
 
@@ -66,7 +70,7 @@ void UCameraRigInterfaceParameter::PostInitProperties()
 	}
 }
 
-void UCameraRigInterfaceParameter::PostDuplicate(EDuplicateMode::Type DuplicateMode)
+void UCameraRigInterfaceParameterBase::PostDuplicate(EDuplicateMode::Type DuplicateMode)
 {
 	Super::PostDuplicate(DuplicateMode);
 
@@ -76,29 +80,39 @@ void UCameraRigInterfaceParameter::PostDuplicate(EDuplicateMode::Type DuplicateM
 	}
 }
 
-UCameraRigInterfaceParameter* FCameraRigInterface::FindInterfaceParameterByName(const FString& ParameterName) const
+UCameraRigBlendableParameter* FCameraRigInterface::FindBlendableParameterByName(const FString& ParameterName) const
 {
-	const TObjectPtr<UCameraRigInterfaceParameter>* FoundItem = InterfaceParameters.FindByPredicate(
-			[&ParameterName](UCameraRigInterfaceParameter* Item)
+	const TObjectPtr<UCameraRigBlendableParameter>* FoundItem = BlendableParameters.FindByPredicate(
+			[&ParameterName](UCameraRigBlendableParameter* Item)
 			{
 				return Item->InterfaceParameterName == ParameterName;
 			});
 	return FoundItem ? *FoundItem : nullptr;
 }
 
-UCameraRigInterfaceParameter* FCameraRigInterface::FindInterfaceParameterByGuid(const FGuid& ParameterGuid) const
+UCameraRigBlendableParameter* FCameraRigInterface::FindBlendableParameterByGuid(const FGuid& ParameterGuid) const
 {
-	const TObjectPtr<UCameraRigInterfaceParameter>* FoundItem = InterfaceParameters.FindByPredicate(
-			[&ParameterGuid](UCameraRigInterfaceParameter* Item)
+	const TObjectPtr<UCameraRigBlendableParameter>* FoundItem = BlendableParameters.FindByPredicate(
+			[&ParameterGuid](UCameraRigBlendableParameter* Item)
 			{
-				return Item->Guid == ParameterGuid;
+				return Item->GetGuid() == ParameterGuid;
 			});
 	return FoundItem ? *FoundItem : nullptr;
 }
 
-bool FCameraRigInterface::HasInterfaceParameter(const FString& ParameterName) const
+UCameraRigDataParameter* FCameraRigInterface::FindDataParameterByGuid(const FGuid& ParameterGuid) const
 {
-	return FindInterfaceParameterByName(ParameterName) != nullptr;
+	const TObjectPtr<UCameraRigDataParameter>* FoundItem = DataParameters.FindByPredicate(
+			[&ParameterGuid](UCameraRigDataParameter* Item)
+			{
+				return Item->GetGuid() == ParameterGuid;
+			});
+	return FoundItem ? *FoundItem : nullptr;
+}
+
+bool FCameraRigInterface::HasBlendableParameter(const FString& ParameterName) const
+{
+	return FindBlendableParameterByName(ParameterName) != nullptr;
 }
 
 const FName UCameraRigAsset::NodeTreeGraphName(TEXT("NodeTree"));
@@ -116,6 +130,25 @@ void UCameraRigAsset::PostLoad()
 		GraphNodePosY_DEPRECATED = 0;
 	}
 
+	// Any interface parameters found in the list of graph objects should be removed and instead
+	// flagged as having a node. This is because the way of handling parameter graph nodes has
+	// changed.
+	for (auto It = AllNodeTreeObjects.CreateIterator(); It; ++It)
+	{
+		UObject* Item(*It);
+		if (!Item)
+		{
+			It.RemoveCurrent();
+			continue;
+		}
+		if (UCameraRigInterfaceParameterBase* InterfaceParameter = Cast<UCameraRigInterfaceParameterBase>(Item))
+		{
+			InterfaceParameter->bHasGraphNode = true;
+			It.RemoveCurrent();
+			continue;
+		}
+	}
+
 #endif
 
 	if (!Guid.IsValid())
@@ -123,6 +156,19 @@ void UCameraRigAsset::PostLoad()
 		Guid = FGuid::NewGuid();
 	}
 
+	// Initialize the ParameterType of blendable parameters that pre-existed the new interface data.
+	// The type defaults to Boolean so check only those.
+	for (UCameraRigBlendableParameter* BlendableParameter : Interface.BlendableParameters)
+	{
+		if (BlendableParameter->ParameterType == ECameraVariableType::Boolean)
+		{
+			if (BlendableParameter->PrivateVariable)
+			{
+				BlendableParameter->ParameterType = BlendableParameter->PrivateVariable->GetVariableType();
+			}
+		}
+	}
+
 	Super::PostLoad();
 }
 
