Cameras: big pass on parameterization of camera rigs

diff --git a/Engine/Plugins/Cameras/GameplayCameras/Source/GameplayCameras/Private/Core/CameraRigAssetBuilder.cpp b/Engine/Plugins/Cameras/GameplayCameras/Source/GameplayCameras/Private/Core/CameraRigAssetBuilder.cpp
--- a/Engine/Plugins/Cameras/GameplayCameras/Source/GameplayCameras/Private/Core/CameraRigAssetBuilder.cpp
+++ b/Engine/Plugins/Cameras/GameplayCameras/Source/GameplayCameras/Private/Core/CameraRigAssetBuilder.cpp
@@ -9,9 +9,10 @@
 #include "Core/CameraRigBuildContext.h"
 #include "Core/CameraVariableAssets.h"
 #include "Core/CameraVariableReferences.h"
+#include "Core/ICustomCameraNodeParameterProvider.h"
 #include "GameplayCamerasDelegates.h"
 #include "Logging/TokenizedMessage.h"
-#include "Nodes/Common/CameraRigCameraNode.h"
+#include "UObject/Object.h"
 
 #define LOCTEXT_NAMESPACE "CameraRigAssetBuilder"
 
@@ -64,11 +65,46 @@ void SetPrivateVariableDefaultValue<UBooleanCameraVariable, bool>(UBooleanCamera
 	}
 }
 
-struct FPrivateVariableBuilder
+UCameraVariableAsset* CreatePrivateVariable(
+		UCameraRigAsset* CameraRig,
+		const FString& InterfaceParameterName,
+		ECameraVariableType ParameterType)
+{
+	const FString VariableName = FString::Format(
+			TEXT("Override_{0}_{1}"), 
+			{ CameraRig->GetName(), InterfaceParameterName });
+
+	TSubclassOf<UCameraVariableAsset> VariableClass;
+	switch (ParameterType)
+	{
+#define UE_CAMERA_VARIABLE_FOR_TYPE(VariableType, VariableName)\
+		case ECameraVariableType::VariableName:\
+			VariableClass = U##VariableName##CameraVariable::StaticClass();\
+			break;
+		UE_CAMERA_VARIABLE_FOR_ALL_TYPES()
+#undef UE_CAMERA_VARIABLE_FOR_TYPE
+	}
+	if (!ensure(VariableClass))
+	{
+		return nullptr;
+	}
+
+	UCameraVariableAsset* PrivateVariable = NewObject<UCameraVariableAsset>(
+			CameraRig, VariableClass, FName(*VariableName), RF_Transactional);
+
+	// Make sure it's a private input variable.
+	PrivateVariable->bIsInput = true;
+	PrivateVariable->bIsPrivate = true;
+	PrivateVariable->bAutoReset = false;
+
+	return PrivateVariable;
+}
+
+struct FInterfaceParameterBindingBuilder
 {
 	UCameraRigAsset* CameraRig;
 
-	FPrivateVariableBuilder(FCameraRigAssetBuilder& InOwner)
+	FInterfaceParameterBindingBuilder(FCameraRigAssetBuilder& InOwner)
 		: Owner(InOwner)
 	{
 		CameraRig = Owner.CameraRig;
@@ -84,206 +120,273 @@ struct FPrivateVariableBuilder
 		Owner.BuildLog.AddMessage(EMessageSeverity::Error, MoveTemp(ErrorMessage));
 	}
 
-	template<typename ExpectedVariableAssetType>
-	ExpectedVariableAssetType* FindReusablePrivateVariable(FStructProperty* ForParameterProperty, UCameraNode* ForCameraNode)
+	template<typename CameraParameterType>
+	void SetCameraParameterOverride(
+			const UCameraRigBlendableParameter* BlendableParameter,
+			FStructProperty* TargetProperty,
+			CameraParameterType* CameraParameterPtr)
 	{
-		using FDrivenParameterKey = FCameraRigAssetBuilder::FDrivenParameterKey;
+		using VariableAssetType = typename CameraParameterType::VariableAssetType;
+		using ValueType = typename VariableAssetType::ValueType;
+
+		ensure(BlendableParameter->PrivateVariable);
+		ensure(BlendableParameter->TargetPropertyName == TargetProperty->GetFName());
 
-		FDrivenParameterKey ParameterKey{ ForParameterProperty, ForCameraNode };
-		UCameraVariableAsset** FoundItem = Owner.OldDrivenParameters.Find(ParameterKey);
-		if (FoundItem)
+		VariableAssetType* NewVariable = EnsureCompatibleVariable<VariableAssetType>(BlendableParameter);
+		const bool bIsValid = CheckTwiceDrivenParameter(BlendableParameter, CameraParameterPtr);
+		if (!NewVariable || !bIsValid)
 		{
-			// We found an existing variable that was driving this camera node's property.
-			// Re-use it and remove it from the re-use pool.
-			UCameraVariableAsset* ReusedVariable = (*FoundItem);
-			Owner.OldDrivenParameters.Remove(ParameterKey);
-			return CastChecked<ExpectedVariableAssetType>(ReusedVariable);
+			return;
 		}
 
-		return nullptr;
+		// Find the variable that was previously driving this parameter. If it has changed, flag the
+		// camera node as having been modified.
+		UCameraNode* TargetNode = BlendableParameter->Target;
+		VariableAssetType* PreviousVariable = FindOldDrivingVariable<VariableAssetType>(TargetProperty->GetFName(), TargetNode);
+		if (PreviousVariable != NewVariable)
+		{
+			TargetNode->Modify();
+		}
+
+		SetPrivateVariableDefaultValue<VariableAssetType, ValueType>(NewVariable, CameraParameterPtr->Value);
+
+		CameraParameterPtr->Variable = NewVariable;
 	}
 
-	template<
-		typename ParameterOverrideType, 
-		typename ExpectedVariableAssetType = typename ParameterOverrideType::CameraParameterType::VariableAssetType>
-	ExpectedVariableAssetType* FindReusablePrivateVariable(ParameterOverrideType* ForParameterOverride, UCameraRigCameraNode* ForCameraRigNode)
+	template<typename VariableReferenceType>
+	void SetVariableReferenceOverride(
+			const UCameraRigBlendableParameter* BlendableParameter,
+			FStructProperty* TargetProperty,
+			VariableReferenceType* VariableReferencePtr)
 	{
-		using FDrivenOverrideKey = FCameraRigAssetBuilder::FDrivenOverrideKey;
+		using VariableAssetType = typename VariableReferenceType::VariableAssetType;
+		using ValueType = typename VariableAssetType::ValueType;
+
+		ensure(BlendableParameter->PrivateVariable);
+		ensure(BlendableParameter->TargetPropertyName == TargetProperty->GetFName());
 
-		FDrivenOverrideKey OverrideKey{ ForParameterOverride->InterfaceParameterGuid, ForCameraRigNode };
-		UCameraVariableAsset** FoundItem = Owner.OldDrivenOverrides.Find(OverrideKey);
-		if (FoundItem)
+		VariableAssetType* NewVariable = EnsureCompatibleVariable<VariableAssetType>(BlendableParameter);
+		const bool bIsValid = CheckTwiceDrivenParameter(BlendableParameter, VariableReferencePtr);
+		if (!NewVariable || !bIsValid)
 		{
-			UCameraVariableAsset* ReusedVariable = (*FoundItem);
-			Owner.OldDrivenOverrides.Remove(OverrideKey);
-			// Don't do a checked cast here because interface parameters can change type if they are
-			// reconnected to a different type of parameter. If the cast fail, we will retur null and
-			// the variable won't get reused, which is what we want.
-			return Cast<ExpectedVariableAssetType>(ReusedVariable);
+			return;
 		}
 
-		return nullptr;
+		// Find the variable that was previously driving this reference. If it has changed, flag the
+		// camera node as having been modified.
+		UCameraNode* TargetNode = BlendableParameter->Target;
+		VariableAssetType* PreviousVariable = FindOldDrivingVariable<VariableAssetType>(TargetProperty->GetFName(), TargetNode);
+		if (PreviousVariable != NewVariable)
+		{
+			TargetNode->Modify();
+		}
+
+		// No default value to set on the driving variable.
+
+		VariableReferencePtr->Variable = NewVariable;
 	}
 
-	bool ReuseInterfaceParameter(UCameraRigInterfaceParameter* InterfaceParameter, UCameraVariableAsset* IntendedPrivateVariable)
+	template<typename VariableAssetType>
+	void SetCustomBlendableParameterOverride(
+			const UCameraRigBlendableParameter* BlendableParameter,
+			const FCustomCameraNodeParameterInfos::FBlendableParameterInfo& CustomParameter)
 	{
-		using FReusableInterfaceParameterInfo = FCameraRigAssetBuilder::FReusableInterfaceParameterInfo;
-		FReusableInterfaceParameterInfo* FoundItem = Owner.OldInterfaceParameters.Find(InterfaceParameter);
-		if (ensure(FoundItem))
+		using ValueType = typename VariableAssetType::ValueType;
+
+		ensure(BlendableParameter->PrivateVariable);
+		ensure(BlendableParameter->TargetPropertyName == CustomParameter.ParameterName);
+
+		VariableAssetType* NewVariable = EnsureCompatibleVariable<VariableAssetType>(BlendableParameter);
+		const bool bIsValid = CheckTwiceDrivenParameter(BlendableParameter, *CustomParameter.OverrideVariable);
+		if (!NewVariable || !bIsValid)
 		{
-			// This is an interface parameter that existed before. Flag things as modified if the 
-			// private variable is changing.
-			ensure(!FoundItem->Value);
-			FoundItem->Value = true;  // This one has now been re-used.
-			return FoundItem->Key != IntendedPrivateVariable;
+			return;
 		}
-		// We should have had this interface parameter in our map, since we built it just a second ago!
-		// Something's wrong... oh well, flag things as modified.
-		return true;
-	}
 
-private:
+		// Find the variable that was previously driving this override. If it has changed, flag the
+		// camera node as having been modified.
+		UCameraNode* TargetNode = BlendableParameter->Target;
+		VariableAssetType* PreviousVariable = FindOldDrivingVariable<VariableAssetType>(CustomParameter.ParameterName, TargetNode);
+		if (PreviousVariable != NewVariable)
+		{
+			TargetNode->Modify();
+		}
 
-	FCameraRigAssetBuilder& Owner;
-};
+		const ValueType* DefaultValue = reinterpret_cast<const ValueType*>(CustomParameter.DefaultValuePtr);
+		SetPrivateVariableDefaultValue<VariableAssetType, ValueType>(NewVariable, *DefaultValue);
 
-template<typename CameraParameterType>
-void CheckTwiceDrivenParameter(
-		FPrivateVariableBuilder& Builder, 
-		UCameraRigInterfaceParameter* InterfaceParameter, 
-		CameraParameterType* CameraParameter)
-{
-	if (CameraParameter->Variable != nullptr)
-	{
-		// We should have cleared all exposed parameters in GatherOldDrivenParameters, so the only variables
-		// left on camera parameters should be user-defined ones.
-		UObject* VariableOuter = CameraParameter->Variable->GetOuter();
-		if (ensureMsgf(
-					VariableOuter != Builder.CameraRig, 
-					TEXT("Unexpected driving variable found: all exposed parameters should have been cleared before rebuilding.")))
-		{
-			// If this parameter is driven by a user-defined variable, emit an error, and replace that 
-			// driving variable with our private variable.
-			Builder.ReportError(
-					InterfaceParameter->Target,
-					FText::Format(
-						LOCTEXT(
-							"CameraParameterDrivenTwice", 
-							"Camera node parameter '{0}.{1}' is both exposed and driven by a variable!"),
-						FText::FromName(InterfaceParameter->Target->GetFName()), 
-						FText::FromName(InterfaceParameter->TargetPropertyName)));
-		}
+		(*CustomParameter.OverrideVariable) = NewVariable;
 	}
-}
-
-template<typename VariableAssetType>
-VariableAssetType* MakeOrRenamePrivateVariable(
-		FPrivateVariableBuilder& Builder,
-		const FString& InterfaceParameterName,
-		VariableAssetType* PrivateVariable)
-{
-	const FString VariableName = FString::Format(
-			TEXT("Override_{0}_{1}"), 
-			{ Builder.CameraRig->GetName(), InterfaceParameterName });
 
-	if (PrivateVariable)
+	void SetDataContextPropertyOverride(
+			const UCameraRigDataParameter* DataParameter,
+			FProperty* TargetProperty,
+			FCameraContextDataID* OverrideDataID)
 	{
-		// We have a pre-existing variable! Make sure it's got the right name, in case the exposed rig parameter
-		// was renamed. Keeping a good name is mostly to help with debugging.
-		FString OriginalName = PrivateVariable->GetName();
-		OriginalName.RemoveFromStart("REUSABLE_", ESearchCase::CaseSensitive);
-		if (OriginalName != VariableName)
+		ensure(DataParameter->PrivateDataID);
+
+		const bool bIsValid = CheckTwiceDrivenParameter(DataParameter, *OverrideDataID);
+		if (!bIsValid)
 		{
-			PrivateVariable->Modify();
+			return;
 		}
-		// Rename non-transactionally because we might be simply setting the variable's name back to what it
-		// always was. We don't want to dirty the package for no-op builds.
-		PrivateVariable->Rename(*VariableName, nullptr, REN_NonTransactional);
+
+		UCameraNode* TargetNode = DataParameter->Target;
+		FCameraContextDataID PreviousDataID = FindOldDrivingDataID(TargetProperty->GetFName(), TargetNode);
+		if (PreviousDataID != DataParameter->PrivateDataID)
+		{
+			TargetNode->Modify();
+		}
+
+		(*OverrideDataID) = DataParameter->PrivateDataID;
 	}
-	else
+
+	void SetCustomDataParameterOverride(
+			const UCameraRigDataParameter* DataParameter,
+			const FCustomCameraNodeParameterInfos::FDataParameterInfo& CustomParameter)
 	{
-		// Make a new variable.
-		PrivateVariable = NewObject<VariableAssetType>(Builder.CameraRig, FName(*VariableName), RF_Transactional);
+		ensure(DataParameter->PrivateDataID);
+		ensure(DataParameter->TargetPropertyName == CustomParameter.ParameterName);
+
+		const bool bIsValid = CheckTwiceDrivenParameter(DataParameter, *CustomParameter.OverrideDataID);
+		if (!bIsValid)
+		{
+			return;
+		}
+
+		UCameraNode* TargetNode = DataParameter->Target;
+		FCameraContextDataID PreviousDataID = FindOldDrivingDataID(CustomParameter.ParameterName, TargetNode);
+		if (PreviousDataID != DataParameter->PrivateDataID)
+		{
+			TargetNode->Modify();
+		}
+
+		(*CustomParameter.OverrideDataID) = DataParameter->PrivateDataID;
 	}
 
-	// Make sure it's a private input variable.
-	PrivateVariable->bIsInput = true;
-	PrivateVariable->bIsPrivate = true;
-	PrivateVariable->bAutoReset = false;
+private:
 
-	return PrivateVariable;
-}
+	template<typename ExpectedVariableAssetType>
+	ExpectedVariableAssetType* FindOldDrivingVariable(FName ForParameterName, UCameraNode* ForCameraNode)
+	{
+		using FDrivenParameterKey = FCameraRigAssetBuilder::FDrivenParameterKey;
 
-template<typename CameraParameterType, typename VariableAssetType = typename CameraParameterType::VariableAssetType>
-void DoSetupPrivateVariable(
-		FPrivateVariableBuilder& Builder,
-		UCameraRigInterfaceParameter* InterfaceParameter,
-		CameraParameterType* CameraParameter,
-		VariableAssetType* ReusedVariable
-		)
-{
-	using ValueType = typename CameraParameterType::ValueType;
+		FDrivenParameterKey ParameterKey{ ForParameterName, ForCameraNode };
+		TObjectPtr<UCameraVariableAsset> ReusedVariable;
+		Owner.OldDrivenBlendableParameters.RemoveAndCopyValue(ParameterKey, ReusedVariable);
+		return CastChecked<ExpectedVariableAssetType>(ReusedVariable, ECastCheckedType::NullAllowed);
+	}
 
-	const bool bIsReusedVariable = (ReusedVariable != nullptr);
+	FCameraContextDataID FindOldDrivingDataID(FName ForParameterName, UCameraNode* ForCameraNode)
+	{
+		using FDrivenParameterKey = FCameraRigAssetBuilder::FDrivenParameterKey;
 
-	// Either rename the camera variable we are re-using, or make a new one with the right name.
-	VariableAssetType* PrivateVariable = MakeOrRenamePrivateVariable(
-			Builder, InterfaceParameter->InterfaceParameterName, ReusedVariable);
-	ensure(PrivateVariable->GetOuter() == Builder.CameraRig);
+		FDrivenParameterKey ParameterKey{ ForParameterName, ForCameraNode };
+		FCameraContextDataID ReusedDataID;
+		Owner.OldDrivenDataParameters.RemoveAndCopyValue(ParameterKey, ReusedDataID);
+		return ReusedDataID;
+	}
 
-	// Set the default value of the variable to be the value in the camera parameter.
-	SetPrivateVariableDefaultValue<VariableAssetType, ValueType>(PrivateVariable, CameraParameter->Value);
+	template<typename CameraParameterType>
+	bool CheckTwiceDrivenParameter(
+			const UCameraRigBlendableParameter* BlendableParameter, 
+			CameraParameterType* CameraParameter)
+	{
+		using VariableAssetType = typename CameraParameterType::VariableAssetType;
+		return CheckTwiceDrivenParameter<VariableAssetType>(BlendableParameter, CameraParameter->Variable);
+	}
 
-	// Set the variable on both the interface parameter and the camera node. Flag them as modified
-	// if we actually changed anything.
-	const bool bShouldModifyInterfaceParameter = Builder.ReuseInterfaceParameter(InterfaceParameter, PrivateVariable);
-	if (bShouldModifyInterfaceParameter)
+	template<typename VariableAssetType>
+	bool CheckTwiceDrivenParameter(
+			const UCameraRigBlendableParameter* BlendableParameter, 
+			TObjectPtr<VariableAssetType> CameraParameterVariable)
 	{
-		InterfaceParameter->Modify();
+		if (CameraParameterVariable != nullptr)
+		{
+			// We should have cleared all exposed parameters in GatherOldDrivenParameters, so the only variables
+			// left on camera parameters should be user-defined ones.
+			UObject* VariableOuter = CameraParameterVariable->GetOuter();
+			if (ensureMsgf(
+						VariableOuter != CameraRig, 
+						TEXT("Unexpected driving variable found: all exposed parameters should have been cleared before rebuilding.")))
+			{
+				// If this parameter is driven by a user-defined variable, emit an error.
+				ReportError(BlendableParameter->Target,
+						FText::Format(
+							LOCTEXT(
+								"BlendableParameterDrivenTwice", 
+								"Camera node parameter '{0}.{1}' is both exposed and driven by a variable!"),
+							FText::FromName(BlendableParameter->Target->GetFName()), 
+							FText::FromName(BlendableParameter->TargetPropertyName)));
+				return false;
+			}
+		}
+		return true;
 	}
-	if (!bIsReusedVariable)
+
+	bool CheckTwiceDrivenParameter(
+			const UCameraRigDataParameter* DataParameter,
+			FCameraContextDataID DataID)
 	{
-		InterfaceParameter->Target->Modify();
+		if (DataID.IsValid())
+		{
+			ReportError(DataParameter->Target,
+					FText::Format(
+						LOCTEXT(
+							"DataParameterDrivenTwice",
+							"Camera node parameter '{0}.{1}' is somehow overriden twice!"),
+						FText::FromName(DataParameter->Target->GetFName()), 
+						FText::FromName(DataParameter->TargetPropertyName)));
+			return false;
+		}
+		return true;
 	}
-	InterfaceParameter->PrivateVariable = PrivateVariable;
-	CameraParameter->Variable = PrivateVariable;
-}
 
-template<typename CameraParameterType>
-void SetupPrivateVariable(
-		FPrivateVariableBuilder& Builder, 
-		UCameraRigInterfaceParameter* InterfaceParameter, 
-		FStructProperty* ParameterTargetProperty,
-		CameraParameterType* CameraParameter)
-{
-	using ValueType = typename CameraParameterType::ValueType;
-	using VariableAssetType = typename CameraParameterType::VariableAssetType;
+	template<typename VariableAssetType>
+	VariableAssetType* EnsureCompatibleVariable(const UCameraRigBlendableParameter* BlendableParameter)
+	{
+		if (!ensure(BlendableParameter->PrivateVariable))
+		{
+			return nullptr;
+		}
 
-	CheckTwiceDrivenParameter(Builder, InterfaceParameter, CameraParameter);
+		VariableAssetType* Variable = Cast<VariableAssetType>(BlendableParameter->PrivateVariable);
 
-	VariableAssetType* ReusedVariable = Builder.FindReusablePrivateVariable<VariableAssetType>(
-			ParameterTargetProperty, InterfaceParameter->Target);
+		if (!Variable)
+		{
+			const UClass* VariableAssetClass = VariableAssetType::StaticClass();
+			const UCameraVariableAsset* DefaultObject = VariableAssetClass->template GetDefaultObject<UCameraVariableAsset>();
 
-	DoSetupPrivateVariable(Builder, InterfaceParameter, CameraParameter, ReusedVariable);
-}
+			ReportError(BlendableParameter->Target,
+				FText::Format(LOCTEXT(
+					"IncompatibleCameraNodeParameter",
+					"Invalid interface parameter '{0}', driving property '{1}' on '{2}': expected {3} camera parameter but was {4}"),
+					FText::FromString(BlendableParameter->InterfaceParameterName),
+					FText::FromName(BlendableParameter->TargetPropertyName),
+					FText::FromName(BlendableParameter->Target->GetFName()),
+					FText::FromName(UEnum::GetValueAsName(DefaultObject->GetVariableType())),
+					FText::FromName(UEnum::GetValueAsName(BlendableParameter->ParameterType))
+					)
+				);
+			return nullptr;
+		}
 
-template<typename ParameterOverrideType>
-void SetupPrivateVariable(
-		FPrivateVariableBuilder& Builder, 
-		UCameraRigInterfaceParameter* InterfaceParameter, 
-		ParameterOverrideType* ParameterOverride)
-{
-	using CameraParameterType = typename ParameterOverrideType::CameraParameterType;
-	using ValueType = typename CameraParameterType::ValueType;
-	using VariableAssetType = typename CameraParameterType::VariableAssetType;
+		return Variable;
+	}
 
-	CheckTwiceDrivenParameter(Builder, InterfaceParameter, &ParameterOverride->Value);
+private:
 
-	UCameraRigCameraNode* CameraRigNode = CastChecked<UCameraRigCameraNode>(InterfaceParameter->Target);
-	VariableAssetType* ReusedVariable = Builder.FindReusablePrivateVariable<ParameterOverrideType>(
-			ParameterOverride, CameraRigNode);
+	FCameraRigAssetBuilder& Owner;
+};
 
-	DoSetupPrivateVariable(Builder, InterfaceParameter, &ParameterOverride->Value, ReusedVariable);
+void TrashPrivateVariable(UCameraVariableAsset* VariableToTrash)
+{
+	TStringBuilder<256> StringBuilder;
+	StringBuilder.Append("TRASH_");
+	StringBuilder.Append(VariableToTrash->GetName());
+
+	VariableToTrash->Modify();
+	VariableToTrash->Rename(StringBuilder.ToString());
 }
 
 void AddCameraVariableToAllocationInfo(UCameraVariableAsset* Variable, FCameraVariableTableAllocationInfo& AllocationInfo)
@@ -299,6 +402,18 @@ void AddCameraVariableToAllocationInfo(UCameraVariableAsset* Variable, FCameraVa
 	}
 }
 
+void AddContextDataToAllocationInfo(FCameraContextDataID DataID, ECameraContextDataType DataType, const UObject* DataTypeObject, FCameraContextDataAllocationInfo& AllocationInfo)
+{
+	if (DataID)
+	{
+		FCameraContextDataDefinition DataDefinition;
+		DataDefinition.DataID = DataID;
+		DataDefinition.DataType = DataType;
+		DataDefinition.DataTypeObject = DataTypeObject;
+		AllocationInfo.DataDefinitions.Add(DataDefinition);
+	}
+}
+
 }  // namespace Internal
 
 FCameraRigAssetBuilder::FCameraRigAssetBuilder(FCameraBuildLog& InBuildLog)
@@ -326,11 +441,11 @@ void FCameraRigAssetBuilder::BuildCameraRig(UCameraRigAsset* InCameraRig, FCusto
 		InCustomBuildStep.ExecuteIfBound(CameraRig, BuildLog);
 
 		CameraRig->EventHandlers.Notify(&ICameraRigAssetEventHandler::OnCameraRigBuilt, CameraRig);
-
-		FGameplayCamerasDelegates::OnCameraRigAssetBuilt().Broadcast(CameraRig, BuildLog);
 	}
 	BuildLog.SetLoggingPrefix(FString());
 	UpdateBuildStatus();
+
+	FGameplayCamerasDelegates::OnCameraRigAssetBuilt().Broadcast(CameraRig);
 }
 
 void FCameraRigAssetBuilder::BuildCameraRigImpl()
@@ -348,10 +463,13 @@ void FCameraRigAssetBuilder::BuildCameraRigImpl()
 	CallPreBuild();
 
 	GatherOldDrivenParameters();
-	BuildNewDrivenParameters();
-	DiscardUnusedPrivateVariables();
+	BuildInterfaceParameters();
+	BuildInterfaceParameterBindings();
+	DiscardUnusedParameters();
 
-	BuildAllocationInfo();
+	CallBuild();
+
+	BuildDefaultParameters();
 }
 
 void FCameraRigAssetBuilder::BuildCameraNodeHierarchy()
@@ -388,49 +506,20 @@ void FCameraRigAssetBuilder::CallPreBuild()
 
 void FCameraRigAssetBuilder::GatherOldDrivenParameters()
 {
-	// Keep track of what camera parameters were previously driven by private variables,
-	// and then clear those variables. This is because it's easier to rebuild this from
-	// a blank slate than trying to figure out what changed.
+	// Keep track of which camera parameters and variable references were previously driven by 
+	// private variables, and then clear those variables. This is because it's easier to rebuild 
+	// all this from a blank slate than trying to figure out what changed.
 	//
-	// As we rebuild things in BuildNewDrivenParameters, we compare to the old state to
+	// As we rebuild things in BuildInterfaceParameterBindings, we compare to the old state to
 	// figure out if we need to flag anything as modified for the current transaction.
 	//
 	// Note that parameters driven by user-defined variables are left alone.
 
-	TSet<UCameraVariableAsset*> GatheredVariables;
+	// First, get the list of nodes, both connected and disconnected from the root hierarchy.
+	// We could use AllNodeTreeObjects for that, but it only exists in editor builds, and we 
+	// don't want to rely on unit tests or runtime data manipulation to have correctly populated 
+	// it, so we'll try to gather any stray nodes by looking at objects outer'ed to the camera rig.
 	TSet<UCameraNode*> CameraNodesToGather(CameraNodeHierarchy.GetFlattenedHierarchy());
-
-	// Start by going through all interface parameters, remembering what private variable 
-	// they were associated with originally. Also collect that private variable to be
-	// renamed and put in the re-use pool.
-
-	OldInterfaceParameters.Reset();
-
-	for (UCameraRigInterfaceParameter* InterfaceParameter : CameraRig->Interface.InterfaceParameters)
-	{
-		OldInterfaceParameters.Add(InterfaceParameter, FReusableInterfaceParameterInfo(InterfaceParameter->PrivateVariable, false));
-		if (InterfaceParameter->PrivateVariable)
-		{
-			GatheredVariables.Add(InterfaceParameter->PrivateVariable);
-		}
-
-		ensureMsgf(
-				InterfaceParameter->Target == nullptr || InterfaceParameter->IsInOuter(CameraRig),
-				TEXT("Interface parameter '%s' points to camera node '%s' which isn't outer'ed to camera rig '%s'."),
-				*InterfaceParameter->InterfaceParameterName,
-				*GetNameSafe(InterfaceParameter->Target),
-				*GetPathNameSafe(CameraRig));
-	}
-
-	// Next go through all the camera nodes we know of. Nodes in CameraNodeHierarchy are the ones
-	// connected to the camera rig's root node, so we are missing nodes that were disconnected
-	// since the last build. We could use AllNodeTreeObjects for that, but it only exists in 
-	// editor builds, and we don't want to rely on unit tests or runtime data manipulation to 
-	// have correctly populated it, so we'll try to gather any stray nodes by looking at
-	// objects outer'ed to the camera rig.
-
-	OldDrivenParameters.Reset();
-
 	ForEachObjectWithOuter(CameraRig, [&CameraNodesToGather](UObject* Obj)
 			{
 				if (UCameraNode* CameraNode = Cast<UCameraNode>(Obj))
@@ -445,6 +534,10 @@ void FCameraRigAssetBuilder::GatherOldDrivenParameters()
 				NumStrayCameraNodes, *GetPathNameSafe(CameraRig));
 	}
 
+	// Second, gather all driven properties, and null them out.
+	OldDrivenBlendableParameters.Reset();
+	OldDrivenDataParameters.Reset();
+
 	for (UCameraNode* CameraNode : CameraNodesToGather)
 	{
 		UClass* CameraNodeClass = CameraNode->GetClass();
@@ -466,203 +559,313 @@ void FCameraRigAssetBuilder::GatherOldDrivenParameters()
 					UObject* VariableOuter = CameraParameterPtr->Variable->GetOuter();\
 					if (VariableOuter == CameraRig)\
 					{\
-						OldDrivenParameters.Add(\
-								FDrivenParameterKey{ StructProperty, CameraNode },\
-								CameraParameterPtr->Variable);\
-						GatheredVariables.Add(CameraParameterPtr->Variable);\
+						OldDrivenBlendableParameters.Add(\
+								FDrivenParameterKey{ StructProperty->GetFName(), CameraNode },\
+								*reinterpret_cast<TObjectPtr<UCameraVariableAsset>*>(\
+									reinterpret_cast<FObjectPtr*>(&CameraParameterPtr->Variable)));\
 						CameraParameterPtr->Variable = nullptr;\
 					}\
 				}\
 			}\
+			else if (StructProperty->Struct == F##ValueName##CameraVariableReference::StaticStruct())\
+			{\
+				auto* VariableReferencePtr = StructProperty->ContainerPtrToValuePtr<F##ValueName##CameraVariableReference>(CameraNode);\
+				if (VariableReferencePtr->Variable)\
+				{\
+					UObject* VariableOuter = VariableReferencePtr->Variable->GetOuter();\
+					if (VariableOuter == CameraRig)\
+					{\
+						OldDrivenBlendableParameters.Add(\
+								FDrivenParameterKey{ StructProperty->GetFName(), CameraNode },\
+								*reinterpret_cast<TObjectPtr<UCameraVariableAsset>*>(\
+									reinterpret_cast<FObjectPtr*>(&VariableReferencePtr->Variable)));\
+						VariableReferencePtr->Variable = nullptr;\
+					}\
+				}\
+			}\
 			else
 UE_CAMERA_VARIABLE_FOR_ALL_TYPES()
 #undef UE_CAMERA_VARIABLE_FOR_TYPE
 			{
-				// Some other struct property.
+				// Other struct type...
 			}
 		}
 
-		if (UCameraRigCameraNode* CameraRigNode = Cast<UCameraRigCameraNode>(CameraNode))
+		if (ICustomCameraNodeParameterProvider* CustomParameterProvider = Cast<ICustomCameraNodeParameterProvider>(CameraNode))
 		{
-			FCameraRigParameterOverrides& ParameterOverrides = CameraRigNode->CameraRigReference.GetParameterOverrides();
+			FCustomCameraNodeParameterInfos CustomParameters;
+			CustomParameterProvider->GetCustomCameraNodeParameters(CustomParameters);
 
-#define UE_CAMERA_VARIABLE_FOR_TYPE(ValueType, ValueName)\
-			for (F##ValueName##CameraRigParameterOverride& ParameterOverride : ParameterOverrides.Get##ValueName##Overrides())\
-			{\
-				if (ParameterOverride.Value.Variable)\
-				{\
-					UObject* VariableOuter = ParameterOverride.Value.Variable->GetOuter();\
-					if (VariableOuter == CameraRig)\
-					{\
-						OldDrivenOverrides.Add(\
-								FDrivenOverrideKey{ ParameterOverride.InterfaceParameterGuid, CameraRigNode },\
-								ParameterOverride.Value.Variable);\
-						GatheredVariables.Add(ParameterOverride.Value.Variable);\
-						ParameterOverride.Value.Variable = nullptr;\
-					}\
-				}\
+			for (const FCustomCameraNodeParameterInfos::FBlendableParameterInfo& BlendableParameter : CustomParameters.BlendableParameters)
+			{
+				if (!ensure(BlendableParameter.OverrideVariable))
+				{
+					continue;
+				}
+
+				if (UCameraVariableAsset* Variable = BlendableParameter.OverrideVariable->Get())
+				{
+					UObject* VariableOuter = Variable->GetOuter();
+					if (VariableOuter == CameraRig)
+					{
+						OldDrivenBlendableParameters.Add(
+								FDrivenParameterKey{ BlendableParameter.ParameterName, CameraNode },
+								*BlendableParameter.OverrideVariable);
+						(*BlendableParameter.OverrideVariable) = nullptr;
+					}
+				}
 			}
-UE_CAMERA_VARIABLE_FOR_ALL_TYPES()
-#undef UE_CAMERA_VARIABLE_FOR_TYPE
-		}
-	}
 
-	// Sanity check: see if we have any stray camera variables, possibly introduced by incorrect
-	// editor code or dynamic data manipulation.
-	const int32 PreviousNumGatheredVariables = GatheredVariables.Num();
-	ForEachObjectWithOuter(CameraRig, [&GatheredVariables](UObject* Obj)
+			for (const FCustomCameraNodeParameterInfos::FDataParameterInfo& DataParameter : CustomParameters.DataParameters)
 			{
-				if (UCameraVariableAsset* CameraVariable = Cast<UCameraVariableAsset>(Obj))
+				if (!ensure(DataParameter.OverrideDataID))
 				{
-					GatheredVariables.Add(CameraVariable);
+					continue;
 				}
-			});
-	if (GatheredVariables.Num() > PreviousNumGatheredVariables)
+
+				if (DataParameter.OverrideDataID->IsValid())
+				{
+					OldDrivenDataParameters.Add(
+							FDrivenParameterKey{ DataParameter.ParameterName, CameraNode },
+							*DataParameter.OverrideDataID);
+					(*DataParameter.OverrideDataID) = FCameraContextDataID();
+				}
+			}
+		}
+	}
+}
+
+void FCameraRigAssetBuilder::BuildInterfaceParameters()
+{
+	using namespace Internal;
+
+	for (auto It = CameraRig->Interface.BlendableParameters.CreateIterator(); It; ++It)
 	{
-		UE_LOG(LogCameraSystem, Verbose, TEXT("Collected %d stray camera variables while building camera rig '%s'."),
-				(GatheredVariables.Num() - PreviousNumGatheredVariables), *GetPathNameSafe(CameraRig));
+		UCameraRigBlendableParameter* BlendableParameter(*It);
+
+		// Basic validations.
+		if (!BlendableParameter)
+		{
+			BuildLog.AddMessage(EMessageSeverity::Warning,
+					CameraRig,
+					LOCTEXT("InvalidBlendableParameter", "Invalid interface parameter was found and removed."));
+
+			CameraRig->Modify();
+			It.RemoveCurrent();
+
+			continue;
+		}
+
+		if (BlendableParameter->InterfaceParameterName.IsEmpty())
+		{
+			BuildLog.AddMessage(EMessageSeverity::Error,
+					BlendableParameter,
+					LOCTEXT(
+						"InvalidBlendableParameterName",
+						"Invalid interface parameter name."));
+			continue;
+		}
+
+		// Create a new private variable for this interface parameter if it doesn't have on yet (e.g. it's a
+		// newly-created parameter), or if its type has changed.
+		UCameraVariableAsset* PrivateVariable = BlendableParameter->PrivateVariable;
+		if (PrivateVariable && PrivateVariable->GetVariableType() == BlendableParameter->ParameterType)
+		{
+			continue;
+		}
+
+		if (PrivateVariable)
+		{
+			// Trash it so its name is available for the new private variable.
+			TrashPrivateVariable(PrivateVariable);
+
+			// We supposedly have removed all references to private variables in GatherOldDrivenParameters,
+			// so once we null-out the reference on the interface parameter, this private variable should
+			// not be used anymore and should get GC'ed soon.
+			BlendableParameter->Modify();
+			BlendableParameter->PrivateVariable = nullptr;
+		}
+
+		BlendableParameter->PrivateVariable = CreatePrivateVariable(
+				CameraRig, BlendableParameter->InterfaceParameterName, BlendableParameter->ParameterType);
 	}
 
-	// Temporarily rename all old camera variables, so their names are available to the new
-	// driven parameters.
-	for (UCameraVariableAsset* GatheredVariable : GatheredVariables)
+	for (auto It = CameraRig->Interface.DataParameters.CreateIterator(); It; ++It)
 	{
-		TStringBuilder<256> StringBuilder;
-		StringBuilder.Append("REUSABLE_");
-		StringBuilder.Append(GatheredVariable->GetName());
-		// Rename non-transactionally because if nothing has changed, we will rename it back
-		// later and we don't want to dirty the package for nothing.
-		GatheredVariable->Rename(StringBuilder.ToString(), nullptr, REN_NonTransactional);
+		UCameraRigDataParameter* DataParameter(*It);
+
+		// Basic validations.
+		if (!DataParameter)
+		{
+			BuildLog.AddMessage(EMessageSeverity::Warning,
+					CameraRig,
+					LOCTEXT("InvalidDataParameter", "Invalid interface parameter was found and removed."));
+
+			CameraRig->Modify();
+			It.RemoveCurrent();
+
+			continue;
+		}
+
+		if (DataParameter->InterfaceParameterName.IsEmpty())
+		{
+			BuildLog.AddMessage(EMessageSeverity::Error,
+					DataParameter,
+					LOCTEXT(
+						"InvalidDataParameterName",
+						"Invalid interface parameter name."));
+			continue;
+		}
+
+		// Create the data ID for this parameter. Flag the parameter as changed if the ID
+		// is different, such as if it was renamed.
+		const FString DataIDName = FString::Format(
+				TEXT("Override_{0}_{1}"), 
+				{ CameraRig->GetName(), DataParameter->InterfaceParameterName });
+
+		FCameraContextDataID DataID = FCameraContextDataID::FromName(FName(DataIDName));
+		if (DataParameter->PrivateDataID != DataID)
+		{
+			DataParameter->Modify();
+			DataParameter->PrivateDataID = DataID;
+		}
 	}
 }
 
-void FCameraRigAssetBuilder::BuildNewDrivenParameters()
+void FCameraRigAssetBuilder::BuildInterfaceParameterBindings()
 {
-	TSet<FString> UsedInterfaceParameterNames;
-
 	using FBuiltDrivenParameter = TTuple<UCameraNode*, FName>;
 	TSet<FBuiltDrivenParameter> BuiltDrivenParameters;
 
 	const FString CameraRigName = CameraRig->GetName();
 	const FString CameraRigPathName = CameraRig->GetPathName();
 
-	// Look at the new interface parameters and setup the driven camera node parameters with
-	// private camera variables. We have gathered the old ones previously so we can re-use them,
-	// instead of creating new variable assets each time.
-	//
-	// Additionally, we need to handle camera rig nodes with special code, for the case of an
-	// interface parameter driving a camera rig override (which in turn drives the inner rig's
-	// interface parameter, and so on). This is basically for multi-level interface parameters
-	// overrides.
-	for (UCameraRigInterfaceParameter* InterfaceParameter : CameraRig->Interface.InterfaceParameters)
+	for (const UCameraRigBlendableParameter* BlendableParameter : CameraRig->Interface.BlendableParameters)
 	{
-		// Do some basic validation.
-		if (!InterfaceParameter)
-		{
-			BuildLog.AddMessage(EMessageSeverity::Error,
-					CameraRig,
-					LOCTEXT("InvalidInterfaceParameter", "Invalid interface parameter or target."));
-			continue;
-		}
-		if (!InterfaceParameter->Target)
+		// Basic validations.
+		if (!BlendableParameter->Target)
 		{
 			BuildLog.AddMessage(EMessageSeverity::Warning,
-					InterfaceParameter,
+					BlendableParameter,
 					LOCTEXT(
-						"DisconnectedInterfaceParameter", 
-						"Interface parameter isn't connected: setting overrides for it will not do anything."));
+						"InvalidBlendableParameterTarget",
+						"Invalid interface parameter: it has no target node."));
 			continue;
 		}
-		if (InterfaceParameter->TargetPropertyName.IsNone())
+		if (BlendableParameter->TargetPropertyName.IsNone())
 		{
 			BuildLog.AddMessage(EMessageSeverity::Error,
-					InterfaceParameter,
+					BlendableParameter,
 					LOCTEXT(
-						"InvalidInterfaceParameterTargetPropertyName", 
-						"Invalid interface parameter target property name."));
+						"InvalidBlendableParameterTargetPropertyName", 
+						"Invalid interface parameter: it has not target property name."));
 			continue;
 		}
-		if (InterfaceParameter->InterfaceParameterName.IsEmpty())
+		if (BlendableParameter->InterfaceParameterName.IsEmpty())
 		{
-			BuildLog.AddMessage(EMessageSeverity::Error,
-					InterfaceParameter,
-					LOCTEXT(
-						"InvalidInterfaceParameterName",
-						"Invalid interface parameter name."));
 			continue;
 		}
 
-		// Check duplicate parameter names.
-		if (UsedInterfaceParameterNames.Contains(InterfaceParameter->InterfaceParameterName))
+		// Check duplicate bindings.
+		FBuiltDrivenParameter NewDrivenParameter(BlendableParameter->Target, BlendableParameter->TargetPropertyName);
+		if (BuiltDrivenParameters.Contains(NewDrivenParameter))
 		{
 			BuildLog.AddMessage(EMessageSeverity::Error,
-					InterfaceParameter,
 					FText::Format(LOCTEXT(
-						"InterfaceParameterNameCollision",
-						"Multiple interface parameters named '{0}'. Ignoring duplicates."),
-						FText::FromString(InterfaceParameter->InterfaceParameterName)));
+						"BlendableParameterTargetCollision",
+						"Multiple interface parameters targeting property '{0}' on camera node '{1}'. Ignoring duplicates."),
+						FText::FromName(BlendableParameter->TargetPropertyName),
+						FText::FromName(BlendableParameter->Target->GetFName())));
 			continue;
 		}
-		UsedInterfaceParameterNames.Add(InterfaceParameter->InterfaceParameterName);
+		BuiltDrivenParameters.Add(NewDrivenParameter);
 
-		// Check duplicate targets.
-		FBuiltDrivenParameter BuiltDrivenParameter(InterfaceParameter->Target, InterfaceParameter->TargetPropertyName);
-		if (BuiltDrivenParameters.Contains(BuiltDrivenParameter))
+		// See if this interface parameter is overriding a camera node parameter.
+		// Otherwise, maybe it's targeting a camera rig node's override for an inner rig interface parameter.
+		if (SetupCameraParameterOrVariableReferenceOverride(BlendableParameter))
+		{
+			// Implicit continue.
+		}
+		else if (SetupCustomBlendableParameterOverride(BlendableParameter))
+		{
+			// Implicit continue.
+		}
+		else
 		{
+			UCameraNode* Target = BlendableParameter->Target;
 			BuildLog.AddMessage(EMessageSeverity::Error,
-					InterfaceParameter,
+					Target,
 					FText::Format(LOCTEXT(
-						"InterfaceParameterTargetCollision",
-						"Multiple interface parameters targeting property '{0}' on camera node '{1}'. Ignoring duplicates."),
-						FText::FromName(InterfaceParameter->Target->GetFName()),
-						FText::FromName(InterfaceParameter->TargetPropertyName)));
+						"InvalidBlendableParameterTargetProperty",
+						"Invalid interface parameter '{0}', driving property '{1}' on '{2}', but no such property found."),
+						FText::FromString(BlendableParameter->InterfaceParameterName), 
+						FText::FromName(BlendableParameter->TargetPropertyName),
+						FText::FromName(Target->GetFName())));
+		}
+	}
+
+	for (const UCameraRigDataParameter* DataParameter : CameraRig->Interface.DataParameters)
+	{
+		// Basic validations.
+		if (!DataParameter->Target)
+		{
+			BuildLog.AddMessage(EMessageSeverity::Warning,
+					DataParameter,
+					LOCTEXT(
+						"InvalidDataParameterTarget",
+						"Invalid interface parameter: it has no target node."));
+			continue;
+		}
+		if (DataParameter->TargetPropertyName.IsNone())
+		{
+			BuildLog.AddMessage(EMessageSeverity::Error,
+					DataParameter,
+					LOCTEXT(
+						"InvalidDataParameterTargetPropertyName", 
+						"Invalid interface parameter: it has not target property name."));
+			continue;
+		}
+		if (DataParameter->InterfaceParameterName.IsEmpty())
+		{
 			continue;
 		}
-		BuiltDrivenParameters.Add(BuiltDrivenParameter);
 
-		// See if this interface parameter is overriding a camera node parameter.
-		// Otherwise, maybe it's targeting a camera rig node's override for an inner rig interface parameter.
-		if (SetupCameraParameterOverride(InterfaceParameter))
+		if (SetupDataContextPropertyOverride(DataParameter))
 		{
 			// Implicit continue.
 		}
-		else if (SetupInnerCameraRigParameterOverride(InterfaceParameter))
+		else if (SetupCustomDataParameterOverride(DataParameter))
 		{
 			// Implicit continue.
 		}
 		else
 		{
-			UCameraNode* Target = InterfaceParameter->Target;
+			UCameraNode* Target = DataParameter->Target;
 			BuildLog.AddMessage(EMessageSeverity::Error,
 					Target,
 					FText::Format(LOCTEXT(
-						"InvalidInterfaceParameterTargetProperty",
+						"InvalidDataParameterTargetProperty",
 						"Invalid interface parameter '{0}', driving property '{1}' on '{2}', but no such property found."),
-						FText::FromString(InterfaceParameter->InterfaceParameterName), 
-						FText::FromName(InterfaceParameter->TargetPropertyName),
+						FText::FromString(DataParameter->InterfaceParameterName), 
+						FText::FromName(DataParameter->TargetPropertyName),
 						FText::FromName(Target->GetFName())));
 		}
 	}
 }
 
-bool FCameraRigAssetBuilder::SetupCameraParameterOverride(UCameraRigInterfaceParameter* InterfaceParameter)
+bool FCameraRigAssetBuilder::SetupCameraParameterOrVariableReferenceOverride(const UCameraRigBlendableParameter* BlendableParameter)
 {
 	using namespace Internal;
 
 	// Here we hook up interface parameters connected to a camera node property. This property is supposed
 	// to be of one of the camera parameter types (FBooleanCameraParameter, FInteger32CameraParameter, etc.)
-	// so they have both a fixed value (bool, int32, etc.) and a "private variable" which is a reference to 
-	// a corresponding camera variable asset (UBooleanCameraVariable, UInteger32CameraVariable, etc.) which
-	// has been set to "private".
-	//
-	// So the goal of this method is to create a private variable and set it on both the interface parameter
-	// and the camera node property. This way, if someone wants to override the value of that interface
-	// parameter, they set the value of the variable defined on it. It will then drive the value of the
-	// corresponding camera node property.
+	// so they have both a fixed value (bool, int32, etc.) and a "private variable" which we will set to the 
+	// private variable of the given interface parameter, checking that the types match (UBooleanCameraVariable, 
+	// UInteger32CameraVariable, etc.)
 
-	UCameraNode* Target = InterfaceParameter->Target;
+	UCameraNode* Target = BlendableParameter->Target;
 	UClass* TargetClass = Target->GetClass();
-	FProperty* TargetProperty = TargetClass->FindPropertyByName(InterfaceParameter->TargetPropertyName);
+	FProperty* TargetProperty = TargetClass->FindPropertyByName(BlendableParameter->TargetPropertyName);
 	if (!TargetProperty)
 	{
 		// No match, try something else.
@@ -675,105 +878,83 @@ bool FCameraRigAssetBuilder::SetupCameraParameterOverride(UCameraRigInterfacePar
 		BuildLog.AddMessage(EMessageSeverity::Error,
 				Target,
 				FText::Format(LOCTEXT(
-						"InvalidCameraNodeProperty",
+						"InvalidCameraNodeParameter",
 						"Invalid interface parameter '{0}', driving property '{1}' on '{2}', but it's not a camera parameter."),
-					FText::FromString(InterfaceParameter->InterfaceParameterName), 
-					FText::FromName(InterfaceParameter->TargetPropertyName),
+					FText::FromString(BlendableParameter->InterfaceParameterName), 
+					FText::FromName(BlendableParameter->TargetPropertyName),
 					FText::FromName(Target->GetFName())));
 		return true;
 	}
 
 	// Get the type of the camera parameter by matching the struct against all the types we support,
 	// and create a private camera variable asset to drive its value.
-	FPrivateVariableBuilder PrivateVariableBuilder(*this);
+	FInterfaceParameterBindingBuilder Builder(*this);
 #define UE_CAMERA_VARIABLE_FOR_TYPE(ValueType, ValueName)\
 	if (TargetStructProperty->Struct == F##ValueName##CameraParameter::StaticStruct())\
 	{\
 		auto* CameraParameterPtr = TargetStructProperty->ContainerPtrToValuePtr<F##ValueName##CameraParameter>(Target);\
-		SetupPrivateVariable(PrivateVariableBuilder, InterfaceParameter, TargetStructProperty, CameraParameterPtr);\
+		Builder.SetCameraParameterOverride<F##ValueName##CameraParameter>(\
+				BlendableParameter, TargetStructProperty, CameraParameterPtr\
+				);\
+	}\
+	else if (TargetStructProperty->Struct == F##ValueName##CameraVariableReference::StaticStruct())\
+	{\
+		auto* VariableReferencePtr = TargetStructProperty->ContainerPtrToValuePtr<F##ValueName##CameraVariableReference>(Target);\
+		Builder.SetVariableReferenceOverride<F##ValueName##CameraVariableReference>(\
+				BlendableParameter, TargetStructProperty, VariableReferencePtr\
+				);\
 	}\
 	else
 	UE_CAMERA_VARIABLE_FOR_ALL_TYPES()
 #undef UE_CAMERA_VARIABLE_FOR_TYPE
 	{
 		BuildLog.AddMessage(EMessageSeverity::Error,
-				InterfaceParameter,
+				BlendableParameter,
 				FText::Format(LOCTEXT(
-						"InvalidCameraNodeProperty",
+						"InvalidCameraNodeParameter",
 						"Invalid interface parameter '{0}', driving property '{1}' on '{2}', but it's not a camera parameter."),
-					FText::FromString(InterfaceParameter->InterfaceParameterName), 
-					FText::FromName(InterfaceParameter->TargetPropertyName),
+					FText::FromString(BlendableParameter->InterfaceParameterName), 
+					FText::FromName(BlendableParameter->TargetPropertyName),
 					FText::FromName(Target->GetFName())));
 	}
 
 	return true;
 }
 
-bool FCameraRigAssetBuilder::SetupInnerCameraRigParameterOverride(UCameraRigInterfaceParameter* InterfaceParameter)
+bool FCameraRigAssetBuilder::SetupCustomBlendableParameterOverride(const UCameraRigBlendableParameter* BlendableParameter)
 {
 	using namespace Internal;
 
-	// Here we hook up interface parameters connected specifically to a camera rig node (aka "prefab node").
-	// Unlike other camera nodes, the camera rig node doesn't have "actual" properties on it, in the sense 
-	// of UClass and FProperty (although we could in the future generate a UClass, like Blueprints do).
-	// Camera rig nodes expose the same properties as their inner camera rig, i.e. they "forward expose" the
-	// interface parameters defined on their inner camera rig.
-	//
-	// So the goal of this method is to handle multi-level exposed parameters. That is: we are exposing the
-	// interface parameter of an inner camera rig as one of our own interface parameter. Just like in the
-	// previous method (see above) we create a private camera variable to set on the interface parameter,
-	// but instead of also setting it on a camera node property, here we set it on an override entry on
-	// the camera rig node's list of overrides.
-	//
-	// Note that this camera rig node may or may not have an existing override. If the user forwards the
-	// parameter without changing its default value, there would not be an existing override and we have
-	// to create our own. If there is an existing override, we set the private variable on it and it
-	// will use the user-defined new override value when the variable isn't set.
-
-	UCameraRigCameraNode* Target = Cast<UCameraRigCameraNode>(InterfaceParameter->Target);
+	ICustomCameraNodeParameterProvider* Target = Cast<ICustomCameraNodeParameterProvider>(BlendableParameter->Target);
 	if (!Target)
 	{
 		// No match, try something else.
 		return false;
 	}
 
-	// Look for an interface parameter matching the target name.
-	UCameraRigAsset* InnerCameraRig = Target->CameraRigReference.GetCameraRig();
-	if (!InnerCameraRig)
-	{
-		return false;
-	}
-	UCameraRigInterfaceParameter* InnerInterfaceParameter = InnerCameraRig->Interface.FindInterfaceParameterByName(
-			InterfaceParameter->TargetPropertyName.ToString());
-	if (!InnerInterfaceParameter)
-	{
-		return false;
-	}
+	// Look for a parameter override matching the target name.
+	// TODO: we're querying the list of custom parameters every time, we may want to cache it for this phase.
+	FCustomCameraNodeParameterInfos CustomParameters;
+	Target->GetCustomCameraNodeParameters(CustomParameters);
 
-	// Found it! Check that the inner camera rig was built.
-	if (InnerInterfaceParameter->PrivateVariable == nullptr)
+	FCustomCameraNodeParameterInfos::FBlendableParameterInfo* TargetCustomParameter = 
+		CustomParameters.BlendableParameters.FindByPredicate(
+			[BlendableParameter](FCustomCameraNodeParameterInfos::FBlendableParameterInfo& CustomParameter)
+			{
+				return CustomParameter.ParameterName == BlendableParameter->TargetPropertyName;
+			});
+	if (!TargetCustomParameter)
 	{
-		BuildLog.AddMessage(EMessageSeverity::Error, Target,
-				FText::Format(
-					LOCTEXT("UnbuiltInnerCameraRig",
-						"Can't expose inner camera rig parameter '{0}': the inner camera rig '{1}' failed to build."),
-					FText::FromName(InterfaceParameter->TargetPropertyName),
-					FText::FromString(GetPathNameSafe(InnerCameraRig))));
-		return true;
+		// No match, try something else.
+		return false;
 	}
 
-	// Look for an override that matches the given interface parameter. Create one if we don't find any.
-	FCameraRigParameterOverrides& ParameterOverrides = Target->CameraRigReference.GetParameterOverrides();
-	switch (InnerInterfaceParameter->PrivateVariable->GetVariableType())
+	FInterfaceParameterBindingBuilder Builder(*this);
+	switch (TargetCustomParameter->ParameterType)
 	{
 #define UE_CAMERA_VARIABLE_FOR_TYPE(ValueType, ValueName)\
 		case ECameraVariableType::ValueName:\
-			{\
-				F##ValueName##CameraRigParameterOverride& ValueName##Override = \
-						ParameterOverrides.FindOrAddParameterOverride<F##ValueName##CameraRigParameterOverride>(InnerInterfaceParameter);\
-				FPrivateVariableBuilder PrivateVariableBuilder(*this);\
-				SetupPrivateVariable(PrivateVariableBuilder, InterfaceParameter, &ValueName##Override);\
-			}\
+			Builder.SetCustomBlendableParameterOverride<U##ValueName##CameraVariable>(BlendableParameter, *TargetCustomParameter);\
 			break;
 	UE_CAMERA_VARIABLE_FOR_ALL_TYPES()
 #undef UE_CAMERA_VARIABLE_FOR_TYPE
@@ -782,47 +963,97 @@ bool FCameraRigAssetBuilder::SetupInnerCameraRigParameterOverride(UCameraRigInte
 	return true;
 }
 
-void FCameraRigAssetBuilder::DiscardUnusedPrivateVariables()
+bool FCameraRigAssetBuilder::SetupDataContextPropertyOverride(const UCameraRigDataParameter* DataParameter)
 {
-	// Now that we've rebuilt all exposed parameters, anything left from the old list 
-	// must be discarded.
-	TSet<UCameraVariableAsset*> VariablesToTrash;
+	using namespace Internal;
 
-	for (TPair<FDrivenParameterKey, UCameraVariableAsset*> Pair : OldDrivenParameters)
+	UCameraNode* Target = DataParameter->Target;
+	UClass* TargetClass = Target->GetClass();
+	FProperty* TargetProperty = TargetClass->FindPropertyByName(DataParameter->TargetPropertyName);
+	if (!TargetProperty)
 	{
-		// We null'ed the driving variable in GatherOldDrivenParameters. Now it's time
-		// to flag the camera node as modified.
-		UCameraNode* Target = Pair.Key.Value;
-		Target->Modify();
+		// No match, try something else.
+		return false;
+	}
 
-		VariablesToTrash.Add(Pair.Value);
+	const FName TargetDataIDPropertyName = FName(TargetProperty->GetName() + TEXT("DataID"));
+	FStructProperty* TargetDataIDProperty = CastField<FStructProperty>(TargetClass->FindPropertyByName(TargetDataIDPropertyName));
+	if (!TargetDataIDProperty || TargetDataIDProperty->Struct != FCameraContextDataID::StaticStruct())
+	{
+		UE_LOG(LogCameraSystem, Error,
+				TEXT("Interface parameter '{0}' is driving data context property '{1}' on '{2}' "
+					 "but no FCameraContextDataID property '{3}' was found to store the override ID."),
+				*DataParameter->InterfaceParameterName,
+				*DataParameter->TargetPropertyName.ToString(),
+				*Target->GetName(),
+				*TargetDataIDPropertyName.ToString());
+		return false;
 	}
-	OldDrivenParameters.Reset();
 
-	for (TPair<FDrivenOverrideKey, UCameraVariableAsset*> Pair : OldDrivenOverrides)
+	FCameraContextDataID* OverrideDataID = TargetDataIDProperty->ContainerPtrToValuePtr<FCameraContextDataID>(Target);
+
+	FInterfaceParameterBindingBuilder Builder(*this);
+	Builder.SetDataContextPropertyOverride(DataParameter, TargetProperty, OverrideDataID);
+
+	return true;
+}
+
+bool FCameraRigAssetBuilder::SetupCustomDataParameterOverride(const UCameraRigDataParameter* DataParameter)
+{
+	using namespace Internal;
+
+	ICustomCameraNodeParameterProvider* Target = Cast<ICustomCameraNodeParameterProvider>(DataParameter->Target);
+	if (!Target)
 	{
-		// We null'ed the override variable in GatherOldDrivenParameters. Flag the 
-		// camera rig node that owns this override as modified.
-		UCameraRigCameraNode* CameraRigNode = Pair.Key.Value;
-		CameraRigNode->Modify();
+		// No match, try something else.
+		return false;
+	}
+
+	FCustomCameraNodeParameterInfos CustomParameters;
+	Target->GetCustomCameraNodeParameters(CustomParameters);
+
+	FCustomCameraNodeParameterInfos::FDataParameterInfo* TargetCustomParameter =
+		CustomParameters.DataParameters.FindByPredicate(
+				[DataParameter](FCustomCameraNodeParameterInfos::FDataParameterInfo& CustomParameter)
+				{
+					return CustomParameter.ParameterName == DataParameter->TargetPropertyName;
+				});
+	if (!TargetCustomParameter)
+	{
+		// No match, try something else.
+		return false;
+	}
+
+	FInterfaceParameterBindingBuilder Builder(*this);
+	Builder.SetCustomDataParameterOverride(DataParameter, *TargetCustomParameter);
+
+	return true;
+}
+
+void FCameraRigAssetBuilder::DiscardUnusedParameters()
+{
+	// Now that we've rebuilt all exposed parameters, anything left from the old list 
+	// must be discarded. These are nodes and properties that used to be driven by
+	// variables and now aren't, so we need to flag them as modified.
 
-		VariablesToTrash.Add(Pair.Value);
+	for (TPair<FDrivenParameterKey, TObjectPtr<UCameraVariableAsset>> Pair : OldDrivenBlendableParameters)
+	{
+		UCameraNode* Target = Pair.Key.Value;
+		Target->Modify();
 	}
-	OldDrivenOverrides.Reset();
-	
-	// Trash the old camera variable. This helps with debugging.
-	for (UCameraVariableAsset* VariableToTrash : VariablesToTrash)
+	OldDrivenBlendableParameters.Reset();
+
+	for (TPair<FDrivenParameterKey, FCameraContextDataID> Pair : OldDrivenDataParameters)
 	{
-		TStringBuilder<256> StringBuilder;
-		StringBuilder.Append("TRASH_");
-		StringBuilder.Append(VariableToTrash->GetName());
-		VariableToTrash->Rename(StringBuilder.ToString());
+		UCameraNode* Target = Pair.Key.Value;
+		Target->Modify();
 	}
+	OldDrivenDataParameters.Reset();
 }
 
-void FCameraRigAssetBuilder::BuildAllocationInfo()
+void FCameraRigAssetBuilder::CallBuild()
 {
-	AllocationInfo = FCameraRigAllocationInfo();
+	FCameraRigBuildContext BuildContext(BuildLog);
 
 	// Build a mock tree of evaluators.
 	FCameraNodeEvaluatorTreeBuildParams BuildParams;
@@ -831,23 +1062,23 @@ void FCameraRigAssetBuilder::BuildAllocationInfo()
 	Storage.BuildEvaluatorTree(BuildParams);
 
 	// Get the size of the evaluators' allocation.
-	Storage.GetAllocationInfo(AllocationInfo.EvaluatorInfo);
+	Storage.GetAllocationInfo(BuildContext.AllocationInfo.EvaluatorInfo);
 
-	// Compute the allocation info for camera variables.
+	// Call Build() on all camera nodes in the hierarchy (detached/orphaned camera nodes don't get called).
 	for (UCameraNode* CameraNode : CameraNodeHierarchy.GetFlattenedHierarchy())
 	{
-		BuildAllocationInfo(CameraNode);
+		CallBuild(BuildContext, CameraNode);
 	}
 
-	// Set it on the camera rig asset.
-	if (CameraRig->AllocationInfo != AllocationInfo)
+	// Set the final allocation info on the camera rig asset.
+	if (CameraRig->AllocationInfo != BuildContext.AllocationInfo)
 	{
 		CameraRig->Modify();
-		CameraRig->AllocationInfo = AllocationInfo;
+		CameraRig->AllocationInfo = BuildContext.AllocationInfo;
 	}
 }
 
-void FCameraRigAssetBuilder::BuildAllocationInfo(UCameraNode* CameraNode)
+void FCameraRigAssetBuilder::CallBuild(FCameraRigBuildContext& BuildContext, UCameraNode* CameraNode)
 {
 	using namespace UE::Cameras::Internal;
 
@@ -855,6 +1086,7 @@ void FCameraRigAssetBuilder::BuildAllocationInfo(UCameraNode* CameraNode)
 	// This is for both exposed rig parameters (which we just built in BuildNewDrivenParameters) and 
 	// for parameters driven by user-defined variables.
 	UClass* CameraNodeClass = CameraNode->GetClass();
+	FCameraRigAllocationInfo& AllocationInfo = BuildContext.AllocationInfo;
 	for (TFieldIterator<FProperty> It(CameraNodeClass); It; ++It)
 	{
 		FStructProperty* StructProperty = CastField<FStructProperty>(*It);
@@ -878,15 +1110,50 @@ void FCameraRigAssetBuilder::BuildAllocationInfo(UCameraNode* CameraNode)
 UE_CAMERA_VARIABLE_FOR_ALL_TYPES()
 #undef UE_CAMERA_VARIABLE_FOR_TYPE
 		{
-			// Some other struct property.
+			// Another struct property...
+		}
+	}
+
+	// Now do the same with custom parameters handled by the node itself. These custom parameters have 
+	// been hooked up to overrides in BuildInterfaceParameterBindings so we should be getting them back 
+	// with the overrides set here.
+	if (ICustomCameraNodeParameterProvider* CustomParameterProvider = Cast<ICustomCameraNodeParameterProvider>(CameraNode))
+	{
+		FCustomCameraNodeParameterInfos CustomParameters;
+		CustomParameterProvider->GetCustomCameraNodeParameters(CustomParameters);
+
+		for (const FCustomCameraNodeParameterInfos::FBlendableParameterInfo& BlendableParameter : CustomParameters.BlendableParameters)
+		{
+			if (ensure(BlendableParameter.OverrideVariable))
+			{
+				AddCameraVariableToAllocationInfo(BlendableParameter.OverrideVariable->Get(), AllocationInfo.VariableTableInfo);
+			}
+		}
+
+		for (const FCustomCameraNodeParameterInfos::FDataParameterInfo& DataParameter : CustomParameters.DataParameters)
+		{
+			if (ensure(DataParameter.OverrideDataID))
+			{
+				AddContextDataToAllocationInfo(*DataParameter.OverrideDataID, DataParameter.ParameterType, DataParameter.ParameterTypeObject, AllocationInfo.ContextDataTableInfo);
+			}
 		}
 	}
 
 	// Let the camera node add any custom variables or extra memory.
-	FCameraRigBuildContext BuildContext(AllocationInfo, BuildLog);
 	CameraNode->Build(BuildContext);
 }
 
+void FCameraRigAssetBuilder::BuildDefaultParameters()
+{
+	FInstancedPropertyBag DefaultParameters;
+	BuildDefaultParameters(CameraRig, DefaultParameters);
+	if (DefaultParameters.GetPropertyBagStruct() != CameraRig->DefaultParameters.GetPropertyBagStruct())
+	{
+		CameraRig->Modify();
+		CameraRig->DefaultParameters = DefaultParameters;
+	}
+}
+
 void FCameraRigAssetBuilder::UpdateBuildStatus()
 {
 	ECameraBuildStatus BuildStatus = ECameraBuildStatus::Clean;
@@ -903,6 +1170,89 @@ void FCameraRigAssetBuilder::UpdateBuildStatus()
 	CameraRig->BuildStatus = BuildStatus;
 }
 
+void FCameraRigAssetBuilder::BuildDefaultParameters(UCameraRigAsset* CameraRigAsset, FInstancedPropertyBag& OutPropertyBag)
+{
+	TArray<FPropertyBagPropertyDesc> DefaultParameterProperties;
+	AppendDefaultParameters(CameraRigAsset->Interface, DefaultParameterProperties);
+	OutPropertyBag.AddProperties(DefaultParameterProperties);
+}
+
+void FCameraRigAssetBuilder::AppendDefaultParameters(const FCameraRigInterface& CameraRigInterface, TArray<FPropertyBagPropertyDesc>& OutProperties)
+{
+	for (const UCameraRigBlendableParameter* BlendableParameter : CameraRigInterface.BlendableParameters)
+	{
+		FName PropertyName(BlendableParameter->InterfaceParameterName);
+
+		EPropertyBagPropertyType PropertyType = EPropertyBagPropertyType::Struct;
+		const UObject* PropertyTypeObject = nullptr;
+
+		switch (BlendableParameter->ParameterType)
+		{
+#define UE_CAMERA_VARIABLE_FOR_TYPE(ValueType, ValueName)\
+			case ECameraVariableType::ValueName:\
+				PropertyTypeObject = F##ValueName##CameraParameter::StaticStruct();\
+				break;
+		UE_CAMERA_VARIABLE_FOR_ALL_TYPES()
+#undef UE_CAMERA_VARIABLE_FOR_TYPE
+			default:
+				ensure(false);
+				break;
+		}
+
+		FPropertyBagPropertyDesc NewProperty(PropertyName, PropertyType, PropertyTypeObject);
+		// Make the property bag match the camera interface parameter GUIDs.
+		NewProperty.ID = BlendableParameter->GetGuid();
+
+		OutProperties.Add(NewProperty);
+	}
+
+	for (const UCameraRigDataParameter* DataParameter : CameraRigInterface.DataParameters)
+	{
+		FName PropertyName(DataParameter->InterfaceParameterName);
+
+		bool bIsValid = true;
+		EPropertyBagPropertyType PropertyType = EPropertyBagPropertyType::Bool;
+		const UObject* PropertyTypeObject = DataParameter->DataTypeObject;
+
+		switch (DataParameter->DataType)
+		{
+			case ECameraContextDataType::Name:
+				PropertyType = EPropertyBagPropertyType::Name;
+				break;
+			case ECameraContextDataType::String:
+				PropertyType = EPropertyBagPropertyType::String;
+				break;
+			case ECameraContextDataType::Enum:
+				PropertyType = EPropertyBagPropertyType::Enum;
+				ensure(PropertyTypeObject && PropertyTypeObject->IsA<UEnum>());
+				break;
+			case ECameraContextDataType::Struct:
+				PropertyType = EPropertyBagPropertyType::Struct;
+				ensure(PropertyTypeObject && PropertyTypeObject->IsA<UScriptStruct>());
+				break;
+			case ECameraContextDataType::Object:
+				PropertyType = EPropertyBagPropertyType::Object;
+				break;
+			case ECameraContextDataType::Class:
+				PropertyType = EPropertyBagPropertyType::Class;
+				break;
+			default:
+				ensure(false);
+				bIsValid = false;
+				break;
+		}
+
+		if (bIsValid)
+		{
+			FPropertyBagPropertyDesc NewProperty(PropertyName, PropertyType, PropertyTypeObject);
+			// Make the property bag match the camera interface parameter GUIDs.
+			NewProperty.ID = DataParameter->GetGuid();
+
+			OutProperties.Add(NewProperty);
+		}
+	}
+}
+
 }  // namespace UE::Cameras
 
 #undef LOCTEXT_NAMESPACE
diff --git a/Engine/Plugins/Cameras/GameplayCameras/Source/GameplayCameras/Private/Core/CameraRigAssetReference.cpp b/Engine/Plugins/Cameras/GameplayCameras/Source/GameplayCameras/Private/Core/CameraRigAssetReference.cpp
--- a/Engine/Plugins/Cameras/GameplayCameras/Source/GameplayCameras/Private/Core/CameraRigAssetReference.cpp
+++ b/Engine/Plugins/Cameras/GameplayCameras/Source/GameplayCameras/Private/Core/CameraRigAssetReference.cpp
@@ -2,15 +2,13 @@
 
 #include "Core/CameraRigAssetReference.h"
 
+#include "Core/CameraRigAsset.h"
+#include "Core/CameraNodeEvaluator.h"
+#include "Helpers/CameraRigParameterOverrideEvaluator.h"
+
 #include UE_INLINE_GENERATED_CPP_BY_NAME(CameraRigAssetReference)
 
-void FCameraRigParameterOverrides::Reset()
-{
-#define UE_CAMERA_VARIABLE_FOR_TYPE(ValueType, ValueName)\
-	ValueName##Overrides.Reset();
-UE_CAMERA_VARIABLE_FOR_ALL_TYPES()
-#undef UE_CAMERA_VARIABLE_FOR_TYPE
-}
+PRAGMA_DISABLE_DEPRECATION_WARNINGS
 
 FCameraRigAssetReference::FCameraRigAssetReference()
 {
@@ -21,53 +19,84 @@ FCameraRigAssetReference::FCameraRigAssetReference(UCameraRigAsset* InCameraRig)
 {
 }
 
-bool FCameraRigAssetReference::UpdateParameterOverrides()
+PRAGMA_ENABLE_DEPRECATION_WARNINGS
+
+void FCameraRigAssetReference::ApplyParameterOverrides(UE::Cameras::FCameraNodeEvaluationResult& OutResult, bool bDrivenOverridesOnly)
+{
+	using namespace UE::Cameras;
+	FCameraRigParameterOverrideEvaluator OverrideEvaluator(*this);
+	OverrideEvaluator.ApplyParameterOverrides(OutResult.VariableTable, OutResult.ContextDataTable, bDrivenOverridesOnly);
+}
+
+bool FCameraRigAssetReference::IsParameterOverriden(const FGuid PropertyID) const
 {
-	if (!CameraRig)
+	return ParameterOverrideGuids.Contains(PropertyID);
+}
+
+void FCameraRigAssetReference::SetParameterOverriden(const FGuid PropertyID, bool bIsOverridden)
+{
+	if (bIsOverridden)
 	{
-		bool bHasAnyOverride = false;
-#define UE_CAMERA_VARIABLE_FOR_TYPE(ValueType, ValueName)\
-		for (F##ValueName##CameraRigParameterOverride& ParameterOverride : ParameterOverrides.ValueName##Overrides)\
-		{\
-			ParameterOverride.bInvalid = true;\
-			bHasAnyOverride = true;\
+		ParameterOverrideGuids.AddUnique(PropertyID);
+	}
+	else
+	{
+		ParameterOverrideGuids.Remove(PropertyID);
+	}
+}
+
+bool FCameraRigAssetReference::NeedsRebuildParameters() const
+{
+	if ((!CameraRig && Parameters.IsValid()) || (CameraRig && !Parameters.IsValid()))
+	{
+		return true;
+	}
+
+	if (CameraRig)
+	{
+		const UPropertyBag* AssetParametersType = CameraRig->GetDefaultParameters().GetPropertyBagStruct();
+		const UPropertyBag* ReferenceParametersType = Parameters.GetPropertyBagStruct();
+		if (AssetParametersType != ReferenceParametersType)
+		{
+			return true;
 		}
-		UE_CAMERA_VARIABLE_FOR_ALL_TYPES()
-#undef UE_CAMERA_VARIABLE_FOR_TYPE
-		return bHasAnyOverride;
 	}
 
-	bool bAnyModified = false;
-	FCameraRigInterface& CameraRigInterface = CameraRig->Interface;
-#define UE_CAMERA_VARIABLE_FOR_TYPE(ValueType, ValueName)\
-	for (F##ValueName##CameraRigParameterOverride& ParameterOverride : ParameterOverrides.ValueName##Overrides)\
-	{\
-		UCameraRigInterfaceParameter* InterfaceParameter = CameraRigInterface.FindInterfaceParameterByGuid(\
-				ParameterOverride.InterfaceParameterGuid);\
-		\
-		const bool bWasInvalid = ParameterOverride.bInvalid;\
-		ParameterOverride.bInvalid = (InterfaceParameter == nullptr);\
-		bAnyModified |= (bWasInvalid != ParameterOverride.bInvalid);\
-		\
-		if (InterfaceParameter)\
-		{\
-			if (ParameterOverride.InterfaceParameterName != InterfaceParameter->InterfaceParameterName)\
-			{\
-				ParameterOverride.InterfaceParameterName = InterfaceParameter->InterfaceParameterName;\
-				bAnyModified = true;\
-			}\
-			UCameraVariableAsset* InterfaceParameterVariable = InterfaceParameter->PrivateVariable;\
-			const FGuid NewPrivateVariableGuid = InterfaceParameterVariable ? InterfaceParameterVariable->GetGuid() : FGuid();\
-			if (ParameterOverride.PrivateVariableGuid != NewPrivateVariableGuid)\
-			{\
-				ParameterOverride.PrivateVariableGuid = NewPrivateVariableGuid;\
-				bAnyModified = true;\
-			}\
-		}\
+	return false;
+}
+
+bool FCameraRigAssetReference::RebuildParametersIfNeeded()
+{
+	if (NeedsRebuildParameters())
+	{
+		RebuildParameters();
+		return true;
+	}
+	return false;
+}
+
+void FCameraRigAssetReference::RebuildParameters()
+{
+	if (CameraRig)
+	{
+		Parameters.MigrateToNewBagInstanceWithOverrides(CameraRig->GetDefaultParameters(), ParameterOverrideGuids);
+		
+		// Remove overrides for parameters that don't exist anymore.
+		if (const UPropertyBag* ParametersType = Parameters.GetPropertyBagStruct())
+		{
+			for (TArray<FGuid>::TIterator It = ParameterOverrideGuids.CreateIterator(); It; ++It)
+			{
+				if (!ParametersType->FindPropertyDescByID(*It))
+				{
+					It.RemoveCurrentSwap();
+				}
+			}
+		}
+	}
+	else
+	{
+		Parameters.Reset();
 	}
-	UE_CAMERA_VARIABLE_FOR_ALL_TYPES()
-#undef UE_CAMERA_VARIABLE_FOR_TYPE
-	return bAnyModified;
 }
 
 bool FCameraRigAssetReference::SerializeFromMismatchedTag(struct FPropertyTag const& Tag, FStructuredArchive::FSlot Slot)
@@ -82,3 +111,49 @@ bool FCameraRigAssetReference::SerializeFromMismatchedTag(struct FPropertyTag co
 	return false;
 }
 
+void FCameraRigAssetReference::PostSerialize(const FArchive& Ar)
+{
+	PRAGMA_DISABLE_DEPRECATION_WARNINGS
+
+	// Make a property bag with the legacy overrides, and then set the values in it.
+	bool bHasAnyLegacyOverride = false;
+	TArray<FPropertyBagPropertyDesc> LegacyParameterProperties;
+	TArray<FGuid> LegacyParameterOverrides;
+
+#define UE_CAMERA_VARIABLE_FOR_TYPE(ValueType, ValueName)\
+	for (F##ValueName##CameraRigParameterOverride& ParameterOverride : ParameterOverrides_DEPRECATED.ValueName##Overrides)\
+	{\
+		FName PropertyName(ParameterOverride.InterfaceParameterName);\
+		EPropertyBagPropertyType PropertyType = EPropertyBagPropertyType::Struct;\
+		const UObject* PropertyTypeObject = F##ValueName##CameraParameter::StaticStruct();\
+		FPropertyBagPropertyDesc LegacyParameterProperty(PropertyName, PropertyType, PropertyTypeObject);\
+		LegacyParameterProperty.ID = ParameterOverride.InterfaceParameterGuid;\
+		LegacyParameterProperties.Add(LegacyParameterProperty);\
+		LegacyParameterOverrides.Add(LegacyParameterProperty.ID);\
+		bHasAnyLegacyOverride = true;\
+	}
+	UE_CAMERA_VARIABLE_FOR_ALL_TYPES()
+#undef UE_CAMERA_VARIABLE_FOR_TYPE
+
+	if (bHasAnyLegacyOverride)
+	{
+		Parameters = FInstancedPropertyBag();
+		Parameters.AddProperties(LegacyParameterProperties);
+
+		ParameterOverrideGuids = LegacyParameterOverrides;
+
+#define UE_CAMERA_VARIABLE_FOR_TYPE(ValueType, ValueName)\
+		for (F##ValueName##CameraRigParameterOverride& ParameterOverride : ParameterOverrides_DEPRECATED.ValueName##Overrides)\
+		{\
+			FName PropertyName(ParameterOverride.InterfaceParameterName);\
+			Parameters.SetValueStruct<F##ValueName##CameraParameter>(PropertyName, ParameterOverride.Value);\
+		}
+		UE_CAMERA_VARIABLE_FOR_ALL_TYPES()
+#undef UE_CAMERA_VARIABLE_FOR_TYPE
+
+		ParameterOverrides_DEPRECATED = FCameraRigParameterOverrides();
+	}
+
+	PRAGMA_ENABLE_DEPRECATION_WARNINGS
+}
+
diff --git a/Engine/Plugins/Cameras/GameplayCameras/Source/GameplayCameras/Private/Core/CameraRigCombinationRegistry.cpp b/Engine/Plugins/Cameras/GameplayCameras/Source/GameplayCameras/Private/Core/CameraRigCombinationRegistry.cpp
--- a/Engine/Plugins/Cameras/GameplayCameras/Source/GameplayCameras/Private/Core/CameraRigCombinationRegistry.cpp
+++ b/Engine/Plugins/Cameras/GameplayCameras/Source/GameplayCameras/Private/Core/CameraRigCombinationRegistry.cpp
@@ -3,7 +3,7 @@
 #include "Core/CameraRigCombinationRegistry.h"
 
 #include "Core/CameraRigAsset.h"
-#include "Core/CameraRigParameterOverrideEvaluator.h"
+#include "Helpers/CameraRigParameterOverrideEvaluator.h"
 
 namespace UE::Cameras
 {
diff --git a/Engine/Plugins/Cameras/GameplayCameras/Source/GameplayCameras/Private/Core/CameraSystemEvaluator.cpp b/Engine/Plugins/Cameras/GameplayCameras/Source/GameplayCameras/Private/Core/CameraSystemEvaluator.cpp
--- a/Engine/Plugins/Cameras/GameplayCameras/Source/GameplayCameras/Private/Core/CameraSystemEvaluator.cpp
+++ b/Engine/Plugins/Cameras/GameplayCameras/Source/GameplayCameras/Private/Core/CameraSystemEvaluator.cpp
@@ -116,6 +116,7 @@ void FCameraSystemEvaluator::AddReferencedObjects(FReferenceCollector& Collector
 {
 	Collector.AddReferencedObject(RootNode);
 	ContextStack.AddReferencedObjects(Collector);
+	RootNodeResult.AddReferencedObjects(Collector);
 	if (RootEvaluator)
 	{
 		RootEvaluator->AddReferencedObjects(Collector);
diff --git a/Engine/Plugins/Cameras/GameplayCameras/Source/GameplayCameras/Private/Core/CameraVariableTable.cpp b/Engine/Plugins/Cameras/GameplayCameras/Source/GameplayCameras/Private/Core/CameraVariableTable.cpp
--- a/Engine/Plugins/Cameras/GameplayCameras/Source/GameplayCameras/Private/Core/CameraVariableTable.cpp
+++ b/Engine/Plugins/Cameras/GameplayCameras/Source/GameplayCameras/Private/Core/CameraVariableTable.cpp
@@ -619,6 +619,37 @@ bool operator==(const FCameraVariableDefinition& A, const FCameraVariableDefinit
 		;
 }
 
+void FCameraVariableTableAllocationInfo::Combine(const FCameraVariableTableAllocationInfo& OtherInfo)
+{
+	TMap<FCameraVariableID, int32> KnownIDs;
+	for (auto It = VariableDefinitions.CreateConstIterator(); It; ++It)
+	{
+		const FCameraVariableDefinition& VariableDefinition(*It);
+		KnownIDs.Add(VariableDefinition.VariableID, It.GetIndex());
+	}
+
+	for (const FCameraVariableDefinition& OtherVariableDefinition : OtherInfo.VariableDefinitions)
+	{
+		const int32 KnownIndex = KnownIDs.FindRef(OtherVariableDefinition.VariableID, INDEX_NONE);
+		if (KnownIndex == INDEX_NONE)
+		{
+			VariableDefinitions.Add(OtherVariableDefinition);
+		}
+		else
+		{
+			const FCameraVariableDefinition& KnownVariableDefinition(VariableDefinitions[KnownIndex]);
+			ensure(KnownVariableDefinition == OtherVariableDefinition);
+		}
+	}
+
+	TSet<TObjectPtr<UCameraVariableAsset>> KnownAutoResetVariables(AutoResetVariables);
+	for (UCameraVariableAsset* OtherAutoResetVariable : OtherInfo.AutoResetVariables)
+	{
+		KnownAutoResetVariables.Add(OtherAutoResetVariable);
+	}
+	AutoResetVariables = KnownAutoResetVariables.Array();
+}
+
 bool operator==(const FCameraVariableTableAllocationInfo& A, const FCameraVariableTableAllocationInfo& B)
 {
 	return A.VariableDefinitions == B.VariableDefinitions
diff --git a/Engine/Plugins/Cameras/GameplayCameras/Source/GameplayCameras/Private/Core/ICustomCameraNodeParameterProvider.cpp b/Engine/Plugins/Cameras/GameplayCameras/Source/GameplayCameras/Private/Core/ICustomCameraNodeParameterProvider.cpp
new file mode 100644
--- /dev/null
+++ b/Engine/Plugins/Cameras/GameplayCameras/Source/GameplayCameras/Private/Core/ICustomCameraNodeParameterProvider.cpp
@@ -0,0 +1,110 @@
+// Copyright Epic Games, Inc. All Rights Reserved.
+
+#include "Core/ICustomCameraNodeParameterProvider.h"
+
+bool operator==(const FCustomCameraNodeBlendableParameter& A, const FCustomCameraNodeBlendableParameter& B)
+{
+	return A.ParameterName == B.ParameterName &&
+		A.ParameterType == B.ParameterType &&
+		A.OverrideVariable == B.OverrideVariable;
+}
+
+bool operator==(const FCustomCameraNodeDataParameter& A, const FCustomCameraNodeDataParameter& B)
+{
+	return A.ParameterName == B.ParameterName &&
+		A.ParameterType == B.ParameterType &&
+		A.ParameterTypeObject == B.ParameterTypeObject &&
+		A.OverrideDataID == B.OverrideDataID;
+}
+
+bool operator==(const FCustomCameraNodeParameters& A, const FCustomCameraNodeParameters& B)
+{
+	return A.BlendableParameters == B.BlendableParameters &&
+		A.DataParameters == B.DataParameters;
+}
+
+void FCustomCameraNodeParameterInfos::AddBlendableParameter(
+		FName ParameterName, 
+		ECameraVariableType ParameterType, 
+		const uint8* DefaultValuePtr,
+		TObjectPtr<UCameraVariableAsset>* OverrideVariable)
+{
+	BlendableParameters.Add({ ParameterName, ParameterType, DefaultValuePtr, OverrideVariable });
+}
+
+void FCustomCameraNodeParameterInfos::AddDataParameter(
+		FName ParameterName, 
+		ECameraContextDataType ParameterType,
+		const UObject* ParameterTypeObject,
+		FCameraContextDataID* OverrideDataID)
+{
+	DataParameters.Add({ ParameterName, ParameterType, ParameterTypeObject, OverrideDataID });
+}
+
+void FCustomCameraNodeParameterInfos::GetBlendableParameters(TArray<FCustomCameraNodeBlendableParameter>& OutBlendableParameters) const
+{
+	for (const FBlendableParameterInfo& BlendableParameter : BlendableParameters)
+	{
+		FCustomCameraNodeBlendableParameter& OutParameter = OutBlendableParameters.Emplace_GetRef();
+		OutParameter.ParameterName = BlendableParameter.ParameterName;
+		OutParameter.ParameterType = BlendableParameter.ParameterType;
+		if (BlendableParameter.OverrideVariable)
+		{
+			OutParameter.OverrideVariable = *BlendableParameter.OverrideVariable;
+		}
+	}
+}
+
+void FCustomCameraNodeParameterInfos::GetDataParameters(TArray<FCustomCameraNodeDataParameter>& OutDataParameters) const
+{
+	for (const FDataParameterInfo& DataParameter : DataParameters)
+	{
+		FCustomCameraNodeDataParameter& OutParameter = OutDataParameters.Emplace_GetRef();
+		OutParameter.ParameterName = DataParameter.ParameterName;
+		OutParameter.ParameterType = DataParameter.ParameterType;
+		OutParameter.ParameterTypeObject = DataParameter.ParameterTypeObject;
+		if (DataParameter.OverrideDataID)
+		{
+			OutParameter.OverrideDataID = *DataParameter.OverrideDataID;
+		}
+	}
+}
+
+bool FCustomCameraNodeParameterInfos::FindBlendableParameter(FName ParameterName, FCustomCameraNodeBlendableParameter& OutParameter) const
+{
+	for (const FBlendableParameterInfo& BlendableParameter : BlendableParameters)
+	{
+		if (BlendableParameter.ParameterName == ParameterName)
+		{
+			OutParameter.ParameterName = BlendableParameter.ParameterName;
+			OutParameter.ParameterType = BlendableParameter.ParameterType;
+			OutParameter.OverrideVariable = nullptr;
+			if (BlendableParameter.OverrideVariable)
+			{
+				OutParameter.OverrideVariable = *BlendableParameter.OverrideVariable;
+			}
+			return true;
+		}
+	}
+	return false;
+}
+
+bool FCustomCameraNodeParameterInfos::FindDataParameter(FName ParameterName, FCustomCameraNodeDataParameter& OutParameter) const
+{
+	for (const FDataParameterInfo& DataParameter : DataParameters)
+	{
+		if (DataParameter.ParameterName == ParameterName)
+		{
+			OutParameter.ParameterName = DataParameter.ParameterName;
+			OutParameter.ParameterType = DataParameter.ParameterType;
+			OutParameter.ParameterTypeObject = DataParameter.ParameterTypeObject;
+			if (DataParameter.OverrideDataID)
+			{
+				OutParameter.OverrideDataID = *DataParameter.OverrideDataID;
+			}
+			return true;
+		}
+	}
+	return false;
+}
+
diff --git a/Engine/Plugins/Cameras/GameplayCameras/Source/GameplayCameras/Private/Debug/CameraNodeEvaluationResultDebugBlock.cpp b/Engine/Plugins/Cameras/GameplayCameras/Source/GameplayCameras/Private/Debug/CameraNodeEvaluationResultDebugBlock.cpp
--- a/Engine/Plugins/Cameras/GameplayCameras/Source/GameplayCameras/Private/Debug/CameraNodeEvaluationResultDebugBlock.cpp
+++ b/Engine/Plugins/Cameras/GameplayCameras/Source/GameplayCameras/Private/Debug/CameraNodeEvaluationResultDebugBlock.cpp
@@ -9,6 +9,7 @@
 #include "Debug/CameraDebugRenderer.h"
 #include "Debug/CameraPoseDebugBlock.h"
 #include "Debug/CameraRigJointsDebugBlock.h"
+#include "Debug/ContextDataTableDebugBlock.h"
 #include "Debug/DebugTextRenderer.h"
 #include "Debug/VariableTableDebugBlock.h"
 #include "Math/ColorList.h"
@@ -32,6 +33,7 @@ void FCameraNodeEvaluationResultDebugBlock::Initialize(const FCameraNodeEvaluati
 
 	AddChild(&Builder.BuildDebugBlock<FCameraPoseDebugBlock>(InResult.CameraPose));
 	AddChild(&Builder.BuildDebugBlock<FVariableTableDebugBlock>(InResult.VariableTable));
+	AddChild(&Builder.BuildDebugBlock<FContextDataTableDebugBlock>(InResult.ContextDataTable));
 	AddChild(&Builder.BuildDebugBlock<FCameraRigJointsDebugBlock>(InResult.CameraRigJoints, InResult.VariableTable));
 }
 
@@ -42,7 +44,7 @@ void FCameraNodeEvaluationResultDebugBlock::Initialize(const FCameraSystemEvalua
 
 	AddChild(&Builder.BuildDebugBlock<FCameraPoseDebugBlock>(InResult.CameraPose));
 	AddChild(&Builder.BuildDebugBlock<FVariableTableDebugBlock>(InResult.VariableTable));
-	// No rig joints exposed on the camera system result.
+	// No context data or rig joints exposed on the camera system result.
 }
 
 FCameraPoseDebugBlock* FCameraNodeEvaluationResultDebugBlock::GetCameraPoseDebugBlock()
@@ -65,6 +67,16 @@ FVariableTableDebugBlock* FCameraNodeEvaluationResultDebugBlock::GetVariableTabl
 	return nullptr;
 }
 
+FContextDataTableDebugBlock* FCameraNodeEvaluationResultDebugBlock::GetContextDataTableDebugBlock()
+{
+	TArrayView<FCameraDebugBlock*> ChildrenView(GetChildren());
+	if (ChildrenView.IsValidIndex(2))
+	{
+		return ChildrenView[2]->CastThis<FContextDataTableDebugBlock>();
+	}
+	return nullptr;
+}
+
 void FCameraNodeEvaluationResultDebugBlock::OnDebugDraw(const FCameraDebugBlockDrawParams& Params, FCameraDebugRenderer& Renderer)
 {
 	const FCameraDebugColors& Colors = FCameraDebugColors::Get();
@@ -110,12 +122,20 @@ void FCameraNodeEvaluationResultDebugBlock::OnDebugDraw(const FCameraDebugBlockD
 	}
 	if (ChildrenView.IsValidIndex(2))
 	{
-		Renderer.AddText(TEXT("{cam_title}Camera Rig Joints:"));
+		Renderer.AddText(TEXT("{cam_title}Context Data Table:"));
 		Renderer.AddIndent();
 		Renderer.SetTextColor(Colors.Default);
 		ChildrenView[2]->DebugDraw(Params, Renderer);
 		Renderer.RemoveIndent();
 	}
+	if (ChildrenView.IsValidIndex(3))
+	{
+		Renderer.AddText(TEXT("{cam_title}Camera Rig Joints:"));
+		Renderer.AddIndent();
+		Renderer.SetTextColor(Colors.Default);
+		ChildrenView[3]->DebugDraw(Params, Renderer);
+		Renderer.RemoveIndent();
+	}
 
 	// We already manually rendered our children.
 	Renderer.SkipAllBlocks();
diff --git a/Engine/Plugins/Cameras/GameplayCameras/Source/GameplayCameras/Private/GameFramework/BlueprintCameraContextDataTable.cpp b/Engine/Plugins/Cameras/GameplayCameras/Source/GameplayCameras/Private/GameFramework/BlueprintCameraContextDataTable.cpp
new file mode 100644
--- /dev/null
+++ b/Engine/Plugins/Cameras/GameplayCameras/Source/GameplayCameras/Private/GameFramework/BlueprintCameraContextDataTable.cpp
@@ -0,0 +1,49 @@
+// Copyright Epic Games, Inc. All Rights Reserved.
+
+#include "GameFramework/BlueprintCameraContextDataTable.h"
+
+#include "Core/CameraContextDataTable.h"
+
+#include UE_INLINE_GENERATED_CPP_BY_NAME(BlueprintCameraContextDataTable)
+
+FBlueprintCameraContextDataTable::FBlueprintCameraContextDataTable()
+	: PrivateContextDataTable(nullptr)
+{
+}
+
+FBlueprintCameraContextDataTable::FBlueprintCameraContextDataTable(FCameraContextDataTable* InContextDataTable)
+	: PrivateContextDataTable(InContextDataTable)
+{
+}
+
+FInstancedStruct UBlueprintCameraContextDataTableFunctionLibrary::GetStructData(const FBlueprintCameraContextDataTable& ContextDataTable, FCameraContextDataID DataID, const UScriptStruct* DataStructType)
+{
+	using namespace UE::Cameras;
+
+	if (FCameraContextDataTable* ActualTable = ContextDataTable.GetContextDataTable())
+	{
+		if (DataID.IsValid())
+		{
+			const uint8* RawData = ActualTable->GetData(DataID, ECameraContextDataType::Struct, DataStructType);
+			FInstancedStruct ReturnValue;
+			ReturnValue.InitializeAs(DataStructType, RawData);
+			return ReturnValue;
+		}
+	}
+
+	return FInstancedStruct();
+}
+
+void UBlueprintCameraContextDataTableFunctionLibrary::SetStructData(const FBlueprintCameraContextDataTable& ContextDataTable, FCameraContextDataID DataID, const FInstancedStruct& Data)
+{
+	using namespace UE::Cameras;
+
+	if (FCameraContextDataTable* ActualTable = ContextDataTable.GetContextDataTable())
+	{
+		if (DataID.IsValid())
+		{
+			ActualTable->SetData(DataID, Data);
+		}
+	}
+}
+
diff --git a/Engine/Plugins/Cameras/GameplayCameras/Source/GameplayCameras/Private/GameFramework/CameraRigParameterInterop.cpp b/Engine/Plugins/Cameras/GameplayCameras/Source/GameplayCameras/Private/GameFramework/CameraRigParameterInterop.cpp
--- a/Engine/Plugins/Cameras/GameplayCameras/Source/GameplayCameras/Private/GameFramework/CameraRigParameterInterop.cpp
+++ b/Engine/Plugins/Cameras/GameplayCameras/Source/GameplayCameras/Private/GameFramework/CameraRigParameterInterop.cpp
@@ -38,12 +38,12 @@ UCameraRigParameterInterop::UCameraRigParameterInterop(const FObjectInitializer&
 {
 }
 
-void UCameraRigParameterInterop::SetBooleanParameter(FBlueprintCameraVariableTable& VariableTable, UCameraRigAsset* CameraRig, const FString& ParameterName, bool bParameterValue)
+void UCameraRigParameterInterop::SetBooleanParameter(FBlueprintCameraVariableTable& VariableTable, UCameraRigAsset* CameraRig, const FString& ParameterName, bool ParameterValue)
 {
 	UE::Cameras::Private::SetCameraRigParameter(
 			VariableTable, 
 			Cast<UBooleanCameraVariable>(GetParameterPrivateVariable(CameraRig, ParameterName)), 
-			bParameterValue);
+			ParameterValue);
 }
 
 void UCameraRigParameterInterop::SetIntegerParameter(FBlueprintCameraVariableTable& VariableTable, UCameraRigAsset* CameraRig, const FString& ParameterName, int32 ParameterValue)
@@ -112,22 +112,22 @@ void UCameraRigParameterInterop::SetTransformParameter(FBlueprintCameraVariableT
 
 UCameraVariableAsset* UCameraRigParameterInterop::GetParameterPrivateVariable(UCameraRigAsset* CameraRig, const FString& ParameterName)
 {
-	UCameraRigInterfaceParameter* InterfaceParameter = CameraRig->Interface.FindInterfaceParameterByName(ParameterName);
-	if (!InterfaceParameter)
+	UCameraRigBlendableParameter* BlendableParameter = CameraRig->Interface.FindBlendableParameterByName(ParameterName);
+	if (!BlendableParameter)
 	{
 		const FText Text = LOCTEXT("NoSuchParameter", "No parameter '{0}' found on camera rig '{1}'. Setting this camera variable table value will most probably accomplish nothing.");
 		FFrame::KismetExecutionMessage(*FText::Format(Text, FText::FromString(ParameterName), FText::FromString(CameraRig->GetPathName())).ToString(), ELogVerbosity::Warning);
 		return nullptr;
 	}
 
-	if (!InterfaceParameter->PrivateVariable)
+	if (!BlendableParameter->PrivateVariable)
 	{
 		const FText Text = LOCTEXT("CameraRigNeedsBuilding", "Parameter '{0}' isn't built. Please build camera rig '{1}'.");
 		FFrame::KismetExecutionMessage(*FText::Format(Text, FText::FromString(ParameterName), FText::FromString(CameraRig->GetPathName())).ToString(), ELogVerbosity::Warning);
 		return nullptr;
 	}
 
-	return InterfaceParameter->PrivateVariable;
+	return BlendableParameter->PrivateVariable;
 }
 
 #undef LOCTEXT_NAMESPACE
diff --git a/Engine/Plugins/Cameras/GameplayCameras/Source/GameplayCameras/Private/GameFramework/GameplayCameraComponent.cpp b/Engine/Plugins/Cameras/GameplayCameras/Source/GameplayCameras/Private/GameFramework/GameplayCameraComponent.cpp
--- a/Engine/Plugins/Cameras/GameplayCameras/Source/GameplayCameras/Private/GameFramework/GameplayCameraComponent.cpp
+++ b/Engine/Plugins/Cameras/GameplayCameras/Source/GameplayCameras/Private/GameFramework/GameplayCameraComponent.cpp
@@ -3,6 +3,7 @@
 #include "GameFramework/GameplayCameraComponent.h"
 
 #include "Components/StaticMeshComponent.h"
+#include "Core/CameraAsset.h"
 #include "Core/CameraAssetBuilder.h"
 #include "Core/CameraBuildLog.h"
 #include "Core/CameraSystemEvaluator.h"
@@ -36,6 +37,17 @@ UGameplayCameraComponent::UGameplayCameraComponent(const FObjectInitializer& Obj
 #endif  // WITH_EDITORONLY_DATA
 }
 
+void UGameplayCameraComponent::PostLoad()
+{
+	Super::PostLoad();
+
+	if (Camera_DEPRECATED)
+	{
+		CameraReference.SetCameraAsset(Camera_DEPRECATED);
+		Camera_DEPRECATED = nullptr;
+	}
+}
+
 TSharedPtr<UE::Cameras::FCameraEvaluationContext> UGameplayCameraComponent::GetEvaluationContext()
 {
 	return EvaluationContext;
@@ -112,7 +124,7 @@ void UGameplayCameraComponent::ActivateCameraEvaluationContext(APlayerController
 		return;
 	}
 
-	if (!Camera)
+	if (!CameraReference.IsValid())
 	{
 		FFrame::KismetExecutionMessage(
 				TEXT("Can't activate gameplay camera component: no camera asset was set!"),
@@ -137,9 +149,16 @@ void UGameplayCameraComponent::ActivateCameraEvaluationContext(APlayerController
 
 		FCameraEvaluationContextInitializeParams InitParams;
 		InitParams.Owner = this;
-		InitParams.CameraAsset = Camera;
+		InitParams.CameraAsset = CameraReference.GetCameraAsset();
 		InitParams.PlayerController = PlayerController;
 		EvaluationContext->Initialize(InitParams);
+
+		FCameraNodeEvaluationResult& InitialResult = EvaluationContext->GetInitialResult();
+		const FCameraAssetAllocationInfo& AllocationInfo = CameraReference.GetCameraAsset()->GetAllocationInfo();
+		InitialResult.VariableTable.Initialize(AllocationInfo.VariableTableInfo);
+		InitialResult.ContextDataTable.Initialize(AllocationInfo.ContextDataTableInfo);
+
+		UpdateCameraEvaluationContext(true);
 	}
 
 	TSharedPtr<FCameraSystemEvaluator> CameraSystemEvaluator = CameraSystemHost->GetCameraSystemEvaluator();
@@ -233,13 +252,15 @@ void UGameplayCameraComponent::BeginPlay()
 	Super::BeginPlay();
 
 #if WITH_EDITOR
-	if (Camera)
+	if (CameraReference.IsValid())
 	{
 		// Auto-build the camera asset on begin play to make sure we've got the latest user edits.
 		using namespace UE::Cameras;
 		FCameraBuildLog BuildLog;
 		FCameraAssetBuilder Builder(BuildLog);
-		Builder.BuildCamera(Camera);
+		Builder.BuildCamera(CameraReference.GetCameraAsset());
+
+		CameraReference.RebuildParametersIfNeeded();
 	}
 #endif
 
@@ -263,13 +284,30 @@ void UGameplayCameraComponent::TickComponent(float DeltaTime, enum ELevelTick Ti
 
 	if (EvaluationContext)
 	{
-		EvaluationContext->Update(this);
+		UpdateCameraEvaluationContext(false);
+	}
+}
+
+void UGameplayCameraComponent::UpdateCameraEvaluationContext(bool bApplyParameterOverrides)
+{
+	using namespace UE::Cameras;
+
+	FCameraNodeEvaluationResult& InitialResult = EvaluationContext->GetInitialResult();
 
-		if (bIsCameraCutNextFrame)
-		{
-			EvaluationContext->GetInitialResult().bIsCameraCut = true;
-			bIsCameraCutNextFrame = false;
-		}
+	const FTransform& OwnerTransform = GetComponentTransform();
+	InitialResult.CameraPose.SetTransform(OwnerTransform);
+	InitialResult.bIsCameraCut = false;
+	InitialResult.bIsValid = true;
+
+	if (bIsCameraCutNextFrame)
+	{
+		InitialResult.bIsCameraCut = true;
+		bIsCameraCutNextFrame = false;
+	}
+
+	if (bApplyParameterOverrides)
+	{
+		CameraReference.ApplyParameterOverrides(InitialResult, false);
 	}
 }
 
@@ -329,14 +367,6 @@ namespace UE::Cameras
 
 UE_DEFINE_CAMERA_EVALUATION_CONTEXT(FGameplayCameraComponentEvaluationContext)
 
-void FGameplayCameraComponentEvaluationContext::Update(UGameplayCameraComponent* Owner)
-{
-	const FTransform& OwnerTransform = Owner->GetComponentTransform();
-	InitialResult.CameraPose.SetTransform(OwnerTransform);
-	InitialResult.bIsCameraCut = false;
-	InitialResult.bIsValid = true;
-}
-
 }  // namespace UE::Cameras
 
 #undef LOCTEXT_NAMESPACE
diff --git a/Engine/Plugins/Cameras/GameplayCameras/Source/GameplayCameras/Private/GameplayCamerasDelegates.cpp b/Engine/Plugins/Cameras/GameplayCameras/Source/GameplayCameras/Private/GameplayCamerasDelegates.cpp
--- a/Engine/Plugins/Cameras/GameplayCameras/Source/GameplayCameras/Private/GameplayCamerasDelegates.cpp
+++ b/Engine/Plugins/Cameras/GameplayCameras/Source/GameplayCameras/Private/GameplayCamerasDelegates.cpp
@@ -5,6 +5,7 @@
 namespace UE::Cameras
 {
 
+FOnCameraAssetBuilt FGameplayCamerasDelegates::OnCameraAssetBuiltDelegates;
 FOnCameraRigAssetBuilt FGameplayCamerasDelegates::OnCameraRigAssetBuiltDelegates;
 
 }  // namespace UE::Cameras
diff --git a/Engine/Plugins/Cameras/GameplayCameras/Source/GameplayCameras/Private/Helpers/CameraAssetParameterOverrideEvaluator.cpp b/Engine/Plugins/Cameras/GameplayCameras/Source/GameplayCameras/Private/Helpers/CameraAssetParameterOverrideEvaluator.cpp
new file mode 100644
--- /dev/null
+++ b/Engine/Plugins/Cameras/GameplayCameras/Source/GameplayCameras/Private/Helpers/CameraAssetParameterOverrideEvaluator.cpp
@@ -0,0 +1,67 @@
+// Copyright Epic Games, Inc. All Rights Reserved.
+
+#include "Helpers/CameraAssetParameterOverrideEvaluator.h"
+
+#include "Core/CameraAsset.h"
+#include "Core/CameraAssetReference.h"
+#include "Helpers/CameraParameterOverrideHelper.h"
+
+namespace UE::Cameras
+{
+
+FCameraAssetParameterOverrideEvaluator::FCameraAssetParameterOverrideEvaluator(const FCameraAssetReference& InCameraReference)
+	: CameraReference(InCameraReference)
+{
+}
+
+void FCameraAssetParameterOverrideEvaluator::ApplyParameterOverrides(FCameraVariableTable& OutVariableTable, bool bDrivenOverridesOnly)
+{
+	ApplyParameterOverrides(&OutVariableTable, nullptr, bDrivenOverridesOnly);
+}
+
+void FCameraAssetParameterOverrideEvaluator::ApplyParameterOverrides(FCameraVariableTable& OutVariableTable, FCameraContextDataTable& OutContextDataTable, bool bDrivenOverridesOnly)
+{
+	ApplyParameterOverrides(&OutVariableTable, &OutContextDataTable, bDrivenOverridesOnly);
+}
+
+void FCameraAssetParameterOverrideEvaluator::ApplyParameterOverrides(FCameraVariableTable* OutVariableTable, FCameraContextDataTable* OutContextDataTable, bool bDrivenOverridesOnly)
+{
+	using namespace Internal;
+
+	check(OutVariableTable);
+
+	const UCameraAsset* CameraAsset = CameraReference.GetCameraAsset();
+	const FInstancedPropertyBag& CameraParameters = CameraReference.GetParameters();
+	const UPropertyBag* CameraParametersStruct = CameraParameters.GetPropertyBagStruct();
+	if (!CameraParametersStruct)
+	{
+		return;
+	}
+
+	TConstArrayView<FPropertyBagPropertyDesc> PropertyDescs = CameraParametersStruct->GetPropertyDescs();
+	const TArray<TWeakObjectPtr<const UCameraRigAsset>>& ParameterOwners = CameraAsset->ParameterOwners;
+	ensure(PropertyDescs.Num() == ParameterOwners.Num());
+
+	FCameraParameterOverrideHelper Helper(*OutVariableTable, *OutContextDataTable);
+
+	for (int32 Index = 0, MaxIndex = FMath::Min(PropertyDescs.Num(), ParameterOwners.Num()); Index < MaxIndex; ++Index)
+	{
+		const FPropertyBagPropertyDesc& PropertyDesc(PropertyDescs[Index]);
+		if (!CameraReference.IsParameterOverriden(PropertyDesc.ID))
+		{
+			continue;
+		}
+
+		const UCameraRigAsset* CameraRig = ParameterOwners[Index].Get();
+		if (!ensure(CameraRig))
+		{
+			continue;
+		}
+
+		TValueOrError<FStructView, EPropertyBagResult> ParameterValueOrError = CameraParameters.GetValueStruct(PropertyDesc);
+		Helper.ApplyParameterOverride(CameraRig, PropertyDesc.ID, ParameterValueOrError, bDrivenOverridesOnly);
+	}
+}
+
+}  // namespace UE::Cameras
+
diff --git a/Engine/Plugins/Cameras/GameplayCameras/Source/GameplayCameras/Private/Helpers/CameraAssetParameterOverrideEvaluator.h b/Engine/Plugins/Cameras/GameplayCameras/Source/GameplayCameras/Private/Helpers/CameraAssetParameterOverrideEvaluator.h
new file mode 100644
--- /dev/null
+++ b/Engine/Plugins/Cameras/GameplayCameras/Source/GameplayCameras/Private/Helpers/CameraAssetParameterOverrideEvaluator.h
@@ -0,0 +1,51 @@
+// Copyright Epic Games, Inc. All Rights Reserved.
+
+#pragma once
+
+struct FCameraAssetReference;
+
+namespace UE::Cameras
+{
+
+class FCameraContextDataTable;
+class FCameraVariableTable;
+
+/**
+ * Utility class for applying interface parameter overrides to a camera rig via a
+ * given variable table.
+ */
+class FCameraAssetParameterOverrideEvaluator
+{
+public:
+
+	/** Creates a new parameter override evaluator. */
+	FCameraAssetParameterOverrideEvaluator(const FCameraAssetReference& InCameraReference);
+
+	/** 
+	 * Applies override values to the given variable table.
+	 *
+	 * @param OutVariableTable      The variable table in which to set the override values.
+	 * @param bDrivenOverridesOnly  Whether only overrides driven by variables should be applied.
+	 */
+	void ApplyParameterOverrides(FCameraVariableTable& OutVariableTable, bool bDrivenOverridesOnly = false);
+
+	/** 
+	 * Applies override values to the given variable and context data tables.
+	 *
+	 * @param OutVariableTable      The variable table in which to set the override values.
+	 * @param OutContextDataTable   The context data table in which to set the override values.
+	 * @param bDrivenOverridesOnly  Whether only overrides driven by variables should be applied.
+	 */
+	void ApplyParameterOverrides(FCameraVariableTable& OutVariableTable, FCameraContextDataTable& OutContextDataTable, bool bDrivenOverridesOnly = false);
+
+private:
+
+	void ApplyParameterOverrides(FCameraVariableTable* OutVariableTable, FCameraContextDataTable* OutContextDataTable, bool bDrivenOverridesOnly);
+
+private:
+
+	const FCameraAssetReference& CameraReference;
+};
+
+}  // namespace UE::Cameras
+
diff --git a/Engine/Plugins/Cameras/GameplayCameras/Source/GameplayCameras/Private/Helpers/CameraParameterOverrideHelper.cpp b/Engine/Plugins/Cameras/GameplayCameras/Source/GameplayCameras/Private/Helpers/CameraParameterOverrideHelper.cpp
new file mode 100644
--- /dev/null
+++ b/Engine/Plugins/Cameras/GameplayCameras/Source/GameplayCameras/Private/Helpers/CameraParameterOverrideHelper.cpp
@@ -0,0 +1,135 @@
+// Copyright Epic Games, Inc. All Rights Reserved.
+
+#include "Helpers/CameraParameterOverrideHelper.h"
+
+#include "Core/CameraContextDataTable.h"
+#include "Core/CameraParameters.h"
+#include "Core/CameraRigAsset.h"
+#include "Core/CameraVariableTable.h"
+
+namespace UE::Cameras
+{
+
+namespace Internal
+{
+
+template<typename ParameterType>
+void ApplyBlendableParameterOverride(
+		const UCameraRigAsset* CameraRig,
+		const UCameraRigBlendableParameter* BlendableParameter,
+		const ParameterType& ParameterValue,
+		FCameraVariableTable& VariableTable,
+		bool bDrivenOverridesOnly)
+{
+	using ValueType = typename ParameterType::ValueType;
+
+	if (!ensure(BlendableParameter))
+	{
+		return;
+	}
+
+	if (!BlendableParameter->PrivateVariable)
+	{
+		// Ignore un-built parameter overrides in the editor since the user could have just added
+		// an override while PIE is running. They need to hit the Build button for the override
+		// to apply.
+		// Outside of the editor, report this as an error.
+#if !WITH_EDITOR
+		UE_LOG(LogCameraSystem, Error,
+				TEXT("Invalid blendable parameter override '%s' in camera rig '%s'. Was it built/cooked?"),
+				*BlendableParameter->InterfaceParameterName,
+				*GetPathNameSafe(CameraRig));
+#endif
+		return;
+	}
+
+	FCameraVariableID ParameterVariableID(BlendableParameter->PrivateVariable->GetVariableID());
+
+	if (ParameterValue.Variable != nullptr)
+	{
+		// The override is driven by a variable... read its value and set it as the value for the
+		// prefab's variable. Basically, we forward the value from one variable to the next.
+		FCameraVariableDefinition OverrideDefinition(ParameterValue.Variable->GetVariableDefinition());
+
+		const ValueType OverrideValue = VariableTable.GetValue<ValueType>(
+				OverrideDefinition.VariableID, ParameterValue.Variable->GetDefaultValue());
+		VariableTable.SetValue<ValueType>(ParameterVariableID, OverrideValue);
+	}
+	else if (!bDrivenOverridesOnly)
+	{
+		// The override is a fixed value. Just set that on the prefab's variable.
+		VariableTable.SetValue<ValueType>(ParameterVariableID, ParameterValue.Value);
+	}
+}
+
+void ApplyDataParameterOverride(
+		const UCameraRigAsset* CameraRig,
+		const UCameraRigDataParameter* DataParameter,
+		const FStructView& ParameterValue,
+		FCameraContextDataTable& ContextDataTable)
+{
+	if (!ensure(DataParameter))
+	{
+		return;
+	}
+
+	if (!DataParameter->PrivateDataID)
+	{
+#if !WITH_EDITOR
+		UE_LOG(LogCameraSystem, Error,
+				TEXT("Invalid data parameter override '%s' in camera rig '%s'. Was it built/cooked?"),
+				*DataParameter->InterfaceParameterName,
+				*GetPathNameSafe(CameraRig));
+		return;
+#endif
+	}
+
+	// Write the override value into the context data table.
+	FCameraContextDataID ParameterDataID = DataParameter->PrivateDataID;
+	ContextDataTable.SetData(ParameterDataID, ParameterValue);
+}
+
+}  // namespace Internal
+
+FCameraParameterOverrideHelper::FCameraParameterOverrideHelper(FCameraVariableTable& OutVariableTable, FCameraContextDataTable& OutContextDataTable)
+	: VariableTable(OutVariableTable)
+	, ContextDataTable(OutContextDataTable)
+{
+}
+
+void FCameraParameterOverrideHelper::ApplyParameterOverride(
+		const UCameraRigAsset* CameraRig,
+		const FGuid& ParameterGuid,
+		TValueOrError<FStructView, EPropertyBagResult> ParameterValueOrError,
+		bool bDrivenOverridesOnly)
+{
+	using namespace Internal;
+
+	if (ensureMsgf(
+				ParameterValueOrError.HasValue() && !ParameterValueOrError.HasError(),
+				TEXT("Camera parameter has no valid value! Error: %s"),
+				*UEnum::GetValueAsString(ParameterValueOrError.GetError())))
+	{
+		const FStructView& ParameterValue = ParameterValueOrError.GetValue();
+		const UScriptStruct* ParameterType = ParameterValue.GetScriptStruct();
+
+		// Check if this is a blendable parameter or a data parameter.
+#define UE_CAMERA_VARIABLE_FOR_TYPE(ValueType, ValueName)\
+		if (ParameterType == F##ValueName##CameraParameter::StaticStruct())\
+		{\
+			const UCameraRigBlendableParameter* BlendableParameter = CameraRig->Interface.FindBlendableParameterByGuid(ParameterGuid);\
+			const F##ValueName##CameraParameter& TypedParameterValue = ParameterValue.Get<F##ValueName##CameraParameter>();\
+			ApplyBlendableParameterOverride(CameraRig, BlendableParameter, TypedParameterValue, VariableTable, bDrivenOverridesOnly);\
+		}\
+		else
+		UE_CAMERA_VARIABLE_FOR_ALL_TYPES()
+#undef UE_CAMERA_VARIABLE_FOR_TYPE
+		{
+			const UCameraRigDataParameter* DataParameter = CameraRig->Interface.FindDataParameterByGuid(ParameterGuid);
+			ApplyDataParameterOverride(CameraRig, DataParameter, ParameterValue, ContextDataTable);
+		}
+	}
+}
+
+}  // namespace UE::Cameras
+
