Cameras: big pass on parameterization of camera rigs

diff --git a/Engine/Plugins/Cameras/GameplayCameras/Source/GameplayCameras/Private/Helpers/CameraParameterOverrideHelper.h b/Engine/Plugins/Cameras/GameplayCameras/Source/GameplayCameras/Private/Helpers/CameraParameterOverrideHelper.h
new file mode 100644
--- /dev/null
+++ b/Engine/Plugins/Cameras/GameplayCameras/Source/GameplayCameras/Private/Helpers/CameraParameterOverrideHelper.h
@@ -0,0 +1,38 @@
+// Copyright Epic Games, Inc. All Rights Reserved.
+
+#pragma once
+
+#include "StructUtils/PropertyBag.h"
+
+class UCameraRigAsset;
+class UCameraRigBlendableParameter;
+class UCameraRigDataParameter;
+
+namespace UE::Cameras
+{
+
+class FCameraContextDataTable;
+class FCameraVariableTable;
+
+/**
+ * A helper class for applying camera rig parameter overrides from a property bag, such as with
+ * camera asset references and camera rig asset references.
+ */
+struct FCameraParameterOverrideHelper
+{
+	FCameraParameterOverrideHelper(FCameraVariableTable& OutVariableTable, FCameraContextDataTable& OutContextDataTable);
+
+	void ApplyParameterOverride(
+			const UCameraRigAsset* CameraRig,
+			const FGuid& ParameterGuid,
+			TValueOrError<FStructView, EPropertyBagResult> ParameterValueOrError,
+			bool bDrivenOverridesOnly);
+
+private:
+
+	FCameraVariableTable& VariableTable;
+	FCameraContextDataTable& ContextDataTable;
+};
+
+}  // namespace UE::Cameras
+
diff --git a/Engine/Plugins/Cameras/GameplayCameras/Source/GameplayCameras/Private/Helpers/CameraRigParameterOverrideEvaluator.cpp b/Engine/Plugins/Cameras/GameplayCameras/Source/GameplayCameras/Private/Helpers/CameraRigParameterOverrideEvaluator.cpp
new file mode 100644
--- /dev/null
+++ b/Engine/Plugins/Cameras/GameplayCameras/Source/GameplayCameras/Private/Helpers/CameraRigParameterOverrideEvaluator.cpp
@@ -0,0 +1,52 @@
+// Copyright Epic Games, Inc. All Rights Reserved.
+
+#include "Helpers/CameraRigParameterOverrideEvaluator.h"
+
+#include "Core/CameraRigAsset.h"
+#include "Core/CameraRigAssetReference.h"
+#include "Helpers/CameraParameterOverrideHelper.h"
+
+namespace UE::Cameras
+{
+
+FCameraRigParameterOverrideEvaluator::FCameraRigParameterOverrideEvaluator(const FCameraRigAssetReference& InCameraRigReference)
+	: CameraRigReference(InCameraRigReference)
+{
+}
+
+void FCameraRigParameterOverrideEvaluator::ApplyParameterOverrides(FCameraVariableTable& OutVariableTable, bool bDrivenOverridesOnly)
+{
+	ApplyParameterOverrides(&OutVariableTable, nullptr, bDrivenOverridesOnly);
+}
+
+void FCameraRigParameterOverrideEvaluator::ApplyParameterOverrides(FCameraVariableTable& OutVariableTable, FCameraContextDataTable& OutContextDataTable, bool bDrivenOverridesOnly)
+{
+	ApplyParameterOverrides(&OutVariableTable, &OutContextDataTable, bDrivenOverridesOnly);
+}
+
+void FCameraRigParameterOverrideEvaluator::ApplyParameterOverrides(FCameraVariableTable* OutVariableTable, FCameraContextDataTable* OutContextDataTable, bool bDrivenOverridesOnly)
+{
+	using namespace Internal;
+
+	check(OutVariableTable);
+
+	const UCameraRigAsset* CameraRig = CameraRigReference.GetCameraRig();
+	const FInstancedPropertyBag& CameraRigParameters = CameraRigReference.GetParameters();
+	const UPropertyBag* CameraRigParametersStruct = CameraRigParameters.GetPropertyBagStruct();
+
+	FCameraParameterOverrideHelper Helper(*OutVariableTable, *OutContextDataTable);
+
+	for (const FPropertyBagPropertyDesc& PropertyDesc : CameraRigParametersStruct->GetPropertyDescs())
+	{
+		if (!CameraRigReference.IsParameterOverriden(PropertyDesc.ID))
+		{
+			continue;
+		}
+
+		TValueOrError<FStructView, EPropertyBagResult> ParameterValueOrError = CameraRigParameters.GetValueStruct(PropertyDesc);
+		Helper.ApplyParameterOverride(CameraRig, PropertyDesc.ID, ParameterValueOrError, bDrivenOverridesOnly);
+	}
+}
+
+}  // namespace UE::Cameras
+
diff --git a/Engine/Plugins/Cameras/GameplayCameras/Source/GameplayCameras/Private/Helpers/CameraRigParameterOverrideEvaluator.h b/Engine/Plugins/Cameras/GameplayCameras/Source/GameplayCameras/Private/Helpers/CameraRigParameterOverrideEvaluator.h
new file mode 100644
--- /dev/null
+++ b/Engine/Plugins/Cameras/GameplayCameras/Source/GameplayCameras/Private/Helpers/CameraRigParameterOverrideEvaluator.h
@@ -0,0 +1,51 @@
+// Copyright Epic Games, Inc. All Rights Reserved.
+
+#pragma once
+
+struct FCameraRigAssetReference;
+
+namespace UE::Cameras
+{
+
+class FCameraContextDataTable;
+class FCameraVariableTable;
+
+/**
+ * Utility class for applying interface parameter overrides to a camera rig via a
+ * given variable table.
+ */
+class FCameraRigParameterOverrideEvaluator
+{
+public:
+
+	/** Creates a new parameter override evaluator. */
+	FCameraRigParameterOverrideEvaluator(const FCameraRigAssetReference& InCameraRigReference);
+
+	/** 
+	 * Applies override values to the given variable table.
+	 *
+	 * @param OutVariableTable      The variable table in which to set the override values.
+	 * @param bDrivenOverridesOnly  Whether only overrides driven by variables should be applied.
+	 */
+	void ApplyParameterOverrides(FCameraVariableTable& OutVariableTable, bool bDrivenOverridesOnly = false);
+
+	/** 
+	 * Applies override values to the given variable and context data tables.
+	 *
+	 * @param OutVariableTable      The variable table in which to set the override values.
+	 * @param OutContextDataTable   The context data table in which to set the override values.
+	 * @param bDrivenOverridesOnly  Whether only overrides driven by variables should be applied.
+	 */
+	void ApplyParameterOverrides(FCameraVariableTable& OutVariableTable, FCameraContextDataTable& OutContextDataTable, bool bDrivenOverridesOnly = false);
+
+private:
+
+	void ApplyParameterOverrides(FCameraVariableTable* OutVariableTable, FCameraContextDataTable* OutContextDataTable, bool bDrivenOverridesOnly);
+
+private:
+
+	const FCameraRigAssetReference& CameraRigReference;
+};
+
+}  // namespace UE::Cameras
+
diff --git a/Engine/Plugins/Cameras/GameplayCameras/Source/GameplayCameras/Private/Nodes/Common/CameraRigCameraNode.cpp b/Engine/Plugins/Cameras/GameplayCameras/Source/GameplayCameras/Private/Nodes/Common/CameraRigCameraNode.cpp
--- a/Engine/Plugins/Cameras/GameplayCameras/Source/GameplayCameras/Private/Nodes/Common/CameraRigCameraNode.cpp
+++ b/Engine/Plugins/Cameras/GameplayCameras/Source/GameplayCameras/Private/Nodes/Common/CameraRigCameraNode.cpp
@@ -6,7 +6,7 @@
 #include "Core/CameraNodeEvaluator.h"
 #include "Core/CameraRigAsset.h"
 #include "Core/CameraRigBuildContext.h"
-#include "Core/CameraRigParameterOverrideEvaluator.h"
+#include "Helpers/CameraRigParameterOverrideEvaluator.h"
 
 #include UE_INLINE_GENERATED_CPP_BY_NAME(CameraRigCameraNode)
 
@@ -42,7 +42,7 @@ void FCameraRigCameraNodeEvaluator::OnBuild(const FCameraNodeEvaluatorBuildParam
 void FCameraRigCameraNodeEvaluator::OnInitialize(const FCameraNodeEvaluatorInitializeParams& Params, FCameraNodeEvaluationResult& OutResult)
 {
 	// Apply overrides right away.
-	ApplyParameterOverrides(OutResult.VariableTable, false);
+	ApplyParameterOverrides(OutResult.VariableTable, OutResult.ContextDataTable, false);
 }
 
 void FCameraRigCameraNodeEvaluator::OnUpdateParameters(const FCameraBlendedParameterUpdateParams& Params, FCameraBlendedParameterUpdateResult& OutResult)
@@ -71,6 +71,17 @@ void FCameraRigCameraNodeEvaluator::ApplyParameterOverrides(FCameraVariableTable
 	}
 }
 
+void FCameraRigCameraNodeEvaluator::ApplyParameterOverrides(FCameraVariableTable& OutVariableTable, FCameraContextDataTable& OutContextDataTable, bool bDrivenOnly)
+{
+	if (bApplyParameterOverrides)
+	{
+		const UCameraRigCameraNode* PrefabNode = GetCameraNodeAs<UCameraRigCameraNode>();
+
+		FCameraRigParameterOverrideEvaluator OverrideEvaluator(PrefabNode->CameraRigReference);
+		OverrideEvaluator.ApplyParameterOverrides(OutVariableTable, OutContextDataTable, bDrivenOnly);
+	}
+}
+
 bool FCameraRigCameraNodeEvaluator::IsApplyingParameterOverrides() const
 {
 	return bApplyParameterOverrides;
@@ -81,96 +92,6 @@ void FCameraRigCameraNodeEvaluator::SetApplyParameterOverrides(bool bShouldApply
 	bApplyParameterOverrides = bShouldApply;
 }
 
-namespace Internal
-{
-
-struct FCameraRigCameraNodeBuilder
-{
-	UCameraRigCameraNode* CameraNode;
-	FCameraRigBuildContext& BuildContext;
-
-	FCameraRigCameraNodeBuilder(UCameraRigCameraNode* InCameraNode, FCameraRigBuildContext& InBuildContext)
-		: CameraNode(InCameraNode)
-		, BuildContext(InBuildContext)
-	{}
-
-	void Setup()
-	{
-		// Build a map matching each of our inner camera rig's interface parameter to its Guid.
-		ParametersByGuid.Reset();
-		const UCameraRigAsset* CameraRig = CameraNode->CameraRigReference.GetCameraRig();
-		for (TObjectPtr<UCameraRigInterfaceParameter> InterfaceParameter : CameraRig->Interface.InterfaceParameters)
-		{
-			ParametersByGuid.Add(InterfaceParameter->Guid, InterfaceParameter);
-		}
-	}
-
-	template<typename ParameterOverrideType>
-	void BuildCameraRigParameterOverride(ParameterOverrideType& ParameterOverride);
-
-private:
-
-	const UCameraRigInterfaceParameter* FindInterfaceParameter(const FGuid& InterfaceParameterGuid)
-	{
-		if (UCameraRigInterfaceParameter** FoundItem = ParametersByGuid.Find(InterfaceParameterGuid))
-		{
-			return *FoundItem;
-		}
-		return nullptr;
-	}
-
-private:
-
-	TMap<FGuid, UCameraRigInterfaceParameter*> ParametersByGuid;
-};
-
-template<typename ParameterOverrideType>
-void FCameraRigCameraNodeBuilder::BuildCameraRigParameterOverride(ParameterOverrideType& ParameterOverride)
-{
-	const UCameraRigAsset* CameraRig = CameraNode->CameraRigReference.GetCameraRig();
-
-	// Each parameter override should point to a valid interface parameter on the inner rig, via its Guid.
-	const UCameraRigInterfaceParameter* InterfaceParameter = FindInterfaceParameter(ParameterOverride.InterfaceParameterGuid);
-	if (!InterfaceParameter)
-	{
-		BuildContext.BuildLog.AddMessage(EMessageSeverity::Error, CameraNode,
-				FText::Format(
-					LOCTEXT("MissingInterfaceParameter", "No camera rig interface parameter named '{0}' exists on '{1}'."),
-					FText::FromString(ParameterOverride.InterfaceParameterName),
-					FText::FromString(GetNameSafe(CameraRig))));
-		return;
-	}
-
-	// The inner rig's interface parameter should have been built, i.e. it should have a private camera variable
-	// assigned for driving its value.
-	if (!InterfaceParameter->PrivateVariable)
-	{
-		BuildContext.BuildLog.AddMessage(EMessageSeverity::Error, CameraNode,
-				FText::Format(
-					LOCTEXT("UnbuiltInterfaceParameter", "Camera rig interface parameter '{0}' was not built correctly on '{1}'."),
-					FText::FromString(ParameterOverride.InterfaceParameterName),
-					FText::FromString(GetNameSafe(CameraRig))));
-		return;
-	}
-
-	// The inner rig's interface parameter is driven by this private variable. Let's remember its Guid so we
-	// can override its value in the variable table at runtime.
-	ParameterOverride.PrivateVariableGuid = InterfaceParameter->PrivateVariable->GetGuid();
-	// Update the last known name for this interface parameter.
-	ParameterOverride.InterfaceParameterName = InterfaceParameter->InterfaceParameterName;
-
-	// The build process automatically gathers variables that drive amera parameters on a camera node, but
-	// nothing else for now. We therefore need to help it out by manually reporting the variables that
-	// drive our parameter overrides.
-	FCameraVariableTableAllocationInfo& VariableTableAllocationInfo = BuildContext.AllocationInfo.VariableTableInfo;
-	if (ParameterOverride.Value.Variable)
-	{
-		VariableTableAllocationInfo.VariableDefinitions.Add(ParameterOverride.Value.Variable->GetVariableDefinition());
-	}
-}
-
-}  // namespace Internal
-
 }  // namespace UE::Cameras
 
 void UCameraRigCameraNode::OnPreBuild(FCameraBuildLog& BuildLog)
@@ -181,12 +102,14 @@ void UCameraRigCameraNode::OnPreBuild(FCameraBuildLog& BuildLog)
 	{
 		CameraRig->BuildCameraRig(BuildLog);
 	}
+
+	// Make sure the property bag of the camera rig reference is up to date.
+	CameraRigReference.RebuildParametersIfNeeded();
 }
 
 void UCameraRigCameraNode::OnBuild(FCameraRigBuildContext& BuildContext)
 {
 	using namespace UE::Cameras;
-	using namespace UE::Cameras::Internal;
 
 	UCameraRigAsset* CameraRig = CameraRigReference.GetCameraRig();
 	if (!CameraRig)
@@ -199,22 +122,67 @@ void UCameraRigCameraNode::OnBuild(FCameraRigBuildContext& BuildContext)
 	// Whatever allocations our inner camera rig needs for its evaluators and
 	// their camera variables, we add that to our camera rig's allocation info.
 	BuildContext.AllocationInfo.Append(CameraRig->AllocationInfo);
+}
 
-	// Next, we set things up for the runtime. Mostly, we want to get the camera variable 
-	// Guids that we need to write the override values to.
-	FCameraRigCameraNodeBuilder InternalBuilder(this, BuildContext);
-	InternalBuilder.Setup();
+void UCameraRigCameraNode::GetCustomCameraNodeParameters(FCustomCameraNodeParameterInfos& OutParameterInfos)
+{
+	const FInstancedPropertyBag& ParameterOverrides = CameraRigReference.GetParameters();
+	const UPropertyBag* ParameterOverridesStruct = ParameterOverrides.GetPropertyBagStruct();
+	if (!ParameterOverridesStruct)
+	{
+		return;
+	}
 
-	FCameraRigParameterOverrides& ParameterOverrides = CameraRigReference.GetParameterOverrides();
+	for (const FPropertyBagPropertyDesc& PropertyDesc : ParameterOverridesStruct->GetPropertyDescs())
+	{
+		switch (PropertyDesc.ValueType)
+		{
+			case EPropertyBagPropertyType::Struct:
+				{
 #define UE_CAMERA_VARIABLE_FOR_TYPE(ValueType, ValueName)\
-	{\
-		for (F##ValueName##CameraRigParameterOverride& ParameterOverride : ParameterOverrides.Get##ValueName##Overrides())\
-		{\
-			InternalBuilder.BuildCameraRigParameterOverride(ParameterOverride);\
-		}\
-	}
-UE_CAMERA_VARIABLE_FOR_ALL_TYPES()
+					if (PropertyDesc.ValueTypeObject == F##ValueName##CameraParameter::StaticStruct())\
+					{\
+						using CameraParameterType = F##ValueName##CameraParameter;\
+						TValueOrError<CameraParameterType*, EPropertyBagResult> PropertyValue =\
+						ParameterOverrides.GetValueStruct<CameraParameterType>(PropertyDesc);\
+						if (ensure(PropertyValue.HasValue() && !PropertyValue.HasError()))\
+						{\
+							CameraParameterType* CameraParameter = PropertyValue.GetValue();\
+							check(CameraParameter);\
+							const uint8* DefaultValuePtr = reinterpret_cast<const uint8*>(&CameraParameter->Value);\
+							OutParameterInfos.AddBlendableParameter(\
+									PropertyDesc.Name, ECameraVariableType::ValueName, DefaultValuePtr, &CameraParameter->Variable);\
+						}\
+					}\
+					else
+					UE_CAMERA_VARIABLE_FOR_ALL_TYPES()
 #undef UE_CAMERA_VARIABLE_FOR_TYPE
+					{
+						const UScriptStruct* DataType = CastChecked<const UScriptStruct>(PropertyDesc.ValueTypeObject);
+						OutParameterInfos.AddDataParameter(PropertyDesc.Name, ECameraContextDataType::Struct, DataType, nullptr);
+					}
+				}
+				break;
+			case EPropertyBagPropertyType::Name:
+				OutParameterInfos.AddDataParameter(PropertyDesc.Name, ECameraContextDataType::Name, nullptr, nullptr);
+				break;
+			case EPropertyBagPropertyType::String:
+				OutParameterInfos.AddDataParameter(PropertyDesc.Name, ECameraContextDataType::String, nullptr, nullptr);
+				break;
+			case EPropertyBagPropertyType::Enum:
+				{
+					const UEnum* EnumType = CastChecked<const UEnum>(PropertyDesc.ValueTypeObject);
+					OutParameterInfos.AddDataParameter(PropertyDesc.Name, ECameraContextDataType::Enum, EnumType, nullptr);
+				}
+				break;
+			case EPropertyBagPropertyType::Object:
+				OutParameterInfos.AddDataParameter(PropertyDesc.Name, ECameraContextDataType::Object, nullptr, nullptr);
+				break;
+			case EPropertyBagPropertyType::Class:
+				OutParameterInfos.AddDataParameter(PropertyDesc.Name, ECameraContextDataType::Object, nullptr, nullptr);
+				break;
+		}
+	}
 }
 
 FCameraNodeEvaluatorPtr UCameraRigCameraNode::OnBuildEvaluator(FCameraNodeEvaluatorBuilder& Builder) const
@@ -223,26 +191,5 @@ FCameraNodeEvaluatorPtr UCameraRigCameraNode::OnBuildEvaluator(FCameraNodeEvalua
 	return Builder.BuildEvaluator<FCameraRigCameraNodeEvaluator>();
 }
 
-void UCameraRigCameraNode::PostLoad()
-{
-	Super::PostLoad();
-
-	if (CameraRig_DEPRECATED)
-	{
-		CameraRigReference.SetCameraRig(CameraRig_DEPRECATED);
-		CameraRig_DEPRECATED = nullptr;
-	}
-
-	FCameraRigParameterOverrides& ParameterOverrides = CameraRigReference.GetParameterOverrides();
-#define UE_CAMERA_VARIABLE_FOR_TYPE(ValueType, ValueName)\
-	if (ValueName##Overrides_DEPRECATED.Num() > 0)\
-	{\
-		ParameterOverrides.AppendParameterOverrides<F##ValueName##CameraRigParameterOverride>(ValueName##Overrides_DEPRECATED);\
-		ValueName##Overrides_DEPRECATED.Reset();\
-	}
-UE_CAMERA_VARIABLE_FOR_ALL_TYPES()
-#undef UE_CAMERA_VARIABLE_FOR_TYPE
-}
-
 #undef LOCTEXT_NAMESPACE
 
diff --git a/Engine/Plugins/Cameras/GameplayCameras/Source/GameplayCameras/Private/Tests/CameraRigAssetBuilderTests.cpp b/Engine/Plugins/Cameras/GameplayCameras/Source/GameplayCameras/Private/Tests/CameraRigAssetBuilderTests.cpp
--- a/Engine/Plugins/Cameras/GameplayCameras/Source/GameplayCameras/Private/Tests/CameraRigAssetBuilderTests.cpp
+++ b/Engine/Plugins/Cameras/GameplayCameras/Source/GameplayCameras/Private/Tests/CameraRigAssetBuilderTests.cpp
@@ -67,12 +67,12 @@ bool FCameraRigAssetBuilderSimpleParameterTest::RunTest(const FString& Parameter
 				.Pin(OffsetNode)
 				.Done()
 			.Done()
-		.ExposeParameter(TEXT("Test"), OffsetNode, GET_MEMBER_NAME_CHECKED(UOffsetCameraNode, TranslationOffset))
+		.AddBlendableParameter(TEXT("Test"), ECameraVariableType::Vector3d, OffsetNode, GET_MEMBER_NAME_CHECKED(UOffsetCameraNode, TranslationOffset))
 		.Get();
 
 	CameraRig->BuildCameraRig();
 
-	UCameraRigInterfaceParameter* Parameter = CameraRig->Interface.InterfaceParameters[0];
+	UCameraRigBlendableParameter* Parameter = CameraRig->Interface.BlendableParameters[0];
 	UTEST_EQUAL("Test parameter", Parameter->InterfaceParameterName, TEXT("Test"));
 	UTEST_NOT_NULL("Test parameter variable", Parameter->PrivateVariable.Get());
 	UTEST_EQUAL("Test parameter variable name", Parameter->PrivateVariable->GetName(), "Override_SimpleTest_Test");
@@ -81,73 +81,6 @@ bool FCameraRigAssetBuilderSimpleParameterTest::RunTest(const FString& Parameter
 	return true;
 }
 
-IMPLEMENT_SIMPLE_AUTOMATION_TEST(FCameraRigAssetBuilderReassignParameterTest, "System.Engine.GameplayCameras.CameraRigAssetBuilder.ReassignParameter", EAutomationTestFlags::EditorContext | EAutomationTestFlags::EngineFilter)
-bool FCameraRigAssetBuilderReassignParameterTest::RunTest(const FString& Parameters)
-{
-	using namespace UE::Cameras::Test;
-
-	UOffsetCameraNode* OffsetNode = nullptr;
-	ULensParametersCameraNode* LensParametersNode = nullptr;
-	UCameraRigAsset* CameraRig = FCameraRigAssetTestBuilder(TEXT("SimpleTest"))
-		.MakeRootNode<UArrayCameraNode>()
-			.AddChild<UOffsetCameraNode>(&UArrayCameraNode::Children)
-				.Pin(OffsetNode)
-				.Done()
-			.AddChild<ULensParametersCameraNode>(&UArrayCameraNode::Children)
-				.Pin(LensParametersNode)
-				.Done()
-			.Done()
-		.ExposeParameter(TEXT("Test1"), OffsetNode, GET_MEMBER_NAME_CHECKED(UOffsetCameraNode, TranslationOffset))
-		.ExposeParameter(TEXT("Test2"), LensParametersNode, GET_MEMBER_NAME_CHECKED(ULensParametersCameraNode, FocalLength))
-		.ExposeParameter(TEXT("Test3"), LensParametersNode, GET_MEMBER_NAME_CHECKED(ULensParametersCameraNode, Aperture))
-		.Get();
-
-	UCameraRigInterfaceParameter* Test1Parameter = CameraRig->Interface.InterfaceParameters[0];
-	UCameraRigInterfaceParameter* Test2Parameter = CameraRig->Interface.InterfaceParameters[1];
-	UCameraRigInterfaceParameter* Test3Parameter = CameraRig->Interface.InterfaceParameters[2];
-
-	CameraRig->BuildCameraRig();
-
-	{
-		UTEST_EQUAL_EXPR(Test1Parameter->PrivateVariable->GetName(), "Override_SimpleTest_Test1");
-		UTEST_TRUE_EXPR(Test1Parameter->PrivateVariable->IsA<UVector3dCameraVariable>());
-		UTEST_EQUAL_EXPR((UCameraVariableAsset*)OffsetNode->TranslationOffset.Variable, Test1Parameter->PrivateVariable.Get());
-
-		UTEST_EQUAL_EXPR(Test2Parameter->PrivateVariable->GetName(), "Override_SimpleTest_Test2");
-		UTEST_TRUE_EXPR(Test2Parameter->PrivateVariable->IsA<UFloatCameraVariable>());
-		UTEST_EQUAL_EXPR((UCameraVariableAsset*)LensParametersNode->FocalLength.Variable, Test2Parameter->PrivateVariable.Get());
-
-		UTEST_EQUAL_EXPR(Test3Parameter->PrivateVariable->GetName(), "Override_SimpleTest_Test3");
-		UTEST_TRUE_EXPR(Test3Parameter->PrivateVariable->IsA<UFloatCameraVariable>());
-		UTEST_EQUAL_EXPR((UCameraVariableAsset*)LensParametersNode->Aperture.Variable, Test3Parameter->PrivateVariable.Get());
-	}
-
-	Test1Parameter->Target = LensParametersNode;
-	Test1Parameter->TargetPropertyName = GET_MEMBER_NAME_CHECKED(ULensParametersCameraNode, FocalLength);
-	Test2Parameter->Target = LensParametersNode;
-	Test2Parameter->TargetPropertyName = GET_MEMBER_NAME_CHECKED(ULensParametersCameraNode, Aperture);
-	Test3Parameter->Target = OffsetNode;
-	Test3Parameter->TargetPropertyName = GET_MEMBER_NAME_CHECKED(UOffsetCameraNode, TranslationOffset);
-
-	CameraRig->BuildCameraRig();
-
-	{
-		UTEST_EQUAL_EXPR(Test1Parameter->PrivateVariable->GetName(), "Override_SimpleTest_Test1");
-		UTEST_TRUE_EXPR(Test1Parameter->PrivateVariable->IsA<UFloatCameraVariable>());
-		UTEST_EQUAL_EXPR((UCameraVariableAsset*)LensParametersNode->FocalLength.Variable, Test1Parameter->PrivateVariable.Get());
-
-		UTEST_EQUAL_EXPR(Test2Parameter->PrivateVariable->GetName(), "Override_SimpleTest_Test2");
-		UTEST_TRUE_EXPR(Test2Parameter->PrivateVariable->IsA<UFloatCameraVariable>());
-		UTEST_EQUAL_EXPR((UCameraVariableAsset*)LensParametersNode->Aperture.Variable, Test2Parameter->PrivateVariable.Get());
-
-		UTEST_EQUAL_EXPR(Test3Parameter->PrivateVariable->GetName(), "Override_SimpleTest_Test3");
-		UTEST_TRUE_EXPR(Test3Parameter->PrivateVariable->IsA<UVector3dCameraVariable>());
-		UTEST_EQUAL_EXPR((UCameraVariableAsset*)OffsetNode->TranslationOffset.Variable, Test3Parameter->PrivateVariable.Get());
-	}
-
-	return true;
-}
-
 IMPLEMENT_SIMPLE_AUTOMATION_TEST(FCameraRigAssetBuilderDrivenOverridesTest, "System.Engine.GameplayCameras.CameraRigAssetBuilder.DrivenOverrides", EAutomationTestFlags::EditorContext | EAutomationTestFlags::EngineFilter)
 bool FCameraRigAssetBuilderDrivenOverridesTest::RunTest(const FString& Parameters)
 {
@@ -166,10 +99,12 @@ bool FCameraRigAssetBuilderDrivenOverridesTest::RunTest(const FString& Parameter
 				.SetParameter(&ULensParametersCameraNode::FocalLength, 20.f)
 				.Done()
 			.Done()
-		.ExposeParameter(TEXT("OffsetParam"), TEXT("Offset"), GET_MEMBER_NAME_CHECKED(UOffsetCameraNode, TranslationOffset))
-		.ExposeParameter(TEXT("FocalLengthParam"), TEXT("Lens"), GET_MEMBER_NAME_CHECKED(ULensParametersCameraNode, FocalLength))
+		.AddBlendableParameter(TEXT("OffsetParam"), ECameraVariableType::Vector3d, TEXT("Offset"), GET_MEMBER_NAME_CHECKED(UOffsetCameraNode, TranslationOffset))
+		.AddBlendableParameter(TEXT("FocalLengthParam"), ECameraVariableType::Float, TEXT("Lens"), GET_MEMBER_NAME_CHECKED(ULensParametersCameraNode, FocalLength))
 		.Get();
 
+	InnerCameraRig->BuildCameraRig();
+
 	// Make a camera rig that uses the previous one, with overrides on both the offset (now 15, 25, 35)
 	// and the focal length (now 25). Expose the offset further up as an interface parameter.
 	UCameraRigCameraNode* MiddlePrefabNode = nullptr;
@@ -181,20 +116,26 @@ bool FCameraRigAssetBuilderDrivenOverridesTest::RunTest(const FString& Parameter
 					{
 						Node->CameraRigReference.SetCameraRig(InnerCameraRig);
 
-						FCameraRigParameterOverrides& ParameterOverrides = Node->CameraRigReference.GetParameterOverrides();
+						FInstancedPropertyBag& ParameterOverrides = Node->CameraRigReference.GetParameters();
 
-						auto* OffsetParam = Registry->Get<UCameraRigInterfaceParameter>(TEXT("OffsetParam"));
-						auto& OffsetParamOverride = ParameterOverrides.FindOrAddParameterOverride<FVector3dCameraRigParameterOverride>(OffsetParam);
-						OffsetParamOverride.Value = FVector3d(15, 25, 35);
+						auto OffsetParamOverride = ParameterOverrides.GetValueStruct<FVector3dCameraParameter>("OffsetParam");
+						if (ensure(OffsetParamOverride.HasValue()))
+						{
+							OffsetParamOverride.GetValue()->Value = FVector3d(15, 25, 35);
+						}
 
-						auto* FocalLengthParam = Registry->Get<UCameraRigInterfaceParameter>(TEXT("FocalLengthParam"));
-						auto& FocalLengthParamOverride = ParameterOverrides.FindOrAddParameterOverride<FFloatCameraRigParameterOverride>(FocalLengthParam);
-						FocalLengthParamOverride.Value = 25.f;
+						auto FocalLengthParamOverride = ParameterOverrides.GetValueStruct<FFloatCameraParameter>("FocalLengthParam");
+						if (ensure(FocalLengthParamOverride.HasValue()))
+						{
+							FocalLengthParamOverride.GetValue()->Value = 25.f;
+						}
 					})
 			.Done()
-			.ExposeParameter(TEXT("MiddleOffsetParam"), MiddlePrefabNode, TEXT("OffsetParam"))
+			.AddBlendableParameter(TEXT("MiddleOffsetParam"), ECameraVariableType::Vector3d, MiddlePrefabNode, TEXT("OffsetParam"))
 		.Get();
 
+	MiddleCameraRig->BuildCameraRig();
+
 	// Make another camera rig that uses the previous one, which makes a total of 3 nesting levels of camera rigs.
 	// This level overrides the offset parameter some more (now 20, 50, 70).
 	UCameraRigCameraNode* OuterPrefabNode = nullptr;
@@ -206,18 +147,21 @@ bool FCameraRigAssetBuilderDrivenOverridesTest::RunTest(const FString& Parameter
 					{
 						Node->CameraRigReference.SetCameraRig(MiddleCameraRig);
 
-						FCameraRigParameterOverrides& ParameterOverrides = Node->CameraRigReference.GetParameterOverrides();
-						auto* MiddleOffsetParam = Registry->Get<UCameraRigInterfaceParameter>(TEXT("MiddleOffsetParam"));
-						auto& MiddleOffsetParamOverride = ParameterOverrides.FindOrAddParameterOverride<FVector3dCameraRigParameterOverride>(MiddleOffsetParam);
-						MiddleOffsetParamOverride.Value = FVector3d(20, 50, 70);
+						FInstancedPropertyBag& ParameterOverrides = Node->CameraRigReference.GetParameters();
+
+						auto MiddleOffsetParamOverride = ParameterOverrides.GetValueStruct<FVector3dCameraParameter>("MiddleOffsetParam");
+						if (ensure(MiddleOffsetParamOverride.HasValue()))
+						{
+							MiddleOffsetParamOverride.GetValue()->Value = FVector3d(20, 50, 70);
+						}
 					})
 			.Done()
 		.Get();
 
 	OuterCameraRig->BuildCameraRig();
 
-	UCameraRigInterfaceParameter* OffsetParam = InnerCameraRig->Interface.InterfaceParameters[0];
-	UCameraRigInterfaceParameter* FocalLengthParam = InnerCameraRig->Interface.InterfaceParameters[1];
+	UCameraRigBlendableParameter* OffsetParam = InnerCameraRig->Interface.BlendableParameters[0];
+	UCameraRigBlendableParameter* FocalLengthParam = InnerCameraRig->Interface.BlendableParameters[1];
 
 	UTEST_EQUAL_EXPR(OffsetParam->PrivateVariable->GetName(), "Override_InnerCameraRig_OffsetParam");
 	UTEST_EQUAL_EXPR(FocalLengthParam->PrivateVariable->GetName(), "Override_InnerCameraRig_FocalLengthParam");
@@ -238,46 +182,46 @@ bool FCameraRigAssetBuilderDrivenOverridesTest::RunTest(const FString& Parameter
 	// Test that the middle prefab node is driving the inner interface parameters, and that one of those
 	// overrides is in turn driven by the middle camera rig's interface parameter.
 	{
-		FCameraRigParameterOverrides& ParameterOverrides = MiddlePrefabNode->CameraRigReference.GetParameterOverrides();
+		FInstancedPropertyBag& ParameterOverrides = MiddlePrefabNode->CameraRigReference.GetParameters();
 
-		FVector3dCameraRigParameterOverride* OffsetParamOverride = 
-			ParameterOverrides.FindParameterOverride<FVector3dCameraRigParameterOverride>(OffsetParam->Guid);
+		const FPropertyBagPropertyDesc* OffsetParamDesc = ParameterOverrides.FindPropertyDescByName("OffsetParam");
+		FVector3dCameraParameter* OffsetParamOverride = ParameterOverrides.GetValueStruct<FVector3dCameraParameter>("OffsetParam").GetValue();
+		UTEST_NOT_NULL("OffsetParamDesc", OffsetParamDesc);
 		UTEST_NOT_NULL("OffsetParamOverride", OffsetParamOverride);
 
-		UTEST_EQUAL_EXPR(OffsetParamOverride->InterfaceParameterName, "OffsetParam");
-		UTEST_EQUAL_EXPR(OffsetParamOverride->PrivateVariableGuid, OffsetParam->PrivateVariable->GetGuid());
-		UTEST_EQUAL_EXPR(OffsetParamOverride->Value.Value, FVector3d(15, 25, 35));
+		UTEST_EQUAL_EXPR(OffsetParamDesc->ID, OffsetParam->GetGuid());
+		UTEST_EQUAL_EXPR(OffsetParamOverride->Value, FVector3d(15, 25, 35));
 
-		FFloatCameraRigParameterOverride* FocalLengthParamOverride =
-			ParameterOverrides.FindParameterOverride<FFloatCameraRigParameterOverride>(FocalLengthParam->Guid);
+		const FPropertyBagPropertyDesc* FocalLengthParamDesc = ParameterOverrides.FindPropertyDescByName("FocalLengthParam");
+		FFloatCameraParameter* FocalLengthParamOverride = ParameterOverrides.GetValueStruct<FFloatCameraParameter>("FocalLengthParam").GetValue();
+		UTEST_NOT_NULL("FocalLengthParamDesc", FocalLengthParamDesc);
 		UTEST_NOT_NULL("FocalLengthParamOverride", FocalLengthParamOverride);
 
-		UTEST_EQUAL_EXPR(FocalLengthParamOverride->InterfaceParameterName, "FocalLengthParam");
-		UTEST_EQUAL_EXPR(FocalLengthParamOverride->PrivateVariableGuid, FocalLengthParam->PrivateVariable->GetGuid());
-		UTEST_EQUAL_EXPR(FocalLengthParamOverride->Value.Value, 25.f);
+		UTEST_EQUAL_EXPR(FocalLengthParamDesc->ID, FocalLengthParam->GetGuid());
+		UTEST_EQUAL_EXPR(FocalLengthParamOverride->Value, 25.f);
 	}
 
-	UCameraRigInterfaceParameter* MiddleOffsetParam = MiddleCameraRig->Interface.InterfaceParameters[0];
+	UCameraRigBlendableParameter* MiddleOffsetParam = MiddleCameraRig->Interface.BlendableParameters[0];
 	{
-		FCameraRigParameterOverrides& ParameterOverrides = MiddlePrefabNode->CameraRigReference.GetParameterOverrides();
+		FInstancedPropertyBag& ParameterOverrides = MiddlePrefabNode->CameraRigReference.GetParameters();
 
-		FVector3dCameraRigParameterOverride* OffsetParamOverride = 
-			ParameterOverrides.FindParameterOverride<FVector3dCameraRigParameterOverride>(OffsetParam->Guid);
-		UTEST_EQUAL_EXPR((UCameraVariableAsset*)OffsetParamOverride->Value.Variable.Get(), MiddleOffsetParam->PrivateVariable.Get());
-		UTEST_EQUAL_EXPR(OffsetParamOverride->Value.Variable->DefaultValue, FVector3d(15, 25, 35));
+		const FPropertyBagPropertyDesc* OffsetParamDesc = ParameterOverrides.FindPropertyDescByName("OffsetParam");
+		FVector3dCameraParameter* OffsetParamOverride = ParameterOverrides.GetValueStruct<FVector3dCameraParameter>("OffsetParam").GetValue();
+		UTEST_EQUAL_EXPR((UCameraVariableAsset*)OffsetParamOverride->Variable.Get(), MiddleOffsetParam->PrivateVariable.Get());
+		UTEST_EQUAL_EXPR(OffsetParamOverride->Variable->DefaultValue, FVector3d(15, 25, 35));
 	}
 
 	// Test that the outer prefab node is driving the middle interface parameters.
 	{
-		FCameraRigParameterOverrides& ParameterOverrides = OuterPrefabNode->CameraRigReference.GetParameterOverrides();
+		FInstancedPropertyBag& ParameterOverrides = OuterPrefabNode->CameraRigReference.GetParameters();
 
-		FVector3dCameraRigParameterOverride* OffsetParamOverride =
-			ParameterOverrides.FindParameterOverride<FVector3dCameraRigParameterOverride>(MiddleOffsetParam->Guid);
+		const FPropertyBagPropertyDesc* OffsetParamDesc = ParameterOverrides.FindPropertyDescByName("MiddleOffsetParam");
+		FVector3dCameraParameter* OffsetParamOverride = ParameterOverrides.GetValueStruct<FVector3dCameraParameter>("MiddleOffsetParam").GetValue();
+		UTEST_NOT_NULL("OffsetParamDesc", OffsetParamDesc);
 		UTEST_NOT_NULL("OffsetParamOverride", OffsetParamOverride);
 
-		UTEST_EQUAL_EXPR(OffsetParamOverride->InterfaceParameterName, "MiddleOffsetParam");
-		UTEST_EQUAL_EXPR(OffsetParamOverride->PrivateVariableGuid, MiddleOffsetParam->PrivateVariable->GetGuid());
-		UTEST_EQUAL_EXPR(OffsetParamOverride->Value.Value, FVector3d(20, 50, 70));
+		UTEST_EQUAL_EXPR(OffsetParamDesc->ID, MiddleOffsetParam->GetGuid());
+		UTEST_EQUAL_EXPR(OffsetParamOverride->Value, FVector3d(20, 50, 70));
 	}
 
 	return true;
diff --git a/Engine/Plugins/Cameras/GameplayCameras/Source/GameplayCameras/Private/Tests/GameplayCamerasTestBuilder.h b/Engine/Plugins/Cameras/GameplayCameras/Source/GameplayCameras/Private/Tests/GameplayCamerasTestBuilder.h
--- a/Engine/Plugins/Cameras/GameplayCameras/Source/GameplayCameras/Private/Tests/GameplayCamerasTestBuilder.h
+++ b/Engine/Plugins/Cameras/GameplayCameras/Source/GameplayCameras/Private/Tests/GameplayCamerasTestBuilder.h
@@ -7,8 +7,11 @@
 #include "Core/CameraDirector.h"
 #include "Core/CameraEvaluationContext.h"
 #include "Core/CameraNode.h"
+#include "Core/CameraParameters.h"
 #include "Core/CameraRigAsset.h"
 #include "Core/CameraRigTransition.h"
+#include "Core/CameraVariableReferences.h"
+#include "Core/CameraVariableTableFwd.h"
 #include "Directors/SingleCameraDirector.h"
 #include "Nodes/Common/ArrayCameraNode.h"
 #include "Templates/PointerIsConvertibleFromTo.h"
@@ -482,14 +485,16 @@ public:
 	 *
 	 * The created parameter is automatically stored in the named object registry under its name.
 	 */
-	ThisType& ExposeParameter(const FString& ParameterName, UCameraNode* Target, FName TargetPropertyName)
-	{
-		UCameraRigInterfaceParameter* InterfaceParameter = NewObject<UCameraRigInterfaceParameter>(CameraRig);
-		InterfaceParameter->InterfaceParameterName = ParameterName;
-		InterfaceParameter->Target = Target;
-		InterfaceParameter->TargetPropertyName = TargetPropertyName;
-		NamedObjectRegistry->Register(InterfaceParameter, ParameterName);
-		CameraRig->Interface.InterfaceParameters.Add(InterfaceParameter);
+	ThisType& AddBlendableParameter(const FString& ParameterName, ECameraVariableType ParameterType, UCameraNode* Target, FName TargetPropertyName)
+	{
+		UCameraRigBlendableParameter* BlendableParameter = NewObject<UCameraRigBlendableParameter>(CameraRig);
+		BlendableParameter->InterfaceParameterName = ParameterName;
+		BlendableParameter->ParameterType = ParameterType;
+		BlendableParameter->Target = Target;
+		BlendableParameter->TargetPropertyName = TargetPropertyName;
+
+		NamedObjectRegistry->Register(BlendableParameter, ParameterName);
+		CameraRig->Interface.BlendableParameters.Add(BlendableParameter);
 		return *static_cast<ThisType*>(this);
 	}
 
@@ -498,11 +503,11 @@ public:
 	 *
 	 * The created parameter is automatically stored in the named object registry under its name.
 	 */
-	ThisType& ExposeParameter(const FString& ParameterName, const FString& TargetName, FName TargetPropertyName)
+	ThisType& AddBlendableParameter(const FString& ParameterName, ECameraVariableType ParameterType, const FString& TargetName, FName TargetPropertyName)
 	{
 		UCameraNode* Target = NamedObjectRegistry->Get<UCameraNode>(TargetName);
 		ensure(Target);
-		return ExposeParameter(ParameterName, Target, TargetPropertyName);
+		return AddBlendableParameter(ParameterName, ParameterType, Target, TargetPropertyName);
 	}
 
 	/** Gets the named object registry. */
diff --git a/Engine/Plugins/Cameras/GameplayCameras/Source/GameplayCameras/Public/Core/CameraAsset.h b/Engine/Plugins/Cameras/GameplayCameras/Source/GameplayCameras/Public/Core/CameraAsset.h
--- a/Engine/Plugins/Cameras/GameplayCameras/Source/GameplayCameras/Public/Core/CameraAsset.h
+++ b/Engine/Plugins/Cameras/GameplayCameras/Source/GameplayCameras/Public/Core/CameraAsset.h
@@ -3,11 +3,14 @@
 #pragma once
 
 #include "Core/CameraBuildStatus.h"
+#include "Core/CameraContextDataAllocationInfo.h"
 #include "Core/CameraEventHandler.h"
 #include "Core/CameraRigTransition.h"
+#include "Core/CameraVariableTableFwd.h"
 #include "Core/ObjectTreeGraphObject.h"
 #include "Core/ObjectTreeGraphRootObject.h"
 #include "CoreTypes.h"
+#include "StructUtils/PropertyBag.h"
 #include "UObject/ObjectPtr.h"
 
 #include "CameraAsset.generated.h"
@@ -18,6 +21,8 @@ class UCameraRigAsset;
 
 namespace UE::Cameras
 {
+	class FCameraAssetBuilder;
+	class FCameraAssetParameterOverrideEvaluator;
 	class FCameraBuildLog;
 
 	/**
@@ -37,6 +42,25 @@ namespace UE::Cameras
 	};
 }
 
+/**
+ * Structure describing various allocations needed by a camera asset.
+ */
+USTRUCT()
+struct FCameraAssetAllocationInfo
+{
+	GENERATED_BODY()
+
+	/** Combined variable table allocation info for all the camera rigs. */
+	UPROPERTY()
+	FCameraVariableTableAllocationInfo VariableTableInfo;
+
+	/** Combined context data table allocation info for all the camera rigs. */
+	UPROPERTY()
+	FCameraContextDataAllocationInfo ContextDataTableInfo;
+
+	GAMEPLAYCAMERAS_API friend bool operator==(const FCameraAssetAllocationInfo& A, const FCameraAssetAllocationInfo& B);
+};
+
 /**
  * A complete camera asset.
  */
@@ -70,6 +94,11 @@ public:
 	/** Removes an exit transition. */
 	GAMEPLAYCAMERAS_API int32 RemoveExitTransition(UCameraRigTransition* InTransition);
 
+	/** Gets the default parameter values for all camera rigs. */
+	const FInstancedPropertyBag& GetDefaultParameters() const { return DefaultParameters; }
+	/** Gets the default parameter values for all camera rigs. */
+	FInstancedPropertyBag& GetDefaultParameters() { return DefaultParameters; }
+
 public:
 
 	/**
@@ -96,6 +125,9 @@ public:
 	/** Sets the build status. */
 	void SetBuildStatus(ECameraBuildStatus InBuildStatus) { BuildStatus = InBuildStatus; }
 
+	/** Gets the allocation info for this camera asset. */
+	const FCameraAssetAllocationInfo& GetAllocationInfo() const { return AllocationInfo; }
+
 protected:
 
 	// IObjectTreeGraphObject interface.
@@ -132,10 +164,6 @@ private:
 	UPROPERTY(Instanced)
 	TObjectPtr<UCameraDirector> CameraDirector;
 
-	/** The list of camera rigs used by this camera (deprecated). */
-	UPROPERTY()
-	TArray<TObjectPtr<UCameraRigAsset>> CameraRigs_DEPRECATED;
-
 	/** A list of default enter transitions for all the camera rigs in this asset. */
 	UPROPERTY(Instanced)
 	TArray<TObjectPtr<UCameraRigTransition>> EnterTransitions;
@@ -148,6 +176,24 @@ private:
 	UPROPERTY(Transient)
 	ECameraBuildStatus BuildStatus = ECameraBuildStatus::Dirty;
 
+	/** Generated values for all camera rig parameters. */
+	UPROPERTY()
+	FInstancedPropertyBag DefaultParameters;
+
+	/** Generated list of camera rigs owning each of the parameters. */
+	UPROPERTY()
+	TArray<TWeakObjectPtr<const UCameraRigAsset>> ParameterOwners;
+
+	/** Allocation info for the camera asset. */
+	UPROPERTY()
+	FCameraAssetAllocationInfo AllocationInfo;
+
+
+	// Deprecated.
+
+	UPROPERTY()
+	TArray<TObjectPtr<UCameraRigAsset>> CameraRigs_DEPRECATED;
+
 #if WITH_EDITORONLY_DATA
 
 	/** Position of the camera node in the shared transitions graph editor. */
@@ -166,5 +212,8 @@ private:
 	friend class UCameraSharedTransitionGraphSchema;
 
 #endif  // WITH_EDITORONLY_DATA
+	
+	friend class UE::Cameras::FCameraAssetBuilder;
+	friend class UE::Cameras::FCameraAssetParameterOverrideEvaluator;
 };
 
diff --git a/Engine/Plugins/Cameras/GameplayCameras/Source/GameplayCameras/Public/Core/CameraContextDataAllocationInfo.h b/Engine/Plugins/Cameras/GameplayCameras/Source/GameplayCameras/Public/Core/CameraContextDataAllocationInfo.h
new file mode 100644
--- /dev/null
+++ b/Engine/Plugins/Cameras/GameplayCameras/Source/GameplayCameras/Public/Core/CameraContextDataAllocationInfo.h
@@ -0,0 +1,57 @@
+// Copyright Epic Games, Inc. All Rights Reserved.
+
+#pragma once
+
+#include "Core/CameraContextDataTableFwd.h"
+#include "UObject/Class.h"
+#include "UObject/ObjectMacros.h"
+#include "UObject/UnrealNames.h"
+
+#include "CameraContextDataAllocationInfo.generated.h"
+
+class UCameraNode;
+
+/**
+ * Definition for one entry in a camera rig's context data registry.
+ */
+USTRUCT()
+struct FCameraContextDataDefinition
+{
+	GENERATED_BODY()
+
+	/** The ID of data. */
+	UPROPERTY()
+	FCameraContextDataID DataID;
+
+	/** The type of the data. */
+	UPROPERTY()
+	ECameraContextDataType DataType;
+
+	/** An extra type object for the data. */
+	UPROPERTY()
+	TObjectPtr<const UObject> DataTypeObject;
+
+	GAMEPLAYCAMERAS_API friend bool operator==(const FCameraContextDataDefinition& A, const FCameraContextDataDefinition& B);
+};
+
+/**
+ * Collection of context data entries for a camera rig.
+ */
+USTRUCT()
+struct FCameraContextDataAllocationInfo
+{
+	GENERATED_BODY()
+
+	/** The list of context data definitions. */
+	UPROPERTY()
+	TArray<FCameraContextDataDefinition> DataDefinitions;
+
+	/** Adds a new context data definition. */
+	GAMEPLAYCAMERAS_API void Add(const UCameraNode* Owner, FName DataName, ECameraContextDataType DataType, const UObject* DataTypeObject);
+
+	/**Combines the given allocation info with this one. */
+	GAMEPLAYCAMERAS_API void Combine(const FCameraContextDataAllocationInfo& OtherInfo);
+
+	GAMEPLAYCAMERAS_API friend bool operator==(const FCameraContextDataAllocationInfo& A, const FCameraContextDataAllocationInfo& B);
+};
+
diff --git a/Engine/Plugins/Cameras/GameplayCameras/Source/GameplayCameras/Public/Core/CameraContextDataTable.h b/Engine/Plugins/Cameras/GameplayCameras/Source/GameplayCameras/Public/Core/CameraContextDataTable.h
new file mode 100644
--- /dev/null
+++ b/Engine/Plugins/Cameras/GameplayCameras/Source/GameplayCameras/Public/Core/CameraContextDataTable.h
@@ -0,0 +1,168 @@
+// Copyright Epic Games, Inc. All Rights Reserved.
+
+#pragma once
+
+#include "Core/CameraContextDataTableFwd.h"
+#include "Containers/Array.h"
+#include "Containers/Map.h"
+#include "GameplayCameras.h"
+#include "Math/NumericLimits.h"
+#include "StructUtils/InstancedStruct.h"
+#include "StructUtils/StructView.h"
+
+struct FCameraContextDataAllocationInfo;
+struct FCameraContextDataDefinition;
+
+namespace UE::Cameras
+{
+
+class FCameraContextDataTable;
+
+#if UE_GAMEPLAY_CAMERAS_DEBUG
+class FContextDataTableDebugBlock;
+#endif  // UE_GAMEPLAY_CAMERAS_DEBUG
+
+template<typename DataType>
+struct TCameraContextDataReader
+{
+	void Initialize(FCameraContextDataID InID, const FCameraContextDataTable& InContextDataTable);
+};
+
+/**
+ * The camera context data table is a container for a collection of arbitrary values
+ * of various types. It is the companion of the camera variable table (see FCameraVariableTable)
+ * but for non-blendable values.
+ */
+class FCameraContextDataTable
+{
+public:
+
+	FCameraContextDataTable();
+	~FCameraContextDataTable();
+
+	/** Initializes the context data table so that it fits the provided allocation info. */
+	void Initialize(const FCameraContextDataAllocationInfo& AllocationInfo);
+
+	/** Adds a data entry to the table. */
+	void AddData(const FCameraContextDataDefinition& DataDefinition);
+
+public:
+
+	// Getter methods.
+
+	template<typename DataType>
+	const DataType* GetData(FCameraContextDataID InID) const;
+
+	// Setter methods.
+
+	template<typename DataType>
+	void SetData(FCameraContextDataID InID, const DataType& InData);
+
+	void SetData(FCameraContextDataID InID, const FStructView& InData);
+	void SetData(FCameraContextDataID InID, const FInstancedStruct& InData);
+
+public:
+
+	// Overriding.
+
+	void OverrideAll(const FCameraContextDataTable& OtherTable);
+	void OverrideKnown(const FCameraContextDataTable& OtherTable);
+
+public:
+
+	/** Collects referenced objects. */
+	void AddReferencedObjects(FReferenceCollector& ReferenceCollector);
+
+public:
+
+	// Low-level API.
+	const uint8* GetData(FCameraContextDataID DataID, ECameraContextDataType ExpectedDataType, const UObject* ExpectedDataTypeObject) const;
+	const uint8* TryGetData(FCameraContextDataID DataID, ECameraContextDataType ExpectedDataType, const UObject* ExpectedDataTypeObject) const;
+
+	void SetData(FCameraContextDataID DataID, ECameraContextDataType ExpectedDataType, const UObject* ExpectedDataTypeObject, const uint8* InRawDataPtr, bool bMarkAsWrittenThisFrame = true);
+	bool TrySetData(FCameraContextDataID DataID, ECameraContextDataType ExpectedDataType, const UObject* ExpectedDataTypeObject, const uint8* InRawDataPtr, bool bMarkAsWrittenThisFrame = true);
+
+	bool IsValueWritten(FCameraContextDataID InID) const;
+	void UnsetValue(FCameraContextDataID InID);
+	void UnsetAllValues();
+
+	bool IsValueWrittenThisFrame(FCameraContextDataID InID) const;
+	void ClearAllWrittenThisFrameFlags();
+
+private:
+
+	enum class EEntryFlags : uint8
+	{
+		None = 0,
+		Written = 1 << 2,
+		WrittenThisFrame = 1 << 3
+	};
+	FRIEND_ENUM_CLASS_FLAGS(EEntryFlags)
+
+	struct FEntry
+	{
+		FCameraContextDataID ID;
+		ECameraContextDataType Type;
+		TObjectPtr<const UObject> TypeObject;
+		uint32 Offset;
+		EEntryFlags Flags;
+#if WITH_EDITORONLY_DATA
+		FName DebugName;
+#endif
+	};
+
+	static bool GetDataTypeAllocationInfo(ECameraContextDataType DataType, const UObject* DataTypeObject, uint32& OutSizeOf, uint32& OutAlignOf);
+	static bool InitializeDefaultDataValue(ECameraContextDataType DataType, const UObject* DataTypeObject, uint8* DataPtr);
+	static bool SetDataValue(ECameraContextDataType DataType, const UObject* DataTypeObject, uint8* DestDataPtr, const uint8* SrcDataPtr);
+
+	const FEntry* FindEntry(FCameraContextDataID InID) const;
+	FEntry* FindEntry(FCameraContextDataID InID);
+
+	void Override(const FCameraContextDataTable& OtherTable, bool bKnownOnly);
+
+	void ReallocateBuffer(uint32 MinRequired = 0);
+	void DestroyBuffer();
+
+private:
+
+	TArray<FEntry> Entries;
+	TMap<FCameraContextDataID, int32> EntryLookup;
+
+	uint8* Memory = nullptr;
+	uint32 Capacity = 0;
+	uint32 Used = 0;
+
+#if UE_GAMEPLAY_CAMERAS_DEBUG
+	friend class FContextDataTableDebugBlock;
+#endif  // UE_GAMEPLAY_CAMERAS_DEBUG
+};
+
+ENUM_CLASS_FLAGS(FCameraContextDataTable::EEntryFlags)
+
+template<typename DataType>
+const DataType* FCameraContextDataTable::GetData(FCameraContextDataID InID) const
+{
+	const FEntry* Entry = FindEntry(InID);
+	if (Entry)
+	{
+		ensureMsgf(Entry->Type == DataType::StaticStruct(), TEXT("Data type mismatch!"));
+		const uint8* RawData = Memory + Entry->Offset;
+		return reinterpret_cast<const DataType*>(RawData);
+	}
+	return nullptr;
+}
+
+template<typename DataType>
+void FCameraContextDataTable::SetData(FCameraContextDataID InID, const DataType& InData)
+{
+	const FEntry* Entry = FindEntry(InID);
+	if (ensure(Entry))
+	{
+		ensureMsgf(Entry->Type == DataType::StaticStruct(), TEXT("Data type mismatch!"));
+		uint8* RawData = Memory + Entry->Offset;
+		*reinterpret_cast<DataType*>(RawData) = InData;
+	}
+}
+
+}  // namespace UE::Cameras
+
diff --git a/Engine/Plugins/Cameras/GameplayCameras/Source/GameplayCameras/Public/Core/CameraNodeEvaluator.h b/Engine/Plugins/Cameras/GameplayCameras/Source/GameplayCameras/Public/Core/CameraNodeEvaluator.h
--- a/Engine/Plugins/Cameras/GameplayCameras/Source/GameplayCameras/Public/Core/CameraNodeEvaluator.h
+++ b/Engine/Plugins/Cameras/GameplayCameras/Source/GameplayCameras/Public/Core/CameraNodeEvaluator.h
@@ -2,6 +2,7 @@
 
 #pragma once
 
+#include "Core/CameraContextDataTable.h"
 #include "Core/CameraObjectRtti.h"
 #include "Core/CameraPose.h"
 #include "Core/CameraRigEvaluationInfo.h"
@@ -164,6 +165,9 @@ struct GAMEPLAYCAMERAS_API FCameraNodeEvaluationResult
 	/** The variable table. */
 	FCameraVariableTable VariableTable;
 
+	/** The context data table. */
+	FCameraContextDataTable ContextDataTable;
+
 	/** The list of joints in the current camera rig. */
 	FCameraRigJoints CameraRigJoints;
 
@@ -189,6 +193,11 @@ public:
 
 	/** Serializes this result to the given archive. */
 	void Serialize(FArchive& Ar);
+
+public:
+
+	/** Collects objects from the context data table. */
+	void AddReferencedObjects(FReferenceCollector& Collector);
 };
 
 /**
diff --git a/Engine/Plugins/Cameras/GameplayCameras/Source/GameplayCameras/Public/Core/CameraRigAsset.h b/Engine/Plugins/Cameras/GameplayCameras/Source/GameplayCameras/Public/Core/CameraRigAsset.h
--- a/Engine/Plugins/Cameras/GameplayCameras/Source/GameplayCameras/Public/Core/CameraRigAsset.h
+++ b/Engine/Plugins/Cameras/GameplayCameras/Source/GameplayCameras/Public/Core/CameraRigAsset.h
@@ -3,6 +3,8 @@
 #pragma once
 
 #include "Core/CameraBuildStatus.h"
+#include "Core/CameraContextDataAllocationInfo.h"
+#include "Core/CameraContextDataTableFwd.h"
 #include "Core/CameraEventHandler.h"
 #include "Core/CameraNodeEvaluatorFwd.h"
 #include "Core/CameraRigTransition.h"
@@ -12,6 +14,7 @@
 #include "CoreTypes.h"
 #include "GameplayTagAssetInterface.h"
 #include "GameplayTagContainer.h"
+#include "StructUtils/PropertyBag.h"
 #include "UObject/ObjectPtr.h"
 
 #include "CameraRigAsset.generated.h"
@@ -36,6 +39,9 @@ namespace UE::Cameras
 		/** Called when the camera rig asset has been built. */
 		virtual void OnCameraRigBuilt(const UCameraRigAsset* CameraRigAsset) {}
 
+		/** Called when the camera rig's interface has changed. */
+		virtual void OnCameraRigInterfaceChanged() {}
+
 #if WITH_EDITOR
 		virtual void OnObjectAddedToGraph(const FName GraphName, UObject* Object) {}
 		virtual void OnObjectRemovedFromGraph(const FName GraphName, UObject* Object) {}
@@ -44,7 +50,7 @@ namespace UE::Cameras
 }
 
 /**
- * Structure describing various allocations needed by a camera node.
+ * Structure describing various allocations needed by a camera rig.
  */
 USTRUCT()
 struct FCameraRigAllocationInfo
@@ -55,10 +61,14 @@ struct FCameraRigAllocationInfo
 	UPROPERTY()
 	FCameraNodeEvaluatorAllocationInfo EvaluatorInfo;
 
-	/** Allocation info for the camera variable. */
+	/** Allocation info for the variable table. */
 	UPROPERTY()
 	FCameraVariableTableAllocationInfo VariableTableInfo;
 
+	/** Allocation info for the context data table. */
+	UPROPERTY()
+	FCameraContextDataAllocationInfo ContextDataTableInfo;
+
 public:
 
 	GAMEPLAYCAMERAS_API void Append(const FCameraRigAllocationInfo& OtherAllocationInfo);
@@ -77,12 +87,11 @@ struct TStructOpsTypeTraits<FCameraRigAllocationInfo> : public TStructOpsTypeTra
 };
 
 /**
- * An exposed camera rig parameter that drives a specific parameter on one of
- * its camera nodes.
+ * Base class for interface parameters on a camera rig asset.
  */
 UCLASS(MinimalAPI, meta=(
 			ObjectTreeGraphSelfPinDirection="Output"))
-class UCameraRigInterfaceParameter
+class UCameraRigInterfaceParameterBase 
 	: public UObject
 	, public IObjectTreeGraphObject
 {
@@ -90,30 +99,41 @@ class UCameraRigInterfaceParameter
 
 public:
 
-	/** The camera node that this parameter drives. */
+	/** The exposed name for this parameter. */
+	UPROPERTY(EditAnywhere, Category=Camera)
+	FString InterfaceParameterName;
+
+	/** The camera node this parameter is connected. */
 	UPROPERTY(meta=(ObjectTreeGraphHidden=true))
 	TObjectPtr<UCameraNode> Target;
 
-	/** The camera parameter on the target camera node that this parameter drives. */
+	/**
+	 * The name of the property this parameter is connected to on the target camera node.
+	 * This may be an actual UObject property, but it may be something else, like the name
+	 * of an interface parameter on a nested camera rig, or the name of a Blueprint property
+	 * on the evaluator class of a Blueprint camera node.
+	 */
 	UPROPERTY()
 	FName TargetPropertyName;
 
-	/** The exposed name for this parameter. */
-	UPROPERTY(EditAnywhere, Category=Camera)
-	FString InterfaceParameterName;
+#if WITH_EDITORONLY_DATA
 
+	/** Whether this parameter has been added to the node graph in the editor. */
 	UPROPERTY()
-	FGuid Guid;
+	bool bHasGraphNode = false;
 
-	// Built on save/cook.
+#endif  // WITH_EDITORONLY_DATA
 
-	/**
-	 * The private camera variable created to drive the target camera parameter on
-	 * the target camera node. This variable is created by the build method on the
-	 * camera rig.
-	 */
+public:
+
+	/** Gets this parameter's unique ID. */
+	const FGuid& GetGuid() const { return Guid; }
+
+protected:
+
+	/** The Guid of this parameter. */
 	UPROPERTY()
-	TObjectPtr<UCameraVariableAsset> PrivateVariable;
+	FGuid Guid;
 
 protected:
 
@@ -138,6 +158,54 @@ private:
 #endif  // WITH_EDITORONLY_DATA
 };
 
+/**
+ * An exposed camera rig parameter that drives a specific parameter on one of
+ * its camera nodes.
+ */
+UCLASS(MinimalAPI)
+class UCameraRigBlendableParameter : public UCameraRigInterfaceParameterBase
+{
+	GENERATED_BODY()
+
+public:
+
+	/** The type of this parameter. */
+	UPROPERTY()
+	ECameraVariableType ParameterType = ECameraVariableType::Boolean;
+
+	// Built on save/cook.
+
+	/**
+	 * The private camera variable created to drive the target camera parameter on
+	 * the target camera node. This variable is created by the build method on the
+	 * camera rig.
+	 */
+	UPROPERTY()
+	TObjectPtr<UCameraVariableAsset> PrivateVariable;
+};
+
+UCLASS(MinimalAPI)
+class UCameraRigDataParameter : public UCameraRigInterfaceParameterBase
+{
+	GENERATED_BODY()
+
+public:
+
+	/** The type of this parameter. */
+	UPROPERTY()
+	ECameraContextDataType DataType;
+
+	/** An additional type object for this parameter. */
+	UPROPERTY()
+	TObjectPtr<const UObject> DataTypeObject;
+
+	// Built on save/cook.
+
+	/** The reference to use to access the underlying data in the context data table. */
+	UPROPERTY()
+	FCameraContextDataID PrivateDataID;
+};
+
 /**
  * Structure defining the public data interface of a camera rig asset.
  */
@@ -150,18 +218,22 @@ public:
 
 	/** The list of exposed parameters on the camera rig. */
 	UPROPERTY(Instanced)
-	TArray<TObjectPtr<UCameraRigInterfaceParameter>> InterfaceParameters;
+	TArray<TObjectPtr<UCameraRigBlendableParameter>> BlendableParameters;
+
+	UPROPERTY(Instanced)
+	TArray<TObjectPtr<UCameraRigDataParameter>> DataParameters;
 
 public:
 	
 	/** Finds an exposed parameter by name. */
-	GAMEPLAYCAMERAS_API UCameraRigInterfaceParameter* FindInterfaceParameterByName(const FString& ParameterName) const;
+	GAMEPLAYCAMERAS_API UCameraRigBlendableParameter* FindBlendableParameterByName(const FString& ParameterName) const;
 
 	/** Finds an exposed parameter by Guid. */
-	GAMEPLAYCAMERAS_API UCameraRigInterfaceParameter* FindInterfaceParameterByGuid(const FGuid& ParameterGuid) const;
+	GAMEPLAYCAMERAS_API UCameraRigBlendableParameter* FindBlendableParameterByGuid(const FGuid& ParameterGuid) const;
+	GAMEPLAYCAMERAS_API UCameraRigDataParameter* FindDataParameterByGuid(const FGuid& ParameterGuid) const;
 
 	/** Returns whether an exposed parameter with the given name exists. */
-	GAMEPLAYCAMERAS_API bool HasInterfaceParameter(const FString& ParameterName) const;
+	GAMEPLAYCAMERAS_API bool HasBlendableParameter(const FString& ParameterName) const;
 };
 
 /**
@@ -223,6 +295,11 @@ public:
 	/** Gets the camera rig's unique ID. */
 	const FGuid& GetGuid() const { return Guid; }
 
+	/** Gets the default values for the parameters exposed on this camera rig. */
+	const FInstancedPropertyBag& GetDefaultParameters() const { return DefaultParameters; }
+
+	/** Gets the default values for the parameters exposed on this camera rig. */
+	FInstancedPropertyBag& GetDefaultParameters() { return DefaultParameters; }
 
 public:
 
@@ -287,9 +364,14 @@ protected:
 
 private:
 
+	/** The camera rig's unique ID. */
 	UPROPERTY()
 	FGuid Guid;
 
+	/** The default interface parameter values, generated during build. */
+	UPROPERTY()
+	FInstancedPropertyBag DefaultParameters;
+
 #if WITH_EDITORONLY_DATA
 
 	/** Position of the camera rig node in the node graph editor. */
diff --git a/Engine/Plugins/Cameras/GameplayCameras/Source/GameplayCameras/Public/Core/CameraRigAssetBuilder.h b/Engine/Plugins/Cameras/GameplayCameras/Source/GameplayCameras/Public/Core/CameraRigAssetBuilder.h
--- a/Engine/Plugins/Cameras/GameplayCameras/Source/GameplayCameras/Public/Core/CameraRigAssetBuilder.h
+++ b/Engine/Plugins/Cameras/GameplayCameras/Source/GameplayCameras/Public/Core/CameraRigAssetBuilder.h
@@ -17,11 +17,12 @@ class UCameraRigAsset;
 class UCameraRigCameraNode;
 class UCameraVariableAsset;
 struct FCameraVariableTableAllocationInfo;
+struct FInstancedPropertyBag;
 
 namespace UE::Cameras
 {
 
-namespace Internal { struct FPrivateVariableBuilder; }
+namespace Internal { struct FInterfaceParameterBindingBuilder; }
 
 /**
  * A class that can prepare a camera rig for runtime use.
@@ -56,18 +57,31 @@ private:
 	void CallPreBuild();
 
 	void GatherOldDrivenParameters();
-	void BuildNewDrivenParameters();
-	void DiscardUnusedPrivateVariables();
+	void BuildInterfaceParameters();
+	void BuildInterfaceParameterBindings();
+	void DiscardUnusedParameters();
 
-	void BuildAllocationInfo();
-	void BuildAllocationInfo(UCameraNode* CameraNode);
+	void CallBuild();
+	void CallBuild(FCameraRigBuildContext& BuildContext, UCameraNode* CameraNode);
+
+	void BuildDefaultParameters();
 
 	void UpdateBuildStatus();
 
 private:
 
-	bool SetupCameraParameterOverride(UCameraRigInterfaceParameter* InterfaceParameter);
-	bool SetupInnerCameraRigParameterOverride(UCameraRigInterfaceParameter* InterfaceParameter);
+	bool SetupCameraParameterOrVariableReferenceOverride(const UCameraRigBlendableParameter* BlendableParameter);
+	bool SetupCustomBlendableParameterOverride(const UCameraRigBlendableParameter* BlendableParameter);
+
+	bool SetupDataContextPropertyOverride(const UCameraRigDataParameter* DataParameter);
+	bool SetupCustomDataParameterOverride(const UCameraRigDataParameter* DataParameter);
+
+public:
+
+	// Internal API.
+
+	static void BuildDefaultParameters(UCameraRigAsset* CameraRigAsset, FInstancedPropertyBag& OutPropertyBag);
+	static void AppendDefaultParameters(const FCameraRigInterface& CameraRigInterface, TArray<FPropertyBagPropertyDesc>& OutProperties);
 
 private:
 
@@ -77,17 +91,11 @@ private:
 
 	FCameraNodeHierarchy CameraNodeHierarchy;
 
-	using FDrivenParameterKey = TTuple<FStructProperty*, UCameraNode*>;
-	TMap<FDrivenParameterKey, UCameraVariableAsset*> OldDrivenParameters;
-	using FDrivenOverrideKey = TTuple<FGuid, UCameraRigCameraNode*>;
-	TMap<FDrivenOverrideKey, UCameraVariableAsset*> OldDrivenOverrides;
-
-	using FReusableInterfaceParameterInfo = TTuple<UCameraVariableAsset*, bool>;
-	TMap<UCameraRigInterfaceParameter*, FReusableInterfaceParameterInfo> OldInterfaceParameters;
-
-	FCameraRigAllocationInfo AllocationInfo;
+	using FDrivenParameterKey = TTuple<FName, UCameraNode*>;
+	TMap<FDrivenParameterKey, TObjectPtr<UCameraVariableAsset>> OldDrivenBlendableParameters;
+	TMap<FDrivenParameterKey, FCameraContextDataID> OldDrivenDataParameters;
 
-	friend struct Internal::FPrivateVariableBuilder;
+	friend struct Internal::FInterfaceParameterBindingBuilder;
 };
 
 }  // namespace UE::Cameras
diff --git a/Engine/Plugins/Cameras/GameplayCameras/Source/GameplayCameras/Public/Core/CameraRigAssetReference.h b/Engine/Plugins/Cameras/GameplayCameras/Source/GameplayCameras/Public/Core/CameraRigAssetReference.h
--- a/Engine/Plugins/Cameras/GameplayCameras/Source/GameplayCameras/Public/Core/CameraRigAssetReference.h
+++ b/Engine/Plugins/Cameras/GameplayCameras/Source/GameplayCameras/Public/Core/CameraRigAssetReference.h
@@ -3,63 +3,67 @@
 #pragma once
 
 #include "Core/CameraParameters.h"
-#include "Core/CameraRigAsset.h"
+#include "StructUtils/PropertyBag.h"
+#include "UObject/ObjectPtr.h"
 
 #include "CameraRigAssetReference.generated.h"
 
+class UCameraRigAsset;
 struct FCameraRigAssetReference;
+struct FPropertyTag;
 
 namespace UE::Cameras
 {
 	class FCameraRigAssetBuilder;
 	class FCameraRigAssetReferenceDetailsCustomization;
+	struct FCameraNodeEvaluationResult;
 }
 
-/** Base struct for camera rig parameter overrides. */
+
+struct UE_DEPRECATED(5.7, "Camera rig references now use instanced property bags.") FCameraRigParameterOverrideBase;
+
+struct UE_DEPRECATED(5.7, "Camera rig references now use instanced property bags.") FBooleanCameraRigParameterOverride;
+struct UE_DEPRECATED(5.7, "Camera rig references now use instanced property bags.") FInteger32CameraRigParameterOverride;
+struct UE_DEPRECATED(5.7, "Camera rig references now use instanced property bags.") FFloatCameraRigParameterOverride;
+struct UE_DEPRECATED(5.7, "Camera rig references now use instanced property bags.") FDoubleCameraRigParameterOverride;
+struct UE_DEPRECATED(5.7, "Camera rig references now use instanced property bags.") FVector2fCameraRigParameterOverride;
+struct UE_DEPRECATED(5.7, "Camera rig references now use instanced property bags.") FVector2dCameraRigParameterOverride;
+struct UE_DEPRECATED(5.7, "Camera rig references now use instanced property bags.") FVector3fCameraRigParameterOverride;
+struct UE_DEPRECATED(5.7, "Camera rig references now use instanced property bags.") FVector3dCameraRigParameterOverride;
+struct UE_DEPRECATED(5.7, "Camera rig references now use instanced property bags.") FVector4fCameraRigParameterOverride;
+struct UE_DEPRECATED(5.7, "Camera rig references now use instanced property bags.") FVector4dCameraRigParameterOverride;
+struct UE_DEPRECATED(5.7, "Camera rig references now use instanced property bags.") FRotator3fCameraRigParameterOverride;
+struct UE_DEPRECATED(5.7, "Camera rig references now use instanced property bags.") FRotator3dCameraRigParameterOverride;
+struct UE_DEPRECATED(5.7, "Camera rig references now use instanced property bags.") FTransform3fCameraRigParameterOverride;
+struct UE_DEPRECATED(5.7, "Camera rig references now use instanced property bags.") FTransform3dCameraRigParameterOverride;
+
+struct UE_DEPRECATED(5.7, "Camera rig references now use instanced property bags.") FCameraRigParameterOverrides;
+
 USTRUCT()
 struct FCameraRigParameterOverrideBase
 {
 	GENERATED_BODY()
 
-	/**
-	 * The Guid of the overriden interface parameter in the inner camera rig.
-	 */
 	UPROPERTY()
 	FGuid InterfaceParameterGuid;
 
-	/**
-	 * The Guid of the overriden interface parameter's private variable in the
-	 * inner camera rig.
-	 *
-	 * This can be derived from InterfaceParameterGuid, but we cache this during
-	 * the build process to avoid searching for interface parameters.
-	 */
 	UPROPERTY()
 	FGuid PrivateVariableGuid;
 
-	/**
-	 * The name of the overriden interface parameter in the inner camera rig.
-	 *
-	 * This can be derived from InterfaceParameterGuid, but we cache this during
-	 * the build process to avoid searching for interface parameters.
-	 */
 	UPROPERTY()
 	FString InterfaceParameterName;
 
-	/**
-	 *
-	 */
 	UPROPERTY()
 	bool bInvalid = false;
 };
 
+PRAGMA_DISABLE_DEPRECATION_WARNINGS
+
 USTRUCT()
 struct FBooleanCameraRigParameterOverride : public FCameraRigParameterOverrideBase
 {
 	GENERATED_BODY()
 
-	using CameraParameterType = FBooleanCameraParameter;
-
 	UPROPERTY()
 	FBooleanCameraParameter Value;
 };
@@ -69,8 +73,6 @@ struct FInteger32CameraRigParameterOverride : public FCameraRigParameterOverride
 {
 	GENERATED_BODY()
 
-	using CameraParameterType = FInteger32CameraParameter;
-
 	UPROPERTY()
 	FInteger32CameraParameter Value;
 };
@@ -80,8 +82,6 @@ struct FFloatCameraRigParameterOverride : public FCameraRigParameterOverrideBase
 {
 	GENERATED_BODY()
 
-	using CameraParameterType = FFloatCameraParameter;
-
 	UPROPERTY()
 	FFloatCameraParameter Value;
 };
@@ -91,8 +91,6 @@ struct FDoubleCameraRigParameterOverride : public FCameraRigParameterOverrideBas
 {
 	GENERATED_BODY()
 
-	using CameraParameterType = FDoubleCameraParameter;
-
 	UPROPERTY()
 	FDoubleCameraParameter Value;
 };
@@ -102,8 +100,6 @@ struct FVector2fCameraRigParameterOverride : public FCameraRigParameterOverrideB
 {
 	GENERATED_BODY()
 
-	using CameraParameterType = FVector2fCameraParameter;
-
 	UPROPERTY()
 	FVector2fCameraParameter Value;
 };
@@ -113,8 +109,6 @@ struct FVector2dCameraRigParameterOverride : public FCameraRigParameterOverrideB
 {
 	GENERATED_BODY()
 
-	using CameraParameterType = FVector2dCameraParameter;
-
 	UPROPERTY()
 	FVector2dCameraParameter Value;
 };
@@ -124,8 +118,6 @@ struct FVector3fCameraRigParameterOverride : public FCameraRigParameterOverrideB
 {
 	GENERATED_BODY()
 
-	using CameraParameterType = FVector3fCameraParameter;
-
 	UPROPERTY()
 	FVector3fCameraParameter Value;
 };
@@ -135,8 +127,6 @@ struct FVector3dCameraRigParameterOverride : public FCameraRigParameterOverrideB
 {
 	GENERATED_BODY()
 
-	using CameraParameterType = FVector3dCameraParameter;
-
 	UPROPERTY()
 	FVector3dCameraParameter Value;
 };
@@ -146,8 +136,6 @@ struct FVector4fCameraRigParameterOverride : public FCameraRigParameterOverrideB
 {
 	GENERATED_BODY()
 
-	using CameraParameterType = FVector4fCameraParameter;
-
 	UPROPERTY()
 	FVector4fCameraParameter Value;
 };
@@ -157,8 +145,6 @@ struct FVector4dCameraRigParameterOverride : public FCameraRigParameterOverrideB
 {
 	GENERATED_BODY()
 
-	using CameraParameterType = FVector4dCameraParameter;
-
 	UPROPERTY()
 	FVector4dCameraParameter Value;
 };
@@ -168,8 +154,6 @@ struct FRotator3fCameraRigParameterOverride : public FCameraRigParameterOverride
 {
 	GENERATED_BODY()
 
-	using CameraParameterType = FRotator3fCameraParameter;
-
 	UPROPERTY()
 	FRotator3fCameraParameter Value;
 };
@@ -179,8 +163,6 @@ struct FRotator3dCameraRigParameterOverride : public FCameraRigParameterOverride
 {
 	GENERATED_BODY()
 
-	using CameraParameterType = FRotator3dCameraParameter;
-
 	UPROPERTY()
 	FRotator3dCameraParameter Value;
 };
@@ -190,8 +172,6 @@ struct FTransform3fCameraRigParameterOverride : public FCameraRigParameterOverri
 {
 	GENERATED_BODY()
 
-	using CameraParameterType = FTransform3fCameraParameter;
-
 	UPROPERTY()
 	FTransform3fCameraParameter Value;
 };
@@ -201,12 +181,12 @@ struct FTransform3dCameraRigParameterOverride : public FCameraRigParameterOverri
 {
 	GENERATED_BODY()
 
-	using CameraParameterType = FTransform3dCameraParameter;
-
 	UPROPERTY()
 	FTransform3dCameraParameter Value;
 };
 
+PRAGMA_ENABLE_DEPRECATION_WARNINGS
+
 /**
  * A structure that holds lists of camera rig interface parameter overrides, one list
  * per parameter type.
@@ -216,67 +196,9 @@ struct GAMEPLAYCAMERAS_API FCameraRigParameterOverrides
 {
 	GENERATED_BODY()
 
-public:
-
-	/** Whether the given camera rig interface parameter is currently overriden. */
-	template<typename ParameterOverrideType>
-	bool IsParameterOverriden(const FGuid& CameraRigParameterGuid) const;
-
-	/** Find a parameter override for the given inner camera rig interface parameter. */
-	template<typename ParameterOverrideType>
-	ParameterOverrideType* FindParameterOverride(const FGuid& CameraRigParameterGuid);
-
-	/** Find or create a parameter override for the given inner camera rig interface parameter. */
-	template<typename ParameterOverrideType>
-	ParameterOverrideType& FindOrAddParameterOverride(const UCameraRigInterfaceParameter* CameraRigParameter);
-
-	/** Remove any parameter override for the given inner camera rig interface parameter. */
-	template<typename ParameterOverrideType>
-	void RemoveParameterOverride(const FGuid& CameraRigParameterGuid);
-
-	/** Remove all parameter overrides. */
-	void Reset();
-
-public:
-
-#define UE_CAMERA_VARIABLE_FOR_TYPE(ValueType, ValueName)\
-	inline TArrayView<const F##ValueName##CameraRigParameterOverride> Get##ValueName##Overrides() const\
-	{\
-		return ValueName##Overrides;\
-	}\
-	inline TArrayView<F##ValueName##CameraRigParameterOverride> Get##ValueName##Overrides()\
-	{\
-		return ValueName##Overrides;\
-	}
-UE_CAMERA_VARIABLE_FOR_ALL_TYPES()
-#undef UE_CAMERA_VARIABLE_FOR_TYPE
-
-
-private:
-
-	template<typename ParameterOverrideType>
-	bool IsParameterOverriden(const TArray<ParameterOverrideType>& OverridesArray, const FGuid& CameraRigParameterGuid) const;
-
-	template<typename ParameterOverrideType>
-	ParameterOverrideType* FindParameterOverride(TArray<ParameterOverrideType>& OverridesArray, const FGuid& CameraRigParameterGuid);
-
-	template<typename ParameterOverrideType>
-	ParameterOverrideType& FindOrAddParameterOverride(TArray<ParameterOverrideType>& OverridesArray, const UCameraRigInterfaceParameter* CameraRigParameter);
-
-	template<typename ParameterOverrideType>
-	void RemoveParameterOverride(TArray<ParameterOverrideType>& OverridesArray, const FGuid& CameraRigParameterGuid);
-
-
-public:
-
-	// Internal API.
-
-	template<typename ParameterOverrideType>
-	inline void AppendParameterOverrides(TArrayView<ParameterOverrideType> NewOverrides);
-
 private:
 
-	// Interface parameter overrides
+	PRAGMA_DISABLE_DEPRECATION_WARNINGS
 
 	UPROPERTY()
 	TArray<FBooleanCameraRigParameterOverride> BooleanOverrides;
@@ -307,94 +229,17 @@ private:
 	UPROPERTY()
 	TArray<FTransform3dCameraRigParameterOverride> Transform3dOverrides;
 
-	friend struct FCameraRigAssetReference;
+	PRAGMA_ENABLE_DEPRECATION_WARNINGS
 
-	friend class UE::Cameras::FCameraRigAssetBuilder;
+	friend struct FCameraRigAssetReference;
 };
 
-template<typename ParameterOverrideType>
-bool FCameraRigParameterOverrides::IsParameterOverriden(const TArray<ParameterOverrideType>& OverridesArray, const FGuid& CameraRigParameterGuid) const
-{
-	return OverridesArray.ContainsByPredicate(
-			[CameraRigParameterGuid](ParameterOverrideType& Item)
-			{
-				return (Item.InterfaceParameterGuid == CameraRigParameterGuid);
-			});
-}
-
-template<typename ParameterOverrideType>
-ParameterOverrideType* FCameraRigParameterOverrides::FindParameterOverride(TArray<ParameterOverrideType>& OverridesArray, const FGuid& CameraRigParameterGuid)
-{
-	ParameterOverrideType* FoundItem = OverridesArray.FindByPredicate(
-			[CameraRigParameterGuid](ParameterOverrideType& Item)
-			{
-				return (Item.InterfaceParameterGuid == CameraRigParameterGuid);
-			});
-	return FoundItem;
-}
-
-template<typename ParameterOverrideType>
-ParameterOverrideType& FCameraRigParameterOverrides::FindOrAddParameterOverride(TArray<ParameterOverrideType>& OverridesArray, const UCameraRigInterfaceParameter* CameraRigParameter)
-{
-	ParameterOverrideType* Existing = FindParameterOverride<ParameterOverrideType>(OverridesArray, CameraRigParameter->Guid);
-	if (Existing)
-	{
-		return *Existing;
-	}
-	else
-	{
-		ParameterOverrideType& NewOverride = OverridesArray.Emplace_GetRef();
-		NewOverride.InterfaceParameterGuid = CameraRigParameter->Guid;
-		NewOverride.InterfaceParameterName = CameraRigParameter->InterfaceParameterName;
-		return NewOverride;
-	}
-}
-
-template<typename ParameterOverrideType>
-void FCameraRigParameterOverrides::RemoveParameterOverride(TArray<ParameterOverrideType>& OverridesArray, const FGuid& CameraRigParameterGuid)
-{
-	OverridesArray.RemoveAll(
-			[CameraRigParameterGuid](ParameterOverrideType& Item)
-			{
-				return (Item.InterfaceParameterGuid == CameraRigParameterGuid);
-			});
-}
-
-#define UE_CAMERA_VARIABLE_FOR_TYPE(ValueType, ValueName)\
-template<>\
-inline bool FCameraRigParameterOverrides::IsParameterOverriden<F##ValueName##CameraRigParameterOverride>(const FGuid& CameraRigParameterGuid) const\
-{\
-	return this->IsParameterOverriden<F##ValueName##CameraRigParameterOverride>(this->ValueName##Overrides, CameraRigParameterGuid);\
-}\
-template<>\
-inline F##ValueName##CameraRigParameterOverride* FCameraRigParameterOverrides::FindParameterOverride<F##ValueName##CameraRigParameterOverride>(const FGuid& CameraRigParameterGuid)\
-{\
-	return this->FindParameterOverride<F##ValueName##CameraRigParameterOverride>(this->ValueName##Overrides, CameraRigParameterGuid);\
-}\
-template<>\
-inline F##ValueName##CameraRigParameterOverride& FCameraRigParameterOverrides::FindOrAddParameterOverride<F##ValueName##CameraRigParameterOverride>(const UCameraRigInterfaceParameter* CameraRigParameter)\
-{\
-	return this->FindOrAddParameterOverride<F##ValueName##CameraRigParameterOverride>(this->ValueName##Overrides, CameraRigParameter);\
-}\
-template<>\
-inline void FCameraRigParameterOverrides::RemoveParameterOverride<F##ValueName##CameraRigParameterOverride>(const FGuid& CameraRigParameterGuid)\
-{\
-	this->RemoveParameterOverride<F##ValueName##CameraRigParameterOverride>(this->ValueName##Overrides, CameraRigParameterGuid);\
-}\
-template<>\
-inline void FCameraRigParameterOverrides::AppendParameterOverrides<F##ValueName##CameraRigParameterOverride>(TArrayView<F##ValueName##CameraRigParameterOverride> NewOverrides)\
-{\
-	this->ValueName##Overrides.Append(NewOverrides);\
-}
-UE_CAMERA_VARIABLE_FOR_ALL_TYPES()
-#undef UE_CAMERA_VARIABLE_FOR_TYPE
-
 /**
  * A structure holding a reference to a camera rig asset, along with the interface parameter
  * override values.
  */
 USTRUCT(BlueprintType)
-struct GAMEPLAYCAMERAS_API FCameraRigAssetReference
+struct FCameraRigAssetReference
 {
 	GENERATED_BODY()
 
@@ -421,37 +266,50 @@ public:
 		return CameraRig;
 	}
 
-	/**
-	 * Sets the referenced camerar rig.
-	 * This will check any existing parameter overrides, tagging them as invalid as needed.
-	 */
+	/** Sets the referenced camera rig. */
 	void SetCameraRig(UCameraRigAsset* InCameraRig)
 	{
 		if (CameraRig != InCameraRig)
 		{
 			CameraRig = InCameraRig;
-			UpdateParameterOverrides();
+			RebuildParameters();
 		}
 	}
 
-	/** Gets the parameter overrides. */
-	const FCameraRigParameterOverrides& GetParameterOverrides() const
+	/** Gets the parameters for this camera rig, some of which containing overrides. */
+	const FInstancedPropertyBag& GetParameters() const
 	{
-		return ParameterOverrides;
+		return Parameters;
 	}
 
-	/** Gets the parameter overrides. */
-	FCameraRigParameterOverrides& GetParameterOverrides()
+	/** Gets the parameters for this camera rig, some of which containing overrides. */
+	FInstancedPropertyBag& GetParameters()
 	{
-		return ParameterOverrides;
+		return Parameters;
 	}
 
+	/** Gets the IDs of the parameters with override values. */
+	TConstArrayView<FGuid> GetOverridenParameterGuids() const
+	{
+		return ParameterOverrideGuids;
+	}
+
+	/** Applies the parameter override values to the given evaluation result. */
+	void ApplyParameterOverrides(UE::Cameras::FCameraNodeEvaluationResult& OutResult, bool bDrivenOverridesOnly);
+
 public:
 
 	// Internal API.
 
-	bool SerializeFromMismatchedTag(struct FPropertyTag const& Tag, FStructuredArchive::FSlot Slot);
-	bool UpdateParameterOverrides();
+	bool SerializeFromMismatchedTag(FPropertyTag const& Tag, FStructuredArchive::FSlot Slot);
+	void PostSerialize(const FArchive& Ar);
+
+	GAMEPLAYCAMERAS_API bool IsParameterOverriden(const FGuid PropertyID) const;
+	GAMEPLAYCAMERAS_API void SetParameterOverriden(const FGuid PropertyID, bool bIsOverridden);
+
+	GAMEPLAYCAMERAS_API bool NeedsRebuildParameters() const;
+	GAMEPLAYCAMERAS_API bool RebuildParametersIfNeeded();
+	GAMEPLAYCAMERAS_API void RebuildParameters();
 
 private:
 
@@ -459,9 +317,23 @@ private:
 	UPROPERTY(EditAnywhere, Category="")
 	TObjectPtr<UCameraRigAsset> CameraRig;
 
-	/** The parameter overrides. */
+	/** The camera rig's parameters. */
 	UPROPERTY(EditAnywhere, Category="", meta=(FixedLayout))
-	FCameraRigParameterOverrides ParameterOverrides;
+	FInstancedPropertyBag Parameters;
+
+	/** The list of camera rig parameters with override values. */
+	UPROPERTY(EditAnywhere, Category="")
+	TArray<FGuid> ParameterOverrideGuids;
+
+
+	// Deprecated
+
+	PRAGMA_DISABLE_DEPRECATION_WARNINGS
+
+	UPROPERTY()
+	FCameraRigParameterOverrides ParameterOverrides_DEPRECATED;
+
+	PRAGMA_ENABLE_DEPRECATION_WARNINGS
 
 	friend class UE::Cameras::FCameraRigAssetReferenceDetailsCustomization;
 };
@@ -471,7 +343,8 @@ struct TStructOpsTypeTraits<FCameraRigAssetReference> : public TStructOpsTypeTra
 {
 	enum
 	{
-		WithStructuredSerializeFromMismatchedTag = true
+		WithStructuredSerializeFromMismatchedTag = true,
+		WithPostSerialize = true
 	};
 };
 
diff --git a/Engine/Plugins/Cameras/GameplayCameras/Source/GameplayCameras/Public/Core/CameraRigBuildContext.h b/Engine/Plugins/Cameras/GameplayCameras/Source/GameplayCameras/Public/Core/CameraRigBuildContext.h
--- a/Engine/Plugins/Cameras/GameplayCameras/Source/GameplayCameras/Public/Core/CameraRigBuildContext.h
+++ b/Engine/Plugins/Cameras/GameplayCameras/Source/GameplayCameras/Public/Core/CameraRigBuildContext.h
@@ -2,10 +2,8 @@
 
 #pragma once
 
-#include "UObject/ObjectPtr.h"
-
-class UCameraRigAsset;
-struct FCameraRigAllocationInfo;
+#include "Containers/Set.h"
+#include "Core/CameraRigAsset.h"
 
 namespace UE::Cameras
 {
@@ -17,16 +15,15 @@ class FCameraBuildLog;
  */
 struct FCameraRigBuildContext
 {
-	FCameraRigBuildContext(
-			FCameraRigAllocationInfo& InAllocationInfo, FCameraBuildLog& InBuildLog)
-		: AllocationInfo(InAllocationInfo)
-		, BuildLog(InBuildLog)
+	FCameraRigBuildContext(FCameraBuildLog& InBuildLog)
+		: BuildLog(InBuildLog)
 	{}
 
-	/** The allocation information to be determined. */
-	FCameraRigAllocationInfo& AllocationInfo;
 	/** The build log for emitting messages. */
 	FCameraBuildLog& BuildLog;
+
+	/** The allocation information for the camera rig. */
+	FCameraRigAllocationInfo AllocationInfo;
 };
 
 }  // namespace UE::Cameras
diff --git a/Engine/Plugins/Cameras/GameplayCameras/Source/GameplayCameras/Public/Core/CameraVariableTableFwd.h b/Engine/Plugins/Cameras/GameplayCameras/Source/GameplayCameras/Public/Core/CameraVariableTableFwd.h
--- a/Engine/Plugins/Cameras/GameplayCameras/Source/GameplayCameras/Public/Core/CameraVariableTableFwd.h
+++ b/Engine/Plugins/Cameras/GameplayCameras/Source/GameplayCameras/Public/Core/CameraVariableTableFwd.h
@@ -207,6 +207,9 @@ struct FCameraVariableTableAllocationInfo
 	UPROPERTY()
 	TArray<TObjectPtr<UCameraVariableAsset>> AutoResetVariables;
 
+	/**Combines the given allocation info with this one. */
+	GAMEPLAYCAMERAS_API void Combine(const FCameraVariableTableAllocationInfo& OtherInfo);
+
 	GAMEPLAYCAMERAS_API friend bool operator==(const FCameraVariableTableAllocationInfo& A, const FCameraVariableTableAllocationInfo& B);
 };
 
diff --git a/Engine/Plugins/Cameras/GameplayCameras/Source/GameplayCameras/Public/Core/ICustomCameraNodeParameterProvider.h b/Engine/Plugins/Cameras/GameplayCameras/Source/GameplayCameras/Public/Core/ICustomCameraNodeParameterProvider.h
new file mode 100644
--- /dev/null
+++ b/Engine/Plugins/Cameras/GameplayCameras/Source/GameplayCameras/Public/Core/ICustomCameraNodeParameterProvider.h
@@ -0,0 +1,181 @@
+// Copyright Epic Games, Inc. All Rights Reserved.
+
+#pragma once
+
+#include "Core/CameraVariableTableFwd.h"
+#include "Core/CameraContextDataAllocationInfo.h"
+#include "UObject/Interface.h"
+#include "UObject/ObjectPtr.h"
+
+#include "ICustomCameraNodeParameterProvider.generated.h"
+
+class UCameraVariableAsset;
+
+namespace UE::Cameras
+{
+	class FCameraRigAssetBuilder;
+	namespace Internal { struct FInterfaceParameterBindingBuilder; }
+}
+
+/** Describes a custom camera blendable parameter. */
+USTRUCT()
+struct FCustomCameraNodeBlendableParameter
+{
+	GENERATED_BODY()
+
+	/** The name of the parameter. */
+	UPROPERTY()
+	FName ParameterName;
+
+	/** The type of the parameter. */
+	UPROPERTY()
+	ECameraVariableType ParameterType;
+
+	/** An optional camera variable that is dynamically driving the parameter's value. */
+	UPROPERTY()
+	TObjectPtr<UCameraVariableAsset> OverrideVariable;
+
+	GAMEPLAYCAMERAS_API friend bool operator==(const FCustomCameraNodeBlendableParameter& A, const FCustomCameraNodeBlendableParameter& B);
+};
+
+/** Describes a custom camera data parameter. */
+USTRUCT()
+struct FCustomCameraNodeDataParameter
+{
+	GENERATED_BODY()
+
+	/** The name of the parameter. */
+	UPROPERTY()
+	FName ParameterName;
+
+	/** The type of the parameter. */
+	UPROPERTY()
+	ECameraContextDataType ParameterType;
+
+	/** An extra type object for the parameter. */
+	UPROPERTY()
+	TObjectPtr<const UObject> ParameterTypeObject;
+
+	/** An optional context data ID for dynamically driving the parameter's value. */
+	UPROPERTY()
+	FCameraContextDataID OverrideDataID;
+
+	GAMEPLAYCAMERAS_API friend bool operator==(const FCustomCameraNodeDataParameter& A, const FCustomCameraNodeDataParameter& B);
+};
+
+/** Describes custom camera parameters. */
+USTRUCT()
+struct FCustomCameraNodeParameters
+{
+	GENERATED_BODY()
+
+	/** The list of blendable parameters. */
+	UPROPERTY()
+	TArray<FCustomCameraNodeBlendableParameter> BlendableParameters;
+
+	/** The list of data parameters. */
+	UPROPERTY()
+	TArray<FCustomCameraNodeDataParameter> DataParameters;
+
+	/** Returns whether this struct has any blendable or data parameter. */
+	bool HasAnyParameters() const { return !BlendableParameters.IsEmpty() || !DataParameters.IsEmpty(); }
+
+	/** Removes all parameters from this structure. */
+	void Reset() { BlendableParameters.Reset(); DataParameters.Reset(); }
+
+	GAMEPLAYCAMERAS_API friend bool operator==(const FCustomCameraNodeParameters& A, const FCustomCameraNodeParameters& B);
+};
+
+/**
+ * A structure for providing custom camera rig parameter information.
+ */
+struct FCustomCameraNodeParameterInfos
+{
+	/** Returns whether there are any blendable or data parameters. */
+	bool HasAnyParameters() const { return !BlendableParameters.IsEmpty() || !DataParameters.IsEmpty(); }
+
+	/** Declares a blendable parameter. */
+	template<typename VariableAssetType>
+	void AddBlendableParameter(
+			FName ParameterName, 
+			ECameraVariableType ParameterType, 
+			const uint8* DefaultValuePtr,
+			TObjectPtr<VariableAssetType>* OverrideVariable)
+	{
+		FObjectPtr* OverrideVariablePtr = reinterpret_cast<FObjectPtr*>(OverrideVariable);
+		BlendableParameters.Add({ 
+				ParameterName, 
+				ParameterType, 
+				DefaultValuePtr,
+				reinterpret_cast<TObjectPtr<UCameraVariableAsset>*>(OverrideVariable) });
+	}
+
+	/** Declares a blendable parameter. */
+	GAMEPLAYCAMERAS_API void AddBlendableParameter(
+			FName ParameterName, 
+			ECameraVariableType ParameterType, 
+			const uint8* DefaultValuePtr,
+			TObjectPtr<UCameraVariableAsset>* OverrideVariable);
+
+	/** Declares a data parameter. */
+	GAMEPLAYCAMERAS_API void AddDataParameter(
+			FName ParameterName, 
+			ECameraContextDataType ParameterType,
+			const UObject* ParameterTypeObject,
+			FCameraContextDataID* OverrideDataID);
+
+	/** Gets the list of blendable parameters. */
+	GAMEPLAYCAMERAS_API void GetBlendableParameters(TArray<FCustomCameraNodeBlendableParameter>& OutBlendableParameters) const;
+	/** Gets the list of data parameters. */
+	GAMEPLAYCAMERAS_API void GetDataParameters(TArray<FCustomCameraNodeDataParameter>& OutDataParameters) const;
+
+	/** Finds a blendable parameter of the given name. */
+	GAMEPLAYCAMERAS_API bool FindBlendableParameter(FName ParameterName, FCustomCameraNodeBlendableParameter& OutParameter) const;
+	/** Finds a data parameter of the given name. */
+	GAMEPLAYCAMERAS_API bool FindDataParameter(FName ParameterName, FCustomCameraNodeDataParameter& OutParameter) const;
+
+private:
+
+	struct FBlendableParameterInfo
+	{
+		FName ParameterName;
+		ECameraVariableType ParameterType;
+		const uint8* DefaultValuePtr = nullptr;
+		TObjectPtr<UCameraVariableAsset>* OverrideVariable = nullptr;
+	};
+
+	struct FDataParameterInfo
+	{
+		FName ParameterName;
+		ECameraContextDataType ParameterType;
+		const UObject* ParameterTypeObject;
+		FCameraContextDataID* OverrideDataID = nullptr;
+	};
+
+	TArray<FBlendableParameterInfo> BlendableParameters;
+	TArray<FDataParameterInfo> DataParameters;
+
+	friend class UE::Cameras::FCameraRigAssetBuilder;
+	friend struct UE::Cameras::Internal::FInterfaceParameterBindingBuilder;
+};
+
+UINTERFACE(MinimalAPI)
+class UCustomCameraNodeParameterProvider : public UInterface
+{
+	GENERATED_BODY()
+};
+
+/**
+ * An interface for camera nodes that want to expose a custom interface of
+ * blendable parameters and data parameters.
+ */
+class ICustomCameraNodeParameterProvider
+{
+	GENERATED_BODY()
+
+public:
+
+	/** Gathers the custom parameters on this node. */
+	virtual void GetCustomCameraNodeParameters(FCustomCameraNodeParameterInfos& OutParameterInfos) {}
+};
+
diff --git a/Engine/Plugins/Cameras/GameplayCameras/Source/GameplayCameras/Public/Debug/CameraNodeEvaluationResultDebugBlock.h b/Engine/Plugins/Cameras/GameplayCameras/Source/GameplayCameras/Public/Debug/CameraNodeEvaluationResultDebugBlock.h
--- a/Engine/Plugins/Cameras/GameplayCameras/Source/GameplayCameras/Public/Debug/CameraNodeEvaluationResultDebugBlock.h
+++ b/Engine/Plugins/Cameras/GameplayCameras/Source/GameplayCameras/Public/Debug/CameraNodeEvaluationResultDebugBlock.h
@@ -10,6 +10,7 @@ namespace UE::Cameras
 {
 
 class FCameraPoseDebugBlock;
+class FContextDataTableDebugBlock;
 class FVariableTableDebugBlock;
 struct FCameraDebugBlockBuilder;
 struct FCameraNodeEvaluationResult;
@@ -49,6 +50,8 @@ public:
 	FCameraPoseDebugBlock* GetCameraPoseDebugBlock();
 	/** Gets the variable table debug block. */
 	FVariableTableDebugBlock* GetVariableTableDebugBlock();
+	/** Gets the context data table debug block. */
+	FContextDataTableDebugBlock* GetContextDataTableDebugBlock();
 
 protected:
 
diff --git a/Engine/Plugins/Cameras/GameplayCameras/Source/GameplayCameras/Public/GameFramework/BlueprintCameraContextDataTable.h b/Engine/Plugins/Cameras/GameplayCameras/Source/GameplayCameras/Public/GameFramework/BlueprintCameraContextDataTable.h
new file mode 100644
--- /dev/null
+++ b/Engine/Plugins/Cameras/GameplayCameras/Source/GameplayCameras/Public/GameFramework/BlueprintCameraContextDataTable.h
@@ -0,0 +1,53 @@
+// Copyright Epic Games, Inc. All Rights Reserved.
+
+#pragma once
+
+#include "Core/CameraContextDataTableFwd.h"
+#include "CoreTypes.h"
+#include "Kismet/BlueprintFunctionLibrary.h"
+#include "StructUtils/InstancedStruct.h"
+#include "UObject/ObjectMacros.h"
+
+#include "BlueprintCameraContextDataTable.generated.h"
+
+namespace UE::Cameras
+{
+
+class FCameraContextDataTable;
+
+}  // namespace UE::Cameras
+
+USTRUCT(BlueprintType, DisplayName="Camera Context Data Table")
+struct GAMEPLAYCAMERAS_API FBlueprintCameraContextDataTable
+{
+	GENERATED_BODY()
+
+public:
+
+	using FCameraContextDataTable = UE::Cameras::FCameraContextDataTable;
+
+	FBlueprintCameraContextDataTable();
+	FBlueprintCameraContextDataTable(FCameraContextDataTable* InContextDataTable);
+
+	FCameraContextDataTable* GetContextDataTable() const { return PrivateContextDataTable; }
+
+
+private:
+
+	FCameraContextDataTable* PrivateContextDataTable = nullptr;
+};
+
+UCLASS()
+class UBlueprintCameraContextDataTableFunctionLibrary : public UBlueprintFunctionLibrary
+{
+	GENERATED_BODY()
+
+public:
+
+	UFUNCTION(BlueprintCallable, Category=Camera)
+	static FInstancedStruct GetStructData(const FBlueprintCameraContextDataTable& ContextDataTable, FCameraContextDataID DataID, const UScriptStruct* DataStructType);
+
+	UFUNCTION(BlueprintCallable, Category=Camera)
+	static void SetStructData(const FBlueprintCameraContextDataTable& ContextDataTable, FCameraContextDataID DataID, const FInstancedStruct& Data);
+};
+
diff --git a/Engine/Plugins/Cameras/GameplayCameras/Source/GameplayCameras/Public/GameFramework/CameraRigParameterInterop.h b/Engine/Plugins/Cameras/GameplayCameras/Source/GameplayCameras/Public/GameFramework/CameraRigParameterInterop.h
--- a/Engine/Plugins/Cameras/GameplayCameras/Source/GameplayCameras/Public/GameFramework/CameraRigParameterInterop.h
+++ b/Engine/Plugins/Cameras/GameplayCameras/Source/GameplayCameras/Public/GameFramework/CameraRigParameterInterop.h
@@ -32,7 +32,7 @@ public:
 public:
 
 	UFUNCTION(BlueprintCallable, Category="Camera", meta=(BlueprintInternalUseOnly="true"))
-	static void SetBooleanParameter(UPARAM(Ref) FBlueprintCameraVariableTable& VariableTable, UCameraRigAsset* CameraRig, const FString& ParameterName, bool bParameterValue);
+	static void SetBooleanParameter(UPARAM(Ref) FBlueprintCameraVariableTable& VariableTable, UCameraRigAsset* CameraRig, const FString& ParameterName, bool ParameterValue);
 
 	UFUNCTION(BlueprintCallable, Category="Camera", meta=(BlueprintInternalUseOnly="true"))
 	static void SetIntegerParameter(UPARAM(Ref) FBlueprintCameraVariableTable& VariableTable, UCameraRigAsset* CameraRig, const FString& ParameterName, int32 ParameterValue);
diff --git a/Engine/Plugins/Cameras/GameplayCameras/Source/GameplayCameras/Public/GameFramework/GameplayCameraComponent.h b/Engine/Plugins/Cameras/GameplayCameras/Source/GameplayCameras/Public/GameFramework/GameplayCameraComponent.h
--- a/Engine/Plugins/Cameras/GameplayCameras/Source/GameplayCameras/Public/GameFramework/GameplayCameraComponent.h
+++ b/Engine/Plugins/Cameras/GameplayCameras/Source/GameplayCameras/Public/GameFramework/GameplayCameraComponent.h
@@ -3,6 +3,7 @@
 #pragma once
 
 #include "Components/SceneComponent.h"
+#include "Core/CameraAssetReference.h"
 #include "Core/CameraEvaluationContext.h"
 #include "GameFramework/BlueprintCameraPose.h"
 #include "GameFramework/BlueprintCameraVariableTable.h"
@@ -82,10 +83,14 @@ public:
 	// USceneComponent interface.
 	virtual void OnUpdateTransform(EUpdateTransformFlags UpdateTransformFlags, ETeleportType Teleport) override;
 
+	// UObject interface.
+	void PostLoad() override;
+
 private:
 
 	void ActivateCameraEvaluationContext(int32 PlayerIndex);
 	void ActivateCameraEvaluationContext(APlayerController* PlayerController);
+	void UpdateCameraEvaluationContext(bool bApplyParameterOverrides);
 	void DeactivateCameraEvaluationContext();
 
 #if WITH_EDITORONLY_DATA
@@ -98,7 +103,7 @@ public:
 
 	/** The camera asset to run. */
 	UPROPERTY(BlueprintReadWrite, EditAnywhere, Category=Camera)
-	TObjectPtr<UCameraAsset> Camera;
+	FCameraAssetReference CameraReference;
 
 	/**
 	 * If AutoActivate is set, auto-activates this component's camera for the given player.
@@ -127,6 +132,11 @@ protected:
 	TObjectPtr<UStaticMeshComponent> PreviewMeshComponent;
 
 #endif	// WITH_EDITORONLY_DATA
+
+private:
+
+	UPROPERTY()
+	TObjectPtr<UCameraAsset> Camera_DEPRECATED;
 };
 
 namespace UE::Cameras
@@ -138,10 +148,6 @@ namespace UE::Cameras
 class FGameplayCameraComponentEvaluationContext : public FCameraEvaluationContext
 {
 	UE_DECLARE_CAMERA_EVALUATION_CONTEXT(GAMEPLAYCAMERAS_API, FGameplayCameraComponentEvaluationContext)
-
-public:
-
-	void Update(UGameplayCameraComponent* Owner);
 };
 
 }  // namespace UE::Cameras
diff --git a/Engine/Plugins/Cameras/GameplayCameras/Source/GameplayCameras/Public/GameplayCamerasDelegates.h b/Engine/Plugins/Cameras/GameplayCameras/Source/GameplayCameras/Public/GameplayCamerasDelegates.h
--- a/Engine/Plugins/Cameras/GameplayCameras/Source/GameplayCameras/Public/GameplayCamerasDelegates.h
+++ b/Engine/Plugins/Cameras/GameplayCameras/Source/GameplayCameras/Public/GameplayCamerasDelegates.h
@@ -5,6 +5,7 @@
 #include "CoreTypes.h"
 #include "Delegates/Delegate.h"
 
+class UCameraAsset;
 class UCameraRigAsset;
 
 namespace UE::Cameras
@@ -12,7 +13,8 @@ namespace UE::Cameras
 
 class FCameraBuildLog;
 
-DECLARE_MULTICAST_DELEGATE_TwoParams(FOnCameraRigAssetBuilt, UCameraRigAsset*, FCameraBuildLog&);
+DECLARE_MULTICAST_DELEGATE_OneParam(FOnCameraAssetBuilt, const UCameraAsset*);
+DECLARE_MULTICAST_DELEGATE_OneParam(FOnCameraRigAssetBuilt, const UCameraRigAsset*);
 
 /**
  * Global delegates for the GameplayCameras module.
@@ -21,7 +23,13 @@ class GAMEPLAYCAMERAS_API FGameplayCamerasDelegates
 {
 public:
 
-	/** Broadcast for when a camera rig has been built. */
+	/** Broadcast when a camera asset has been built. */
+	static inline FOnCameraAssetBuilt& OnCameraAssetBuilt()
+	{
+		return OnCameraAssetBuiltDelegates;
+	}
+
+	/** Broadcast when a camera rig has been built. */
 	static inline FOnCameraRigAssetBuilt& OnCameraRigAssetBuilt()
 	{
 		return OnCameraRigAssetBuiltDelegates;
@@ -29,6 +37,7 @@ public:
 
 private:
 
+	static FOnCameraAssetBuilt OnCameraAssetBuiltDelegates;
 	static FOnCameraRigAssetBuilt OnCameraRigAssetBuiltDelegates;
 };
 
diff --git a/Engine/Plugins/Cameras/GameplayCameras/Source/GameplayCameras/Public/Nodes/Common/CameraRigCameraNode.h b/Engine/Plugins/Cameras/GameplayCameras/Source/GameplayCameras/Public/Nodes/Common/CameraRigCameraNode.h
--- a/Engine/Plugins/Cameras/GameplayCameras/Source/GameplayCameras/Public/Nodes/Common/CameraRigCameraNode.h
+++ b/Engine/Plugins/Cameras/GameplayCameras/Source/GameplayCameras/Public/Nodes/Common/CameraRigCameraNode.h
@@ -5,6 +5,7 @@
 #include "Core/CameraNode.h"
 #include "Core/CameraRigAsset.h"
 #include "Core/CameraRigAssetReference.h"
+#include "Core/ICustomCameraNodeParameterProvider.h"
 
 #include "CameraRigCameraNode.generated.h"
 
@@ -12,7 +13,9 @@
  * A camera node that runs a camera rig's own node tree.
  */
 UCLASS(MinimalAPI, meta=(DisplayName="Camera Rig Prefab", CameraNodeCategories="Common,Utility"))
-class UCameraRigCameraNode : public UCameraNode
+class UCameraRigCameraNode 
+	: public UCameraNode
+	, public ICustomCameraNodeParameterProvider
 {
 	GENERATED_BODY()
 
@@ -23,50 +26,14 @@ protected:
 	virtual void OnBuild(FCameraRigBuildContext& BuildContext) override;
 	virtual FCameraNodeEvaluatorPtr OnBuildEvaluator(FCameraNodeEvaluatorBuilder& Builder) const override;
 
-	// UObject interface.
-	void PostLoad() override;
+	// ICustomCameraNodeParameterProvider interface.
+	virtual void GetCustomCameraNodeParameters(FCustomCameraNodeParameterInfos& OutParameterInfos) override;
 
 public:
 
 	/** The camera rig to run. */
 	UPROPERTY(EditAnywhere, Category=Common, meta=(ObjectTreeGraphHidden=true))
 	FCameraRigAssetReference CameraRigReference;
-
-private:
-
-	// Deprecated properties, predating FCameraRigAssetReference
-
-	UPROPERTY(meta=(ObjectTreeGraphHidden=true))
-	TObjectPtr<UCameraRigAsset> CameraRig_DEPRECATED;
-
-	UPROPERTY()
-	TArray<FBooleanCameraRigParameterOverride> BooleanOverrides_DEPRECATED;
-	UPROPERTY()
-	TArray<FInteger32CameraRigParameterOverride> Integer32Overrides_DEPRECATED;
-	UPROPERTY()
-	TArray<FFloatCameraRigParameterOverride> FloatOverrides_DEPRECATED;
-	UPROPERTY()
-	TArray<FDoubleCameraRigParameterOverride> DoubleOverrides_DEPRECATED;
-	UPROPERTY()
-	TArray<FVector2fCameraRigParameterOverride> Vector2fOverrides_DEPRECATED;
-	UPROPERTY()
-	TArray<FVector2dCameraRigParameterOverride> Vector2dOverrides_DEPRECATED;
-	UPROPERTY()
-	TArray<FVector3fCameraRigParameterOverride> Vector3fOverrides_DEPRECATED;
-	UPROPERTY()
-	TArray<FVector3dCameraRigParameterOverride> Vector3dOverrides_DEPRECATED;
-	UPROPERTY()
-	TArray<FVector4fCameraRigParameterOverride> Vector4fOverrides_DEPRECATED;
-	UPROPERTY()
-	TArray<FVector4dCameraRigParameterOverride> Vector4dOverrides_DEPRECATED;
-	UPROPERTY()
-	TArray<FRotator3fCameraRigParameterOverride> Rotator3fOverrides_DEPRECATED;
-	UPROPERTY()
-	TArray<FRotator3dCameraRigParameterOverride> Rotator3dOverrides_DEPRECATED;
-	UPROPERTY()
-	TArray<FTransform3fCameraRigParameterOverride> Transform3fOverrides_DEPRECATED;
-	UPROPERTY()
-	TArray<FTransform3dCameraRigParameterOverride> Transform3dOverrides_DEPRECATED;
 };
 
 namespace UE::Cameras
@@ -96,6 +63,7 @@ protected:
 private:
 
 	void ApplyParameterOverrides(FCameraVariableTable& OutVariableTable, bool bDrivenOnly);
+	void ApplyParameterOverrides(FCameraVariableTable& OutVariableTable, FCameraContextDataTable& OutContextDataTable, bool bDrivenOnly);
 
 private:
 
diff --git a/Engine/Plugins/Cameras/GameplayCameras/Source/GameplayCamerasEditor/GameplayCamerasEditor.Build.cs b/Engine/Plugins/Cameras/GameplayCameras/Source/GameplayCamerasEditor/GameplayCamerasEditor.Build.cs
--- a/Engine/Plugins/Cameras/GameplayCameras/Source/GameplayCamerasEditor/GameplayCamerasEditor.Build.cs
+++ b/Engine/Plugins/Cameras/GameplayCameras/Source/GameplayCamerasEditor/GameplayCamerasEditor.Build.cs
@@ -36,6 +36,7 @@ public class GameplayCamerasEditor : ModuleRules
 				"InputCore",
 				"InteractiveToolsFramework",
 				"Kismet",
+				"KismetWidgets",
 				"LevelEditor",
 				"Projects",
 				"RenderCore",
@@ -43,6 +44,8 @@ public class GameplayCamerasEditor : ModuleRules
 				"RHI",
 				"Slate",
 				"SlateCore",
+				"StructUtilsEditor",
+				"StructViewer",
 				"TimeManagement",
 				"ToolMenus",
 				"ToolWidgets",
