Cameras: big pass on parameterization of camera rigs

diff --git a/Engine/Plugins/Cameras/GameplayCameras/Source/GameplayCamerasEditor/Private/Customizations/CameraRigAssetReferenceDetailsCustomization.cpp b/Engine/Plugins/Cameras/GameplayCameras/Source/GameplayCamerasEditor/Private/Customizations/CameraRigAssetReferenceDetailsCustomization.cpp
--- a/Engine/Plugins/Cameras/GameplayCameras/Source/GameplayCamerasEditor/Private/Customizations/CameraRigAssetReferenceDetailsCustomization.cpp
+++ b/Engine/Plugins/Cameras/GameplayCameras/Source/GameplayCamerasEditor/Private/Customizations/CameraRigAssetReferenceDetailsCustomization.cpp
@@ -2,534 +2,169 @@
 
 #include "Customizations/CameraRigAssetReferenceDetailsCustomization.h"
 
-#include "Core/CameraParameters.h"
 #include "Core/CameraRigAsset.h"
 #include "Core/CameraRigAssetReference.h"
 #include "GameplayCamerasDelegates.h"
 #include "IDetailChildrenBuilder.h"
-#include "IDetailGroup.h"
-#include "IPropertyTypeCustomization.h"
 #include "IPropertyUtilities.h"
-#include "IStructureDataProvider.h"
+#include "PropertyBagDetails.h"
 #include "PropertyCustomizationHelpers.h"
-#include "Widgets/Images/SImage.h"
-#include "Widgets/Text/STextBlock.h"
 
 #define LOCTEXT_NAMESPACE "CameraRigAssetReferenceDetailsCustomization"
 
 namespace UE::Cameras
 {
 
-/**
- * A class that holds information regarding a given camera rig parameter override, and how it should
- * be shown in a details view panel.
- */
-class FCameraRigParameterOverrideDetailRow 
-	: public IStructureDataProvider 
-	, public TSharedFromThis<FCameraRigParameterOverrideDetailRow>
+class FCameraRigAssetParameterOverrideDataDetails : public FPropertyBagInstanceDataDetails
 {
 public:
 
-	FCameraRigParameterOverrideDetailRow(TSharedPtr<IPropertyHandle> InCameraRigReferenceProperty, TSharedPtr<IPropertyUtilities> InPropertyUtilities);
-
-	void Initialize(UCameraRigInterfaceParameter* InInterfaceParameter);
-	void InitializeAsInvalid(const FString& InInterfaceParameterName, const FGuid& InInterfaceParameterGuid);
-
-	void BuildDetailPropertyRow(IDetailChildrenBuilder& StructBuilder);
-
-	const FString& GetDisplayName() const;
-	
-private:
-
-	template<typename ParameterOverrideType>
-	void InitializeValues();
-
-	template<typename ParameterOverrideType>
-	void BuildDetailPropertyRowImpl(IDetailChildrenBuilder& StructBuilder);
-
-	void AccessCameraRigReferences(TArray<FCameraRigAssetReference*>& OutCameraRigReferences) const;
-	void ModifyOuterObjects() const;
-
-	template<typename ParameterOverrideType, typename CameraParameterType = typename ParameterOverrideType::CameraParameterType>
-	void OnPropertyValueChanged();
-	template<typename ParameterOverrideType>
-	bool OnIsResetToDefaultVisible(TSharedPtr<IPropertyHandle> PropertyHandle);
-	template<typename ParameterOverrideType>
-	void OnResetToDefault(TSharedPtr<IPropertyHandle> PropertyHandle);
-
-	void OnRemoveInvalidOverride();
-
-public:
-
-	// IStructureDataProvider interface.
-	virtual bool IsValid() const override
-	{
-		return true; 
-	}
-	
-	virtual const UStruct* GetBaseStructure() const override
-	{
-		return ParameterType;
-	}
-	
-	virtual void GetInstances(TArray<TSharedPtr<FStructOnScope>>& OutInstances, const UStruct* ExpectedBaseStructure) const override
+	FCameraRigAssetParameterOverrideDataDetails(
+			TSharedPtr<IPropertyHandle> InStructPropertyHandle,
+			TSharedPtr<IPropertyHandle> InParametersPropertyHandle,
+			TSharedPtr<IPropertyUtilities>& InPropertyUtilities)
+		: FPropertyBagInstanceDataDetails(InParametersPropertyHandle, InPropertyUtilities, true)
+		, StructPropertyHandle(InStructPropertyHandle)
 	{
-		OutInstances = ParameterOverrideStructs;
 	}
 
-private:
-
-	TSharedPtr<IPropertyHandle> CameraRigReferenceProperty;
-	TSharedPtr<IPropertyUtilities> PropertyUtilities;
-
-	UCameraRigInterfaceParameter* InterfaceParameter = nullptr;
-	FString InvalidInterfaceParameterName;
-	FGuid InvalidInterfaceParameterGuid;
-
-	const UStruct* ParameterType = nullptr;
-	TArray<TSharedPtr<FStructOnScope>> ParameterOverrideStructs;
-	TSharedPtr<FStructOnScope> DefaultValue;
-};
-
-FCameraRigParameterOverrideDetailRow::FCameraRigParameterOverrideDetailRow(TSharedPtr<IPropertyHandle> InCameraRigReferenceProperty, TSharedPtr<IPropertyUtilities> InPropertyUtilities)
-	: CameraRigReferenceProperty(InCameraRigReferenceProperty)
-	, PropertyUtilities(InPropertyUtilities)
-{
-}
-
-void FCameraRigParameterOverrideDetailRow::Initialize(UCameraRigInterfaceParameter* InInterfaceParameter)
-{
-	InterfaceParameter = InInterfaceParameter;
+protected:
 
-	if (InterfaceParameter && InterfaceParameter->PrivateVariable != nullptr)
+	struct FCameraRigAssetReferenceOverrideProvider : public IPropertyBagOverrideProvider
 	{
-		switch (InterfaceParameter->PrivateVariable->GetVariableType())
+		FCameraRigAssetReferenceOverrideProvider(FCameraRigAssetReference& InCameraRigAssetReference)
+			: CameraRigAssetReference(InCameraRigAssetReference)
 		{
-#define UE_CAMERA_VARIABLE_FOR_TYPE(ValueType, ValueName)\
-		case ECameraVariableType::ValueName:\
-			InitializeValues<F##ValueName##CameraRigParameterOverride>();\
-			break;
-			UE_CAMERA_VARIABLE_FOR_ALL_TYPES()
-#undef UE_CAMERA_VARIABLE_FOR_TYPE
 		}
-	}
-}
-
-void FCameraRigParameterOverrideDetailRow::InitializeAsInvalid(const FString& InInterfaceParameterName, const FGuid& InInterfaceParameterGuid)
-{
-	InvalidInterfaceParameterName = InInterfaceParameterName;
-	InvalidInterfaceParameterGuid = InInterfaceParameterGuid;
-}
-
-template<typename ParameterOverrideType>
-void FCameraRigParameterOverrideDetailRow::InitializeValues()
-{
-	using CameraParameterType = typename ParameterOverrideType::CameraParameterType;
-	using VariableAssetType = typename CameraParameterType::VariableAssetType;
-
-	// Get the camera parameter type.
-	ParameterType = CameraParameterType::StaticStruct();
-
-	// Get the default value of the camera parameter for this override.
-	const VariableAssetType* TypedPrivateVariable = CastChecked<VariableAssetType>(InterfaceParameter->PrivateVariable);
-	const CameraParameterType DefaultValueParameter(TypedPrivateVariable->GetDefaultValue());
-	DefaultValue = MakeShared<TStructOnScope<CameraParameterType>>(DefaultValueParameter);
-
-	// Add all the actual parameter override values, if any, or the default value, if not.
-	TArray<FCameraRigAssetReference*> CameraRigReferences;
-	AccessCameraRigReferences(CameraRigReferences);
-	for (FCameraRigAssetReference* CameraRigReference : CameraRigReferences)
-	{
-		check(CameraRigReference);
-
-		FCameraRigParameterOverrides& ParameterOverrides = CameraRigReference->GetParameterOverrides();
-		ParameterOverrideType* FoundOverride = ParameterOverrides.FindParameterOverride<ParameterOverrideType>(
-				InterfaceParameter->Guid);
-		if (FoundOverride)
+		
+		virtual bool IsPropertyOverridden(const FGuid PropertyID) const override
 		{
-			ParameterOverrideStructs.Add(MakeShared<TStructOnScope<CameraParameterType>>(FoundOverride->Value));
+			return CameraRigAssetReference.IsParameterOverriden(PropertyID);
 		}
-		else
+		
+		virtual void SetPropertyOverride(const FGuid PropertyID, const bool bIsOverridden) const override
 		{
-			ParameterOverrideStructs.Add(MakeShared<TStructOnScope<CameraParameterType>>(DefaultValueParameter));
+			CameraRigAssetReference.SetParameterOverriden(PropertyID, bIsOverridden);
 		}
-	}
-}
 
-void FCameraRigParameterOverrideDetailRow::BuildDetailPropertyRow(IDetailChildrenBuilder& StructBuilder)
-{
-	if (InterfaceParameter == nullptr)
-	{
-		// This is an old parameter override that isn't valid anymore (e.g. the camera rig's interface
-		// parameter was deleted).
-		FResetToDefaultOverride ResetToDefault = FResetToDefaultOverride::Create(
-				TAttribute<bool>(true),
-				FSimpleDelegate::CreateSP(this, &FCameraRigParameterOverrideDetailRow::OnRemoveInvalidOverride));
+	private:
+		FCameraRigAssetReference& CameraRigAssetReference;
+	};
 
-		StructBuilder.AddCustomRow(FText::FromString(InvalidInterfaceParameterName))
-			.NameContent()
-			[
-				SNew(STextBlock)
-				.Font(IPropertyTypeCustomizationUtils::GetRegularFont())
-				.Text(FText::FromString(InvalidInterfaceParameterName))
-			]
-			.ValueContent()
-			[
-				SNew(SHorizontalBox)
-				+SHorizontalBox::Slot()
-				.AutoWidth()
-				[
-					SNew(SImage)
-					.Image(FAppStyle::GetBrush("Icons.Warning"))
-				]
-				+SHorizontalBox::Slot()
-				.FillWidth(1.f)
-				[
-					SNew(STextBlock)
-					.Font(IPropertyTypeCustomizationUtils::GetRegularFont())
-					.Text(LOCTEXT("InvalidParameterOverrideWarning", "No such parameter found."))
-				]
-			]
-			.OverrideResetToDefault(ResetToDefault);
-	}
-	else if (InterfaceParameter->PrivateVariable == nullptr)
+	virtual bool HasPropertyOverrides() const override
 	{
-		// The camera rig wasn't built, show a message.
-		StructBuilder.AddCustomRow(FText::FromString(InterfaceParameter->InterfaceParameterName))
-			.NameContent()
-			[
-				SNew(STextBlock)
-				.Font(IPropertyTypeCustomizationUtils::GetRegularFont())
-				.Text(FText::FromString(InterfaceParameter->InterfaceParameterName))
-			]
-			.ValueContent()
-			[
-				SNew(STextBlock)
-				.Font(IPropertyTypeCustomizationUtils::GetRegularFont())
-				.Text(LOCTEXT("UnbuiltInnerCameraRigWarning", "Please build the child camera rig."))
-			];
+		return true;
 	}
-	else
-	{
-		switch (InterfaceParameter->PrivateVariable->GetVariableType())
-		{
-#define UE_CAMERA_VARIABLE_FOR_TYPE(ValueType, ValueName)\
-			case ECameraVariableType::ValueName:\
-				BuildDetailPropertyRowImpl<F##ValueName##CameraRigParameterOverride>(StructBuilder);\
-				break;
-			UE_CAMERA_VARIABLE_FOR_ALL_TYPES()
-#undef UE_CAMERA_VARIABLE_FOR_TYPE
-		}
-	}
-}
 
-const FString& FCameraRigParameterOverrideDetailRow::GetDisplayName() const
-{
-	if (InterfaceParameter)
+	virtual void PreChangeOverrides() override
 	{
-		return InterfaceParameter->InterfaceParameterName;
+		StructPropertyHandle->NotifyPreChange();
 	}
-	else
-	{
-		return InvalidInterfaceParameterName;
-	}
-}
 
-void FCameraRigParameterOverrideDetailRow::AccessCameraRigReferences(TArray<FCameraRigAssetReference*>& OutCameraRigReferences) const
-{
-	TArray<void*> RawData;
-	CameraRigReferenceProperty->AccessRawData(RawData);
-	for (int32 Index = 0; Index < RawData.Num(); ++Index)
+	virtual void PostChangeOverrides() override
 	{
-		OutCameraRigReferences.Add(static_cast<FCameraRigAssetReference*>(RawData[Index]));
+		StructPropertyHandle->NotifyPostChange(EPropertyChangeType::ValueSet);
+		StructPropertyHandle->NotifyFinishedChangingProperties();
 	}
-}
 
-void FCameraRigParameterOverrideDetailRow::ModifyOuterObjects() const
-{
-	TArray<UObject*> OuterObjects;
-	CameraRigReferenceProperty->GetOuterObjects(OuterObjects);
-	for (UObject* OuterObject : OuterObjects)
+	virtual void EnumeratePropertyBags(TSharedPtr<IPropertyHandle> PropertyBagHandle, const EnumeratePropertyBagFuncRef& Func) const override
 	{
-		OuterObject->Modify();
+		StructPropertyHandle->EnumerateRawData([Func](void* RawData, const int32 DataIndex, const int32 NumDatas)
+		{
+			if (FCameraRigAssetReference* CameraRigAssetReference = static_cast<FCameraRigAssetReference*>(RawData))
+			{
+				if (const UCameraRigAsset* CameraRigAsset = CameraRigAssetReference->GetCameraRig())
+				{
+					const FInstancedPropertyBag& DefaultParameters = CameraRigAsset->GetDefaultParameters();
+					FInstancedPropertyBag& Parameters = CameraRigAssetReference->GetParameters();
+					FCameraRigAssetReferenceOverrideProvider OverrideProvider(*CameraRigAssetReference);
+					if (!Func(DefaultParameters, Parameters, OverrideProvider))
+					{
+						return false;
+					}
+				}
+			}
+			return true;
+		});
 	}
-}
 
-template<typename ParameterOverrideType>
-void FCameraRigParameterOverrideDetailRow::BuildDetailPropertyRowImpl(IDetailChildrenBuilder& StructBuilder)
-{
-	using CameraParameterType = typename ParameterOverrideType::CameraParameterType;
-
-	// Add a new property row to the detail view. The row shows our copy of a camera parameter that
-	// corresponds to the sort of camera parameter needed to override a given camera rig parameter.
-	IDetailPropertyRow* ParameterOverrideRow = StructBuilder.AddExternalStructure(SharedThis(this));
-
-	ParameterOverrideRow->DisplayName(FText::FromString(InterfaceParameter->InterfaceParameterName));
-
-	// Set up callbacks so that when the user edits the camera parameter, we replicate the edit onto
-	// the parameter overrides list (by adding/removing the override, or updating its value).
-	ParameterOverrideRow->GetPropertyHandle()->SetOnPropertyValueChanged(FSimpleDelegate::CreateSP(
-				this, &FCameraRigParameterOverrideDetailRow::OnPropertyValueChanged<ParameterOverrideType>));
-	ParameterOverrideRow->GetPropertyHandle()->SetOnChildPropertyValueChanged(FSimpleDelegate::CreateSP(
-				this, &FCameraRigParameterOverrideDetailRow::OnPropertyValueChanged<ParameterOverrideType>));
-
-	// We need to add our own custom reset-to-default logic, which sets the default value as determined
-	// by the values inside the camera rig prefab. However, camera parameters also have a detail 
-	// customization that has a custom reset-to-default logic. Slate complaines when both the outer and
-	// inner rows have custom reset-to-default logic, so suppress the inner one.
-	ParameterOverrideRow->GetPropertyHandle()->SetInstanceMetaData("NoCustomCameraParameterResetToDefault", "true");
-
-	FResetToDefaultOverride ResetToDefault = FResetToDefaultOverride::Create(
-			FIsResetToDefaultVisible::CreateSP(
-				this, &FCameraRigParameterOverrideDetailRow::OnIsResetToDefaultVisible<ParameterOverrideType>),
-			FResetToDefaultHandler::CreateSP(
-				this, &FCameraRigParameterOverrideDetailRow::OnResetToDefault<ParameterOverrideType>));
-	ParameterOverrideRow->OverrideResetToDefault(ResetToDefault);
-}
+private:
+	
+	TSharedPtr<IPropertyHandle> StructPropertyHandle;
+};
 
-template<typename ParameterOverrideType, typename CameraParameterType>
-void FCameraRigParameterOverrideDetailRow::OnPropertyValueChanged()
+TSharedRef<IPropertyTypeCustomization> FCameraRigAssetReferenceDetailsCustomization::MakeInstance()
 {
-	const CameraParameterType* TypedDefaultValue = reinterpret_cast<CameraParameterType*>(DefaultValue->GetStructMemory());
-
-	TArray<FCameraRigAssetReference*> CameraRigReferences;
-	AccessCameraRigReferences(CameraRigReferences);
-	check(CameraRigReferences.Num() == ParameterOverrideStructs.Num());
-
-	ModifyOuterObjects();
-
-	for (int32 Index = 0; Index < CameraRigReferences.Num(); ++Index)
-	{
-		FCameraRigAssetReference* CameraRigReference(CameraRigReferences[Index]);
-		TSharedPtr<FStructOnScope> ParameterOverrideStruct(ParameterOverrideStructs[Index]);
-
-		const CameraParameterType* TypedOverrideParameter = reinterpret_cast<CameraParameterType*>(
-				ParameterOverrideStruct->GetStructMemory());
-
-		FCameraRigParameterOverrides& ParameterOverrides = CameraRigReference->GetParameterOverrides();
-		const bool bEqualValues = CameraParameterValueEquals<typename CameraParameterType::ValueType>(
-				TypedOverrideParameter->Value, TypedDefaultValue->Value);
-		if (bEqualValues && TypedOverrideParameter->Variable == nullptr)
-		{
-			ParameterOverrides.RemoveParameterOverride<ParameterOverrideType>(InterfaceParameter->Guid);
-		}
-		else
-		{
-			ParameterOverrideType& Override = ParameterOverrides.FindOrAddParameterOverride<ParameterOverrideType>(InterfaceParameter);
-			Override.Value = *TypedOverrideParameter;
-		}
-	}
+	return MakeShared<FCameraRigAssetReferenceDetailsCustomization>();
 }
 
-template<typename ParameterOverrideType>
-bool FCameraRigParameterOverrideDetailRow::OnIsResetToDefaultVisible(TSharedPtr<IPropertyHandle> PropertyHandle)
+FCameraRigAssetReferenceDetailsCustomization::~FCameraRigAssetReferenceDetailsCustomization()
 {
-	using CameraParameterType = typename ParameterOverrideType::CameraParameterType;
-
-	const CameraParameterType* TypedDefaultValue = reinterpret_cast<CameraParameterType*>(DefaultValue->GetStructMemory());
-
-	for (TSharedPtr<FStructOnScope> ParameterOverrideStruct : ParameterOverrideStructs)
-	{
-		const CameraParameterType* TypedOverrideParameter = reinterpret_cast<CameraParameterType*>(ParameterOverrideStruct->GetStructMemory());
-		if (TypedOverrideParameter->Variable != nullptr ||
-				!CameraParameterValueEquals<typename CameraParameterType::ValueType>(
-					TypedOverrideParameter->Value, TypedDefaultValue->Value))
-		{
-			return true;
-		}
-	}
-
-	return false;
+	FGameplayCamerasDelegates::OnCameraRigAssetBuilt().RemoveAll(this);
 }
 
-template<typename ParameterOverrideType>
-void FCameraRigParameterOverrideDetailRow::OnResetToDefault(TSharedPtr<IPropertyHandle> PropertyHandle)
+void FCameraRigAssetReferenceDetailsCustomization::CustomizeHeader(TSharedRef<IPropertyHandle> InStructPropertyHandle, FDetailWidgetRow& InHeaderRow, IPropertyTypeCustomizationUtils& InCustomizationUtils)
 {
-	using CameraParameterType = typename ParameterOverrideType::CameraParameterType;
+	StructPropertyHandle = InStructPropertyHandle;
+	PropertyUtilities = InCustomizationUtils.GetPropertyUtilities();
 
-	const CameraParameterType* TypedDefaultValue = reinterpret_cast<CameraParameterType*>(DefaultValue->GetStructMemory());
+	CameraRigAssetPropertyHandle = StructPropertyHandle->GetChildHandle(GET_MEMBER_NAME_CHECKED(FCameraRigAssetReference, CameraRig));
+	ParametersPropertyHandle = StructPropertyHandle->GetChildHandle(GET_MEMBER_NAME_CHECKED(FCameraRigAssetReference, Parameters));
 
-	ModifyOuterObjects();
-
-	for (TSharedPtr<FStructOnScope> ParameterOverrideStruct : ParameterOverrideStructs)
-	{
-		CameraParameterType* TypedOverrideParameter = reinterpret_cast<CameraParameterType*>(ParameterOverrideStruct->GetStructMemory());
+	CameraRigAssetPropertyHandle->SetOnPropertyValueChanged(
+			FSimpleDelegate::CreateSP(this, &FCameraRigAssetReferenceDetailsCustomization::RebuildParametersIfNeeded));
 
-		TypedOverrideParameter->Value = TypedDefaultValue->Value;
-		TypedOverrideParameter->Variable = nullptr;
-	}
-}
-
-void FCameraRigParameterOverrideDetailRow::OnRemoveInvalidOverride()
-{
-	if (InvalidInterfaceParameterGuid.IsValid())
+	if (!GIsTransacting)
 	{
-		ModifyOuterObjects();
-
-		TArray<FCameraRigAssetReference*> CameraRigReferences;
-		AccessCameraRigReferences(CameraRigReferences);
-
-		for (FCameraRigAssetReference* CameraRigReference : CameraRigReferences)
-		{
-			FCameraRigParameterOverrides& ParameterOverrides = CameraRigReference->GetParameterOverrides();
-#define UE_CAMERA_VARIABLE_FOR_TYPE(ValueType, ValueName)\
-			ParameterOverrides.RemoveParameterOverride<F##ValueName##CameraRigParameterOverride>(InvalidInterfaceParameterGuid);
-			UE_CAMERA_VARIABLE_FOR_ALL_TYPES()
-#undef UE_CAMERA_VARIABLE_FOR_TYPE
-		}
-
-		if (PropertyUtilities)
-		{
-			PropertyUtilities->RequestForceRefresh();
-		}
+		RebuildParametersIfNeeded();
 	}
-}
 
-TSharedRef<IPropertyTypeCustomization> FCameraRigAssetReferenceDetailsCustomization::MakeInstance()
-{
-	return MakeShared<FCameraRigAssetReferenceDetailsCustomization>();
-}
-
-void FCameraRigAssetReferenceDetailsCustomization::CustomizeHeader(TSharedRef<IPropertyHandle> StructPropertyHandle, FDetailWidgetRow& HeaderRow, IPropertyTypeCustomizationUtils& StructCustomizationUtils)
-{
-	// Remember a few things.
-	CameraRigReferenceProperty = StructPropertyHandle;
-	PropertyUtilities = StructCustomizationUtils.GetPropertyUtilities();
-
-	// Display the camera rig picker widget.
-	TSharedPtr<IPropertyHandle> CameraRigProperty = StructPropertyHandle->GetChildHandle(
-			GET_MEMBER_NAME_CHECKED(FCameraRigAssetReference, CameraRig));
-
-	HeaderRow
+	InHeaderRow
+	.ShouldAutoExpand(true)
 	.NameContent()
 	[
-		CameraRigProperty->CreatePropertyNameWidget()
+		CameraRigAssetPropertyHandle->CreatePropertyNameWidget()
 	]
 	.ValueContent()
 	[
-		CameraRigProperty->CreatePropertyValueWidgetWithCustomization(nullptr)
-	]
-	.ShouldAutoExpand(true);
+		CameraRigAssetPropertyHandle->CreatePropertyValueWidgetWithCustomization(nullptr)
+	];
 	
-	// Setup a callback so that when a different camera rig is picked, we refresh the list
-	// of parameter overrides.
-	CameraRigProperty->SetOnPropertyValueChanged(
-			FSimpleDelegate::CreateSP(this, &FCameraRigAssetReferenceDetailsCustomization::OnCameraRigChanged));
-
-	// Setup a similar callback for when the selected camera rig get (re)built.
-	FGameplayCamerasDelegates::OnCameraRigAssetBuilt().AddSP(this, &FCameraRigAssetReferenceDetailsCustomization::OnCameraRigBuilt);
+	FGameplayCamerasDelegates::OnCameraRigAssetBuilt().AddSP(this, &FCameraRigAssetReferenceDetailsCustomization::OnCameraRigAssetBuilt);
 }
 
-void FCameraRigAssetReferenceDetailsCustomization::CustomizeChildren(TSharedRef<IPropertyHandle> StructPropertyHandle, IDetailChildrenBuilder& StructBuilder, IPropertyTypeCustomizationUtils& StructCustomizationUtils)
+void FCameraRigAssetReferenceDetailsCustomization::CustomizeChildren(TSharedRef<IPropertyHandle> InStructPropertyHandle, IDetailChildrenBuilder& InChildrenBuilder, IPropertyTypeCustomizationUtils& InCustomizationUtils)
 {
-	// Update our list of parameter overrides and add each of them as a detail view row
-	// below the camera rig picker.
-	UpdateParameterOverrides(nullptr, false);
-	BuildParameterOverrideRows(StructBuilder);
+	const TSharedRef<FCameraRigAssetParameterOverrideDataDetails> ParameterOverrideDataDetails = 
+		MakeShared<FCameraRigAssetParameterOverrideDataDetails>(StructPropertyHandle, ParametersPropertyHandle, PropertyUtilities);
+	InChildrenBuilder.AddCustomBuilder(ParameterOverrideDataDetails);
 }
 
-void FCameraRigAssetReferenceDetailsCustomization::OnCameraRigChanged()
+void FCameraRigAssetReferenceDetailsCustomization::OnCameraRigAssetBuilt(const UCameraRigAsset* CameraRig)
 {
-	UpdateParameterOverrides(nullptr, true);
+	RebuildParametersIfNeeded();
 }
 
-void FCameraRigAssetReferenceDetailsCustomization::OnCameraRigBuilt(UCameraRigAsset* CameraRig, FCameraBuildLog& BuildLog)
+void FCameraRigAssetReferenceDetailsCustomization::RebuildParametersIfNeeded()
 {
-	UpdateParameterOverrides(CameraRig, true);
-}
-
-void FCameraRigAssetReferenceDetailsCustomization::UpdateParameterOverrides(const UCameraRigAsset* CameraRigToUpdate, bool bRequestRefresh)
-{
-	// Get all the camera rig references we are editing.
 	TArray<void*> RawData;
-	CameraRigReferenceProperty->AccessRawData(RawData);
-
-	TArray<FCameraRigAssetReference*> CameraRigReferences;
-	for (int32 Index = 0; Index < RawData.Num(); Index++)
-	{
-		FCameraRigAssetReference* CameraRigReference = static_cast<FCameraRigAssetReference*>(RawData[Index]);
-		check(CameraRigReference);
-		CameraRigReferences.Add(CameraRigReference);
-	}
-
-	// Update the parameter overrides for the given camera rig, or for all the ones we're editing if none
-	// was specified. As we go, determine if we're editing parameter overrides for the same camera rig
-	// across all edited references, or if we have a mix.
-	bool bIsEditingNullCameraRig = false;
-	TSet<const UCameraRigAsset*> UsedCameraRigs;
-	for (FCameraRigAssetReference* CameraRigReference : CameraRigReferences)
-	{
-		const UCameraRigAsset* CameraRig = CameraRigReference->GetCameraRig();
-		if (CameraRigToUpdate == nullptr || CameraRigToUpdate == CameraRig)
-		{
-			CameraRigReference->UpdateParameterOverrides();
-		}
+	StructPropertyHandle->AccessRawData(RawData);
 
-		if (CameraRig)
-		{
-			UsedCameraRigs.Add(CameraRig);
-		}
-		else
-		{
-			bIsEditingNullCameraRig = true;
-		}
-	}
+	bool bRebuiltAny = false;
 
-	// Re-create our list of parameter overrides.
-	ParameterOverrideRows.Reset();
-	TSet<FGuid> UsedInterfaceParameterGuids;
-	if (UsedCameraRigs.Num() == 1 && !bIsEditingNullCameraRig)
+	for (int32 Index = 0; Index < RawData.Num(); ++Index)
 	{
-		// Add rows for all of the camera rig's parameters.
-		const UCameraRigAsset* CameraRig = UsedCameraRigs.Array()[0];
-		for (UCameraRigInterfaceParameter* InterfaceParameter : CameraRig->Interface.InterfaceParameters)
-		{
-			TSharedPtr<FCameraRigParameterOverrideDetailRow> Row = MakeShared<FCameraRigParameterOverrideDetailRow>(
-					CameraRigReferenceProperty, PropertyUtilities);
-			Row->Initialize(InterfaceParameter);
-			ParameterOverrideRows.Add(Row);
-			UsedInterfaceParameterGuids.Add(InterfaceParameter->Guid);
-		}
-
-		// Add any left-over/invalid parameter overrides.
-		for (FCameraRigAssetReference* CameraRigReference : CameraRigReferences)
+		FCameraRigAssetReference* CameraRigAssetReference = static_cast<FCameraRigAssetReference*>(RawData[Index]);
+		if (CameraRigAssetReference)
 		{
-			FCameraRigParameterOverrides& ParameterOverrides = CameraRigReference->GetParameterOverrides();
-#define UE_CAMERA_VARIABLE_FOR_TYPE(ValueType, ValueName)\
-			for (F##ValueName##CameraRigParameterOverride& ParameterOverride : ParameterOverrides.Get##ValueName##Overrides())\
-			{\
-				if (!UsedInterfaceParameterGuids.Contains(ParameterOverride.InterfaceParameterGuid))\
-				{\
-					TSharedPtr<FCameraRigParameterOverrideDetailRow> InvalidRow = MakeShared<FCameraRigParameterOverrideDetailRow>(\
-							CameraRigReferenceProperty, PropertyUtilities);\
-					InvalidRow->InitializeAsInvalid(ParameterOverride.InterfaceParameterName, ParameterOverride.InterfaceParameterGuid);\
-					ParameterOverrideRows.Add(InvalidRow);\
-				}\
-			}
-			UE_CAMERA_VARIABLE_FOR_ALL_TYPES()
-#undef UE_CAMERA_VARIABLE_FOR_TYPE
+			const bool bRebuiltOne = CameraRigAssetReference->RebuildParametersIfNeeded();
+			bRebuiltAny |= bRebuiltOne;
 		}
-
-		// Sort rows by name.
-		ParameterOverrideRows.StableSort(
-				[](TSharedPtr<FCameraRigParameterOverrideDetailRow> A,
-					TSharedPtr<FCameraRigParameterOverrideDetailRow> B)
-				{
-					return A->GetDisplayName().Compare(B->GetDisplayName()) < 0;
-				});
 	}
 
-	if (bRequestRefresh && PropertyUtilities)
-	{
-		PropertyUtilities->RequestForceRefresh();
-	}
-}
-
-void FCameraRigAssetReferenceDetailsCustomization::BuildParameterOverrideRows(IDetailChildrenBuilder& StructBuilder)
-{
-	for (TSharedPtr<FCameraRigParameterOverrideDetailRow> Row : ParameterOverrideRows)
+	if (bRebuiltAny && PropertyUtilities)
 	{
-		Row->BuildDetailPropertyRow(StructBuilder);
+		PropertyUtilities->RequestRefresh();
 	}
 }
 
diff --git a/Engine/Plugins/Cameras/GameplayCameras/Source/GameplayCamerasEditor/Private/Editors/CameraNodeGraphDragDropOp.cpp b/Engine/Plugins/Cameras/GameplayCameras/Source/GameplayCamerasEditor/Private/Editors/CameraNodeGraphDragDropOp.cpp
new file mode 100644
--- /dev/null
+++ b/Engine/Plugins/Cameras/GameplayCameras/Source/GameplayCamerasEditor/Private/Editors/CameraNodeGraphDragDropOp.cpp
@@ -0,0 +1,63 @@
+// Copyright Epic Games, Inc. All Rights Reserved.
+
+#include "Editors/CameraNodeGraphDragDropOp.h"
+
+#include "Core/CameraRigAsset.h"
+#include "Editors/CameraNodeGraphSchema.h"
+#include "GraphEditor.h"
+#include "ScopedTransaction.h"
+
+#define LOCTEXT_NAMESPACE "CameraNodeGraphDragDropOp"
+
+TSharedRef<FCameraNodeGraphInterfaceParameterDragDropOp> FCameraNodeGraphInterfaceParameterDragDropOp::New(UCameraRigInterfaceParameterBase* InInterfaceParameter)
+{
+	TSharedRef<FCameraNodeGraphInterfaceParameterDragDropOp> Operation = MakeShared<FCameraNodeGraphInterfaceParameterDragDropOp>();
+	Operation->InterfaceParameter = InInterfaceParameter;
+	Operation->Construct();
+	return Operation;
+}
+
+FReply FCameraNodeGraphInterfaceParameterDragDropOp::ExecuteDragOver(TSharedPtr<SGraphEditor> GraphEditor)
+{
+	if (InterfaceParameter)
+	{
+		if (!InterfaceParameter->bHasGraphNode)
+		{
+			const FSlateBrush* OKIcon = FAppStyle::GetBrush(TEXT("Graph.ConnectorFeedback.OK"));
+			SetToolTip(LOCTEXT("OnDragOver_Success", "Add interface parameter"), OKIcon);
+		}
+		else
+		{
+			const FSlateBrush* ErrorIcon = FAppStyle::GetBrush(TEXT("Graph.ConnectorFeedback.Error"));
+			SetToolTip(
+					LOCTEXT("OnDragOver_Error", "This interface parameter is already in the graph"),
+					ErrorIcon);
+		}
+	}
+
+	return FReply::Handled();
+}
+
+FReply FCameraNodeGraphInterfaceParameterDragDropOp::ExecuteDrop(TSharedPtr<SGraphEditor> GraphEditor, const FVector2D& NewLocation)
+{
+	if (!InterfaceParameter || InterfaceParameter->bHasGraphNode)
+	{
+		return FReply::Handled();
+	}
+
+	const FScopedTransaction Transaction(LOCTEXT("DropObjectClasses", "Drop New Nodes"));
+
+	UEdGraph* Graph = GraphEditor->GetCurrentGraph();
+
+	GraphEditor->ClearSelectionSet();
+
+	FCameraNodeGraphSchemaAction_AddInterfaceParameterNode Action;
+	Action.InterfaceParameter = InterfaceParameter;
+	UEdGraphNode* NewNode = Action.PerformAction(Graph, nullptr, NewLocation, false);
+	GraphEditor->SetNodeSelection(NewNode, true);
+
+	return FReply::Handled();
+}
+
+#undef LOCTEXT_NAMESPACE
+
diff --git a/Engine/Plugins/Cameras/GameplayCameras/Source/GameplayCamerasEditor/Private/Editors/CameraNodeGraphNode.cpp b/Engine/Plugins/Cameras/GameplayCameras/Source/GameplayCamerasEditor/Private/Editors/CameraNodeGraphNode.cpp
--- a/Engine/Plugins/Cameras/GameplayCameras/Source/GameplayCamerasEditor/Private/Editors/CameraNodeGraphNode.cpp
+++ b/Engine/Plugins/Cameras/GameplayCameras/Source/GameplayCamerasEditor/Private/Editors/CameraNodeGraphNode.cpp
@@ -3,6 +3,8 @@
 #include "Editors/CameraNodeGraphNode.h"
 
 #include "Core/CameraParameters.h"
+#include "Core/CameraVariableReferences.h"
+#include "Core/ICustomCameraNodeParameterProvider.h"
 #include "EdGraph/EdGraphPin.h"
 #include "Editors/CameraNodeGraphSchema.h"
 #include "Editors/SCameraNodeGraphNode.h"
@@ -17,8 +19,10 @@ void UCameraNodeGraphNode::AllocateDefaultPins()
 {
 	Super::AllocateDefaultPins();
 
-	// Add extra input pins for any camera parameter.
-	UClass* CameraNodeClass = GetObject()->GetClass();
+	// Add extra input pins for any camera parameter, variable reference, and context data.
+	UObject* Object = GetObject();
+	UClass* CameraNodeClass = Object->GetClass();
+	const FName ContextDataMetaData = TEXT("CameraContextData");
 	for (TFieldIterator<FProperty> PropertyIt(CameraNodeClass); PropertyIt; ++PropertyIt)
 	{
 		FStructProperty* StructProperty = CastField<FStructProperty>(*PropertyIt);
@@ -30,46 +34,104 @@ void UCameraNodeGraphNode::AllocateDefaultPins()
 		const FName PropertyName = PropertyIt->GetFName();
 		
 		FEdGraphPinType PinType;
-		PinType.PinCategory = UCameraNodeGraphSchema::PC_CameraParameter;
+		const FText PinFriendlyName = FText::FromName(PropertyName);
+		const FString PinToolTip = StructProperty->Struct->GetDisplayNameText().ToString();
 
 #define UE_CAMERA_VARIABLE_FOR_TYPE(ValueType, ValueName)\
 		if (StructProperty->Struct == F##ValueName##CameraParameter::StaticStruct())\
 		{\
+			PinType.PinCategory = UCameraNodeGraphSchema::PC_CameraParameter;\
+			PinType.PinSubCategory = UEnum::GetValueAsName(ECameraVariableType::ValueName);\
+			PinType.PinSubCategoryObject = F##ValueName##CameraParameter::StaticStruct();\
 			UEdGraphPin* ParameterPin = CreatePin(EGPD_Input, PinType, PropertyName);\
-			ParameterPin->PinFriendlyName = FText::FromName(PropertyName);\
+			ParameterPin->PinFriendlyName = PinFriendlyName;\
+			ParameterPin->PinToolTip = PinToolTip;\
+			continue;\
+		}\
+		if (StructProperty->Struct == F##ValueName##CameraVariableReference::StaticStruct())\
+		{\
+			PinType.PinCategory = UCameraNodeGraphSchema::PC_CameraVariableReference;\
+			PinType.PinSubCategory = UEnum::GetValueAsName(ECameraVariableType::ValueName);\
+			PinType.PinSubCategoryObject = F##ValueName##CameraVariableReference::StaticStruct();\
+			UEdGraphPin* VariableReferencePin = CreatePin(EGPD_Input, PinType, PropertyName);\
+			VariableReferencePin->PinFriendlyName = PinFriendlyName;\
+			VariableReferencePin->PinToolTip = PinToolTip;\
 			continue;\
 		}
 UE_CAMERA_VARIABLE_FOR_ALL_TYPES()
 #undef UE_CAMERA_VARIABLE_FOR_TYPE
-	}
-}
 
-TSharedPtr<SGraphNode> UCameraNodeGraphNode::CreateVisualWidget()
-{
-	return SNew(SCameraNodeGraphNode).GraphNode(this);
-}
+		if (StructProperty->HasMetaData(ContextDataMetaData))
+		{
+			PinType.PinCategory = UCameraNodeGraphSchema::PC_CameraContextData;
+			PinType.PinSubCategory = UEnum::GetValueAsName(ECameraContextDataType::Struct);
+			PinType.PinSubCategoryObject = StructProperty->Struct;
+			UEdGraphPin* ContextDataPin = CreatePin(EGPD_Input, PinType, PropertyName);
+			ContextDataPin->PinFriendlyName = PinFriendlyName;
+			ContextDataPin->PinToolTip = PinToolTip;
+			continue;
+		}
+	}
 
-UEdGraphPin* UCameraNodeGraphNode::GetPinForCameraParameterProperty(const FName& InPropertyName) const
-{
-	UEdGraphPin* const* FoundItem = Pins.FindByPredicate(
-			[InPropertyName](UEdGraphPin* Item)
-			{ 
-				return Item->PinType.PinCategory == UCameraNodeGraphSchema::PC_CameraParameter &&
-					Item->GetFName() == InPropertyName;
-			});
-	if (FoundItem)
+	ICustomCameraNodeParameterProvider* CustomParameterProvider = Cast<ICustomCameraNodeParameterProvider>(Object);
+	if (CustomParameterProvider)
 	{
-		return *FoundItem;
+		FCustomCameraNodeParameterInfos CustomParameters;
+		CustomParameterProvider->GetCustomCameraNodeParameters(CustomParameters);
+
+		// Add pins for blendable parameters.
+		TArray<FCustomCameraNodeBlendableParameter> CustomBlendableParameters;
+		CustomParameters.GetBlendableParameters(CustomBlendableParameters);
+
+		const UEnum* VariableTypeEnum = StaticEnum<ECameraVariableType>();
+		for (const FCustomCameraNodeBlendableParameter& BlendableParameter : CustomBlendableParameters)
+		{
+			FEdGraphPinType PinType;
+			PinType.PinCategory = UCameraNodeGraphSchema::PC_CameraParameter;
+			PinType.PinSubCategory = VariableTypeEnum->GetValueAsName(BlendableParameter.ParameterType);
+
+			switch (BlendableParameter.ParameterType)
+			{
+#define UE_CAMERA_VARIABLE_FOR_TYPE(ValueType, ValueName)\
+				case ECameraVariableType::ValueName:\
+					PinType.PinSubCategoryObject = F##ValueName##CameraVariableReference::StaticStruct();\
+					break;
+				UE_CAMERA_VARIABLE_FOR_ALL_TYPES()
+#undef UE_CAMERA_VARIABLE_FOR_TYPE
+			}
+
+			UEdGraphPin* ParameterPin = CreatePin(EGPD_Input, PinType, BlendableParameter.ParameterName);
+			ParameterPin->PinFriendlyName = FText::FromName(BlendableParameter.ParameterName);
+			ParameterPin->PinToolTip = VariableTypeEnum->GetNameStringByValue((int64)BlendableParameter.ParameterType);
+		}
+
+		// Add pins for data parameters.
+		TArray<FCustomCameraNodeDataParameter> CustomDataParameters;
+		CustomParameters.GetDataParameters(CustomDataParameters);
+		const UEnum* DataTypeEnum = StaticEnum<ECameraContextDataType>();
+
+		for (const FCustomCameraNodeDataParameter& DataParameter : CustomDataParameters)
+		{
+			FEdGraphPinType PinType;
+			PinType.PinCategory = UCameraNodeGraphSchema::PC_CameraContextData;
+			PinType.PinSubCategory = DataTypeEnum->GetNameByValue((int64)DataParameter.ParameterType);
+			PinType.PinSubCategoryObject = const_cast<UObject*>(DataParameter.ParameterTypeObject.Get());
+
+			UEdGraphPin* ContextDataPin = CreatePin(EGPD_Input, PinType, DataParameter.ParameterName);
+			ContextDataPin->PinFriendlyName = FText::FromName(DataParameter.ParameterName);
+
+			FString PinToolTip = DataTypeEnum->GetNameStringByValue((int64)DataParameter.ParameterType);
+			if (DataParameter.ParameterTypeObject)
+			{
+				PinToolTip = DataParameter.ParameterTypeObject->GetName();
+			}
+			ContextDataPin->PinToolTip = PinToolTip;
+		}
 	}
-	return nullptr;
 }
 
-FName UCameraNodeGraphNode::GetCameraParameterPropertyForPin(const UEdGraphPin* InPin) const
+TSharedPtr<SGraphNode> UCameraNodeGraphNode::CreateVisualWidget()
 {
-	if (InPin->PinType.PinCategory == UCameraNodeGraphSchema::PC_CameraParameter)
-	{
-		return InPin->GetFName();
-	}
-	return NAME_None;
+	return SNew(SCameraNodeGraphNode).GraphNode(this);
 }
 
diff --git a/Engine/Plugins/Cameras/GameplayCameras/Source/GameplayCamerasEditor/Private/Editors/CameraNodeGraphSchema.cpp b/Engine/Plugins/Cameras/GameplayCameras/Source/GameplayCamerasEditor/Private/Editors/CameraNodeGraphSchema.cpp
--- a/Engine/Plugins/Cameras/GameplayCameras/Source/GameplayCamerasEditor/Private/Editors/CameraNodeGraphSchema.cpp
+++ b/Engine/Plugins/Cameras/GameplayCameras/Source/GameplayCamerasEditor/Private/Editors/CameraNodeGraphSchema.cpp
@@ -5,18 +5,20 @@
 #include "Core/BlendCameraNode.h"
 #include "Core/CameraNode.h"
 #include "Core/CameraNodeHierarchy.h"
+#include "Core/CameraParameters.h"
 #include "Core/CameraRigAsset.h"
+#include "Core/CameraVariableReferences.h"
+#include "Core/ICustomCameraNodeParameterProvider.h"
 #include "EdGraph/EdGraphPin.h"
+#include "EdGraphSchema_K2.h"
 #include "Editors/CameraNodeGraphNode.h"
 #include "Editors/CameraRigInterfaceParameterGraphNode.h"
-#include "Editors/CameraRigNodeGraphNode.h"
 #include "Editors/ObjectTreeGraph.h"
 #include "Editors/ObjectTreeGraphConfig.h"
 #include "Editors/ObjectTreeGraphNode.h"
 #include "Framework/Notifications/NotificationManager.h"
 #include "GameplayCamerasEditorSettings.h"
 #include "Nodes/Common/ArrayCameraNode.h"
-#include "Nodes/Common/CameraRigCameraNode.h"
 #include "Widgets/Notifications/SNotificationList.h"
 
 #include "ScopedTransaction.h"
@@ -26,6 +28,14 @@
 #define LOCTEXT_NAMESPACE "CameraNodeGraphSchema"
 
 const FName UCameraNodeGraphSchema::PC_CameraParameter("CameraParameter");
+const FName UCameraNodeGraphSchema::PC_CameraVariableReference("CameraVariableReference");
+const FName UCameraNodeGraphSchema::PC_CameraContextData("CameraContextData");
+
+UCameraNodeGraphSchema::UCameraNodeGraphSchema(const FObjectInitializer& ObjInit)
+	: Super(ObjInit)
+{
+	PinColors.Initialize();
+}
 
 FObjectTreeGraphConfig UCameraNodeGraphSchema::BuildGraphConfig() const
 {
@@ -35,7 +45,6 @@ FObjectTreeGraphConfig UCameraNodeGraphSchema::BuildGraphConfig() const
 	GraphConfig.GraphName = UCameraRigAsset::NodeTreeGraphName;
 	GraphConfig.ConnectableObjectClasses.Add(UCameraRigAsset::StaticClass());
 	GraphConfig.ConnectableObjectClasses.Add(UCameraNode::StaticClass());
-	GraphConfig.ConnectableObjectClasses.Add(UCameraRigInterfaceParameter::StaticClass());
 	GraphConfig.NonConnectableObjectClasses.Add(UBlendCameraNode::StaticClass());
 	GraphConfig.GraphDisplayInfo.PlainName = LOCTEXT("NodeGraphPlainName", "CameraNodes");
 	GraphConfig.GraphDisplayInfo.DisplayName = LOCTEXT("NodeGraphDisplayName", "Camera Nodes");
@@ -50,11 +59,6 @@ FObjectTreeGraphConfig UCameraNodeGraphSchema::BuildGraphConfig() const
 		.CreateCategoryMetaData(TEXT("CameraNodeCategories"))
 		.NodeTitleColor(Settings->CameraNodeTitleColor)
 		.GraphNodeClass(UCameraNodeGraphNode::StaticClass());
-	GraphConfig.ObjectClassConfigs.Emplace(UCameraRigCameraNode::StaticClass())
-		.GraphNodeClass(UCameraRigNodeGraphNode::StaticClass());
-	GraphConfig.ObjectClassConfigs.Emplace(UCameraRigInterfaceParameter::StaticClass())
-		.CanCreateNew(false)
-		.GraphNodeClass(UCameraRigInterfaceParameterGraphNode::StaticClass());
 	GraphConfig.ObjectClassConfigs.Emplace(UArrayCameraNode::StaticClass())
 		.OnSetupNewObject(FOnSetupNewObject::CreateLambda([](UObject* NewObject)
 				{
@@ -64,6 +68,9 @@ FObjectTreeGraphConfig UCameraNodeGraphSchema::BuildGraphConfig() const
 					ArrayNode->Children.AddDefaulted();
 				}));
 
+	// Note that we don't add the interface parameter types to the config, we will manage
+	// them ourselves.
+
 	return GraphConfig;
 }
 
@@ -119,63 +126,67 @@ void UCameraNodeGraphSchema::OnCreateAllNodes(UObjectTreeGraph* InGraph, const F
 		return;
 	}
 	
-	for (UCameraRigInterfaceParameter* InterfaceParameter : CameraRig->Interface.InterfaceParameters)
+	// Add nodes for all interface parameters that have been added to the graph.
+	// These nodes are UObjectTreeGraphNode instances, but they are "unmanaged" by the UObjectTreeGraphSchema since
+	// their object types are not in the ConnectableObjectClasses. Instead, we manage them ourselves in this schema.
+	TArray<TObjectPtr<UCameraRigInterfaceParameterBase>> InterfaceParameters;
+	InterfaceParameters.Append(CameraRig->Interface.BlendableParameters);
+	InterfaceParameters.Append(CameraRig->Interface.DataParameters);
+
+	for (UCameraRigInterfaceParameterBase* InterfaceParameter : InterfaceParameters)
 	{
-		UObjectTreeGraphNode* const* InterfaceParameterNode = InCreatedNodes.CreatedNodes.Find(InterfaceParameter);
-		UObjectTreeGraphNode* const* CameraNodeNode = InCreatedNodes.CreatedNodes.Find(InterfaceParameter->Target);
-		if (InterfaceParameterNode && CameraNodeNode)
+		if (!InterfaceParameter->bHasGraphNode)
 		{
-			UEdGraphPin* InterfaceParameterSelfPin = (*InterfaceParameterNode)->GetSelfPin();
-			UEdGraphPin* CameraParameterPin = Cast<UCameraNodeGraphNode>(*CameraNodeNode)->GetPinForCameraParameterProperty(InterfaceParameter->TargetPropertyName);
-			InterfaceParameterSelfPin->MakeLinkTo(CameraParameterPin);
+			continue;
 		}
-	}
-}
-
-void UCameraNodeGraphSchema::OnAddConnectableObject(UObjectTreeGraph* InGraph, UObjectTreeGraphNode* InNewNode) const
-{
-	Super::OnAddConnectableObject(InGraph, InNewNode);
-
-	if (UCameraRigInterfaceParameter* InterfaceParameter = Cast<UCameraRigInterfaceParameter>(InNewNode->GetObject()))
-	{
-		UCameraRigAsset* CameraRig = Cast<UCameraRigAsset>(InGraph->GetRootObject());
-		if (ensure(CameraRig))
+		
+		UCameraRigInterfaceParameterGraphNode* InterfaceParameterNode = CreateInterfaceParameterNode(InGraph, InterfaceParameter);
+		UObjectTreeGraphNode* const* CameraNodeNode = InCreatedNodes.CreatedNodes.Find(InterfaceParameter->Target);
+		if (CameraNodeNode)
 		{
-			CameraRig->Modify();
-
-			const int32 Index = CameraRig->Interface.InterfaceParameters.AddUnique(InterfaceParameter);
-			ensure(Index == CameraRig->Interface.InterfaceParameters.Num() - 1);
+			UEdGraphPin* InterfaceParameterSelfPin = InterfaceParameterNode->GetSelfPin();
+			UEdGraphPin* NodePin = FindPin(*CameraNodeNode, InterfaceParameter->TargetPropertyName, NAME_None);
+			ensure(NodePin &&
+					(NodePin->PinType.PinCategory == PC_CameraParameter ||
+					 NodePin->PinType.PinCategory == PC_CameraVariableReference ||
+					 NodePin->PinType.PinCategory == PC_CameraContextData));
+			InterfaceParameterSelfPin->MakeLinkTo(NodePin);
 		}
 	}
 }
 
-void UCameraNodeGraphSchema::OnRemoveConnectableObject(UObjectTreeGraph* InGraph, UObjectTreeGraphNode* InRemovedNode) const
+UCameraRigInterfaceParameterGraphNode* UCameraNodeGraphSchema::CreateInterfaceParameterNode(UEdGraph* InGraph, UCameraRigInterfaceParameterBase* InterfaceParameter) const
 {
-	Super::OnRemoveConnectableObject(InGraph, InRemovedNode);
-
-	if (UCameraRigInterfaceParameter* InterfaceParameter = Cast<UCameraRigInterfaceParameter>(InRemovedNode->GetObject()))
-	{
-		UCameraRigAsset* CameraRig = Cast<UCameraRigAsset>(InGraph->GetRootObject());
-		if (ensure(CameraRig))
-		{
-			CameraRig->Modify();
-
-			const int32 NumRemoved = CameraRig->Interface.InterfaceParameters.Remove(InterfaceParameter);
-			ensure(NumRemoved == 1);
-		}
-	}
+	FGraphNodeCreator<UCameraRigInterfaceParameterGraphNode> GraphNodeCreator(*InGraph);
+	UCameraRigInterfaceParameterGraphNode* InterfaceParameterNode = GraphNodeCreator.CreateNode(false);
+	InterfaceParameterNode->Initialize(InterfaceParameter);
+	GraphNodeCreator.Finalize();
+	return InterfaceParameterNode;
 }
 
 void UCameraNodeGraphSchema::GetGraphContextActions(FGraphContextMenuBuilder& ContextMenuBuilder) const
 {
-	// See if we were dragging a camera parameter pin.
+	// See if we were dragging a camera parameter pin or camera variable reference pin.
 	if (const UEdGraphPin* DraggedPin = ContextMenuBuilder.FromPin)
 	{
-		if (DraggedPin->PinType.PinCategory == PC_CameraParameter)
+		UCameraNodeGraphNode* CameraNodeNode = Cast<UCameraNodeGraphNode>(DraggedPin->GetOwningNode());
+
+		if (DraggedPin->PinType.PinCategory == PC_CameraParameter ||
+				DraggedPin->PinType.PinCategory == PC_CameraVariableReference ||
+				DraggedPin->PinType.PinCategory == PC_CameraContextData)
 		{
-			UCameraNodeGraphNode* CameraNodeNode = Cast<UCameraNodeGraphNode>(DraggedPin->GetOwningNode());
-			const FName PropertyName = CameraNodeNode->GetCameraParameterPropertyForPin(DraggedPin);
-			ensure(PropertyName != NAME_None);
+			ensure(DraggedPin->PinName != NAME_None);
+
+			// Find the property being dragged, so we know what kind of parameter to create.
+			const UClass* CameraNodeClass = CameraNodeNode->GetObject()->GetClass();
+			FProperty* Property = CameraNodeClass->FindPropertyByName(DraggedPin->PinName);
+
+			FCustomCameraNodeParameterInfos CustomParameters;
+			ICustomCameraNodeParameterProvider* CustomParameterProvider = Cast<ICustomCameraNodeParameterProvider>(CameraNodeNode->GetObject());
+			if (CustomParameterProvider)
+			{
+				CustomParameterProvider->GetCustomCameraNodeParameters(CustomParameters);
+			}
 
 			TSharedRef<FCameraNodeGraphSchemaAction_NewInterfaceParameterNode> Action = 
 				MakeShared<FCameraNodeGraphSchemaAction_NewInterfaceParameterNode>(
@@ -183,7 +194,113 @@ void UCameraNodeGraphSchema::GetGraphContextActions(FGraphContextMenuBuilder& Co
 						LOCTEXT("NewInterfaceParameterAction", "Camera Rig Parameter"),
 						LOCTEXT("NewInterfaceParameterActionToolTip", "Exposes this parameter on the camera rig"));
 			Action->Target = Cast<UCameraNode>(CameraNodeNode->GetObject());
-			Action->TargetPropertyName = PropertyName;
+			Action->TargetPropertyName = DraggedPin->PinName;
+
+			if (DraggedPin->PinType.PinCategory == PC_CameraParameter ||
+				DraggedPin->PinType.PinCategory == PC_CameraVariableReference)
+			{
+				ECameraVariableType ParameterType;
+
+				if (FStructProperty* StructProperty = CastField<FStructProperty>(Property))
+				{
+#define UE_CAMERA_VARIABLE_FOR_TYPE(ValueType, ValueName)\
+					if ((StructProperty->Struct == F##ValueName##CameraParameter::StaticStruct()) ||\
+						(StructProperty->Struct == F##ValueName##CameraVariableReference::StaticStruct()))\
+					{\
+						ParameterType = ECameraVariableType::ValueName;\
+					}\
+					else
+					UE_CAMERA_VARIABLE_FOR_ALL_TYPES()
+#undef UE_CAMERA_VARIABLE_FOR_TYPE
+					{
+						// Unexpected: if there was a camera parameter pin or a variable reference pin, we should
+						// have had a camera parameter property or variable reference property!
+						ensure(false);
+						return;
+					}
+				}
+				else
+				{
+					FCustomCameraNodeBlendableParameter BlendableParameter;
+					if (CustomParameters.FindBlendableParameter(DraggedPin->PinName, BlendableParameter))
+					{
+						ParameterType = BlendableParameter.ParameterType;
+					}
+					else
+					{
+						// Unexpected: a parameter pin was created, but we found no property or custom
+						// parameter for it!
+						ensure(false);
+						return;
+					}
+				}
+
+				Action->NewNodeType = EInterfaceParameterCreateNodeType::BlendableParameter;
+				Action->BlendableParameterType = ParameterType;
+			}
+			else if (DraggedPin->PinType.PinCategory == PC_CameraContextData)
+			{
+				ECameraContextDataType DataType;
+				const UObject* DataTypeObject = nullptr;
+				
+				if (Property)
+				{
+					if (FNameProperty* NameProperty = CastField<FNameProperty>(Property))
+					{
+						DataType = ECameraContextDataType::Name;
+					}
+					else if (FStrProperty* StringProperty = CastField<FStrProperty>(Property))
+					{
+						DataType = ECameraContextDataType::String;
+					}
+					else if (FEnumProperty* EnumProperty = CastField<FEnumProperty>(Property))
+					{
+						DataType = ECameraContextDataType::Enum;
+						DataTypeObject = EnumProperty->GetEnum();
+					}
+					else if (FStructProperty* StructProperty = CastField<FStructProperty>(Property))
+					{
+						DataType = ECameraContextDataType::Struct;
+						DataTypeObject = StructProperty->Struct;
+					}
+					else if (FClassProperty* ClassProperty = CastField<FClassProperty>(Property))
+					{
+						DataType = ECameraContextDataType::Class;
+						DataTypeObject = ClassProperty->PropertyClass;
+					}
+					else if (FObjectProperty* ObjectProperty = CastField<FObjectProperty>(Property))
+					{
+						DataType = ECameraContextDataType::Object;
+						DataTypeObject = ObjectProperty->PropertyClass;
+					}
+					else
+					{
+						// Unexpected as per previous comments.
+						ensure(false);
+						return;
+					}
+				}
+				else
+				{
+					FCustomCameraNodeDataParameter DataParameter;
+					if (CustomParameters.FindDataParameter(DraggedPin->PinName, DataParameter))
+					{
+						DataType = DataParameter.ParameterType;
+						DataTypeObject = DataParameter.ParameterTypeObject;
+					}
+					else
+					{
+						// Unexpected as per previous comments.
+						ensure(false);
+						return;
+					}
+				}
+
+				Action->NewNodeType = EInterfaceParameterCreateNodeType::DataParameter;
+				Action->DataParameterType = DataType;
+				Action->DataParameterTypeObject = DataTypeObject;
+			}
+
 			ContextMenuBuilder.AddAction(StaticCastSharedPtr<FEdGraphSchemaAction>(Action.ToSharedPtr()));
 
 			return;
@@ -195,137 +312,214 @@ void UCameraNodeGraphSchema::GetGraphContextActions(FGraphContextMenuBuilder& Co
 
 const FPinConnectionResponse UCameraNodeGraphSchema::CanCreateConnection(const UEdGraphPin* A, const UEdGraphPin* B) const
 {
-	if (A->PinType.PinCategory == PC_CameraParameter && B->PinType.PinCategory == PC_Self)
+	// Check if we are connecting parameter pins of compatible types.
+	if ((A->PinType.PinCategory == PC_CameraParameter || 
+				A->PinType.PinCategory == PC_CameraVariableReference ||
+				A->PinType.PinCategory == PC_CameraContextData) && 
+			B->PinType.PinCategory == PC_Self)
 	{
-		UObjectTreeGraphNode* NodeB = Cast<UObjectTreeGraphNode>(B->GetOwningNode());
-		if (NodeB && NodeB->IsObjectA<UCameraRigInterfaceParameter>())
+		UCameraRigInterfaceParameterGraphNode* NodeB = Cast<UCameraRigInterfaceParameterGraphNode>(B->GetOwningNode());
+		if (NodeB)
 		{
-			return FPinConnectionResponse(CONNECT_RESPONSE_BREAK_OTHERS_AB, TEXT("Compatible pin types"));
+			UCameraRigBlendableParameter* BlendableParameter = NodeB->CastObject<UCameraRigBlendableParameter>();
+			if (BlendableParameter && 
+					A->PinType.PinSubCategory == UEnum::GetValueAsName(BlendableParameter->ParameterType))
+			{
+				return FPinConnectionResponse(CONNECT_RESPONSE_BREAK_OTHERS_AB, TEXT("Compatible pin types"));
+			}
+			UCameraRigDataParameter* DataParameter = NodeB->CastObject<UCameraRigDataParameter>();
+			if (DataParameter && 
+					A->PinType.PinSubCategory == UEnum::GetValueAsName(DataParameter->DataType) &&
+					A->PinType.PinSubCategoryObject == DataParameter->DataTypeObject)
+			{
+				return FPinConnectionResponse(CONNECT_RESPONSE_BREAK_OTHERS_AB, TEXT("Compatible pin types"));
+			}
 		}
 	}
-	else if (A->PinType.PinCategory == PC_Self && B->PinType.PinCategory == PC_CameraParameter)
+	else if (A->PinType.PinCategory == PC_Self && 
+			(B->PinType.PinCategory == PC_CameraParameter || 
+			 B->PinType.PinCategory == PC_CameraVariableReference ||
+			 B->PinType.PinCategory == PC_CameraContextData))
 	{
-		UObjectTreeGraphNode* NodeA = Cast<UObjectTreeGraphNode>(A->GetOwningNode());
-		if (NodeA && NodeA->IsObjectA<UCameraRigInterfaceParameter>())
+		UCameraRigInterfaceParameterGraphNode* NodeA = Cast<UCameraRigInterfaceParameterGraphNode>(A->GetOwningNode());
+		if (NodeA)
 		{
-			return FPinConnectionResponse(CONNECT_RESPONSE_BREAK_OTHERS_AB, TEXT("Compatible pin types"));
+			UCameraRigBlendableParameter* BlendableParameter = NodeA->CastObject<UCameraRigBlendableParameter>();
+			if (BlendableParameter && 
+					B->PinType.PinSubCategory == UEnum::GetValueAsName(BlendableParameter->ParameterType))
+			{
+				return FPinConnectionResponse(CONNECT_RESPONSE_BREAK_OTHERS_AB, TEXT("Compatible pin types"));
+			}
+			UCameraRigDataParameter* DataParameter = NodeA->CastObject<UCameraRigDataParameter>();
+			if (DataParameter && 
+					B->PinType.PinSubCategory == UEnum::GetValueAsName(DataParameter->DataType) &&
+					B->PinType.PinSubCategoryObject == DataParameter->DataTypeObject)
+			{
+				return FPinConnectionResponse(CONNECT_RESPONSE_BREAK_OTHERS_AB, TEXT("Compatible pin types"));
+			}
 		}
 	}
 
 	return Super::CanCreateConnection(A, B);
 }
 
-bool UCameraNodeGraphSchema::OnApplyConnection(UEdGraphPin* A, UEdGraphPin* B) const
+bool UCameraNodeGraphSchema::OnTryCreateCustomConnection(UEdGraphPin* A, UEdGraphPin* B) const
 {
-	// Try to make a connection between a camera node's parameter pin and a camera rig interface parameter.
-	// First, figure out which is which.
-	UEdGraphPin* RigInterfacePin = nullptr;
-	UEdGraphPin* CameraParameterPin = nullptr;
-
-	if (A->PinType.PinCategory == PC_CameraParameter && B->PinType.PinCategory == PC_Self)
+	// See if we are in the situation of connecting an interface parameter to a camera node property.
+	UEdGraphPin* TargetPin = nullptr;
+	UObjectTreeGraphNode* TargetNode = nullptr;
+	UCameraRigInterfaceParameterGraphNode* InterfaceParameterNode = nullptr;
+
+	if ((A->PinType.PinCategory == PC_CameraParameter || 
+				A->PinType.PinCategory == PC_CameraVariableReference ||
+				A->PinType.PinCategory == PC_CameraContextData) && 
+			B->PinType.PinCategory == PC_Self)
 	{
-		RigInterfacePin = B;
-		CameraParameterPin = A;
+		TargetPin = A;
+		TargetNode = Cast<UObjectTreeGraphNode>(A->GetOwningNode());
+		InterfaceParameterNode = Cast<UCameraRigInterfaceParameterGraphNode>(B->GetOwningNode());
 	}
-	else if (A->PinType.PinCategory == PC_Self && B->PinType.PinCategory == PC_CameraParameter)
+	else if (A->PinType.PinCategory == PC_Self && 
+			(B->PinType.PinCategory == PC_CameraParameter || 
+			 B->PinType.PinCategory == PC_CameraVariableReference ||
+			 B->PinType.PinCategory == PC_CameraContextData))
 	{
-		RigInterfacePin = A;
-		CameraParameterPin = B;
+		InterfaceParameterNode = Cast<UCameraRigInterfaceParameterGraphNode>(A->GetOwningNode());
+		TargetNode = Cast<UObjectTreeGraphNode>(B->GetOwningNode());
+		TargetPin = B;
 	}
 
-	if (!RigInterfacePin || !CameraParameterPin)
+	if (TargetNode && TargetPin && InterfaceParameterNode)
 	{
-		return false;
-	}
+		UCameraNode* Target = TargetNode->CastObject<UCameraNode>();
+		UCameraRigInterfaceParameterBase* InterfaceParameter = InterfaceParameterNode->GetInterfaceParameter();
+		if (Target && InterfaceParameter)
+		{
+			InterfaceParameter->Modify();
 
-	// Now make sure both nodes are what we expect, and that they have what we need.
-	UObjectTreeGraphNode* RigParameterNode = Cast<UObjectTreeGraphNode>(RigInterfacePin->GetOwningNode());
-	if (!RigParameterNode)
-	{
-		return false;
-	}
-	UCameraRigInterfaceParameter* RigParameter = RigParameterNode->CastObject<UCameraRigInterfaceParameter>();
-	if (!RigParameter)
-	{
-		return false;
+			InterfaceParameter->Target = Target;
+			InterfaceParameter->TargetPropertyName = TargetPin->PinName;
+		}
+
+		return true;
 	}
 
-	UCameraNodeGraphNode* CameraNodeNode = Cast<UCameraNodeGraphNode>(CameraParameterPin->GetOwningNode());
-	if (!CameraNodeNode)
+	return false;
+}
+
+bool UCameraNodeGraphSchema::OnBreakCustomPinLinks(UEdGraphPin& TargetPin) const
+{
+	// See if we are in the situation of an interface parameter node being disconnected from 
+	// a camera node property pin.
+	UCameraRigInterfaceParameterGraphNode* InterfaceParameterNode = nullptr;
+	if (TargetPin.PinType.PinCategory == PC_CameraParameter ||
+			TargetPin.PinType.PinCategory == PC_CameraVariableReference ||
+			TargetPin.PinType.PinCategory == PC_CameraContextData)
 	{
-		return false;
+		if (TargetPin.LinkedTo.Num() > 0)
+		{
+			InterfaceParameterNode = Cast<UCameraRigInterfaceParameterGraphNode>(TargetPin.LinkedTo[0]->GetOwningNode());
+		}
 	}
-	UCameraNode* CameraNode = CameraNodeNode->CastObject<UCameraNode>();
-	if (!CameraNode)
+	else if (TargetPin.PinType.PinCategory == PC_Self)
 	{
-		return false;
+		InterfaceParameterNode = Cast<UCameraRigInterfaceParameterGraphNode>(TargetPin.GetOwningNode());
 	}
-	const FName PropertyName = CameraNodeNode->GetCameraParameterPropertyForPin(CameraParameterPin);
-	if (PropertyName.IsNone())
+
+	if (InterfaceParameterNode)
 	{
-		return false;
-	}
+		UCameraRigInterfaceParameterBase* InterfaceParameter = InterfaceParameterNode->GetInterfaceParameter();
+		if (InterfaceParameter)
+		{
+			InterfaceParameter->Modify();
 
-	// Make the connection.
-	RigParameter->Modify();
+			InterfaceParameter->Target = nullptr;
+			InterfaceParameter->TargetPropertyName = NAME_None;
+		}
 
-	RigParameter->Target = CameraNode;
-	RigParameter->TargetPropertyName = PropertyName;
-	if (RigParameter->InterfaceParameterName.IsEmpty())
-	{
-		RigParameter->InterfaceParameterName = PropertyName.ToString();
+		return true;
 	}
 
-	return true;
+	return false;
 }
 
-bool UCameraNodeGraphSchema::OnApplyDisconnection(UEdGraphPin* TargetPin) const
+bool UCameraNodeGraphSchema::OnBreakSingleCustomPinLink(UEdGraphPin* SourcePin, UEdGraphPin* TargetPin) const
 {
-	// See if we have a rig parameter connection to break.
-	if (TargetPin->PinType.PinCategory == PC_Self || TargetPin->PinType.PinCategory == PC_CameraParameter)
+	// See if we are in the situation of an interface parameter node being disconnected from 
+	// a camera node property pin.
+	UCameraRigInterfaceParameterGraphNode* InterfaceParameterNode = nullptr;
+	if (SourcePin->PinType.PinCategory == PC_Self)
 	{
-		UEdGraphPin* RigParameterSelfPin = TargetPin;
-		if (TargetPin->PinType.PinCategory == PC_CameraParameter)
-		{
-			RigParameterSelfPin = TargetPin->LinkedTo[0];
-		}
+		InterfaceParameterNode = Cast<UCameraRigInterfaceParameterGraphNode>(SourcePin->GetOwningNode());
+	}
+	else if (TargetPin->PinType.PinCategory == PC_Self)
+	{
+		InterfaceParameterNode = Cast<UCameraRigInterfaceParameterGraphNode>(TargetPin->GetOwningNode());
+	}
 
-		UObjectTreeGraphNode* RigParameterNode = Cast<UObjectTreeGraphNode>(RigParameterSelfPin->GetOwningNode());
-		if (!RigParameterNode)
+	if (InterfaceParameterNode)
+	{
+		UCameraRigInterfaceParameterBase* InterfaceParameter = InterfaceParameterNode->GetInterfaceParameter();
+		if (InterfaceParameter)
 		{
-			return false;
+			InterfaceParameter->Modify();
+
+			InterfaceParameter->Target = nullptr;
+			InterfaceParameter->TargetPropertyName = NAME_None;
 		}
 
-		UCameraRigInterfaceParameter* RigParameter = RigParameterNode->CastObject<UCameraRigInterfaceParameter>();
-		if (RigParameter)
-		{
-			RigParameter->Modify();
+		return true;
+	}
 
-			RigParameter->Target = nullptr;
-			RigParameter->TargetPropertyName = NAME_None;
-			RigParameter->PrivateVariable = nullptr;
+	return false;
+}
 
-			return true;
-		}
+FLinearColor UCameraNodeGraphSchema::GetPinTypeColor(const FEdGraphPinType& PinType) const
+{
+	if (PinType.PinCategory == PC_CameraParameter || PinType.PinCategory == PC_CameraVariableReference)
+	{
+		const FName TypeName = PinType.PinSubCategory;
+		return PinColors.GetPinColor(TypeName);
+	}
+	if (PinType.PinCategory == PC_CameraContextData)
+	{
+		return PinColors.GetStructPinColor();
 	}
 
-	return false;
+	return UObjectTreeGraphSchema::GetPinTypeColor(PinType);
 }
 
-bool UCameraNodeGraphSchema::OnApplyDisconnection(UEdGraphPin* SourcePin, UEdGraphPin* TargetPin) const
+bool UCameraNodeGraphSchema::SafeDeleteNodeFromGraph(UEdGraph* Graph, UEdGraphNode* Node) const
 {
-	UObjectTreeGraphNode* RigParameterNode = nullptr;
+	Super::SafeDeleteNodeFromGraph(Graph, Node);
 
-	if (SourcePin->PinType.PinCategory == PC_Self && TargetPin->PinType.PinCategory == PC_CameraParameter)
+	// Deleting an interface parameter node simply removes its bHasGraphNode flag.
+	// To actually delete the parameter, the user needs to remove it from the "parameters" panel.
+	if (UCameraRigInterfaceParameterGraphNode* InterfaceParameterNode = Cast<UCameraRigInterfaceParameterGraphNode>(Node))
 	{
-		return OnApplyDisconnection(SourcePin);
+		if (UCameraRigInterfaceParameterBase* InterfaceParameter = InterfaceParameterNode->GetInterfaceParameter())
+		{
+			InterfaceParameter->Modify();
+			InterfaceParameter->bHasGraphNode = false;
+		}
 	}
-	else if (SourcePin->PinType.PinCategory == PC_CameraParameter && TargetPin->PinType.PinCategory == PC_Self)
+
+	return true;
+}
+
+UEdGraphPin* UCameraNodeGraphSchema::FindPin(UEdGraphNode* InNode, const FName& InPinName, const FName& InPinCategoryName) const
+{
+	UEdGraphPin* const* FoundItem = InNode->Pins.FindByPredicate(
+			[InPinName, InPinCategoryName](UEdGraphPin* Item)
+			{ 
+				return Item->GetFName() == InPinName &&
+					(InPinCategoryName.IsNone() || Item->PinType.PinCategory == InPinCategoryName);
+			});
+	if (FoundItem)
 	{
-		return OnApplyDisconnection(TargetPin);
+		return *FoundItem;
 	}
-
-	return false;
+	return nullptr;
 }
 
 FCameraNodeGraphSchemaAction_NewInterfaceParameterNode::FCameraNodeGraphSchemaAction_NewInterfaceParameterNode()
@@ -353,17 +547,98 @@ UEdGraphNode* FCameraNodeGraphSchemaAction_NewInterfaceParameterNode::PerformAct
 
 	const FScopedTransaction Transaction(LOCTEXT("CreateNewNodeAction", "Create New Node"));
 
-	const UObjectTreeGraphSchema* Schema = CastChecked<UObjectTreeGraphSchema>(ParentGraph->GetSchema());
+	const UCameraNodeGraphSchema* Schema = CastChecked<UCameraNodeGraphSchema>(ParentGraph->GetSchema());
 
-	UCameraRigInterfaceParameter* NewInterfaceParameter = NewObject<UCameraRigInterfaceParameter>(CameraRig, NAME_None, RF_Transactional);
-	// The interface parameter's properties will be set correctly inside AutowireNewNode by virtue
+	CameraRig->Modify();
+
+	// Create a new interface parameter and set it up based on the pin we're creating it from, if any.
+	UCameraRigInterfaceParameterBase* NewInterfaceParameter = nullptr;
+	if (NewNodeType == EInterfaceParameterCreateNodeType::BlendableParameter)
+	{
+		UCameraRigBlendableParameter* NewBlendableParameter = NewObject<UCameraRigBlendableParameter>(CameraRig, NAME_None, RF_Transactional);
+		if (FromPin)
+		{
+			NewBlendableParameter->ParameterType = (ECameraVariableType)StaticEnum<ECameraVariableType>()->GetValueByName(FromPin->PinType.PinSubCategory);
+		}
+		CameraRig->Interface.BlendableParameters.Add(NewBlendableParameter);
+		NewInterfaceParameter = NewBlendableParameter;
+	}
+	else if (NewNodeType == EInterfaceParameterCreateNodeType::DataParameter)
+	{
+		UCameraRigDataParameter* NewDataParameter = NewObject<UCameraRigDataParameter>(CameraRig, NAME_None, RF_Transactional);
+		if (FromPin)
+		{
+			const UEnum* DataTypeEnum = StaticEnum<ECameraContextDataType>();
+			NewDataParameter->DataType = (ECameraContextDataType)DataTypeEnum->GetValueByName(FromPin->PinType.PinSubCategory);
+			NewDataParameter->DataTypeObject = FromPin->PinType.PinSubCategoryObject.Get();
+		}
+		CameraRig->Interface.DataParameters.Add(NewDataParameter);
+		NewInterfaceParameter = NewDataParameter;
+	}
+
+	if (ensure(NewInterfaceParameter))
+	{
+		NewInterfaceParameter->InterfaceParameterName = FromPin ? FromPin->GetName() : NewInterfaceParameter->GetName();
+		NewInterfaceParameter->bHasGraphNode = true;
+	}
+
+	// The interface parameter's other properties will be set correctly inside AutowireNewNode by virtue
 	// of getting connected to the dragged camera node pin.
 
 	ObjectTreeGraph->Modify();
 
-	UObjectTreeGraphNode* NewGraphNode = Schema->CreateObjectNode(ObjectTreeGraph, NewInterfaceParameter);
+	UCameraRigInterfaceParameterGraphNode* NewGraphNode = Schema->CreateInterfaceParameterNode(ObjectTreeGraph, NewInterfaceParameter);
 
-	Schema->AddConnectableObject(ObjectTreeGraph, NewGraphNode);
+	NewGraphNode->NodePosX = Location.X;
+	NewGraphNode->NodePosY = Location.Y;
+	NewGraphNode->OnGraphNodeMoved(false);
+
+	NewGraphNode->AutowireNewNode(FromPin);
+
+	CameraRig->EventHandlers.Notify(&UE::Cameras::ICameraRigAssetEventHandler::OnCameraRigInterfaceChanged);
+
+	return NewGraphNode;
+}
+
+FCameraNodeGraphSchemaAction_AddInterfaceParameterNode::FCameraNodeGraphSchemaAction_AddInterfaceParameterNode()
+{
+}
+
+FCameraNodeGraphSchemaAction_AddInterfaceParameterNode::FCameraNodeGraphSchemaAction_AddInterfaceParameterNode(FText InNodeCategory, FText InMenuDesc, FText InToolTip, const int32 InGrouping, FText InKeywords)
+	: FEdGraphSchemaAction(InNodeCategory, InMenuDesc, InToolTip, InGrouping, InKeywords)
+{
+}
+
+UEdGraphNode* FCameraNodeGraphSchemaAction_AddInterfaceParameterNode::PerformAction(UEdGraph* ParentGraph, UEdGraphPin* FromPin, const FVector2D Location, bool bSelectNewNode)
+{
+	if (!InterfaceParameter || InterfaceParameter->bHasGraphNode)
+	{
+		return nullptr;
+	}
+
+	UObjectTreeGraph* ObjectTreeGraph = Cast<UObjectTreeGraph>(ParentGraph);
+	if (!ensure(ObjectTreeGraph))
+	{
+		return nullptr;
+	}
+
+	UCameraRigAsset* CameraRig = Cast<UCameraRigAsset>(ObjectTreeGraph->GetRootObject());
+	if (!ensure(CameraRig))
+	{
+		return nullptr;
+	}
+
+	const FScopedTransaction Transaction(LOCTEXT("CreateNewNodeAction", "Create New Node"));
+
+	const UCameraNodeGraphSchema* Schema = CastChecked<UCameraNodeGraphSchema>(ParentGraph->GetSchema());
+
+	// Simply flag the interface parameter has having been added to the graph, and create a node for it.
+	InterfaceParameter->Modify();
+	InterfaceParameter->bHasGraphNode = true;
+
+	ParentGraph->Modify();
+	
+	UCameraRigInterfaceParameterGraphNode* NewGraphNode = Schema->CreateInterfaceParameterNode(ParentGraph, InterfaceParameter);
 
 	NewGraphNode->NodePosX = Location.X;
 	NewGraphNode->NodePosY = Location.Y;
@@ -371,6 +646,8 @@ UEdGraphNode* FCameraNodeGraphSchemaAction_NewInterfaceParameterNode::PerformAct
 
 	NewGraphNode->AutowireNewNode(FromPin);
 
+	CameraRig->EventHandlers.Notify(&UE::Cameras::ICameraRigAssetEventHandler::OnCameraRigInterfaceChanged);
+
 	return NewGraphNode;
 }
 
diff --git a/Engine/Plugins/Cameras/GameplayCameras/Source/GameplayCamerasEditor/Private/Editors/CameraRigInterfaceParameterGraphNode.cpp b/Engine/Plugins/Cameras/GameplayCameras/Source/GameplayCamerasEditor/Private/Editors/CameraRigInterfaceParameterGraphNode.cpp
--- a/Engine/Plugins/Cameras/GameplayCameras/Source/GameplayCamerasEditor/Private/Editors/CameraRigInterfaceParameterGraphNode.cpp
+++ b/Engine/Plugins/Cameras/GameplayCameras/Source/GameplayCamerasEditor/Private/Editors/CameraRigInterfaceParameterGraphNode.cpp
@@ -4,12 +4,18 @@
 
 #include "Core/CameraRigAsset.h"
 #include "Editors/SCameraRigInterfaceParameterGraphNode.h"
+#include "UObject/Object.h"
 
 UCameraRigInterfaceParameterGraphNode::UCameraRigInterfaceParameterGraphNode(const FObjectInitializer& ObjInit)
 	: UObjectTreeGraphNode(ObjInit)
 {
 }
 
+UCameraRigInterfaceParameterBase* UCameraRigInterfaceParameterGraphNode::GetInterfaceParameter() const
+{
+	return CastChecked<UCameraRigInterfaceParameterBase>(GetObject(), ECastCheckedType::NullAllowed);
+}
+
 TSharedPtr<SGraphNode> UCameraRigInterfaceParameterGraphNode::CreateVisualWidget()
 {
 	return SNew(SCameraRigInterfaceParameterGraphNode).GraphNode(this);
diff --git a/Engine/Plugins/Cameras/GameplayCameras/Source/GameplayCamerasEditor/Private/Editors/SCameraNodeGraphNode.cpp b/Engine/Plugins/Cameras/GameplayCameras/Source/GameplayCamerasEditor/Private/Editors/SCameraNodeGraphNode.cpp
--- a/Engine/Plugins/Cameras/GameplayCameras/Source/GameplayCamerasEditor/Private/Editors/SCameraNodeGraphNode.cpp
+++ b/Engine/Plugins/Cameras/GameplayCameras/Source/GameplayCamerasEditor/Private/Editors/SCameraNodeGraphNode.cpp
@@ -22,7 +22,8 @@ TSharedPtr<SGraphPin> SCameraNodeGraphNode::CreatePinWidget(UEdGraphPin* InPin)
 	{
 		TSharedRef<FGameplayCamerasEditorStyle> GraphStyle = FGameplayCamerasEditorStyle::Get();
 
-		if (InPin->PinType.PinCategory == UCameraNodeGraphSchema::PC_CameraParameter)
+		if (InPin->PinType.PinCategory == UCameraNodeGraphSchema::PC_CameraParameter ||
+				InPin->PinType.PinCategory == UCameraNodeGraphSchema::PC_CameraVariableReference)
 		{
 			const FSlateBrush* ConnectedBrush = GraphStyle->GetBrush("Graph.CameraRigParameterPin.Connected");
 			const FSlateBrush* DisconnectedBrush = GraphStyle->GetBrush("Graph.CameraRigParameterPin.Disconnected");
diff --git a/Engine/Plugins/Cameras/GameplayCameras/Source/GameplayCamerasEditor/Private/Editors/SCameraRigAssetEditor.cpp b/Engine/Plugins/Cameras/GameplayCameras/Source/GameplayCamerasEditor/Private/Editors/SCameraRigAssetEditor.cpp
--- a/Engine/Plugins/Cameras/GameplayCameras/Source/GameplayCamerasEditor/Private/Editors/SCameraRigAssetEditor.cpp
+++ b/Engine/Plugins/Cameras/GameplayCameras/Source/GameplayCamerasEditor/Private/Editors/SCameraRigAssetEditor.cpp
@@ -7,10 +7,12 @@
 #include "Core/CameraRigTransition.h"
 #include "EdGraph/EdGraphPin.h"
 #include "Editors/CameraNodeGraphSchema.h"
+#include "Editors/CameraRigInterfaceParameterGraphNode.h"
 #include "Editors/CameraRigTransitionGraphSchema.h"
 #include "Editors/ObjectTreeGraph.h"
 #include "Editors/ObjectTreeGraphConfig.h"
 #include "Editors/ObjectTreeGraphNode.h"
+#include "Editors/SCameraNodeGraphEditor.h"
 #include "Editors/SObjectTreeGraphEditor.h"
 #include "ObjectEditorUtils.h"
 #include "Widgets/Layout/SBox.h"
@@ -30,6 +32,8 @@ void SCameraRigAssetEditor::Construct(const FArguments& InArgs)
 
 	CreateGraphEditors();
 
+	CameraRigAsset->EventHandlers.Register(EventHandler, this);
+
 	ChildSlot
 	[
 		SAssignNew(BoxPanel, SBox)
@@ -51,6 +55,8 @@ void SCameraRigAssetEditor::SetCameraRigAsset(UCameraRigAsset* InCameraRig)
 {
 	if (CameraRigAsset != InCameraRig)
 	{
+		EventHandler.Unlink();
+
 		DiscardGraphEditors();
 
 		CameraRigAsset = InCameraRig;
@@ -58,6 +64,8 @@ void SCameraRigAssetEditor::SetCameraRigAsset(UCameraRigAsset* InCameraRig)
 		CreateGraphEditors();
 
 		SetEditorModeImpl(CurrentMode, true);
+
+		CameraRigAsset->EventHandlers.Register(EventHandler, this);
 	}
 }
 
@@ -83,7 +91,7 @@ void SCameraRigAssetEditor::CreateNodeGraphEditor()
 	FGraphAppearanceInfo Appearance;
 	Appearance.CornerText = LOCTEXT("CameraRigGraphText", "CAMERA NODES");
 
-	NodeGraphEditor = SNew(SObjectTreeGraphEditor)
+	NodeGraphEditor = SNew(SCameraNodeGraphEditor)
 		.Appearance(Appearance)
 		.DetailsView(DetailsView)
 		.GraphTitle(this, &SCameraRigAssetEditor::GetCameraRigAssetName, NodeGraph.Get())
@@ -296,6 +304,53 @@ void SCameraRigAssetEditor::RemoveOnAnyGraphChanged(FDelegateUserObjectConst InU
 	OnAnyGraphChanged.RemoveAll(InUserObject);
 }
 
+void SCameraRigAssetEditor::OnCameraRigInterfaceChanged()
+{
+	// List all the interface parameters that want a node.
+	TSet<UCameraRigInterfaceParameterBase*> InterfaceParametersWithNodes;
+	for (UCameraRigInterfaceParameterBase* InterfaceParameter : CameraRigAsset->Interface.BlendableParameters)
+	{
+		if (InterfaceParameter->bHasGraphNode)
+		{
+			InterfaceParametersWithNodes.Add(InterfaceParameter);
+		}
+	}
+	for (UCameraRigInterfaceParameterBase* InterfaceParameter : CameraRigAsset->Interface.DataParameters)
+	{
+		if (InterfaceParameter->bHasGraphNode)
+		{
+			InterfaceParametersWithNodes.Add(InterfaceParameter);
+		}
+	}
+
+	// Find all the interface parameter nodes that already exist.
+	TArray<UCameraRigInterfaceParameterGraphNode*> InterfaceParameterNodes;
+	NodeGraph->GetNodesOfClass(InterfaceParameterNodes);
+	TMap<UCameraRigInterfaceParameterBase*, UCameraRigInterfaceParameterGraphNode*> InterfaceParameterToNodeMap;
+
+	// Remove nodes that aren't needed anymore.
+	for (UCameraRigInterfaceParameterGraphNode* InterfaceParameterNode : InterfaceParameterNodes)
+	{
+		InterfaceParameterToNodeMap.Add(InterfaceParameterNode->GetInterfaceParameter(), InterfaceParameterNode);
+		if (!InterfaceParametersWithNodes.Contains(InterfaceParameterNode->GetInterfaceParameter()))
+		{
+			NodeGraph->Modify();
+			NodeGraph->RemoveNode(InterfaceParameterNode);
+		}
+	}
+
+	// Add nodes that are newly needed.
+	const UCameraNodeGraphSchema* Schema = CastChecked<UCameraNodeGraphSchema>(NodeGraph->GetSchema());
+	for (UCameraRigInterfaceParameterBase* InterfaceParameter : InterfaceParametersWithNodes)
+	{
+		if (!InterfaceParameterToNodeMap.Contains(InterfaceParameter))
+		{
+			NodeGraph->Modify();
+			Schema->CreateInterfaceParameterNode(NodeGraph, InterfaceParameter);
+		}
+	}
+}
+
 }  // namespace UE::Cameras
 
 #undef LOCTEXT_NAMESPACE
diff --git a/Engine/Plugins/Cameras/GameplayCameras/Source/GameplayCamerasEditor/Private/Editors/SCameraRigInterfaceParameterGraphNode.cpp b/Engine/Plugins/Cameras/GameplayCameras/Source/GameplayCamerasEditor/Private/Editors/SCameraRigInterfaceParameterGraphNode.cpp
--- a/Engine/Plugins/Cameras/GameplayCameras/Source/GameplayCamerasEditor/Private/Editors/SCameraRigInterfaceParameterGraphNode.cpp
+++ b/Engine/Plugins/Cameras/GameplayCameras/Source/GameplayCamerasEditor/Private/Editors/SCameraRigInterfaceParameterGraphNode.cpp
@@ -3,7 +3,7 @@
 #include "Editors/SCameraRigInterfaceParameterGraphNode.h"
 
 #include "Core/CameraRigAsset.h"
-#include "Editors/ObjectTreeGraphNode.h"
+#include "Editors/CameraRigInterfaceParameterGraphNode.h"
 #include "SGraphNode.h"
 #include "SNodePanel.h"
 #include "Styles/GameplayCamerasEditorStyle.h"
@@ -132,9 +132,9 @@ TSharedPtr<SGraphPin> SCameraRigInterfaceParameterGraphNode::CreatePinWidget(UEd
 
 FText SCameraRigInterfaceParameterGraphNode::GetInterfaceParameterName() const
 {
-	if (UCameraRigInterfaceParameter* RigParameterObject = GetObjectGraphNode()->CastObject<UCameraRigInterfaceParameter>())
+	if (UCameraRigInterfaceParameterBase* InterfaceParameter = GetObjectGraphNode()->CastObject<UCameraRigInterfaceParameterBase>())
 	{
-		return FText::FromString(RigParameterObject->InterfaceParameterName);
+		return FText::FromString(InterfaceParameter->InterfaceParameterName);
 	}
 	return LOCTEXT("InvalidParameterName", "Invalid");
 }
diff --git a/Engine/Plugins/Cameras/GameplayCameras/Source/GameplayCamerasEditor/Private/GameplayCamerasEditorModule.cpp b/Engine/Plugins/Cameras/GameplayCameras/Source/GameplayCamerasEditor/Private/GameplayCamerasEditorModule.cpp
--- a/Engine/Plugins/Cameras/GameplayCameras/Source/GameplayCamerasEditor/Private/GameplayCamerasEditorModule.cpp
+++ b/Engine/Plugins/Cameras/GameplayCameras/Source/GameplayCamerasEditor/Private/GameplayCamerasEditorModule.cpp
@@ -12,6 +12,7 @@
 #include "Commands/CameraVariableCollectionEditorCommands.h"
 #include "Commands/GameplayCamerasDebuggerCommands.h"
 #include "Commands/ObjectTreeGraphEditorCommands.h"
+#include "Customizations/CameraAssetReferenceDetailsCustomization.h"
 #include "Customizations/CameraParameterDetailsCustomizations.h"
 #include "Customizations/CameraProxyTableDetailsCustomization.h"
 #include "Customizations/CameraRigAssetReferenceDetailsCustomization.h"
@@ -87,7 +88,6 @@ public:
 
 		RegisterSettings();
 		RegisterCameraDirectorEditors();
-		RegisterBuilders();
 		RegisterCoreDebugCategories();
 		RegisterRewindDebuggerFeatures();
 		RegisterDetailsCustomizations();
@@ -114,7 +114,6 @@ public:
 
 		UnregisterSettings();
 		UnregisterCameraDirectorEditors();
-		UnregisterBuilders();
 		UnregisterCoreDebugCategories();
 		UnregisterRewindDebuggerFeatures();
 		UnregisterDetailsCustomizations();
@@ -186,46 +185,6 @@ public:
 				});
 	}
 
-	virtual FDelegateHandle RegisterCameraAssetBuilder(FOnBuildCameraAsset InOnBuildCameraAsset) override
-	{
-		CameraAssetBuilders.Add(InOnBuildCameraAsset);
-		return CameraAssetBuilders.Last().GetHandle();
-	}
-
-	virtual TArrayView<const FOnBuildCameraAsset> GetCameraAssetBuilders() const override
-	{
-		return CameraAssetBuilders;
-	}
-
-	virtual void UnregisterCameraAssetBuilder(FDelegateHandle InHandle) override
-	{
-		CameraAssetBuilders.RemoveAll(
-				[=](const FOnBuildCameraAsset& Delegate)
-				{
-					return Delegate.GetHandle() == InHandle;
-				});
-	}
-
-	virtual FDelegateHandle RegisterCameraRigAssetBuilder(FOnBuildCameraRigAsset InOnBuildCameraRigAsset) override
-	{
-		CameraRigAssetBuilders.Add(InOnBuildCameraRigAsset);
-		return CameraRigAssetBuilders.Last().GetHandle();
-	}
-	
-	virtual TArrayView<const FOnBuildCameraRigAsset> GetCameraRigAssetBuilders() const override
-	{
-		return CameraRigAssetBuilders;
-	}
-
-	virtual void UnregisterCameraRigAssetBuilder(FDelegateHandle InHandle) override
-	{
-		CameraRigAssetBuilders.RemoveAll(
-				[=](const FOnBuildCameraRigAsset& Delegate)
-				{
-					return Delegate.GetHandle() == InHandle;
-				});
-	}
-
 	virtual void RegisterDebugCategory(const UE::Cameras::FCameraDebugCategoryInfo& InCategoryInfo) override
 	{
 		if (!ensureMsgf(!InCategoryInfo.Name.IsEmpty(), TEXT("A debug category must at least specify a name!")))
@@ -335,26 +294,6 @@ private:
 		BuiltInDirectorCreatorHandles.Reset();
 	}
 
-	void RegisterBuilders()
-	{
-		// No default builders yet.
-	}
-
-	void UnregisterBuilders()
-	{
-		for (FDelegateHandle Handle : BuiltInCameraAssetBuilders)
-		{
-			UnregisterCameraAssetBuilder(Handle);
-		}
-		BuiltInCameraAssetBuilders.Reset();
-
-		for (FDelegateHandle Handle : BuiltInCameraRigAssetBuilders)
-		{
-			UnregisterCameraRigAssetBuilder(Handle);
-		}
-		BuiltInCameraRigAssetBuilders.Reset();
-	}
-
 	void RegisterCoreDebugCategories()
 	{
 		using namespace UE::Cameras;
@@ -473,12 +412,22 @@ private:
 		FCameraVariableReferenceDetailsCustomization::Register(PropertyEditorModule);
 		FRichCurveDetailsCustomization::Register(PropertyEditorModule);
 
-		PropertyEditorModule.RegisterCustomPropertyTypeLayout("CameraRigProxyTableEntry", FOnGetPropertyTypeCustomizationInstance::CreateStatic(
+		PropertyEditorModule.RegisterCustomPropertyTypeLayout(
+				"CameraAssetReference",
+				FOnGetPropertyTypeCustomizationInstance::CreateStatic(
+					&FCameraAssetReferenceDetailsCustomization::MakeInstance));
+		PropertyEditorModule.RegisterCustomPropertyTypeLayout(
+				"CameraRigProxyTableEntry", 
+				FOnGetPropertyTypeCustomizationInstance::CreateStatic(
 					&FCameraProxyTableEntryDetailsCustomization::MakeInstance));
-		PropertyEditorModule.RegisterCustomPropertyTypeLayout("CameraRigAssetReference", FOnGetPropertyTypeCustomizationInstance::CreateStatic(
+		PropertyEditorModule.RegisterCustomPropertyTypeLayout(
+				"CameraRigAssetReference", 
+				FOnGetPropertyTypeCustomizationInstance::CreateStatic(
 					&FCameraRigAssetReferenceDetailsCustomization::MakeInstance));
 
-		PropertyEditorModule.RegisterCustomClassLayout("FilmbackCameraNode", FOnGetDetailCustomizationInstance::CreateStatic(
+		PropertyEditorModule.RegisterCustomClassLayout(
+				"FilmbackCameraNode", 
+				FOnGetDetailCustomizationInstance::CreateStatic(
 					&FFilmbackCameraNodeDetailsCustomization::MakeInstance));
 	}
 
@@ -601,12 +550,6 @@ private:
 	TArray<FOnCreateCameraDirectorAssetEditorMode> CameraDirectorEditorCreators;
 	TArray<FDelegateHandle> BuiltInDirectorCreatorHandles;
 
-	TArray<FOnBuildCameraAsset> CameraAssetBuilders;
-	TArray<FDelegateHandle> BuiltInCameraAssetBuilders;
-
-	TArray<FOnBuildCameraRigAsset> CameraRigAssetBuilders;
-	TArray<FDelegateHandle> BuiltInCameraRigAssetBuilders;
-
 	TSharedPtr<UE::Cameras::FGameplayCamerasGraphPanelPinFactory> GraphPanelPinFactory;
 
 	TMap<FString, UE::Cameras::FCameraDebugCategoryInfo> DebugCategoryInfos;
diff --git a/Engine/Plugins/Cameras/GameplayCameras/Source/GameplayCamerasEditor/Private/Styles/GameplayCamerasEditorStyle.cpp b/Engine/Plugins/Cameras/GameplayCameras/Source/GameplayCamerasEditor/Private/Styles/GameplayCamerasEditorStyle.cpp
--- a/Engine/Plugins/Cameras/GameplayCameras/Source/GameplayCamerasEditor/Private/Styles/GameplayCamerasEditorStyle.cpp
+++ b/Engine/Plugins/Cameras/GameplayCameras/Source/GameplayCamerasEditor/Private/Styles/GameplayCamerasEditorStyle.cpp
@@ -80,6 +80,7 @@ FGameplayCamerasEditorStyle::FGameplayCamerasEditorStyle()
 	Set("CameraRigAssetEditor.Tabs.NodeHierarchy", new IMAGE_BRUSH_SVG("Icons/CameraRig-NodeHierarchy", Icon16x16));
 	Set("CameraRigAssetEditor.Tabs.Transitions", new IMAGE_BRUSH_SVG("Icons/CameraRig-Transitions", Icon16x16));
 	Set("CameraRigAssetEditor.Tabs.Curves", new IMAGE_BRUSH_SVG("Icons/CurveEditor", Icon16x16));
+	Set("CameraRigAssetEditor.Tabs.InterfaceParameters", new IMAGE_BRUSH_SVG("Icons/CameraRig-InterfaceParameters", Icon16x16));
 
 	Set("CameraRigAssetEditor.Build", new IMAGE_BRUSH_SVG("Icons/CameraRig-BuildStatus_Background", Icon20x20));
 	Set("CameraRigAssetEditor.BuildStatus.Background", new IMAGE_BRUSH_SVG("Icons/CameraRig-BuildStatus_Background", Icon20x20));
@@ -96,8 +97,11 @@ FGameplayCamerasEditorStyle::FGameplayCamerasEditorStyle()
 	Set("CameraRigAssetEditor.Apply", new IMAGE_BRUSH_SVG("CameraEditor-Apply", Icon20x20));
 	Set("CameraRigAssetEditor.FindInCameraRig", new CORE_IMAGE_BRUSH_SVG("Starship/Common/Search", Icon20x20));
 
+	Set("CameraRigAssetEditor.InterfaceParameter.Message", FTextBlockStyle(NormalText).SetFont(DEFAULT_FONT("Italic", 10)));
+
 	// Camera parameters icons.
 	Set("CameraParameter.VariableBrowser", new IMAGE_BRUSH_SVG("Icons/CameraParameter-Variable", Icon16x16));
+	Set("CameraParameter.TypeIcon", new IMAGE_BRUSH_SVG("Icons/CameraParameter-Pill", Icon16x16));
 
 	// Camera variable collection icons.
 	Set("CameraVariableCollectionEditor.CreateVariable", new CORE_IMAGE_BRUSH_SVG("Starship/Common/plus", Icon16x16));
diff --git a/Engine/Plugins/Cameras/GameplayCameras/Source/GameplayCamerasEditor/Private/Toolkits/CameraAssetEditorToolkit.cpp b/Engine/Plugins/Cameras/GameplayCameras/Source/GameplayCamerasEditor/Private/Toolkits/CameraAssetEditorToolkit.cpp
--- a/Engine/Plugins/Cameras/GameplayCameras/Source/GameplayCamerasEditor/Private/Toolkits/CameraAssetEditorToolkit.cpp
+++ b/Engine/Plugins/Cameras/GameplayCameras/Source/GameplayCamerasEditor/Private/Toolkits/CameraAssetEditorToolkit.cpp
@@ -283,29 +283,7 @@ void FCameraAssetEditorToolkit::OnBuild()
 
 	FCameraBuildLog BuildLog;
 	FCameraAssetBuilder Builder(BuildLog);
-	Builder.BuildCamera(
-			CameraAsset,
-			FCameraAssetBuilder::FCustomBuildStep::CreateLambda(
-				[&UsageInfo](UCameraAsset* InCameraAsset, FCameraBuildLog& BuildLog)
-				{
-					IGameplayCamerasEditorModule& GameplayCamerasEditorModule = IGameplayCamerasEditorModule::Get();
-					for (const FOnBuildCameraAsset& Builder : GameplayCamerasEditorModule.GetCameraAssetBuilders())
-					{
-						Builder.ExecuteIfBound(InCameraAsset, BuildLog);
-					}
-
-					if (UCameraDirector* CameraDirector = InCameraAsset->GetCameraDirector())
-					{
-						CameraDirector->GatherRigUsageInfo(UsageInfo);
-						for (UCameraRigAsset* CameraRig : UsageInfo.CameraRigs)
-						{
-							for (const FOnBuildCameraRigAsset& Builder : GameplayCamerasEditorModule.GetCameraRigAssetBuilders())
-							{
-								Builder.ExecuteIfBound(CameraRig, BuildLog);
-							}
-						}
-					}
-				}));
+	Builder.BuildCamera(CameraAsset);
 	
 	BuildLogToolkit->PopulateMessageListing(BuildLog);
 
diff --git a/Engine/Plugins/Cameras/GameplayCameras/Source/GameplayCamerasEditor/Private/Toolkits/CameraRigAssetEditorToolkit.cpp b/Engine/Plugins/Cameras/GameplayCameras/Source/GameplayCamerasEditor/Private/Toolkits/CameraRigAssetEditorToolkit.cpp
--- a/Engine/Plugins/Cameras/GameplayCameras/Source/GameplayCamerasEditor/Private/Toolkits/CameraRigAssetEditorToolkit.cpp
+++ b/Engine/Plugins/Cameras/GameplayCameras/Source/GameplayCamerasEditor/Private/Toolkits/CameraRigAssetEditorToolkit.cpp
@@ -24,6 +24,7 @@
 #include "Toolkits/BuildButtonToolkit.h"
 #include "Toolkits/CameraBuildLogToolkit.h"
 #include "Toolkits/CameraRigAssetEditorToolkitBase.h"
+#include "Toolkits/CameraRigInterfaceParametersToolkit.h"
 #include "Toolkits/CurveEditorToolkit.h"
 #include "Toolkits/StandardToolkitLayout.h"
 #include "Widgets/Docking/SDockTab.h"
@@ -38,6 +39,7 @@ namespace UE::Cameras
 const FName FCameraRigAssetEditorToolkit::SearchTabId(TEXT("CameraRigAssetEditor_Search"));
 const FName FCameraRigAssetEditorToolkit::MessagesTabId(TEXT("CameraRigAssetEditor_Messages"));
 const FName FCameraRigAssetEditorToolkit::CurvesTabId(TEXT("CameraRigAssetEditor_Curves"));
+const FName FCameraRigAssetEditorToolkit::InterfaceParametersTabId(TEXT("CameraRigAssetEditor_InterfaceParameters"));
 
 FCameraRigAssetEditorToolkit::FCameraRigAssetEditorToolkit(UAssetEditor* InOwningAssetEditor)
 	: FBaseAssetToolkit(InOwningAssetEditor)
@@ -46,6 +48,7 @@ FCameraRigAssetEditorToolkit::FCameraRigAssetEditorToolkit(UAssetEditor* InOwnin
 	BuildButtonToolkit = MakeShared<FBuildButtonToolkit>();
 	BuildLogToolkit = MakeShared<FCameraBuildLogToolkit>();
 	CurveEditorToolkit = MakeShared<FCurveEditorToolkit>();
+	InterfaceParametersToolkit = MakeShared<FCameraRigInterfaceParametersToolkit>();
 
 	// Override base class default layout.
 	TSharedPtr<FStandardToolkitLayout> StandardLayout = Impl->GetStandardLayout();
@@ -53,6 +56,7 @@ FCameraRigAssetEditorToolkit::FCameraRigAssetEditorToolkit(UAssetEditor* InOwnin
 		StandardLayout->AddBottomTab(SearchTabId);
 		StandardLayout->AddBottomTab(MessagesTabId);
 		StandardLayout->AddBottomTab(CurvesTabId);
+		StandardLayout->AddLeftTab(InterfaceParametersTabId, ETabState::OpenedTab);
 	}
 	StandaloneDefaultLayout = StandardLayout->GetLayout();
 
@@ -76,6 +80,7 @@ void FCameraRigAssetEditorToolkit::SetCameraRigAsset(UCameraRigAsset* InCameraRi
 
 	Impl->SetCameraRigAsset(InCameraRig);
 	BuildButtonToolkit->SetTarget(InCameraRig);
+	InterfaceParametersToolkit->SetCameraRigAsset(InCameraRig);
 
 	if (InCameraRig)
 	{
@@ -106,6 +111,11 @@ void FCameraRigAssetEditorToolkit::RegisterTabSpawners(const TSharedRef<class FT
 		.SetDisplayName(LOCTEXT("Curves", "Curves"))
 		.SetGroup(AssetEditorTabsCategory.ToSharedRef())
 		.SetIcon(FSlateIcon(CamerasStyleSetName, "CameraRigAssetEditor.Tabs.Curves"));
+
+	InTabManager->RegisterTabSpawner(InterfaceParametersTabId, FOnSpawnTab::CreateSP(this, &FCameraRigAssetEditorToolkit::SpawnTab_InterfaceParameters))
+		.SetDisplayName(LOCTEXT("InterfaceParameters", "InterfaceParameters"))
+		.SetGroup(AssetEditorTabsCategory.ToSharedRef())
+		.SetIcon(FSlateIcon(CamerasStyleSetName, "CameraRigAssetEditor.Tabs.InterfaceParameters"));
 }
 
 TSharedRef<SDockTab> FCameraRigAssetEditorToolkit::SpawnTab_Search(const FSpawnTabArgs& Args)
@@ -150,6 +160,17 @@ TSharedRef<SDockTab> FCameraRigAssetEditorToolkit::SpawnTab_Curves(const FSpawnT
 	return CurvesTab.ToSharedRef();
 }
 
+TSharedRef<SDockTab> FCameraRigAssetEditorToolkit::SpawnTab_InterfaceParameters(const FSpawnTabArgs& Args)
+{
+	TSharedPtr<SDockTab> InterfaceParametersTab = SNew(SDockTab)
+		.Label(LOCTEXT("InterfaceParametersTabTitle", "Parameters"))
+		[
+			InterfaceParametersToolkit->GetInterfaceParametersPanel().ToSharedRef()
+		];
+
+	return InterfaceParametersTab.ToSharedRef();
+}
+
 void FCameraRigAssetEditorToolkit::OnCurvesTabClosed(TSharedRef<SDockTab> InTab)
 {
 	if (CurveEditorToolkit->IsInitialized())
@@ -171,6 +192,7 @@ void FCameraRigAssetEditorToolkit::UnregisterTabSpawners(const TSharedRef<FTabMa
 	InTabManager->UnregisterTabSpawner(SearchTabId);
 	InTabManager->UnregisterTabSpawner(MessagesTabId);
 	InTabManager->UnregisterTabSpawner(CurvesTabId);
+	InTabManager->UnregisterTabSpawner(InterfaceParametersTabId);
 }
 
 void FCameraRigAssetEditorToolkit::CreateWidgets()
@@ -201,6 +223,9 @@ void FCameraRigAssetEditorToolkit::CreateWidgets()
 
 	// Create the message log.
 	BuildLogToolkit->Initialize("CameraRigAssetBuildMessages");
+
+	// Hook-up the selection of interface parameters.
+	InterfaceParametersToolkit->OnInterfaceParameterSelected().AddSP(this, &FCameraRigAssetEditorToolkit::OnCameraRigInterfaceParameterSelected);
 }
 
 void FCameraRigAssetEditorToolkit::RegisterToolbar()
@@ -282,6 +307,11 @@ void FCameraRigAssetEditorToolkit::PostRegenerateMenusAndToolbars()
 	SetMenuOverlay(FAssetTypeMenuOverlayHelper::CreateMenuOverlay(UCameraRigAsset::StaticClass()));
 }
 
+void FCameraRigAssetEditorToolkit::OnCameraRigInterfaceParameterSelected(UCameraRigInterfaceParameterBase* Object)
+{
+	OnJumpToObject(Object, NAME_None);
+}
+
 void FCameraRigAssetEditorToolkit::OnBuild()
 {
 	UCameraRigAsset* CameraRigAsset = Impl->GetCameraRigAsset();
@@ -292,17 +322,7 @@ void FCameraRigAssetEditorToolkit::OnBuild()
 
 	FCameraBuildLog BuildLog;
 	FCameraRigAssetBuilder Builder(BuildLog);
-	Builder.BuildCameraRig(
-			CameraRigAsset,
-			FCameraRigAssetBuilder::FCustomBuildStep::CreateLambda(
-				[](UCameraRigAsset* InCameraRigAsset, FCameraBuildLog& BuildLog)
-				{
-					IGameplayCamerasEditorModule& GameplayCamerasEditorModule = IGameplayCamerasEditorModule::Get();
-					for (const FOnBuildCameraRigAsset& Builder : GameplayCamerasEditorModule.GetCameraRigAssetBuilders())
-					{
-						Builder.ExecuteIfBound(InCameraRigAsset, BuildLog);
-					}
-				}));
+	Builder.BuildCameraRig(CameraRigAsset);
 
 	BuildLogToolkit->PopulateMessageListing(BuildLog);
 
