Cameras: big pass on parameterization of camera rigs

diff --git a/Engine/Plugins/Cameras/GameplayCameras/Source/GameplayCamerasEditor/Private/Toolkits/CameraRigInterfaceParametersToolkit.cpp b/Engine/Plugins/Cameras/GameplayCameras/Source/GameplayCamerasEditor/Private/Toolkits/CameraRigInterfaceParametersToolkit.cpp
new file mode 100644
--- /dev/null
+++ b/Engine/Plugins/Cameras/GameplayCameras/Source/GameplayCamerasEditor/Private/Toolkits/CameraRigInterfaceParametersToolkit.cpp
@@ -0,0 +1,933 @@
+// Copyright Epic Games, Inc. All Rights Reserved.
+
+#include "Toolkits/CameraRigInterfaceParametersToolkit.h"
+
+#include "Core/CameraRigAsset.h"
+#include "EdGraphSchema_K2.h"
+#include "Editor.h"
+#include "Editors/CameraNodeGraphDragDropOp.h"
+#include "Framework/MultiBox/MultiBoxBuilder.h"
+#include "ScopedTransaction.h"
+#include "SPinTypeSelector.h"
+#include "Styles/GameplayCamerasEditorStyle.h"
+#include "Styling/SlateIconFinder.h"
+#include "ToolMenus.h"
+#include "Types/SlateEnums.h"
+#include "Widgets/Input/SButton.h"
+#include "Widgets/Input/SComboButton.h"
+#include "Widgets/Layout/SBox.h"
+#include "Widgets/Layout/SSplitter.h"
+#include "Widgets/SBoxPanel.h"
+#include "Widgets/SCompoundWidget.h"
+#include "Widgets/Text/SInlineEditableTextBlock.h"
+#include "Widgets/Text/SRichTextBlock.h"
+#include "Widgets/Text/STextBlock.h"
+#include "Widgets/Views/SListView.h"
+
+#define LOCTEXT_NAMESPACE "CameraRigInterfaceParametersToolkit"
+
+namespace UE::Cameras
+{
+
+static const FName ParameterTypeColumn(TEXT("ParameterType"));
+static const FName ParameterNameColumn(TEXT("ParameterName"));
+static const FName ParameterMessageColumn(TEXT("ParameterMessage"));
+
+/**
+ * List entry for any interface parameter panel.
+ */
+template<typename ParameterType>
+class SCameraRigInterfaceParameterTableRowBase : public SMultiColumnTableRow<TObjectPtr<ParameterType>>
+{
+public:
+
+	SLATE_BEGIN_ARGS(SCameraRigInterfaceParameterTableRowBase<ParameterType>)
+	{}
+		SLATE_ARGUMENT(TObjectPtr<ParameterType>, Item)
+	SLATE_END_ARGS()
+
+	typedef SMultiColumnTableRow<TObjectPtr<ParameterType>> FSuperRowType;
+	typedef typename STableRow<TObjectPtr<ParameterType>>::FArguments FTableRowArgs;
+		
+	void Construct(const FArguments& Args, const TSharedRef<STableViewBase>& OwnerTable)
+	{
+		Item = Args._Item;
+
+		FSuperRowType::Construct(
+				FTableRowArgs(),
+				OwnerTable);
+	}
+
+	void EnterNameEditingMode()
+	{
+		NameTextBlock->EnterEditingMode();
+	}
+
+protected:
+
+	virtual TSharedRef<SWidget> GenerateWidgetForColumn(const FName& InColumnName) override
+	{
+		if (InColumnName == ParameterNameColumn)
+		{
+			return SAssignNew(NameTextBlock, SInlineEditableTextBlock)
+				.IsSelected(this, &SCameraRigInterfaceParameterTableRowBase<ParameterType>::IsSelected)
+				.Text_Lambda([this]() { return FText::FromString(Item->InterfaceParameterName); })
+				.OnTextCommitted(this, &SCameraRigInterfaceParameterTableRowBase<ParameterType>::OnParameterNameTextCommitted);
+		}
+		else if (InColumnName == ParameterMessageColumn)
+		{
+			TSharedRef<FGameplayCamerasEditorStyle> GameplayCamerasEditorStyle = FGameplayCamerasEditorStyle::Get();
+			const FTextBlockStyle& MessageStyle = GameplayCamerasEditorStyle->GetWidgetStyle<FTextBlockStyle>("CameraRigAssetEditor.InterfaceParameter.Message");
+
+			return SNew(STextBlock)
+				.TextStyle(&MessageStyle)
+				.Text(this, &SCameraRigInterfaceParameterTableRowBase<ParameterType>::GetParameterMessageText)
+				.ToolTipText(this, &SCameraRigInterfaceParameterTableRowBase<ParameterType>::GetParameterMessageToolTip);
+		}
+
+		return SNullWidget::NullWidget;
+	}
+
+	virtual FReply OnDragDetected(const FGeometry& MyGeometry, const FPointerEvent& MouseEvent) override
+	{
+		if (MouseEvent.IsMouseButtonDown(EKeys::LeftMouseButton))
+		{
+			return FReply::Handled().BeginDragDrop(FCameraNodeGraphInterfaceParameterDragDropOp::New(Item));
+		}
+		return SMultiColumnTableRow<TObjectPtr<ParameterType>>::OnDragDetected(MyGeometry, MouseEvent);
+	}
+
+	void OnParameterNameTextCommitted(const FText& Text, ETextCommit::Type CommitType)
+	{
+		const FScopedTransaction Transaction(LOCTEXT("RenameInterfaceParameter", "Rename Interface Parameter"));
+
+		Item->Modify();
+		Item->InterfaceParameterName = Text.ToString();
+	}
+
+	FText GetParameterMessageText() const
+	{
+		if (!Item->Target || Item->TargetPropertyName.IsNone())
+		{
+			return LOCTEXT("UnboundInterfaceParameterMessage", "Unbound");
+		}
+		return FText::GetEmpty();
+	}
+
+	FText GetParameterMessageToolTip() const
+	{
+		if (!Item->Target || Item->TargetPropertyName.IsNone())
+		{
+			return LOCTEXT(
+					"UnboundInterfaceParameterMessageToolTip", 
+					"This interface parameter is not bound to any camera node. Setting this parameter will have no effect.");
+		}
+		return FText::GetEmpty();
+	}
+
+protected:
+
+	TObjectPtr<ParameterType> Item;
+	TSharedPtr<SInlineEditableTextBlock> NameTextBlock;
+};
+
+/**
+ * List entry for the blendable parameters panel.
+ */
+class SCameraRigBlendableParameterTableRow : public SCameraRigInterfaceParameterTableRowBase<UCameraRigBlendableParameter>
+{
+public:
+
+	void Construct(const FArguments& Args, const TSharedRef<STableViewBase>& OwnerTable)
+	{
+		SCameraRigInterfaceParameterTableRowBase<UCameraRigBlendableParameter>::Construct(Args, OwnerTable);
+	}
+
+protected:
+
+	virtual TSharedRef<SWidget> GenerateWidgetForColumn(const FName& InColumnName) override
+	{
+		if (InColumnName == ParameterTypeColumn)
+		{
+			TSharedRef<FGameplayCamerasEditorStyle> GameplayCamerasEditorStyle = FGameplayCamerasEditorStyle::Get();
+
+			const UEdGraphSchema_K2* K2Schema = GetDefault<UEdGraphSchema_K2>();
+
+			return SNew(SBox)
+				.MinDesiredWidth(125.f)
+				.MaxDesiredWidth(400.f)
+				.Padding(this, &SCameraRigBlendableParameterTableRow::GetPinTypeSelectorPadding)
+				[
+					SNew(SPinTypeSelector, FGetPinTypeTree::CreateSP(this, &SCameraRigBlendableParameterTableRow::GetPinTypeTreeInfos))
+					.OnPinTypeChanged(this, &SCameraRigBlendableParameterTableRow::OnBlendableParameterPinTypeChanged)
+					.TargetPinType(this, &SCameraRigBlendableParameterTableRow::GetBlendableParameterPinType)
+					.ReadOnly(this, &SCameraRigBlendableParameterTableRow::IsPinTypeSelectorReadOnly)
+					.Schema(K2Schema)
+					.bAllowArrays(false)
+				];
+		}
+
+		return SCameraRigInterfaceParameterTableRowBase<UCameraRigBlendableParameter>::GenerateWidgetForColumn(InColumnName);
+	}
+
+	FMargin GetPinTypeSelectorPadding() const
+	{
+		// Add some horizontal margin when the pin type selector is read-only, so that the transition on hover is seamless.
+		return IsHovered() ? FMargin(0) : FMargin(8, 0);
+	}
+
+	void GetPinTypeTreeInfos(TArray<FPinTypeTreeItem>& TypeTree, ETypeTreeFilter TypeTreeFilter) const
+	{
+		using FPinTypeTreeInfo = UEdGraphSchema_K2::FPinTypeTreeInfo;
+
+		const UEdGraphSchema_K2* K2Schema = GetDefault<UEdGraphSchema_K2>();
+
+		TypeTree.Reset();
+
+		TypeTree.Add(MakeShared<FPinTypeTreeInfo>(UEdGraphSchema_K2::GetCategoryText(UEdGraphSchema_K2::PC_Boolean, true), UEdGraphSchema_K2::PC_Boolean, K2Schema, LOCTEXT("BooleanType", "True or false value")));
+		TypeTree.Add(MakeShared<FPinTypeTreeInfo>(UEdGraphSchema_K2::GetCategoryText(UEdGraphSchema_K2::PC_Int, true), UEdGraphSchema_K2::PC_Int, K2Schema, LOCTEXT("IntegerType", "Integer number")));
+		TypeTree.Add(MakeShared<FPinTypeTreeInfo>(UEdGraphSchema_K2::GetCategoryText(UEdGraphSchema_K2::PC_Float, true), UEdGraphSchema_K2::PC_Float, K2Schema, LOCTEXT("FloatType", "Floating point number")));
+		TypeTree.Add(MakeShared<FPinTypeTreeInfo>(UEdGraphSchema_K2::GetCategoryText(UEdGraphSchema_K2::PC_Double, true), UEdGraphSchema_K2::PC_Double, K2Schema, LOCTEXT("DoubleType", "Double precision floating point number")));
+
+		TypeTree.Add(MakeShared<FPinTypeTreeInfo>(UEdGraphSchema_K2::PC_Struct, TVariantStructure<FVector2f>::Get(), LOCTEXT("Vector2fType", "A 2D vector")));
+		TypeTree.Add(MakeShared<FPinTypeTreeInfo>(UEdGraphSchema_K2::PC_Struct, TBaseStructure<FVector2D>::Get(), LOCTEXT("Vector2dType", "A double precision 2D vector")));
+		TypeTree.Add(MakeShared<FPinTypeTreeInfo>(UEdGraphSchema_K2::PC_Struct, TVariantStructure<FVector3f>::Get(), LOCTEXT("Vector3fType", "A 3D vector")));
+		TypeTree.Add(MakeShared<FPinTypeTreeInfo>(UEdGraphSchema_K2::PC_Struct, TBaseStructure<FVector>::Get(), LOCTEXT("Vector3dType", "A double precision 3D vector")));
+		TypeTree.Add(MakeShared<FPinTypeTreeInfo>(UEdGraphSchema_K2::PC_Struct, TVariantStructure<FVector4f>::Get(), LOCTEXT("Vector4fType", "A 4D vector")));
+		TypeTree.Add(MakeShared<FPinTypeTreeInfo>(UEdGraphSchema_K2::PC_Struct, TBaseStructure<FVector4>::Get(), LOCTEXT("Vector4dType", "A double precision 4D vector")));
+		TypeTree.Add(MakeShared<FPinTypeTreeInfo>(UEdGraphSchema_K2::PC_Struct, TVariantStructure<FRotator3f>::Get(), LOCTEXT("Rotator3fType", "A 3D rotation")));
+		TypeTree.Add(MakeShared<FPinTypeTreeInfo>(UEdGraphSchema_K2::PC_Struct, TBaseStructure<FRotator>::Get(), LOCTEXT("Rotator3dType", "A double precision 3D rotation")));
+		TypeTree.Add(MakeShared<FPinTypeTreeInfo>(UEdGraphSchema_K2::PC_Struct, TVariantStructure<FTransform3f>::Get(), LOCTEXT("TransformType", "A 3D transformation")));
+		TypeTree.Add(MakeShared<FPinTypeTreeInfo>(UEdGraphSchema_K2::PC_Struct, TBaseStructure<FTransform>::Get(), LOCTEXT("TransformType", "A double precision 3D transformation")));
+	}
+
+	FEdGraphPinType GetBlendableParameterPinType() const
+	{
+		FEdGraphPinType PinType;
+		switch (Item->ParameterType)
+		{
+			case ECameraVariableType::Boolean:
+				PinType.PinCategory = UEdGraphSchema_K2::PC_Boolean;
+				break;
+			case ECameraVariableType::Integer32:
+				PinType.PinCategory = UEdGraphSchema_K2::PC_Int;
+				break;
+			case ECameraVariableType::Float:
+				PinType.PinCategory = UEdGraphSchema_K2::PC_Float;
+				break;
+			case ECameraVariableType::Double:
+				PinType.PinCategory = UEdGraphSchema_K2::PC_Double;
+				break;
+			case ECameraVariableType::Vector2f:
+				PinType.PinCategory = UEdGraphSchema_K2::PC_Struct;
+				PinType.PinSubCategoryObject = TVariantStructure<FVector2f>::Get();
+				break;
+			case ECameraVariableType::Vector2d:
+				PinType.PinCategory = UEdGraphSchema_K2::PC_Struct;
+				PinType.PinSubCategoryObject = TBaseStructure<FVector2D>::Get();
+				break;
+			case ECameraVariableType::Vector3f:
+				PinType.PinCategory = UEdGraphSchema_K2::PC_Struct;
+				PinType.PinSubCategoryObject = TVariantStructure<FVector3f>::Get();
+				break;
+			case ECameraVariableType::Vector3d:
+				PinType.PinCategory = UEdGraphSchema_K2::PC_Struct;
+				PinType.PinSubCategoryObject = TBaseStructure<FVector>::Get();
+				break;
+			case ECameraVariableType::Vector4f:
+				PinType.PinCategory = UEdGraphSchema_K2::PC_Struct;
+				PinType.PinSubCategoryObject = TVariantStructure<FVector4f>::Get();
+				break;
+			case ECameraVariableType::Vector4d:
+				PinType.PinCategory = UEdGraphSchema_K2::PC_Struct;
+				PinType.PinSubCategoryObject = TBaseStructure<FVector4>::Get();
+				break;
+			case ECameraVariableType::Rotator3f:
+				PinType.PinCategory = UEdGraphSchema_K2::PC_Struct;
+				PinType.PinSubCategoryObject = TVariantStructure<FRotator3f>::Get();
+				break;
+			case ECameraVariableType::Rotator3d:
+				PinType.PinCategory = UEdGraphSchema_K2::PC_Struct;
+				PinType.PinSubCategoryObject = TBaseStructure<FRotator>::Get();
+				break;
+			case ECameraVariableType::Transform3f:
+				PinType.PinCategory = UEdGraphSchema_K2::PC_Struct;
+				PinType.PinSubCategoryObject = TVariantStructure<FTransform3f>::Get();
+				break;
+			case ECameraVariableType::Transform3d:
+				PinType.PinCategory = UEdGraphSchema_K2::PC_Struct;
+				PinType.PinSubCategoryObject = TBaseStructure<FTransform>::Get();
+				break;
+		}
+		return PinType;
+	}
+
+	void OnBlendableParameterPinTypeChanged(const FEdGraphPinType& PinType)
+	{
+		bool bIsValidType = true;
+		ECameraVariableType NewParameterType = ECameraVariableType::Boolean;
+
+		if (PinType.PinCategory == UEdGraphSchema_K2::PC_Boolean)
+		{
+			NewParameterType = ECameraVariableType::Boolean;
+		}
+		else if (PinType.PinCategory == UEdGraphSchema_K2::PC_Int)
+		{
+			NewParameterType = ECameraVariableType::Integer32;
+		}
+		else if (PinType.PinCategory == UEdGraphSchema_K2::PC_Float)
+		{
+			NewParameterType = ECameraVariableType::Float;
+		}
+		else if (PinType.PinCategory == UEdGraphSchema_K2::PC_Double)
+		{
+			NewParameterType = ECameraVariableType::Double;
+		}
+		else if (PinType.PinCategory == UEdGraphSchema_K2::PC_Struct)
+		{
+			const UObject* TypeObject = PinType.PinSubCategoryObject.Get();
+			if (TypeObject == TVariantStructure<FVector2f>::Get())
+			{
+				NewParameterType = ECameraVariableType::Vector2f;
+			}
+			else if (TypeObject == TBaseStructure<FVector2D>::Get())
+			{
+				NewParameterType = ECameraVariableType::Vector2d;
+			}
+			else if (TypeObject == TVariantStructure<FVector3f>::Get())
+			{
+				NewParameterType = ECameraVariableType::Vector3f;
+			}
+			else if (TypeObject == TBaseStructure<FVector>::Get())
+			{
+				NewParameterType = ECameraVariableType::Vector3d;
+			}
+			else if (TypeObject == TVariantStructure<FVector4f>::Get())
+			{
+				NewParameterType = ECameraVariableType::Vector4f;
+			}
+			else if (TypeObject == TBaseStructure<FVector4>::Get())
+			{
+				NewParameterType = ECameraVariableType::Vector4d;
+			}
+			else if (TypeObject == TVariantStructure<FRotator3f>::Get())
+			{
+				NewParameterType = ECameraVariableType::Rotator3f;
+			}
+			else if (TypeObject == TBaseStructure<FRotator>::Get())
+			{
+				NewParameterType = ECameraVariableType::Rotator3d;
+			}
+			else if (TypeObject == TVariantStructure<FTransform3f>::Get())
+			{
+				NewParameterType = ECameraVariableType::Transform3f;
+			}
+			else if (TypeObject == TBaseStructure<FTransform>::Get())
+			{
+				NewParameterType = ECameraVariableType::Transform3d;
+			}
+			else
+			{
+				bIsValidType = false;
+			}
+		}
+		else
+		{
+			bIsValidType = false;
+		}
+
+		if (ensure(bIsValidType) && Item->ParameterType != NewParameterType)
+		{
+			const FScopedTransaction Transaction(LOCTEXT("ChangeBlendableParameterType", "Change Blendable Parameter Type"));
+
+			Item->Modify();
+			Item->ParameterType = NewParameterType;
+		}
+	}
+
+	bool IsPinTypeSelectorReadOnly() const
+	{
+		return !IsHovered();
+	}
+};
+
+class FDataParameterPinTypeSelectorFilter : public IPinTypeSelectorFilter
+{
+	virtual bool ShouldShowPinTypeTreeItem(FPinTypeTreeItem InItem) const override
+	{
+		// Remove types that are for blendable parameters.
+		FEdGraphPinType PinType = InItem->GetPinType(false);
+		if (PinType.PinCategory == UEdGraphSchema_K2::PC_Boolean ||
+				PinType.PinCategory == UEdGraphSchema_K2::PC_Int ||
+				PinType.PinCategory == UEdGraphSchema_K2::PC_Float ||
+				PinType.PinCategory == UEdGraphSchema_K2::PC_Double ||
+				PinType.PinCategory == UEdGraphSchema_K2::PC_Real)
+		{
+			return false;
+		}
+
+		if (PinType.PinCategory == UEdGraphSchema_K2::PC_Struct)
+		{
+			const UObject* TypeObject = PinType.PinSubCategoryObject.Get();
+			if (TypeObject == TVariantStructure<FVector2f>::Get() ||
+					TypeObject == TBaseStructure<FVector2D>::Get() ||
+					TypeObject == TVariantStructure<FVector3f>::Get() ||
+					TypeObject == TBaseStructure<FVector>::Get() ||
+					TypeObject == TVariantStructure<FVector4f>::Get() ||
+					TypeObject == TBaseStructure<FVector4>::Get() ||
+					TypeObject == TVariantStructure<FRotator3f>::Get() ||
+					TypeObject == TBaseStructure<FRotator>::Get() ||
+					TypeObject == TVariantStructure<FTransform3f>::Get() ||
+					TypeObject == TBaseStructure<FTransform>::Get())
+			{
+				return false;
+			}
+		}
+		
+		return true;
+	};
+};
+
+/**
+ * List entry for the data parameters panel.
+ */
+class SCameraRigDataParameterTableRow : public SCameraRigInterfaceParameterTableRowBase<UCameraRigDataParameter>
+{
+protected:
+
+	virtual TSharedRef<SWidget> GenerateWidgetForColumn(const FName& InColumnName) override
+	{
+		if (InColumnName == ParameterTypeColumn)
+		{
+			TSharedRef<FGameplayCamerasEditorStyle> GameplayCamerasEditorStyle = FGameplayCamerasEditorStyle::Get();
+
+			const UEdGraphSchema_K2* K2Schema = GetDefault<UEdGraphSchema_K2>();
+			TArray<TSharedPtr<IPinTypeSelectorFilter>> PinTypeSelectorFilters{ MakeShared<FDataParameterPinTypeSelectorFilter>() };
+
+			return SNew(SBox)
+				.MinDesiredWidth(125.f)
+				.MaxDesiredWidth(400.f)
+				.Padding(this, &SCameraRigDataParameterTableRow::GetPinTypeSelectorPadding)
+				[
+					SNew(SPinTypeSelector, FGetPinTypeTree::CreateUObject(K2Schema, &UEdGraphSchema_K2::GetVariableTypeTree))
+					.OnPinTypeChanged(this, &SCameraRigDataParameterTableRow::OnDataParameterPinTypeChanged)
+					.TargetPinType(this, &SCameraRigDataParameterTableRow::GetDataParameterPinType)
+					.ReadOnly(this, &SCameraRigDataParameterTableRow::IsPinTypeSelectorReadOnly)
+					.Schema(K2Schema)
+					.CustomFilters(PinTypeSelectorFilters)
+					.bAllowArrays(false)
+				];
+		}
+
+		return SCameraRigInterfaceParameterTableRowBase<UCameraRigDataParameter>::GenerateWidgetForColumn(InColumnName);
+	}
+
+	FMargin GetPinTypeSelectorPadding() const
+	{
+		// Add some horizontal margin when the pin type selector is read-only, so that the transition on hover is seamless.
+		return IsHovered() ? FMargin(0) : FMargin(8, 0);
+	}
+
+	FEdGraphPinType GetDataParameterPinType() const
+	{
+		FEdGraphPinType PinType;
+		PinType.PinSubCategoryObject = const_cast<UObject*>(Item->DataTypeObject.Get());
+
+		switch (Item->DataType)
+		{
+			case ECameraContextDataType::Name:
+				PinType.PinCategory = UEdGraphSchema_K2::PC_Name;
+				break;
+			case ECameraContextDataType::String:
+				PinType.PinCategory = UEdGraphSchema_K2::PC_String;
+				break;
+			case ECameraContextDataType::Enum:
+				PinType.PinCategory = UEdGraphSchema_K2::PC_Enum;
+				break;
+			case ECameraContextDataType::Struct:
+				PinType.PinCategory = UEdGraphSchema_K2::PC_Struct;
+				break;
+			case ECameraContextDataType::Object:
+				PinType.PinCategory = UEdGraphSchema_K2::PC_Object;
+				break;
+			case ECameraContextDataType::Class:
+				PinType.PinCategory = UEdGraphSchema_K2::PC_Class;
+				break;
+		}
+
+		return PinType;
+	}
+
+	void OnDataParameterPinTypeChanged(const FEdGraphPinType& PinType)
+	{
+		bool bIsValidType = true;
+		ECameraContextDataType NewDataType = ECameraContextDataType::Name;
+		const UObject* NewDataTypeObject = PinType.PinSubCategoryObject.Get();
+
+		if (PinType.PinCategory == UEdGraphSchema_K2::PC_Name)
+		{
+			NewDataType = ECameraContextDataType::Name;
+		}
+		else if (PinType.PinCategory == UEdGraphSchema_K2::PC_String)
+		{
+			NewDataType = ECameraContextDataType::String;
+		}
+		else if (PinType.PinCategory == UEdGraphSchema_K2::PC_Enum)
+		{
+			NewDataType = ECameraContextDataType::Enum;
+		}
+		else if (PinType.PinCategory == UEdGraphSchema_K2::PC_Struct)
+		{
+			NewDataType = ECameraContextDataType::Struct;
+		}
+		else if (PinType.PinCategory == UEdGraphSchema_K2::PC_Object)
+		{
+			NewDataType = ECameraContextDataType::Object;
+		}
+		else if (PinType.PinCategory == UEdGraphSchema_K2::PC_Class)
+		{
+			NewDataType = ECameraContextDataType::Class;
+		}
+		else
+		{
+			bIsValidType = false;
+		}
+
+		if (ensure(bIsValidType) && (NewDataType != Item->DataType || NewDataTypeObject != Item->DataTypeObject))
+		{
+			const FScopedTransaction Transaction(LOCTEXT("ChangeDataParameterType", "Change Data Parameter Type"));
+
+			Item->DataType = NewDataType;
+			Item->DataTypeObject = NewDataTypeObject;
+		}
+	}
+
+	bool IsPinTypeSelectorReadOnly() const
+	{
+		return !IsHovered();
+	}
+
+private:
+
+	mutable ECameraContextDataType CachedDataType;
+	mutable const UObject* CachedDataTypeObject = nullptr;
+	mutable const FSlateBrush* CachedDataTypeImage = nullptr;
+	mutable FSlateColor CachedDataTypeImageColor;
+};
+
+/**
+ * The overall interface parameters panel, showing two sub-panels, one for blendable parameters, and
+ * one for data parameters.
+ */
+class SCameraRigInterfaceParametersPanel 
+	: public SCompoundWidget
+	, public ICameraRigAssetEventHandler
+{
+public:
+
+	SLATE_BEGIN_ARGS(SCameraRigInterfaceParametersPanel)
+	{}
+	SLATE_END_ARGS()
+
+	void Construct(const FArguments& Args, FCameraRigInterfaceParametersToolkit* OwnerToolkit);
+
+	void RequestListRefresh();
+
+protected:
+
+	// SWidget interface.
+	virtual void Tick(const FGeometry& AllottedGeometry, const double InCurrentTime, const float InDeltaTime) override;
+
+	// ICameraRigAssetEventHandler interface.
+	virtual void OnCameraRigInterfaceChanged() override;
+
+private:
+
+	TSharedRef<ITableRow> OnGenerateBlendableParameterRow(TObjectPtr<UCameraRigBlendableParameter> Item, const TSharedRef<STableViewBase>& OwnerTable);
+	TSharedRef<ITableRow> OnGenerateDataParameterRow(TObjectPtr<UCameraRigDataParameter> Item, const TSharedRef<STableViewBase>& OwnerTable);
+
+	void OnBlendableSelectionChanged(TObjectPtr<UCameraRigBlendableParameter> Item, ESelectInfo::Type Type);
+	void OnDataParameterSelectionChanged(TObjectPtr<UCameraRigDataParameter> Item, ESelectInfo::Type Type);
+
+	TSharedPtr<SWidget> OnBlendableParameterContextMenuOpening();
+	TSharedPtr<SWidget> OnDataParameterContextMenuOpening();
+
+	template<typename ItemType>
+	TSharedPtr<SWidget> OnInterfaceParameterContextMenuOpening(
+			TSharedPtr<SListView<TObjectPtr<ItemType>>> ListView, TObjectPtr<ItemType> Item);
+	template<typename ItemType>
+	void OnRenameInterfaceParameter(
+			TSharedPtr<SListView<TObjectPtr<ItemType>>> ListView, TObjectPtr<ItemType> Item);
+	template<typename ItemType>
+	void OnDeleteInterfaceParameter(
+			TSharedPtr<SListView<TObjectPtr<ItemType>>> ListView, TObjectPtr<ItemType> Item);
+
+	FReply OnAddBlendableParameter();
+	FReply OnAddDataParameter();
+
+private:
+
+	UCameraRigAsset* CameraRigAsset = nullptr;
+	FCameraRigInterfaceParametersToolkit* Toolkit = nullptr;
+
+	TCameraEventHandler<ICameraRigAssetEventHandler> EventHandler;
+
+	TSharedPtr<SListView<TObjectPtr<UCameraRigBlendableParameter>>> BlendableParametersListView;
+	TSharedPtr<SListView<TObjectPtr<UCameraRigDataParameter>>> DataParametersListView;
+
+	bool bListRefreshRequested = false;
+};
+
+void SCameraRigInterfaceParametersPanel::Construct(const FArguments& Args, FCameraRigInterfaceParametersToolkit* OwnerToolkit)
+{
+	CameraRigAsset = OwnerToolkit->GetCameraRigAsset();
+	CameraRigAsset->EventHandlers.Register(EventHandler, this);
+
+	Toolkit = OwnerToolkit;
+
+	ChildSlot
+	[
+		SNew(SSplitter)
+		.Orientation(Orient_Vertical)
+
+		// The blendable parameters panel.
+		+SSplitter::Slot()
+		[
+			SNew(SVerticalBox)
+			+SVerticalBox::Slot()
+			.AutoHeight()
+			[
+				SNew(SBorder)
+				.BorderImage(FAppStyle::GetBrush("ToolPanel.GroupBorder"))
+				.Padding(FMargin(3, 5))
+				[
+					SNew(SHorizontalBox)
+					+SHorizontalBox::Slot()
+					.FillWidth(1.f)
+					.VAlign(VAlign_Center)
+					.Padding(5.f)
+					[
+						SNew(SRichTextBlock)
+						.Text(LOCTEXT("BlendableParametersPanelTitle", "Blendable Parameters"))
+						.TransformPolicy(ETextTransformPolicy::ToUpper)
+						.DecoratorStyleSet(&FAppStyle::Get())
+						.TextStyle(FAppStyle::Get(), "DetailsView.CategoryTextStyle")
+					]
+					+SHorizontalBox::Slot()
+					.AutoWidth()
+					.VAlign(VAlign_Center)
+					[
+						SNew(SButton)
+						.ButtonStyle(FAppStyle::Get(), "SimpleButton")
+						.ContentPadding(FMargin(1, 0))
+						.ToolTipText(LOCTEXT("AddBlendableToolTip", "Add a blendable parameter"))
+						.OnClicked(this, &SCameraRigInterfaceParametersPanel::OnAddBlendableParameter)
+						[
+							SNew(SImage)
+							.Image(FAppStyle::Get().GetBrush("Icons.PlusCircle"))
+							.ColorAndOpacity(FSlateColor::UseForeground())
+						]
+					]
+				]
+			]
+			+SVerticalBox::Slot()
+			.FillHeight(1.f)
+			[
+				SAssignNew(BlendableParametersListView, SListView<TObjectPtr<UCameraRigBlendableParameter>>)
+				.ListItemsSource(&CameraRigAsset->Interface.BlendableParameters)
+				.OnGenerateRow(this, &SCameraRigInterfaceParametersPanel::OnGenerateBlendableParameterRow)
+				.OnSelectionChanged(this, &SCameraRigInterfaceParametersPanel::OnBlendableSelectionChanged)
+				.OnContextMenuOpening(this, &SCameraRigInterfaceParametersPanel::OnBlendableParameterContextMenuOpening)
+				.HeaderRow
+				(
+					SNew(SHeaderRow)
+
+					+SHeaderRow::Column(ParameterTypeColumn)
+					.FillWidth(0.3f)
+					.DefaultLabel(LOCTEXT("ParameterTypeColumnLabel", "Type"))
+
+					+SHeaderRow::Column(ParameterNameColumn)
+					.FillWidth(0.5f)
+					.DefaultLabel(LOCTEXT("ParameterNameColumnLabel", "Name"))
+
+					+SHeaderRow::Column(ParameterMessageColumn)
+					.FillWidth(0.2f)
+					.DefaultLabel(LOCTEXT("ParameterMessageColumnLabel", "Note"))
+				)
+			]
+		]
+
+		// The data parameters panel.
+		+SSplitter::Slot()
+		[
+			SNew(SVerticalBox)
+			+SVerticalBox::Slot()
+			.AutoHeight()
+			[
+				SNew(SBorder)
+				.BorderImage(FAppStyle::GetBrush("ToolPanel.GroupBorder"))
+				.Padding(FMargin(3, 5))
+				[
+					SNew(SHorizontalBox)
+					+SHorizontalBox::Slot()
+					.FillWidth(1.f)
+					.VAlign(VAlign_Center)
+					.Padding(5.f)
+					[
+						SNew(SRichTextBlock)
+						.Text(LOCTEXT("DataParametersPanelTitle", "Data Parameters"))
+						.TransformPolicy(ETextTransformPolicy::ToUpper)
+						.DecoratorStyleSet(&FAppStyle::Get())
+						.TextStyle(FAppStyle::Get(), "DetailsView.CategoryTextStyle")
+					]
+					+SHorizontalBox::Slot()
+					.AutoWidth()
+					.VAlign(VAlign_Center)
+					[
+						SNew(SButton)
+						.ButtonStyle(FAppStyle::Get(), "SimpleButton")
+						.ContentPadding(FMargin(1, 0))
+						.ToolTipText(LOCTEXT("AddDataParameterToolTip", "Add a data parameter"))
+						.OnClicked(this, &SCameraRigInterfaceParametersPanel::OnAddDataParameter)
+						[
+							SNew(SImage)
+							.Image(FAppStyle::Get().GetBrush("Icons.PlusCircle"))
+							.ColorAndOpacity(FSlateColor::UseForeground())
+						]
+					]
+				]
+			]
+			+SVerticalBox::Slot()
+			.FillHeight(1.f)
+			[
+				SAssignNew(DataParametersListView, SListView<TObjectPtr<UCameraRigDataParameter>>)
+				.ListItemsSource(&CameraRigAsset->Interface.DataParameters)
+				.OnGenerateRow(this, &SCameraRigInterfaceParametersPanel::OnGenerateDataParameterRow)
+				.OnSelectionChanged(this, &SCameraRigInterfaceParametersPanel::OnDataParameterSelectionChanged)
+				.OnContextMenuOpening(this, &SCameraRigInterfaceParametersPanel::OnDataParameterContextMenuOpening)
+				.HeaderRow
+				(
+					SNew(SHeaderRow)
+
+					+SHeaderRow::Column(ParameterTypeColumn)
+					.FillWidth(0.3f)
+					.DefaultLabel(LOCTEXT("ParameterTypeColumnLabel", "Type"))
+
+					+SHeaderRow::Column(ParameterNameColumn)
+					.FillWidth(0.5f)
+					.DefaultLabel(LOCTEXT("ParameterNameColumnLabel", "Name"))
+
+					+SHeaderRow::Column(ParameterMessageColumn)
+					.FillWidth(0.2f)
+					.DefaultLabel(LOCTEXT("ParameterMessageColumnLabel", "Note"))
+				)
+			]
+		]
+	];
+}
+
+void SCameraRigInterfaceParametersPanel::RequestListRefresh()
+{
+	bListRefreshRequested = true;
+}
+
+void SCameraRigInterfaceParametersPanel::Tick(const FGeometry& AllottedGeometry, const double InCurrentTime, const float InDeltaTime)
+{
+	SCompoundWidget::Tick(AllottedGeometry, InCurrentTime, InDeltaTime);
+
+	if (bListRefreshRequested)
+	{
+		bListRefreshRequested = false;
+
+		BlendableParametersListView->RequestListRefresh();
+		DataParametersListView->RequestListRefresh();
+	}
+}
+
+TSharedRef<ITableRow> SCameraRigInterfaceParametersPanel::OnGenerateBlendableParameterRow(TObjectPtr<UCameraRigBlendableParameter> Item, const TSharedRef<STableViewBase>& OwnerTable)
+{
+	return SNew(SCameraRigBlendableParameterTableRow, OwnerTable)
+		.Item(Item);
+}
+
+TSharedRef<ITableRow> SCameraRigInterfaceParametersPanel::OnGenerateDataParameterRow(TObjectPtr<UCameraRigDataParameter> Item, const TSharedRef<STableViewBase>& OwnerTable)
+{
+	return SNew(SCameraRigDataParameterTableRow, OwnerTable)
+		.Item(Item);
+}
+
+void SCameraRigInterfaceParametersPanel::OnBlendableSelectionChanged(TObjectPtr<UCameraRigBlendableParameter> Item, ESelectInfo::Type Type)
+{
+	Toolkit->OnInterfaceParameterSelected().Broadcast(Item);
+}
+
+void SCameraRigInterfaceParametersPanel::OnDataParameterSelectionChanged(TObjectPtr<UCameraRigDataParameter> Item, ESelectInfo::Type Type)
+{
+	Toolkit->OnInterfaceParameterSelected().Broadcast(Item);
+}
+
+TSharedPtr<SWidget> SCameraRigInterfaceParametersPanel::OnBlendableParameterContextMenuOpening()
+{
+	TArray<TObjectPtr<UCameraRigBlendableParameter>> SelectedItems = BlendableParametersListView->GetSelectedItems();
+	if (SelectedItems.Num() > 0)
+	{
+		return OnInterfaceParameterContextMenuOpening(BlendableParametersListView, SelectedItems[0]);
+	}
+	return nullptr;
+}
+
+TSharedPtr<SWidget> SCameraRigInterfaceParametersPanel::OnDataParameterContextMenuOpening()
+{
+	TArray<TObjectPtr<UCameraRigDataParameter>> SelectedItems = DataParametersListView->GetSelectedItems();
+	if (SelectedItems.Num() > 0)
+	{
+		return OnInterfaceParameterContextMenuOpening(DataParametersListView, SelectedItems[0]);
+	}
+	return nullptr;
+}
+
+template<typename ItemType>
+TSharedPtr<SWidget> SCameraRigInterfaceParametersPanel::OnInterfaceParameterContextMenuOpening(
+		TSharedPtr<SListView<TObjectPtr<ItemType>>> ListView, TObjectPtr<ItemType> Item)
+{
+	FMenuBuilder MenuBuilder(true, nullptr);
+
+	MenuBuilder.AddMenuEntry(
+			LOCTEXT("RenameInterfaceParameter", "Rename"),
+			LOCTEXT("RenameInterfaceParameterToolTip", "Renames this interface parameter"),
+			FSlateIcon(),
+			FExecuteAction::CreateSP(this, &SCameraRigInterfaceParametersPanel::OnRenameInterfaceParameter, ListView, Item));
+	MenuBuilder.AddMenuEntry(
+			LOCTEXT("DeleteInterfaceParameter", "Delete"),
+			LOCTEXT("DeleteInterfaceParameterToolTip", "Deletes this interface parameter"),
+			FSlateIcon(),
+			FExecuteAction::CreateSP(this, &SCameraRigInterfaceParametersPanel::OnDeleteInterfaceParameter, ListView, Item));
+
+	return MenuBuilder.MakeWidget();
+}
+
+template<typename ItemType>
+void SCameraRigInterfaceParametersPanel::OnRenameInterfaceParameter(
+		TSharedPtr<SListView<TObjectPtr<ItemType>>> ListView, TObjectPtr<ItemType> Item)
+{
+	TSharedPtr<ITableRow> RowWidget = ListView->WidgetFromItem(Item);
+	if (!ensure(RowWidget))
+	{
+		return;
+	}
+
+	TSharedPtr<SCameraRigInterfaceParameterTableRowBase<ItemType>> TypedRowWidget = 
+		StaticCastSharedPtr<SCameraRigInterfaceParameterTableRowBase<ItemType>>(RowWidget);
+	TypedRowWidget->EnterNameEditingMode();
+}
+
+template<typename ItemType>
+void SCameraRigInterfaceParametersPanel::OnDeleteInterfaceParameter(
+		TSharedPtr<SListView<TObjectPtr<ItemType>>> ListView, TObjectPtr<ItemType> Item)
+{
+	const FScopedTransaction Transaction(LOCTEXT("RemoveInterfaceParameter", "Remove Interface Parameter"));
+
+	CameraRigAsset->Modify();
+
+	if constexpr (std::is_same_v<ItemType, UCameraRigBlendableParameter>)
+	{
+		const int32 NumRemoved = CameraRigAsset->Interface.BlendableParameters.Remove(Item);
+		ensure(NumRemoved == 1);
+	}
+	else if constexpr (std::is_same_v<ItemType, UCameraRigDataParameter>)
+	{
+		const int32 NumRemoved = CameraRigAsset->Interface.DataParameters.Remove(Item);
+		ensure(NumRemoved == 1);
+	}
+
+	CameraRigAsset->EventHandlers.Notify(&ICameraRigAssetEventHandler::OnCameraRigInterfaceChanged);
+
+	ListView->RequestListRefresh();
+}
+
+FReply SCameraRigInterfaceParametersPanel::OnAddBlendableParameter()
+{
+	const FScopedTransaction Transaction(LOCTEXT("AddBlendableParameter", "Add Blendable Parameter"));
+
+	UCameraRigBlendableParameter* NewBlendableParameter = NewObject<UCameraRigBlendableParameter>(CameraRigAsset, NAME_None, RF_Transactional);
+	NewBlendableParameter->InterfaceParameterName = NewBlendableParameter->GetName();
+
+	CameraRigAsset->Modify();
+	CameraRigAsset->Interface.BlendableParameters.Add(NewBlendableParameter);
+	CameraRigAsset->EventHandlers.Notify(&ICameraRigAssetEventHandler::OnCameraRigInterfaceChanged);
+
+	return FReply::Handled();
+}
+
+FReply SCameraRigInterfaceParametersPanel::OnAddDataParameter()
+{
+	const FScopedTransaction Transaction(LOCTEXT("AddDataParameter", "Add Data Parameter"));
+
+	UCameraRigDataParameter* NewDataParameter = NewObject<UCameraRigDataParameter>(CameraRigAsset, NAME_None, RF_Transactional);
+	NewDataParameter->InterfaceParameterName = NewDataParameter->GetName();
+
+	CameraRigAsset->Modify();
+	CameraRigAsset->Interface.DataParameters.Add(NewDataParameter);
+	CameraRigAsset->EventHandlers.Notify(&ICameraRigAssetEventHandler::OnCameraRigInterfaceChanged);
+
+	return FReply::Handled();
+}
+
+void SCameraRigInterfaceParametersPanel::OnCameraRigInterfaceChanged()
+{
+	bListRefreshRequested = true;
+}
+
+FCameraRigInterfaceParametersToolkit::FCameraRigInterfaceParametersToolkit()
+{
+	SAssignNew(PanelContainer, SBox);
+
+	if (GEditor)
+	{
+		GEditor->RegisterForUndo(this);
+	}
+}
+
+FCameraRigInterfaceParametersToolkit::~FCameraRigInterfaceParametersToolkit()
+{
+	if (GEditor)
+	{
+		GEditor->UnregisterForUndo(this);
+	}
+}
+
+void FCameraRigInterfaceParametersToolkit::SetCameraRigAsset(UCameraRigAsset* InCameraRigAsset)
+{
+	if (CameraRigAsset != InCameraRigAsset)
+	{
+		PanelContainer->SetContent(SNullWidget::NullWidget);
+
+		CameraRigAsset = InCameraRigAsset;
+
+		if (CameraRigAsset)
+		{
+			Panel = SNew(SCameraRigInterfaceParametersPanel, this);
+			PanelContainer->SetContent(Panel.ToSharedRef());
+		}
+	}
+}
+
+TSharedPtr<SWidget> FCameraRigInterfaceParametersToolkit::GetInterfaceParametersPanel() const
+{
+	return PanelContainer;
+}
+
+void FCameraRigInterfaceParametersToolkit::PostUndo(bool bSuccess)
+{
+	Panel->RequestListRefresh();
+}
+
+void FCameraRigInterfaceParametersToolkit::PostRedo(bool bSuccess)
+{
+	Panel->RequestListRefresh();
+}
+
+}  // namespace UE::Cameras
+
+#undef LOCTEXT_NAMESPACE
+
diff --git a/Engine/Plugins/Cameras/GameplayCameras/Source/GameplayCamerasEditor/Public/Customizations/CameraRigAssetReferenceDetailsCustomization.h b/Engine/Plugins/Cameras/GameplayCameras/Source/GameplayCamerasEditor/Public/Customizations/CameraRigAssetReferenceDetailsCustomization.h
--- a/Engine/Plugins/Cameras/GameplayCameras/Source/GameplayCamerasEditor/Public/Customizations/CameraRigAssetReferenceDetailsCustomization.h
+++ b/Engine/Plugins/Cameras/GameplayCameras/Source/GameplayCamerasEditor/Public/Customizations/CameraRigAssetReferenceDetailsCustomization.h
@@ -4,18 +4,13 @@
 
 #include "IPropertyTypeCustomization.h"
 
-class FStructOnScope;
-class IStructureDataProvider;
+class IPropertyHandle;
+class IPropertyUtilities;
 class UCameraRigAsset;
-class UCameraRigInterfaceParameter;
-struct FCameraRigAssetReference;
 
 namespace UE::Cameras
 {
 
-class FCameraBuildLog;
-class FCameraRigParameterOverrideDetailRow;
-
 class FCameraRigAssetReferenceDetailsCustomization : public IPropertyTypeCustomization
 {
 public:
@@ -24,25 +19,24 @@ public:
 
 public:
 
+	~FCameraRigAssetReferenceDetailsCustomization();
+
 	// IPropertyTypeCustomization interface.
-	virtual void CustomizeHeader(TSharedRef<IPropertyHandle> StructPropertyHandle, FDetailWidgetRow& HeaderRow, IPropertyTypeCustomizationUtils& StructCustomizationUtils) override;
-	virtual void CustomizeChildren(TSharedRef<IPropertyHandle> StructPropertyHandle, IDetailChildrenBuilder& StructBuilder, IPropertyTypeCustomizationUtils& StructCustomizationUtils) override;
+	virtual void CustomizeHeader(TSharedRef<IPropertyHandle> InStructPropertyHandle, FDetailWidgetRow& InHeaderRow, IPropertyTypeCustomizationUtils& InCustomizationUtils) override;
+	virtual void CustomizeChildren(TSharedRef<IPropertyHandle> InStructPropertyHandle, IDetailChildrenBuilder& InChildrenBuilder, IPropertyTypeCustomizationUtils& InCustomizationUtils) override;
 
 private:
 
-	void OnCameraRigChanged();
-	void OnCameraRigBuilt(UCameraRigAsset* CameraRig, FCameraBuildLog& BuildLog);
-
-	void UpdateParameterOverrides(const UCameraRigAsset* CameraRigToUpdate, bool bRequestRefresh);
-
-	void BuildParameterOverrideRows(IDetailChildrenBuilder& StructBuilder);
+	void OnCameraRigAssetBuilt(const UCameraRigAsset* CameraRig);
+	void RebuildParametersIfNeeded();
 
 private:
 
-	TSharedPtr<IPropertyHandle> CameraRigReferenceProperty;
+	TSharedPtr<IPropertyHandle> StructPropertyHandle;
+	TSharedPtr<IPropertyHandle> CameraRigAssetPropertyHandle;
+	TSharedPtr<IPropertyHandle> ParametersPropertyHandle;
 	TSharedPtr<IPropertyUtilities> PropertyUtilities;
-
-	TArray<TSharedPtr<FCameraRigParameterOverrideDetailRow>> ParameterOverrideRows;
 };
 
 }  // namespace UE::Cameras
+
diff --git a/Engine/Plugins/Cameras/GameplayCameras/Source/GameplayCamerasEditor/Public/Editors/CameraNodeGraphDragDropOp.h b/Engine/Plugins/Cameras/GameplayCameras/Source/GameplayCamerasEditor/Public/Editors/CameraNodeGraphDragDropOp.h
new file mode 100644
--- /dev/null
+++ b/Engine/Plugins/Cameras/GameplayCameras/Source/GameplayCamerasEditor/Public/Editors/CameraNodeGraphDragDropOp.h
@@ -0,0 +1,25 @@
+// Copyright Epic Games, Inc. All Rights Reserved.
+
+#pragma once
+
+#include "DragAndDrop/DecoratedDragDropOp.h"
+
+class SGraphEditor;
+class UCameraRigInterfaceParameterBase;
+
+class FCameraNodeGraphInterfaceParameterDragDropOp : public FDecoratedDragDropOp
+{
+public:
+	
+	DRAG_DROP_OPERATOR_TYPE(FCameraNodeGraphInterfaceParameterDragDropOp, FDecoratedDragDropOp)
+
+	static TSharedRef<FCameraNodeGraphInterfaceParameterDragDropOp> New(UCameraRigInterfaceParameterBase* InInterfaceParameter);
+
+	FReply ExecuteDragOver(TSharedPtr<SGraphEditor> GraphEditor);
+	FReply ExecuteDrop(TSharedPtr<SGraphEditor> GraphEditor, const FVector2D& NewLocation);
+
+private:
+
+	UCameraRigInterfaceParameterBase* InterfaceParameter;
+};
+
diff --git a/Engine/Plugins/Cameras/GameplayCameras/Source/GameplayCamerasEditor/Public/Editors/CameraNodeGraphNode.h b/Engine/Plugins/Cameras/GameplayCameras/Source/GameplayCamerasEditor/Public/Editors/CameraNodeGraphNode.h
--- a/Engine/Plugins/Cameras/GameplayCameras/Source/GameplayCamerasEditor/Public/Editors/CameraNodeGraphNode.h
+++ b/Engine/Plugins/Cameras/GameplayCameras/Source/GameplayCamerasEditor/Public/Editors/CameraNodeGraphNode.h
@@ -20,11 +20,6 @@ public:
 	/** Creates a new graph node. */
 	UCameraNodeGraphNode(const FObjectInitializer& ObjInit);
 
-	/** Gets input pin for given camera parameter property on the underlying camera node. */
-	UEdGraphPin* GetPinForCameraParameterProperty(const FName& InPropertyName) const;
-	/** Gets camera parameter property name on the underlying camera node for the given input pin. */
-	FName GetCameraParameterPropertyForPin(const UEdGraphPin* InPin) const;
-
 public:
 
 	// UEdGraphNode interface
diff --git a/Engine/Plugins/Cameras/GameplayCameras/Source/GameplayCamerasEditor/Public/Editors/CameraNodeGraphSchema.h b/Engine/Plugins/Cameras/GameplayCameras/Source/GameplayCamerasEditor/Public/Editors/CameraNodeGraphSchema.h
--- a/Engine/Plugins/Cameras/GameplayCameras/Source/GameplayCamerasEditor/Public/Editors/CameraNodeGraphSchema.h
+++ b/Engine/Plugins/Cameras/GameplayCameras/Source/GameplayCamerasEditor/Public/Editors/CameraNodeGraphSchema.h
@@ -2,11 +2,17 @@
 
 #pragma once
 
+#include "Core/CameraContextDataTableFwd.h"
+#include "Core/CameraVariableTableFwd.h"
+#include "EdGraph/EdGraphPin.h"
+#include "Editors/CameraNodeGraphPinColors.h"
 #include "Editors/ObjectTreeGraphSchema.h"
 
 #include "CameraNodeGraphSchema.generated.h"
 
 class UCameraNode;
+class UCameraRigInterfaceParameterBase;
+class UCameraRigInterfaceParameterGraphNode;
 struct FObjectTreeGraphConfig;
 
 /**
@@ -20,23 +26,42 @@ class UCameraNodeGraphSchema : public UObjectTreeGraphSchema
 public:
 
 	static const FName PC_CameraParameter;			// A camera parameter pin.
+	static const FName PC_CameraVariableReference;	// A variable reference pin.
+	static const FName PC_CameraContextData;		// A context data pin.
+
+	UCameraNodeGraphSchema(const FObjectInitializer& ObjInit);
 
 	FObjectTreeGraphConfig BuildGraphConfig() const;
 
+	UCameraRigInterfaceParameterGraphNode* CreateInterfaceParameterNode(UEdGraph* InGraph, UCameraRigInterfaceParameterBase* InterfaceParameter) const;
+
 protected:
 
 	// UEdGraphSchema interface.
 	virtual void GetGraphContextActions(FGraphContextMenuBuilder& ContextMenuBuilder) const override;
 	virtual const FPinConnectionResponse CanCreateConnection(const UEdGraphPin* A, const UEdGraphPin* B) const override;
+	virtual bool OnTryCreateCustomConnection(UEdGraphPin* A, UEdGraphPin* B) const;
+	virtual bool OnBreakCustomPinLinks(UEdGraphPin& TargetPin) const;
+	virtual bool OnBreakSingleCustomPinLink(UEdGraphPin* SourcePin, UEdGraphPin* TargetPin) const;
+	virtual FLinearColor GetPinTypeColor(const FEdGraphPinType& PinType) const override;
+	virtual bool SafeDeleteNodeFromGraph(UEdGraph* Graph, UEdGraphNode* Node) const override;
 
 	// UObjectTreeGraphSchema interface.
 	virtual void CollectAllObjects(UObjectTreeGraph* InGraph, TSet<UObject*>& OutAllObjects) const override;
 	virtual void OnCreateAllNodes(UObjectTreeGraph* InGraph, const FCreatedNodes& InCreatedNodes) const override;
-	virtual void OnAddConnectableObject(UObjectTreeGraph* InGraph, UObjectTreeGraphNode* InNewNode) const override;
-	virtual void OnRemoveConnectableObject(UObjectTreeGraph* InGraph, UObjectTreeGraphNode* InRemovedNode) const override;
-	virtual bool OnApplyConnection(UEdGraphPin* A, UEdGraphPin* B) const override;
-	virtual bool OnApplyDisconnection(UEdGraphPin* TargetPin) const override;
-	virtual bool OnApplyDisconnection(UEdGraphPin* SourcePin, UEdGraphPin* TargetPin) const override;
+
+private:
+
+	UEdGraphPin* FindPin(UEdGraphNode* InNode, const FName& InPinName, const FName& InPinCategoryName) const;
+
+	UE::Cameras::FCameraNodeGraphPinColors PinColors;
+};
+
+UENUM()
+enum class EInterfaceParameterCreateNodeType
+{
+	BlendableParameter,
+	DataParameter
 };
 
 /**
@@ -57,6 +82,18 @@ public:
 	UPROPERTY()
 	FName TargetPropertyName;
 
+	UPROPERTY()
+	EInterfaceParameterCreateNodeType NewNodeType = EInterfaceParameterCreateNodeType::BlendableParameter;
+
+	UPROPERTY()
+	ECameraVariableType BlendableParameterType;
+
+	UPROPERTY()
+	ECameraContextDataType DataParameterType;
+
+	UPROPERTY()
+	TObjectPtr<const UObject> DataParameterTypeObject;
+
 public:
 
 	FCameraNodeGraphSchemaAction_NewInterfaceParameterNode();
@@ -70,3 +107,24 @@ public:
 	virtual UEdGraphNode* PerformAction(UEdGraph* ParentGraph, UEdGraphPin* FromPin, const FVector2D Location, bool bSelectNewNode = true) override;
 };
 
+USTRUCT()
+struct FCameraNodeGraphSchemaAction_AddInterfaceParameterNode : public FEdGraphSchemaAction
+{
+	GENERATED_BODY()
+	
+	UPROPERTY()
+	TObjectPtr<UCameraRigInterfaceParameterBase> InterfaceParameter;
+
+public:
+	
+	FCameraNodeGraphSchemaAction_AddInterfaceParameterNode();
+	FCameraNodeGraphSchemaAction_AddInterfaceParameterNode(FText InNodeCategory, FText InMenuDesc, FText InToolTip, const int32 InGrouping = 0, FText InKeywords = FText());
+
+public:
+
+	// FEdGraphSchemaAction interface.
+	static FName StaticGetTypeId() { static FName Type("FCameraNodeGraphSchemaAction_AddInterfaceParameterNode"); return Type; }
+	virtual FName GetTypeId() const override { return StaticGetTypeId(); } 
+	virtual UEdGraphNode* PerformAction(UEdGraph* ParentGraph, UEdGraphPin* FromPin, const FVector2D Location, bool bSelectNewNode = true) override;
+};
+
diff --git a/Engine/Plugins/Cameras/GameplayCameras/Source/GameplayCamerasEditor/Public/Editors/CameraRigInterfaceParameterGraphNode.h b/Engine/Plugins/Cameras/GameplayCameras/Source/GameplayCamerasEditor/Public/Editors/CameraRigInterfaceParameterGraphNode.h
--- a/Engine/Plugins/Cameras/GameplayCameras/Source/GameplayCamerasEditor/Public/Editors/CameraRigInterfaceParameterGraphNode.h
+++ b/Engine/Plugins/Cameras/GameplayCameras/Source/GameplayCamerasEditor/Public/Editors/CameraRigInterfaceParameterGraphNode.h
@@ -6,6 +6,8 @@
 
 #include "CameraRigInterfaceParameterGraphNode.generated.h"
 
+class UCameraRigInterfaceParameterBase;
+
 /**
  * Custom graph editor node for a camera rig parameter.
  */
@@ -19,6 +21,9 @@ public:
 	/** Creates a new graph node. */
 	UCameraRigInterfaceParameterGraphNode(const FObjectInitializer& ObjInit);
 
+	/** Gets the underlying object as a camera rig interface parameter. */
+	UCameraRigInterfaceParameterBase* GetInterfaceParameter() const;
+
 public:
 
 	// UEdGraphNode interface
diff --git a/Engine/Plugins/Cameras/GameplayCameras/Source/GameplayCamerasEditor/Public/Editors/SCameraRigAssetEditor.h b/Engine/Plugins/Cameras/GameplayCameras/Source/GameplayCamerasEditor/Public/Editors/SCameraRigAssetEditor.h
--- a/Engine/Plugins/Cameras/GameplayCameras/Source/GameplayCamerasEditor/Public/Editors/SCameraRigAssetEditor.h
+++ b/Engine/Plugins/Cameras/GameplayCameras/Source/GameplayCamerasEditor/Public/Editors/SCameraRigAssetEditor.h
@@ -2,6 +2,7 @@
 
 #pragma once
 
+#include "Core/CameraRigAsset.h"
 #include "CoreTypes.h"
 #include "GraphEditor.h"
 #include "Templates/SharedPointerFwd.h"
@@ -12,7 +13,6 @@ class FAssetEditorToolkit;
 class IDetailsView;
 class SBox;
 class SObjectTreeGraphEditor;
-class UCameraRigAsset;
 class UEdGraphNode;
 class UObjectTreeGraph;
 struct FObjectTreeGraphConfig;
@@ -35,7 +35,9 @@ enum class ECameraRigAssetEditorMode
  * This implements only the dual-graph editor, for the node hierarchy and transitions.
  * The rest of the UI such as the details view or the toolbox aren't included here.
  */
-class SCameraRigAssetEditor : public SCompoundWidget
+class SCameraRigAssetEditor 
+	: public SCompoundWidget
+	, public ICameraRigAssetEventHandler
 {
 public:
 
@@ -83,6 +85,11 @@ public:
 	/** Removes a previous added callback. */
 	void RemoveOnAnyGraphChanged(FDelegateUserObjectConst InUserObject);
 
+protected:
+
+	// ICameraRigAssetEventHandler interface.
+	virtual void OnCameraRigInterfaceChanged() override;
+
 protected:
 
 	void CreateGraphEditors();
@@ -102,8 +109,13 @@ private:
 	/** The asset being edited */
 	TObjectPtr<UCameraRigAsset> CameraRigAsset;
 
+	/** Event handler */
+	TCameraEventHandler<ICameraRigAssetEventHandler> EventHandler;
+
+	/** Reference to the details view */
 	TSharedPtr<IDetailsView> DetailsView;
 
+	/** Reference to the owning asset editor */
 	TWeakPtr<FAssetEditorToolkit> AssetEditorToolkit;
 
 	/** The node hierarchy graph */
diff --git a/Engine/Plugins/Cameras/GameplayCameras/Source/GameplayCamerasEditor/Public/Editors/SCameraRigInterfaceParameterGraphNode.h b/Engine/Plugins/Cameras/GameplayCameras/Source/GameplayCamerasEditor/Public/Editors/SCameraRigInterfaceParameterGraphNode.h
--- a/Engine/Plugins/Cameras/GameplayCameras/Source/GameplayCamerasEditor/Public/Editors/SCameraRigInterfaceParameterGraphNode.h
+++ b/Engine/Plugins/Cameras/GameplayCameras/Source/GameplayCamerasEditor/Public/Editors/SCameraRigInterfaceParameterGraphNode.h
@@ -4,6 +4,8 @@
 
 #include "Editors/SObjectTreeGraphNode.h"
 
+class UCameraRigInterfaceParameterGraphNode;
+
 /**
  * Custom graph editor node widget for a camera rig parameter node.
  */
@@ -14,7 +16,7 @@ public:
 	SLATE_BEGIN_ARGS(SCameraRigInterfaceParameterGraphNode)
 		: _GraphNode(nullptr)
 	{}
-		SLATE_ARGUMENT(UObjectTreeGraphNode*, GraphNode)
+		SLATE_ARGUMENT(UCameraRigInterfaceParameterGraphNode*, GraphNode)
 	SLATE_END_ARGS()
 
 	void Construct(const FArguments& InArgs);
@@ -27,7 +29,7 @@ protected:
 	virtual TSharedPtr<SGraphPin> CreatePinWidget(UEdGraphPin* InPin) const override;
 
 private:
-	
+
 	FText GetInterfaceParameterName() const;
 };
 
diff --git a/Engine/Plugins/Cameras/GameplayCameras/Source/GameplayCamerasEditor/Public/IGameplayCamerasEditorModule.h b/Engine/Plugins/Cameras/GameplayCameras/Source/GameplayCamerasEditor/Public/IGameplayCamerasEditorModule.h
--- a/Engine/Plugins/Cameras/GameplayCameras/Source/GameplayCamerasEditor/Public/IGameplayCamerasEditorModule.h
+++ b/Engine/Plugins/Cameras/GameplayCameras/Source/GameplayCamerasEditor/Public/IGameplayCamerasEditorModule.h
@@ -36,8 +36,6 @@ struct FCameraDebugCategoryInfo
 }  // namespace UE::Cameras
 
 DECLARE_DELEGATE_RetVal_OneParam(TSharedPtr<UE::Cameras::FCameraDirectorAssetEditorMode>, FOnCreateCameraDirectorAssetEditorMode, UCameraAsset*);
-DECLARE_DELEGATE_TwoParams(FOnBuildCameraAsset, UCameraAsset*, UE::Cameras::FCameraBuildLog&);
-DECLARE_DELEGATE_TwoParams(FOnBuildCameraRigAsset, UCameraRigAsset*, UE::Cameras::FCameraBuildLog&);
 DECLARE_DELEGATE_RetVal_OneParam(TSharedRef<SWidget>, FOnCreateDebugCategoryPanel, const FString&);
 
 /**
@@ -82,20 +80,6 @@ public:
 	/** Unregisters a camera director editor creator. */
 	virtual void UnregisterCameraDirectorEditor(FDelegateHandle InHandle) = 0;
 
-	/** Registers a custom camera asset builder. */
-	virtual FDelegateHandle RegisterCameraAssetBuilder(FOnBuildCameraAsset InOnBuildCameraAsset) = 0;
-	/** Gets the registered custom camera asset builders. */
-	virtual TArrayView<const FOnBuildCameraAsset> GetCameraAssetBuilders() const = 0;
-	/** Unregisters a custom camera asset builder. */
-	virtual void UnregisterCameraAssetBuilder(FDelegateHandle InHandle) = 0;
-
-	/** Registers a custom camera rig builder. */
-	virtual FDelegateHandle RegisterCameraRigAssetBuilder(FOnBuildCameraRigAsset InOnBuildCameraRigAsset) = 0;
-	/** Gets the registered custom camera rig builders. */
-	virtual TArrayView<const FOnBuildCameraRigAsset> GetCameraRigAssetBuilders() const = 0;
-	/** Unregisters a custom camera rig builder. */
-	virtual void UnregisterCameraRigAssetBuilder(FDelegateHandle InHandle) = 0;
-
 public:
 
 	using FCameraDebugCategoryInfo = UE::Cameras::FCameraDebugCategoryInfo;
diff --git a/Engine/Plugins/Cameras/GameplayCameras/Source/GameplayCamerasEditor/Public/Toolkits/CameraRigAssetEditorToolkit.h b/Engine/Plugins/Cameras/GameplayCameras/Source/GameplayCamerasEditor/Public/Toolkits/CameraRigAssetEditorToolkit.h
--- a/Engine/Plugins/Cameras/GameplayCameras/Source/GameplayCamerasEditor/Public/Toolkits/CameraRigAssetEditorToolkit.h
+++ b/Engine/Plugins/Cameras/GameplayCameras/Source/GameplayCamerasEditor/Public/Toolkits/CameraRigAssetEditorToolkit.h
@@ -20,6 +20,7 @@ namespace UE::Cameras
 class FBuildButtonToolkit;
 class FCameraBuildLogToolkit;
 class FCameraRigAssetEditorToolkitBase;
+class FCameraRigInterfaceParametersToolkit;
 class FCurveEditorToolkit;
 class IGameplayCamerasLiveEditManager;
 
@@ -60,9 +61,12 @@ private:
 	TSharedRef<SDockTab> SpawnTab_Search(const FSpawnTabArgs& Args);
 	TSharedRef<SDockTab> SpawnTab_Messages(const FSpawnTabArgs& Args);
 	TSharedRef<SDockTab> SpawnTab_Curves(const FSpawnTabArgs& Args);
+	TSharedRef<SDockTab> SpawnTab_InterfaceParameters(const FSpawnTabArgs& Args);
 
 	void OnCurvesTabClosed(TSharedRef<SDockTab> InTab);
 
+	void OnCameraRigInterfaceParameterSelected(UCameraRigInterfaceParameterBase* Object);
+
 	void OnBuild();
 	void OnFindInCameraRig();
 
@@ -76,6 +80,7 @@ private:
 	static const FName SearchTabId;
 	static const FName MessagesTabId;
 	static const FName CurvesTabId;
+	static const FName InterfaceParametersTabId;
 
 	/** Base implementation */
 	TSharedPtr<FCameraRigAssetEditorToolkitBase> Impl;
@@ -91,6 +96,8 @@ private:
 	TSharedPtr<FCameraBuildLogToolkit> BuildLogToolkit;
 	/** The curve editor */
 	TSharedPtr<FCurveEditorToolkit> CurveEditorToolkit;
+	/** The interface parameters panel */
+	TSharedPtr<FCameraRigInterfaceParametersToolkit> InterfaceParametersToolkit;
 
 	/** Search widget */
 	TSharedPtr<SFindInObjectTreeGraph> SearchWidget;
diff --git a/Engine/Plugins/Cameras/GameplayCameras/Source/GameplayCamerasEditor/Public/Toolkits/CameraRigInterfaceParametersToolkit.h b/Engine/Plugins/Cameras/GameplayCameras/Source/GameplayCamerasEditor/Public/Toolkits/CameraRigInterfaceParametersToolkit.h
new file mode 100644
--- /dev/null
+++ b/Engine/Plugins/Cameras/GameplayCameras/Source/GameplayCamerasEditor/Public/Toolkits/CameraRigInterfaceParametersToolkit.h
@@ -0,0 +1,61 @@
+// Copyright Epic Games, Inc. All Rights Reserved.
+
+#pragma once
+
+#include "EditorUndoClient.h"
+#include "Templates/SharedPointer.h"
+#include "UObject/ObjectPtr.h"
+
+class SBox;
+class SWidget;
+class UCameraRigAsset;
+class UCameraRigInterfaceParameterBase;
+
+namespace UE::Cameras
+{
+
+class SCameraRigInterfaceParametersPanel;
+
+DECLARE_MULTICAST_DELEGATE_OneParam(FOnCameraRigInterfaceParameterEvent, UCameraRigInterfaceParameterBase*);
+
+/**
+ * Utility toolkit for the "interface parameters" panel of the camera rig editor.
+ */
+class FCameraRigInterfaceParametersToolkit 
+	: public TSharedFromThis<FCameraRigInterfaceParametersToolkit>
+	, public FEditorUndoClient
+{
+public:
+
+	FCameraRigInterfaceParametersToolkit();
+	~FCameraRigInterfaceParametersToolkit();
+
+	/** Gets the camera rig asset to edit. */
+	UCameraRigAsset* GetCameraRigAsset() const { return CameraRigAsset; }
+	/** Sets the camera rig asset to edit. This re-creates the panel widget. */
+	void SetCameraRigAsset(UCameraRigAsset* InCameraRigAsset);
+
+	/** Gets the panel widget. */
+	TSharedPtr<SWidget> GetInterfaceParametersPanel() const;
+
+	/** Delegate invoked when a parmeter is selected in the panel. */
+	FOnCameraRigInterfaceParameterEvent& OnInterfaceParameterSelected() { return OnInterfaceParameterSelectedDelegate; }
+
+protected:
+
+	// FEditorUndoClient interface
+	virtual void PostUndo(bool bSuccess) override;
+	virtual void PostRedo(bool bSuccess) override;
+
+private:
+
+	TObjectPtr<UCameraRigAsset> CameraRigAsset;
+
+	FOnCameraRigInterfaceParameterEvent OnInterfaceParameterSelectedDelegate;
+
+	TSharedPtr<SBox> PanelContainer;
+	TSharedPtr<SCameraRigInterfaceParametersPanel> Panel;
+};
+
+}  // namespace UE::Cameras
+
diff --git a/Engine/Plugins/Cameras/GameplayCameras/Source/GameplayCamerasUncookedOnly/Private/BlueprintGraph/K2Node_SetCameraRigParameters.cpp b/Engine/Plugins/Cameras/GameplayCameras/Source/GameplayCamerasUncookedOnly/Private/BlueprintGraph/K2Node_SetCameraRigParameters.cpp
--- a/Engine/Plugins/Cameras/GameplayCameras/Source/GameplayCamerasUncookedOnly/Private/BlueprintGraph/K2Node_SetCameraRigParameters.cpp
+++ b/Engine/Plugins/Cameras/GameplayCameras/Source/GameplayCamerasUncookedOnly/Private/BlueprintGraph/K2Node_SetCameraRigParameters.cpp
@@ -50,12 +50,12 @@ void UK2Node_SetCameraRigParameters::ReallocatePinsDuringReconstruction(TArray<U
 	{
 		// The camera rig might not be loaded yet.
 		PreloadObject(CameraRig);
-		for (UCameraRigInterfaceParameter* InterfaceParameter : CameraRig->Interface.InterfaceParameters)
+		for (UCameraRigBlendableParameter* BlendableParameter : CameraRig->Interface.BlendableParameters)
 		{
-			PreloadObject(InterfaceParameter);
-			if (InterfaceParameter)
+			PreloadObject(BlendableParameter);
+			if (BlendableParameter)
 			{
-				PreloadObject(InterfaceParameter->PrivateVariable);
+				PreloadObject(BlendableParameter->PrivateVariable);
 			}
 		}
 
@@ -145,14 +145,14 @@ void UK2Node_SetCameraRigParameters::ExpandNode(class FKismetCompilerContext& Co
 
 	for (UEdGraphPin* RigParameterPin : RigParameterPins)
 	{
-		UCameraRigInterfaceParameter* InterfaceParameter = CameraRig->Interface.FindInterfaceParameterByName(RigParameterPin->GetName());
-		if (!InterfaceParameter)
+		UCameraRigBlendableParameter* BlendableParameter = CameraRig->Interface.FindBlendableParameterByName(RigParameterPin->GetName());
+		if (!BlendableParameter)
 		{
 			CompilerContext.MessageLog.Error(*LOCTEXT("ErrorMissingParameter", "SetCameraRigParameters node @@ is trying to set parameter @@ but camera rig @@ has no such parameter.").ToString(), this, *RigParameterPin->GetName(), CameraRig);
 			continue;
 		}
 
-		if (!InterfaceParameter->PrivateVariable)
+		if (!BlendableParameter->PrivateVariable)
 		{
 			CompilerContext.MessageLog.Error(*LOCTEXT("ErrorMissingParameterVariable", "SetCameraRigParameters node @@ needs camera rig @@ to be built.").ToString(), this, CameraRig);
 			continue;
@@ -160,7 +160,7 @@ void UK2Node_SetCameraRigParameters::ExpandNode(class FKismetCompilerContext& Co
 
 		// Figure out the sort of SetXxxParameter function we want to call for this parameter.
 		FName CallSetParameterFuncName;
-		switch (InterfaceParameter->PrivateVariable->GetVariableType())
+		switch (BlendableParameter->PrivateVariable->GetVariableType())
 		{
 			case ECameraVariableType::Boolean:
 				CallSetParameterFuncName = GET_FUNCTION_NAME_CHECKED(UCameraRigParameterInterop, SetBooleanParameter);
@@ -211,7 +211,7 @@ void UK2Node_SetCameraRigParameters::ExpandNode(class FKismetCompilerContext& Co
 
 		// Set the parameter name argument.
 		UEdGraphPin* CallSetParameterNamePin = CallSetParameter->FindPinChecked(TEXT("ParameterName"));
-		CallSetParameterNamePin->DefaultValue = InterfaceParameter->InterfaceParameterName;
+		CallSetParameterNamePin->DefaultValue = BlendableParameter->InterfaceParameterName;
 
 		// Set or connect the parameter value argument.
 		UEdGraphPin* CallSetParameterValuePin = CallSetParameter->FindPinChecked(TEXT("ParameterValue"));
@@ -311,14 +311,14 @@ void UK2Node_SetCameraRigParameters::CreatePinsForCameraRig(UCameraRigAsset* Cam
 
 	const UEdGraphSchema_K2* K2Schema = GetDefault<UEdGraphSchema_K2>();
 
-	for (UCameraRigInterfaceParameter* InterfaceParameter : CameraRig->Interface.InterfaceParameters)
+	for (UCameraRigBlendableParameter* BlendableParameter : CameraRig->Interface.BlendableParameters)
 	{
-		if (!ensure(InterfaceParameter))
+		if (!ensure(BlendableParameter))
 		{
 			continue;
 		}
 
-		if (!InterfaceParameter->PrivateVariable)
+		if (!BlendableParameter->PrivateVariable)
 		{
 			// Camera rig isn't fully built.
 			continue;
@@ -327,7 +327,7 @@ void UK2Node_SetCameraRigParameters::CreatePinsForCameraRig(UCameraRigAsset* Cam
 		FName NewPinCategory;
 		FName NewPinSubCategory;
 		UObject* NewPinSubCategoryObject = nullptr;
-		switch (InterfaceParameter->PrivateVariable->GetVariableType())
+		switch (BlendableParameter->PrivateVariable->GetVariableType())
 		{
 			case ECameraVariableType::Boolean:
 				NewPinCategory = UEdGraphSchema_K2::PC_Boolean;
@@ -374,7 +374,7 @@ void UK2Node_SetCameraRigParameters::CreatePinsForCameraRig(UCameraRigAsset* Cam
 		UEdGraphPin* NewPin = CreatePin(
 				EGPD_Input, 
 				NewPinCategory, NewPinSubCategory, NewPinSubCategoryObject, 
-				FName(InterfaceParameter->InterfaceParameterName));
+				FName(BlendableParameter->InterfaceParameterName));
 		if (CreatedPins)
 		{
 			CreatedPins->Add(NewPin);
