Cameras: add support for reconstructing graph nodes if their custom parameters change

- Add global delegate for custom parameter providers changing their parameters.
- Make the graph nodes listen to it if needed, and reconstruct their pins.
- Make the camera rig prefab node implement it.

diff --git a/Engine/Plugins/Cameras/GameplayCameras/Source/GameplayCameras/Private/Core/ICustomCameraNodeParameterProvider.cpp b/Engine/Plugins/Cameras/GameplayCameras/Source/GameplayCameras/Private/Core/ICustomCameraNodeParameterProvider.cpp
--- a/Engine/Plugins/Cameras/GameplayCameras/Source/GameplayCameras/Private/Core/ICustomCameraNodeParameterProvider.cpp
+++ b/Engine/Plugins/Cameras/GameplayCameras/Source/GameplayCameras/Private/Core/ICustomCameraNodeParameterProvider.cpp
@@ -2,6 +2,8 @@
 
 #include "Core/ICustomCameraNodeParameterProvider.h"
 
+#include "GameplayCamerasDelegates.h"
+
 bool operator==(const FCustomCameraNodeBlendableParameter& A, const FCustomCameraNodeBlendableParameter& B)
 {
 	return A.ParameterName == B.ParameterName &&
@@ -108,3 +110,10 @@ bool FCustomCameraNodeParameterInfos::FindDataParameter(FName ParameterName, FCu
 	return false;
 }
 
+void ICustomCameraNodeParameterProvider::OnCustomCameraNodeParametersChanged(const UCameraNode* ThisAsCameraNode) const
+{
+	using namespace UE::Cameras;
+
+	FGameplayCamerasDelegates::OnCustomCameraNodeParametersChanged().Broadcast(ThisAsCameraNode);
+}
+
diff --git a/Engine/Plugins/Cameras/GameplayCameras/Source/GameplayCameras/Private/GameplayCamerasDelegates.cpp b/Engine/Plugins/Cameras/GameplayCameras/Source/GameplayCameras/Private/GameplayCamerasDelegates.cpp
--- a/Engine/Plugins/Cameras/GameplayCameras/Source/GameplayCameras/Private/GameplayCamerasDelegates.cpp
+++ b/Engine/Plugins/Cameras/GameplayCameras/Source/GameplayCameras/Private/GameplayCamerasDelegates.cpp
@@ -5,8 +5,9 @@
 namespace UE::Cameras
 {
 
-FOnCameraAssetBuilt FGameplayCamerasDelegates::OnCameraAssetBuiltDelegates;
-FOnCameraRigAssetBuilt FGameplayCamerasDelegates::OnCameraRigAssetBuiltDelegates;
+FOnCameraAssetBuilt FGameplayCamerasDelegates::OnCameraAssetBuiltDelegate;
+FOnCameraRigAssetBuilt FGameplayCamerasDelegates::OnCameraRigAssetBuiltDelegate;
+FOnCameraNodeChanged FGameplayCamerasDelegates::OnCustomCameraNodeParametersChangedDelegate;
 
 }  // namespace UE::Cameras
 
diff --git a/Engine/Plugins/Cameras/GameplayCameras/Source/GameplayCameras/Private/Nodes/Common/CameraRigCameraNode.cpp b/Engine/Plugins/Cameras/GameplayCameras/Source/GameplayCameras/Private/Nodes/Common/CameraRigCameraNode.cpp
--- a/Engine/Plugins/Cameras/GameplayCameras/Source/GameplayCameras/Private/Nodes/Common/CameraRigCameraNode.cpp
+++ b/Engine/Plugins/Cameras/GameplayCameras/Source/GameplayCameras/Private/Nodes/Common/CameraRigCameraNode.cpp
@@ -191,5 +191,19 @@ FCameraNodeEvaluatorPtr UCameraRigCameraNode::OnBuildEvaluator(FCameraNodeEvalua
 	return Builder.BuildEvaluator<FCameraRigCameraNodeEvaluator>();
 }
 
+#if WITH_EDITOR
+
+void UCameraRigCameraNode::PostEditChangeProperty( struct FPropertyChangedEvent& PropertyChangedEvent)
+{
+	Super::PostEditChangeProperty(PropertyChangedEvent);
+
+	if (PropertyChangedEvent.GetPropertyName() == GET_MEMBER_NAME_CHECKED(UCameraRigCameraNode, CameraRigReference))
+	{
+		OnCustomCameraNodeParametersChanged(this);
+	}
+}
+
+#endif
+
 #undef LOCTEXT_NAMESPACE
 
diff --git a/Engine/Plugins/Cameras/GameplayCameras/Source/GameplayCameras/Public/Core/ICustomCameraNodeParameterProvider.h b/Engine/Plugins/Cameras/GameplayCameras/Source/GameplayCameras/Public/Core/ICustomCameraNodeParameterProvider.h
--- a/Engine/Plugins/Cameras/GameplayCameras/Source/GameplayCameras/Public/Core/ICustomCameraNodeParameterProvider.h
+++ b/Engine/Plugins/Cameras/GameplayCameras/Source/GameplayCameras/Public/Core/ICustomCameraNodeParameterProvider.h
@@ -9,6 +9,7 @@
 
 #include "ICustomCameraNodeParameterProvider.generated.h"
 
+class UCameraNode;
 class UCameraVariableAsset;
 
 namespace UE::Cameras
@@ -177,5 +178,8 @@ public:
 
 	/** Gathers the custom parameters on this node. */
 	virtual void GetCustomCameraNodeParameters(FCustomCameraNodeParameterInfos& OutParameterInfos) {}
+
+	/** Utility function for sub-classes to broadcast when the custom parameters have changed. */
+	GAMEPLAYCAMERAS_API void OnCustomCameraNodeParametersChanged(const UCameraNode* ThisAsCameraNode) const;
 };
 
diff --git a/Engine/Plugins/Cameras/GameplayCameras/Source/GameplayCameras/Public/GameplayCamerasDelegates.h b/Engine/Plugins/Cameras/GameplayCameras/Source/GameplayCameras/Public/GameplayCamerasDelegates.h
--- a/Engine/Plugins/Cameras/GameplayCameras/Source/GameplayCameras/Public/GameplayCamerasDelegates.h
+++ b/Engine/Plugins/Cameras/GameplayCameras/Source/GameplayCameras/Public/GameplayCamerasDelegates.h
@@ -6,6 +6,7 @@
 #include "Delegates/Delegate.h"
 
 class UCameraAsset;
+class UCameraNode;
 class UCameraRigAsset;
 
 namespace UE::Cameras
@@ -15,6 +16,7 @@ class FCameraBuildLog;
 
 DECLARE_MULTICAST_DELEGATE_OneParam(FOnCameraAssetBuilt, const UCameraAsset*);
 DECLARE_MULTICAST_DELEGATE_OneParam(FOnCameraRigAssetBuilt, const UCameraRigAsset*);
+DECLARE_MULTICAST_DELEGATE_OneParam(FOnCameraNodeChanged, const UCameraNode*);
 
 /**
  * Global delegates for the GameplayCameras module.
@@ -26,19 +28,26 @@ public:
 	/** Broadcast when a camera asset has been built. */
 	static inline FOnCameraAssetBuilt& OnCameraAssetBuilt()
 	{
-		return OnCameraAssetBuiltDelegates;
+		return OnCameraAssetBuiltDelegate;
 	}
 
 	/** Broadcast when a camera rig has been built. */
 	static inline FOnCameraRigAssetBuilt& OnCameraRigAssetBuilt()
 	{
-		return OnCameraRigAssetBuiltDelegates;
+		return OnCameraRigAssetBuiltDelegate;
+	}
+
+	/** Broadcast when a custom camera parameter provider node changes it parameters. */
+	static inline FOnCameraNodeChanged& OnCustomCameraNodeParametersChanged()
+	{
+		return OnCustomCameraNodeParametersChangedDelegate;
 	}
 
 private:
 
-	static FOnCameraAssetBuilt OnCameraAssetBuiltDelegates;
-	static FOnCameraRigAssetBuilt OnCameraRigAssetBuiltDelegates;
+	static FOnCameraAssetBuilt OnCameraAssetBuiltDelegate;
+	static FOnCameraRigAssetBuilt OnCameraRigAssetBuiltDelegate;
+	static FOnCameraNodeChanged OnCustomCameraNodeParametersChangedDelegate;
 };
 
 }  // namespace UE::Cameras
diff --git a/Engine/Plugins/Cameras/GameplayCameras/Source/GameplayCameras/Public/Nodes/Common/CameraRigCameraNode.h b/Engine/Plugins/Cameras/GameplayCameras/Source/GameplayCameras/Public/Nodes/Common/CameraRigCameraNode.h
--- a/Engine/Plugins/Cameras/GameplayCameras/Source/GameplayCameras/Public/Nodes/Common/CameraRigCameraNode.h
+++ b/Engine/Plugins/Cameras/GameplayCameras/Source/GameplayCameras/Public/Nodes/Common/CameraRigCameraNode.h
@@ -29,6 +29,11 @@ protected:
 	// ICustomCameraNodeParameterProvider interface.
 	virtual void GetCustomCameraNodeParameters(FCustomCameraNodeParameterInfos& OutParameterInfos) override;
 
+	// UObject interface.
+#if WITH_EDITOR
+	virtual void PostEditChangeProperty( struct FPropertyChangedEvent& PropertyChangedEvent) override;
+#endif  // WITH_EDITOR
+
 public:
 
 	/** The camera rig to run. */
diff --git a/Engine/Plugins/Cameras/GameplayCameras/Source/GameplayCamerasEditor/Private/Editors/CameraNodeGraphNode.cpp b/Engine/Plugins/Cameras/GameplayCameras/Source/GameplayCamerasEditor/Private/Editors/CameraNodeGraphNode.cpp
--- a/Engine/Plugins/Cameras/GameplayCameras/Source/GameplayCamerasEditor/Private/Editors/CameraNodeGraphNode.cpp
+++ b/Engine/Plugins/Cameras/GameplayCameras/Source/GameplayCamerasEditor/Private/Editors/CameraNodeGraphNode.cpp
@@ -2,12 +2,14 @@
 
 #include "Editors/CameraNodeGraphNode.h"
 
+#include "Core/CameraNode.h"
 #include "Core/CameraParameters.h"
 #include "Core/CameraVariableReferences.h"
 #include "Core/ICustomCameraNodeParameterProvider.h"
 #include "EdGraph/EdGraphPin.h"
 #include "Editors/CameraNodeGraphSchema.h"
 #include "Editors/SCameraNodeGraphNode.h"
+#include "GameplayCamerasDelegates.h"
 #include "ToolMenus.h"
 
 UCameraNodeGraphNode::UCameraNodeGraphNode(const FObjectInitializer& ObjInit)
@@ -15,6 +17,35 @@ UCameraNodeGraphNode::UCameraNodeGraphNode(const FObjectInitializer& ObjInit)
 {
 }
 
+void UCameraNodeGraphNode::OnInitialize()
+{
+	using namespace UE::Cameras;
+
+	const bool bIsCustomParameterProvider = GetObject()->Implements<UCustomCameraNodeParameterProvider>();
+	if (bIsCustomParameterProvider)
+	{
+		FGameplayCamerasDelegates::OnCustomCameraNodeParametersChanged().AddUObject(
+				this, &UCameraNodeGraphNode::OnCustomCameraNodeParametersChanged);
+	}
+}
+
+void UCameraNodeGraphNode::BeginDestroy()
+{
+	using namespace UE::Cameras;
+
+	FGameplayCamerasDelegates::OnCustomCameraNodeParametersChanged().RemoveAll(this);
+
+	Super::BeginDestroy();
+}
+
+void UCameraNodeGraphNode::OnCustomCameraNodeParametersChanged(const UCameraNode* CameraNode)
+{
+	if (CameraNode == GetObject())
+	{
+		ReconstructNode();
+	}
+}
+
 void UCameraNodeGraphNode::AllocateDefaultPins()
 {
 	Super::AllocateDefaultPins();
diff --git a/Engine/Plugins/Cameras/GameplayCameras/Source/GameplayCamerasEditor/Private/Editors/ObjectTreeGraphNode.cpp b/Engine/Plugins/Cameras/GameplayCameras/Source/GameplayCamerasEditor/Private/Editors/ObjectTreeGraphNode.cpp
--- a/Engine/Plugins/Cameras/GameplayCameras/Source/GameplayCamerasEditor/Private/Editors/ObjectTreeGraphNode.cpp
+++ b/Engine/Plugins/Cameras/GameplayCameras/Source/GameplayCamerasEditor/Private/Editors/ObjectTreeGraphNode.cpp
@@ -32,6 +32,8 @@ void UObjectTreeGraphNode::Initialize(UObject* InObject)
 	{
 		NodeComment = GraphObject->GetGraphNodeCommentText(NodeContext.GraphConfig.GraphName);
 	}
+
+	OnInitialize();
 }
 
 FText UObjectTreeGraphNode::GetNodeTitle(ENodeTitleType::Type TitleType) const
@@ -182,6 +184,43 @@ void UObjectTreeGraphNode::OnPinRemoved(UEdGraphPin* InRemovedPin)
 	}
 }
 
+void UObjectTreeGraphNode::ReconstructNode()
+{
+	Modify(false);
+
+	ErrorMsg.Reset();
+
+	// Save old pins.
+	TMap<FName, UEdGraphPin*> OldPins;
+	for (auto NodeIt = Pins.CreateIterator(); NodeIt; ++NodeIt)
+	{
+		UEdGraphPin* Pin(*NodeIt);
+		OldPins.Add(Pin->PinName, Pin);
+	}
+
+	// Reconstruct all pins from scratch.
+	Pins.Reset();
+	AllocateDefaultPins();
+
+	// Rewire existing connections to new pins, matched by name, direction, and type.
+	for (UEdGraphPin* NewPin : Pins)
+	{
+		UEdGraphPin* OldPin = OldPins.FindRef(NewPin->PinName);
+		if (!OldPin || 
+				OldPin->Direction != NewPin->Direction ||
+				OldPin->PinType != NewPin->PinType)
+		{
+			continue;
+		}
+
+		NewPin->MovePersistentDataFromOldPin(*OldPin);
+	}
+
+	GetGraph()->NotifyNodeChanged(this);
+
+	Super::ReconstructNode();
+}
+
 void UObjectTreeGraphNode::GetArrayProperties(TArray<FArrayProperty*>& OutArrayProperties, EEdGraphPinDirection Direction) const
 {
 	const FNodeContext NodeContext = GetNodeContext();
diff --git a/Engine/Plugins/Cameras/GameplayCameras/Source/GameplayCamerasEditor/Public/Editors/CameraNodeGraphNode.h b/Engine/Plugins/Cameras/GameplayCameras/Source/GameplayCamerasEditor/Public/Editors/CameraNodeGraphNode.h
--- a/Engine/Plugins/Cameras/GameplayCameras/Source/GameplayCamerasEditor/Public/Editors/CameraNodeGraphNode.h
+++ b/Engine/Plugins/Cameras/GameplayCameras/Source/GameplayCamerasEditor/Public/Editors/CameraNodeGraphNode.h
@@ -6,6 +6,8 @@
 
 #include "CameraNodeGraphNode.generated.h"
 
+class UCameraNode;
+
 /**
  * Custom graph node for camera nodes. They mostly differ by showing input pins for any 
  * camera parameter property.
@@ -20,10 +22,19 @@ public:
 	/** Creates a new graph node. */
 	UCameraNodeGraphNode(const FObjectInitializer& ObjInit);
 
+	virtual void BeginDestroy() override;
+
 public:
 
-	// UEdGraphNode interface
+	// UEdGraphNode interface.
 	virtual void AllocateDefaultPins() override;
 	virtual TSharedPtr<SGraphNode> CreateVisualWidget() override;
+
+	// UObjectTreeGraphNode interface.
+	virtual void OnInitialize() override;
+
+private:
+
+	void OnCustomCameraNodeParametersChanged(const UCameraNode* CameraNode);
 };
 
diff --git a/Engine/Plugins/Cameras/GameplayCameras/Source/GameplayCamerasEditor/Public/Editors/ObjectTreeGraphNode.h b/Engine/Plugins/Cameras/GameplayCameras/Source/GameplayCamerasEditor/Public/Editors/ObjectTreeGraphNode.h
--- a/Engine/Plugins/Cameras/GameplayCameras/Source/GameplayCamerasEditor/Public/Editors/ObjectTreeGraphNode.h
+++ b/Engine/Plugins/Cameras/GameplayCameras/Source/GameplayCamerasEditor/Public/Editors/ObjectTreeGraphNode.h
@@ -72,6 +72,7 @@ public:
 	virtual void PinConnectionListChanged(UEdGraphPin* Pin) override;
 	virtual void NodeConnectionListChanged() override;
 	virtual void OnPinRemoved(UEdGraphPin* InRemovedPin) override;
+	virtual void ReconstructNode() override;
 	virtual void GetNodeContextMenuActions(class UToolMenu* Menu, class UGraphNodeContextMenuContext* Context) const override;
 	virtual bool GetCanRenameNode() const override;
 	virtual void OnRenameNode(const FString& NewName) override;
@@ -81,6 +82,7 @@ public:
 	virtual void OnUpdateCommentText(const FString& NewComment) override;
 
 	// UObjectTreeGraphNode interface.
+	virtual void OnInitialize() {}
 	virtual void OnGraphNodeMoved(bool bMarkDirty);
 	virtual void OnDoubleClicked() {}
 
