Cameras: make the Blueprint node implement dynamic camera parameters

- Change the reference to a Blueprint class into a Blueprint instance, so that users can set values of exposed BP properties on the node.
- Implement the custom parameter provider interface and make the node refresh its pins when the BP is recompiled.

diff --git a/Engine/Plugins/Cameras/GameplayCameras/Source/GameplayCameras/Private/Nodes/Utility/BlueprintCameraNode.cpp b/Engine/Plugins/Cameras/GameplayCameras/Source/GameplayCameras/Private/Nodes/Utility/BlueprintCameraNode.cpp
--- a/Engine/Plugins/Cameras/GameplayCameras/Source/GameplayCameras/Private/Nodes/Utility/BlueprintCameraNode.cpp
+++ b/Engine/Plugins/Cameras/GameplayCameras/Source/GameplayCameras/Private/Nodes/Utility/BlueprintCameraNode.cpp
@@ -4,14 +4,18 @@
 
 #include "Components/ActorComponent.h"
 #include "Core/CameraBuildLog.h"
+#include "Core/CameraContextDataTableFwd.h"
 #include "Core/CameraEvaluationContext.h"
 #include "Core/CameraNodeEvaluator.h"
 #include "Core/CameraRigBuildContext.h"
+#include "Core/CameraVariableAssets.h"
 #include "Core/CameraVariableTable.h"
 #include "Engine/World.h"
 #include "GameFramework/Actor.h"
 #include "GameplayCameras.h"
-#include "Templates/UnrealTemplate.h"
+#include "Misc/AssertionMacros.h"
+#include "UObject/Object.h"
+#include "UObject/ObjectMacros.h"
 
 #include UE_INLINE_GENERATED_CPP_BY_NAME(BlueprintCameraNode)
 
@@ -31,6 +35,14 @@ protected:
 	virtual void OnRun(const FCameraNodeEvaluationParams& Params, FCameraNodeEvaluationResult& OutResult) override;
 	virtual void OnAddReferencedObjects(FReferenceCollector& Collector) override;
 
+private:
+
+	void ApplyParameterOverrides(const FCameraVariableTable& VariableTable, const FCameraContextDataTable& ContextDataTable);
+
+	template<typename CameraParameterType>
+	void ApplyParameterOverride(
+			TValueOrError<CameraParameterType*, EPropertyBagResult> ParameterOrError);
+
 private:
 
 	TObjectPtr<UBlueprintCameraNodeEvaluator> EvaluatorBlueprint;
@@ -46,10 +58,14 @@ void FBlueprintCameraNodeEvaluator::OnInitialize(const FCameraNodeEvaluatorIniti
 		return;
 	}
 
-	if (BlueprintNode->CameraNodeEvaluatorClass)
+	if (const UBlueprintCameraNodeEvaluator* EvaluatorTemplate = BlueprintNode->GetCameraNodeEvaluatorTemplate())
 	{
 		UObject* Outer = Params.EvaluationContext->GetOwner();
-		EvaluatorBlueprint = NewObject<UBlueprintCameraNodeEvaluator>(Outer, BlueprintNode->CameraNodeEvaluatorClass);
+		EvaluatorBlueprint = CastChecked<UBlueprintCameraNodeEvaluator>(StaticDuplicateObject(EvaluatorTemplate, Outer, NAME_None));
+
+		ApplyParameterOverrides(OutResult.VariableTable, OutResult.ContextDataTable);
+
+		EvaluatorBlueprint->NativeInitializeCameraNode(Params, OutResult);
 	}
 	else
 	{
@@ -61,6 +77,8 @@ void FBlueprintCameraNodeEvaluator::OnRun(const FCameraNodeEvaluationParams& Par
 {
 	if (EvaluatorBlueprint)
 	{
+		ApplyParameterOverrides(OutResult.VariableTable, OutResult.ContextDataTable);
+
 		EvaluatorBlueprint->NativeRunCameraNode(Params, OutResult);
 	}
 }
@@ -70,25 +88,96 @@ void FBlueprintCameraNodeEvaluator::OnAddReferencedObjects(FReferenceCollector&
 	Collector.AddReferencedObject(EvaluatorBlueprint);
 }
 
+void FBlueprintCameraNodeEvaluator::ApplyParameterOverrides(const FCameraVariableTable& VariableTable, const FCameraContextDataTable& ContextDataTable)
+{
+	using namespace Internal;
+
+	const UClass* EvaluatorBlueprintClass = EvaluatorBlueprint->GetClass();
+	const UBlueprintCameraNode* BlueprintNode = GetCameraNodeAs<UBlueprintCameraNode>();
+	const FCustomCameraNodeParameters& Overrides = BlueprintNode->CameraNodeEvaluatorOverrides;
+
+	// Set the value of any properties driven by a camera variable.
+	for (const FCustomCameraNodeBlendableParameter& BlendableParameter : Overrides.BlendableParameters)
+	{
+		FProperty* Property = EvaluatorBlueprintClass->FindPropertyByName(BlendableParameter.ParameterName);
+		if (ensure(Property))
+		{
+			if (BlendableParameter.OverrideVariable)
+			{
+				const uint8* ValuePtr = VariableTable.GetValue(
+						BlendableParameter.OverrideVariable->GetVariableID(), 
+						BlendableParameter.ParameterType);
+				Property->SetValue_InContainer(EvaluatorBlueprint, ValuePtr);
+			}
+		}
+	}
+
+	// Set the value of any properties driven by a data parameter.
+	for (const FCustomCameraNodeDataParameter& DataParameter : Overrides.DataParameters)
+	{
+		FProperty* Property = EvaluatorBlueprintClass->FindPropertyByName(DataParameter.ParameterName);
+		if (ensure(Property))
+		{
+			if (DataParameter.OverrideDataID)
+			{
+				const uint8* DataPtr = ContextDataTable.GetData(
+						DataParameter.OverrideDataID,
+						DataParameter.ParameterType,
+						DataParameter.ParameterTypeObject);
+				Property->SetValue_InContainer(EvaluatorBlueprint, DataPtr);
+			}
+		}
+	}
+}
+
 }  // namespace UE::Cameras
 
+void UBlueprintCameraNodeEvaluator::NativeInitializeCameraNode(const FCameraNodeEvaluatorInitializeParams& Params, FCameraNodeEvaluationResult& OutResult)
+{
+	using namespace UE::Cameras;
+
+	SetupExecution(Params.EvaluationContext, OutResult);
+	{
+		bIsFirstFrame = true;
+
+		InitializeCameraNode();
+
+		CameraPose.ApplyTo(OutResult.CameraPose);
+	}
+	TeardownExecution();
+}
+
 void UBlueprintCameraNodeEvaluator::NativeRunCameraNode(const FCameraNodeEvaluationParams& Params, FCameraNodeEvaluationResult& OutResult)
 {
 	using namespace UE::Cameras;
 
-	bIsFirstFrame = Params.bIsFirstFrame;
-	EvaluationContextOwner = Params.EvaluationContext->GetOwner();
+	SetupExecution(Params.EvaluationContext, OutResult);
+	{
+		bIsFirstFrame = Params.bIsFirstFrame;
+
+		TickCameraNode(Params.DeltaTime);
+
+		CameraPose.ApplyTo(OutResult.CameraPose);
+	}
+	TeardownExecution();
+}
+
+void UBlueprintCameraNodeEvaluator::SetupExecution(TSharedPtr<const FCameraEvaluationContext> EvaluationContext, FCameraNodeEvaluationResult& OutResult)
+{
+	EvaluationContextOwner = EvaluationContext->GetOwner();
 	CameraPose = FBlueprintCameraPose::FromCameraPose(OutResult.CameraPose);
 
 	ensure(!CurrentContext.IsValid());
-	TGuardValue<TSharedPtr<const FCameraEvaluationContext>> CurrentContextGuard(CurrentContext, Params.EvaluationContext);
+	CurrentContext = EvaluationContext;
 
 	ensure(CurrentResult == nullptr);
-	TGuardValue<FCameraNodeEvaluationResult*> CurrentResultGuard(CurrentResult, &OutResult);
-
-	TickCameraNode(Params.DeltaTime);
+	CurrentResult = &OutResult;
+}
 
-	CameraPose.ApplyTo(OutResult.CameraPose);
+void UBlueprintCameraNodeEvaluator::TeardownExecution()
+{
+	CurrentContext = nullptr;
+	CurrentResult = nullptr;
 }
 
 AActor* UBlueprintCameraNodeEvaluator::FindEvaluationContextOwnerActor(TSubclassOf<AActor> ActorClass) const
@@ -145,13 +234,317 @@ UWorld* UBlueprintCameraNodeEvaluator::GetWorld() const
 	return nullptr;
 }
 
+UBlueprintCameraNode::UBlueprintCameraNode(const FObjectInitializer& ObjInit)
+	: Super(ObjInit)
+{
+#if WITH_EDITOR
+	FCoreUObjectDelegates::OnObjectsReplaced.AddUObject(this, &UBlueprintCameraNode::OnObjectsReplaced);
+#endif
+}
+
+void UBlueprintCameraNode::PostLoad()
+{
+	Super::PostLoad();
+
+	if (CameraNodeEvaluatorClass_DEPRECATED)
+	{
+		CameraNodeEvaluatorTemplate = NewObject<UBlueprintCameraNodeEvaluator>(
+				this, CameraNodeEvaluatorClass_DEPRECATED, NAME_None, RF_Transactional);
+
+		RebuildOverrides();
+
+		CameraNodeEvaluatorClass_DEPRECATED = nullptr;
+	}
+}
+
+void UBlueprintCameraNode::BeginDestroy()
+{
+#if WITH_EDITOR
+	FCoreUObjectDelegates::OnObjectsReplaced.RemoveAll(this);
+#endif
+
+	Super::BeginDestroy();
+}
+
+void UBlueprintCameraNode::RebuildOverrides()
+{
+	using namespace UE::Cameras::Internal;
+
+	// If there is no evaluator set, clear all overrides.
+	if (!CameraNodeEvaluatorTemplate)
+	{
+		if (CameraNodeEvaluatorOverrides.HasAnyParameters())
+		{
+			Modify();
+
+			CameraNodeEvaluatorOverrides.Reset();
+		}
+		return;
+	}
+
+	// Remember the overrides already present on parameters.
+	TMap<FName, UCameraVariableAsset*> OldOverrideVariableMap;
+	for (const FCustomCameraNodeBlendableParameter& OldOverride : CameraNodeEvaluatorOverrides.BlendableParameters)
+	{
+		if (OldOverride.OverrideVariable)
+		{
+			OldOverrideVariableMap.Add(OldOverride.ParameterName, OldOverride.OverrideVariable);
+		}
+	}
+	TMap<FName, FCameraContextDataID> OldOverrideDataIDMap;
+	for (const FCustomCameraNodeDataParameter& OldOverride : CameraNodeEvaluatorOverrides.DataParameters)
+	{
+		if (OldOverride.OverrideDataID)
+		{
+			OldOverrideDataIDMap.Add(OldOverride.ParameterName, OldOverride.OverrideDataID);
+		}
+	}
+
+	// Build the new list of blendable and data parameters.
+	// All exposed blendable properties on the Blueprint class show up as blendable parameters.
+	// All other exposed properties on the Blueprint class show up as data parameters.
+	FCustomCameraNodeParameters NewOverrides;
+
+	for (TFieldIterator<FProperty> PropertyIt(CameraNodeEvaluatorTemplate->GetClass()); PropertyIt; ++PropertyIt)
+	{
+		FProperty* Property(*PropertyIt);
+
+		if (!Property->GetOwnerClass()->HasAnyClassFlags(CLASS_CompiledFromBlueprint))
+		{
+			continue;
+		}
+		if (Property->HasAnyPropertyFlags(CPF_EditorOnly | CPF_Protected | CPF_DisableEditOnInstance))
+		{
+			continue;
+		}
+		if (!Property->HasAnyPropertyFlags(CPF_Edit))
+		{
+			continue;
+		}
+
+		bool bIsBlendableProperty = false;
+		ECameraVariableType BlendablePropertyType = ECameraVariableType::Boolean;
+
+		bool bIsDataProperty = false;
+		ECameraContextDataType DataPropertyType = ECameraContextDataType::Name;
+		const UObject* DataPropertyTypeObject = nullptr;
+
+		if (FBoolProperty* BoolProperty = CastField<FBoolProperty>(Property))
+		{
+			bIsBlendableProperty = true;
+			BlendablePropertyType = ECameraVariableType::Boolean;
+		}
+		else if (FIntProperty* Int32Property = CastField<FIntProperty>(Property))
+		{
+			bIsBlendableProperty = true;
+			BlendablePropertyType = ECameraVariableType::Integer32;
+		}
+		else if (FFloatProperty* FloatProperty = CastField<FFloatProperty>(Property))
+		{
+			bIsBlendableProperty = true;
+			BlendablePropertyType = ECameraVariableType::Float;
+		}
+		else if (FDoubleProperty* DoubleProperty = CastField<FDoubleProperty>(Property))
+		{
+			bIsBlendableProperty = true;
+			BlendablePropertyType = ECameraVariableType::Double;
+		}
+		else if (FStructProperty* StructProperty = CastField<FStructProperty>(Property))
+		{
+			if (StructProperty->Struct == TVariantStructure<FVector2f>::Get())
+			{
+				bIsBlendableProperty = true;
+				BlendablePropertyType = ECameraVariableType::Vector2f;
+			}
+			else if (StructProperty->Struct == TBaseStructure<FVector2D>::Get())
+			{
+				bIsBlendableProperty = true;
+				BlendablePropertyType = ECameraVariableType::Vector2d;
+			}
+			else if (StructProperty->Struct == TVariantStructure<FVector3f>::Get())
+			{
+				bIsBlendableProperty = true;
+				BlendablePropertyType = ECameraVariableType::Vector3f;
+			}
+			else if (StructProperty->Struct == TBaseStructure<FVector>::Get())
+			{
+				bIsBlendableProperty = true;
+				BlendablePropertyType = ECameraVariableType::Vector3d;
+			}
+			else if (StructProperty->Struct == TVariantStructure<FVector4f>::Get())
+			{
+				bIsBlendableProperty = true;
+				BlendablePropertyType = ECameraVariableType::Vector4f;
+			}
+			else if (StructProperty->Struct == TBaseStructure<FVector4>::Get())
+			{
+				bIsBlendableProperty = true;
+				BlendablePropertyType = ECameraVariableType::Vector4d;
+			}
+			else if (StructProperty->Struct == TVariantStructure<FRotator3f>::Get())
+			{
+				bIsBlendableProperty = true;
+				BlendablePropertyType = ECameraVariableType::Rotator3f;
+			}
+			else if (StructProperty->Struct == TBaseStructure<FRotator>::Get())
+			{
+				bIsBlendableProperty = true;
+				BlendablePropertyType = ECameraVariableType::Rotator3d;
+			}
+			else if (StructProperty->Struct == TVariantStructure<FTransform3f>::Get())
+			{
+				bIsBlendableProperty = true;
+				BlendablePropertyType = ECameraVariableType::Transform3f;
+			}
+			else if (StructProperty->Struct == TBaseStructure<FTransform>::Get())
+			{
+				bIsBlendableProperty = true;
+				BlendablePropertyType = ECameraVariableType::Transform3d;
+			}
+			else
+			{
+				bIsDataProperty = true;
+				DataPropertyType = ECameraContextDataType::Struct;
+				DataPropertyTypeObject = StructProperty->Struct;
+			}
+		}
+		else if (FNameProperty* NameProperty = CastField<FNameProperty>(Property))
+		{
+			bIsDataProperty = true;
+			DataPropertyType = ECameraContextDataType::Name;
+		}
+		else if (FStrProperty* StringProperty = CastField<FStrProperty>(Property))
+		{
+			bIsDataProperty = true;
+			DataPropertyType = ECameraContextDataType::String;
+		}
+		else if (FEnumProperty* EnumProperty = CastField<FEnumProperty>(Property))
+		{
+			bIsDataProperty = true;
+			DataPropertyType = ECameraContextDataType::Enum;
+			DataPropertyTypeObject = EnumProperty->GetEnum();
+		}
+		else if (FObjectProperty* ObjectProperty = CastField<FObjectProperty>(Property))
+		{
+			bIsDataProperty = true;
+			DataPropertyType = ECameraContextDataType::Object;
+		}
+		else if (FClassProperty* ClassProperty = CastField<FClassProperty>(Property))
+		{
+			bIsDataProperty = true;
+			DataPropertyType = ECameraContextDataType::Class;
+		}
+
+		if (bIsBlendableProperty)
+		{
+			FCustomCameraNodeBlendableParameter NewOverride;
+			NewOverride.ParameterName = Property->GetFName();
+			NewOverride.ParameterType = BlendablePropertyType;
+
+			// If this blendable parameter existed before and had an overriding variable set,
+			// preserve that override.
+			UCameraVariableAsset* OldOverrideVariable = nullptr;
+			OldOverrideVariableMap.RemoveAndCopyValue(NewOverride.ParameterName, OldOverrideVariable);
+			if (OldOverrideVariable && OldOverrideVariable->GetVariableType() == BlendablePropertyType)
+			{
+				NewOverride.OverrideVariable = OldOverrideVariable;
+			}
+
+			NewOverrides.BlendableParameters.Add(NewOverride);
+		}
+		else if (bIsDataProperty)
+		{
+			FCustomCameraNodeDataParameter NewOverride;
+			NewOverride.ParameterName = Property->GetFName();
+			NewOverride.ParameterType = DataPropertyType;
+			NewOverride.ParameterTypeObject = DataPropertyTypeObject;
+
+			// If this data parameter existed before and had an override data ID set,
+			// preserve that override.
+			FCameraContextDataID OldOverrideDataID;
+			OldOverrideDataIDMap.RemoveAndCopyValue(NewOverride.ParameterName, OldOverrideDataID);
+			if (OldOverrideDataID)
+			{
+				NewOverride.OverrideDataID = OldOverrideDataID;
+			}
+
+			NewOverrides.DataParameters.Add(NewOverride);
+		}
+		else
+		{
+			UE_LOG(
+					LogCameraSystem, Warning, 
+					TEXT("Property '%s' on Blueprint camera node evaluator class '%s' cannot be exposed as "
+						 "neither a blendable or data parameter. The property type is not (yet) supported."),
+					*Property->GetName(),
+					*CameraNodeEvaluatorTemplate->GetClass()->GetName());
+		}
+	}
+
+	if (NewOverrides != CameraNodeEvaluatorOverrides)
+	{
+		Modify();
+
+		CameraNodeEvaluatorOverrides = NewOverrides;
+	}
+}
+
+void UBlueprintCameraNode::OnPreBuild(FCameraBuildLog& BuildLog)
+{
+	RebuildOverrides();
+}
+
 void UBlueprintCameraNode::OnBuild(FCameraRigBuildContext& BuildContext)
 {
-	if (!CameraNodeEvaluatorClass)
+	if (!CameraNodeEvaluatorTemplate)
 	{
 		BuildContext.BuildLog.AddMessage(
 				EMessageSeverity::Error, this,
-				LOCTEXT("MissingBlueprintClass", "No evaluator Blueprint class is set."));
+				LOCTEXT("MissingBlueprintEvaluatorTemplate", "No evaluator Blueprint is set."));
+		return;
+	}
+}
+
+void UBlueprintCameraNode::GetCustomCameraNodeParameters(FCustomCameraNodeParameterInfos& OutParameterInfos)
+{
+	using namespace UE::Cameras::Internal;
+
+	if (!CameraNodeEvaluatorTemplate)
+	{
+		return;
+	}
+
+	const UClass* CameraNodeEvaluatorClass = CameraNodeEvaluatorTemplate->GetClass();
+
+	for (FCustomCameraNodeBlendableParameter& BlendableParameter : CameraNodeEvaluatorOverrides.BlendableParameters)
+	{
+		FProperty* BlendableProperty = CameraNodeEvaluatorClass->FindPropertyByName(BlendableParameter.ParameterName);
+		if (!BlendableProperty)
+		{
+			continue;
+		}
+
+		const void* DefaultValuePtr = BlendableProperty->ContainerPtrToValuePtr<void>(CameraNodeEvaluatorTemplate);
+		OutParameterInfos.AddBlendableParameter(
+				BlendableParameter.ParameterName,
+				BlendableParameter.ParameterType, 
+				reinterpret_cast<const uint8*>(DefaultValuePtr),
+				&BlendableParameter.OverrideVariable);
+	}
+
+	for (FCustomCameraNodeDataParameter& DataParameter : CameraNodeEvaluatorOverrides.DataParameters)
+	{
+		FProperty* DataProperty = CameraNodeEvaluatorClass->FindPropertyByName(DataParameter.ParameterName);
+		if (!DataProperty)
+		{
+			continue;
+		}
+
+		OutParameterInfos.AddDataParameter(
+				DataParameter.ParameterName,
+				DataParameter.ParameterType,
+				DataParameter.ParameterTypeObject,
+				&DataParameter.OverrideDataID);
 	}
 }
 
@@ -161,5 +554,34 @@ FCameraNodeEvaluatorPtr UBlueprintCameraNode::OnBuildEvaluator(FCameraNodeEvalua
 	return Builder.BuildEvaluator<FBlueprintCameraNodeEvaluator>();
 }
 
+#if WITH_EDITOR
+
+void UBlueprintCameraNode::OnObjectsReplaced(const TMap<UObject*, UObject*>& ReplacementMap)
+{
+	UObject* NewEvaluatorTemplate = ReplacementMap.FindRef(CameraNodeEvaluatorTemplate);
+	if (NewEvaluatorTemplate)
+	{
+		CameraNodeEvaluatorTemplate = CastChecked<UBlueprintCameraNodeEvaluator>(NewEvaluatorTemplate);
+
+		RebuildOverrides();
+
+		OnCustomCameraNodeParametersChanged(this);
+	}
+}
+
+void UBlueprintCameraNode::PostEditChangeProperty( struct FPropertyChangedEvent& PropertyChangedEvent)
+{
+	Super::PostEditChangeProperty(PropertyChangedEvent);
+
+	if (PropertyChangedEvent.GetPropertyName() == GET_MEMBER_NAME_CHECKED(UBlueprintCameraNode, CameraNodeEvaluatorTemplate))
+	{
+		RebuildOverrides();
+
+		OnCustomCameraNodeParametersChanged(this);
+	}
+}
+
+#endif
+
 #undef LOCTEXT_NAMESPACE
 
diff --git a/Engine/Plugins/Cameras/GameplayCameras/Source/GameplayCameras/Public/Nodes/Utility/BlueprintCameraNode.h b/Engine/Plugins/Cameras/GameplayCameras/Source/GameplayCameras/Public/Nodes/Utility/BlueprintCameraNode.h
--- a/Engine/Plugins/Cameras/GameplayCameras/Source/GameplayCameras/Public/Nodes/Utility/BlueprintCameraNode.h
+++ b/Engine/Plugins/Cameras/GameplayCameras/Source/GameplayCameras/Public/Nodes/Utility/BlueprintCameraNode.h
@@ -2,19 +2,21 @@
 
 #pragma once
 
+#include "Core/CameraEvaluationContext.h"
 #include "Core/CameraNode.h"
-#include "Core/CameraRigJoints.h"
-#include "Core/CameraVariableAssets.h"
+#include "Core/CameraVariableTableFwd.h"
+#include "Core/ICustomCameraNodeParameterProvider.h"
 #include "GameFramework/BlueprintCameraPose.h"
 #include "GameFramework/BlueprintCameraVariableTable.h"
-#include "Templates/SubclassOf.h"
 
 #include "BlueprintCameraNode.generated.h"
 
-struct FBlueprintCameraPose;
+class UCameraVariableAsset;
+class UPropertyBag;
 
 namespace UE::Cameras
 {
+	class FBlueprintCameraNodeEvaluator;
 	class FCameraVariableTable;
 	struct FCameraNodeEvaluationParams;
 	struct FCameraNodeEvaluationResult;
@@ -23,13 +25,16 @@ namespace UE::Cameras
 /**
  * The base class for Blueprint camera node evaluators.
  */
-UCLASS(MinimalAPI, Blueprintable, Abstract)
+UCLASS(MinimalAPI, Blueprintable, Abstract, EditInlineNew, CollapseCategories)
 class UBlueprintCameraNodeEvaluator : public UObject
 {
 	GENERATED_BODY()
 
 public:
 
+	UFUNCTION(BlueprintImplementableEvent, Category="Evaluation")
+	void InitializeCameraNode();
+
 	/** The main execution callback for the camera node. Call SetCameraPose to affect the result. */
 	UFUNCTION(BlueprintImplementableEvent, Category="Evaluation")
 	void TickCameraNode(float DeltaTime);
@@ -44,8 +49,13 @@ public:
 
 public:
 
+	using FCameraNodeEvaluatorInitializeParams = UE::Cameras::FCameraNodeEvaluatorInitializeParams;
 	using FCameraNodeEvaluationParams = UE::Cameras::FCameraNodeEvaluationParams;
 	using FCameraNodeEvaluationResult = UE::Cameras::FCameraNodeEvaluationResult;
+	using FCameraEvaluationContext = UE::Cameras::FCameraEvaluationContext;
+
+	/** Initialize this camera node. */
+	void NativeInitializeCameraNode(const FCameraNodeEvaluatorInitializeParams& Params, FCameraNodeEvaluationResult& OutResult);
 
 	/** Runs this camera node. */
 	void NativeRunCameraNode(const FCameraNodeEvaluationParams& Params, FCameraNodeEvaluationResult& OutResult);
@@ -58,6 +68,11 @@ public:
 	virtual bool ImplementsGetWorld() const override { return true; }
 #endif  // WITH_EDITOR
 
+private:
+
+	void SetupExecution(TSharedPtr<const FCameraEvaluationContext> EvaluationContext, FCameraNodeEvaluationResult& OutResult);
+	void TeardownExecution();
+
 protected:
 
 	/** Whether this is the first frame of this camera node's lifetime. */
@@ -89,20 +104,60 @@ private:
  * A camera node that runs arbitrary Blueprint logic.
  */
 UCLASS(MinimalAPI, meta=(CameraNodeCategories="Common,Transform"))
-class UBlueprintCameraNode : public UCameraNode
+class UBlueprintCameraNode
+	: public UCameraNode
+	, public ICustomCameraNodeParameterProvider
 {
 	GENERATED_BODY()
 
+public:
+
+	UBlueprintCameraNode(const FObjectInitializer& ObjInit);
+
 protected:
 
 	// UCameraNode interface.
+	virtual void OnPreBuild(FCameraBuildLog& BuildLog) override;
 	virtual void OnBuild(FCameraRigBuildContext& BuildContext) override;
 	virtual FCameraNodeEvaluatorPtr OnBuildEvaluator(FCameraNodeEvaluatorBuilder& Builder) const override;
 
+	// ICustomCameraNodeParameterProvider interface.
+	virtual void GetCustomCameraNodeParameters(FCustomCameraNodeParameterInfos& OutParameterInfos) override;
+
+	// UObject interface.
+	virtual void PostLoad() override;
+	virtual void BeginDestroy() override;
+#if WITH_EDITOR
+	virtual void PostEditChangeProperty( struct FPropertyChangedEvent& PropertyChangedEvent) override;
+#endif
+
 public:
 
-	/** The camera node evaluator class to instantiate and run. */
-	UPROPERTY(EditAnywhere, Category=Common)
-	TSubclassOf<UBlueprintCameraNodeEvaluator> CameraNodeEvaluatorClass;
+	const UBlueprintCameraNodeEvaluator* GetCameraNodeEvaluatorTemplate() const { return CameraNodeEvaluatorTemplate; }
+
+private:
+
+	void RebuildOverrides();
+
+#if WITH_EDITOR
+	void OnObjectsReplaced(const TMap<UObject*, UObject*>& ReplacementMap);
+#endif
+
+private:
+
+	/** The camera node evaluator to instantiate and run. */
+	UPROPERTY(Instanced, EditAnywhere, Category=Common)
+	TObjectPtr<UBlueprintCameraNodeEvaluator> CameraNodeEvaluatorTemplate;
+
+	/** Overrides for the evaluator instance. */
+	UPROPERTY()
+	FCustomCameraNodeParameters CameraNodeEvaluatorOverrides;
+
+	// Deprecated.
+	
+	UPROPERTY()
+	TSubclassOf<UBlueprintCameraNodeEvaluator> CameraNodeEvaluatorClass_DEPRECATED;
+
+	friend class UE::Cameras::FBlueprintCameraNodeEvaluator;
 };
 
