Cameras: handle orphaned camera parameter pins

diff --git a/Engine/Plugins/Cameras/GameplayCameras/Source/GameplayCamerasEditor/Private/Editors/CameraNodeGraphSchema.cpp b/Engine/Plugins/Cameras/GameplayCameras/Source/GameplayCamerasEditor/Private/Editors/CameraNodeGraphSchema.cpp
--- a/Engine/Plugins/Cameras/GameplayCameras/Source/GameplayCamerasEditor/Private/Editors/CameraNodeGraphSchema.cpp
+++ b/Engine/Plugins/Cameras/GameplayCameras/Source/GameplayCamerasEditor/Private/Editors/CameraNodeGraphSchema.cpp
@@ -141,16 +141,35 @@ void UCameraNodeGraphSchema::OnCreateAllNodes(UObjectTreeGraph* InGraph, const F
 		}
 		
 		UCameraRigInterfaceParameterGraphNode* InterfaceParameterNode = CreateInterfaceParameterNode(InGraph, InterfaceParameter);
-		UObjectTreeGraphNode* const* CameraNodeNode = InCreatedNodes.CreatedNodes.Find(InterfaceParameter->Target);
+		UObjectTreeGraphNode* CameraNodeNode = InCreatedNodes.CreatedNodes.FindRef(InterfaceParameter->Target);
 		if (CameraNodeNode)
 		{
 			UEdGraphPin* InterfaceParameterSelfPin = InterfaceParameterNode->GetSelfPin();
-			UEdGraphPin* NodePin = FindPin(*CameraNodeNode, InterfaceParameter->TargetPropertyName, NAME_None);
-			ensure(NodePin &&
+			UEdGraphPin* NodePin = FindPin(CameraNodeNode, InterfaceParameter->TargetPropertyName, NAME_None);
+			if (NodePin &&
 					(NodePin->PinType.PinCategory == PC_CameraParameter ||
 					 NodePin->PinType.PinCategory == PC_CameraVariableReference ||
-					 NodePin->PinType.PinCategory == PC_CameraContextData));
-			InterfaceParameterSelfPin->MakeLinkTo(NodePin);
+					 NodePin->PinType.PinCategory == PC_CameraContextData))
+			{
+				InterfaceParameterSelfPin->MakeLinkTo(NodePin);
+			}
+			else
+			{
+				FName ErrorPinCategory = NAME_None;
+				if (InterfaceParameter->IsA<UCameraRigBlendableParameter>())
+				{
+					ErrorPinCategory = PC_CameraParameter;
+				}
+				else if (InterfaceParameter->IsA<UCameraRigDataParameter>())
+				{
+					ErrorPinCategory = PC_CameraContextData;
+				}
+
+				UEdGraphPin* ErrorPin = CameraNodeNode->CreatePin(
+						EGPD_Input, ErrorPinCategory, InterfaceParameter->TargetPropertyName);
+				InterfaceParameterSelfPin->MakeLinkTo(ErrorPin);
+				ErrorPin->bOrphanedPin = true;
+			}
 		}
 	}
 }
@@ -177,6 +196,12 @@ void UCameraNodeGraphSchema::GetGraphContextActions(FGraphContextMenuBuilder& Co
 		{
 			ensure(DraggedPin->PinName != NAME_None);
 
+			// If this is an invalid parameter/data pin, don't show any actions.
+			if (DraggedPin->bOrphanedPin)
+			{
+				return;
+			}
+
 			// Find the property being dragged, so we know what kind of parameter to create.
 			const UClass* CameraNodeClass = CameraNodeNode->GetObject()->GetClass();
 			FProperty* Property = CameraNodeClass->FindPropertyByName(DraggedPin->PinName);
@@ -316,7 +341,8 @@ const FPinConnectionResponse UCameraNodeGraphSchema::CanCreateConnection(const U
 	if ((A->PinType.PinCategory == PC_CameraParameter || 
 				A->PinType.PinCategory == PC_CameraVariableReference ||
 				A->PinType.PinCategory == PC_CameraContextData) && 
-			B->PinType.PinCategory == PC_Self)
+			B->PinType.PinCategory == PC_Self &&
+			!A->bOrphanedPin)
 	{
 		UCameraRigInterfaceParameterGraphNode* NodeB = Cast<UCameraRigInterfaceParameterGraphNode>(B->GetOwningNode());
 		if (NodeB)
@@ -339,7 +365,8 @@ const FPinConnectionResponse UCameraNodeGraphSchema::CanCreateConnection(const U
 	else if (A->PinType.PinCategory == PC_Self && 
 			(B->PinType.PinCategory == PC_CameraParameter || 
 			 B->PinType.PinCategory == PC_CameraVariableReference ||
-			 B->PinType.PinCategory == PC_CameraContextData))
+			 B->PinType.PinCategory == PC_CameraContextData) &&
+			!B->bOrphanedPin)
 	{
 		UCameraRigInterfaceParameterGraphNode* NodeA = Cast<UCameraRigInterfaceParameterGraphNode>(A->GetOwningNode());
 		if (NodeA)
@@ -412,6 +439,7 @@ bool UCameraNodeGraphSchema::OnBreakCustomPinLinks(UEdGraphPin& TargetPin) const
 	// See if we are in the situation of an interface parameter node being disconnected from 
 	// a camera node property pin.
 	UCameraRigInterfaceParameterGraphNode* InterfaceParameterNode = nullptr;
+
 	if (TargetPin.PinType.PinCategory == PC_CameraParameter ||
 			TargetPin.PinType.PinCategory == PC_CameraVariableReference ||
 			TargetPin.PinType.PinCategory == PC_CameraContextData)
diff --git a/Engine/Plugins/Cameras/GameplayCameras/Source/GameplayCamerasEditor/Private/Editors/ObjectTreeGraphNode.cpp b/Engine/Plugins/Cameras/GameplayCameras/Source/GameplayCamerasEditor/Private/Editors/ObjectTreeGraphNode.cpp
--- a/Engine/Plugins/Cameras/GameplayCameras/Source/GameplayCamerasEditor/Private/Editors/ObjectTreeGraphNode.cpp
+++ b/Engine/Plugins/Cameras/GameplayCameras/Source/GameplayCamerasEditor/Private/Editors/ObjectTreeGraphNode.cpp
@@ -164,6 +164,23 @@ void UObjectTreeGraphNode::AutowireNewNode(UEdGraphPin* FromPin)
 
 void UObjectTreeGraphNode::PinConnectionListChanged(UEdGraphPin* Pin)
 {
+	// Auto-remove orphaned pins when they are disconnected.
+	if (Pin->bOrphanedPin && Pin->LinkedTo.IsEmpty())
+	{
+		if (Pin->ParentPin)
+		{
+			Pin->ParentPin->SubPins.Remove(Pin);
+		}
+
+		RemovePin(Pin);
+
+		UEdGraph* OuterGraph = GetGraph();
+		if (OuterGraph)
+		{
+			OuterGraph->NotifyNodeChanged(this);
+		}
+	}
+
 	Super::PinConnectionListChanged(Pin);
 }
 
@@ -203,17 +220,39 @@ void UObjectTreeGraphNode::ReconstructNode()
 	AllocateDefaultPins();
 
 	// Rewire existing connections to new pins, matched by name, direction, and type.
+	TArray<UEdGraphPin*> ErrorPins;
 	for (UEdGraphPin* NewPin : Pins)
 	{
-		UEdGraphPin* OldPin = OldPins.FindRef(NewPin->PinName);
-		if (!OldPin || 
-				OldPin->Direction != NewPin->Direction ||
-				OldPin->PinType != NewPin->PinType)
+		UEdGraphPin* OldPin = nullptr;
+		OldPins.RemoveAndCopyValue(NewPin->PinName, OldPin);
+		if (OldPin)
 		{
-			continue;
+			const bool bOldMatchesNew = (
+				OldPin->Direction == NewPin->Direction &&
+				OldPin->PinType == NewPin->PinType);
+			if (bOldMatchesNew)
+			{
+				NewPin->MovePersistentDataFromOldPin(*OldPin);
+			}
+			else if (OldPin->LinkedTo.Num() > 0)
+			{
+				ErrorPins.Add(OldPin);
+			}
 		}
+	}
 
-		NewPin->MovePersistentDataFromOldPin(*OldPin);
+	// Old pins that had connections must be preserved, but made into orphans.
+	for (TPair<FName, UEdGraphPin*> Pair : OldPins)
+	{
+		if (Pair.Value->LinkedTo.Num() > 0)
+		{
+			ErrorPins.Add(Pair.Value);
+		}
+	}
+	for (UEdGraphPin* ErrorPin : ErrorPins)
+	{
+		Pins.Add(ErrorPin);
+		ErrorPin->bOrphanedPin = true;
 	}
 
 	GetGraph()->NotifyNodeChanged(this);
diff --git a/Engine/Plugins/Cameras/GameplayCameras/Source/GameplayCamerasEditor/Private/Editors/ObjectTreeGraphSchema.cpp b/Engine/Plugins/Cameras/GameplayCameras/Source/GameplayCamerasEditor/Private/Editors/ObjectTreeGraphSchema.cpp
--- a/Engine/Plugins/Cameras/GameplayCameras/Source/GameplayCamerasEditor/Private/Editors/ObjectTreeGraphSchema.cpp
+++ b/Engine/Plugins/Cameras/GameplayCameras/Source/GameplayCamerasEditor/Private/Editors/ObjectTreeGraphSchema.cpp
@@ -616,6 +616,11 @@ const FPinConnectionResponse UObjectTreeGraphSchema::CanCreateConnection(const U
 	{
 		return FPinConnectionResponse(CONNECT_RESPONSE_DISALLOW, TEXT("Unsupported node types"));
 	}
+	
+	if (A->bOrphanedPin || B->bOrphanedPin)
+	{
+		return FPinConnectionResponse(CONNECT_RESPONSE_DISALLOW, TEXT("Can't connect an orphaned pin"));
+	}
 
 	if (A->Direction == B->Direction)
 	{
