Cameras: expose pre-blending as an option for blendable parameters

Also add debug information about camera variable flags, so the user can tell if a variable is pre-blended, private, etc.

diff --git a/Engine/Plugins/Cameras/GameplayCameras/Source/GameplayCameras/Private/Core/CameraRigAssetBuilder.cpp b/Engine/Plugins/Cameras/GameplayCameras/Source/GameplayCameras/Private/Core/CameraRigAssetBuilder.cpp
--- a/Engine/Plugins/Cameras/GameplayCameras/Source/GameplayCameras/Private/Core/CameraRigAssetBuilder.cpp
+++ b/Engine/Plugins/Cameras/GameplayCameras/Source/GameplayCameras/Private/Core/CameraRigAssetBuilder.cpp
@@ -65,17 +65,17 @@ void SetPrivateVariableDefaultValue<UBooleanCameraVariable, bool>(UBooleanCamera
 	}
 }
 
-UCameraVariableAsset* CreatePrivateVariable(
-		UCameraRigAsset* CameraRig,
-		const FString& InterfaceParameterName,
-		ECameraVariableType ParameterType)
+FString MakePrivateVariableName(UCameraRigAsset* CameraRig, const UCameraRigBlendableParameter* BlendableParameter)
 {
-	const FString VariableName = FString::Format(
+	return FString::Format(
 			TEXT("Override_{0}_{1}"), 
-			{ CameraRig->GetName(), InterfaceParameterName });
+			{ CameraRig->GetName(), BlendableParameter->InterfaceParameterName });
+}
 
+UCameraVariableAsset* CreatePrivateVariable(UCameraRigAsset* CameraRig, const UCameraRigBlendableParameter* BlendableParameter)
+{
 	TSubclassOf<UCameraVariableAsset> VariableClass;
-	switch (ParameterType)
+	switch (BlendableParameter->ParameterType)
 	{
 #define UE_CAMERA_VARIABLE_FOR_TYPE(VariableType, VariableName)\
 		case ECameraVariableType::VariableName:\
@@ -89,17 +89,41 @@ UCameraVariableAsset* CreatePrivateVariable(
 		return nullptr;
 	}
 
+	const FString VariableName = MakePrivateVariableName(CameraRig, BlendableParameter);
 	UCameraVariableAsset* PrivateVariable = NewObject<UCameraVariableAsset>(
 			CameraRig, VariableClass, FName(*VariableName), RF_Transactional);
 
-	// Make sure it's a private input variable.
-	PrivateVariable->bIsInput = true;
+	// Make sure it's a private variable.
 	PrivateVariable->bIsPrivate = true;
 	PrivateVariable->bAutoReset = false;
 
+	// Make it an input variable if the parameter is pre-blended.
+	PrivateVariable->bIsInput = BlendableParameter->bIsPreBlended;
+
 	return PrivateVariable;
 }
 
+void UpdatePrivateVariable(UCameraRigAsset* CameraRig, const UCameraRigBlendableParameter* BlendableParameter, UCameraVariableAsset* PrivateVariable)
+{
+	check(PrivateVariable);
+
+	ensure(PrivateVariable->bIsPrivate);
+	ensure(!PrivateVariable->bAutoReset);
+
+	const FString VariableName = MakePrivateVariableName(CameraRig, BlendableParameter);
+	if (PrivateVariable->GetName() != VariableName)
+	{
+		ensure(!PrivateVariable->HasAnyFlags(RF_Public));
+		PrivateVariable->Rename(*VariableName, nullptr);
+	}
+
+	if (PrivateVariable->bIsInput != BlendableParameter->bIsPreBlended)
+	{
+		PrivateVariable->Modify();
+		PrivateVariable->bIsInput = BlendableParameter->bIsPreBlended;
+	}
+}
+
 struct FInterfaceParameterBindingBuilder
 {
 	UCameraRigAsset* CameraRig;
@@ -671,6 +695,7 @@ void FCameraRigAssetBuilder::BuildInterfaceParameters()
 		UCameraVariableAsset* PrivateVariable = BlendableParameter->PrivateVariable;
 		if (PrivateVariable && PrivateVariable->GetVariableType() == BlendableParameter->ParameterType)
 		{
+			UpdatePrivateVariable(CameraRig, BlendableParameter, PrivateVariable);
 			continue;
 		}
 
@@ -686,8 +711,7 @@ void FCameraRigAssetBuilder::BuildInterfaceParameters()
 			BlendableParameter->PrivateVariable = nullptr;
 		}
 
-		BlendableParameter->PrivateVariable = CreatePrivateVariable(
-				CameraRig, BlendableParameter->InterfaceParameterName, BlendableParameter->ParameterType);
+		BlendableParameter->PrivateVariable = CreatePrivateVariable(CameraRig, BlendableParameter);
 	}
 
 	for (auto It = CameraRig->Interface.DataParameters.CreateIterator(); It; ++It)
diff --git a/Engine/Plugins/Cameras/GameplayCameras/Source/GameplayCameras/Private/Debug/VariableTableDebugBlock.cpp b/Engine/Plugins/Cameras/GameplayCameras/Source/GameplayCameras/Private/Debug/VariableTableDebugBlock.cpp
--- a/Engine/Plugins/Cameras/GameplayCameras/Source/GameplayCameras/Private/Debug/VariableTableDebugBlock.cpp
+++ b/Engine/Plugins/Cameras/GameplayCameras/Source/GameplayCameras/Private/Debug/VariableTableDebugBlock.cpp
@@ -49,6 +49,8 @@ void FVariableTableDebugBlock::Initialize(const FCameraVariableTable& InVariable
 #undef UE_CAMERA_VARIABLE_FOR_TYPE
 
 		FEntryDebugInfo EntryDebugInfo{ Entry.ID.GetValue(), EntryName, EntryValueStr};
+		EntryDebugInfo.bIsInput = EnumHasAnyFlags(Entry.Flags, FCameraVariableTable::EEntryFlags::Input);
+		EntryDebugInfo.bIsPrivate = EnumHasAnyFlags(Entry.Flags, FCameraVariableTable::EEntryFlags::Private);
 		EntryDebugInfo.bWritten = EnumHasAnyFlags(Entry.Flags, FCameraVariableTable::EEntryFlags::Written);
 		EntryDebugInfo.bWrittenThisFrame = EnumHasAnyFlags(Entry.Flags, FCameraVariableTable::EEntryFlags::WrittenThisFrame);
 		Entries.Add(EntryDebugInfo);
@@ -87,6 +89,7 @@ void FVariableTableDebugBlock::OnDebugDraw(const FCameraDebugBlockDrawParams& Pa
 		{
 			Renderer.AddText(TEXT("{cam_passive}[%d]{cam_default} "), Entry.ID);
 		}
+
 		if (!Entry.Name.IsEmpty())
 		{
 			Renderer.AddText(TEXT("%s : "), *Entry.Name);
@@ -102,9 +105,10 @@ void FVariableTableDebugBlock::OnDebugDraw(const FCameraDebugBlockDrawParams& Pa
 		if (Entry.bWritten)
 		{
 			Renderer.AddText(Entry.Value);
+			Renderer.AddText(TEXT(" "));
 			if (Entry.bWrittenThisFrame)
 			{
-				Renderer.AddText(TEXT(" {cam_passive}[WrittenThisFrame]"));
+				Renderer.AddText(TEXT("{cam_passive}[WrittenThisFrame]"));
 			}
 		}
 		else
@@ -112,6 +116,16 @@ void FVariableTableDebugBlock::OnDebugDraw(const FCameraDebugBlockDrawParams& Pa
 			Renderer.AddText("{cam_warning}[Uninitialized]");
 		}
 
+		if (Entry.bIsInput)
+		{
+			Renderer.AddText("{cam_notice2}[Input]");
+		}
+
+		if (Entry.bIsPrivate)
+		{
+			Renderer.AddText("{cam_notice2}[Private]");
+		}
+
 		Renderer.NewLine();
 		Renderer.SetTextColor(Colors.Default);
 	}
@@ -128,6 +142,8 @@ FArchive& operator<< (FArchive& Ar, FVariableTableDebugBlock::FEntryDebugInfo& E
 	Ar << EntryDebugInfo.ID;
 	Ar << EntryDebugInfo.Name;
 	Ar << EntryDebugInfo.Value;
+	Ar << EntryDebugInfo.bIsInput;
+	Ar << EntryDebugInfo.bIsPrivate;
 	Ar << EntryDebugInfo.bWritten;
 	Ar << EntryDebugInfo.bWrittenThisFrame;
 	return Ar;
diff --git a/Engine/Plugins/Cameras/GameplayCameras/Source/GameplayCameras/Public/Core/CameraRigAsset.h b/Engine/Plugins/Cameras/GameplayCameras/Source/GameplayCameras/Public/Core/CameraRigAsset.h
--- a/Engine/Plugins/Cameras/GameplayCameras/Source/GameplayCameras/Public/Core/CameraRigAsset.h
+++ b/Engine/Plugins/Cameras/GameplayCameras/Source/GameplayCameras/Public/Core/CameraRigAsset.h
@@ -173,6 +173,16 @@ public:
 	UPROPERTY()
 	ECameraVariableType ParameterType = ECameraVariableType::Boolean;
 
+	/**
+	 * Whether this parameter's value should be pre-blended.
+	 *
+	 * Pre-blending means that if two blending camera rigs share this parameter, 
+	 * each of their values will be blended in a first evaluation pass, and then
+	 * both camera rigs will evaluate with the same blended value.
+	 */
+	UPROPERTY()
+	bool bIsPreBlended = true;
+
 	// Built on save/cook.
 
 	/**
diff --git a/Engine/Plugins/Cameras/GameplayCameras/Source/GameplayCameras/Public/Debug/VariableTableDebugBlock.h b/Engine/Plugins/Cameras/GameplayCameras/Source/GameplayCameras/Public/Debug/VariableTableDebugBlock.h
--- a/Engine/Plugins/Cameras/GameplayCameras/Source/GameplayCameras/Public/Debug/VariableTableDebugBlock.h
+++ b/Engine/Plugins/Cameras/GameplayCameras/Source/GameplayCameras/Public/Debug/VariableTableDebugBlock.h
@@ -46,6 +46,8 @@ private:
 		uint32 ID;
 		FString Name;
 		FString Value;
+		bool bIsInput;
+		bool bIsPrivate;
 		bool bWritten;
 		bool bWrittenThisFrame;
 	};
diff --git a/Engine/Plugins/Cameras/GameplayCameras/Source/GameplayCamerasEditor/Private/Toolkits/CameraRigInterfaceParametersToolkit.cpp b/Engine/Plugins/Cameras/GameplayCameras/Source/GameplayCamerasEditor/Private/Toolkits/CameraRigInterfaceParametersToolkit.cpp
--- a/Engine/Plugins/Cameras/GameplayCameras/Source/GameplayCamerasEditor/Private/Toolkits/CameraRigInterfaceParametersToolkit.cpp
+++ b/Engine/Plugins/Cameras/GameplayCameras/Source/GameplayCamerasEditor/Private/Toolkits/CameraRigInterfaceParametersToolkit.cpp
@@ -14,6 +14,7 @@
 #include "ToolMenus.h"
 #include "Types/SlateEnums.h"
 #include "Widgets/Input/SButton.h"
+#include "Widgets/Input/SCheckBox.h"
 #include "Widgets/Input/SComboButton.h"
 #include "Widgets/Layout/SBox.h"
 #include "Widgets/Layout/SSplitter.h"
@@ -32,6 +33,7 @@ namespace UE::Cameras
 static const FName ParameterTypeColumn(TEXT("ParameterType"));
 static const FName ParameterNameColumn(TEXT("ParameterName"));
 static const FName ParameterMessageColumn(TEXT("ParameterMessage"));
+static const FName ParameterIsPreBlendedColumn(TEXT("ParameterIsPreBlended"));
 
 /**
  * List entry for any interface parameter panel.
@@ -166,6 +168,12 @@ protected:
 					.bAllowArrays(false)
 				];
 		}
+		else if (InColumnName == ParameterIsPreBlendedColumn)
+		{
+			return SNew(SCheckBox)
+				.IsChecked(this, &SCameraRigBlendableParameterTableRow::IsBlendableParameterPreBlended)
+				.OnCheckStateChanged(this, &SCameraRigBlendableParameterTableRow::OnBlendableParameterPreBlendedChanged);
+		}
 
 		return SCameraRigInterfaceParameterTableRowBase<UCameraRigBlendableParameter>::GenerateWidgetForColumn(InColumnName);
 	}
@@ -345,6 +353,23 @@ protected:
 		}
 	}
 
+	ECheckBoxState IsBlendableParameterPreBlended() const
+	{
+		return Item->bIsPreBlended ? ECheckBoxState::Checked : ECheckBoxState::Unchecked;
+	}
+	
+	void OnBlendableParameterPreBlendedChanged(ECheckBoxState CheckState)
+	{
+		const bool bIsPreBlended = (CheckState == ECheckBoxState::Checked);
+		if (Item->bIsPreBlended != bIsPreBlended)
+		{
+			const FScopedTransaction Transaction(LOCTEXT("ChangeBlendableParameterIsPreBlended", "Change Blendable Parameter Pre-Blending"));
+
+			Item->Modify();
+			Item->bIsPreBlended = bIsPreBlended;
+		}
+	}
+
 	bool IsPinTypeSelectorReadOnly() const
 	{
 		return !IsHovered();
@@ -650,6 +675,10 @@ void SCameraRigInterfaceParametersPanel::Construct(const FArguments& Args, FCame
 					.FillWidth(0.5f)
 					.DefaultLabel(LOCTEXT("ParameterNameColumnLabel", "Name"))
 
+					+SHeaderRow::Column(ParameterIsPreBlendedColumn)
+					.ManualWidth(60)
+					.DefaultLabel(LOCTEXT("ParameterIsPreBlendedColumnLabel", "Pre-Blend"))
+
 					+SHeaderRow::Column(ParameterMessageColumn)
 					.FillWidth(0.2f)
 					.DefaultLabel(LOCTEXT("ParameterMessageColumnLabel", "Note"))
