Cameras: remove the auto-reset variable service

We can actually achieve pretty much the same result by flagging entries in the variable table as "unwritten".

However, to properly reset these variables at the correct time, we also change a bit how the end-of-update happens. We generate debug info immediately if needed, so that the debug info gets the proper variable flags to show in the debug HUD. Then we can reset everything. The debug update now simply prints the generated debug info.

diff --git a/Engine/Plugins/Cameras/GameplayCameras/Source/GameplayCameras/Private/Core/CameraEvaluationContext.cpp b/Engine/Plugins/Cameras/GameplayCameras/Source/GameplayCameras/Private/Core/CameraEvaluationContext.cpp
--- a/Engine/Plugins/Cameras/GameplayCameras/Source/GameplayCameras/Private/Core/CameraEvaluationContext.cpp
+++ b/Engine/Plugins/Cameras/GameplayCameras/Source/GameplayCameras/Private/Core/CameraEvaluationContext.cpp
@@ -64,6 +64,13 @@ void FCameraEvaluationContext::AddReferencedObjects(FReferenceCollector& Collect
 	}
 }
 
+void FCameraEvaluationContext::OnEndCameraSystemUpdate()
+{
+	InitialResult.VariableTable.AutoResetValues();
+	InitialResult.VariableTable.ClearAllWrittenThisFrameFlags();
+	InitialResult.ContextDataTable.ClearAllWrittenThisFrameFlags();
+}
+
 void FCameraEvaluationContext::AutoCreateDirectorEvaluator()
 {
 	if (DirectorEvaluator == nullptr)
diff --git a/Engine/Plugins/Cameras/GameplayCameras/Source/GameplayCameras/Private/Core/CameraEvaluationContextStack.cpp b/Engine/Plugins/Cameras/GameplayCameras/Source/GameplayCameras/Private/Core/CameraEvaluationContextStack.cpp
--- a/Engine/Plugins/Cameras/GameplayCameras/Source/GameplayCameras/Private/Core/CameraEvaluationContextStack.cpp
+++ b/Engine/Plugins/Cameras/GameplayCameras/Source/GameplayCameras/Private/Core/CameraEvaluationContextStack.cpp
@@ -160,7 +160,7 @@ void FCameraEvaluationContextStack::OnEndCameraSystemUpdate()
 	while (!ContextsToVisit.IsEmpty())
 	{
 		TSharedPtr<FCameraEvaluationContext> Context = ContextsToVisit.Pop();
-		Context->GetInitialResult().VariableTable.ClearAllWrittenThisFrameFlags();
+		Context->OnEndCameraSystemUpdate();
 
 		TArrayView<const TSharedPtr<FCameraEvaluationContext>> ChildrenContexts(Context->GetChildrenContexts());
 		for (TSharedPtr<FCameraEvaluationContext> ChildContext : ReverseIterate(ChildrenContexts))
diff --git a/Engine/Plugins/Cameras/GameplayCameras/Source/GameplayCameras/Private/Core/CameraRigAsset.cpp b/Engine/Plugins/Cameras/GameplayCameras/Source/GameplayCameras/Private/Core/CameraRigAsset.cpp
--- a/Engine/Plugins/Cameras/GameplayCameras/Source/GameplayCameras/Private/Core/CameraRigAsset.cpp
+++ b/Engine/Plugins/Cameras/GameplayCameras/Source/GameplayCameras/Private/Core/CameraRigAsset.cpp
@@ -17,7 +17,6 @@ void FCameraRigAllocationInfo::Append(const FCameraRigAllocationInfo& OtherAlloc
 	EvaluatorInfo.TotalSizeof = Align(EvaluatorInfo.TotalSizeof, OtherEvaluatorInfo.MaxAlignof) + OtherEvaluatorInfo.TotalSizeof;
 
 	const FCameraVariableTableAllocationInfo& OtherVariableTableInfo(OtherAllocationInfo.VariableTableInfo);
-	VariableTableInfo.AutoResetVariables.Append(OtherVariableTableInfo.AutoResetVariables);
 	VariableTableInfo.VariableDefinitions.Append(OtherVariableTableInfo.VariableDefinitions);
 
 	const FCameraContextDataAllocationInfo& OtherContextDataTableInfo(OtherAllocationInfo.ContextDataTableInfo);
diff --git a/Engine/Plugins/Cameras/GameplayCameras/Source/GameplayCameras/Private/Core/CameraRigAssetBuilder.cpp b/Engine/Plugins/Cameras/GameplayCameras/Source/GameplayCameras/Private/Core/CameraRigAssetBuilder.cpp
--- a/Engine/Plugins/Cameras/GameplayCameras/Source/GameplayCameras/Private/Core/CameraRigAssetBuilder.cpp
+++ b/Engine/Plugins/Cameras/GameplayCameras/Source/GameplayCameras/Private/Core/CameraRigAssetBuilder.cpp
@@ -419,10 +419,6 @@ void AddCameraVariableToAllocationInfo(UCameraVariableAsset* Variable, FCameraVa
 	{
 		FCameraVariableDefinition VariableDefinition = Variable->GetVariableDefinition();
 		AllocationInfo.VariableDefinitions.Add(VariableDefinition);
-		if (Variable->bAutoReset)
-		{
-			AllocationInfo.AutoResetVariables.Add(Variable);
-		}
 	}
 }
 
diff --git a/Engine/Plugins/Cameras/GameplayCameras/Source/GameplayCameras/Private/Core/CameraSystemEvaluator.cpp b/Engine/Plugins/Cameras/GameplayCameras/Source/GameplayCameras/Private/Core/CameraSystemEvaluator.cpp
--- a/Engine/Plugins/Cameras/GameplayCameras/Source/GameplayCameras/Private/Core/CameraSystemEvaluator.cpp
+++ b/Engine/Plugins/Cameras/GameplayCameras/Source/GameplayCameras/Private/Core/CameraSystemEvaluator.cpp
@@ -15,7 +15,6 @@
 #include "Debug/CameraSystemTrace.h"
 #include "Debug/RootCameraDebugBlock.h"
 #include "GameplayCamerasSettings.h"
-#include "Services/AutoResetCameraVariableService.h"
 #include "Services/OrientationInitializationService.h"
 #include "UObject/Package.h"
 #include "UObject/UObjectGlobals.h"
@@ -90,8 +89,6 @@ void FCameraSystemEvaluator::Initialize(const FCameraSystemEvaluatorCreateParams
 		RootEvaluator->Initialize(InitParams, RootNodeResult);
 	}
 
-	VariableAutoResetService = MakeShared<FAutoResetCameraVariableService>();
-	RegisterEvaluationService(VariableAutoResetService.ToSharedRef());
 	RegisterEvaluationService(MakeShared<FOrientationInitializationService>());
 
 	CameraRigCombinationRegistry = MakeShared<FCameraRigCombinationRegistry>();
@@ -201,13 +198,10 @@ void FCameraSystemEvaluator::Update(const FCameraSystemEvaluationParams& Params)
 	// Reset our result' flags.
 	RootNodeResult.CameraPose.ClearAllChangedFlags();
 	RootNodeResult.VariableTable.ClearAllWrittenThisFrameFlags();
+	RootNodeResult.ContextDataTable.ClearAllWrittenThisFrameFlags();
 
-	// Run the variable auto-reset service here, because the other (third party) services
-	// should get the reset variable values.
-	if (VariableAutoResetService)
-	{
-		VariableAutoResetService->PerformVariableResets(RootNodeResult.VariableTable, ContextStack);
-	}
+	// Reset variables and data.
+	RootNodeResult.VariableTable.AutoResetValues();
 
 	// Pre-update all services.
 	PreUpdateServices(Params.DeltaTime, ECameraEvaluationServiceFlags::None);
@@ -297,6 +291,11 @@ void FCameraSystemEvaluator::Update(const FCameraSystemEvaluationParams& Params)
 	// Harvest the result.
 	Result.Reset(RootNodeResult);
 
+	// Generate debug information if needed.
+#if UE_GAMEPLAY_CAMERAS_DEBUG
+	BuildDebugBlocksIfNeeded();
+#endif  // UE_GAMEPLAY_CAMERAS_DEBUG
+
 	// End of update things...
 	ContextStack.OnEndCameraSystemUpdate();
 }
@@ -359,24 +358,27 @@ void FCameraSystemEvaluator::GetEvaluatedCameraView(FMinimalViewInfo& DesiredVie
 
 #if UE_GAMEPLAY_CAMERAS_DEBUG
 
-void FCameraSystemEvaluator::DebugUpdate(const FCameraSystemDebugUpdateParams& Params)
+bool FCameraSystemEvaluator::IsDebugTraceEnabled()
 {
 #if UE_GAMEPLAY_CAMERAS_TRACE
-	const bool bTraceEnabled = FCameraSystemTrace::IsTraceEnabled();
+	return FCameraSystemTrace::IsTraceEnabled();
 #else
-	const bool bTraceEnabled = false;
+	return false;
 #endif  // UE_GAMEPLAY_CAMERAS_TRACE
-	if (!bTraceEnabled && !GGameplayCamerasDebugEnable)
-	{
-		return;
-	}
+}
 
-#if UE_GAMEPLAY_CAMERAS_TRACE
-	if (FCameraSystemTrace::IsTraceReplay())
+bool FCameraSystemEvaluator::ShouldBuildOrDrawDebugBlocks()
+{
+	const bool bTraceEnabled = IsDebugTraceEnabled();
+	return bTraceEnabled || GGameplayCamerasDebugEnable;
+}
+
+void FCameraSystemEvaluator::BuildDebugBlocksIfNeeded()
+{
+	if (!ShouldBuildOrDrawDebugBlocks())
 	{
 		return;
 	}
-#endif  // UE_GAMEPLAY_CAMERAS_TRACE
 
 	// Clear previous frame's debug info and make room for this frame's.
 	DebugBlockStorage.DestroyDebugBlocks();
@@ -387,12 +389,20 @@ void FCameraSystemEvaluator::DebugUpdate(const FCameraSystemDebugUpdateParams& P
 	FCameraDebugBlockBuildParams BuildParams;
 	FCameraDebugBlockBuilder DebugBlockBuilder(DebugBlockStorage, *RootDebugBlock);
 	RootDebugBlock->BuildDebugBlocks(*this, BuildParams, DebugBlockBuilder);
+}
+
+void FCameraSystemEvaluator::DebugUpdate(const FCameraSystemDebugUpdateParams& Params)
+{
+	if (!ShouldBuildOrDrawDebugBlocks() || !RootDebugBlock)
+	{
+		return;
+	}
 
 	UObject* Owner = WeakOwner.Get();
 	UWorld* OwnerWorld = Owner ? Owner->GetWorld() : nullptr;
 
 #if UE_GAMEPLAY_CAMERAS_TRACE
-	if (bTraceEnabled)
+	if (IsDebugTraceEnabled())
 	{
 		FCameraSystemTrace::TraceEvaluation(OwnerWorld, Result, *RootDebugBlock);
 	}
diff --git a/Engine/Plugins/Cameras/GameplayCameras/Source/GameplayCameras/Private/Core/CameraVariableAssets.cpp b/Engine/Plugins/Cameras/GameplayCameras/Source/GameplayCameras/Private/Core/CameraVariableAssets.cpp
--- a/Engine/Plugins/Cameras/GameplayCameras/Source/GameplayCameras/Private/Core/CameraVariableAssets.cpp
+++ b/Engine/Plugins/Cameras/GameplayCameras/Source/GameplayCameras/Private/Core/CameraVariableAssets.cpp
@@ -22,6 +22,7 @@ FCameraVariableDefinition UCameraVariableAsset::GetVariableDefinition() const
 	VariableDefinition.VariableType = GetVariableType();
 	VariableDefinition.bIsPrivate = bIsPrivate;
 	VariableDefinition.bIsInput = bIsInput;
+	VariableDefinition.bAutoReset = bAutoReset;
 #if WITH_EDITORONLY_DATA
 	VariableDefinition.VariableName = GetDisplayName();
 #endif
diff --git a/Engine/Plugins/Cameras/GameplayCameras/Source/GameplayCameras/Private/Core/CameraVariableTable.cpp b/Engine/Plugins/Cameras/GameplayCameras/Source/GameplayCameras/Private/Core/CameraVariableTable.cpp
--- a/Engine/Plugins/Cameras/GameplayCameras/Source/GameplayCameras/Private/Core/CameraVariableTable.cpp
+++ b/Engine/Plugins/Cameras/GameplayCameras/Source/GameplayCameras/Private/Core/CameraVariableTable.cpp
@@ -101,6 +101,10 @@ void FCameraVariableTable::Initialize(const FCameraVariableTableAllocationInfo&
 		{
 			NewEntry.Flags |= EEntryFlags::Input;
 		}
+		if (VariableDefinition.bAutoReset)
+		{
+			NewEntry.Flags |= EEntryFlags::AutoReset;
+		}
 #if WITH_EDITORONLY_DATA
 		NewEntry.DebugName = VariableDefinition.VariableName;
 #endif
@@ -163,6 +167,10 @@ void FCameraVariableTable::AddVariable(const FCameraVariableDefinition& Variable
 	{
 		NewEntry.Flags |= EEntryFlags::Input;
 	}
+	if (VariableDefinition.bAutoReset)
+	{
+		NewEntry.Flags |= EEntryFlags::AutoReset;
+	}
 #if WITH_EDITORONLY_DATA
 	NewEntry.DebugName = VariableDefinition.VariableName;
 #endif
@@ -333,6 +341,17 @@ void FCameraVariableTable::ClearAllWrittenThisFrameFlags()
 	}
 }
 
+void FCameraVariableTable::AutoResetValues()
+{
+	for (FEntry& Entry : Entries)
+	{
+		if (EnumHasAnyFlags(Entry.Flags, EEntryFlags::AutoReset))
+		{
+			EnumRemoveFlags(Entry.Flags, EEntryFlags::Written | EEntryFlags::WrittenThisFrame);
+		}
+	}
+}
+
 bool FCameraVariableTable::TryGetVariableDefinition(FCameraVariableID VariableID, FCameraVariableDefinition& OutVariableDefinition) const
 {
 	if (const FEntry* Entry = FindEntry(VariableID))
@@ -641,18 +660,10 @@ void FCameraVariableTableAllocationInfo::Combine(const FCameraVariableTableAlloc
 			ensure(KnownVariableDefinition == OtherVariableDefinition);
 		}
 	}
-
-	TSet<TObjectPtr<UCameraVariableAsset>> KnownAutoResetVariables(AutoResetVariables);
-	for (UCameraVariableAsset* OtherAutoResetVariable : OtherInfo.AutoResetVariables)
-	{
-		KnownAutoResetVariables.Add(OtherAutoResetVariable);
-	}
-	AutoResetVariables = KnownAutoResetVariables.Array();
 }
 
 bool operator==(const FCameraVariableTableAllocationInfo& A, const FCameraVariableTableAllocationInfo& B)
 {
-	return A.VariableDefinitions == B.VariableDefinitions
-		&& A.AutoResetVariables == B.AutoResetVariables;
+	return A.VariableDefinitions == B.VariableDefinitions;
 }
 
diff --git a/Engine/Plugins/Cameras/GameplayCameras/Source/GameplayCameras/Private/Core/RootCameraNode.cpp b/Engine/Plugins/Cameras/GameplayCameras/Source/GameplayCameras/Private/Core/RootCameraNode.cpp
--- a/Engine/Plugins/Cameras/GameplayCameras/Source/GameplayCameras/Private/Core/RootCameraNode.cpp
+++ b/Engine/Plugins/Cameras/GameplayCameras/Source/GameplayCameras/Private/Core/RootCameraNode.cpp
@@ -5,7 +5,6 @@
 #include "Core/CameraEvaluationService.h"
 #include "Core/CameraSystemEvaluator.h"
 #include "Core/RootCameraNodeCameraRigEvent.h"
-#include "Services/AutoResetCameraVariableService.h"
 
 #include UE_INLINE_GENERATED_CPP_BY_NAME(RootCameraNode)
 
@@ -39,14 +38,11 @@ void FRootCameraNodeEvaluator::BuildSingleCameraRigHierarchy(const FSingleCamera
 
 void FRootCameraNodeEvaluator::RunSingleCameraRig(const FSingleCameraRigEvaluationParams& Params, FCameraNodeEvaluationResult& OutResult)
 {
-	// Before we do the actual evaluation we need to ask the system to auto-reset
-	// any camera variable that needs auto-resetting. Otherwise, we might end up with
-	// an update result that isn't representative of what would happen normally.
-	if (FCameraSystemEvaluator* Evaluator = Params.EvaluationParams.Evaluator)
-	{
-		// TODO: we might have to reset variables on the context's initial result too?
-		Evaluator->VariableAutoResetService->PerformVariableResets(OutResult.VariableTable);
-	}
+	// Before we do the actual evaluation we need to auto-reset any camera variable 
+	// that needs auto-resetting. Otherwise, we might end up with an update result 
+	// that isn't representative of what would happen normally.
+	// TODO: we might have to reset variables on the context's initial result too?
+	OutResult.VariableTable.AutoResetValues();
 
 	OnRunSingleCameraRig(Params, OutResult);
 }
diff --git a/Engine/Plugins/Cameras/GameplayCameras/Source/GameplayCameras/Private/Debug/CameraSystemTrace.cpp b/Engine/Plugins/Cameras/GameplayCameras/Source/GameplayCameras/Private/Debug/CameraSystemTrace.cpp
--- a/Engine/Plugins/Cameras/GameplayCameras/Source/GameplayCameras/Private/Debug/CameraSystemTrace.cpp
+++ b/Engine/Plugins/Cameras/GameplayCameras/Source/GameplayCameras/Private/Debug/CameraSystemTrace.cpp
@@ -276,18 +276,6 @@ FString FCameraSystemTrace::ChannelName("CameraSystemChannel");
 FString FCameraSystemTrace::LoggerName("CameraSystem");
 FString FCameraSystemTrace::EvaluationEventName("CameraSystemEvaluation");
 
-bool FCameraSystemTrace::bIsReplaying(false);
-
-bool FCameraSystemTrace::IsTraceReplay()
-{
-	return bIsReplaying;
-}
-
-void FCameraSystemTrace::SetTraceReplay(bool bInIsReplaying)
-{
-	bIsReplaying = bInIsReplaying;
-}
-
 bool FCameraSystemTrace::IsTraceEnabled()
 {
 	return GGameplayCamerasDebugTrace || UE_TRACE_CHANNELEXPR_IS_ENABLED(CameraSystemChannel);
diff --git a/Engine/Plugins/Cameras/GameplayCameras/Source/GameplayCameras/Private/Directors/BlueprintCameraDirector.cpp b/Engine/Plugins/Cameras/GameplayCameras/Source/GameplayCameras/Private/Directors/BlueprintCameraDirector.cpp
--- a/Engine/Plugins/Cameras/GameplayCameras/Source/GameplayCameras/Private/Directors/BlueprintCameraDirector.cpp
+++ b/Engine/Plugins/Cameras/GameplayCameras/Source/GameplayCameras/Private/Directors/BlueprintCameraDirector.cpp
@@ -17,7 +17,6 @@
 #include "GameFramework/ControllerGameplayCameraEvaluationComponent.h"
 #include "GameplayCameras.h"
 #include "Helpers/OutgoingReferenceFinder.h"
-#include "Services/AutoResetCameraVariableService.h"
 
 #include UE_INLINE_GENERATED_CPP_BY_NAME(BlueprintCameraDirector)
 
@@ -345,7 +344,7 @@ FBlueprintCameraVariableTable UBlueprintCameraDirectorEvaluator::GetInitialConte
 	if (EvaluationContext)
 	{
 		FCameraVariableTable& VariableTable = EvaluationContext->GetInitialResult().VariableTable;
-		return FBlueprintCameraVariableTable(&VariableTable, VariableAutoResetService);
+		return FBlueprintCameraVariableTable(&VariableTable);
 	}
 	else
 	{
@@ -375,7 +374,6 @@ void UBlueprintCameraDirectorEvaluator::NativeActivateCameraDirector(const UE::C
 	using namespace UE::Cameras;
 
 	EvaluationContext = Params.OwnerContext;
-	VariableAutoResetService = Params.Evaluator->FindEvaluationService<FAutoResetCameraVariableService>();
 
 	EvaluationResult.Reset();
 	{
@@ -402,7 +400,6 @@ void UBlueprintCameraDirectorEvaluator::NativeDeactivateCameraDirector(const UE:
 		DeactivateCameraDirector(BlueprintParams);
 	}
 
-	VariableAutoResetService = nullptr;
 	EvaluationContext = nullptr;
 }
 
diff --git a/Engine/Plugins/Cameras/GameplayCameras/Source/GameplayCameras/Private/GameFramework/BlueprintCameraVariableTable.cpp b/Engine/Plugins/Cameras/GameplayCameras/Source/GameplayCameras/Private/GameFramework/BlueprintCameraVariableTable.cpp
--- a/Engine/Plugins/Cameras/GameplayCameras/Source/GameplayCameras/Private/GameFramework/BlueprintCameraVariableTable.cpp
+++ b/Engine/Plugins/Cameras/GameplayCameras/Source/GameplayCameras/Private/GameFramework/BlueprintCameraVariableTable.cpp
@@ -4,7 +4,6 @@
 
 #include "Core/CameraVariableAssets.h"
 #include "Core/CameraVariableTable.h"
-#include "Services/AutoResetCameraVariableService.h"
 
 #include UE_INLINE_GENERATED_CPP_BY_NAME(BlueprintCameraVariableTable)
 
@@ -31,20 +30,15 @@
 #define UE_PRIVATE_BLUEPRINT_CAMERA_VARIABLE_TABLE_SET_VARIABLE(VariableType)\
 	UE_PRIVATE_BLUEPRINT_CAMERA_VARIABLE_TABLE_VALIDATE()\
 	UE_PRIVATE_BLUEPRINT_CAMERA_VARIABLE_TABLE_VALIDATE_PARAM()\
-	VariableTable.GetVariableTable()->SetValue(Variable, Value, true);\
-	if (Variable->bAutoReset && VariableTable.VariableAutoResetService)\
-	{\
-		VariableTable.VariableAutoResetService->RegisterVariableUseFromScripting(Variable);\
-	}
+	VariableTable.GetVariableTable()->SetValue(Variable, Value, true);
 
 FBlueprintCameraVariableTable::FBlueprintCameraVariableTable()
 	: PrivateVariableTable(nullptr)
 {
 }
 
-FBlueprintCameraVariableTable::FBlueprintCameraVariableTable(FCameraVariableTable* InVariableTable, TSharedPtr<UE::Cameras::FAutoResetCameraVariableService> InVariableAutoResetService)
+FBlueprintCameraVariableTable::FBlueprintCameraVariableTable(FCameraVariableTable* InVariableTable)
 	: PrivateVariableTable(InVariableTable)
-	, VariableAutoResetService(InVariableAutoResetService)
 {
 }
 
diff --git a/Engine/Plugins/Cameras/GameplayCameras/Source/GameplayCameras/Private/GameFramework/GameplayCameraComponent.cpp b/Engine/Plugins/Cameras/GameplayCameras/Source/GameplayCameras/Private/GameFramework/GameplayCameraComponent.cpp
--- a/Engine/Plugins/Cameras/GameplayCameras/Source/GameplayCameras/Private/GameFramework/GameplayCameraComponent.cpp
+++ b/Engine/Plugins/Cameras/GameplayCameras/Source/GameplayCameras/Private/GameFramework/GameplayCameraComponent.cpp
@@ -14,7 +14,6 @@
 #include "GameFramework/GameplayCameraSystemHost.h"
 #include "GameplayCameras.h"
 #include "Kismet/GameplayStatics.h"
-#include "Services/AutoResetCameraVariableService.h"
 #include "UObject/ConstructorHelpers.h"
 
 #include UE_INLINE_GENERATED_CPP_BY_NAME(GameplayCameraComponent)
@@ -168,63 +167,46 @@ void UGameplayCameraComponent::ActivateCameraEvaluationContext(APlayerController
 	Activate();
 }
 
+#define UE_PRIVATE_GAMEPLAY_CAMERA_COMPONENT_VALIDATE_EVALUATION_CONTEXT(ErrorMsg, ErrorResult)\
+	using namespace UE::Cameras;\
+	if (!EvaluationContext)\
+	{\
+		FFrame::KismetExecutionMessage(\
+				*FString::Format(\
+					TEXT(#ErrorResult " on Gameplay Camera component '{0}': it isn't active."),\
+					{ *GetNameSafe(this) }),\
+				ELogVerbosity::Error);\
+		return ErrorResult;\
+	}
+
 FBlueprintCameraPose UGameplayCameraComponent::GetInitialPose() const
 {
-	if (EvaluationContext)
-	{
-		return FBlueprintCameraPose::FromCameraPose(EvaluationContext->GetInitialResult().CameraPose);
-	}
-	else
-	{
-		FFrame::KismetExecutionMessage(
-				*FString::Format(
-					TEXT("Can't get initial camera pose on Gameplay Camera component '{0}': it isn't active."),
-					{ *GetNameSafe(this) }),
-				ELogVerbosity::Error);
-		return FBlueprintCameraPose();
-	}
+	UE_PRIVATE_GAMEPLAY_CAMERA_COMPONENT_VALIDATE_EVALUATION_CONTEXT("Can't get initial camera pose", FBlueprintCameraPose());
+
+	return FBlueprintCameraPose::FromCameraPose(EvaluationContext->GetInitialResult().CameraPose);
 }
 
-void UGameplayCameraComponent::SetInitialPose(const FBlueprintCameraPose& CameraPose)
+bool UGameplayCameraComponent::SetInitialPose(const FBlueprintCameraPose& CameraPose)
 {
-	if (EvaluationContext)
-	{
-		FCameraPose InitialPose = EvaluationContext->GetInitialResult().CameraPose;
-		CameraPose.ApplyTo(InitialPose);
-	}
-	else
-	{
-		FFrame::KismetExecutionMessage(
-				*FString::Format(
-					TEXT("Can't set initial camera pose on Gameplay Camera component '{0}': it isn't active."),
-					{ GetNameSafe(this) }),
-				ELogVerbosity::Error);
-	}
+	UE_PRIVATE_GAMEPLAY_CAMERA_COMPONENT_VALIDATE_EVALUATION_CONTEXT("Can't set initial camera pose", false);
+
+	FCameraPose InitialPose = EvaluationContext->GetInitialResult().CameraPose;
+	CameraPose.ApplyTo(InitialPose);
+	return true;
 }
 
 FBlueprintCameraVariableTable UGameplayCameraComponent::GetInitialVariableTable() const
 {
 	using namespace UE::Cameras;
 
-	if (EvaluationContext)
-	{
-		FCameraVariableTable& VariableTable = EvaluationContext->GetInitialResult().VariableTable;
-		FCameraSystemEvaluator* CameraSystemEvaluator = EvaluationContext->GetCameraSystemEvaluator();
-		TSharedPtr<FAutoResetCameraVariableService> VariableAutoResetService = 
-			CameraSystemEvaluator->FindEvaluationService<FAutoResetCameraVariableService>();
-		return FBlueprintCameraVariableTable(&VariableTable, VariableAutoResetService);
-	}
-	else
-	{
-		FFrame::KismetExecutionMessage(
-				*FString::Format(
-					TEXT("Can't get initial camera variable table on Gameplay Camera component '{0}': it isn't active."),
-					{ GetNameSafe(this) }),
-				ELogVerbosity::Error);
-		return FBlueprintCameraVariableTable();
-	}
+	UE_PRIVATE_GAMEPLAY_CAMERA_COMPONENT_VALIDATE_EVALUATION_CONTEXT("Can't get initial camera variable table", FBlueprintCameraVariableTable());
+
+	FCameraVariableTable& VariableTable = EvaluationContext->GetInitialResult().VariableTable;
+	return FBlueprintCameraVariableTable(&VariableTable);
 }
 
+#undef UE_PRIVATE_GAMEPLAY_CAMERA_COMPONENT_VALIDATE_EVALUATION_CONTEXT
+
 void UGameplayCameraComponent::OnRegister()
 {
 	Super::OnRegister();
diff --git a/Engine/Plugins/Cameras/GameplayCameras/Source/GameplayCameras/Private/Services/AutoResetCameraVariableService.cpp b/Engine/Plugins/Cameras/GameplayCameras/Source/GameplayCameras/Private/Services/AutoResetCameraVariableService.cpp
deleted file mode 100644
diff --git a/Engine/Plugins/Cameras/GameplayCameras/Source/GameplayCameras/Private/Services/AutoResetCameraVariableService.h b/Engine/Plugins/Cameras/GameplayCameras/Source/GameplayCameras/Private/Services/AutoResetCameraVariableService.h
deleted file mode 100644
diff --git a/Engine/Plugins/Cameras/GameplayCameras/Source/GameplayCameras/Public/Core/CameraEvaluationContext.h b/Engine/Plugins/Cameras/GameplayCameras/Source/GameplayCameras/Public/Core/CameraEvaluationContext.h
--- a/Engine/Plugins/Cameras/GameplayCameras/Source/GameplayCameras/Public/Core/CameraEvaluationContext.h
+++ b/Engine/Plugins/Cameras/GameplayCameras/Source/GameplayCameras/Public/Core/CameraEvaluationContext.h
@@ -119,6 +119,8 @@ public:
 	// Internal API.
 	void AddReferencedObjects(FReferenceCollector& Collector);
 
+	void OnEndCameraSystemUpdate();
+
 	bool RegisterChildContext(TSharedRef<FCameraEvaluationContext> ChildContext);
 	bool UnregisterChildContext(TSharedRef<FCameraEvaluationContext> ChildContext);
 
diff --git a/Engine/Plugins/Cameras/GameplayCameras/Source/GameplayCameras/Public/Core/CameraSystemEvaluator.h b/Engine/Plugins/Cameras/GameplayCameras/Source/GameplayCameras/Public/Core/CameraSystemEvaluator.h
--- a/Engine/Plugins/Cameras/GameplayCameras/Source/GameplayCameras/Public/Core/CameraSystemEvaluator.h
+++ b/Engine/Plugins/Cameras/GameplayCameras/Source/GameplayCameras/Public/Core/CameraSystemEvaluator.h
@@ -23,7 +23,6 @@ struct FMinimalViewInfo;
 namespace UE::Cameras
 {
 
-class FAutoResetCameraVariableService;
 class FCameraEvaluationContext;
 class FCameraEvaluationService;
 class FCameraRigCombinationRegistry;
@@ -179,6 +178,12 @@ private:
 
 	void NotifyRootCameraNodeEvent(const FRootCameraNodeCameraRigEvent& InEvent);
 
+#if UE_GAMEPLAY_CAMERAS_DEBUG
+	static bool IsDebugTraceEnabled();
+	static bool ShouldBuildOrDrawDebugBlocks();
+	void BuildDebugBlocksIfNeeded();
+#endif  // UE_GAMEPLAY_CAMERAS_DEBUG
+
 private:
 
 	/** The owner (if any) of this camera system evaluator. */
@@ -193,9 +198,6 @@ private:
 	/** The list of evaluation services. */
 	TArray<TSharedPtr<FCameraEvaluationService>> EvaluationServices;
 
-	/** Quick access to the variable auto-reset service. */
-	TSharedPtr<FAutoResetCameraVariableService> VariableAutoResetService;
-
 	/** Registry for programmatically building combinations of camera rigs. */
 	TSharedPtr<FCameraRigCombinationRegistry> CameraRigCombinationRegistry;
 
diff --git a/Engine/Plugins/Cameras/GameplayCameras/Source/GameplayCameras/Public/Core/CameraVariableTable.h b/Engine/Plugins/Cameras/GameplayCameras/Source/GameplayCameras/Public/Core/CameraVariableTable.h
--- a/Engine/Plugins/Cameras/GameplayCameras/Source/GameplayCameras/Public/Core/CameraVariableTable.h
+++ b/Engine/Plugins/Cameras/GameplayCameras/Source/GameplayCameras/Public/Core/CameraVariableTable.h
@@ -179,6 +179,8 @@ public:
 	bool IsValueWrittenThisFrame(FCameraVariableID VariableID) const;
 	void ClearAllWrittenThisFrameFlags();
 
+	void AutoResetValues();
+
 	bool TryGetVariableDefinition(FCameraVariableID VariableID, FCameraVariableDefinition& OutVariableDefinition) const;
 
 	void Serialize(FArchive& Ar);
@@ -210,8 +212,9 @@ private:
 		None = 0,
 		Private = 1 << 0,
 		Input = 1 << 1,
-		Written = 1 << 2,
-		WrittenThisFrame = 1 << 3
+		AutoReset = 1 << 2,
+		Written = 1 << 3,
+		WrittenThisFrame = 1 << 4
 	};
 	FRIEND_ENUM_CLASS_FLAGS(EEntryFlags)
 
diff --git a/Engine/Plugins/Cameras/GameplayCameras/Source/GameplayCameras/Public/Core/CameraVariableTableFwd.h b/Engine/Plugins/Cameras/GameplayCameras/Source/GameplayCameras/Public/Core/CameraVariableTableFwd.h
--- a/Engine/Plugins/Cameras/GameplayCameras/Source/GameplayCameras/Public/Core/CameraVariableTableFwd.h
+++ b/Engine/Plugins/Cameras/GameplayCameras/Source/GameplayCameras/Public/Core/CameraVariableTableFwd.h
@@ -145,6 +145,14 @@ struct FCameraVariableDefinition
 	UPROPERTY()
 	bool bIsInput = false;
 
+	/**
+	 * Whether the variable should auto-reset to an "unset" state after every
+	 * evaluation. Doing so resets it to its default value if bound to a
+	 * camera variable or camera parameter, or to a non-existent state otherwise.
+	 */
+	UPROPERTY()
+	bool bAutoReset = false;
+
 #if WITH_EDITORONLY_DATA
 	/** The name of the variable, for debugging purposes. */
 	UPROPERTY()
@@ -203,10 +211,6 @@ struct FCameraVariableTableAllocationInfo
 	UPROPERTY()
 	TArray<FCameraVariableDefinition> VariableDefinitions;
 
-	/** The list of variables that should be auto-reset to their default value every frame. */
-	UPROPERTY()
-	TArray<TObjectPtr<UCameraVariableAsset>> AutoResetVariables;
-
 	/**Combines the given allocation info with this one. */
 	GAMEPLAYCAMERAS_API void Combine(const FCameraVariableTableAllocationInfo& OtherInfo);
 
diff --git a/Engine/Plugins/Cameras/GameplayCameras/Source/GameplayCameras/Public/Debug/CameraSystemTrace.h b/Engine/Plugins/Cameras/GameplayCameras/Source/GameplayCameras/Public/Debug/CameraSystemTrace.h
--- a/Engine/Plugins/Cameras/GameplayCameras/Source/GameplayCameras/Public/Debug/CameraSystemTrace.h
+++ b/Engine/Plugins/Cameras/GameplayCameras/Source/GameplayCameras/Public/Debug/CameraSystemTrace.h
@@ -31,21 +31,12 @@ public:
 
 public:
 
-	/** Gets whether we are currently replaying traced information (such as with rewind debugger). */
-	GAMEPLAYCAMERAS_API static bool IsTraceReplay();
-	/** Sets whether we are currently replaying traced information (such as with rewind debugger). */
-	GAMEPLAYCAMERAS_API static void SetTraceReplay(bool bInIsReplaying);
-
 	/** Returns whether tracing of camera system evaluation is enabled. */
 	GAMEPLAYCAMERAS_API static bool IsTraceEnabled();
 	/** Records one frame of camera system evaluation. */
 	GAMEPLAYCAMERAS_API static void TraceEvaluation(UWorld* InWorld, const FCameraSystemEvaluationResult& InResult, FCameraDebugBlock& InRootDebugBlock);
 	/** Reads back one frame of camera system evaluation. */
 	GAMEPLAYCAMERAS_API static FCameraDebugBlock* ReadEvaluationTrace(TArray<uint8> InSerializedBlocks, FCameraDebugBlockStorage& InStorage);
-
-private:
-
-	static bool bIsReplaying;
 };
 
 }  // namespace UE::Cameras
diff --git a/Engine/Plugins/Cameras/GameplayCameras/Source/GameplayCameras/Public/Directors/BlueprintCameraDirector.h b/Engine/Plugins/Cameras/GameplayCameras/Source/GameplayCameras/Public/Directors/BlueprintCameraDirector.h
--- a/Engine/Plugins/Cameras/GameplayCameras/Source/GameplayCameras/Public/Directors/BlueprintCameraDirector.h
+++ b/Engine/Plugins/Cameras/GameplayCameras/Source/GameplayCameras/Public/Directors/BlueprintCameraDirector.h
@@ -18,8 +18,6 @@ enum class ECameraRigLayer : uint8;
 namespace UE::Cameras
 {
 
-class FAutoResetCameraVariableService;
-
 /** Information about a persitent camera rig to be activated or deactivated. */
 struct FBlueprintPersistentCameraRigInfo
 {
@@ -220,9 +218,6 @@ private:
 
 	/** The current evaluation context. */
 	TSharedPtr<FCameraEvaluationContext> EvaluationContext;
-
-	/** The variable auto-reset service, for using when returning the variable table. */
-	TSharedPtr<UE::Cameras::FAutoResetCameraVariableService> VariableAutoResetService;
 };
 
 /**
diff --git a/Engine/Plugins/Cameras/GameplayCameras/Source/GameplayCameras/Public/GameFramework/BlueprintCameraVariableTable.h b/Engine/Plugins/Cameras/GameplayCameras/Source/GameplayCameras/Public/GameFramework/BlueprintCameraVariableTable.h
--- a/Engine/Plugins/Cameras/GameplayCameras/Source/GameplayCameras/Public/GameFramework/BlueprintCameraVariableTable.h
+++ b/Engine/Plugins/Cameras/GameplayCameras/Source/GameplayCameras/Public/GameFramework/BlueprintCameraVariableTable.h
@@ -23,7 +23,6 @@ namespace UE::Cameras
 {
 
 class FCameraVariableTable;
-class FAutoResetCameraVariableService;
 
 }  // namespace UE::Cameras
 
@@ -38,7 +37,7 @@ public:
 	using FCameraVariableTable = UE::Cameras::FCameraVariableTable;
 
 	FBlueprintCameraVariableTable();
-	FBlueprintCameraVariableTable(FCameraVariableTable* InVariableTable, TSharedPtr<UE::Cameras::FAutoResetCameraVariableService> InVariableAutoResetService);
+	FBlueprintCameraVariableTable(FCameraVariableTable* InVariableTable);
 
 	/** Gets the underlying variable table. */
 	FCameraVariableTable* GetVariableTable() const { return PrivateVariableTable; }
@@ -49,7 +48,6 @@ public:
 private:
 
 	FCameraVariableTable* PrivateVariableTable = nullptr;
-	TSharedPtr<UE::Cameras::FAutoResetCameraVariableService> VariableAutoResetService;
 
 	friend class UBlueprintCameraVariableTableFunctionLibrary;
 };
diff --git a/Engine/Plugins/Cameras/GameplayCameras/Source/GameplayCameras/Public/GameFramework/GameplayCameraComponent.h b/Engine/Plugins/Cameras/GameplayCameras/Source/GameplayCameras/Public/GameFramework/GameplayCameraComponent.h
--- a/Engine/Plugins/Cameras/GameplayCameras/Source/GameplayCameras/Public/GameFramework/GameplayCameraComponent.h
+++ b/Engine/Plugins/Cameras/GameplayCameras/Source/GameplayCameras/Public/GameFramework/GameplayCameraComponent.h
@@ -62,7 +62,7 @@ public:
 
 	/** Sets the initial camera pose for this component's camera evaluation context. */
 	UFUNCTION(BlueprintCallable, Category=Camera)
-	GAMEPLAYCAMERAS_API void SetInitialPose(const FBlueprintCameraPose& CameraPose);
+	GAMEPLAYCAMERAS_API bool SetInitialPose(const FBlueprintCameraPose& CameraPose);
 
 	/** Gets the initial camera variable table for this component's camera evaluation context. */
 	UFUNCTION(BlueprintPure, Category=Camera)
