Cameras: add menu entry for changing a camera asset's director

- Add new camera asset editor toolkit command for changing the camera director type
- Support adding/removing editor modes in the base editor mode manager class
- Switch the camera director editor mode when the camera director changes
- Remove unused commands for the camera asset editor toolkit

diff --git a/Engine/Plugins/Cameras/GameplayCameras/Source/GameplayCamerasEditor/Private/Commands/CameraAssetEditorCommands.cpp b/Engine/Plugins/Cameras/GameplayCameras/Source/GameplayCamerasEditor/Private/Commands/CameraAssetEditorCommands.cpp
--- a/Engine/Plugins/Cameras/GameplayCameras/Source/GameplayCamerasEditor/Private/Commands/CameraAssetEditorCommands.cpp
+++ b/Engine/Plugins/Cameras/GameplayCameras/Source/GameplayCamerasEditor/Private/Commands/CameraAssetEditorCommands.cpp
@@ -35,16 +35,9 @@ void FCameraAssetEditorCommands::RegisterCommands()
 	UI_COMMAND(ShowSharedTransitions, "Shared Transitions", "Shows the shared transitions",
 			EUserInterfaceActionType::ToggleButton, FInputChord());
 
-	UI_COMMAND(EditCameraRig, "Edit Camera Rig", "Edits the selected camera rig",
+	UI_COMMAND(ChangeCameraDirector, "Change Camera Director", "Opens a dialog for changing the type of camera director",
 			EUserInterfaceActionType::Button, FInputChord());
 
-	UI_COMMAND(AddCameraRig, "Add Camera Rig", "Adds a new camera rig",
-			EUserInterfaceActionType::Button, FInputChord());
-	UI_COMMAND(DeleteCameraRig, "Delete Camera Rig", "Deletes the selected camera rig(s)",
-			EUserInterfaceActionType::Button, FInputChord(EKeys::Delete));
-	UI_COMMAND(RenameCameraRig, "Rename Camera Rig", "Renames the selected camera rig",
-			EUserInterfaceActionType::Button, FInputChord(EKeys::F2));
-
 	UI_COMMAND(ShowMessages, "Messages", "Shows the message log for this camera asset",
 			EUserInterfaceActionType::ToggleButton, FInputChord());
 	UI_COMMAND(FindInCamera, "Search", "Searches for nodes in this camera asset",
diff --git a/Engine/Plugins/Cameras/GameplayCameras/Source/GameplayCamerasEditor/Private/Factories/CameraAssetFactory.cpp b/Engine/Plugins/Cameras/GameplayCameras/Source/GameplayCamerasEditor/Private/Factories/CameraAssetFactory.cpp
--- a/Engine/Plugins/Cameras/GameplayCameras/Source/GameplayCamerasEditor/Private/Factories/CameraAssetFactory.cpp
+++ b/Engine/Plugins/Cameras/GameplayCameras/Source/GameplayCamerasEditor/Private/Factories/CameraAssetFactory.cpp
@@ -4,50 +4,12 @@
 
 #include "Core/CameraAsset.h"
 #include "Core/CameraDirector.h"
-#include "Core/CameraRigAsset.h"
-#include "ClassViewerFilter.h"
-#include "ClassViewerModule.h"
-#include "Directors/BlueprintCameraDirector.h"
-#include "Directors/SingleCameraDirector.h"
-#include "Kismet2/SClassPickerDialog.h"
-#include "Modules/ModuleManager.h"
+#include "Helpers/CameraDirectorClassPicker.h"
 
 #include UE_INLINE_GENERATED_CPP_BY_NAME(CameraAssetFactory)
 
 #define LOCTEXT_NAMESPACE "CameraAssetFactory"
 
-namespace UE::Cameras::Private
-{
-
-class FCameraDirectorClassFilter : public IClassViewerFilter
-{
-public:
-	TSet<const UClass*> AllowedClasses;
-	EClassFlags DisallowedClassFlags;
-
-	FCameraDirectorClassFilter()
-	{
-		AllowedClasses.Add(UCameraDirector::StaticClass());
-		DisallowedClassFlags = CLASS_Abstract | CLASS_Deprecated;
-	}
-
-	virtual bool IsClassAllowed(const FClassViewerInitializationOptions& InInitOptions, const UClass* InClass, TSharedRef<FClassViewerFilterFuncs> InFilterFuncs) override
-	{
-		return 
-			!InClass->HasAnyClassFlags(DisallowedClassFlags) &&
-			InFilterFuncs->IfInChildOfClassesSet(AllowedClasses, InClass) != EFilterReturn::Failed;
-	}
-
-	virtual bool IsUnloadedClassAllowed(const FClassViewerInitializationOptions& InInitOptions, const TSharedRef< const IUnloadedBlueprintData > InUnloadedClassData, TSharedRef< FClassViewerFilterFuncs > InFilterFuncs) override
-	{
-		return 
-			!InUnloadedClassData->HasAnyClassFlags(DisallowedClassFlags) && 
-			InFilterFuncs->IfInChildOfClassesSet(AllowedClasses, InUnloadedClassData) != EFilterReturn::Failed;
-	}
-};
-
-}  // namespace UE::Cameras::Private
-
 UCameraAssetFactory::UCameraAssetFactory(const FObjectInitializer& ObjectInitializer)
 	: Super(ObjectInitializer)
 {
@@ -75,25 +37,13 @@ UObject* UCameraAssetFactory::FactoryCreateNew(UClass* Class, UObject* Parent, F
 
 bool UCameraAssetFactory::ConfigureProperties()
 {
-	using namespace UE::Cameras::Private;
-
-	FClassViewerModule& ClassViewerModule = FModuleManager::LoadModuleChecked<FClassViewerModule>("ClassViewer");
-
-	FClassViewerInitializationOptions Options;
-	Options.Mode = EClassViewerMode::ClassPicker;
-	Options.NameTypeToDisplay = EClassViewerNameTypeToDisplay::DisplayName;
-	Options.bShowNoneOption = true;
-
-	Options.ExtraPickerCommonClasses.Add(UBlueprintCameraDirector::StaticClass());
-	Options.ExtraPickerCommonClasses.Add(USingleCameraDirector::StaticClass());
+	using namespace UE::Cameras;
 
-	TSharedPtr<FCameraDirectorClassFilter> Filter = MakeShareable(new FCameraDirectorClassFilter);
-	Options.ClassFilters.Add(Filter.ToSharedRef());
+	FCameraDirectorClassPicker Picker;
 
 	CameraDirectorClass = nullptr;
-	UClass* ChosenClass = nullptr;
-	const FText TitleText = LOCTEXT("CameraDirectorPicker", "Pick Camera Director Type");
-	const bool bPressedOk = SClassPickerDialog::PickClass(TitleText, Options, ChosenClass, UCameraDirector::StaticClass());
+	TSubclassOf<UCameraDirector> ChosenClass;
+	const bool bPressedOk = Picker.PickCameraDirectorClass(ChosenClass);
 	if (bPressedOk)
 	{
 		CameraDirectorClass = ChosenClass;
diff --git a/Engine/Plugins/Cameras/GameplayCameras/Source/GameplayCamerasEditor/Private/Helpers/CameraDirectorClassPicker.cpp b/Engine/Plugins/Cameras/GameplayCameras/Source/GameplayCamerasEditor/Private/Helpers/CameraDirectorClassPicker.cpp
new file mode 100644
--- /dev/null
+++ b/Engine/Plugins/Cameras/GameplayCameras/Source/GameplayCamerasEditor/Private/Helpers/CameraDirectorClassPicker.cpp
@@ -0,0 +1,86 @@
+// Copyright Epic Games, Inc. All Rights Reserved.
+
+#include "Helpers/CameraDirectorClassPicker.h"
+
+#include "ClassViewerFilter.h"
+#include "ClassViewerModule.h"
+#include "Directors/BlueprintCameraDirector.h"
+#include "Directors/SingleCameraDirector.h"
+#include "Kismet2/SClassPickerDialog.h"
+#include "Modules/ModuleManager.h"
+
+#define LOCTEXT_NAMESPACE "CameraDirectorClassPicker"
+
+namespace UE::Cameras
+{
+
+class FCameraDirectorClassFilter : public IClassViewerFilter
+{
+public:
+	TSet<const UClass*> AllowedClasses;
+	EClassFlags DisallowedClassFlags;
+
+	FCameraDirectorClassFilter()
+	{
+		AllowedClasses.Add(UCameraDirector::StaticClass());
+		DisallowedClassFlags = CLASS_Abstract | CLASS_Deprecated;
+	}
+
+	virtual bool IsClassAllowed(const FClassViewerInitializationOptions& InInitOptions, const UClass* InClass, TSharedRef<FClassViewerFilterFuncs> InFilterFuncs) override
+	{
+		return 
+			!InClass->HasAnyClassFlags(DisallowedClassFlags) &&
+			InFilterFuncs->IfInChildOfClassesSet(AllowedClasses, InClass) != EFilterReturn::Failed;
+	}
+
+	virtual bool IsUnloadedClassAllowed(const FClassViewerInitializationOptions& InInitOptions, const TSharedRef< const IUnloadedBlueprintData > InUnloadedClassData, TSharedRef< FClassViewerFilterFuncs > InFilterFuncs) override
+	{
+		return 
+			!InUnloadedClassData->HasAnyClassFlags(DisallowedClassFlags) && 
+			InFilterFuncs->IfInChildOfClassesSet(AllowedClasses, InUnloadedClassData) != EFilterReturn::Failed;
+	}
+};
+
+FCameraDirectorClassPicker::FCameraDirectorClassPicker()
+{
+	CommonCameraDirectorClasses.Add(UBlueprintCameraDirector::StaticClass());
+	CommonCameraDirectorClasses.Add(USingleCameraDirector::StaticClass());
+}
+
+void FCameraDirectorClassPicker::AddCommonCameraDirector(TSubclassOf<UCameraDirector> InClass)
+{
+	CommonCameraDirectorClasses.Add(InClass);
+}
+
+void FCameraDirectorClassPicker::ResetCommonCameraDirectors()
+{
+	CommonCameraDirectorClasses.Reset();
+}
+
+bool FCameraDirectorClassPicker::PickCameraDirectorClass(TSubclassOf<UCameraDirector>& OutChosenClass)
+{
+	FClassViewerModule& ClassViewerModule = FModuleManager::LoadModuleChecked<FClassViewerModule>("ClassViewer");
+
+	FClassViewerInitializationOptions Options;
+	Options.Mode = EClassViewerMode::ClassPicker;
+	Options.NameTypeToDisplay = EClassViewerNameTypeToDisplay::DisplayName;
+	Options.bShowNoneOption = true;
+	Options.ExtraPickerCommonClasses.Append(CommonCameraDirectorClasses);
+
+	TSharedPtr<FCameraDirectorClassFilter> Filter = MakeShareable(new FCameraDirectorClassFilter);
+	Options.ClassFilters.Add(Filter.ToSharedRef());
+
+	UClass* ChosenClass = nullptr;
+	const FText TitleText = LOCTEXT("CameraDirectorPicker", "Pick Camera Director Type");
+	const bool bPressedOk = SClassPickerDialog::PickClass(TitleText, Options, ChosenClass, UCameraDirector::StaticClass());
+	if (bPressedOk)
+	{
+		OutChosenClass = ChosenClass;
+	}
+	return bPressedOk;
+}
+
+}  // namespace UE::Cameras
+
+#undef LOCTEXT_NAMESPACE
+
diff --git a/Engine/Plugins/Cameras/GameplayCameras/Source/GameplayCamerasEditor/Private/Helpers/CameraDirectorClassPicker.h b/Engine/Plugins/Cameras/GameplayCameras/Source/GameplayCamerasEditor/Private/Helpers/CameraDirectorClassPicker.h
new file mode 100644
--- /dev/null
+++ b/Engine/Plugins/Cameras/GameplayCameras/Source/GameplayCamerasEditor/Private/Helpers/CameraDirectorClassPicker.h
@@ -0,0 +1,30 @@
+// Copyright Epic Games, Inc. All Rights Reserved.
+
+#pragma once
+
+#include "Containers/Array.h"
+#include "Templates/SubclassOf.h"
+
+class UCameraDirector;
+
+namespace UE::Cameras
+{
+
+class FCameraDirectorClassPicker
+{
+public:
+
+	FCameraDirectorClassPicker();
+
+	void AddCommonCameraDirector(TSubclassOf<UCameraDirector> InClass);
+	void ResetCommonCameraDirectors();
+
+	bool PickCameraDirectorClass(TSubclassOf<UCameraDirector>& OutChosenClass);
+
+private:
+
+	TArray<TSubclassOf<UCameraDirector>> CommonCameraDirectorClasses;
+};
+
+}  // namespace UE::Cameras
+
diff --git a/Engine/Plugins/Cameras/GameplayCameras/Source/GameplayCamerasEditor/Private/Toolkits/AssetEditorModeManagerToolkit.cpp b/Engine/Plugins/Cameras/GameplayCameras/Source/GameplayCamerasEditor/Private/Toolkits/AssetEditorModeManagerToolkit.cpp
--- a/Engine/Plugins/Cameras/GameplayCameras/Source/GameplayCamerasEditor/Private/Toolkits/AssetEditorModeManagerToolkit.cpp
+++ b/Engine/Plugins/Cameras/GameplayCameras/Source/GameplayCamerasEditor/Private/Toolkits/AssetEditorModeManagerToolkit.cpp
@@ -41,6 +41,29 @@ void FAssetEditorModeManagerToolkit::AddEditorMode(TSharedRef<FAssetEditorMode>
 	EditorModes.Add(InMode->GetModeName(), InMode);
 }
 
+void FAssetEditorModeManagerToolkit::RemoveEditorMode(TSharedRef<FAssetEditorMode> InMode)
+{
+	RemoveEditorMode(InMode->GetModeName());
+}
+
+void FAssetEditorModeManagerToolkit::RemoveEditorMode(FName InModeName)
+{
+	TSharedPtr<FAssetEditorMode> ModeToRemove = EditorModes.FindRef(InModeName);
+	if (!ensureMsgf(
+				ModeToRemove,
+				TEXT("No editor mode named '%s' was added!"), *InModeName.ToString()))
+	{
+		return;
+	}
+
+	if (CurrentEditorMode == ModeToRemove)
+	{
+		SetEditorMode(NAME_None);
+	}
+
+	EditorModes.Remove(InModeName);
+}
+
 void FAssetEditorModeManagerToolkit::GetEditorModes(TArray<TSharedPtr<FAssetEditorMode>>& OutModes) const
 {
 	for (const TPair<FName, TSharedPtr<FAssetEditorMode>>& Pair : EditorModes)
diff --git a/Engine/Plugins/Cameras/GameplayCameras/Source/GameplayCamerasEditor/Private/Toolkits/CameraAssetEditorToolkit.cpp b/Engine/Plugins/Cameras/GameplayCameras/Source/GameplayCamerasEditor/Private/Toolkits/CameraAssetEditorToolkit.cpp
--- a/Engine/Plugins/Cameras/GameplayCameras/Source/GameplayCamerasEditor/Private/Toolkits/CameraAssetEditorToolkit.cpp
+++ b/Engine/Plugins/Cameras/GameplayCameras/Source/GameplayCamerasEditor/Private/Toolkits/CameraAssetEditorToolkit.cpp
@@ -17,12 +17,14 @@
 #include "Framework/Docking/TabManager.h"
 #include "GameplayCamerasEditorSettings.h"
 #include "Helpers/AssetTypeMenuOverlayHelper.h"
+#include "Helpers/CameraDirectorClassPicker.h"
 #include "Helpers/ObjectReferenceFinder.h"
 #include "IAssetTools.h"
 #include "IGameplayCamerasEditorModule.h"
 #include "IGameplayCamerasLiveEditManager.h"
 #include "IGameplayCamerasModule.h"
 #include "PropertyEditorModule.h"
+#include "ScopedTransaction.h"
 #include "Styles/GameplayCamerasEditorStyle.h"
 #include "ToolMenus.h"
 #include "Toolkits/BuildButtonToolkit.h"
@@ -50,6 +52,8 @@ FCameraAssetEditorToolkit::FCameraAssetEditorToolkit(UCameraAssetEditor* InOwnin
 	, BuildButtonToolkit(MakeShared<FBuildButtonToolkit>(CameraAsset))
 	, BuildLogToolkit(MakeShared<FCameraBuildLogToolkit>())
 {
+	CameraAsset->EventHandlers.Register(CameraAssetEventHandler, this);
+
 	StandardLayout->AddBottomTab(SearchTabId);
 	StandardLayout->AddBottomTab(MessagesTabId);
 
@@ -201,27 +205,20 @@ void FCameraAssetEditorToolkit::InitToolMenuContext(FToolMenuContext& MenuContex
 
 void FCameraAssetEditorToolkit::PostInitAssetEditor()
 {
-	Settings = GetMutableDefault<UGameplayCamerasEditorSettings>();
+	TSharedPtr<FExtender> MenuExtender = MakeShared<FExtender>();
+	{
+		MenuExtender->AddMenuExtension(
+			"AssetEditorActions",
+			EExtensionHook::After,
+			GetToolkitCommands(),
+			FMenuExtensionDelegate::CreateSP(this, &FCameraAssetEditorToolkit::FillCameraMenu));
+	}
+	AddMenuExtender(MenuExtender);
 
-	IGameplayCamerasEditorModule& GameplayCamerasEditorModule = IGameplayCamerasEditorModule::Get();
+	Settings = GetMutableDefault<UGameplayCamerasEditorSettings>();
 
 	const FName CameraDirectorModeName = FCameraDirectorAssetEditorMode::ModeName;
-	{
-		TSharedPtr<FCameraDirectorAssetEditorMode> CameraDirectorEditor;
-		for (const FOnCreateCameraDirectorAssetEditorMode& EditorCreator : GameplayCamerasEditorModule.GetCameraDirectorEditorCreators())
-		{
-			CameraDirectorEditor = EditorCreator.Execute(CameraAsset);
-			if (CameraDirectorEditor)
-			{
-				break;
-			}
-		}
-		if (!CameraDirectorEditor)
-		{
-			CameraDirectorEditor = MakeShared<FCameraDirectorAssetEditorMode>(CameraAsset);
-		}
-		AddEditorMode(CameraDirectorEditor.ToSharedRef());
-	}
+	AddEditorMode(CreateCameraDirectorAssetEditorMode().ToSharedRef());
 
 	const FName SharedTransitionsModeName = FCameraSharedTransitionsAssetEditorMode::ModeName;
 	AddEditorMode(MakeShared<FCameraSharedTransitionsAssetEditorMode>(CameraAsset));
@@ -239,6 +236,10 @@ void FCameraAssetEditorToolkit::PostInitAssetEditor()
 			FIsActionChecked::CreateSP(this, &FCameraAssetEditorToolkit::IsEditorMode, Pair.Key));
 	}
 
+	ToolkitCommands->MapAction(
+		Commands.ChangeCameraDirector,
+		FExecuteAction::CreateSP(this, &FCameraAssetEditorToolkit::OnChangeCameraDirector));
+
 	ToolkitCommands->MapAction(
 		Commands.Build,
 		FExecuteAction::CreateSP(this, &FCameraAssetEditorToolkit::OnBuild));
@@ -259,6 +260,35 @@ void FCameraAssetEditorToolkit::PostInitAssetEditor()
 	UpgradeLegacyCameraAssets();
 }
 
+TSharedPtr<FAssetEditorMode> FCameraAssetEditorToolkit::CreateCameraDirectorAssetEditorMode()
+{
+	IGameplayCamerasEditorModule& GameplayCamerasEditorModule = IGameplayCamerasEditorModule::Get();
+
+	TSharedPtr<FCameraDirectorAssetEditorMode> CameraDirectorEditor;
+	for (const FOnCreateCameraDirectorAssetEditorMode& EditorCreator : GameplayCamerasEditorModule.GetCameraDirectorEditorCreators())
+	{
+		CameraDirectorEditor = EditorCreator.Execute(CameraAsset);
+		if (CameraDirectorEditor)
+		{
+			break;
+		}
+	}
+	if (!CameraDirectorEditor)
+	{
+		CameraDirectorEditor = MakeShared<FCameraDirectorAssetEditorMode>(CameraAsset);
+	}
+	return CameraDirectorEditor;
+}
+
+void FCameraAssetEditorToolkit::FillCameraMenu(FMenuBuilder& MenuBuilder)
+{
+	const FCameraAssetEditorCommands& Commands = FCameraAssetEditorCommands::Get();
+
+	MenuBuilder.BeginSection(TEXT("Camera"), LOCTEXT("CameraMenuTitle", "Camera"));
+	MenuBuilder.AddMenuEntry(Commands.ChangeCameraDirector);
+	MenuBuilder.EndSection();
+}
+
 void FCameraAssetEditorToolkit::PostRegenerateMenusAndToolbars()
 {
 	SetMenuOverlay(FAssetTypeMenuOverlayHelper::CreateMenuOverlay(UCameraAsset::StaticClass()));
@@ -270,6 +300,34 @@ void FCameraAssetEditorToolkit::OnEditorToolkitModeActivated()
 	Settings->SaveConfig();
 }
 
+void FCameraAssetEditorToolkit::OnCameraDirectorChanged(UCameraAsset* InCameraAsset, const TCameraPropertyChangedEvent<UCameraDirector*>& Event)
+{
+	const FName CameraDirectorModeName = FCameraDirectorAssetEditorMode::ModeName;
+
+	RemoveEditorMode(CameraDirectorModeName);
+
+	AddEditorMode(CreateCameraDirectorAssetEditorMode().ToSharedRef());
+
+	if (GetCurrentEditorModeName().IsNone())
+	{
+		SetEditorMode(CameraDirectorModeName);
+	}
+}
+
+void FCameraAssetEditorToolkit::OnChangeCameraDirector()
+{
+	FCameraDirectorClassPicker Picker;
+	TSubclassOf<UCameraDirector> ChosenClass;
+	const bool bPressedOk = Picker.PickCameraDirectorClass(ChosenClass);
+	if (bPressedOk && ChosenClass != CameraAsset->GetCameraDirector()->GetClass())
+	{
+		const FScopedTransaction Transaction(LOCTEXT("ChangeCameraDirector", "Change Camera Director"));
+
+		UCameraDirector* NewCameraDirector = NewObject<UCameraDirector>(CameraAsset, ChosenClass, NAME_None, RF_Transactional);
+		CameraAsset->SetCameraDirector(NewCameraDirector);
+	}
+}
+
 void FCameraAssetEditorToolkit::OnBuild()
 {
 	using namespace UE::Cameras;
diff --git a/Engine/Plugins/Cameras/GameplayCameras/Source/GameplayCamerasEditor/Public/Commands/CameraAssetEditorCommands.h b/Engine/Plugins/Cameras/GameplayCameras/Source/GameplayCamerasEditor/Public/Commands/CameraAssetEditorCommands.h
--- a/Engine/Plugins/Cameras/GameplayCameras/Source/GameplayCamerasEditor/Public/Commands/CameraAssetEditorCommands.h
+++ b/Engine/Plugins/Cameras/GameplayCameras/Source/GameplayCamerasEditor/Public/Commands/CameraAssetEditorCommands.h
@@ -28,11 +28,7 @@ public:
 	TSharedPtr<FUICommandInfo> ShowCameraRigs;
 	TSharedPtr<FUICommandInfo> ShowSharedTransitions;
 
-	TSharedPtr<FUICommandInfo> EditCameraRig;
-
-	TSharedPtr<FUICommandInfo> AddCameraRig;
-	TSharedPtr<FUICommandInfo> DeleteCameraRig;
-	TSharedPtr<FUICommandInfo> RenameCameraRig;
+	TSharedPtr<FUICommandInfo> ChangeCameraDirector;
 
 	TSharedPtr<FUICommandInfo> ShowMessages;
 	TSharedPtr<FUICommandInfo> FindInCamera;
diff --git a/Engine/Plugins/Cameras/GameplayCameras/Source/GameplayCamerasEditor/Public/Toolkits/AssetEditorModeManagerToolkit.h b/Engine/Plugins/Cameras/GameplayCameras/Source/GameplayCamerasEditor/Public/Toolkits/AssetEditorModeManagerToolkit.h
--- a/Engine/Plugins/Cameras/GameplayCameras/Source/GameplayCamerasEditor/Public/Toolkits/AssetEditorModeManagerToolkit.h
+++ b/Engine/Plugins/Cameras/GameplayCameras/Source/GameplayCamerasEditor/Public/Toolkits/AssetEditorModeManagerToolkit.h
@@ -38,6 +38,8 @@ protected:
 protected:
 
 	void AddEditorMode(TSharedRef<FAssetEditorMode> InMode);
+	void RemoveEditorMode(TSharedRef<FAssetEditorMode> InMode);
+	void RemoveEditorMode(FName InModeName);
 
 	TSharedPtr<FAssetEditorMode> GetEditorMode(FName InModeName) const;
 	void GetEditorModes(TArray<TSharedPtr<FAssetEditorMode>>& OutModes) const;
diff --git a/Engine/Plugins/Cameras/GameplayCameras/Source/GameplayCamerasEditor/Public/Toolkits/CameraAssetEditorToolkit.h b/Engine/Plugins/Cameras/GameplayCameras/Source/GameplayCamerasEditor/Public/Toolkits/CameraAssetEditorToolkit.h
--- a/Engine/Plugins/Cameras/GameplayCameras/Source/GameplayCamerasEditor/Public/Toolkits/CameraAssetEditorToolkit.h
+++ b/Engine/Plugins/Cameras/GameplayCameras/Source/GameplayCamerasEditor/Public/Toolkits/CameraAssetEditorToolkit.h
@@ -2,7 +2,8 @@
 
 #pragma once
 
-#include "CoreTypes.h"
+#include "Core/CameraAsset.h"
+#include "Core/CameraEventHandler.h"
 #include "Toolkits/AssetEditorModeManagerToolkit.h"
 #include "Tools/BaseAssetToolkit.h"
 #include "UObject/GCObject.h"
@@ -10,7 +11,6 @@
 #include "CameraAssetEditorToolkit.generated.h"
 
 class SFindInObjectTreeGraph;
-class UCameraAsset;
 class UCameraAssetEditor;
 class UEdGraph;
 class UEdGraphNode;
@@ -31,6 +31,7 @@ class IGameplayCamerasLiveEditManager;
 class FCameraAssetEditorToolkit 
 	: public FAssetEditorModeManagerToolkit
 	, public FGCObject
+	, public UE::Cameras::ICameraAssetEventHandler
 {
 public:
 
@@ -61,11 +62,19 @@ protected:
 	// FAssetEditorModeManagerToolkit interface
 	virtual void OnEditorToolkitModeActivated() override;
 
+	// ICameraAssetEventHandler interface
+	virtual void OnCameraDirectorChanged(UCameraAsset* InCameraAsset, const TCameraPropertyChangedEvent<UCameraDirector*>& Event) override;
+
 private:
 
 	TSharedRef<SDockTab> SpawnTab_Search(const FSpawnTabArgs& Args);
 	TSharedRef<SDockTab> SpawnTab_Messages(const FSpawnTabArgs& Args);
 
+	void FillCameraMenu(FMenuBuilder& MenuBuilder);
+
+	TSharedPtr<FAssetEditorMode> CreateCameraDirectorAssetEditorMode();
+
+	void OnChangeCameraDirector();
 	void OnBuild();
 	void OnFindInCamera();
 
@@ -83,6 +92,9 @@ private:
 	/** The asset being edited */
 	TObjectPtr<UCameraAsset> CameraAsset;
 
+	/** Event listener for the camera asset */
+	UE::Cameras::TCameraEventHandler<UE::Cameras::ICameraAssetEventHandler> CameraAssetEventHandler;
+
 	/** The layout for this toolkit */
 	TSharedPtr<FStandardToolkitLayout> StandardLayout;
 
