Cameras: add more Blueprint functions for getting/setting camera rig parameters

Since users may now want to get values from either blendable or data parameters, they need both the variable and context data tables exposed in Blueprint. This change therefore exposes a "result" wrapper that combines the BP wrappers for both of these objects. Previous BP functions that returned only the variable table are marked deprecated.

The SetCameraRigParameters node now also shows data parameter pins, and the internal "interop" class that it uses to generate all the "setter" function calls has gained new functions for setting data values. This required expanding the API of the context data table class itself.

diff --git a/Engine/Plugins/Cameras/GameplayCameras/Source/GameplayCameras/Private/Core/CameraContextDataTable.cpp b/Engine/Plugins/Cameras/GameplayCameras/Source/GameplayCameras/Private/Core/CameraContextDataTable.cpp
--- a/Engine/Plugins/Cameras/GameplayCameras/Source/GameplayCameras/Private/Core/CameraContextDataTable.cpp
+++ b/Engine/Plugins/Cameras/GameplayCameras/Source/GameplayCameras/Private/Core/CameraContextDataTable.cpp
@@ -269,7 +269,106 @@ void FCameraContextDataTable::DestroyBuffer()
 	Used = 0;
 }
 
-void FCameraContextDataTable::SetData(FCameraContextDataID InID, const FStructView& InData)
+const FName& FCameraContextDataTable::GetNameData(FCameraContextDataID InID) const
+{
+	if (const FName* Value = GetDataImpl<FName>(InID, ECameraContextDataType::Name, nullptr))
+	{
+		return *Value;
+	}
+
+	static FName DefaultValue(NAME_None);
+	return DefaultValue;
+}
+
+const FString& FCameraContextDataTable::GetStringData(FCameraContextDataID InID) const
+{
+	if (const FString* Value = GetDataImpl<FString>(InID, ECameraContextDataType::Name, nullptr))
+	{
+		return *Value;
+	}
+
+	static FString DefaultValue;
+	return DefaultValue;
+}
+
+uint8 FCameraContextDataTable::GetEnumData(FCameraContextDataID InID, const UEnum* EnumType) const
+{
+	if (const uint8* Value = GetDataImpl<uint8>(InID, ECameraContextDataType::Enum, EnumType))
+	{
+		return *Value;
+	}
+	return 0;
+}
+
+FConstStructView FCameraContextDataTable::GetStructViewData(FCameraContextDataID InID, const UScriptStruct* StructType) const
+{
+	const uint8* RawData = TryGetData(InID, ECameraContextDataType::Struct, StructType);
+	if (RawData)
+	{
+		FConstStructView ReturnValue(StructType, RawData);
+		return ReturnValue;
+	}
+	return FStructView();
+}
+
+FInstancedStruct FCameraContextDataTable::GetInstancedStructData(FCameraContextDataID InID, const UScriptStruct* StructType) const
+{
+	const uint8* RawData = TryGetData(InID, ECameraContextDataType::Struct, StructType);
+	if (RawData)
+	{
+		FInstancedStruct ReturnValue;
+		ReturnValue.InitializeAs(StructType, RawData);
+		return ReturnValue;
+	}
+	return FInstancedStruct();
+}
+
+UObject* FCameraContextDataTable::GetObjectData(FCameraContextDataID InID) const
+{
+	if (const TObjectPtr<UObject>* Value = GetDataImpl<TObjectPtr<UObject>>(InID, ECameraContextDataType::Name, nullptr))
+	{
+		return Value->Get();
+	}
+	return nullptr;
+}
+
+UClass* FCameraContextDataTable::GetClassData(FCameraContextDataID InID) const
+{
+	if (const TObjectPtr<UClass>* Value = GetDataImpl<TObjectPtr<UClass>>(InID, ECameraContextDataType::Name, nullptr))
+	{
+		return Value->Get();
+	}
+	return nullptr;
+}
+
+void FCameraContextDataTable::SetNameData(FCameraContextDataID InID, const FName& InData)
+{
+	SetDataImpl(InID, ECameraContextDataType::Name, nullptr, InData);
+}
+
+void FCameraContextDataTable::SetStringData(FCameraContextDataID InID, const FString& InData)
+{
+	SetDataImpl(InID, ECameraContextDataType::String, nullptr, InData);
+}
+
+void FCameraContextDataTable::SetEnumData(FCameraContextDataID InID, const UEnum* EnumType, uint8 InData)
+{
+	SetDataImpl(InID, ECameraContextDataType::Enum, EnumType, InData);
+}
+
+void FCameraContextDataTable::SetObjectData(FCameraContextDataID InID, UObject* InData)
+{
+	TObjectPtr<UObject> ActualData(InData);
+	SetDataImpl(InID, ECameraContextDataType::Object, nullptr, ActualData);
+}
+
+void FCameraContextDataTable::SetClassData(FCameraContextDataID InID, UClass* InData)
+{
+	TObjectPtr<UClass> ActualData(InData);
+	SetDataImpl(InID, ECameraContextDataType::Class, nullptr, ActualData);
+}
+
+void FCameraContextDataTable::SetStructViewData(FCameraContextDataID InID, const FStructView& InData)
 {
 	FEntry* Entry = FindEntry(InID);
 	if (ensure(Entry && Entry->Type == ECameraContextDataType::Struct && InData.GetScriptStruct() == Entry->TypeObject))
@@ -281,7 +380,7 @@ void FCameraContextDataTable::SetData(FCameraContextDataID InID, const FStructVi
 	}
 }
 
-void FCameraContextDataTable::SetData(FCameraContextDataID InID, const FInstancedStruct& InData)
+void FCameraContextDataTable::SetInstancedStructData(FCameraContextDataID InID, const FInstancedStruct& InData)
 {
 	FEntry* Entry = FindEntry(InID);
 	if (ensure(Entry && Entry->Type == ECameraContextDataType::Struct && InData.GetScriptStruct() == Entry->TypeObject))
@@ -318,7 +417,7 @@ const uint8* FCameraContextDataTable::GetData(FCameraContextDataID DataID, ECame
 	const uint8* Data = TryGetData(DataID, ExpectedDataType, ExpectedDataTypeObject);
 	ensureMsgf(
 			Data, 
-			TEXT("Can't get camera context data (ID '%s') because it doesn't exist in the table."), 
+			TEXT("Can't get camera context data (ID '%s') because it doesn't exist in the table, or isn't of the expected data type."), 
 			*DataID.DataName.ToString());
 	return Data;
 }
@@ -328,7 +427,7 @@ const uint8* FCameraContextDataTable::TryGetData(FCameraContextDataID DataID, EC
 	const FEntry* Entry = FindEntry(DataID);
 	if (Entry)
 	{
-		if (ensure(Entry->Type == ExpectedDataType && Entry->TypeObject == ExpectedDataTypeObject))
+		if (Entry->Type == ExpectedDataType && Entry->TypeObject == ExpectedDataTypeObject)
 		{
 			return Memory + Entry->Offset;
 		}
diff --git a/Engine/Plugins/Cameras/GameplayCameras/Source/GameplayCameras/Private/Core/CameraRigAsset.cpp b/Engine/Plugins/Cameras/GameplayCameras/Source/GameplayCameras/Private/Core/CameraRigAsset.cpp
--- a/Engine/Plugins/Cameras/GameplayCameras/Source/GameplayCameras/Private/Core/CameraRigAsset.cpp
+++ b/Engine/Plugins/Cameras/GameplayCameras/Source/GameplayCameras/Private/Core/CameraRigAsset.cpp
@@ -89,6 +89,16 @@ UCameraRigBlendableParameter* FCameraRigInterface::FindBlendableParameterByName(
 	return FoundItem ? *FoundItem : nullptr;
 }
 
+UCameraRigDataParameter* FCameraRigInterface::FindDataParameterByName(const FString& ParameterName) const
+{
+	const TObjectPtr<UCameraRigDataParameter>* FoundItem = DataParameters.FindByPredicate(
+			[&ParameterName](UCameraRigDataParameter* Item)
+			{
+				return Item->InterfaceParameterName == ParameterName;
+			});
+	return FoundItem ? *FoundItem : nullptr;
+}
+
 UCameraRigBlendableParameter* FCameraRigInterface::FindBlendableParameterByGuid(const FGuid& ParameterGuid) const
 {
 	const TObjectPtr<UCameraRigBlendableParameter>* FoundItem = BlendableParameters.FindByPredicate(
diff --git a/Engine/Plugins/Cameras/GameplayCameras/Source/GameplayCameras/Private/Directors/BlueprintCameraDirector.cpp b/Engine/Plugins/Cameras/GameplayCameras/Source/GameplayCameras/Private/Directors/BlueprintCameraDirector.cpp
--- a/Engine/Plugins/Cameras/GameplayCameras/Source/GameplayCameras/Private/Directors/BlueprintCameraDirector.cpp
+++ b/Engine/Plugins/Cameras/GameplayCameras/Source/GameplayCameras/Private/Directors/BlueprintCameraDirector.cpp
@@ -322,6 +322,21 @@ AActor* UBlueprintCameraDirectorEvaluator::FindEvaluationContextOwnerActor(TSubc
 	}
 }
 
+FBlueprintCameraNodeEvaluationResult UBlueprintCameraDirectorEvaluator::GetInitialContextResult() const
+{
+	if (EvaluationContext)
+	{
+		return FBlueprintCameraNodeEvaluationResult(&EvaluationContext->GetInitialResult());
+	}
+	else
+	{
+		FFrame::KismetExecutionMessage(
+				TEXT("Can't access evaluation context's initial result outside of RunCameraDirector"), 
+				ELogVerbosity::Error);
+		return FBlueprintCameraNodeEvaluationResult();
+	}
+}
+
 FBlueprintCameraPose UBlueprintCameraDirectorEvaluator::GetInitialContextCameraPose() const
 {
 	if (EvaluationContext)
diff --git a/Engine/Plugins/Cameras/GameplayCameras/Source/GameplayCameras/Private/GameFramework/BlueprintCameraContextDataTable.cpp b/Engine/Plugins/Cameras/GameplayCameras/Source/GameplayCameras/Private/GameFramework/BlueprintCameraContextDataTable.cpp
--- a/Engine/Plugins/Cameras/GameplayCameras/Source/GameplayCameras/Private/GameFramework/BlueprintCameraContextDataTable.cpp
+++ b/Engine/Plugins/Cameras/GameplayCameras/Source/GameplayCameras/Private/GameFramework/BlueprintCameraContextDataTable.cpp
@@ -6,6 +6,20 @@
 
 #include UE_INLINE_GENERATED_CPP_BY_NAME(BlueprintCameraContextDataTable)
 
+#define UE_PRIVATE_BLUEPRINT_CAMERA_CONTEXT_DATA_TABLE_VALIDATE(ErrorResult)\
+	using namespace UE::Cameras;\
+	if (!DataID.IsValid())\
+	{\
+		FFrame::KismetExecutionMessage(TEXT("Invalid camera context data ID"), ELogVerbosity::Error);\
+		return ErrorResult;\
+	}\
+	if (!ContextDataTable.IsValid())\
+	{\
+		FFrame::KismetExecutionMessage(TEXT("No camera context data table has been set"), ELogVerbosity::Error);\
+		return ErrorResult;\
+	}\
+	FCameraContextDataTable* ActualTable = ContextDataTable.GetContextDataTable();
+
 FBlueprintCameraContextDataTable::FBlueprintCameraContextDataTable()
 	: PrivateContextDataTable(nullptr)
 {
@@ -16,34 +30,81 @@ FBlueprintCameraContextDataTable::FBlueprintCameraContextDataTable(FCameraContex
 {
 }
 
+FName UBlueprintCameraContextDataTableFunctionLibrary::GetNameData(const FBlueprintCameraContextDataTable& ContextDataTable, FCameraContextDataID DataID)
+{
+	UE_PRIVATE_BLUEPRINT_CAMERA_CONTEXT_DATA_TABLE_VALIDATE(NAME_None);
+	return ActualTable->GetNameData(DataID);
+}
+
+FString UBlueprintCameraContextDataTableFunctionLibrary::GetStringData(const FBlueprintCameraContextDataTable& ContextDataTable, FCameraContextDataID DataID)
+{
+	UE_PRIVATE_BLUEPRINT_CAMERA_CONTEXT_DATA_TABLE_VALIDATE(FString());
+	return ActualTable->GetStringData(DataID);
+}
+
+uint8 UBlueprintCameraContextDataTableFunctionLibrary::GetEnumData(const FBlueprintCameraContextDataTable& ContextDataTable, FCameraContextDataID DataID, const UEnum* EnumType)
+{
+	UE_PRIVATE_BLUEPRINT_CAMERA_CONTEXT_DATA_TABLE_VALIDATE(0);
+	return ActualTable->GetEnumData(DataID, EnumType);
+}
+
 FInstancedStruct UBlueprintCameraContextDataTableFunctionLibrary::GetStructData(const FBlueprintCameraContextDataTable& ContextDataTable, FCameraContextDataID DataID, const UScriptStruct* DataStructType)
 {
-	using namespace UE::Cameras;
+	UE_PRIVATE_BLUEPRINT_CAMERA_CONTEXT_DATA_TABLE_VALIDATE(FInstancedStruct());
+	return ActualTable->GetInstancedStructData(DataID, DataStructType);
+}
+
+UObject* UBlueprintCameraContextDataTableFunctionLibrary::GetObjectData(const FBlueprintCameraContextDataTable& ContextDataTable, FCameraContextDataID DataID)
+{
+	UE_PRIVATE_BLUEPRINT_CAMERA_CONTEXT_DATA_TABLE_VALIDATE(nullptr);
+	return ActualTable->GetObjectData(DataID);
+}
+
+UClass* UBlueprintCameraContextDataTableFunctionLibrary::GetClassData(const FBlueprintCameraContextDataTable& ContextDataTable, FCameraContextDataID DataID)
+{
+	UE_PRIVATE_BLUEPRINT_CAMERA_CONTEXT_DATA_TABLE_VALIDATE(nullptr);
+	return ActualTable->GetClassData(DataID);
+}
+
+bool UBlueprintCameraContextDataTableFunctionLibrary::SetNameData(const FBlueprintCameraContextDataTable& ContextDataTable, FCameraContextDataID DataID, const FName& Data)
+{
+	UE_PRIVATE_BLUEPRINT_CAMERA_CONTEXT_DATA_TABLE_VALIDATE(false);
+	ActualTable->SetNameData(DataID, Data);
+	return true;
+}
+
+bool UBlueprintCameraContextDataTableFunctionLibrary::SetStringData(const FBlueprintCameraContextDataTable& ContextDataTable, FCameraContextDataID DataID, const FString& Data)
+{
+	UE_PRIVATE_BLUEPRINT_CAMERA_CONTEXT_DATA_TABLE_VALIDATE(false);
+	ActualTable->SetStringData(DataID, Data);
+	return true;
+}
 
-	if (FCameraContextDataTable* ActualTable = ContextDataTable.GetContextDataTable())
-	{
-		if (DataID.IsValid())
-		{
-			const uint8* RawData = ActualTable->GetData(DataID, ECameraContextDataType::Struct, DataStructType);
-			FInstancedStruct ReturnValue;
-			ReturnValue.InitializeAs(DataStructType, RawData);
-			return ReturnValue;
-		}
-	}
+bool UBlueprintCameraContextDataTableFunctionLibrary::SetEnumData(const FBlueprintCameraContextDataTable& ContextDataTable, FCameraContextDataID DataID, const UEnum* EnumType, uint8 Data)
+{
+	UE_PRIVATE_BLUEPRINT_CAMERA_CONTEXT_DATA_TABLE_VALIDATE(false);
+	ActualTable->SetEnumData(DataID, EnumType, Data);
+	return true;
+}
 
-	return FInstancedStruct();
+bool UBlueprintCameraContextDataTableFunctionLibrary::SetStructData(const FBlueprintCameraContextDataTable& ContextDataTable, FCameraContextDataID DataID, const FInstancedStruct& Data)
+{
+	UE_PRIVATE_BLUEPRINT_CAMERA_CONTEXT_DATA_TABLE_VALIDATE(false);
+	ActualTable->SetInstancedStructData(DataID, Data);
+	return true;
 }
 
-void UBlueprintCameraContextDataTableFunctionLibrary::SetStructData(const FBlueprintCameraContextDataTable& ContextDataTable, FCameraContextDataID DataID, const FInstancedStruct& Data)
+bool UBlueprintCameraContextDataTableFunctionLibrary::SetObjectData(const FBlueprintCameraContextDataTable& ContextDataTable, FCameraContextDataID DataID, UObject* Data)
 {
-	using namespace UE::Cameras;
+	UE_PRIVATE_BLUEPRINT_CAMERA_CONTEXT_DATA_TABLE_VALIDATE(false);
+	ActualTable->SetObjectData(DataID, Data);
+	return true;
+}
 
-	if (FCameraContextDataTable* ActualTable = ContextDataTable.GetContextDataTable())
-	{
-		if (DataID.IsValid())
-		{
-			ActualTable->SetData(DataID, Data);
-		}
-	}
+bool UBlueprintCameraContextDataTableFunctionLibrary::SetClassData(const FBlueprintCameraContextDataTable& ContextDataTable, FCameraContextDataID DataID, UClass* Data)
+{
+	UE_PRIVATE_BLUEPRINT_CAMERA_CONTEXT_DATA_TABLE_VALIDATE(false);
+	ActualTable->SetClassData(DataID, Data);
+	return true;
 }
 
diff --git a/Engine/Plugins/Cameras/GameplayCameras/Source/GameplayCameras/Private/GameFramework/BlueprintCameraNodeEvaluationResult.cpp b/Engine/Plugins/Cameras/GameplayCameras/Source/GameplayCameras/Private/GameFramework/BlueprintCameraNodeEvaluationResult.cpp
new file mode 100644
--- /dev/null
+++ b/Engine/Plugins/Cameras/GameplayCameras/Source/GameplayCameras/Private/GameFramework/BlueprintCameraNodeEvaluationResult.cpp
@@ -0,0 +1,60 @@
+// Copyright Epic Games, Inc. All Rights Reserved.
+
+#include "GameFramework/BlueprintCameraNodeEvaluationResult.h"
+
+FBlueprintCameraNodeEvaluationResult::FBlueprintCameraNodeEvaluationResult()
+{
+}
+
+FBlueprintCameraNodeEvaluationResult::FBlueprintCameraNodeEvaluationResult(FCameraNodeEvaluationResult* InResult)
+	: Result(InResult)
+{
+}
+
+FBlueprintCameraPose FBlueprintCameraNodeEvaluationResult::GetCameraPose() const
+{
+	if (Result)
+	{
+		return FBlueprintCameraPose::FromCameraPose(Result->CameraPose);
+	}
+	return FBlueprintCameraPose();
+}
+
+FBlueprintCameraVariableTable FBlueprintCameraNodeEvaluationResult::GetVariableTable() const
+{
+	return FBlueprintCameraVariableTable(Result ? &Result->VariableTable : nullptr);
+}
+
+FBlueprintCameraContextDataTable FBlueprintCameraNodeEvaluationResult::GetContextDataTable() const
+{
+	return FBlueprintCameraContextDataTable(Result ? &Result->ContextDataTable : nullptr);
+}
+
+void FBlueprintCameraNodeEvaluationResult::SetCameraPose(const FBlueprintCameraPose& CameraPose)
+{
+	if (Result)
+	{
+		CameraPose.ApplyTo(Result->CameraPose);
+	}
+}
+
+FBlueprintCameraPose UBlueprintCameraNodeEvaluationResultFunctionLibrary::GetCameraPose(const FBlueprintCameraNodeEvaluationResult& Data)
+{
+	return Data.GetCameraPose();
+}
+
+FBlueprintCameraVariableTable UBlueprintCameraNodeEvaluationResultFunctionLibrary::GetVariableTable(const FBlueprintCameraNodeEvaluationResult& Data)
+{
+	return Data.GetVariableTable();
+}
+
+FBlueprintCameraContextDataTable UBlueprintCameraNodeEvaluationResultFunctionLibrary::GetContextDataTable(const FBlueprintCameraNodeEvaluationResult& Data)
+{
+	return Data.GetContextDataTable();
+}
+
+void UBlueprintCameraNodeEvaluationResultFunctionLibrary::SetCameraPose(FBlueprintCameraNodeEvaluationResult& Data, const FBlueprintCameraPose& CameraPose)
+{
+	Data.SetCameraPose(CameraPose);
+}
+
diff --git a/Engine/Plugins/Cameras/GameplayCameras/Source/GameplayCameras/Private/GameFramework/CameraRigParameterInterop.cpp b/Engine/Plugins/Cameras/GameplayCameras/Source/GameplayCameras/Private/GameFramework/CameraRigParameterInterop.cpp
--- a/Engine/Plugins/Cameras/GameplayCameras/Source/GameplayCameras/Private/GameFramework/CameraRigParameterInterop.cpp
+++ b/Engine/Plugins/Cameras/GameplayCameras/Source/GameplayCameras/Private/GameFramework/CameraRigParameterInterop.cpp
@@ -5,7 +5,10 @@
 #include "Core/CameraRigAsset.h"
 #include "Core/CameraNodeEvaluator.h"
 #include "Core/CameraVariableTable.h"
+#include "GameFramework/BlueprintCameraContextDataTable.h"
+#include "GameFramework/BlueprintCameraNodeEvaluationResult.h"
 #include "GameFramework/BlueprintCameraVariableTable.h"
+#include "Templates/UnrealTypeTraits.h"
 
 #define LOCTEXT_NAMESPACE "CameraRigParameterInterop"
 
@@ -15,7 +18,7 @@ namespace UE::Cameras::Private
 {
 
 template<typename VariableAssetType>
-void SetCameraRigParameter(FBlueprintCameraVariableTable& VariableTable, VariableAssetType* PrivateVariable, typename VariableAssetType::ValueType Value)
+void SetCameraRigBlendableParameter(FBlueprintCameraVariableTable VariableTable, VariableAssetType* PrivateVariable, typename VariableAssetType::ValueType Value)
 {
 	if (!VariableTable.IsValid())
 	{
@@ -31,6 +34,61 @@ void SetCameraRigParameter(FBlueprintCameraVariableTable& VariableTable, Variabl
 	VariableTable.GetVariableTable()->SetValue(PrivateVariable, Value, true);
 }
 
+UCameraVariableAsset* GetParameterPrivateVariable(UCameraRigAsset* CameraRig, const FString& ParameterName)
+{
+	UCameraRigBlendableParameter* BlendableParameter = CameraRig->Interface.FindBlendableParameterByName(ParameterName);
+	if (!BlendableParameter)
+	{
+		const FText Text = LOCTEXT("NoSuchBlendableParameter", "No parameter '{0}' found on camera rig '{1}'. Setting this camera variable table value will most probably accomplish nothing.");
+		FFrame::KismetExecutionMessage(*FText::Format(Text, FText::FromString(ParameterName), FText::FromString(CameraRig->GetPathName())).ToString(), ELogVerbosity::Warning);
+		return nullptr;
+	}
+
+	if (!BlendableParameter->PrivateVariable)
+	{
+		const FText Text = LOCTEXT("CameraRigNeedsBuilding", "Parameter '{0}' isn't built. Please build camera rig '{1}'.");
+		FFrame::KismetExecutionMessage(*FText::Format(Text, FText::FromString(ParameterName), FText::FromString(CameraRig->GetPathName())).ToString(), ELogVerbosity::Warning);
+		return nullptr;
+	}
+
+	return BlendableParameter->PrivateVariable;
+}
+
+bool ValidateSetCameraRigDataParameter(FBlueprintCameraContextDataTable& ContextDataTable, FCameraContextDataID DataID)
+{
+	if (!ContextDataTable.IsValid())
+	{
+		FFrame::KismetExecutionMessage(TEXT("Invalid camera variable table was passed."), ELogVerbosity::Error);
+		return false;
+	}
+	if (!DataID.IsValid())
+	{
+		FFrame::KismetExecutionMessage(TEXT("No camera rig was passed."), ELogVerbosity::Error);
+		return false;
+	}
+	return true;
+}
+
+FCameraContextDataID GetParameterPrivateDataID(UCameraRigAsset* CameraRig, const FString& ParameterName)
+{
+	UCameraRigDataParameter* DataParameter = CameraRig->Interface.FindDataParameterByName(ParameterName);
+	if (!DataParameter)
+	{
+		const FText Text = LOCTEXT("NoSuchDataParameter", "No parameter '{0}' found on camera rig '{1}'. Setting this data will most probably accomplish nothing.");
+		FFrame::KismetExecutionMessage(*FText::Format(Text, FText::FromString(ParameterName), FText::FromString(CameraRig->GetPathName())).ToString(), ELogVerbosity::Warning);
+		return FCameraContextDataID();
+	}
+
+	if (!DataParameter->PrivateDataID)
+	{
+		const FText Text = LOCTEXT("CameraRigNeedsBuilding", "Parameter '{0}' isn't built. Please build camera rig '{1}'.");
+		FFrame::KismetExecutionMessage(*FText::Format(Text, FText::FromString(ParameterName), FText::FromString(CameraRig->GetPathName())).ToString(), ELogVerbosity::Warning);
+		return FCameraContextDataID();
+	}
+
+	return DataParameter->PrivateDataID;
+}
+
 }  // namespace UE::Cameras::Private
 
 UCameraRigParameterInterop::UCameraRigParameterInterop(const FObjectInitializer& ObjectInit)
@@ -38,96 +96,136 @@ UCameraRigParameterInterop::UCameraRigParameterInterop(const FObjectInitializer&
 {
 }
 
-void UCameraRigParameterInterop::SetBooleanParameter(FBlueprintCameraVariableTable& VariableTable, UCameraRigAsset* CameraRig, const FString& ParameterName, bool ParameterValue)
+void UCameraRigParameterInterop::SetBooleanParameter(FBlueprintCameraNodeEvaluationResult& Result, UCameraRigAsset* CameraRig, const FString& ParameterName, bool ParameterValue)
 {
-	UE::Cameras::Private::SetCameraRigParameter(
-			VariableTable, 
-			Cast<UBooleanCameraVariable>(GetParameterPrivateVariable(CameraRig, ParameterName)), 
+	UE::Cameras::Private::SetCameraRigBlendableParameter(
+			Result.GetVariableTable(), 
+			Cast<UBooleanCameraVariable>(UE::Cameras::Private::GetParameterPrivateVariable(CameraRig, ParameterName)), 
 			ParameterValue);
 }
 
-void UCameraRigParameterInterop::SetIntegerParameter(FBlueprintCameraVariableTable& VariableTable, UCameraRigAsset* CameraRig, const FString& ParameterName, int32 ParameterValue)
+void UCameraRigParameterInterop::SetIntegerParameter(FBlueprintCameraNodeEvaluationResult& Result, UCameraRigAsset* CameraRig, const FString& ParameterName, int32 ParameterValue)
 {
-	UE::Cameras::Private::SetCameraRigParameter(
-			VariableTable, 
-			Cast<UInteger32CameraVariable>(GetParameterPrivateVariable(CameraRig, ParameterName)), 
+	UE::Cameras::Private::SetCameraRigBlendableParameter(
+			Result.GetVariableTable(), 
+			Cast<UInteger32CameraVariable>(UE::Cameras::Private::GetParameterPrivateVariable(CameraRig, ParameterName)), 
 			ParameterValue);
 }
 
-void UCameraRigParameterInterop::SetFloatParameter(FBlueprintCameraVariableTable& VariableTable, UCameraRigAsset* CameraRig, const FString& ParameterName, double ParameterValue)
+void UCameraRigParameterInterop::SetFloatParameter(FBlueprintCameraNodeEvaluationResult& Result, UCameraRigAsset* CameraRig, const FString& ParameterName, double ParameterValue)
 {
-	UE::Cameras::Private::SetCameraRigParameter(
-			VariableTable, 
-			Cast<UFloatCameraVariable>(GetParameterPrivateVariable(CameraRig, ParameterName)),
+	UE::Cameras::Private::SetCameraRigBlendableParameter(
+			Result.GetVariableTable(), 
+			Cast<UFloatCameraVariable>(UE::Cameras::Private::GetParameterPrivateVariable(CameraRig, ParameterName)),
 			(float)ParameterValue);
 }
 
-void UCameraRigParameterInterop::SetDoubleParameter(FBlueprintCameraVariableTable& VariableTable, UCameraRigAsset* CameraRig, const FString& ParameterName, double ParameterValue)
+void UCameraRigParameterInterop::SetDoubleParameter(FBlueprintCameraNodeEvaluationResult& Result, UCameraRigAsset* CameraRig, const FString& ParameterName, double ParameterValue)
 {
-	UE::Cameras::Private::SetCameraRigParameter(
-			VariableTable, 
-			Cast<UDoubleCameraVariable>(GetParameterPrivateVariable(CameraRig, ParameterName)),
+	UE::Cameras::Private::SetCameraRigBlendableParameter(
+			Result.GetVariableTable(), 
+			Cast<UDoubleCameraVariable>(UE::Cameras::Private::GetParameterPrivateVariable(CameraRig, ParameterName)),
 			ParameterValue);
 }
 
-void UCameraRigParameterInterop::SetVector2Parameter(FBlueprintCameraVariableTable& VariableTable, UCameraRigAsset* CameraRig, const FString& ParameterName, FVector2D ParameterValue)
+void UCameraRigParameterInterop::SetVector2Parameter(FBlueprintCameraNodeEvaluationResult& Result, UCameraRigAsset* CameraRig, const FString& ParameterName, FVector2D ParameterValue)
 {
-	UE::Cameras::Private::SetCameraRigParameter(
-			VariableTable, 
-			Cast<UVector2dCameraVariable>(GetParameterPrivateVariable(CameraRig, ParameterName)),
+	UE::Cameras::Private::SetCameraRigBlendableParameter(
+			Result.GetVariableTable(), 
+			Cast<UVector2dCameraVariable>(UE::Cameras::Private::GetParameterPrivateVariable(CameraRig, ParameterName)),
 			ParameterValue);
 }
 
-void UCameraRigParameterInterop::SetVector3Parameter(FBlueprintCameraVariableTable& VariableTable, UCameraRigAsset* CameraRig, const FString& ParameterName, FVector ParameterValue)
+void UCameraRigParameterInterop::SetVector3Parameter(FBlueprintCameraNodeEvaluationResult& Result, UCameraRigAsset* CameraRig, const FString& ParameterName, FVector ParameterValue)
 {
-	UE::Cameras::Private::SetCameraRigParameter(
-			VariableTable, 
-			Cast<UVector3dCameraVariable>(GetParameterPrivateVariable(CameraRig, ParameterName)),
+	UE::Cameras::Private::SetCameraRigBlendableParameter(
+			Result.GetVariableTable(), 
+			Cast<UVector3dCameraVariable>(UE::Cameras::Private::GetParameterPrivateVariable(CameraRig, ParameterName)),
 			ParameterValue);
 }
 
-void UCameraRigParameterInterop::SetVector4Parameter(FBlueprintCameraVariableTable& VariableTable, UCameraRigAsset* CameraRig, const FString& ParameterName, FVector4 ParameterValue)
+void UCameraRigParameterInterop::SetVector4Parameter(FBlueprintCameraNodeEvaluationResult& Result, UCameraRigAsset* CameraRig, const FString& ParameterName, FVector4 ParameterValue)
 {
-	UE::Cameras::Private::SetCameraRigParameter(
-			VariableTable, 
-			Cast<UVector4dCameraVariable>(GetParameterPrivateVariable(CameraRig, ParameterName)),
+	UE::Cameras::Private::SetCameraRigBlendableParameter(
+			Result.GetVariableTable(), 
+			Cast<UVector4dCameraVariable>(UE::Cameras::Private::GetParameterPrivateVariable(CameraRig, ParameterName)),
 			ParameterValue);
 }
 
-void UCameraRigParameterInterop::SetRotatorParameter(FBlueprintCameraVariableTable& VariableTable, UCameraRigAsset* CameraRig, const FString& ParameterName, FRotator ParameterValue)
+void UCameraRigParameterInterop::SetRotatorParameter(FBlueprintCameraNodeEvaluationResult& Result, UCameraRigAsset* CameraRig, const FString& ParameterName, FRotator ParameterValue)
 {
-	UE::Cameras::Private::SetCameraRigParameter(
-			VariableTable, 
-			Cast<URotator3dCameraVariable>(GetParameterPrivateVariable(CameraRig, ParameterName)),
+	UE::Cameras::Private::SetCameraRigBlendableParameter(
+			Result.GetVariableTable(), 
+			Cast<URotator3dCameraVariable>(UE::Cameras::Private::GetParameterPrivateVariable(CameraRig, ParameterName)),
 			ParameterValue);
 }
 
-void UCameraRigParameterInterop::SetTransformParameter(FBlueprintCameraVariableTable& VariableTable, UCameraRigAsset* CameraRig, const FString& ParameterName, FTransform ParameterValue)
+void UCameraRigParameterInterop::SetTransformParameter(FBlueprintCameraNodeEvaluationResult& Result, UCameraRigAsset* CameraRig, const FString& ParameterName, FTransform ParameterValue)
 {
-	UE::Cameras::Private::SetCameraRigParameter(
-			VariableTable, 
-			Cast<UTransform3dCameraVariable>(GetParameterPrivateVariable(CameraRig, ParameterName)),
+	UE::Cameras::Private::SetCameraRigBlendableParameter(
+			Result.GetVariableTable(), 
+			Cast<UTransform3dCameraVariable>(UE::Cameras::Private::GetParameterPrivateVariable(CameraRig, ParameterName)),
 			ParameterValue);
 }
 
-UCameraVariableAsset* UCameraRigParameterInterop::GetParameterPrivateVariable(UCameraRigAsset* CameraRig, const FString& ParameterName)
+void UCameraRigParameterInterop::SetNameParameter(UPARAM(Ref) FBlueprintCameraNodeEvaluationResult& Result, UCameraRigAsset* CameraRig, const FString& ParameterName, FName ParameterValue)
 {
-	UCameraRigBlendableParameter* BlendableParameter = CameraRig->Interface.FindBlendableParameterByName(ParameterName);
-	if (!BlendableParameter)
+	const FCameraContextDataID DataID(UE::Cameras::Private::GetParameterPrivateDataID(CameraRig, ParameterName));
+	FBlueprintCameraContextDataTable ContextDataTable(Result.GetContextDataTable());
+	if (UE::Cameras::Private::ValidateSetCameraRigDataParameter(ContextDataTable, DataID))
 	{
-		const FText Text = LOCTEXT("NoSuchParameter", "No parameter '{0}' found on camera rig '{1}'. Setting this camera variable table value will most probably accomplish nothing.");
-		FFrame::KismetExecutionMessage(*FText::Format(Text, FText::FromString(ParameterName), FText::FromString(CameraRig->GetPathName())).ToString(), ELogVerbosity::Warning);
-		return nullptr;
+		ContextDataTable.GetContextDataTable()->SetNameData(DataID, ParameterValue);
 	}
+}
 
-	if (!BlendableParameter->PrivateVariable)
+void UCameraRigParameterInterop::SetStringParameter(FBlueprintCameraNodeEvaluationResult& Result, UCameraRigAsset* CameraRig, const FString& ParameterName, FString ParameterValue)
+{
+	const FCameraContextDataID DataID(UE::Cameras::Private::GetParameterPrivateDataID(CameraRig, ParameterName));
+	FBlueprintCameraContextDataTable ContextDataTable(Result.GetContextDataTable());
+	if (UE::Cameras::Private::ValidateSetCameraRigDataParameter(ContextDataTable, DataID))
 	{
-		const FText Text = LOCTEXT("CameraRigNeedsBuilding", "Parameter '{0}' isn't built. Please build camera rig '{1}'.");
-		FFrame::KismetExecutionMessage(*FText::Format(Text, FText::FromString(ParameterName), FText::FromString(CameraRig->GetPathName())).ToString(), ELogVerbosity::Warning);
-		return nullptr;
+		ContextDataTable.GetContextDataTable()->SetStringData(DataID, ParameterValue);
 	}
+}
 
-	return BlendableParameter->PrivateVariable;
+void UCameraRigParameterInterop::SetEnumParameter(FBlueprintCameraNodeEvaluationResult& Result, UCameraRigAsset* CameraRig, const FString& ParameterName, const UEnum* EnumType, uint8 ParameterValue)
+{
+	const FCameraContextDataID DataID(UE::Cameras::Private::GetParameterPrivateDataID(CameraRig, ParameterName));
+	FBlueprintCameraContextDataTable ContextDataTable(Result.GetContextDataTable());
+	if (UE::Cameras::Private::ValidateSetCameraRigDataParameter(ContextDataTable, DataID))
+	{
+		ContextDataTable.GetContextDataTable()->SetEnumData(DataID, EnumType, ParameterValue);
+	}
+}
+
+void UCameraRigParameterInterop::SetStructParameter(FBlueprintCameraNodeEvaluationResult& Result, UCameraRigAsset* CameraRig, const FString& ParameterName, const FInstancedStruct& ParameterValue)
+{
+	const FCameraContextDataID DataID(UE::Cameras::Private::GetParameterPrivateDataID(CameraRig, ParameterName));
+	FBlueprintCameraContextDataTable ContextDataTable(Result.GetContextDataTable());
+	if (UE::Cameras::Private::ValidateSetCameraRigDataParameter(ContextDataTable, DataID))
+	{
+		ContextDataTable.GetContextDataTable()->SetInstancedStructData(DataID, ParameterValue);
+	}
+}
+
+void UCameraRigParameterInterop::SetObjectParameter(FBlueprintCameraNodeEvaluationResult& Result, UCameraRigAsset* CameraRig, const FString& ParameterName, UObject* ParameterValue)
+{
+	const FCameraContextDataID DataID(UE::Cameras::Private::GetParameterPrivateDataID(CameraRig, ParameterName));
+	FBlueprintCameraContextDataTable ContextDataTable(Result.GetContextDataTable());
+	if (UE::Cameras::Private::ValidateSetCameraRigDataParameter(ContextDataTable, DataID))
+	{
+		ContextDataTable.GetContextDataTable()->SetObjectData(DataID, ParameterValue);
+	}
+}
+
+void UCameraRigParameterInterop::SetClassParameter(FBlueprintCameraNodeEvaluationResult& Result, UCameraRigAsset* CameraRig, const FString& ParameterName, UClass* ParameterValue)
+{
+	const FCameraContextDataID DataID(UE::Cameras::Private::GetParameterPrivateDataID(CameraRig, ParameterName));
+	FBlueprintCameraContextDataTable ContextDataTable(Result.GetContextDataTable());
+	if (UE::Cameras::Private::ValidateSetCameraRigDataParameter(ContextDataTable, DataID))
+	{
+		ContextDataTable.GetContextDataTable()->SetClassData(DataID, ParameterValue);
+	}
 }
 
 #undef LOCTEXT_NAMESPACE
diff --git a/Engine/Plugins/Cameras/GameplayCameras/Source/GameplayCameras/Private/GameFramework/GameplayCameraComponent.cpp b/Engine/Plugins/Cameras/GameplayCameras/Source/GameplayCameras/Private/GameFramework/GameplayCameraComponent.cpp
--- a/Engine/Plugins/Cameras/GameplayCameras/Source/GameplayCameras/Private/GameFramework/GameplayCameraComponent.cpp
+++ b/Engine/Plugins/Cameras/GameplayCameras/Source/GameplayCameras/Private/GameFramework/GameplayCameraComponent.cpp
@@ -179,6 +179,13 @@ void UGameplayCameraComponent::ActivateCameraEvaluationContext(APlayerController
 		return ErrorResult;\
 	}
 
+FBlueprintCameraNodeEvaluationResult UGameplayCameraComponent::GetInitialResult() const
+{
+	UE_PRIVATE_GAMEPLAY_CAMERA_COMPONENT_VALIDATE_EVALUATION_CONTEXT("Can't get initial result", FBlueprintCameraNodeEvaluationResult());
+
+	return FBlueprintCameraNodeEvaluationResult(&EvaluationContext->GetInitialResult());
+}
+
 FBlueprintCameraPose UGameplayCameraComponent::GetInitialPose() const
 {
 	UE_PRIVATE_GAMEPLAY_CAMERA_COMPONENT_VALIDATE_EVALUATION_CONTEXT("Can't get initial camera pose", FBlueprintCameraPose());
diff --git a/Engine/Plugins/Cameras/GameplayCameras/Source/GameplayCameras/Private/Helpers/CameraParameterOverrideHelper.cpp b/Engine/Plugins/Cameras/GameplayCameras/Source/GameplayCameras/Private/Helpers/CameraParameterOverrideHelper.cpp
--- a/Engine/Plugins/Cameras/GameplayCameras/Source/GameplayCameras/Private/Helpers/CameraParameterOverrideHelper.cpp
+++ b/Engine/Plugins/Cameras/GameplayCameras/Source/GameplayCameras/Private/Helpers/CameraParameterOverrideHelper.cpp
@@ -86,7 +86,7 @@ void ApplyDataParameterOverride(
 
 	// Write the override value into the context data table.
 	FCameraContextDataID ParameterDataID = DataParameter->PrivateDataID;
-	ContextDataTable.SetData(ParameterDataID, ParameterValue);
+	ContextDataTable.SetStructViewData(ParameterDataID, ParameterValue);
 }
 
 }  // namespace Internal
diff --git a/Engine/Plugins/Cameras/GameplayCameras/Source/GameplayCameras/Public/Core/CameraContextDataTable.h b/Engine/Plugins/Cameras/GameplayCameras/Source/GameplayCameras/Public/Core/CameraContextDataTable.h
--- a/Engine/Plugins/Cameras/GameplayCameras/Source/GameplayCameras/Public/Core/CameraContextDataTable.h
+++ b/Engine/Plugins/Cameras/GameplayCameras/Source/GameplayCameras/Public/Core/CameraContextDataTable.h
@@ -50,16 +50,42 @@ public:
 
 	// Getter methods.
 
-	template<typename DataType>
-	const DataType* GetData(FCameraContextDataID InID) const;
+	const FName& GetNameData(FCameraContextDataID InID) const;
+	const FString& GetStringData(FCameraContextDataID InID) const;
+	uint8 GetEnumData(FCameraContextDataID InID, const UEnum* EnumType) const;
+	FConstStructView GetStructViewData(FCameraContextDataID InID, const UScriptStruct* StructType) const;
+	FInstancedStruct GetInstancedStructData(FCameraContextDataID InID, const UScriptStruct* StructType) const;
+	UObject* GetObjectData(FCameraContextDataID InID) const;
+	UClass* GetClassData(FCameraContextDataID InID) const;
+	
+	template<typename EnumType>
+	EnumType GetEnumData(FCameraContextDataID InID) const;
+
+	template<typename StructType>
+	const StructType& GetStructData(FCameraContextDataID InID) const;
+
+	template<typename ObjectClass>
+	ObjectClass* GetObjectData(FCameraContextDataID InID) const;
+
+	template<typename BaseClass>
+	TSubclassOf<BaseClass> GetClassData(FCameraContextDataID InID) const;
 
 	// Setter methods.
 
-	template<typename DataType>
-	void SetData(FCameraContextDataID InID, const DataType& InData);
+	void SetNameData(FCameraContextDataID InID, const FName& InData);
+	void SetStringData(FCameraContextDataID InID, const FString& InData);
+	void SetEnumData(FCameraContextDataID InID, const UEnum* EnumType, uint8 InData);
+	void SetObjectData(FCameraContextDataID InID, UObject* InData);
+	void SetClassData(FCameraContextDataID InID, UClass* InData);
 
-	void SetData(FCameraContextDataID InID, const FStructView& InData);
-	void SetData(FCameraContextDataID InID, const FInstancedStruct& InData);
+	template<typename EnumType>
+	void SetEnumData(FCameraContextDataID InID, EnumType InData);
+
+	template<typename StructType>
+	void SetStructData(FCameraContextDataID InID, const StructType& InData);
+
+	void SetStructViewData(FCameraContextDataID InID, const FStructView& InData);
+	void SetInstancedStructData(FCameraContextDataID InID, const FInstancedStruct& InData);
 
 public:
 
@@ -123,6 +149,12 @@ private:
 	void ReallocateBuffer(uint32 MinRequired = 0);
 	void DestroyBuffer();
 
+	template<typename StorageType>
+	const StorageType* GetDataImpl(FCameraContextDataID InID, ECameraContextDataType DataType, const UObject* DataTypeObject) const;
+
+	template<typename StorageType>
+	bool SetDataImpl(FCameraContextDataID InID, ECameraContextDataType DataType, const UObject* DataTypeObject, const StorageType& InData);
+
 private:
 
 	TArray<FEntry> Entries;
@@ -139,29 +171,84 @@ private:
 
 ENUM_CLASS_FLAGS(FCameraContextDataTable::EEntryFlags)
 
-template<typename DataType>
-const DataType* FCameraContextDataTable::GetData(FCameraContextDataID InID) const
+template<typename EnumType>
+EnumType FCameraContextDataTable::GetEnumData(FCameraContextDataID InID) const
+{
+	if (const uint8* Value = GetDataImpl<uint8>(InID, ECameraContextDataType::Enum, StaticEnum<EnumType>()))
+	{
+		return EnumType(*Value);
+	}
+	return EnumType();
+}
+
+template<typename StructType>
+const StructType& FCameraContextDataTable::GetStructData(FCameraContextDataID InID) const
+{
+	if (const StructType* Value = GetDataImpl<StructType>(InID, ECameraContextDataType::Struct, StructType::StaticStruct()))
+	{
+		return *Value;
+	}
+
+	static const StructType DefaultValue;
+	return DefaultValue;
+}
+
+template<typename ObjectClass>
+ObjectClass* FCameraContextDataTable::GetObjectData(FCameraContextDataID InID) const
+{
+	if (const TObjectPtr<UObject>* Value = GetDataImpl<TObjectPtr<UObject>>(InID, ECameraContextDataType::Object, nullptr))
+	{
+		return Cast<ObjectClass>(Value->Get());
+	}
+	return nullptr;
+}
+
+template<typename BaseClass>
+TSubclassOf<BaseClass> FCameraContextDataTable::GetClassData(FCameraContextDataID InID) const
+{
+	if (const TObjectPtr<UClass>* Value = GetDataImpl<TObjectPtr<UClass>>(InID, ECameraContextDataType::Class, nullptr))
+	{
+		return TSubclassOf<BaseClass>(Value->Get());
+	}
+	return nullptr;
+}
+
+template<typename EnumType>
+void FCameraContextDataTable::SetEnumData(FCameraContextDataID InID, EnumType InData)
+{
+	SetDataImpl(InID, ECameraContextDataType::Enum, StaticEnum<EnumType>(), InData);
+}
+
+template<typename StructType>
+void FCameraContextDataTable::SetStructData(FCameraContextDataID InID, const StructType& InData)
+{
+	SetDataImpl(InID, ECameraContextDataType::Struct, StructType::StaticStruct(), InData);
+}
+
+template<typename StorageType>
+const StorageType* FCameraContextDataTable::GetDataImpl(FCameraContextDataID InID, ECameraContextDataType DataType, const UObject* DataTypeObject) const
 {
 	const FEntry* Entry = FindEntry(InID);
-	if (Entry)
+	if (Entry && Entry->Type == DataType && Entry->TypeObject == DataTypeObject)
 	{
-		ensureMsgf(Entry->Type == DataType::StaticStruct(), TEXT("Data type mismatch!"));
 		const uint8* RawData = Memory + Entry->Offset;
-		return reinterpret_cast<const DataType*>(RawData);
+		return reinterpret_cast<const StorageType*>(RawData);
 	}
 	return nullptr;
 }
 
-template<typename DataType>
-void FCameraContextDataTable::SetData(FCameraContextDataID InID, const DataType& InData)
+template<typename StorageType>
+bool FCameraContextDataTable::SetDataImpl(FCameraContextDataID InID, ECameraContextDataType DataType, const UObject* DataTypeObject, const StorageType& InData)
 {
-	const FEntry* Entry = FindEntry(InID);
-	if (ensure(Entry))
+	FEntry* Entry = FindEntry(InID);
+	if (Entry && Entry->Type == DataType && Entry->TypeObject == DataTypeObject)
 	{
-		ensureMsgf(Entry->Type == DataType::StaticStruct(), TEXT("Data type mismatch!"));
 		uint8* RawData = Memory + Entry->Offset;
-		*reinterpret_cast<DataType*>(RawData) = InData;
+		*reinterpret_cast<StorageType*>(RawData) = InData;
+		EnumAddFlags(Entry->Flags, EEntryFlags::Written | EEntryFlags::WrittenThisFrame);
+		return true;
 	}
+	return false;
 }
 
 }  // namespace UE::Cameras
diff --git a/Engine/Plugins/Cameras/GameplayCameras/Source/GameplayCameras/Public/Core/CameraRigAsset.h b/Engine/Plugins/Cameras/GameplayCameras/Source/GameplayCameras/Public/Core/CameraRigAsset.h
--- a/Engine/Plugins/Cameras/GameplayCameras/Source/GameplayCameras/Public/Core/CameraRigAsset.h
+++ b/Engine/Plugins/Cameras/GameplayCameras/Source/GameplayCameras/Public/Core/CameraRigAsset.h
@@ -237,6 +237,7 @@ public:
 	
 	/** Finds an exposed parameter by name. */
 	GAMEPLAYCAMERAS_API UCameraRigBlendableParameter* FindBlendableParameterByName(const FString& ParameterName) const;
+	GAMEPLAYCAMERAS_API UCameraRigDataParameter* FindDataParameterByName(const FString& ParameterName) const;
 
 	/** Finds an exposed parameter by Guid. */
 	GAMEPLAYCAMERAS_API UCameraRigBlendableParameter* FindBlendableParameterByGuid(const FGuid& ParameterGuid) const;
diff --git a/Engine/Plugins/Cameras/GameplayCameras/Source/GameplayCameras/Public/Directors/BlueprintCameraDirector.h b/Engine/Plugins/Cameras/GameplayCameras/Source/GameplayCameras/Public/Directors/BlueprintCameraDirector.h
--- a/Engine/Plugins/Cameras/GameplayCameras/Source/GameplayCameras/Public/Directors/BlueprintCameraDirector.h
+++ b/Engine/Plugins/Cameras/GameplayCameras/Source/GameplayCameras/Public/Directors/BlueprintCameraDirector.h
@@ -4,6 +4,7 @@
 
 #include "Core/CameraDirector.h"
 #include "Core/CameraDirectorEvaluator.h"
+#include "GameFramework/BlueprintCameraNodeEvaluationResult.h"
 #include "GameFramework/BlueprintCameraPose.h"
 #include "GameFramework/BlueprintCameraVariableTable.h"
 #include "Templates/SubclassOf.h"
@@ -108,14 +109,14 @@ public:
 	 * Override this method in Blueprint to execute custom logic when this
 	 * camera director gets activated.
 	 */
-	UFUNCTION(BlueprintCallable, BlueprintImplementableEvent, Category="Activation")
+	UFUNCTION(BlueprintCallable, BlueprintImplementableEvent, Category="Camera Director|Activation")
 	void ActivateCameraDirector(const FBlueprintCameraDirectorActivateParams& Params);
 
 	/**
 	 * Override this method in Blueprint to execute custom logic when this
 	 * camera director gets deactivated.
 	 */
-	UFUNCTION(BlueprintCallable, BlueprintImplementableEvent, Category="Activation")
+	UFUNCTION(BlueprintCallable, BlueprintImplementableEvent, Category="Camera Director|Activation")
 	void DeactivateCameraDirector(const FBlueprintCameraDirectorDeactivateParams& Params);
 	
 	/**
@@ -128,40 +129,40 @@ public:
 public:
 
 	/** Activates the given camera rig prefab in the base layer. */
-	UFUNCTION(BlueprintCallable, Category="Activation")
+	UFUNCTION(BlueprintCallable, Category="Camera Director|Activation")
 	void ActivatePersistentBaseCameraRig(UCameraRigAsset* CameraRigPrefab);
 
 	/** Activates the given camera rig prefab in the global layer. */
-	UFUNCTION(BlueprintCallable, Category="Activation")
+	UFUNCTION(BlueprintCallable, Category="Camera Director|Activation")
 	void ActivatePersistentGlobalCameraRig(UCameraRigAsset* CameraRigPrefab);
 
 	/** Activates the given camera rig prefab in the visual layer. */
-	UFUNCTION(BlueprintCallable, Category="Activation")
+	UFUNCTION(BlueprintCallable, Category="Camera Director|Activation")
 	void ActivatePersistentVisualCameraRig(UCameraRigAsset* CameraRigPrefab);
 
 	/** Deactivates the given camera rig prefab in the base layer. */
-	UFUNCTION(BlueprintCallable, Category="Activation")
+	UFUNCTION(BlueprintCallable, Category="Camera Director|Activation")
 	void DeactivatePersistentBaseCameraRig(UCameraRigAsset* CameraRigPrefab);
 
 	/** Deactivates the given camera rig prefab in the global layer. */
-	UFUNCTION(BlueprintCallable, Category="Activation")
+	UFUNCTION(BlueprintCallable, Category="Camera Director|Activation")
 	void DeactivatePersistentGlobalCameraRig(UCameraRigAsset* CameraRigPrefab);
 
 	/** Deactivates the given camera rig prefab in the visual layer. */
-	UFUNCTION(BlueprintCallable, Category="Activation")
+	UFUNCTION(BlueprintCallable, Category="Camera Director|Activation")
 	void DeactivatePersistentVisualCameraRig(UCameraRigAsset* CameraRigPrefab);
 
 public:
 
 	/** Specifies a camera rig to be active this frame. */
-	UFUNCTION(BlueprintCallable, Category="Evaluation")
+	UFUNCTION(BlueprintCallable, Category="Camera Director|Evaluation")
 	void ActivateCameraRig(UCameraRigAsset* CameraRig, bool bForceNewInstance = false);
 
 	/**
 	 * Specifies a camera rig to be active this frame, via a proxy which is later resolved
 	 * via the proxy table of the Blueprint camera director.
 	 */
-	UFUNCTION(BlueprintCallable, Category="Evaluation")
+	UFUNCTION(BlueprintCallable, Category="Camera Director|Evaluation")
 	void ActivateCameraRigViaProxy(UCameraRigProxyAsset* CameraRigProxy);
 
 public:
@@ -171,27 +172,23 @@ public:
 	 * Handles the situation where the evaluation context is an actor component (like a
 	 * UGameplayCameraComponent) or an actor itself.
 	 */
-	UFUNCTION(BlueprintPure, Category="Evaluation", meta=(DeterminesOutputType="ActorClass"))
+	UFUNCTION(BlueprintPure, Category="Camera Director|Evaluation", meta=(DeterminesOutputType="ActorClass"))
 	AActor* FindEvaluationContextOwnerActor(TSubclassOf<AActor> ActorClass) const;
 
-	/**
-	 * Gets the initial evaluation context camera pose.
-	 */
-	UFUNCTION(BlueprintPure, Category="Evaluation")
+	/** Gets the shared evaluation context data. */
+	UFUNCTION(BlueprintPure, Category="Camera Director|Evaluation", meta=(DisplayName="Get Shared Evaluation Data"))
+	FBlueprintCameraNodeEvaluationResult GetInitialContextResult() const;
+
+	UFUNCTION(BlueprintPure, Category="Camera Director|Evaluation",
+			meta=(DeprecatedFunction, DeprecationMessage="Please use GetSharedEvaluationData"))
 	FBlueprintCameraPose GetInitialContextCameraPose() const;
 
-	/**
-	 * Sets the initial evaluation context camera pose.
-	 * WARNING: this will change the initial pose of ALL running camera rigs!
-	 */
-	UFUNCTION(BlueprintCallable, Category="Evaluation")
+	UFUNCTION(BlueprintCallable, Category="Camera Director|Evaluation",
+			meta=(DeprecatedFunction, DeprecationMessage="Please use GetSharedEvaluationData"))
 	void SetInitialContextCameraPose(const FBlueprintCameraPose& InCameraPose);
 
-	/**
-	 * Gets the initial evaluation context camera variable table.
-	 * WARNING: setting variables here will affect ALL running camera rigs!
-	 */
-	UFUNCTION(BlueprintPure, Category="Evaluation")
+	UFUNCTION(BlueprintPure, Category="Camera Director|Evaluation", 
+			meta=(DeprecatedFunction, DeprecationMessage="Please use GetSharedEvaluationData"))
 	FBlueprintCameraVariableTable GetInitialContextVariableTable() const;
 
 public:
diff --git a/Engine/Plugins/Cameras/GameplayCameras/Source/GameplayCameras/Public/GameFramework/BlueprintCameraContextDataTable.h b/Engine/Plugins/Cameras/GameplayCameras/Source/GameplayCameras/Public/GameFramework/BlueprintCameraContextDataTable.h
--- a/Engine/Plugins/Cameras/GameplayCameras/Source/GameplayCameras/Public/GameFramework/BlueprintCameraContextDataTable.h
+++ b/Engine/Plugins/Cameras/GameplayCameras/Source/GameplayCameras/Public/GameFramework/BlueprintCameraContextDataTable.h
@@ -17,6 +17,7 @@ class FCameraContextDataTable;
 
 }  // namespace UE::Cameras
 
+/** Provides access to a camera context data table. */
 USTRUCT(BlueprintType, DisplayName="Camera Context Data Table")
 struct GAMEPLAYCAMERAS_API FBlueprintCameraContextDataTable
 {
@@ -29,14 +30,20 @@ public:
 	FBlueprintCameraContextDataTable();
 	FBlueprintCameraContextDataTable(FCameraContextDataTable* InContextDataTable);
 
+	/** Gets the underlying context data table. */
 	FCameraContextDataTable* GetContextDataTable() const { return PrivateContextDataTable; }
 
+	/** Returns whether this context data table is valid. */
+	bool IsValid() const { return PrivateContextDataTable != nullptr; }
 
 private:
 
 	FCameraContextDataTable* PrivateContextDataTable = nullptr;
 };
 
+/**
+ * Utility Blueprint functions for camera context data tables.
+ */
 UCLASS()
 class UBlueprintCameraContextDataTableFunctionLibrary : public UBlueprintFunctionLibrary
 {
@@ -44,10 +51,54 @@ class UBlueprintCameraContextDataTableFunctionLibrary : public UBlueprintFunctio
 
 public:
 
+	/** Gets a value from the given camera context data table. */
+	UFUNCTION(BlueprintCallable, Category=Camera)
+	static FName GetNameData(const FBlueprintCameraContextDataTable& ContextDataTable, FCameraContextDataID DataID);
+
+	/** Gets a value from the given camera context data table. */
+	UFUNCTION(BlueprintCallable, Category=Camera)
+	static FString GetStringData(const FBlueprintCameraContextDataTable& ContextDataTable, FCameraContextDataID DataID);
+
+	/** Gets a value from the given camera context data table. */
+	UFUNCTION(BlueprintCallable, Category=Camera, meta=(DeterminesOutputType="EnumType"))
+	static uint8 GetEnumData(const FBlueprintCameraContextDataTable& ContextDataTable, FCameraContextDataID DataID, const UEnum* EnumType);
+
+	/** Gets a value from the given camera context data table. */
 	UFUNCTION(BlueprintCallable, Category=Camera)
 	static FInstancedStruct GetStructData(const FBlueprintCameraContextDataTable& ContextDataTable, FCameraContextDataID DataID, const UScriptStruct* DataStructType);
 
+	/** Gets a value from the given camera context data table. */
+	UFUNCTION(BlueprintCallable, Category=Camera)
+	static UObject* GetObjectData(const FBlueprintCameraContextDataTable& ContextDataTable, FCameraContextDataID DataID);
+
+	/** Gets a value from the given camera context data table. */
+	UFUNCTION(BlueprintCallable, Category=Camera)
+	static UClass* GetClassData(const FBlueprintCameraContextDataTable& ContextDataTable, FCameraContextDataID DataID);
+
+public:
+
+	/** Sets a value in the given camera context data table. */
+	UFUNCTION(BlueprintCallable, Category=Camera)
+	static bool SetNameData(const FBlueprintCameraContextDataTable& ContextDataTable, FCameraContextDataID DataID, const FName& Data);
+
+	/** Sets a value in the given camera context data table. */
+	UFUNCTION(BlueprintCallable, Category=Camera)
+	static bool SetStringData(const FBlueprintCameraContextDataTable& ContextDataTable, FCameraContextDataID DataID, const FString& Data);
+
+	/** Sets a value in the given camera context data table. */
+	UFUNCTION(BlueprintCallable, Category=Camera)
+	static bool SetEnumData(const FBlueprintCameraContextDataTable& ContextDataTable, FCameraContextDataID DataID, const UEnum* EnumType, uint8 Data);
+
+	/** Sets a value in the given camera context data table. */
+	UFUNCTION(BlueprintCallable, Category=Camera)
+	static bool SetStructData(const FBlueprintCameraContextDataTable& ContextDataTable, FCameraContextDataID DataID, const FInstancedStruct& Data);
+
+	/** Sets a value in the given camera context data table. */
+	UFUNCTION(BlueprintCallable, Category=Camera)
+	static bool SetObjectData(const FBlueprintCameraContextDataTable& ContextDataTable, FCameraContextDataID DataID, UObject* Data);
+
+	/** Sets a value in the given camera context data table. */
 	UFUNCTION(BlueprintCallable, Category=Camera)
-	static void SetStructData(const FBlueprintCameraContextDataTable& ContextDataTable, FCameraContextDataID DataID, const FInstancedStruct& Data);
+	static bool SetClassData(const FBlueprintCameraContextDataTable& ContextDataTable, FCameraContextDataID DataID, UClass* Data);
 };
 
diff --git a/Engine/Plugins/Cameras/GameplayCameras/Source/GameplayCameras/Public/GameFramework/BlueprintCameraNodeEvaluationResult.h b/Engine/Plugins/Cameras/GameplayCameras/Source/GameplayCameras/Public/GameFramework/BlueprintCameraNodeEvaluationResult.h
new file mode 100644
--- /dev/null
+++ b/Engine/Plugins/Cameras/GameplayCameras/Source/GameplayCameras/Public/GameFramework/BlueprintCameraNodeEvaluationResult.h
@@ -0,0 +1,81 @@
+// Copyright Epic Games, Inc. All Rights Reserved.
+
+#pragma once
+
+#include "Core/CameraNodeEvaluator.h"
+#include "GameFramework/BlueprintCameraContextDataTable.h"
+#include "GameFramework/BlueprintCameraPose.h"
+#include "GameFramework/BlueprintCameraVariableTable.h"
+#include "Kismet/BlueprintFunctionLibrary.h"
+
+#include "BlueprintCameraNodeEvaluationResult.generated.h"
+
+namespace UE::Cameras
+{
+
+struct FCameraNodeEvaluationResult;
+
+}  // namespace UE::Cameras
+
+/**
+ * Blueprint wrapper for camera evaluation data.
+ */
+USTRUCT(BlueprintType, DisplayName="Camera Evalution Data")
+struct GAMEPLAYCAMERAS_API FBlueprintCameraNodeEvaluationResult
+{
+	GENERATED_BODY()
+
+public:
+
+	using FCameraNodeEvaluationResult = UE::Cameras::FCameraNodeEvaluationResult;
+
+	FBlueprintCameraNodeEvaluationResult();
+	FBlueprintCameraNodeEvaluationResult(FCameraNodeEvaluationResult* InResult);
+
+public:
+
+	/** Gets the camera pose wrapper. */
+	FBlueprintCameraPose GetCameraPose() const;
+	/** Gets the variable table wrapper. */
+	FBlueprintCameraVariableTable GetVariableTable() const;
+	/** Gets the context data table wrapper. */
+	FBlueprintCameraContextDataTable GetContextDataTable() const;
+
+	/** Sets the camera pose on the underlying evaluation result. */
+	void SetCameraPose(const FBlueprintCameraPose& CameraPose);
+
+private:
+
+	/** The underlying camera evaluation result. */
+	FCameraNodeEvaluationResult* Result = nullptr;
+};
+
+/**
+ * Blueprint function library for a camera node evaluation data wrapper.
+ */
+UCLASS()
+class UBlueprintCameraNodeEvaluationResultFunctionLibrary : public UBlueprintFunctionLibrary
+{
+	GENERATED_BODY()
+
+public:
+
+	/** Gets the camera pose. */
+	UFUNCTION(BlueprintPure, Category="Camera")
+	static FBlueprintCameraPose GetCameraPose(const FBlueprintCameraNodeEvaluationResult& Data);
+
+	/** Gets the camera variable table. */
+	UFUNCTION(BlueprintPure, Category="Camera")
+	static FBlueprintCameraVariableTable GetVariableTable(const FBlueprintCameraNodeEvaluationResult& Data);
+
+	/** Gets the camera context data table. */
+	UFUNCTION(BlueprintPure, Category="Camera")
+	static FBlueprintCameraContextDataTable GetContextDataTable(const FBlueprintCameraNodeEvaluationResult& Data);
+
+public:
+
+	/** Sets the camera pose. */
+	UFUNCTION(BlueprintCallable, Category="Camera")
+	static void SetCameraPose(FBlueprintCameraNodeEvaluationResult& Data, const FBlueprintCameraPose& CameraPose);
+};
+
diff --git a/Engine/Plugins/Cameras/GameplayCameras/Source/GameplayCameras/Public/GameFramework/BlueprintCameraVariableTable.h b/Engine/Plugins/Cameras/GameplayCameras/Source/GameplayCameras/Public/GameFramework/BlueprintCameraVariableTable.h
--- a/Engine/Plugins/Cameras/GameplayCameras/Source/GameplayCameras/Public/GameFramework/BlueprintCameraVariableTable.h
+++ b/Engine/Plugins/Cameras/GameplayCameras/Source/GameplayCameras/Public/GameFramework/BlueprintCameraVariableTable.h
@@ -48,8 +48,6 @@ public:
 private:
 
 	FCameraVariableTable* PrivateVariableTable = nullptr;
-
-	friend class UBlueprintCameraVariableTableFunctionLibrary;
 };
 
 /**
diff --git a/Engine/Plugins/Cameras/GameplayCameras/Source/GameplayCameras/Public/GameFramework/CameraRigParameterInterop.h b/Engine/Plugins/Cameras/GameplayCameras/Source/GameplayCameras/Public/GameFramework/CameraRigParameterInterop.h
--- a/Engine/Plugins/Cameras/GameplayCameras/Source/GameplayCameras/Public/GameFramework/CameraRigParameterInterop.h
+++ b/Engine/Plugins/Cameras/GameplayCameras/Source/GameplayCameras/Public/GameFramework/CameraRigParameterInterop.h
@@ -12,7 +12,10 @@
 
 class UCameraRigAsset;
 class UCameraVariableAsset;
+struct FBlueprintCameraContextDataTable;
+struct FBlueprintCameraNodeEvaluationResult;
 struct FBlueprintCameraVariableTable;
+struct FInstancedStruct;
 
 /**
  * Blueprint internal methods to set values on a camera rig's exposed parameters.
@@ -32,34 +35,50 @@ public:
 public:
 
 	UFUNCTION(BlueprintCallable, Category="Camera", meta=(BlueprintInternalUseOnly="true"))
-	static void SetBooleanParameter(UPARAM(Ref) FBlueprintCameraVariableTable& VariableTable, UCameraRigAsset* CameraRig, const FString& ParameterName, bool ParameterValue);
+	static void SetBooleanParameter(UPARAM(Ref) FBlueprintCameraNodeEvaluationResult& Result, UCameraRigAsset* CameraRig, const FString& ParameterName, bool ParameterValue);
 
 	UFUNCTION(BlueprintCallable, Category="Camera", meta=(BlueprintInternalUseOnly="true"))
-	static void SetIntegerParameter(UPARAM(Ref) FBlueprintCameraVariableTable& VariableTable, UCameraRigAsset* CameraRig, const FString& ParameterName, int32 ParameterValue);
+	static void SetIntegerParameter(UPARAM(Ref) FBlueprintCameraNodeEvaluationResult& Result, UCameraRigAsset* CameraRig, const FString& ParameterName, int32 ParameterValue);
 
 	UFUNCTION(BlueprintCallable, Category="Camera", meta=(BlueprintInternalUseOnly="true"))
-	static void SetFloatParameter(UPARAM(Ref) FBlueprintCameraVariableTable& VariableTable, UCameraRigAsset* CameraRig, const FString& ParameterName, double ParameterValue);
+	static void SetFloatParameter(UPARAM(Ref) FBlueprintCameraNodeEvaluationResult& Result, UCameraRigAsset* CameraRig, const FString& ParameterName, double ParameterValue);
 
 	UFUNCTION(BlueprintCallable, Category="Camera", meta=(BlueprintInternalUseOnly="true"))
-	static void SetDoubleParameter(UPARAM(Ref) FBlueprintCameraVariableTable& VariableTable, UCameraRigAsset* CameraRig, const FString& ParameterName, double ParameterValue);
+	static void SetDoubleParameter(UPARAM(Ref) FBlueprintCameraNodeEvaluationResult& Result, UCameraRigAsset* CameraRig, const FString& ParameterName, double ParameterValue);
 
 	UFUNCTION(BlueprintCallable, Category="Camera", meta=(BlueprintInternalUseOnly="true"))
-	static void SetVector2Parameter(UPARAM(Ref) FBlueprintCameraVariableTable& VariableTable, UCameraRigAsset* CameraRig, const FString& ParameterName, FVector2D ParameterValue);
+	static void SetVector2Parameter(UPARAM(Ref) FBlueprintCameraNodeEvaluationResult& Result, UCameraRigAsset* CameraRig, const FString& ParameterName, FVector2D ParameterValue);
 
 	UFUNCTION(BlueprintCallable, Category="Camera", meta=(BlueprintInternalUseOnly="true"))
-	static void SetVector3Parameter(UPARAM(Ref) FBlueprintCameraVariableTable& VariableTable, UCameraRigAsset* CameraRig, const FString& ParameterName, FVector ParameterValue);
+	static void SetVector3Parameter(UPARAM(Ref) FBlueprintCameraNodeEvaluationResult& Result, UCameraRigAsset* CameraRig, const FString& ParameterName, FVector ParameterValue);
 
 	UFUNCTION(BlueprintCallable, Category="Camera", meta=(BlueprintInternalUseOnly="true"))
-	static void SetVector4Parameter(UPARAM(Ref) FBlueprintCameraVariableTable& VariableTable, UCameraRigAsset* CameraRig, const FString& ParameterName, FVector4 ParameterValue);
+	static void SetVector4Parameter(UPARAM(Ref) FBlueprintCameraNodeEvaluationResult& Result, UCameraRigAsset* CameraRig, const FString& ParameterName, FVector4 ParameterValue);
 
 	UFUNCTION(BlueprintCallable, Category="Camera", meta=(BlueprintInternalUseOnly="true"))
-	static void SetRotatorParameter(UPARAM(Ref) FBlueprintCameraVariableTable& VariableTable, UCameraRigAsset* CameraRig, const FString& ParameterName, FRotator ParameterValue);
+	static void SetRotatorParameter(UPARAM(Ref) FBlueprintCameraNodeEvaluationResult& Result, UCameraRigAsset* CameraRig, const FString& ParameterName, FRotator ParameterValue);
 
 	UFUNCTION(BlueprintCallable, Category="Camera", meta=(BlueprintInternalUseOnly="true"))
-	static void SetTransformParameter(UPARAM(Ref) FBlueprintCameraVariableTable& VariableTable, UCameraRigAsset* CameraRig, const FString& ParameterName, FTransform ParameterValue);
+	static void SetTransformParameter(UPARAM(Ref) FBlueprintCameraNodeEvaluationResult& Result, UCameraRigAsset* CameraRig, const FString& ParameterName, FTransform ParameterValue);
 
-private:
+public:
+
+	UFUNCTION(BlueprintCallable, Category="Camera", meta=(BlueprintInternalUseOnly="true"))
+	static void SetNameParameter(UPARAM(Ref) FBlueprintCameraNodeEvaluationResult& Result, UCameraRigAsset* CameraRig, const FString& ParameterName, FName ParameterValue);
+
+	UFUNCTION(BlueprintCallable, Category="Camera", meta=(BlueprintInternalUseOnly="true"))
+	static void SetStringParameter(UPARAM(Ref) FBlueprintCameraNodeEvaluationResult& Result, UCameraRigAsset* CameraRig, const FString& ParameterName, FString ParameterValue);
+
+	UFUNCTION(BlueprintCallable, Category="Camera", meta=(BlueprintInternalUseOnly="true"))
+	static void SetEnumParameter(UPARAM(Ref) FBlueprintCameraNodeEvaluationResult& Result, UCameraRigAsset* CameraRig, const FString& ParameterName, const UEnum* EnumType, uint8 ParameterValue);
+
+	UFUNCTION(BlueprintCallable, Category="Camera", meta=(BlueprintInternalUseOnly="true"))
+	static void SetStructParameter(UPARAM(Ref) FBlueprintCameraNodeEvaluationResult& Result, UCameraRigAsset* CameraRig, const FString& ParameterName, const FInstancedStruct& ParameterValue);
 
-	static UCameraVariableAsset* GetParameterPrivateVariable(UCameraRigAsset* CameraRig, const FString& ParameterName);
+	UFUNCTION(BlueprintCallable, Category="Camera", meta=(BlueprintInternalUseOnly="true"))
+	static void SetObjectParameter(UPARAM(Ref) FBlueprintCameraNodeEvaluationResult& Result, UCameraRigAsset* CameraRig, const FString& ParameterName, UObject* ParameterValue);
+
+	UFUNCTION(BlueprintCallable, Category="Camera", meta=(BlueprintInternalUseOnly="true"))
+	static void SetClassParameter(UPARAM(Ref) FBlueprintCameraNodeEvaluationResult& Result, UCameraRigAsset* CameraRig, const FString& ParameterName, UClass* ParameterValue);
 };
 
diff --git a/Engine/Plugins/Cameras/GameplayCameras/Source/GameplayCameras/Public/GameFramework/GameplayCameraComponent.h b/Engine/Plugins/Cameras/GameplayCameras/Source/GameplayCameras/Public/GameFramework/GameplayCameraComponent.h
--- a/Engine/Plugins/Cameras/GameplayCameras/Source/GameplayCameras/Public/GameFramework/GameplayCameraComponent.h
+++ b/Engine/Plugins/Cameras/GameplayCameras/Source/GameplayCameras/Public/GameFramework/GameplayCameraComponent.h
@@ -5,6 +5,7 @@
 #include "Components/SceneComponent.h"
 #include "Core/CameraAssetReference.h"
 #include "Core/CameraEvaluationContext.h"
+#include "GameFramework/BlueprintCameraNodeEvaluationResult.h"
 #include "GameFramework/BlueprintCameraPose.h"
 #include "GameFramework/BlueprintCameraVariableTable.h"
 #include "UObject/ObjectMacros.h"
@@ -56,16 +57,20 @@ public:
 	UFUNCTION(BlueprintCallable, Category=Camera)
 	GAMEPLAYCAMERAS_API void DeactivateCamera();
 
+	/** Gets the shared camera evaluation data for this component's evaluation context. */
+	UFUNCTION(BlueprintPure, Category=Camera, meta=(DisplayName="Get Shared Camera Data"))
+	GAMEPLAYCAMERAS_API FBlueprintCameraNodeEvaluationResult GetInitialResult() const;
+
 	/** Gets the initial camera pose for this component's camera evaluation context. */
-	UFUNCTION(BlueprintPure, Category=Camera)
+	UFUNCTION(BlueprintPure, Category=Camera, meta=(DeprecatedFunction, DeprecationMessage="Please use GetSharedCameraData"))
 	GAMEPLAYCAMERAS_API FBlueprintCameraPose GetInitialPose() const;
 
 	/** Sets the initial camera pose for this component's camera evaluation context. */
-	UFUNCTION(BlueprintCallable, Category=Camera)
+	UFUNCTION(BlueprintCallable, Category=Camera, meta=(DeprecatedFunction, DeprecationMessage="Please use GetSharedCameraData"))
 	GAMEPLAYCAMERAS_API bool SetInitialPose(const FBlueprintCameraPose& CameraPose);
 
 	/** Gets the initial camera variable table for this component's camera evaluation context. */
-	UFUNCTION(BlueprintPure, Category=Camera)
+	UFUNCTION(BlueprintPure, Category=Camera, meta=(DeprecatedFunction, DeprecationMessage="Please use GetSharedCameraData"))
 	GAMEPLAYCAMERAS_API FBlueprintCameraVariableTable GetInitialVariableTable() const;
 
 public:
diff --git a/Engine/Plugins/Cameras/GameplayCameras/Source/GameplayCamerasUncookedOnly/Private/BlueprintGraph/K2Node_SetCameraRigParameters.cpp b/Engine/Plugins/Cameras/GameplayCameras/Source/GameplayCamerasUncookedOnly/Private/BlueprintGraph/K2Node_SetCameraRigParameters.cpp
--- a/Engine/Plugins/Cameras/GameplayCameras/Source/GameplayCamerasUncookedOnly/Private/BlueprintGraph/K2Node_SetCameraRigParameters.cpp
+++ b/Engine/Plugins/Cameras/GameplayCameras/Source/GameplayCamerasUncookedOnly/Private/BlueprintGraph/K2Node_SetCameraRigParameters.cpp
@@ -8,16 +8,17 @@
 #include "Core/CameraVariableAssets.h"
 #include "EdGraphSchema_K2.h"
 #include "EditorCategoryUtils.h"
-#include "GameFramework/BlueprintCameraVariableTable.h"
+#include "GameFramework/BlueprintCameraNodeEvaluationResult.h"
 #include "GameFramework/CameraRigParameterInterop.h"
 #include "K2Node_CallFunction.h"
-#include "KismetCompiler.h"
+#include "Kismet/BlueprintInstancedStructLibrary.h"
 #include "Kismet2/BlueprintEditorUtils.h"
+#include "KismetCompiler.h"
 
 #define LOCTEXT_NAMESPACE "K2Node_SetCameraRigParameters"
 
 const FName UK2Node_SetCameraRigParameters::CameraRigPinName(TEXT("CameraRig"));
-const FName UK2Node_SetCameraRigParameters::CameraVariableTablePinName(TEXT("CameraVariableTable"));
+const FName UK2Node_SetCameraRigParameters::CameraNodeEvaluationResultPinName(TEXT("CameraEvaluationResult"));
 
 UK2Node_SetCameraRigParameters::UK2Node_SetCameraRigParameters(const FObjectInitializer& ObjectInit)
 	: Super(ObjectInit)
@@ -33,7 +34,7 @@ void UK2Node_SetCameraRigParameters::AllocateDefaultPins()
 	CreatePin(EGPD_Output, UEdGraphSchema_K2::PC_Exec, UEdGraphSchema_K2::PN_Then);
 
 	// Add evalation result pin.
-	CreatePin(EGPD_Input, UEdGraphSchema_K2::PC_Struct, FBlueprintCameraVariableTable::StaticStruct(), CameraVariableTablePinName);
+	CreatePin(EGPD_Input, UEdGraphSchema_K2::PC_Struct, FBlueprintCameraNodeEvaluationResult::StaticStruct(), CameraNodeEvaluationResultPinName);
 
 	// Add camera rig pin.
 	CreatePin(EGPD_Input, UEdGraphSchema_K2::PC_Object, UCameraRigAsset::StaticClass(), CameraRigPinName);
@@ -58,6 +59,14 @@ void UK2Node_SetCameraRigParameters::ReallocatePinsDuringReconstruction(TArray<U
 				PreloadObject(BlendableParameter->PrivateVariable);
 			}
 		}
+		for (UCameraRigDataParameter* DataParameter : CameraRig->Interface.DataParameters)
+		{
+			PreloadObject(DataParameter);
+			if (DataParameter)
+			{
+				PreloadObject(const_cast<UObject*>(DataParameter->DataTypeObject.Get()));
+			}
+		}
 
 		CreatePinsForCameraRig(CameraRig);
 	}
@@ -133,17 +142,20 @@ void UK2Node_SetCameraRigParameters::ExpandNode(class FKismetCompilerContext& Co
 		return;
 	}
 
-	// Get all the pins that correspond to parameters we want to override.
-	TArray<UEdGraphPin*> RigParameterPins;
-	GetCameraRigParameterPins(RigParameterPins);
-
 	UEdGraphPin* const CameraRigPin = FindPinChecked(CameraRigPinName);
-	UEdGraphPin* const CameraVariableTablePin = FindPinChecked(CameraVariableTablePinName);
+	UEdGraphPin* const CameraNodeEvaluationResultPin = FindPinChecked(CameraNodeEvaluationResultPinName);
 
 	UEdGraphPin* OriginalThenPin = GetThenPin();
 	UEdGraphPin* PreviousThenPin = nullptr;
-
-	for (UEdGraphPin* RigParameterPin : RigParameterPins)
+	
+	// For each blendable and data parameter, we figure out the type of SetXxxParameter function to call on the UCameraRigParameterInterop
+	// function library. We then make a K2Node_CallFunction node for it, and connect all its inputs, including connecting the parameter
+	// value to whatever our node's corresponding parameter value pin was connected to. As we go, we chain the exec/then pins, basically
+	// transforming our SetCameraRigParameters node into a chain of individual setter function calls.
+
+	TArray<UEdGraphPin*> BlendableParameterPins;
+	FindBlendableParameterPins(BlendableParameterPins);
+	for (UEdGraphPin* RigParameterPin : BlendableParameterPins)
 	{
 		UCameraRigBlendableParameter* BlendableParameter = CameraRig->Interface.FindBlendableParameterByName(RigParameterPin->GetName());
 		if (!BlendableParameter)
@@ -201,9 +213,9 @@ void UK2Node_SetCameraRigParameters::ExpandNode(class FKismetCompilerContext& Co
 		CallSetParameter->FunctionReference.SetExternalMember(CallSetParameterFuncName, UCameraRigParameterInterop::StaticClass());
 		CallSetParameter->AllocateDefaultPins();
 
-		// Connect the variable table pin that specifies where the parameter should be overriden.
-		UEdGraphPin* CallSetParameterVariableTablePin = CallSetParameter->FindPinChecked(TEXT("VariableTable"));
-		CompilerContext.CopyPinLinksToIntermediate(*CameraVariableTablePin, *CallSetParameterVariableTablePin);
+		// Connect the camera evaluation result argument.
+		UEdGraphPin* CallSetParameterResultPin = CallSetParameter->FindPinChecked(TEXT("Result"));
+		CompilerContext.CopyPinLinksToIntermediate(*CameraNodeEvaluationResultPin, *CallSetParameterResultPin);
 
 		// Connect the camera rig argument.
 		UEdGraphPin* CallSetParameterCameraRigPin = CallSetParameter->FindPinChecked(TEXT("CameraRig"));
@@ -224,8 +236,7 @@ void UK2Node_SetCameraRigParameters::ExpandNode(class FKismetCompilerContext& Co
 			CompilerContext.MovePinLinksToIntermediate(*RigParameterPin, *CallSetParameterValuePin);
 		}
 
-		// Connect the SetXxxParameter node to the chain of other similar nodes. The SetCameraRigParameters node
-		// effectively transforms into a chain of individual setter function calls.
+		// Chain the execution.
 		UEdGraphPin* CallSetParameterExecPin = CallSetParameter->GetExecPin();
 		if (PreviousThenPin)
 		{
@@ -240,6 +251,122 @@ void UK2Node_SetCameraRigParameters::ExpandNode(class FKismetCompilerContext& Co
 		PreviousThenPin = CallSetParameter->GetThenPin();
 	}
 
+	TArray<UEdGraphPin*> DataParameterPins;
+	FindDataParameterPins(DataParameterPins);
+	for (UEdGraphPin* RigParameterPin : DataParameterPins)
+	{
+		UCameraRigDataParameter* DataParameter = CameraRig->Interface.FindDataParameterByName(RigParameterPin->GetName());
+		if (!DataParameter)
+		{
+			CompilerContext.MessageLog.Error(*LOCTEXT("ErrorMissingParameter", "SetCameraRigParameters node @@ is trying to set parameter @@ but camera rig @@ has no such parameter.").ToString(), this, *RigParameterPin->GetName(), CameraRig);
+			continue;
+		}
+
+		if (!DataParameter->PrivateDataID.IsValid())
+		{
+			CompilerContext.MessageLog.Error(*LOCTEXT("ErrorMissingParameterVariable", "SetCameraRigParameters node @@ needs camera rig @@ to be built.").ToString(), this, CameraRig);
+			continue;
+		}
+
+		// Figure out the sort of SetXxxParameter function we want to call for this parameter.
+		FName CallSetParameterFuncName;
+		switch (DataParameter->DataType)
+		{
+			case ECameraContextDataType::Name:
+				CallSetParameterFuncName = GET_FUNCTION_NAME_CHECKED(UCameraRigParameterInterop, SetNameParameter);
+				break;
+			case ECameraContextDataType::String:
+				CallSetParameterFuncName = GET_FUNCTION_NAME_CHECKED(UCameraRigParameterInterop, SetStringParameter);
+				break;
+			case ECameraContextDataType::Enum:
+				CallSetParameterFuncName = GET_FUNCTION_NAME_CHECKED(UCameraRigParameterInterop, SetEnumParameter);
+				break;
+			case ECameraContextDataType::Struct:
+				CallSetParameterFuncName = GET_FUNCTION_NAME_CHECKED(UCameraRigParameterInterop, SetStructParameter);
+				break;
+			case ECameraContextDataType::Object:
+				CallSetParameterFuncName = GET_FUNCTION_NAME_CHECKED(UCameraRigParameterInterop, SetObjectParameter);
+				break;
+			case ECameraContextDataType::Class:
+				CallSetParameterFuncName = GET_FUNCTION_NAME_CHECKED(UCameraRigParameterInterop, SetClassParameter);
+				break;
+		}
+		if (CallSetParameterFuncName.IsNone())
+		{
+			CompilerContext.MessageLog.Error(*LOCTEXT("ErrorUnsupportedParameterType", "SetCameraRigParameters node @@ is trying to set parameter @@ but it has an unsupported type.").ToString(), this, *RigParameterPin->GetName());
+			continue;
+		}
+
+		// Make the SetXxxData function call node.
+		UK2Node_CallFunction* CallSetParameter = CompilerContext.SpawnIntermediateNode<UK2Node_CallFunction>(this, SourceGraph);
+		CallSetParameter->FunctionReference.SetExternalMember(CallSetParameterFuncName, UCameraRigParameterInterop::StaticClass());
+		CallSetParameter->AllocateDefaultPins();
+		UEdGraphPin* CurExecPin = CallSetParameter->GetExecPin();
+
+		// Connect the camera evaluation result argument.
+		UEdGraphPin* CallSetParameterResultPin = CallSetParameter->FindPinChecked(TEXT("Result"));
+		CompilerContext.CopyPinLinksToIntermediate(*CameraNodeEvaluationResultPin, *CallSetParameterResultPin);
+
+		// Connect the camera rig argument.
+		UEdGraphPin* CallSetParameterCameraRigPin = CallSetParameter->FindPinChecked(TEXT("CameraRig"));
+		CompilerContext.CopyPinLinksToIntermediate(*CameraRigPin, *CallSetParameterCameraRigPin);
+
+		// Set the parameter name argument.
+		UEdGraphPin* CallSetParameterNamePin = CallSetParameter->FindPinChecked(TEXT("ParameterName"));
+		CallSetParameterNamePin->DefaultValue = DataParameter->InterfaceParameterName;
+
+		// Set or connect the parameter value argument.
+		UEdGraphPin* CallSetParameterValuePin = CallSetParameter->FindPinChecked(TEXT("ParameterValue"));
+		if (DataParameter->DataType == ECameraContextDataType::Struct)
+		{
+			// If we are setting a struct, first turn it into an FInstancedStruct. Insert this node in the
+			// execution chain, and use that node's Value pin as the pin to set the incoming struct value.
+			UK2Node_CallFunction* CallMakeInstancedStruct = CompilerContext.SpawnIntermediateNode<UK2Node_CallFunction>(this, SourceGraph);
+			CallMakeInstancedStruct->FunctionReference.SetExternalMember(GET_FUNCTION_NAME_CHECKED(UBlueprintInstancedStructLibrary, MakeInstancedStruct), UBlueprintInstancedStructLibrary::StaticClass());
+			CallMakeInstancedStruct->AllocateDefaultPins();
+
+			CallMakeInstancedStruct->GetReturnValuePin()->MakeLinkTo(CallSetParameterValuePin);
+			CallMakeInstancedStruct->GetThenPin()->MakeLinkTo(CurExecPin);
+			CurExecPin = CallMakeInstancedStruct->GetExecPin();
+
+			UEdGraphPin* CallMakeInstancedStructValuePin = CallMakeInstancedStruct->FindPinChecked(TEXT("Value"));
+			CompilerContext.MovePinLinksToIntermediate(*RigParameterPin, *CallMakeInstancedStructValuePin);
+			CallMakeInstancedStruct->PinConnectionListChanged(CallMakeInstancedStructValuePin);
+		}
+		else
+		{
+			CallSetParameterValuePin->DefaultValue = RigParameterPin->DefaultValue;
+			CallSetParameterValuePin->DefaultTextValue = RigParameterPin->DefaultTextValue;
+			CallSetParameterValuePin->AutogeneratedDefaultValue = RigParameterPin->AutogeneratedDefaultValue;
+			CallSetParameterValuePin->DefaultObject = RigParameterPin->DefaultObject;
+			if (RigParameterPin->LinkedTo.Num() > 0)
+			{
+				CompilerContext.MovePinLinksToIntermediate(*RigParameterPin, *CallSetParameterValuePin);
+			}
+		}
+
+		// Set extra type pin for enums.
+		if (DataParameter->DataType == ECameraContextDataType::Enum)
+		{
+			const UEnum* EnumType = CastChecked<const UEnum>(DataParameter->DataTypeObject);
+			UEdGraphPin* CallSetParameterEnumTypePin = CallSetParameter->FindPinChecked(TEXT("EnumType"));
+			CallSetParameterEnumTypePin->DefaultObject = const_cast<UEnum*>(EnumType);
+		}
+
+		// Chain the execution.
+		if (PreviousThenPin)
+		{
+			PreviousThenPin->MakeLinkTo(CurExecPin);
+		}
+		else
+		{
+			UEdGraphPin* ThisExecPin = GetExecPin();
+			CompilerContext.MovePinLinksToIntermediate(*ThisExecPin, *CurExecPin);
+		}
+
+		PreviousThenPin = CallSetParameter->GetThenPin();
+	}
+
 	// Connect the last node if necessary.
 	if (OriginalThenPin && PreviousThenPin && OriginalThenPin->LinkedTo.Num() > 0)
 	{
@@ -270,12 +397,12 @@ UEdGraphPin* UK2Node_SetCameraRigParameters::GetCameraRigPin(TArrayView<UEdGraph
 	return CameraRigPin;
 }
 
-UEdGraphPin* UK2Node_SetCameraRigParameters::GetCameraEvaluationResultPin() const
+UEdGraphPin* UK2Node_SetCameraRigParameters::GetCameraNodeEvaluationResultPin() const
 {
 	UEdGraphPin* ResultPin = nullptr;
 	for (UEdGraphPin* Pin : Pins)
 	{
-		if (Pin && Pin->PinName == CameraRigPinName)
+		if (Pin && Pin->PinName == CameraNodeEvaluationResultPinName)
 		{
 			ResultPin = Pin;
 			break;
@@ -285,33 +412,15 @@ UEdGraphPin* UK2Node_SetCameraRigParameters::GetCameraEvaluationResultPin() cons
 	return ResultPin;
 }
 
-void UK2Node_SetCameraRigParameters::GetCameraRigParameterPins(TArray<UEdGraphPin*>& OutParameterPins) const
-{
-	for (UEdGraphPin* Pin : Pins)
-	{
-		if (IsCameraRigParameterPin(Pin))
-		{
-			OutParameterPins.Add(Pin);
-		}
-	}
-}
-
-bool UK2Node_SetCameraRigParameters::IsCameraRigParameterPin(UEdGraphPin* Pin) const
-{
-	return Pin->PinName != UEdGraphSchema_K2::PN_Execute &&
-		Pin->PinName != UEdGraphSchema_K2::PN_Then &&
-		Pin->PinName != UEdGraphSchema_K2::PN_ReturnValue &&
-		Pin->PinName != CameraRigPinName &&
-		Pin->PinName != CameraVariableTablePinName;
-}
-
 void UK2Node_SetCameraRigParameters::CreatePinsForCameraRig(UCameraRigAsset* CameraRig, TArray<UEdGraphPin*>* CreatedPins)
 {
 	check(CameraRig);
 
 	const UEdGraphSchema_K2* K2Schema = GetDefault<UEdGraphSchema_K2>();
 
-	for (UCameraRigBlendableParameter* BlendableParameter : CameraRig->Interface.BlendableParameters)
+	BlendableParameterPinNames.Reset();
+
+	for (const UCameraRigBlendableParameter* BlendableParameter : CameraRig->Interface.BlendableParameters)
 	{
 		if (!ensure(BlendableParameter))
 		{
@@ -375,11 +484,91 @@ void UK2Node_SetCameraRigParameters::CreatePinsForCameraRig(UCameraRigAsset* Cam
 				EGPD_Input, 
 				NewPinCategory, NewPinSubCategory, NewPinSubCategoryObject, 
 				FName(BlendableParameter->InterfaceParameterName));
+		BlendableParameterPinNames.Add(NewPin->PinName);
 		if (CreatedPins)
 		{
 			CreatedPins->Add(NewPin);
 		}
 	}
+
+	DataParameterPinNames.Reset();
+
+	for (const UCameraRigDataParameter* DataParameter : CameraRig->Interface.DataParameters)
+	{
+		if (!ensure(DataParameter))
+		{
+			continue;
+		}
+		
+		if (!DataParameter->PrivateDataID.IsValid())
+		{
+			// Camera rig isn't fully built.
+			continue;
+		}
+
+		FName NewPinCategory;
+		UObject* NewPinSubCategoryObject = const_cast<UObject*>(DataParameter->DataTypeObject.Get());
+		switch (DataParameter->DataType)
+		{
+			case ECameraContextDataType::Name:
+				NewPinCategory = UEdGraphSchema_K2::PC_Name;
+				break;
+			case ECameraContextDataType::String:
+				NewPinCategory = UEdGraphSchema_K2::PC_String;
+				break;
+			case ECameraContextDataType::Enum:
+				NewPinCategory = UEdGraphSchema_K2::PC_Enum;
+				break;
+			case ECameraContextDataType::Struct:
+				NewPinCategory = UEdGraphSchema_K2::PC_Struct;
+				break;
+			case ECameraContextDataType::Object:
+				NewPinCategory = UEdGraphSchema_K2::PC_Object;
+				break;
+			case ECameraContextDataType::Class:
+				NewPinCategory = UEdGraphSchema_K2::PC_Class;
+				break;
+		}
+		if (NewPinCategory.IsNone())
+		{
+			// Unsupported type for Blueprints.
+			continue;
+		}
+
+		UEdGraphPin* NewPin = CreatePin(
+				EGPD_Input, 
+				NewPinCategory, NAME_None, NewPinSubCategoryObject, 
+				FName(DataParameter->InterfaceParameterName));
+		DataParameterPinNames.Add(NewPin->PinName);
+		if (CreatedPins)
+		{
+			CreatedPins->Add(NewPin);
+		}
+	}
+}
+
+void UK2Node_SetCameraRigParameters::FindBlendableParameterPins(TArray<UEdGraphPin*>& OutPins) const
+{
+	for (const FName& PinName : BlendableParameterPinNames)
+	{
+		UEdGraphPin* Pin = FindPin(PinName);
+		if (ensure(Pin))
+		{
+			OutPins.Add(Pin);
+		}
+	}
+}
+
+void UK2Node_SetCameraRigParameters::FindDataParameterPins(TArray<UEdGraphPin*>& OutPins) const
+{
+	for (const FName& PinName : DataParameterPinNames)
+	{
+		UEdGraphPin* Pin = FindPin(PinName);
+		if (ensure(Pin))
+		{
+			OutPins.Add(Pin);
+		}
+	}
 }
 
 UCameraRigAsset* UK2Node_SetCameraRigParameters::GetCameraRig(TArrayView<UEdGraphPin* const>* InPinsToSearch) const
@@ -407,12 +596,16 @@ UCameraRigAsset* UK2Node_SetCameraRigParameters::GetCameraRig(TArrayView<UEdGrap
 
 void UK2Node_SetCameraRigParameters::OnCameraRigChanged()
 {
+	TArray<FName> OldCameraRigPinNames;
+	OldCameraRigPinNames.Append(BlendableParameterPinNames);
+	OldCameraRigPinNames.Append(DataParameterPinNames);
+
 	TArray<UEdGraphPin*> OldPins = Pins;
 	TArray<UEdGraphPin*> OldCameraRigPins;
 
 	for (UEdGraphPin* OldPin : OldPins)
 	{
-		if (IsCameraRigParameterPin(OldPin))
+		if (OldCameraRigPinNames.Contains(OldPin->PinName))
 		{
 			Pins.Remove(OldPin);
 			OldCameraRigPins.Add(OldPin);
diff --git a/Engine/Plugins/Cameras/GameplayCameras/Source/GameplayCamerasUncookedOnly/Private/BlueprintGraph/K2Node_SetCameraRigParameters.h b/Engine/Plugins/Cameras/GameplayCameras/Source/GameplayCamerasUncookedOnly/Private/BlueprintGraph/K2Node_SetCameraRigParameters.h
--- a/Engine/Plugins/Cameras/GameplayCameras/Source/GameplayCamerasUncookedOnly/Private/BlueprintGraph/K2Node_SetCameraRigParameters.h
+++ b/Engine/Plugins/Cameras/GameplayCameras/Source/GameplayCamerasUncookedOnly/Private/BlueprintGraph/K2Node_SetCameraRigParameters.h
@@ -46,15 +46,23 @@ public:
 protected:
 
 	static const FName CameraRigPinName;
-	static const FName CameraVariableTablePinName;
+	static const FName CameraNodeEvaluationResultPinName;
 
 	UEdGraphPin* GetCameraRigPin(TArrayView<UEdGraphPin* const>* InPinsToSearch = nullptr) const;
-	UEdGraphPin* GetCameraEvaluationResultPin() const;
-	void GetCameraRigParameterPins(TArray<UEdGraphPin*>& OutParameterPins) const;
-	bool IsCameraRigParameterPin(UEdGraphPin* Pin) const;
+	UEdGraphPin* GetCameraNodeEvaluationResultPin() const;
 	void CreatePinsForCameraRig(UCameraRigAsset* CameraRig, TArray<UEdGraphPin*>* CreatedPins = nullptr);
+	void FindBlendableParameterPins(TArray<UEdGraphPin*>& OutPins) const;
+	void FindDataParameterPins(TArray<UEdGraphPin*>& OutPins) const;
 
 	UCameraRigAsset* GetCameraRig(TArrayView<UEdGraphPin* const>* InPinsToSearch = nullptr) const;
 	void OnCameraRigChanged();
+
+private:
+
+	UPROPERTY()
+	TArray<FName> BlendableParameterPinNames;
+
+	UPROPERTY()
+	TArray<FName> DataParameterPinNames;
 };
 
