Cameras: add support for conditional context results

Conditional results on an evaluation context are data that is "overlaid" on top on some camera rig instances that match the condition. In that way they work in the same way as the initial result, only the initial result is applied to all running camera rig instances for that evaluation context.

Right now we only have one condition, which is "the active camera rig", i.e. the camera rig instance that's on the top of the main blend stack. This allows the user to set data only for a new camera rig instance, even if other instances of that camera rig are still running.

Since we need to be able to overlay multiple context results, change how we write them onto running camera rig instances: we now only write values that have been written this frame. This requires force-setting all values on initialization, to make sure new camera rig instances are up to speed.

diff --git a/Engine/Plugins/Cameras/GameplayCameras/Source/GameplayCameras/Private/Core/BlendStackCameraNode.cpp b/Engine/Plugins/Cameras/GameplayCameras/Source/GameplayCameras/Private/Core/BlendStackCameraNode.cpp
--- a/Engine/Plugins/Cameras/GameplayCameras/Source/GameplayCameras/Private/Core/BlendStackCameraNode.cpp
+++ b/Engine/Plugins/Cameras/GameplayCameras/Source/GameplayCameras/Private/Core/BlendStackCameraNode.cpp
@@ -86,12 +86,17 @@ bool FBlendStackCameraNodeEvaluator::InitializeEntry(
 		return false;
 	}
 
-	// Allocate variables in the variable table.
+	// Allocate variable table and context data table.
 	NewEntry.Result.VariableTable.Initialize(CameraRig->AllocationInfo.VariableTableInfo);
-
-	// Allocate context data in the data table.
 	NewEntry.Result.ContextDataTable.Initialize(CameraRig->AllocationInfo.ContextDataTableInfo);
 
+	// Set all the data from the context.
+	const FCameraNodeEvaluationResult& ContextResult = EvaluationContext->GetInitialResult();
+	NewEntry.Result.VariableTable.Override(
+			ContextResult.VariableTable,
+			ECameraVariableTableFilter::InputOutput | ECameraVariableTableFilter::Private);
+	NewEntry.Result.ContextDataTable.OverrideAll(ContextResult.ContextDataTable);
+
 	// Initialize the node evaluators.
 	FCameraNodeEvaluatorInitializeParams InitParams(&NewEntry.EvaluatorHierarchy);
 	InitParams.Evaluator = OwningEvaluator;
@@ -170,7 +175,12 @@ void FBlendStackCameraNodeEvaluator::ResolveEntries(TArray<FResolvedEntry>& OutR
 		FCameraRigEntry& Entry(Entries[Index]);
 		TSharedPtr<const FCameraEvaluationContext> CurContext = Entry.EvaluationContext.Pin();
 
-		OutResolvedEntries.Add({ Entry, CurContext, Index });
+		FResolvedEntry& ResolvedEntry = OutResolvedEntries.Emplace_GetRef(Entry, CurContext);
+		ResolvedEntry.EntryIndex = Index;
+		if (Index == Entries.Num() - 1)
+		{
+			ResolvedEntry.bIsActiveEntry = true;
+		}
 
 		// While we make these resolved entries, emit warnings and errors as needed.
 		if (!Entry.bIsFrozen)
@@ -637,6 +647,9 @@ void FTransientBlendStackCameraNodeEvaluator::OnRun(const FCameraNodeEvaluationP
 
 void FTransientBlendStackCameraNodeEvaluator::InternalPreBlendPrepare(TArrayView<FResolvedEntry> ResolvedEntries, const FCameraNodeEvaluationParams& Params, FCameraNodeEvaluationResult& OutResult)
 {
+	constexpr ECameraVariableTableFilter VariableTableFilter = ECameraVariableTableFilter::InputOutput | ECameraVariableTableFilter::Private | ECameraVariableTableFilter::ChangedOnly;
+	constexpr ECameraContextDataTableFilter ContextDataTableFilter = ECameraContextDataTableFilter::ChangedOnly;
+
 	for (FResolvedEntry& ResolvedEntry : ResolvedEntries)
 	{
 		FCameraRigEntry& Entry(ResolvedEntry.Entry);
@@ -657,10 +670,18 @@ void FTransientBlendStackCameraNodeEvaluator::InternalPreBlendPrepare(TArrayView
 
 		// Override it with whatever the evaluation context has set on its result.
 		// Evaluation contexts may have private variables we need to pass along, such as when rig parameter
-		// overrides have been set on them.
+		// overrides have been set on them, so include private variables in the filter.
 		const FCameraNodeEvaluationResult& ContextResult(ResolvedEntry.Context->GetInitialResult());
-		CurResult.VariableTable.Override(ContextResult.VariableTable, ECameraVariableTableFilter::AllPublic | ECameraVariableTableFilter::Private);
-		CurResult.ContextDataTable.OverrideAll(ContextResult.ContextDataTable);
+		CurResult.VariableTable.Override(ContextResult.VariableTable, VariableTableFilter);
+		CurResult.ContextDataTable.Override(ContextResult.ContextDataTable, ContextDataTableFilter);
+		if (ResolvedEntry.bIsActiveEntry)
+		{
+			if (const FCameraNodeEvaluationResult* ActiveOnlyResult = ResolvedEntry.Context->GetConditionalResult(ECameraEvaluationDataCondition::ActiveCameraRig))
+			{
+				CurResult.VariableTable.Override(ActiveOnlyResult->VariableTable, VariableTableFilter);
+				CurResult.ContextDataTable.Override(ActiveOnlyResult->ContextDataTable, ContextDataTableFilter);
+			}
+		}
 
 		// Gather input parameters if needed (and remember if it was indeed needed).
 		if (!Entry.bInputRunThisFrame)
@@ -756,6 +777,14 @@ void FTransientBlendStackCameraNodeEvaluator::InternalUpdate(TArrayView<FResolve
 		const FCameraNodeEvaluationResult& ContextResult(ResolvedEntry.Context->GetInitialResult());
 		CurResult.CameraPose.OverrideChanged(ContextResult.CameraPose);
 		CurResult.bIsCameraCut = OutResult.bIsCameraCut || ContextResult.bIsCameraCut || Entry.bForceCameraCut;
+		if (ResolvedEntry.bIsActiveEntry)
+		{
+			if (const FCameraNodeEvaluationResult* ActiveOnlyResult = ResolvedEntry.Context->GetConditionalResult(ECameraEvaluationDataCondition::ActiveCameraRig))
+			{
+				CurResult.CameraPose.OverrideChanged(ActiveOnlyResult->CameraPose);
+				CurResult.bIsCameraCut |= ActiveOnlyResult->bIsCameraCut;
+			}
+		}
 		CurResult.bIsValid = true;
 
 		// Run the camera rig's root node.
@@ -1025,6 +1054,9 @@ void FPersistentBlendStackCameraNodeEvaluator::OnRun(const FCameraNodeEvaluation
 
 void FPersistentBlendStackCameraNodeEvaluator::InternalUpdate(TArrayView<FResolvedEntry> ResolvedEntries, const FCameraNodeEvaluationParams& Params, FCameraNodeEvaluationResult& OutResult)
 {
+	constexpr ECameraVariableTableFilter VariableTableFilter = ECameraVariableTableFilter::InputOutput | ECameraVariableTableFilter::Private | ECameraVariableTableFilter::ChangedOnly;
+	constexpr ECameraContextDataTableFilter ContextDataTableFilter = ECameraContextDataTableFilter::ChangedOnly;
+
 	for (FResolvedEntry& ResolvedEntry : ResolvedEntries)
 	{
 		FCameraRigEntry& Entry(ResolvedEntry.Entry);
@@ -1046,11 +1078,11 @@ void FPersistentBlendStackCameraNodeEvaluator::InternalUpdate(TArrayView<FResolv
 
 				// Override it with whatever the evaluation context has set on its result.
 				// Evaluation contexts may have private variables we need to pass along, such as when rig parameter
-				// overrides have been set on them.
+				// overrides have been set on them, so include private variables in the filter.
 				const FCameraNodeEvaluationResult& ContextResult(ResolvedEntry.Context->GetInitialResult());
 				CurResult.CameraPose.OverrideChanged(ContextResult.CameraPose);
-				CurResult.VariableTable.Override(ContextResult.VariableTable, ECameraVariableTableFilter::AllPublic | ECameraVariableTableFilter::Private);
-				CurResult.ContextDataTable.OverrideAll(ContextResult.ContextDataTable);
+				CurResult.VariableTable.Override(ContextResult.VariableTable, VariableTableFilter);
+				CurResult.ContextDataTable.Override(ContextResult.ContextDataTable, ContextDataTableFilter);
 
 				// Setup flags.
 				CurResult.bIsCameraCut = OutResult.bIsCameraCut || ContextResult.bIsCameraCut || Entry.bForceCameraCut;
diff --git a/Engine/Plugins/Cameras/GameplayCameras/Source/GameplayCameras/Private/Core/CameraContextDataTable.cpp b/Engine/Plugins/Cameras/GameplayCameras/Source/GameplayCameras/Private/Core/CameraContextDataTable.cpp
--- a/Engine/Plugins/Cameras/GameplayCameras/Source/GameplayCameras/Private/Core/CameraContextDataTable.cpp
+++ b/Engine/Plugins/Cameras/GameplayCameras/Source/GameplayCameras/Private/Core/CameraContextDataTable.cpp
@@ -540,20 +540,29 @@ void FCameraContextDataTable::ClearAllWrittenThisFrameFlags()
 
 void FCameraContextDataTable::OverrideAll(const FCameraContextDataTable& OtherTable)
 {
-	Override(OtherTable, false);
+	InternalOverride(OtherTable, ECameraContextDataTableFilter::None);
 }
 
 void FCameraContextDataTable::OverrideKnown(const FCameraContextDataTable& OtherTable)
 {
-	Override(OtherTable, true);
+	InternalOverride(OtherTable, ECameraContextDataTableFilter::KnownOnly);
 }
 
-void FCameraContextDataTable::Override(const FCameraContextDataTable& OtherTable, bool bKnownOnly)
+void FCameraContextDataTable::Override(const FCameraContextDataTable& OtherTable, ECameraContextDataTableFilter Filter)
 {
+	InternalOverride(OtherTable, Filter);
+}
+
+void FCameraContextDataTable::InternalOverride(const FCameraContextDataTable& OtherTable, ECameraContextDataTableFilter Filter)
+{
+	const bool bKnownOnly = EnumHasAllFlags(Filter, ECameraContextDataTableFilter::KnownOnly);
+	const bool bChangedOnly = EnumHasAllFlags(Filter, ECameraContextDataTableFilter::ChangedOnly);
+
 	for (const FEntry& OtherEntry : OtherTable.Entries)
 	{
 		const EEntryFlags OtherFlags = OtherEntry.Flags;
-		if (EnumHasAnyFlags(OtherFlags, EEntryFlags::Written))
+		if (EnumHasAnyFlags(OtherFlags, EEntryFlags::Written)
+				&& (!bChangedOnly || EnumHasAnyFlags(OtherFlags, EEntryFlags::WrittenThisFrame)))
 		{
 			int32 ThisIndex = EntryLookup.FindRef(OtherEntry.ID, INDEX_NONE);
 			if (ThisIndex == INDEX_NONE)
diff --git a/Engine/Plugins/Cameras/GameplayCameras/Source/GameplayCameras/Private/Core/CameraEvaluationContext.cpp b/Engine/Plugins/Cameras/GameplayCameras/Source/GameplayCameras/Private/Core/CameraEvaluationContext.cpp
--- a/Engine/Plugins/Cameras/GameplayCameras/Source/GameplayCameras/Private/Core/CameraEvaluationContext.cpp
+++ b/Engine/Plugins/Cameras/GameplayCameras/Source/GameplayCameras/Private/Core/CameraEvaluationContext.cpp
@@ -31,6 +31,14 @@ void FCameraEvaluationContext::Initialize(const FCameraEvaluationContextInitiali
 	CameraAsset = Params.CameraAsset;
 	WeakPlayerController = Params.PlayerController;
 
+	if (CameraAsset)
+	{
+		const FCameraAssetAllocationInfo& AllocationInfo = CameraAsset->GetAllocationInfo();
+
+		InitialResult.VariableTable.Initialize(AllocationInfo.VariableTableInfo);
+		InitialResult.ContextDataTable.Initialize(AllocationInfo.ContextDataTableInfo);
+	}
+
 	bInitialized = true;
 }
 
@@ -66,9 +74,19 @@ void FCameraEvaluationContext::AddReferencedObjects(FReferenceCollector& Collect
 
 void FCameraEvaluationContext::OnEndCameraSystemUpdate()
 {
+	InitialResult.CameraPose.ClearAllChangedFlags();
 	InitialResult.VariableTable.AutoResetValues();
 	InitialResult.VariableTable.ClearAllWrittenThisFrameFlags();
 	InitialResult.ContextDataTable.ClearAllWrittenThisFrameFlags();
+
+	for (FConditionalResults::ElementType& Pair : ConditionalResults)
+	{
+		FCameraNodeEvaluationResult& Result = Pair.Value;
+		Result.CameraPose.ClearAllChangedFlags();
+		Result.VariableTable.AutoResetValues();
+		Result.VariableTable.ClearAllWrittenThisFrameFlags();
+		Result.ContextDataTable.ClearAllWrittenThisFrameFlags();
+	}
 }
 
 void FCameraEvaluationContext::AutoCreateDirectorEvaluator()
@@ -96,6 +114,26 @@ void FCameraEvaluationContext::AutoCreateDirectorEvaluator()
 	}
 }
 
+FCameraNodeEvaluationResult& FCameraEvaluationContext::GetOrAddConditionalResult(ECameraEvaluationDataCondition Condition)
+{
+	if (FCameraNodeEvaluationResult* ExistingResult = ConditionalResults.Find(Condition))
+	{
+		return *ExistingResult;
+	}
+
+	FCameraNodeEvaluationResult& NewResult = ConditionalResults.Add(Condition);
+
+	if (CameraAsset)
+	{
+		const FCameraAssetAllocationInfo& AllocationInfo = CameraAsset->GetAllocationInfo();
+
+		NewResult.VariableTable.Initialize(AllocationInfo.VariableTableInfo);
+		NewResult.ContextDataTable.Initialize(AllocationInfo.ContextDataTableInfo);
+	}
+
+	return NewResult;
+}
+
 void FCameraEvaluationContext::Activate(const FCameraEvaluationContextActivateParams& Params)
 {
 	if (!ensureMsgf(bInitialized, TEXT("This evaluation context needs to be initialized!")))
diff --git a/Engine/Plugins/Cameras/GameplayCameras/Source/GameplayCameras/Private/Core/CameraIKAim.cpp b/Engine/Plugins/Cameras/GameplayCameras/Source/GameplayCameras/Private/Core/CameraIKAim.cpp
--- a/Engine/Plugins/Cameras/GameplayCameras/Source/GameplayCameras/Private/Core/CameraIKAim.cpp
+++ b/Engine/Plugins/Cameras/GameplayCameras/Source/GameplayCameras/Private/Core/CameraIKAim.cpp
@@ -164,7 +164,7 @@ void FCameraIKAim::RunRootCameraNode(const FCameraIKAimParams& Params, const FCa
 	// Make sure the camera rig will get its private variables, such as rig interface parameter overrides.
 	ScratchResult.VariableTable.Override(
 			CameraRigInfo.LastResult->VariableTable, 
-			ECameraVariableTableFilter::AllPublic | ECameraVariableTableFilter::Private);
+			ECameraVariableTableFilter::InputOutput | ECameraVariableTableFilter::Private);
 
 	FRootCameraNodeEvaluator* RootEvaluator = Params.Evaluator->GetRootNodeEvaluator();
 
diff --git a/Engine/Plugins/Cameras/GameplayCameras/Source/GameplayCameras/Private/Core/CameraPose.cpp b/Engine/Plugins/Cameras/GameplayCameras/Source/GameplayCameras/Private/Core/CameraPose.cpp
--- a/Engine/Plugins/Cameras/GameplayCameras/Source/GameplayCameras/Private/Core/CameraPose.cpp
+++ b/Engine/Plugins/Cameras/GameplayCameras/Source/GameplayCameras/Private/Core/CameraPose.cpp
@@ -112,10 +112,10 @@ FTransform3d FCameraPose::GetTransform() const
 	return Transform;
 }
 
-void FCameraPose::SetTransform(FTransform3d Transform)
+void FCameraPose::SetTransform(FTransform3d Transform, bool bForceSet)
 {
-	SetLocation(Transform.GetLocation());
-	SetRotation(Transform.GetRotation().Rotator());
+	SetLocation(Transform.GetLocation(), bForceSet);
+	SetRotation(Transform.GetRotation().Rotator(), bForceSet);
 }
 
 double FCameraPose::GetEffectiveFieldOfView() const
diff --git a/Engine/Plugins/Cameras/GameplayCameras/Source/GameplayCameras/Private/Core/CameraVariableTable.cpp b/Engine/Plugins/Cameras/GameplayCameras/Source/GameplayCameras/Private/Core/CameraVariableTable.cpp
--- a/Engine/Plugins/Cameras/GameplayCameras/Source/GameplayCameras/Private/Core/CameraVariableTable.cpp
+++ b/Engine/Plugins/Cameras/GameplayCameras/Source/GameplayCameras/Private/Core/CameraVariableTable.cpp
@@ -420,7 +420,7 @@ void FCameraVariableTable::Serialize(FArchive& Ar)
 
 void FCameraVariableTable::OverrideAll(const FCameraVariableTable& OtherTable)
 {
-	const ECameraVariableTableFilter Filter = ECameraVariableTableFilter::AllPublic;
+	const ECameraVariableTableFilter Filter = ECameraVariableTableFilter::InputOutput;
 	InternalOverride(OtherTable, Filter, nullptr, false, nullptr);
 }
 
@@ -513,7 +513,7 @@ void FCameraVariableTable::InternalOverride(const FCameraVariableTable& OtherTab
 
 void FCameraVariableTable::LerpAll(const FCameraVariableTable& ToTable, float Factor)
 {
-	const ECameraVariableTableFilter Filter = ECameraVariableTableFilter::AllPublic;
+	const ECameraVariableTableFilter Filter = ECameraVariableTableFilter::InputOutput;
 	InternalLerp(ToTable, Filter, Factor, nullptr, false, nullptr);
 }
 
diff --git a/Engine/Plugins/Cameras/GameplayCameras/Source/GameplayCameras/Private/Directors/BlueprintCameraDirector.cpp b/Engine/Plugins/Cameras/GameplayCameras/Source/GameplayCameras/Private/Directors/BlueprintCameraDirector.cpp
--- a/Engine/Plugins/Cameras/GameplayCameras/Source/GameplayCameras/Private/Directors/BlueprintCameraDirector.cpp
+++ b/Engine/Plugins/Cameras/GameplayCameras/Source/GameplayCameras/Private/Directors/BlueprintCameraDirector.cpp
@@ -337,6 +337,21 @@ FBlueprintCameraNodeEvaluationResult UBlueprintCameraDirectorEvaluator::GetIniti
 	}
 }
 
+FBlueprintCameraNodeEvaluationResult UBlueprintCameraDirectorEvaluator::GetConditionalContextResult(ECameraEvaluationDataCondition Condition) const
+{
+	if (EvaluationContext)
+	{
+		return FBlueprintCameraNodeEvaluationResult(&EvaluationContext->GetOrAddConditionalResult(Condition));
+	}
+	else
+	{
+		FFrame::KismetExecutionMessage(
+				TEXT("Can't access evaluation context's initial result outside of RunCameraDirector"), 
+				ELogVerbosity::Error);
+		return FBlueprintCameraNodeEvaluationResult();
+	}
+}
+
 FBlueprintCameraPose UBlueprintCameraDirectorEvaluator::GetInitialContextCameraPose() const
 {
 	if (EvaluationContext)
diff --git a/Engine/Plugins/Cameras/GameplayCameras/Source/GameplayCameras/Private/GameFramework/GameplayCameraComponent.cpp b/Engine/Plugins/Cameras/GameplayCameras/Source/GameplayCameras/Private/GameFramework/GameplayCameraComponent.cpp
--- a/Engine/Plugins/Cameras/GameplayCameras/Source/GameplayCameras/Private/GameFramework/GameplayCameraComponent.cpp
+++ b/Engine/Plugins/Cameras/GameplayCameras/Source/GameplayCameras/Private/GameFramework/GameplayCameraComponent.cpp
@@ -152,11 +152,6 @@ void UGameplayCameraComponent::ActivateCameraEvaluationContext(APlayerController
 		InitParams.PlayerController = PlayerController;
 		EvaluationContext->Initialize(InitParams);
 
-		FCameraNodeEvaluationResult& InitialResult = EvaluationContext->GetInitialResult();
-		const FCameraAssetAllocationInfo& AllocationInfo = CameraReference.GetCameraAsset()->GetAllocationInfo();
-		InitialResult.VariableTable.Initialize(AllocationInfo.VariableTableInfo);
-		InitialResult.ContextDataTable.Initialize(AllocationInfo.ContextDataTableInfo);
-
 		UpdateCameraEvaluationContext(true);
 	}
 
@@ -181,11 +176,18 @@ void UGameplayCameraComponent::ActivateCameraEvaluationContext(APlayerController
 
 FBlueprintCameraNodeEvaluationResult UGameplayCameraComponent::GetInitialResult() const
 {
-	UE_PRIVATE_GAMEPLAY_CAMERA_COMPONENT_VALIDATE_EVALUATION_CONTEXT("Can't get initial result", FBlueprintCameraNodeEvaluationResult());
+	UE_PRIVATE_GAMEPLAY_CAMERA_COMPONENT_VALIDATE_EVALUATION_CONTEXT("Can't get shared camera data", FBlueprintCameraNodeEvaluationResult());
 
 	return FBlueprintCameraNodeEvaluationResult(&EvaluationContext->GetInitialResult());
 }
 
+FBlueprintCameraNodeEvaluationResult UGameplayCameraComponent::GetConditionalResult(ECameraEvaluationDataCondition Condition) const
+{
+	UE_PRIVATE_GAMEPLAY_CAMERA_COMPONENT_VALIDATE_EVALUATION_CONTEXT("Can't get conditional camera data", FBlueprintCameraNodeEvaluationResult());
+
+	return FBlueprintCameraNodeEvaluationResult(&EvaluationContext->GetOrAddConditionalResult(Condition));
+}
+
 FBlueprintCameraPose UGameplayCameraComponent::GetInitialPose() const
 {
 	UE_PRIVATE_GAMEPLAY_CAMERA_COMPONENT_VALIDATE_EVALUATION_CONTEXT("Can't get initial camera pose", FBlueprintCameraPose());
@@ -284,7 +286,7 @@ void UGameplayCameraComponent::UpdateCameraEvaluationContext(bool bApplyParamete
 	FCameraNodeEvaluationResult& InitialResult = EvaluationContext->GetInitialResult();
 
 	const FTransform& OwnerTransform = GetComponentTransform();
-	InitialResult.CameraPose.SetTransform(OwnerTransform);
+	InitialResult.CameraPose.SetTransform(OwnerTransform, true);
 	InitialResult.bIsCameraCut = false;
 	InitialResult.bIsValid = true;
 
diff --git a/Engine/Plugins/Cameras/GameplayCameras/Source/GameplayCameras/Public/Core/BlendStackCameraNode.h b/Engine/Plugins/Cameras/GameplayCameras/Source/GameplayCameras/Public/Core/BlendStackCameraNode.h
--- a/Engine/Plugins/Cameras/GameplayCameras/Source/GameplayCameras/Public/Core/BlendStackCameraNode.h
+++ b/Engine/Plugins/Cameras/GameplayCameras/Source/GameplayCameras/Public/Core/BlendStackCameraNode.h
@@ -155,10 +155,15 @@ protected:
 
 	struct FResolvedEntry
 	{
+		FResolvedEntry(FCameraRigEntry& InEntry, TSharedPtr<const FCameraEvaluationContext> InContext)
+			: Entry(InEntry), Context(InContext)
+		{}
+
 		FCameraRigEntry& Entry;
 		TSharedPtr<const FCameraEvaluationContext> Context;
-		int32 EntryIndex;
-		bool bHasPreBlendedParameters;
+		int32 EntryIndex = INDEX_NONE;
+		bool bIsActiveEntry = false;
+		bool bHasPreBlendedParameters = false;
 	};
 
 	void ResolveEntries(TArray<FResolvedEntry>& OutResolvedEntries);
diff --git a/Engine/Plugins/Cameras/GameplayCameras/Source/GameplayCameras/Public/Core/CameraContextDataTable.h b/Engine/Plugins/Cameras/GameplayCameras/Source/GameplayCameras/Public/Core/CameraContextDataTable.h
--- a/Engine/Plugins/Cameras/GameplayCameras/Source/GameplayCameras/Public/Core/CameraContextDataTable.h
+++ b/Engine/Plugins/Cameras/GameplayCameras/Source/GameplayCameras/Public/Core/CameraContextDataTable.h
@@ -28,6 +28,19 @@ struct TCameraContextDataReader
 	void Initialize(FCameraContextDataID InID, const FCameraContextDataTable& InContextDataTable);
 };
 
+/**
+ * Filter for context data table operations.
+ */
+enum class ECameraContextDataTableFilter
+{
+	None = 0,
+	/** Only include data that is common to both tables. */
+	KnownOnly = 1 << 0,
+	/** Only include data that was written this frame. */
+	ChangedOnly = 1 << 1,
+};
+ENUM_CLASS_FLAGS(ECameraContextDataTableFilter)
+
 /**
  * The camera context data table is a container for a collection of arbitrary values
  * of various types. It is the companion of the camera variable table (see FCameraVariableTable)
@@ -93,6 +106,7 @@ public:
 
 	void OverrideAll(const FCameraContextDataTable& OtherTable);
 	void OverrideKnown(const FCameraContextDataTable& OtherTable);
+	void Override(const FCameraContextDataTable& OtherTable, ECameraContextDataTableFilter Filter);
 
 public:
 
@@ -144,7 +158,7 @@ private:
 	const FEntry* FindEntry(FCameraContextDataID InID) const;
 	FEntry* FindEntry(FCameraContextDataID InID);
 
-	void Override(const FCameraContextDataTable& OtherTable, bool bKnownOnly);
+	void InternalOverride(const FCameraContextDataTable& OtherTable, ECameraContextDataTableFilter Filter);
 
 	void ReallocateBuffer(uint32 MinRequired = 0);
 	void DestroyBuffer();
diff --git a/Engine/Plugins/Cameras/GameplayCameras/Source/GameplayCameras/Public/Core/CameraEvaluationContext.h b/Engine/Plugins/Cameras/GameplayCameras/Source/GameplayCameras/Public/Core/CameraEvaluationContext.h
--- a/Engine/Plugins/Cameras/GameplayCameras/Source/GameplayCameras/Public/Core/CameraEvaluationContext.h
+++ b/Engine/Plugins/Cameras/GameplayCameras/Source/GameplayCameras/Public/Core/CameraEvaluationContext.h
@@ -9,10 +9,18 @@
 #include "Core/CameraObjectRtti.h"
 #include "Templates/SharedPointer.h"
 
+#include "CameraEvaluationContext.generated.h"
+
 class APlayerController;
 class UCameraAsset;
 class UCameraDirector;
 
+UENUM(BlueprintType)
+enum class ECameraEvaluationDataCondition : uint8
+{
+	ActiveCameraRig
+};
+
 namespace UE::Cameras
 {
 
@@ -101,6 +109,17 @@ public:
 	/** Gets the children evaluation contexts running inside this context. */
 	TArrayView<const TSharedPtr<FCameraEvaluationContext>> GetChildrenContexts() const { return ChildrenContexts; }
 
+public:
+
+	/** Gets an evaluation result that will be conditionally overlaid on some camera rigs in this context. */
+	const FCameraNodeEvaluationResult* GetConditionalResult(ECameraEvaluationDataCondition Condition) const { return ConditionalResults.Find(Condition); }
+
+	/** Gets an evaluation result that will be conditionally overlaid on some camera rigs in this context. */
+	FCameraNodeEvaluationResult* GetConditionalResult(ECameraEvaluationDataCondition Condition) { return ConditionalResults.Find(Condition); }
+
+	/** Gets or adds an evaluation result that will be conditionally overlaid on some camera rigs in this context. */
+	GAMEPLAYCAMERAS_API FCameraNodeEvaluationResult& GetOrAddConditionalResult(ECameraEvaluationDataCondition Condition);
+
 public:
 
 	/**
@@ -160,6 +179,9 @@ private:
 	using FChildrenContexts = TArray<TSharedPtr<FCameraEvaluationContext>>;
 	FChildrenContexts ChildrenContexts;
 
+	using FConditionalResults = TMap<ECameraEvaluationDataCondition, FCameraNodeEvaluationResult>;
+	FConditionalResults ConditionalResults;
+
 	bool bInitialized = false;
 	bool bActivated = false;
 };
diff --git a/Engine/Plugins/Cameras/GameplayCameras/Source/GameplayCameras/Public/Core/CameraPose.h b/Engine/Plugins/Cameras/GameplayCameras/Source/GameplayCameras/Public/Core/CameraPose.h
--- a/Engine/Plugins/Cameras/GameplayCameras/Source/GameplayCameras/Public/Core/CameraPose.h
+++ b/Engine/Plugins/Cameras/GameplayCameras/Source/GameplayCameras/Public/Core/CameraPose.h
@@ -144,7 +144,7 @@ public:
 	FTransform3d GetTransform() const;
 
 	/** Sets the transform of the camera. */
-	void SetTransform(FTransform3d Transform);
+	void SetTransform(FTransform3d Transform, bool bForceSet = false);
 
 	/**
 	 * Computes the horizontal field of view of the camera.
diff --git a/Engine/Plugins/Cameras/GameplayCameras/Source/GameplayCameras/Public/Core/CameraVariableTable.h b/Engine/Plugins/Cameras/GameplayCameras/Source/GameplayCameras/Public/Core/CameraVariableTable.h
--- a/Engine/Plugins/Cameras/GameplayCameras/Source/GameplayCameras/Public/Core/CameraVariableTable.h
+++ b/Engine/Plugins/Cameras/GameplayCameras/Source/GameplayCameras/Public/Core/CameraVariableTable.h
@@ -32,13 +32,11 @@ enum class ECameraVariableTableFilter
 	Output = 1 << 1,
 	/** Include private variables. */
 	Private = 1 << 3,
-	/** Include changed variables. */
+	/** Only include variables that were written this frame. */
 	ChangedOnly = 1 << 2,
 
-	/** All variables. */
-	AllPublic = Input | Output,
-	/** All changed variables. */
-	AllPublicChanged = Input | Output | ChangedOnly
+	/** Both input and output variables. */
+	InputOutput = Input | Output,
 };
 ENUM_CLASS_FLAGS(ECameraVariableTableFilter)
 
diff --git a/Engine/Plugins/Cameras/GameplayCameras/Source/GameplayCameras/Public/Directors/BlueprintCameraDirector.h b/Engine/Plugins/Cameras/GameplayCameras/Source/GameplayCameras/Public/Directors/BlueprintCameraDirector.h
--- a/Engine/Plugins/Cameras/GameplayCameras/Source/GameplayCameras/Public/Directors/BlueprintCameraDirector.h
+++ b/Engine/Plugins/Cameras/GameplayCameras/Source/GameplayCameras/Public/Directors/BlueprintCameraDirector.h
@@ -4,6 +4,7 @@
 
 #include "Core/CameraDirector.h"
 #include "Core/CameraDirectorEvaluator.h"
+#include "Core/CameraEvaluationContext.h"
 #include "GameFramework/BlueprintCameraNodeEvaluationResult.h"
 #include "GameFramework/BlueprintCameraPose.h"
 #include "GameFramework/BlueprintCameraVariableTable.h"
@@ -176,9 +177,13 @@ public:
 	AActor* FindEvaluationContextOwnerActor(TSubclassOf<AActor> ActorClass) const;
 
 	/** Gets the shared evaluation context data. */
-	UFUNCTION(BlueprintPure, Category="Camera Director|Evaluation", meta=(DisplayName="Get Shared Evaluation Data"))
+	UFUNCTION(BlueprintPure, Category="Camera Director|Evaluation", meta=(DisplayName="Get Shared Camera Data"))
 	FBlueprintCameraNodeEvaluationResult GetInitialContextResult() const;
 
+	/** Gets the evaluation context data for a sub-set of camera rigs. */	
+	UFUNCTION(BlueprintPure, Category="Camera Director|Evaluation", meta=(DisplayName="Get Conditional Camera Data"))
+	FBlueprintCameraNodeEvaluationResult GetConditionalContextResult(ECameraEvaluationDataCondition Condition) const;
+
 	UFUNCTION(BlueprintPure, Category="Camera Director|Evaluation",
 			meta=(DeprecatedFunction, DeprecationMessage="Please use GetSharedEvaluationData"))
 	FBlueprintCameraPose GetInitialContextCameraPose() const;
diff --git a/Engine/Plugins/Cameras/GameplayCameras/Source/GameplayCameras/Public/GameFramework/GameplayCameraComponent.h b/Engine/Plugins/Cameras/GameplayCameras/Source/GameplayCameras/Public/GameFramework/GameplayCameraComponent.h
--- a/Engine/Plugins/Cameras/GameplayCameras/Source/GameplayCameras/Public/GameFramework/GameplayCameraComponent.h
+++ b/Engine/Plugins/Cameras/GameplayCameras/Source/GameplayCameras/Public/GameFramework/GameplayCameraComponent.h
@@ -60,6 +60,10 @@ public:
 	/** Gets the shared camera evaluation data for this component's evaluation context. */
 	UFUNCTION(BlueprintPure, Category=Camera, meta=(DisplayName="Get Shared Camera Data"))
 	GAMEPLAYCAMERAS_API FBlueprintCameraNodeEvaluationResult GetInitialResult() const;
+	
+	/** Gets the camera evaluation data for a given sub-set of camera rigs in this component's evaluation context. */
+	UFUNCTION(BlueprintPure, Category=Camera, meta=(DisplayName="Get Conditional Camera Data"))
+	GAMEPLAYCAMERAS_API FBlueprintCameraNodeEvaluationResult GetConditionalResult(ECameraEvaluationDataCondition Condition) const;
 
 	/** Gets the initial camera pose for this component's camera evaluation context. */
 	UFUNCTION(BlueprintPure, Category=Camera, meta=(DeprecatedFunction, DeprecationMessage="Please use GetSharedCameraData"))
