Cameras: add a single-setter BP node for setting a camera rig parameter

- Make the new K2 node class for setting a single parameter
- Extract common code from the multi-setter K2 node into a base class
- Make those two node classes into more strongly typed setter classes tied to a specific camera rig asset, just like normal BP setters are tied to a specific BP class.
- Add all setters to all known camera rigs in the BP action database

diff --git a/Engine/Plugins/Cameras/GameplayCameras/Source/GameplayCameras/Private/Core/CameraRigAsset.cpp b/Engine/Plugins/Cameras/GameplayCameras/Source/GameplayCameras/Private/Core/CameraRigAsset.cpp
--- a/Engine/Plugins/Cameras/GameplayCameras/Source/GameplayCameras/Private/Core/CameraRigAsset.cpp
+++ b/Engine/Plugins/Cameras/GameplayCameras/Source/GameplayCameras/Private/Core/CameraRigAsset.cpp
@@ -6,6 +6,7 @@
 #include "Core/CameraNode.h"
 #include "Core/CameraRigAssetBuilder.h"
 #include "Core/CameraVariableAssets.h"
+#include "UObject/AssetRegistryTagsContext.h"
 #include "UObject/ObjectSaveContext.h"
 
 #include UE_INLINE_GENERATED_CPP_BY_NAME(CameraRigAsset)
@@ -403,3 +404,43 @@ void UCameraRigAsset::RemoveConnectableObject(FName InGraphName, UObject* InObje
 
 #endif  // WITH_EDITOR
 
+void UCameraRigAsset::GetAssetRegistryTags(FAssetRegistryTagsContext Context) const
+{
+	const UEnum* VariableTypeEnum = StaticEnum<ECameraVariableType>();
+	for (const UCameraRigBlendableParameter* BlendableParameter : Interface.BlendableParameters)
+	{
+		if (BlendableParameter)
+		{
+			FAssetRegistryTag BlendableParameterTag;
+			BlendableParameterTag.Name = FName(BlendableParameter->InterfaceParameterName);
+			BlendableParameterTag.Value = VariableTypeEnum->GetNameStringByValue((int64)BlendableParameter->ParameterType);
+			BlendableParameterTag.Type = FAssetRegistryTag::TT_Alphabetical;
+			Context.AddTag(BlendableParameterTag);
+		}
+	}
+
+	const UEnum* ContextDataTypeEnum = StaticEnum<ECameraContextDataType>();
+	for (const UCameraRigDataParameter* DataParameter : Interface.DataParameters)
+	{
+		if (DataParameter)
+		{
+			FAssetRegistryTag DataParameterTag;
+			DataParameterTag.Name = FName(DataParameter->InterfaceParameterName);
+			DataParameterTag.Value = ContextDataTypeEnum->GetNameStringByValue((int64)DataParameter->DataType);
+			DataParameterTag.Type = FAssetRegistryTag::TT_Alphabetical;
+			Context.AddTag(DataParameterTag);
+		}
+	}
+
+	Super::GetAssetRegistryTags(Context);
+}
+
+#if WITH_EDITOR
+
+void UCameraRigAsset::GetAssetRegistryTagMetadata(TMap<FName, FAssetRegistryTagMetadata>& OutMetadata) const
+{
+	Super::GetAssetRegistryTagMetadata(OutMetadata);
+}
+
+#endif  // WITH_EDITOR
+
diff --git a/Engine/Plugins/Cameras/GameplayCameras/Source/GameplayCameras/Public/Core/CameraRigAsset.h b/Engine/Plugins/Cameras/GameplayCameras/Source/GameplayCameras/Public/Core/CameraRigAsset.h
--- a/Engine/Plugins/Cameras/GameplayCameras/Source/GameplayCameras/Public/Core/CameraRigAsset.h
+++ b/Engine/Plugins/Cameras/GameplayCameras/Source/GameplayCameras/Public/Core/CameraRigAsset.h
@@ -86,6 +86,16 @@ struct TStructOpsTypeTraits<FCameraRigAllocationInfo> : public TStructOpsTypeTra
 	};
 };
 
+/**
+ * The type of a camera rig parameter.
+ */
+UENUM()
+enum class ECameraRigInterfaceParameterType : uint8
+{
+	Blendable,
+	Data
+};
+
 /**
  * Base class for interface parameters on a camera rig asset.
  */
@@ -372,6 +382,10 @@ protected:
 	virtual void PostInitProperties() override;
 	virtual void PostDuplicate(EDuplicateMode::Type DuplicateMode) override;
 	virtual void PreSave(FObjectPreSaveContext ObjectSaveContext) override;
+	virtual void GetAssetRegistryTags(FAssetRegistryTagsContext Context) const override;
+#if WITH_EDITOR
+	virtual void GetAssetRegistryTagMetadata(TMap<FName, FAssetRegistryTagMetadata>& OutMetadata) const override;
+#endif
 
 private:
 
diff --git a/Engine/Plugins/Cameras/GameplayCameras/Source/GameplayCamerasEditor/Private/Editors/GameplayCamerasGraphPanelPinFactory.cpp b/Engine/Plugins/Cameras/GameplayCameras/Source/GameplayCamerasEditor/Private/Editors/GameplayCamerasGraphPanelPinFactory.cpp
--- a/Engine/Plugins/Cameras/GameplayCameras/Source/GameplayCamerasEditor/Private/Editors/GameplayCamerasGraphPanelPinFactory.cpp
+++ b/Engine/Plugins/Cameras/GameplayCameras/Source/GameplayCamerasEditor/Private/Editors/GameplayCamerasGraphPanelPinFactory.cpp
@@ -24,17 +24,12 @@ TSharedPtr<SGraphPin> FGameplayCamerasGraphPanelPinFactory::CreatePin(UEdGraphPi
 	{
 		if (PinPropertyClass->IsChildOf<UCameraVariableAsset>())
 		{
-			return CreateCameraVariablePickerPin(Pin);
+			return SNew(SCameraVariableNameGraphPin, Pin);
 		}
 	}
 
 	return nullptr;
 }
 
-TSharedPtr<SGraphPin> FGameplayCamerasGraphPanelPinFactory::CreateCameraVariablePickerPin(UEdGraphPin* Pin) const
-{
-	return SNew(SCameraVariableNameGraphPin, Pin);
-}
-
 }  // namespace UE::Cameras
 
diff --git a/Engine/Plugins/Cameras/GameplayCameras/Source/GameplayCamerasEditor/Public/Editors/GameplayCamerasGraphPanelPinFactory.h b/Engine/Plugins/Cameras/GameplayCameras/Source/GameplayCamerasEditor/Public/Editors/GameplayCamerasGraphPanelPinFactory.h
--- a/Engine/Plugins/Cameras/GameplayCameras/Source/GameplayCamerasEditor/Public/Editors/GameplayCamerasGraphPanelPinFactory.h
+++ b/Engine/Plugins/Cameras/GameplayCameras/Source/GameplayCamerasEditor/Public/Editors/GameplayCamerasGraphPanelPinFactory.h
@@ -19,10 +19,6 @@ public:
 
 	// FGraphPanelPinFactory interface.
 	virtual TSharedPtr<SGraphPin> CreatePin(UEdGraphPin* Pin) const override;
-
-private:
-
-	TSharedPtr<SGraphPin> CreateCameraVariablePickerPin(UEdGraphPin* Pin) const;
 };
 
 }  // namespace UE::Cameras
diff --git a/Engine/Plugins/Cameras/GameplayCameras/Source/GameplayCamerasUncookedOnly/GameplayCamerasUncookedOnly.Build.cs b/Engine/Plugins/Cameras/GameplayCameras/Source/GameplayCamerasUncookedOnly/GameplayCamerasUncookedOnly.Build.cs
--- a/Engine/Plugins/Cameras/GameplayCameras/Source/GameplayCamerasUncookedOnly/GameplayCamerasUncookedOnly.Build.cs
+++ b/Engine/Plugins/Cameras/GameplayCameras/Source/GameplayCamerasUncookedOnly/GameplayCamerasUncookedOnly.Build.cs
@@ -23,6 +23,7 @@ namespace UnrealBuildTool.Rules
 					"BlueprintGraph",
 					"Kismet",
 					"KismetCompiler",
+					"SlateCore",
 					"UnrealEd",
 				}
 			);
diff --git a/Engine/Plugins/Cameras/GameplayCameras/Source/GameplayCamerasUncookedOnly/Private/BlueprintGraph/K2Node_CameraRigBase.cpp b/Engine/Plugins/Cameras/GameplayCameras/Source/GameplayCamerasUncookedOnly/Private/BlueprintGraph/K2Node_CameraRigBase.cpp
new file mode 100644
--- /dev/null
+++ b/Engine/Plugins/Cameras/GameplayCameras/Source/GameplayCamerasUncookedOnly/Private/BlueprintGraph/K2Node_CameraRigBase.cpp
@@ -0,0 +1,312 @@
+// Copyright Epic Games, Inc. All Rights Reserved.
+
+#include "BlueprintGraph/K2Node_CameraRigBase.h"
+
+#include "AssetRegistry/AssetRegistryModule.h"
+#include "AssetRegistry/IAssetRegistry.h"
+#include "BlueprintActionDatabaseRegistrar.h"
+#include "Core/CameraRigAsset.h"
+#include "EdGraphSchema_K2.h"
+#include "Editor.h"
+#include "EditorCategoryUtils.h"
+#include "GameFramework/BlueprintCameraNodeEvaluationResult.h"
+#include "GameFramework/CameraRigParameterInterop.h"
+#include "GameplayCamerasDelegates.h"
+#include "Kismet2/BlueprintEditorUtils.h"
+#include "KismetCompiler.h"
+#include "Subsystems/AssetEditorSubsystem.h"
+
+#define LOCTEXT_NAMESPACE "K2Node_CameraRigBase"
+
+const FName UK2Node_CameraRigBase::CameraNodeEvaluationResultPinName(TEXT("CameraData"));
+
+UK2Node_CameraRigBase::UK2Node_CameraRigBase(const FObjectInitializer& ObjectInit)
+	: Super(ObjectInit)
+{
+	using namespace UE::Cameras;
+	
+	FGameplayCamerasDelegates::OnCameraRigAssetBuilt().AddUObject(this, &UK2Node_CameraRigBase::OnCameraRigAssetBuilt);
+}
+
+void UK2Node_CameraRigBase::BeginDestroy()
+{
+	using namespace UE::Cameras;
+	
+	FGameplayCamerasDelegates::OnCameraRigAssetBuilt().RemoveAll(this);
+
+	Super::BeginDestroy();
+}
+
+void UK2Node_CameraRigBase::AllocateDefaultPins()
+{
+	using namespace UE::Cameras;
+
+	// Add execution pins.
+	CreatePin(EGPD_Input, UEdGraphSchema_K2::PC_Exec, UEdGraphSchema_K2::PN_Execute);
+	CreatePin(EGPD_Output, UEdGraphSchema_K2::PC_Exec, UEdGraphSchema_K2::PN_Then);
+
+	// Add evalation result pin.
+	CreatePin(EGPD_Input, UEdGraphSchema_K2::PC_Struct, FBlueprintCameraNodeEvaluationResult::StaticStruct(), CameraNodeEvaluationResultPinName);
+
+	Super::AllocateDefaultPins();
+}
+
+void UK2Node_CameraRigBase::ValidateNodeDuringCompilation(FCompilerResultsLog& MessageLog) const
+{
+	Super::ValidateNodeDuringCompilation(MessageLog);
+
+	if (!CameraRig)
+	{
+		const FText MessageText = LOCTEXT("MissingCameraRig", "Invalid camera rig reference inside node @@");
+		MessageLog.Error(*MessageText.ToString(), this);
+	}
+}
+
+bool UK2Node_CameraRigBase::CanJumpToDefinition() const
+{
+	return CameraRig != nullptr;
+}
+
+void UK2Node_CameraRigBase::JumpToDefinition() const
+{
+	if (CameraRig)
+	{
+		GEditor->GetEditorSubsystem<UAssetEditorSubsystem>()->OpenEditorForAsset(CameraRig);
+	}
+}
+
+FText UK2Node_CameraRigBase::GetMenuCategory() const
+{
+	const FText BaseCategoryString = FEditorCategoryUtils::BuildCategoryString(
+			FCommonEditorCategory::Gameplay, 
+			LOCTEXT("CameraRigAssetsEditorCategory", "Camera Rigs"));
+	return BaseCategoryString;
+}
+
+void UK2Node_CameraRigBase::GetMenuActions(FBlueprintActionDatabaseRegistrar& ActionRegistrar) const
+{
+	UClass* ActionKey = GetClass();
+	if (ActionRegistrar.IsOpenForRegistration(ActionKey))
+	{
+		FARFilter Filter;
+		Filter.ClassPaths.Add(UCameraRigAsset::StaticClass()->GetClassPathName());
+
+		TArray<FAssetData> CameraRigAssetDatas;
+		IAssetRegistry& AssetRegistry = FAssetRegistryModule::GetRegistry();
+		AssetRegistry.GetAssets(Filter, CameraRigAssetDatas);
+
+		for (const FAssetData& CameraRigAssetData : CameraRigAssetDatas)
+		{
+			GetMenuActions(ActionRegistrar, CameraRigAssetData);
+		}
+	}
+	else if (const UCameraRigAsset* CameraRigKeyFilter = Cast<const UCameraRigAsset>(ActionRegistrar.GetActionKeyFilter()))
+	{
+		const FAssetData CameraRigAssetData(CameraRigKeyFilter);
+		GetMenuActions(ActionRegistrar, CameraRigAssetData);
+	}
+}
+
+UEdGraphPin* UK2Node_CameraRigBase::GetCameraNodeEvaluationResultPin() const
+{
+	return FindPinChecked(CameraNodeEvaluationResultPinName);
+}
+
+bool UK2Node_CameraRigBase::ValidateCameraRigBeforeExpandNode(FKismetCompilerContext& CompilerContext) const
+{
+	if (!CameraRig)
+	{
+		CompilerContext.MessageLog.Error(*LOCTEXT("ErrorMissingCameraRig", "SetCameraRigParameter node @@ doesn't have a valid camera rig set.").ToString(), this);
+		return false;
+	}
+	return true;
+}
+
+void UK2Node_CameraRigBase::OnCameraRigAssetBuilt(const UCameraRigAsset* InBuiltCameraRig)
+{
+	if (CameraRig && CameraRig == InBuiltCameraRig)
+	{
+		ReconstructNode();
+	}
+}
+
+FEdGraphPinType UK2Node_CameraRigBase::MakeBlendableParameterPinType(const UCameraRigBlendableParameter* BlendableParameter)
+{
+	return MakeBlendableParameterPinType(BlendableParameter->PrivateVariable->GetVariableType());
+}
+
+FEdGraphPinType UK2Node_CameraRigBase::MakeBlendableParameterPinType(ECameraVariableType CameraVariableType)
+{
+	FName PinCategory;
+	FName PinSubCategory;
+	UObject* PinSubCategoryObject = nullptr;
+	switch (CameraVariableType)
+	{
+		case ECameraVariableType::Boolean:
+			PinCategory = UEdGraphSchema_K2::PC_Boolean;
+			break;
+		case ECameraVariableType::Integer32:
+			PinCategory = UEdGraphSchema_K2::PC_Int;
+			break;
+		case ECameraVariableType::Float:
+			// We'll cast down to float.
+			PinCategory = UEdGraphSchema_K2::PC_Real;
+			PinSubCategory = UEdGraphSchema_K2::PC_Float;
+			break;
+		case ECameraVariableType::Double:
+			PinCategory = UEdGraphSchema_K2::PC_Real;
+			PinSubCategory = UEdGraphSchema_K2::PC_Double;
+			break;
+		case ECameraVariableType::Vector2d:
+			PinCategory = UEdGraphSchema_K2::PC_Struct;
+			PinSubCategoryObject = TBaseStructure<FVector2D>::Get();
+			break;
+		case ECameraVariableType::Vector3d:
+			PinCategory = UEdGraphSchema_K2::PC_Struct;
+			PinSubCategoryObject = TBaseStructure<FVector>::Get();
+			break;
+		case ECameraVariableType::Vector4d:
+			PinCategory = UEdGraphSchema_K2::PC_Struct;
+			PinSubCategoryObject = TBaseStructure<FVector4>::Get();
+			break;
+		case ECameraVariableType::Rotator3d:
+			PinCategory = UEdGraphSchema_K2::PC_Struct;
+			PinSubCategoryObject = TBaseStructure<FRotator>::Get();
+			break;
+		case ECameraVariableType::Transform3d:
+			PinCategory = UEdGraphSchema_K2::PC_Struct;
+			PinSubCategoryObject = TBaseStructure<FTransform>::Get();
+			break;
+	}
+
+	FEdGraphPinType PinType;
+	PinType.PinCategory = PinCategory;
+	PinType.PinSubCategory = PinSubCategory;
+	PinType.PinSubCategoryObject = PinSubCategoryObject;
+	return PinType;
+}
+
+FEdGraphPinType UK2Node_CameraRigBase::MakeDataParameterPinType(const UCameraRigDataParameter* DataParameter)
+{
+	return MakeDataParameterPinType(DataParameter->DataType, DataParameter->DataTypeObject);
+}
+
+FEdGraphPinType UK2Node_CameraRigBase::MakeDataParameterPinType(ECameraContextDataType CameraContextDataType, const UObject* CameraContextDataTypeObject)
+{
+	FName PinCategory;
+	UObject* PinSubCategoryObject = const_cast<UObject*>(CameraContextDataTypeObject);
+	switch (CameraContextDataType)
+	{
+		case ECameraContextDataType::Name:
+			PinCategory = UEdGraphSchema_K2::PC_Name;
+			break;
+		case ECameraContextDataType::String:
+			PinCategory = UEdGraphSchema_K2::PC_String;
+			break;
+		case ECameraContextDataType::Enum:
+			PinCategory = UEdGraphSchema_K2::PC_Enum;
+			break;
+		case ECameraContextDataType::Struct:
+			PinCategory = UEdGraphSchema_K2::PC_Struct;
+			break;
+		case ECameraContextDataType::Object:
+			PinCategory = UEdGraphSchema_K2::PC_Object;
+			break;
+		case ECameraContextDataType::Class:
+			PinCategory = UEdGraphSchema_K2::PC_Class;
+			break;
+	}
+
+	FEdGraphPinType PinType;
+	PinType.PinCategory = PinCategory;
+	PinType.PinSubCategoryObject = PinSubCategoryObject;
+	return PinType;
+}
+
+FName UK2Node_CameraRigBase::GetBlendableParameterInteropSettingFunctionName(const UCameraRigBlendableParameter* BlendableParameter)
+{
+	if (ensure(BlendableParameter && BlendableParameter->PrivateVariable))
+	{
+		return GetBlendableParameterInteropSettingFunctionName(BlendableParameter->PrivateVariable->GetVariableType());
+	}
+	return NAME_None;
+}
+
+FName UK2Node_CameraRigBase::GetBlendableParameterInteropSettingFunctionName(ECameraVariableType CameraVariableType)
+{
+	// Figure out the sort of SetXxxParameter function we want to call for this parameter.
+	FName CallSetParameterFuncName;
+	switch (CameraVariableType)
+	{
+		case ECameraVariableType::Boolean:
+			CallSetParameterFuncName = GET_FUNCTION_NAME_CHECKED(UCameraRigParameterInterop, SetBooleanParameter);
+			break;
+		case ECameraVariableType::Integer32:
+			CallSetParameterFuncName = GET_FUNCTION_NAME_CHECKED(UCameraRigParameterInterop, SetIntegerParameter);
+			break;
+		case ECameraVariableType::Float:
+			CallSetParameterFuncName = GET_FUNCTION_NAME_CHECKED(UCameraRigParameterInterop, SetFloatParameter);
+			break;
+		case ECameraVariableType::Double:
+			CallSetParameterFuncName = GET_FUNCTION_NAME_CHECKED(UCameraRigParameterInterop, SetDoubleParameter);
+			break;
+		case ECameraVariableType::Vector2d:
+			CallSetParameterFuncName = GET_FUNCTION_NAME_CHECKED(UCameraRigParameterInterop, SetVector2Parameter);
+			break;
+		case ECameraVariableType::Vector3d:
+			CallSetParameterFuncName = GET_FUNCTION_NAME_CHECKED(UCameraRigParameterInterop, SetVector3Parameter);
+			break;
+		case ECameraVariableType::Vector4d:
+			CallSetParameterFuncName = GET_FUNCTION_NAME_CHECKED(UCameraRigParameterInterop, SetVector4Parameter);
+			break;
+		case ECameraVariableType::Rotator3d:
+			CallSetParameterFuncName = GET_FUNCTION_NAME_CHECKED(UCameraRigParameterInterop, SetRotatorParameter);
+			break;
+		case ECameraVariableType::Transform3d:
+			CallSetParameterFuncName = GET_FUNCTION_NAME_CHECKED(UCameraRigParameterInterop, SetTransformParameter);
+			break;
+	}
+
+	return CallSetParameterFuncName;
+}
+
+FName UK2Node_CameraRigBase::GetDataParameterInteropSettingFunctionName(const UCameraRigDataParameter* DataParameter)
+{
+	if (ensure(DataParameter && DataParameter->PrivateDataID.IsValid()))
+	{
+		return GetDataParameterInteropSettingFunctionName(DataParameter->DataType, DataParameter->DataTypeObject);
+	}
+	return NAME_None;
+}
+
+FName UK2Node_CameraRigBase::GetDataParameterInteropSettingFunctionName(ECameraContextDataType CameraContextDataType, const UObject* CameraContextDataTypeObject)
+{
+	// Figure out the sort of SetXxxParameter function we want to call for this parameter.
+	FName CallSetParameterFuncName;
+	switch (CameraContextDataType)
+	{
+		case ECameraContextDataType::Name:
+			CallSetParameterFuncName = GET_FUNCTION_NAME_CHECKED(UCameraRigParameterInterop, SetNameParameter);
+			break;
+		case ECameraContextDataType::String:
+			CallSetParameterFuncName = GET_FUNCTION_NAME_CHECKED(UCameraRigParameterInterop, SetStringParameter);
+			break;
+		case ECameraContextDataType::Enum:
+			CallSetParameterFuncName = GET_FUNCTION_NAME_CHECKED(UCameraRigParameterInterop, SetEnumParameter);
+			break;
+		case ECameraContextDataType::Struct:
+			CallSetParameterFuncName = GET_FUNCTION_NAME_CHECKED(UCameraRigParameterInterop, SetStructParameter);
+			break;
+		case ECameraContextDataType::Object:
+			CallSetParameterFuncName = GET_FUNCTION_NAME_CHECKED(UCameraRigParameterInterop, SetObjectParameter);
+			break;
+		case ECameraContextDataType::Class:
+			CallSetParameterFuncName = GET_FUNCTION_NAME_CHECKED(UCameraRigParameterInterop, SetClassParameter);
+			break;
+	}
+
+	return CallSetParameterFuncName;
+}
+
+#undef LOCTEXT_NAMESPACE
+
diff --git a/Engine/Plugins/Cameras/GameplayCameras/Source/GameplayCamerasUncookedOnly/Private/BlueprintGraph/K2Node_CameraRigBase.h b/Engine/Plugins/Cameras/GameplayCameras/Source/GameplayCamerasUncookedOnly/Private/BlueprintGraph/K2Node_CameraRigBase.h
new file mode 100644
--- /dev/null
+++ b/Engine/Plugins/Cameras/GameplayCameras/Source/GameplayCamerasUncookedOnly/Private/BlueprintGraph/K2Node_CameraRigBase.h
@@ -0,0 +1,76 @@
+// Copyright Epic Games, Inc. All Rights Reserved.
+
+#pragma once
+
+#include "Core/CameraContextDataTableFwd.h"
+#include "Core/CameraVariableTableFwd.h"
+#include "EdGraph/EdGraphNode.h"
+#include "K2Node.h"
+#include "UObject/ObjectMacros.h"
+#include "UObject/UObjectGlobals.h"
+
+#include "K2Node_CameraRigBase.generated.h"
+
+class UCameraRigAsset;
+class UCameraRigBlendableParameter;
+class UCameraRigDataParameter;
+
+/**
+ * Utility base class for Blueprint nodes that can set camera rig parameters.
+ */
+UCLASS(MinimalAPI, Abstract)
+class UK2Node_CameraRigBase : public UK2Node
+{
+	GENERATED_BODY()
+
+public:
+
+	UK2Node_CameraRigBase(const FObjectInitializer& ObjectInit);
+
+public:
+
+	// UObject interface.
+	virtual void BeginDestroy() override;
+
+	// UEdGraphNode interface.
+	virtual void AllocateDefaultPins() override;
+	virtual void ValidateNodeDuringCompilation(FCompilerResultsLog& MessageLog) const override;
+	virtual bool CanJumpToDefinition() const override;
+	virtual void JumpToDefinition() const override;
+
+	// UK2Node interface.
+	virtual FText GetMenuCategory() const override;
+	virtual void GetMenuActions(FBlueprintActionDatabaseRegistrar& ActionRegistrar) const override;
+
+public:
+
+	static FEdGraphPinType MakeBlendableParameterPinType(const UCameraRigBlendableParameter* BlendableParameter);
+	static FEdGraphPinType MakeBlendableParameterPinType(ECameraVariableType CameraVariableType);
+	static FEdGraphPinType MakeDataParameterPinType(const UCameraRigDataParameter* DataParameter);
+	static FEdGraphPinType MakeDataParameterPinType(ECameraContextDataType CameraContextDataType, const UObject* CameraContextDataTypeObject);
+
+	static FName GetBlendableParameterInteropSettingFunctionName(const UCameraRigBlendableParameter* BlendableParameter);
+	static FName GetBlendableParameterInteropSettingFunctionName(ECameraVariableType CameraVariableType);
+	static FName GetDataParameterInteropSettingFunctionName(const UCameraRigDataParameter* DataParameter);
+	static FName GetDataParameterInteropSettingFunctionName(ECameraContextDataType CameraContextDataType, const UObject* CameraContextDataTypeObject);
+
+protected:
+
+	// UK2Node_CameraRigBase interface.
+	virtual void GetMenuActions(FBlueprintActionDatabaseRegistrar& ActionRegistrar, const FAssetData& CameraRigAssetData) const {}
+
+protected:
+
+	static const FName CameraNodeEvaluationResultPinName;
+
+	UEdGraphPin* GetCameraNodeEvaluationResultPin() const;
+	bool ValidateCameraRigBeforeExpandNode(FKismetCompilerContext& CompilerContext) const;
+
+	void OnCameraRigAssetBuilt(const UCameraRigAsset* InBuiltCameraRig);
+
+protected:
+
+	UPROPERTY()
+	TObjectPtr<UCameraRigAsset> CameraRig;
+};
+
diff --git a/Engine/Plugins/Cameras/GameplayCameras/Source/GameplayCamerasUncookedOnly/Private/BlueprintGraph/K2Node_SetCameraRigParameter.cpp b/Engine/Plugins/Cameras/GameplayCameras/Source/GameplayCamerasUncookedOnly/Private/BlueprintGraph/K2Node_SetCameraRigParameter.cpp
new file mode 100644
--- /dev/null
+++ b/Engine/Plugins/Cameras/GameplayCameras/Source/GameplayCamerasUncookedOnly/Private/BlueprintGraph/K2Node_SetCameraRigParameter.cpp
@@ -0,0 +1,257 @@
+// Copyright Epic Games, Inc. All Rights Reserved.
+
+#include "BlueprintGraph/K2Node_SetCameraRigParameter.h"
+
+#include "BlueprintActionDatabaseRegistrar.h"
+#include "BlueprintGraph/K2Node_CameraRigBase.h"
+#include "BlueprintNodeSpawner.h"
+#include "Core/CameraRigAsset.h"
+#include "EdGraphSchema_K2.h"
+#include "EditorCategoryUtils.h"
+#include "GameFramework/CameraRigParameterInterop.h"
+#include "K2Node_CallFunction.h"
+#include "Kismet/BlueprintInstancedStructLibrary.h"
+#include "Kismet2/BlueprintEditorUtils.h"
+#include "KismetCompiler.h"
+#include "Styling/AppStyle.h"
+#include "Styling/SlateIconFinder.h"
+
+#define LOCTEXT_NAMESPACE "K2Node_SetCameraRigParameter"
+
+UK2Node_SetCameraRigParameter::UK2Node_SetCameraRigParameter(const FObjectInitializer& ObjectInit)
+	: Super(ObjectInit)
+{
+}
+
+void UK2Node_SetCameraRigParameter::Initialize(const FAssetData& UnloadedCameraRig, const FString& InCameraParameterName)
+{
+	UCameraRigAsset* LoadedCameraRig = Cast<UCameraRigAsset>(UnloadedCameraRig.GetAsset());
+	if (ensure(LoadedCameraRig))
+	{
+		if (const UCameraRigBlendableParameter* BlendableParameter = LoadedCameraRig->Interface.FindBlendableParameterByName(InCameraParameterName))
+		{
+			Initialize(LoadedCameraRig, InCameraParameterName, BlendableParameter->ParameterType);
+		}
+		else if (const UCameraRigDataParameter* DataParameter = LoadedCameraRig->Interface.FindDataParameterByName(InCameraParameterName))
+		{
+			Initialize(LoadedCameraRig, InCameraParameterName, DataParameter->DataType, DataParameter->DataTypeObject);
+		}
+		// else, no parameter of that name found...
+	}
+}
+
+void UK2Node_SetCameraRigParameter::Initialize(UCameraRigAsset* InCameraRig, const FString& InCameraParameterName, ECameraVariableType InCameraVariableType)
+{
+	CameraRig = InCameraRig;
+	CameraParameterName = InCameraParameterName;
+	CameraParameterType = EK2Node_CameraParameterType::Blendable;
+	BlendableCameraParameterType = InCameraVariableType;
+}
+
+void UK2Node_SetCameraRigParameter::Initialize(UCameraRigAsset* InCameraRig, const FString& InCameraParameterName, ECameraContextDataType InCameraContextDataType, const UObject* InCameraContextDataTypeObject)
+{
+	CameraRig = InCameraRig;
+	CameraParameterName = InCameraParameterName;
+	CameraParameterType = EK2Node_CameraParameterType::Data;
+	DataCameraParameterType = InCameraContextDataType;
+	DataCameraParameterTypeObject = const_cast<UObject*>(InCameraContextDataTypeObject);
+}
+
+void UK2Node_SetCameraRigParameter::AllocateDefaultPins()
+{
+	Super::AllocateDefaultPins();
+
+	// Add parameter value pin.
+	FEdGraphPinType PinType;
+	switch (CameraParameterType)
+	{
+		case EK2Node_CameraParameterType::Blendable:
+			PinType = UK2Node_CameraRigBase::MakeBlendableParameterPinType(BlendableCameraParameterType);
+			break;
+		case EK2Node_CameraParameterType::Data:
+			PinType = UK2Node_CameraRigBase::MakeDataParameterPinType(DataCameraParameterType, DataCameraParameterTypeObject);
+			break;
+		default:
+			PinType.PinCategory = UEdGraphSchema_K2::PC_Wildcard;
+			break;
+	}
+	CreatePin(EGPD_Input, PinType, FName(CameraParameterName));
+}
+
+FText UK2Node_SetCameraRigParameter::GetNodeTitle(ENodeTitleType::Type TitleType) const
+{
+	return FText::Format(LOCTEXT("BaseNodeTitle", "SET on {0}"), FText::FromString(GetNameSafe(CameraRig)));
+}
+
+FText UK2Node_SetCameraRigParameter::GetTooltipText() const
+{
+	return FText::Format(
+			LOCTEXT("NodeTooltip", "Sets the value of camera rig {0}'s parameter {1} on the given evaluation data."),
+			FText::FromString(GetNameSafe(CameraRig)),
+			FText::FromString(CameraParameterName));
+}
+
+FSlateIcon UK2Node_SetCameraRigParameter::GetIconAndTint(FLinearColor& OutColor) const
+{
+	UEdGraphPin* ParameterValuePin = FindPin(CameraParameterName);
+	if (ParameterValuePin)
+	{
+		const UEdGraphSchema_K2* K2Schema = GetDefault<UEdGraphSchema_K2>();
+		OutColor = K2Schema->GetPinTypeColor(ParameterValuePin->PinType);
+
+		if (UStruct* Struct = Cast<UStruct>(ParameterValuePin->PinType.PinSubCategoryObject.Get()))
+		{
+			return FSlateIconFinder::FindIconForClass(Struct);
+		}
+	}
+
+	return FSlateIcon(FAppStyle::GetAppStyleSetName(), "Kismet.AllClasses.VariableIcon");
+}
+
+void UK2Node_SetCameraRigParameter::GetMenuActions(FBlueprintActionDatabaseRegistrar& ActionRegistrar, const FAssetData& CameraRigAssetData) const
+{
+	const FText BaseCategoryString = GetMenuCategory();
+
+	for (TPair<FName, FAssetTagValueRef> It : CameraRigAssetData.TagsAndValues)
+	{
+		const FName ParameterName = It.Key;
+
+		UBlueprintNodeSpawner* NodeSpawner = UBlueprintNodeSpawner::Create(GetClass());
+		NodeSpawner->DefaultMenuSignature.Category = FText::Join(
+				FText::FromString(TEXT("|")), BaseCategoryString, FText::FromName(CameraRigAssetData.AssetName));
+		NodeSpawner->DefaultMenuSignature.MenuName = FText::Format(
+				LOCTEXT("SetCameraRigParameterActionMenuName", "Set {0}"),
+				FText::FromName(ParameterName));
+		NodeSpawner->CustomizeNodeDelegate = UBlueprintNodeSpawner::FCustomizeNodeDelegate::CreateLambda(
+				[ParameterName, CameraRigAssetData](UEdGraphNode* NewNode, bool bIsTemplateNode)
+				{
+					UK2Node_SetCameraRigParameter* NewSetter = CastChecked<UK2Node_SetCameraRigParameter>(NewNode);
+					NewSetter->Initialize(CameraRigAssetData, ParameterName.ToString());
+				});
+
+		ActionRegistrar.AddBlueprintAction(CameraRigAssetData, NodeSpawner);
+	}
+}
+
+void UK2Node_SetCameraRigParameter::ExpandNode(FKismetCompilerContext& CompilerContext, UEdGraph* SourceGraph)
+{
+	Super::ExpandNode(CompilerContext, SourceGraph);
+
+	if (!ValidateCameraRigBeforeExpandNode(CompilerContext))
+	{
+		BreakAllNodeLinks();
+		return;
+	}
+
+	if (CameraParameterName.IsEmpty())
+	{
+		CompilerContext.MessageLog.Error(*LOCTEXT("ErrorMissingCameraParameterName", "SetCameraRigParameter node @@ doesn't have a valid camera parameter name set.").ToString(), this);
+		BreakAllNodeLinks();
+		return;
+	}
+
+	FName CallSetParameterFuncName;
+	switch (CameraParameterType)
+	{
+		case EK2Node_CameraParameterType::Blendable:
+			CallSetParameterFuncName = UK2Node_CameraRigBase::GetBlendableParameterInteropSettingFunctionName(
+					BlendableCameraParameterType);
+			break;
+		case EK2Node_CameraParameterType::Data:
+			CallSetParameterFuncName = UK2Node_CameraRigBase::GetDataParameterInteropSettingFunctionName(
+					DataCameraParameterType, DataCameraParameterTypeObject);
+			break;
+		default:
+			break;
+	}
+	if (CallSetParameterFuncName.IsNone())
+	{
+		CompilerContext.MessageLog.Error(*LOCTEXT("ErrorUnknownParameterType", "SetCameraRigParameter node @@ is trying to set parameter @@ but that parameter is missing or invalid.").ToString(), this, *CameraParameterName);
+		BreakAllNodeLinks();
+		return;
+	}
+
+	// Additional setup is needed for some types.
+	UEnum* ExtraEnumTypePinObject = nullptr;
+	bool bMakeExtraInstanceStructSetup = false;
+	if (DataCameraParameterType == ECameraContextDataType::Enum)
+	{
+		const UEnum* EnumType = CastChecked<const UEnum>(DataCameraParameterTypeObject);
+		ExtraEnumTypePinObject = const_cast<UEnum*>(EnumType);
+	}
+	else if (DataCameraParameterType == ECameraContextDataType::Struct)
+	{
+		bMakeExtraInstanceStructSetup = true;
+	}
+
+	UEdGraphPin* const CameraNodeEvaluationResultPin = GetCameraNodeEvaluationResultPin();
+	UEdGraphPin* const CameraParameterValuePin = FindPinChecked(CameraParameterName);
+
+	// Make the SetXxxParameter function call node.
+	UK2Node_CallFunction* CallSetParameter = CompilerContext.SpawnIntermediateNode<UK2Node_CallFunction>(this, SourceGraph);
+	CallSetParameter->FunctionReference.SetExternalMember(CallSetParameterFuncName, UCameraRigParameterInterop::StaticClass());
+	CallSetParameter->AllocateDefaultPins();
+	UEdGraphPin* FirstExecPin = CallSetParameter->GetExecPin();
+
+	// Connect the camera evaluation result argument.
+	UEdGraphPin* CallSetParameterResultPin = CallSetParameter->FindPinChecked(TEXT("Result"));
+	CompilerContext.CopyPinLinksToIntermediate(*CameraNodeEvaluationResultPin, *CallSetParameterResultPin);
+
+	// Set the camera rig argument.
+	UEdGraphPin* CallSetParameterCameraRigPin = CallSetParameter->FindPinChecked(TEXT("CameraRig"));
+	CallSetParameterCameraRigPin->DefaultObject = CameraRig;
+
+	// Set the parameter name argument.
+	UEdGraphPin* CallSetParameterNamePin = CallSetParameter->FindPinChecked(TEXT("ParameterName"));
+	CallSetParameterNamePin->DefaultValue = CameraParameterName;
+
+	// Set or connect the parameter value argument.
+	UEdGraphPin* CallSetParameterValuePin = CallSetParameter->FindPinChecked(TEXT("ParameterValue"));
+	if (bMakeExtraInstanceStructSetup)
+	{
+		// If we are setting a struct, first turn it into an FInstancedStruct. Insert this node in the
+		// execution chain, and use that node's Value pin as the pin to set the incoming struct value.
+		UK2Node_CallFunction* CallMakeInstancedStruct = CompilerContext.SpawnIntermediateNode<UK2Node_CallFunction>(this, SourceGraph);
+		CallMakeInstancedStruct->FunctionReference.SetExternalMember(GET_FUNCTION_NAME_CHECKED(UBlueprintInstancedStructLibrary, MakeInstancedStruct), UBlueprintInstancedStructLibrary::StaticClass());
+		CallMakeInstancedStruct->AllocateDefaultPins();
+
+		CallMakeInstancedStruct->GetReturnValuePin()->MakeLinkTo(CallSetParameterValuePin);
+		CallMakeInstancedStruct->GetThenPin()->MakeLinkTo(FirstExecPin);
+		FirstExecPin = CallMakeInstancedStruct->GetExecPin();
+
+		UEdGraphPin* CallMakeInstancedStructValuePin = CallMakeInstancedStruct->FindPinChecked(TEXT("Value"));
+		CompilerContext.MovePinLinksToIntermediate(*CameraParameterValuePin, *CallMakeInstancedStructValuePin);
+		CallMakeInstancedStruct->PinConnectionListChanged(CallMakeInstancedStructValuePin);
+	}
+	else
+	{
+		CallSetParameterValuePin->DefaultValue = CameraParameterValuePin->DefaultValue;
+		CallSetParameterValuePin->DefaultTextValue = CameraParameterValuePin->DefaultTextValue;
+		CallSetParameterValuePin->AutogeneratedDefaultValue = CameraParameterValuePin->AutogeneratedDefaultValue;
+		CallSetParameterValuePin->DefaultObject = CameraParameterValuePin->DefaultObject;
+		if (CameraParameterValuePin->LinkedTo.Num() > 0)
+		{
+			CompilerContext.MovePinLinksToIntermediate(*CameraParameterValuePin, *CallSetParameterValuePin);
+		}
+	}
+
+	// Set extra type pin for enums.
+	if (ExtraEnumTypePinObject)
+	{
+		UEdGraphPin* CallSetParameterEnumTypePin = CallSetParameter->FindPinChecked(TEXT("EnumType"));
+		CallSetParameterEnumTypePin->DefaultObject = ExtraEnumTypePinObject;
+	}
+
+	// Setup the execution flow.
+	UEdGraphPin* ThisExecPin = GetExecPin();
+	CompilerContext.MovePinLinksToIntermediate(*ThisExecPin, *FirstExecPin);
+
+	UEdGraphPin* ThisThenPin = GetThenPin();
+	UEdGraphPin* CallSetParameterThenPin = CallSetParameter->GetThenPin();
+	CompilerContext.MovePinLinksToIntermediate(*ThisThenPin, *CallSetParameterThenPin);
+
+	BreakAllNodeLinks();
+}
+
+#undef LOCTEXT_NAMESPACE
+
diff --git a/Engine/Plugins/Cameras/GameplayCameras/Source/GameplayCamerasUncookedOnly/Private/BlueprintGraph/K2Node_SetCameraRigParameter.h b/Engine/Plugins/Cameras/GameplayCameras/Source/GameplayCamerasUncookedOnly/Private/BlueprintGraph/K2Node_SetCameraRigParameter.h
new file mode 100644
--- /dev/null
+++ b/Engine/Plugins/Cameras/GameplayCameras/Source/GameplayCamerasUncookedOnly/Private/BlueprintGraph/K2Node_SetCameraRigParameter.h
@@ -0,0 +1,68 @@
+// Copyright Epic Games, Inc. All Rights Reserved.
+
+#pragma once
+
+#include "K2Node_CameraRigBase.h"
+
+#include "K2Node_SetCameraRigParameter.generated.h"
+
+class FKismetCompilerContext;
+class UCameraRigAsset;
+
+UENUM()
+enum class EK2Node_CameraParameterType : uint8
+{
+	Unknown,
+	Blendable,
+	Data
+};
+
+/**
+ * Blueprint node that, given a camera rig, lets the user set the value of one single
+ * exposed parameter.
+ */
+UCLASS(MinimalAPI)
+class UK2Node_SetCameraRigParameter : public UK2Node_CameraRigBase
+{
+	GENERATED_BODY()
+
+public:
+
+	UK2Node_SetCameraRigParameter(const FObjectInitializer& ObjectInit);
+
+	void Initialize(const FAssetData& UnloadedCameraRig, const FString& InCameraParameterName);
+	void Initialize(UCameraRigAsset* InCameraRig, const FString& InCameraParameterName, ECameraVariableType InCameraVariableType);
+	void Initialize(UCameraRigAsset* InCameraRig, const FString& InCameraParameterName, ECameraContextDataType InCameraContextDataType, const UObject* InCameraContextDataTypeObject);
+
+public:
+
+	// UEdGraphNode interface.
+	virtual void AllocateDefaultPins() override;
+	virtual FText GetNodeTitle(ENodeTitleType::Type TitleType) const override;
+	virtual FText GetTooltipText() const override;
+	virtual FSlateIcon GetIconAndTint(FLinearColor& OutColor) const override;
+	virtual void ExpandNode(FKismetCompilerContext& CompilerContext, UEdGraph* SourceGraph) override;
+
+protected:
+
+	// UK2Node_CameraRigBase interface.
+	void GetMenuActions(FBlueprintActionDatabaseRegistrar& ActionRegistrar, const FAssetData& CameraRigAssetData) const;
+
+private:
+
+	UPROPERTY()
+	FString CameraParameterName;
+
+	UPROPERTY()
+	EK2Node_CameraParameterType CameraParameterType = EK2Node_CameraParameterType::Unknown;
+
+	UPROPERTY()
+	ECameraVariableType BlendableCameraParameterType;
+
+	UPROPERTY()
+	ECameraContextDataType DataCameraParameterType;
+
+	UPROPERTY()
+	TObjectPtr<UObject> DataCameraParameterTypeObject;
+};
+
diff --git a/Engine/Plugins/Cameras/GameplayCameras/Source/GameplayCamerasUncookedOnly/Private/BlueprintGraph/K2Node_SetCameraRigParameters.cpp b/Engine/Plugins/Cameras/GameplayCameras/Source/GameplayCamerasUncookedOnly/Private/BlueprintGraph/K2Node_SetCameraRigParameters.cpp
--- a/Engine/Plugins/Cameras/GameplayCameras/Source/GameplayCamerasUncookedOnly/Private/BlueprintGraph/K2Node_SetCameraRigParameters.cpp
+++ b/Engine/Plugins/Cameras/GameplayCameras/Source/GameplayCamerasUncookedOnly/Private/BlueprintGraph/K2Node_SetCameraRigParameters.cpp
@@ -5,10 +5,7 @@
 #include "BlueprintActionDatabaseRegistrar.h"
 #include "BlueprintNodeSpawner.h"
 #include "Core/CameraRigAsset.h"
-#include "Core/CameraVariableAssets.h"
 #include "EdGraphSchema_K2.h"
-#include "EditorCategoryUtils.h"
-#include "GameFramework/BlueprintCameraNodeEvaluationResult.h"
 #include "GameFramework/CameraRigParameterInterop.h"
 #include "K2Node_CallFunction.h"
 #include "Kismet/BlueprintInstancedStructLibrary.h"
@@ -17,39 +14,42 @@
 
 #define LOCTEXT_NAMESPACE "K2Node_SetCameraRigParameters"
 
-const FName UK2Node_SetCameraRigParameters::CameraRigPinName(TEXT("CameraRig"));
-const FName UK2Node_SetCameraRigParameters::CameraNodeEvaluationResultPinName(TEXT("CameraEvaluationResult"));
-
 UK2Node_SetCameraRigParameters::UK2Node_SetCameraRigParameters(const FObjectInitializer& ObjectInit)
 	: Super(ObjectInit)
 {
 }
 
-void UK2Node_SetCameraRigParameters::AllocateDefaultPins()
+void UK2Node_SetCameraRigParameters::Initialize(const FAssetData& UnloadedCameraRig)
 {
-	using namespace UE::Cameras;
-
-	// Add execution pins.
-	CreatePin(EGPD_Input, UEdGraphSchema_K2::PC_Exec, UEdGraphSchema_K2::PN_Execute);
-	CreatePin(EGPD_Output, UEdGraphSchema_K2::PC_Exec, UEdGraphSchema_K2::PN_Then);
-
-	// Add evalation result pin.
-	CreatePin(EGPD_Input, UEdGraphSchema_K2::PC_Struct, FBlueprintCameraNodeEvaluationResult::StaticStruct(), CameraNodeEvaluationResultPinName);
+	CameraRig = Cast<UCameraRigAsset>(UnloadedCameraRig.GetAsset());
+}
 
-	// Add camera rig pin.
-	CreatePin(EGPD_Input, UEdGraphSchema_K2::PC_Object, UCameraRigAsset::StaticClass(), CameraRigPinName);
-	
+void UK2Node_SetCameraRigParameters::AllocateDefaultPins()
+{
 	Super::AllocateDefaultPins();
+
+	CreateParameterPins();
 }
 
 void UK2Node_SetCameraRigParameters::ReallocatePinsDuringReconstruction(TArray<UEdGraphPin*>& OldPins) 
 {
-	AllocateDefaultPins();
+	// Remove old camera rig pin name. Do this before calling AllocateDefaultPins() because
+	// we need our CameraRig property set to (re)create our parameter pins.
+	const FName OldCameraRigPinName = TEXT("CameraRig");
+	UEdGraphPin** OldCameraRigPin = OldPins.FindByPredicate([&OldCameraRigPinName](UEdGraphPin* OldPin)
+			{
+				return OldPin->PinName == OldCameraRigPinName;
+			});
+	if (OldCameraRigPin)
+	{
+		CameraRig = Cast<UCameraRigAsset>((*OldCameraRigPin)->DefaultObject);
+		(*OldCameraRigPin)->DefaultObject = nullptr;
+		(*OldCameraRigPin)->DefaultValue.Reset();
+	}
 
-	TArrayView<UEdGraphPin* const> PinsToSearch = OldPins;
-	if (UCameraRigAsset* CameraRig = GetCameraRig(&PinsToSearch))
+	// The camera rig might not be loaded yet when we are rebuilt on startup.
+	if (CameraRig)
 	{
-		// The camera rig might not be loaded yet.
 		PreloadObject(CameraRig);
 		for (UCameraRigBlendableParameter* BlendableParameter : CameraRig->Interface.BlendableParameters)
 		{
@@ -67,83 +67,71 @@ void UK2Node_SetCameraRigParameters::ReallocatePinsDuringReconstruction(TArray<U
 				PreloadObject(const_cast<UObject*>(DataParameter->DataTypeObject.Get()));
 			}
 		}
-
-		CreatePinsForCameraRig(CameraRig);
 	}
 
-	RestoreSplitPins(OldPins);
-}
-
-void UK2Node_SetCameraRigParameters::PostPlacedNewNode()
-{
-	Super::PostPlacedNewNode();
+	// Create all default pins, including the parameter pins.
+	AllocateDefaultPins();
 
-	if (UCameraRigAsset* CameraRig = GetCameraRig())
+	// Upgrade old result pin name to new pin name.
+	const FName OldResultPinName = TEXT("CameraEvaluationResult");
+	UEdGraphPin** OldResultPin = OldPins.FindByPredicate([&OldResultPinName](UEdGraphPin* OldPin)
+			{
+				return OldPin->PinName == OldResultPinName;
+			});
+	if (OldResultPin)
 	{
-		CreatePinsForCameraRig(CameraRig);
+		UEdGraphPin* NewResultPin = FindPin(CameraNodeEvaluationResultPinName);
+		check(NewResultPin);
+		NewResultPin->MovePersistentDataFromOldPin(**OldResultPin);
 	}
-}
-
-void UK2Node_SetCameraRigParameters::PinConnectionListChanged(UEdGraphPin* Pin)
-{
-	Super::PinConnectionListChanged(Pin);
 
-	if (Pin && (Pin->PinName == CameraRigPinName))
-	{
-		OnCameraRigChanged();
-	}
+	RestoreSplitPins(OldPins);
 }
 
-void UK2Node_SetCameraRigParameters::PinDefaultValueChanged(UEdGraphPin* ChangedPin) 
+FText UK2Node_SetCameraRigParameters::GetNodeTitle(ENodeTitleType::Type TitleType) const
 {
-	if (ChangedPin && (ChangedPin->PinName == CameraRigPinName))
-	{
-		OnCameraRigChanged();
-	}
+	return FText::Format(
+			LOCTEXT("BaseNodeTitle", "SET on {0}"),
+			FText::FromString(GetNameSafe(CameraRig)));
 }
 
 FText UK2Node_SetCameraRigParameters::GetTooltipText() const
 {
-	return LOCTEXT("NodeTooltip", "Sets values for the exposed parameters on the given camera rig.");
-}
-
-FText UK2Node_SetCameraRigParameters::GetNodeTitle(ENodeTitleType::Type TitleType) const
-{
-	return LOCTEXT("BaseNodeTitle", "Set Camera Rig Parameters");
+	return FText::Format(
+			LOCTEXT("NodeTooltip", "Sets the values of all camera rig parameters on {0} on the given evaluation data."),
+			FText::FromString(GetNameSafe(CameraRig)));
 }
 
-void UK2Node_SetCameraRigParameters::GetMenuActions(FBlueprintActionDatabaseRegistrar& ActionRegistrar) const
+void UK2Node_SetCameraRigParameters::GetMenuActions(FBlueprintActionDatabaseRegistrar& ActionRegistrar, const FAssetData& CameraRigAssetData) const
 {
-	UClass* ActionKey = GetClass();
-	if (ActionRegistrar.IsOpenForRegistration(ActionKey))
-	{
-		UBlueprintNodeSpawner* NodeSpawner = UBlueprintNodeSpawner::Create(GetClass());
-		check(NodeSpawner != nullptr);
-
-		ActionRegistrar.AddBlueprintAction(ActionKey, NodeSpawner);
-	}
-}
+	const FText BaseCategoryString = GetMenuCategory();
+
+	UBlueprintNodeSpawner* NodeSpawner = UBlueprintNodeSpawner::Create(GetClass());
+	NodeSpawner->DefaultMenuSignature.Category = BaseCategoryString;
+	NodeSpawner->DefaultMenuSignature.MenuName = FText::Format(
+			LOCTEXT("SetCameraRigParameterActionMenuName", "Set all parameters on {0}"),
+			FText::FromName(CameraRigAssetData.AssetName));
+	NodeSpawner->CustomizeNodeDelegate = UBlueprintNodeSpawner::FCustomizeNodeDelegate::CreateLambda(
+			[CameraRigAssetData](UEdGraphNode* NewNode, bool bIsTemplateNode)
+			{
+				UK2Node_SetCameraRigParameters* NewSetter = CastChecked<UK2Node_SetCameraRigParameters>(NewNode);
+				NewSetter->Initialize(CameraRigAssetData);
+			});
 
-FText UK2Node_SetCameraRigParameters::GetMenuCategory() const
-{
-	return FEditorCategoryUtils::GetCommonCategory(FCommonEditorCategory::Gameplay);
+	ActionRegistrar.AddBlueprintAction(CameraRigAssetData, NodeSpawner);
 }
 
-void UK2Node_SetCameraRigParameters::ExpandNode(class FKismetCompilerContext& CompilerContext, UEdGraph* SourceGraph)
+void UK2Node_SetCameraRigParameters::ExpandNode(FKismetCompilerContext& CompilerContext, UEdGraph* SourceGraph)
 {
 	Super::ExpandNode(CompilerContext, SourceGraph);
 
-	UCameraRigAsset* CameraRig = GetCameraRig();
-
-	if (!CameraRig)
+	if (!ValidateCameraRigBeforeExpandNode(CompilerContext))
 	{
-		CompilerContext.MessageLog.Error(*LOCTEXT("ErrorMissingCameraRig", "SetCameraRigParameters node @@ doesn't have a camera rig set.").ToString(), this);
 		BreakAllNodeLinks();
 		return;
 	}
 
-	UEdGraphPin* const CameraRigPin = FindPinChecked(CameraRigPinName);
-	UEdGraphPin* const CameraNodeEvaluationResultPin = FindPinChecked(CameraNodeEvaluationResultPinName);
+	UEdGraphPin* const CameraNodeEvaluationResultPin = GetCameraNodeEvaluationResultPin();
 
 	UEdGraphPin* OriginalThenPin = GetThenPin();
 	UEdGraphPin* PreviousThenPin = nullptr;
@@ -157,7 +145,7 @@ void UK2Node_SetCameraRigParameters::ExpandNode(class FKismetCompilerContext& Co
 	FindBlendableParameterPins(BlendableParameterPins);
 	for (UEdGraphPin* RigParameterPin : BlendableParameterPins)
 	{
-		UCameraRigBlendableParameter* BlendableParameter = CameraRig->Interface.FindBlendableParameterByName(RigParameterPin->GetName());
+		const UCameraRigBlendableParameter* BlendableParameter = CameraRig->Interface.FindBlendableParameterByName(RigParameterPin->GetName());
 		if (!BlendableParameter)
 		{
 			CompilerContext.MessageLog.Error(*LOCTEXT("ErrorMissingParameter", "SetCameraRigParameters node @@ is trying to set parameter @@ but camera rig @@ has no such parameter.").ToString(), this, *RigParameterPin->GetName(), CameraRig);
@@ -171,37 +159,7 @@ void UK2Node_SetCameraRigParameters::ExpandNode(class FKismetCompilerContext& Co
 		}
 
 		// Figure out the sort of SetXxxParameter function we want to call for this parameter.
-		FName CallSetParameterFuncName;
-		switch (BlendableParameter->PrivateVariable->GetVariableType())
-		{
-			case ECameraVariableType::Boolean:
-				CallSetParameterFuncName = GET_FUNCTION_NAME_CHECKED(UCameraRigParameterInterop, SetBooleanParameter);
-				break;
-			case ECameraVariableType::Integer32:
-				CallSetParameterFuncName = GET_FUNCTION_NAME_CHECKED(UCameraRigParameterInterop, SetIntegerParameter);
-				break;
-			case ECameraVariableType::Float:
-				CallSetParameterFuncName = GET_FUNCTION_NAME_CHECKED(UCameraRigParameterInterop, SetFloatParameter);
-				break;
-			case ECameraVariableType::Double:
-				CallSetParameterFuncName = GET_FUNCTION_NAME_CHECKED(UCameraRigParameterInterop, SetDoubleParameter);
-				break;
-			case ECameraVariableType::Vector2d:
-				CallSetParameterFuncName = GET_FUNCTION_NAME_CHECKED(UCameraRigParameterInterop, SetVector2Parameter);
-				break;
-			case ECameraVariableType::Vector3d:
-				CallSetParameterFuncName = GET_FUNCTION_NAME_CHECKED(UCameraRigParameterInterop, SetVector3Parameter);
-				break;
-			case ECameraVariableType::Vector4d:
-				CallSetParameterFuncName = GET_FUNCTION_NAME_CHECKED(UCameraRigParameterInterop, SetVector4Parameter);
-				break;
-			case ECameraVariableType::Rotator3d:
-				CallSetParameterFuncName = GET_FUNCTION_NAME_CHECKED(UCameraRigParameterInterop, SetRotatorParameter);
-				break;
-			case ECameraVariableType::Transform3d:
-				CallSetParameterFuncName = GET_FUNCTION_NAME_CHECKED(UCameraRigParameterInterop, SetTransformParameter);
-				break;
-		}
+		FName CallSetParameterFuncName = GetBlendableParameterInteropSettingFunctionName(BlendableParameter);
 		if (CallSetParameterFuncName.IsNone())
 		{
 			CompilerContext.MessageLog.Error(*LOCTEXT("ErrorUnsupportedParameterType", "SetCameraRigParameters node @@ is trying to set parameter @@ but it has an unsupported type.").ToString(), this, *RigParameterPin->GetName());
@@ -217,9 +175,9 @@ void UK2Node_SetCameraRigParameters::ExpandNode(class FKismetCompilerContext& Co
 		UEdGraphPin* CallSetParameterResultPin = CallSetParameter->FindPinChecked(TEXT("Result"));
 		CompilerContext.CopyPinLinksToIntermediate(*CameraNodeEvaluationResultPin, *CallSetParameterResultPin);
 
-		// Connect the camera rig argument.
+		// Set the camera rig argument.
 		UEdGraphPin* CallSetParameterCameraRigPin = CallSetParameter->FindPinChecked(TEXT("CameraRig"));
-		CompilerContext.CopyPinLinksToIntermediate(*CameraRigPin, *CallSetParameterCameraRigPin);
+		CallSetParameterCameraRigPin->DefaultObject = CameraRig;
 
 		// Set the parameter name argument.
 		UEdGraphPin* CallSetParameterNamePin = CallSetParameter->FindPinChecked(TEXT("ParameterName"));
@@ -255,7 +213,7 @@ void UK2Node_SetCameraRigParameters::ExpandNode(class FKismetCompilerContext& Co
 	FindDataParameterPins(DataParameterPins);
 	for (UEdGraphPin* RigParameterPin : DataParameterPins)
 	{
-		UCameraRigDataParameter* DataParameter = CameraRig->Interface.FindDataParameterByName(RigParameterPin->GetName());
+		const UCameraRigDataParameter* DataParameter = CameraRig->Interface.FindDataParameterByName(RigParameterPin->GetName());
 		if (!DataParameter)
 		{
 			CompilerContext.MessageLog.Error(*LOCTEXT("ErrorMissingParameter", "SetCameraRigParameters node @@ is trying to set parameter @@ but camera rig @@ has no such parameter.").ToString(), this, *RigParameterPin->GetName(), CameraRig);
@@ -269,28 +227,7 @@ void UK2Node_SetCameraRigParameters::ExpandNode(class FKismetCompilerContext& Co
 		}
 
 		// Figure out the sort of SetXxxParameter function we want to call for this parameter.
-		FName CallSetParameterFuncName;
-		switch (DataParameter->DataType)
-		{
-			case ECameraContextDataType::Name:
-				CallSetParameterFuncName = GET_FUNCTION_NAME_CHECKED(UCameraRigParameterInterop, SetNameParameter);
-				break;
-			case ECameraContextDataType::String:
-				CallSetParameterFuncName = GET_FUNCTION_NAME_CHECKED(UCameraRigParameterInterop, SetStringParameter);
-				break;
-			case ECameraContextDataType::Enum:
-				CallSetParameterFuncName = GET_FUNCTION_NAME_CHECKED(UCameraRigParameterInterop, SetEnumParameter);
-				break;
-			case ECameraContextDataType::Struct:
-				CallSetParameterFuncName = GET_FUNCTION_NAME_CHECKED(UCameraRigParameterInterop, SetStructParameter);
-				break;
-			case ECameraContextDataType::Object:
-				CallSetParameterFuncName = GET_FUNCTION_NAME_CHECKED(UCameraRigParameterInterop, SetObjectParameter);
-				break;
-			case ECameraContextDataType::Class:
-				CallSetParameterFuncName = GET_FUNCTION_NAME_CHECKED(UCameraRigParameterInterop, SetClassParameter);
-				break;
-		}
+		FName CallSetParameterFuncName = GetDataParameterInteropSettingFunctionName(DataParameter);
 		if (CallSetParameterFuncName.IsNone())
 		{
 			CompilerContext.MessageLog.Error(*LOCTEXT("ErrorUnsupportedParameterType", "SetCameraRigParameters node @@ is trying to set parameter @@ but it has an unsupported type.").ToString(), this, *RigParameterPin->GetName());
@@ -307,9 +244,9 @@ void UK2Node_SetCameraRigParameters::ExpandNode(class FKismetCompilerContext& Co
 		UEdGraphPin* CallSetParameterResultPin = CallSetParameter->FindPinChecked(TEXT("Result"));
 		CompilerContext.CopyPinLinksToIntermediate(*CameraNodeEvaluationResultPin, *CallSetParameterResultPin);
 
-		// Connect the camera rig argument.
+		// Set the camera rig argument.
 		UEdGraphPin* CallSetParameterCameraRigPin = CallSetParameter->FindPinChecked(TEXT("CameraRig"));
-		CompilerContext.CopyPinLinksToIntermediate(*CameraRigPin, *CallSetParameterCameraRigPin);
+		CallSetParameterCameraRigPin->DefaultObject = CameraRig;
 
 		// Set the parameter name argument.
 		UEdGraphPin* CallSetParameterNamePin = CallSetParameter->FindPinChecked(TEXT("ParameterName"));
@@ -376,49 +313,15 @@ void UK2Node_SetCameraRigParameters::ExpandNode(class FKismetCompilerContext& Co
 	BreakAllNodeLinks();
 }
 
-UEdGraphPin* UK2Node_SetCameraRigParameters::GetCameraRigPin(TArrayView<UEdGraphPin* const>* InPinsToSearch) const
+void UK2Node_SetCameraRigParameters::CreateParameterPins()
 {
-	TArrayView<UEdGraphPin* const> PinsToSearch = MakeArrayView(Pins);
-	if (InPinsToSearch)
-	{
-		PinsToSearch = *InPinsToSearch;
-	}
-
-	UEdGraphPin* CameraRigPin = nullptr;
-	for (UEdGraphPin* Pin : PinsToSearch)
-	{
-		if (Pin && Pin->PinName == CameraRigPinName)
-		{
-			CameraRigPin = Pin;
-			break;
-		}
-	}
-	check(CameraRigPin == nullptr || CameraRigPin->Direction == EGPD_Input);
-	return CameraRigPin;
-}
+	BlendableParameterPinNames.Reset();
+	DataParameterPinNames.Reset();
 
-UEdGraphPin* UK2Node_SetCameraRigParameters::GetCameraNodeEvaluationResultPin() const
-{
-	UEdGraphPin* ResultPin = nullptr;
-	for (UEdGraphPin* Pin : Pins)
+	if (!CameraRig)
 	{
-		if (Pin && Pin->PinName == CameraNodeEvaluationResultPinName)
-		{
-			ResultPin = Pin;
-			break;
-		}
+		return;
 	}
-	check(ResultPin == nullptr || ResultPin->Direction == EGPD_Input);
-	return ResultPin;
-}
-
-void UK2Node_SetCameraRigParameters::CreatePinsForCameraRig(UCameraRigAsset* CameraRig, TArray<UEdGraphPin*>* CreatedPins)
-{
-	check(CameraRig);
-
-	const UEdGraphSchema_K2* K2Schema = GetDefault<UEdGraphSchema_K2>();
-
-	BlendableParameterPinNames.Reset();
 
 	for (const UCameraRigBlendableParameter* BlendableParameter : CameraRig->Interface.BlendableParameters)
 	{
@@ -432,67 +335,18 @@ void UK2Node_SetCameraRigParameters::CreatePinsForCameraRig(UCameraRigAsset* Cam
 			// Camera rig isn't fully built.
 			continue;
 		}
-
-		FName NewPinCategory;
-		FName NewPinSubCategory;
-		UObject* NewPinSubCategoryObject = nullptr;
-		switch (BlendableParameter->PrivateVariable->GetVariableType())
-		{
-			case ECameraVariableType::Boolean:
-				NewPinCategory = UEdGraphSchema_K2::PC_Boolean;
-				break;
-			case ECameraVariableType::Integer32:
-				NewPinCategory = UEdGraphSchema_K2::PC_Int;
-				break;
-			case ECameraVariableType::Float:
-				// We'll cast down to float.
-				NewPinCategory = UEdGraphSchema_K2::PC_Real;
-				NewPinSubCategory = UEdGraphSchema_K2::PC_Float;
-				break;
-			case ECameraVariableType::Double:
-				NewPinCategory = UEdGraphSchema_K2::PC_Real;
-				NewPinSubCategory = UEdGraphSchema_K2::PC_Double;
-				break;
-			case ECameraVariableType::Vector2d:
-				NewPinCategory = UEdGraphSchema_K2::PC_Struct;
-				NewPinSubCategoryObject = TBaseStructure<FVector2D>::Get();
-				break;
-			case ECameraVariableType::Vector3d:
-				NewPinCategory = UEdGraphSchema_K2::PC_Struct;
-				NewPinSubCategoryObject = TBaseStructure<FVector>::Get();
-				break;
-			case ECameraVariableType::Vector4d:
-				NewPinCategory = UEdGraphSchema_K2::PC_Struct;
-				NewPinSubCategoryObject = TBaseStructure<FVector4>::Get();
-				break;
-			case ECameraVariableType::Rotator3d:
-				NewPinCategory = UEdGraphSchema_K2::PC_Struct;
-				NewPinSubCategoryObject = TBaseStructure<FRotator>::Get();
-				break;
-			case ECameraVariableType::Transform3d:
-				NewPinCategory = UEdGraphSchema_K2::PC_Struct;
-				NewPinSubCategoryObject = TBaseStructure<FTransform>::Get();
-				break;
-		}
-		if (NewPinCategory.IsNone())
+		
+		FEdGraphPinType PinType = MakeBlendableParameterPinType(BlendableParameter);
+		if (PinType.PinCategory.IsNone())
 		{
 			// Unsupported type for Blueprints.
 			continue;
 		}
 
-		UEdGraphPin* NewPin = CreatePin(
-				EGPD_Input, 
-				NewPinCategory, NewPinSubCategory, NewPinSubCategoryObject, 
-				FName(BlendableParameter->InterfaceParameterName));
+		UEdGraphPin* NewPin = CreatePin(EGPD_Input, PinType, FName(BlendableParameter->InterfaceParameterName));
 		BlendableParameterPinNames.Add(NewPin->PinName);
-		if (CreatedPins)
-		{
-			CreatedPins->Add(NewPin);
-		}
 	}
 
-	DataParameterPinNames.Reset();
-
 	for (const UCameraRigDataParameter* DataParameter : CameraRig->Interface.DataParameters)
 	{
 		if (!ensure(DataParameter))
@@ -506,44 +360,15 @@ void UK2Node_SetCameraRigParameters::CreatePinsForCameraRig(UCameraRigAsset* Cam
 			continue;
 		}
 
-		FName NewPinCategory;
-		UObject* NewPinSubCategoryObject = const_cast<UObject*>(DataParameter->DataTypeObject.Get());
-		switch (DataParameter->DataType)
-		{
-			case ECameraContextDataType::Name:
-				NewPinCategory = UEdGraphSchema_K2::PC_Name;
-				break;
-			case ECameraContextDataType::String:
-				NewPinCategory = UEdGraphSchema_K2::PC_String;
-				break;
-			case ECameraContextDataType::Enum:
-				NewPinCategory = UEdGraphSchema_K2::PC_Enum;
-				break;
-			case ECameraContextDataType::Struct:
-				NewPinCategory = UEdGraphSchema_K2::PC_Struct;
-				break;
-			case ECameraContextDataType::Object:
-				NewPinCategory = UEdGraphSchema_K2::PC_Object;
-				break;
-			case ECameraContextDataType::Class:
-				NewPinCategory = UEdGraphSchema_K2::PC_Class;
-				break;
-		}
-		if (NewPinCategory.IsNone())
+		FEdGraphPinType PinType = MakeDataParameterPinType(DataParameter);
+		if (PinType.PinCategory.IsNone())
 		{
 			// Unsupported type for Blueprints.
 			continue;
 		}
 
-		UEdGraphPin* NewPin = CreatePin(
-				EGPD_Input, 
-				NewPinCategory, NAME_None, NewPinSubCategoryObject, 
-				FName(DataParameter->InterfaceParameterName));
+		UEdGraphPin* NewPin = CreatePin(EGPD_Input, PinType, FName(DataParameter->InterfaceParameterName));
 		DataParameterPinNames.Add(NewPin->PinName);
-		if (CreatedPins)
-		{
-			CreatedPins->Add(NewPin);
-		}
 	}
 }
 
@@ -571,58 +396,5 @@ void UK2Node_SetCameraRigParameters::FindDataParameterPins(TArray<UEdGraphPin*>&
 	}
 }
 
-UCameraRigAsset* UK2Node_SetCameraRigParameters::GetCameraRig(TArrayView<UEdGraphPin* const>* InPinsToSearch) const
-{
-	TArrayView<UEdGraphPin* const> PinsToSearch(Pins);
-	if (InPinsToSearch)
-	{
-		PinsToSearch = *InPinsToSearch;
-	}
-
-	UEdGraphPin* CameraRigPin = GetCameraRigPin(&PinsToSearch);
-	if (CameraRigPin && CameraRigPin->DefaultObject && CameraRigPin->LinkedTo.Num() == 0)
-	{
-		return CastChecked<UCameraRigAsset>(CameraRigPin->DefaultObject);
-	}
-	else if (CameraRigPin && CameraRigPin->LinkedTo.Num() > 0)
-	{
-		if (UEdGraphPin* CameraRigSource = CameraRigPin->LinkedTo[0])
-		{
-			return Cast<UCameraRigAsset>(CameraRigSource->PinType.PinSubCategoryObject.Get());
-		}
-	}
-	return nullptr;
-}
-
-void UK2Node_SetCameraRigParameters::OnCameraRigChanged()
-{
-	TArray<FName> OldCameraRigPinNames;
-	OldCameraRigPinNames.Append(BlendableParameterPinNames);
-	OldCameraRigPinNames.Append(DataParameterPinNames);
-
-	TArray<UEdGraphPin*> OldPins = Pins;
-	TArray<UEdGraphPin*> OldCameraRigPins;
-
-	for (UEdGraphPin* OldPin : OldPins)
-	{
-		if (OldCameraRigPinNames.Contains(OldPin->PinName))
-		{
-			Pins.Remove(OldPin);
-			OldCameraRigPins.Add(OldPin);
-		}
-	}
-
-	TArray<UEdGraphPin*> NewPins;
-	if (UCameraRigAsset* CameraRig = GetCameraRig())
-	{
-		CreatePinsForCameraRig(CameraRig, &NewPins);
-	}
-
-	RewireOldPinsToNewPins(OldCameraRigPins, Pins, nullptr);
-
-	GetGraph()->NotifyGraphChanged();
-	FBlueprintEditorUtils::MarkBlueprintAsModified(GetBlueprint());
-}
-
 #undef LOCTEXT_NAMESPACE
 
diff --git a/Engine/Plugins/Cameras/GameplayCameras/Source/GameplayCamerasUncookedOnly/Private/BlueprintGraph/K2Node_SetCameraRigParameters.h b/Engine/Plugins/Cameras/GameplayCameras/Source/GameplayCamerasUncookedOnly/Private/BlueprintGraph/K2Node_SetCameraRigParameters.h
--- a/Engine/Plugins/Cameras/GameplayCameras/Source/GameplayCamerasUncookedOnly/Private/BlueprintGraph/K2Node_SetCameraRigParameters.h
+++ b/Engine/Plugins/Cameras/GameplayCameras/Source/GameplayCamerasUncookedOnly/Private/BlueprintGraph/K2Node_SetCameraRigParameters.h
@@ -2,23 +2,19 @@
 
 #pragma once
 
-#include "CoreTypes.h"
-#include "EdGraph/EdGraphNode.h"
-#include "K2Node.h"
-#include "UObject/ObjectMacros.h"
-#include "UObject/UObjectGlobals.h"
+#include "K2Node_CameraRigBase.h"
 
 #include "K2Node_SetCameraRigParameters.generated.h"
 
 class UCameraRigAsset;
 
 /**
- * Dynamic Blueprint node that, given a camera rig, exposes input pins for any camera rig parameters
- * found on it. On compile, this node gets expanded into the appropriate number of individual setter
- * function calls for each parameter (see UCameraRigParameterInterop).
+ * Blueprint node that, given a camera rig, lets the user set the values of all exposed parameters
+ * on that camera rig. Any camera rig running with the specific evaluation data will use those
+ * values.
  */
 UCLASS(MinimalAPI)
-class UK2Node_SetCameraRigParameters : public UK2Node
+class UK2Node_SetCameraRigParameters : public UK2Node_CameraRigBase
 {
 	GENERATED_BODY()
 
@@ -26,37 +22,30 @@ public:
 
 	UK2Node_SetCameraRigParameters(const FObjectInitializer& ObjectInit);
 
+	void Initialize(const FAssetData& UnloadedCameraRig);
+
 public:
 
 	// UEdGraphNode interface.
 	virtual void AllocateDefaultPins() override;
 	virtual FText GetNodeTitle(ENodeTitleType::Type TitleType) const override;
-	virtual void PinDefaultValueChanged(UEdGraphPin* Pin) override;
 	virtual FText GetTooltipText() const override;
-	virtual void PinConnectionListChanged(UEdGraphPin* Pin);
-	virtual void PostPlacedNewNode() override;
 	virtual void ExpandNode(class FKismetCompilerContext& CompilerContext, UEdGraph* SourceGraph) override;
 
 	// UK2Node interface.
-	virtual bool IsNodeSafeToIgnore() const override { return true; }
 	virtual void ReallocatePinsDuringReconstruction(TArray<UEdGraphPin*>& OldPins) override;
-	virtual void GetMenuActions(FBlueprintActionDatabaseRegistrar& ActionRegistrar) const override;
-	virtual FText GetMenuCategory() const override;
 
 protected:
 
-	static const FName CameraRigPinName;
-	static const FName CameraNodeEvaluationResultPinName;
+	// UK2Node_CameraRigBase interface.
+	virtual void GetMenuActions(FBlueprintActionDatabaseRegistrar& ActionRegistrar, const FAssetData& CameraRigAssetData) const override;
+
+private:
 
-	UEdGraphPin* GetCameraRigPin(TArrayView<UEdGraphPin* const>* InPinsToSearch = nullptr) const;
-	UEdGraphPin* GetCameraNodeEvaluationResultPin() const;
-	void CreatePinsForCameraRig(UCameraRigAsset* CameraRig, TArray<UEdGraphPin*>* CreatedPins = nullptr);
+	void CreateParameterPins();
 	void FindBlendableParameterPins(TArray<UEdGraphPin*>& OutPins) const;
 	void FindDataParameterPins(TArray<UEdGraphPin*>& OutPins) const;
 
-	UCameraRigAsset* GetCameraRig(TArrayView<UEdGraphPin* const>* InPinsToSearch = nullptr) const;
-	void OnCameraRigChanged();
-
 private:
 
 	UPROPERTY()
