Cameras: add location/rotation blend node

diff --git a/Engine/Plugins/Cameras/GameplayCameras/Source/GameplayCameras/Private/Nodes/Blends/LocationRotationBlendCameraNode.cpp b/Engine/Plugins/Cameras/GameplayCameras/Source/GameplayCameras/Private/Nodes/Blends/LocationRotationBlendCameraNode.cpp
new file mode 100644
--- /dev/null
+++ b/Engine/Plugins/Cameras/GameplayCameras/Source/GameplayCameras/Private/Nodes/Blends/LocationRotationBlendCameraNode.cpp
@@ -0,0 +1,198 @@
+// Copyright Epic Games, Inc. All Rights Reserved.
+
+#include "Nodes/Blends/LocationRotationBlendCameraNode.h"
+
+#include "Debug/CameraDebugBlock.h"
+#include "Debug/CameraDebugBlockBuilder.h"
+#include "Debug/CameraDebugRenderer.h"
+#include "Nodes/Blends/PopBlendCameraNode.h"
+
+#include UE_INLINE_GENERATED_CPP_BY_NAME(LocationRotationBlendCameraNode)
+
+namespace UE::Cameras
+{
+
+class FLocationRotationBlendCameraNodeEvaluator : public FBlendCameraNodeEvaluator
+{
+	UE_DECLARE_BLEND_CAMERA_NODE_EVALUATOR(GAMEPLAYCAMERAS_API, FLocationRotationBlendCameraNodeEvaluator)
+
+protected:
+
+	// FCameraNodeEvaluator interface.
+	virtual void OnBuild(const FCameraNodeEvaluatorBuildParams& Params) override;
+	virtual FCameraNodeEvaluatorChildrenView OnGetChildren() override;
+	virtual void OnRun(const FCameraNodeEvaluationParams& Params, FCameraNodeEvaluationResult& OutResult) override;
+	virtual void OnBlendParameters(const FCameraNodePreBlendParams& Params, FCameraNodePreBlendResult& OutResult) override;
+	virtual void OnBlendResults(const FCameraNodeBlendParams& Params, FCameraNodeBlendResult& OutResult) override;
+#if UE_GAMEPLAY_CAMERAS_DEBUG
+	virtual void OnBuildDebugBlocks(const FCameraDebugBlockBuildParams& Params, FCameraDebugBlockBuilder& Builder) override;
+#endif  // UE_GAMEPLAY_CAMERAS_DEBUG
+
+private:
+
+	FSimpleBlendCameraNodeEvaluator* LocationBlendEvaluator = nullptr;
+	FSimpleBlendCameraNodeEvaluator* RotationBlendEvaluator = nullptr;
+	FSimpleBlendCameraNodeEvaluator* OtherBlendEvaluator = nullptr;
+};
+
+UE_DEFINE_BLEND_CAMERA_NODE_EVALUATOR(FLocationRotationBlendCameraNodeEvaluator)
+
+UE_DECLARE_CAMERA_DEBUG_BLOCK_START(GAMEPLAYCAMERAS_API, FLocationRotationBlendCameraDebugBlock)
+	UE_DECLARE_CAMERA_DEBUG_BLOCK_FIELD(float, LocationBlendFactor);
+	UE_DECLARE_CAMERA_DEBUG_BLOCK_FIELD(float, RotationBlendFactor);
+	UE_DECLARE_CAMERA_DEBUG_BLOCK_FIELD(float, OtherBlendFactor);
+UE_DECLARE_CAMERA_DEBUG_BLOCK_END()
+
+UE_DEFINE_CAMERA_DEBUG_BLOCK_WITH_FIELDS(FLocationRotationBlendCameraDebugBlock)
+
+void FLocationRotationBlendCameraNodeEvaluator::OnBuild(const FCameraNodeEvaluatorBuildParams& Params)
+{
+	const ULocationRotationBlendCameraNode* LocationRotationBlend = GetCameraNodeAs<ULocationRotationBlendCameraNode>();
+	if (LocationRotationBlend->LocationBlend)
+	{
+		LocationBlendEvaluator = Params.BuildEvaluatorAs<FSimpleBlendCameraNodeEvaluator>(LocationRotationBlend->LocationBlend);
+	}
+	if (LocationRotationBlend->RotationBlend)
+	{
+		RotationBlendEvaluator = Params.BuildEvaluatorAs<FSimpleBlendCameraNodeEvaluator>(LocationRotationBlend->RotationBlend);
+	}
+	if (LocationRotationBlend->OtherBlend)
+	{
+		OtherBlendEvaluator = Params.BuildEvaluatorAs<FSimpleBlendCameraNodeEvaluator>(LocationRotationBlend->OtherBlend);
+	}
+}
+
+FCameraNodeEvaluatorChildrenView FLocationRotationBlendCameraNodeEvaluator::OnGetChildren()
+{
+	return FCameraNodeEvaluatorChildrenView({ LocationBlendEvaluator, RotationBlendEvaluator, OtherBlendEvaluator });
+}
+
+void FLocationRotationBlendCameraNodeEvaluator::OnBlendParameters(const FCameraNodePreBlendParams& Params, FCameraNodePreBlendResult& OutResult)
+{
+	if (OtherBlendEvaluator)
+	{
+		OtherBlendEvaluator->BlendParameters(Params, OutResult);
+	}
+	else
+	{
+		FPopBlendCameraNodeHelper::PopParameters(Params, OutResult);
+	}
+}
+
+void FLocationRotationBlendCameraNodeEvaluator::OnRun(const FCameraNodeEvaluationParams& Params, FCameraNodeEvaluationResult& OutResult)
+{
+	if (LocationBlendEvaluator)
+	{
+		LocationBlendEvaluator->Run(Params, OutResult);
+	}
+	if (RotationBlendEvaluator)
+	{
+		RotationBlendEvaluator->Run(Params, OutResult);
+	}
+	if (OtherBlendEvaluator)
+	{
+		OtherBlendEvaluator->Run(Params, OutResult);
+	}
+}
+
+void FLocationRotationBlendCameraNodeEvaluator::OnBlendResults(const FCameraNodeBlendParams& Params, FCameraNodeBlendResult& OutResult)
+{
+	const FCameraNodeEvaluationResult& ChildResult(Params.ChildResult);
+	FCameraNodeEvaluationResult& BlendedResult(OutResult.BlendedResult);
+
+	// Save the "from" location and rotation.
+	const FVector FromLocation = BlendedResult.CameraPose.GetLocation();
+	const FRotator FromRotation = BlendedResult.CameraPose.GetRotation();
+
+	// Blend the whole result. We will overwrite location and rotation with our own logic.
+	bool bOtherBlendFull = true;
+	bool bOtherBlendFinished = true;
+	if (OtherBlendEvaluator)
+	{
+		OtherBlendEvaluator->BlendResults(Params, OutResult);
+
+		bOtherBlendFull = OutResult.bIsBlendFull;
+		bOtherBlendFinished = OutResult.bIsBlendFinished;
+	}
+	else
+	{
+		FPopBlendCameraNodeHelper::PopResults(Params, OutResult);
+	}
+
+	bool bLocationBlendFull = bOtherBlendFull;
+	bool bLocationBlendFinished = bOtherBlendFinished;
+	if (LocationBlendEvaluator)
+	{
+		const float LocationBlendFactor = LocationBlendEvaluator->GetBlendFactor();
+		const FVector ToLocation = ChildResult.CameraPose.GetLocation();
+		const FVector BlendedLocation = FMath::Lerp(FromLocation, ToLocation, LocationBlendFactor);
+		BlendedResult.CameraPose.SetLocation(BlendedLocation);
+
+		bLocationBlendFull = LocationBlendEvaluator->IsBlendFull();
+		bLocationBlendFinished = LocationBlendEvaluator->IsBlendFinished();
+	}
+
+	bool bRotationBlendFull = bOtherBlendFull;
+	bool bRotationBlendFinished = bOtherBlendFinished;
+	if (RotationBlendEvaluator)
+	{
+		const float RotationBlendFactor = RotationBlendEvaluator->GetBlendFactor();
+		const FRotator ToRotation = ChildResult.CameraPose.GetRotation();
+		const FRotator BlendedRotation = FMath::Lerp(FromRotation, ToRotation, RotationBlendFactor);
+		BlendedResult.CameraPose.SetRotation(BlendedRotation);
+
+		bRotationBlendFull = RotationBlendEvaluator->IsBlendFull();
+		bRotationBlendFinished = RotationBlendEvaluator->IsBlendFinished();
+	}
+
+	OutResult.bIsBlendFull = (bOtherBlendFull && bLocationBlendFull && bRotationBlendFull);
+	OutResult.bIsBlendFinished = (bOtherBlendFinished && bLocationBlendFinished && bRotationBlendFinished);
+}
+
+#if UE_GAMEPLAY_CAMERAS_DEBUG
+
+void FLocationRotationBlendCameraNodeEvaluator::OnBuildDebugBlocks(const FCameraDebugBlockBuildParams& Params, FCameraDebugBlockBuilder& Builder)
+{
+	FLocationRotationBlendCameraDebugBlock& DebugBlock = Builder.AttachDebugBlock<FLocationRotationBlendCameraDebugBlock>();
+	DebugBlock.LocationBlendFactor = LocationBlendEvaluator ? LocationBlendEvaluator->GetBlendFactor() : -1.f;
+	DebugBlock.RotationBlendFactor = RotationBlendEvaluator ? RotationBlendEvaluator->GetBlendFactor() : -1.f;
+	DebugBlock.OtherBlendFactor = OtherBlendEvaluator ? OtherBlendEvaluator->GetBlendFactor() : -1.f;
+}
+
+void FLocationRotationBlendCameraDebugBlock::OnDebugDraw(const FCameraDebugBlockDrawParams& Params, FCameraDebugRenderer& Renderer)
+{
+	if (LocationBlendFactor >= 0.f)
+	{
+		Renderer.AddText(TEXT("location %.2f%% "), LocationBlendFactor * 100.f);
+	}
+	if (RotationBlendFactor >= 0.f)
+	{
+		Renderer.AddText(TEXT("rotation %.2f%% "), RotationBlendFactor * 100.f);
+	}
+	if (OtherBlendFactor >= 0.f)
+	{
+		Renderer.AddText(TEXT("other %.2f%% "), OtherBlendFactor * 100.f);
+	}
+}
+
+#endif  // UE_GAMEPLAY_CAMERAS_DEBUG
+
+}  // namespace UE::Cameras
+
+ULocationRotationBlendCameraNode::ULocationRotationBlendCameraNode(const FObjectInitializer& ObjInit)
+	: Super(ObjInit)
+{
+	AddNodeFlags(ECameraNodeFlags::CustomGetChildren);
+}
+
+FCameraNodeChildrenView ULocationRotationBlendCameraNode::OnGetChildren()
+{
+	return FCameraNodeChildrenView({ LocationBlend, RotationBlend, OtherBlend });
+}
+
+FCameraNodeEvaluatorPtr ULocationRotationBlendCameraNode::OnBuildEvaluator(FCameraNodeEvaluatorBuilder& Builder) const
+{
+	using namespace UE::Cameras;
+	return Builder.BuildEvaluator<FLocationRotationBlendCameraNodeEvaluator>();
+}
+
diff --git a/Engine/Plugins/Cameras/GameplayCameras/Source/GameplayCameras/Private/Nodes/Blends/PopBlendCameraNode.cpp b/Engine/Plugins/Cameras/GameplayCameras/Source/GameplayCameras/Private/Nodes/Blends/PopBlendCameraNode.cpp
--- a/Engine/Plugins/Cameras/GameplayCameras/Source/GameplayCameras/Private/Nodes/Blends/PopBlendCameraNode.cpp
+++ b/Engine/Plugins/Cameras/GameplayCameras/Source/GameplayCameras/Private/Nodes/Blends/PopBlendCameraNode.cpp
@@ -24,6 +24,16 @@ void FPopBlendCameraNodeEvaluator::OnRun(const FCameraNodeEvaluationParams& Para
 }
 
 void FPopBlendCameraNodeEvaluator::OnBlendParameters(const FCameraNodePreBlendParams& Params, FCameraNodePreBlendResult& OutResult)
+{
+	FPopBlendCameraNodeHelper::PopParameters(Params, OutResult);
+}
+
+void FPopBlendCameraNodeEvaluator::OnBlendResults(const FCameraNodeBlendParams& Params, FCameraNodeBlendResult& OutResult)
+{
+	FPopBlendCameraNodeHelper::PopResults(Params, OutResult);
+}
+
+void FPopBlendCameraNodeHelper::PopParameters(const UE::Cameras::FCameraNodePreBlendParams& Params, FCameraNodePreBlendResult& OutResult)
 {
 	const FCameraVariableTable& ChildVariableTable(Params.ChildVariableTable);
 	OutResult.VariableTable.Override(ChildVariableTable, ECameraVariableTableFilter::Input);
@@ -32,7 +42,7 @@ void FPopBlendCameraNodeEvaluator::OnBlendParameters(const FCameraNodePreBlendPa
 	OutResult.bIsBlendFinished = true;
 }
 
-void FPopBlendCameraNodeEvaluator::OnBlendResults(const FCameraNodeBlendParams& Params, FCameraNodeBlendResult& OutResult)
+void FPopBlendCameraNodeHelper::PopResults(const UE::Cameras::FCameraNodeBlendParams& Params, FCameraNodeBlendResult& OutResult)
 {
 	const FCameraNodeEvaluationResult& ChildResult(Params.ChildResult);
 	FCameraNodeEvaluationResult& BlendedResult(OutResult.BlendedResult);
diff --git a/Engine/Plugins/Cameras/GameplayCameras/Source/GameplayCameras/Public/Nodes/Blends/LocationRotationBlendCameraNode.h b/Engine/Plugins/Cameras/GameplayCameras/Source/GameplayCameras/Public/Nodes/Blends/LocationRotationBlendCameraNode.h
new file mode 100644
--- /dev/null
+++ b/Engine/Plugins/Cameras/GameplayCameras/Source/GameplayCameras/Public/Nodes/Blends/LocationRotationBlendCameraNode.h
@@ -0,0 +1,42 @@
+// Copyright Epic Games, Inc. All Rights Reserved.
+
+#pragma once
+
+#include "Core/BlendCameraNode.h"
+#include "Nodes/Blends/SimpleBlendCameraNode.h"
+
+#include "LocationRotationBlendCameraNode.generated.h"
+
+/**
+ * A blend node that blends translation and rotation differently.
+ */
+UCLASS(MinimalAPI)
+class ULocationRotationBlendCameraNode : public UBlendCameraNode
+{
+	GENERATED_BODY()
+
+public:
+
+	ULocationRotationBlendCameraNode(const FObjectInitializer& ObjInit);
+
+public:
+
+	/** The blend to use for the camera's translation. */
+	UPROPERTY()
+	TObjectPtr<USimpleBlendCameraNode> LocationBlend;
+
+	/** The blend to use for the camera's rotation. */
+	UPROPERTY()
+	TObjectPtr<USimpleBlendCameraNode> RotationBlend;
+
+	/** The blend to use for everything else. */
+	UPROPERTY()
+	TObjectPtr<USimpleBlendCameraNode> OtherBlend;
+
+protected:
+
+	// UCameraNode interface.
+	virtual FCameraNodeChildrenView OnGetChildren() override;
+	virtual FCameraNodeEvaluatorPtr OnBuildEvaluator(FCameraNodeEvaluatorBuilder& Builder) const override;
+};
+
diff --git a/Engine/Plugins/Cameras/GameplayCameras/Source/GameplayCameras/Public/Nodes/Blends/PopBlendCameraNode.h b/Engine/Plugins/Cameras/GameplayCameras/Source/GameplayCameras/Public/Nodes/Blends/PopBlendCameraNode.h
--- a/Engine/Plugins/Cameras/GameplayCameras/Source/GameplayCameras/Public/Nodes/Blends/PopBlendCameraNode.h
+++ b/Engine/Plugins/Cameras/GameplayCameras/Source/GameplayCameras/Public/Nodes/Blends/PopBlendCameraNode.h
@@ -6,6 +6,25 @@
 
 #include "PopBlendCameraNode.generated.h"
 
+namespace UE::Cameras
+{
+
+struct FCameraNodeBlendParams;
+struct FCameraNodeBlendResult;
+struct FCameraNodePreBlendParams;
+struct FCameraNodePreBlendResult;
+
+/**
+ * Utility function for cutting evaluation results without blending.
+ */
+struct FPopBlendCameraNodeHelper
+{
+	static void PopParameters(const UE::Cameras::FCameraNodePreBlendParams& Params, FCameraNodePreBlendResult& OutResult);
+	static void PopResults(const UE::Cameras::FCameraNodeBlendParams& Params, FCameraNodeBlendResult& OutResult);
+};
+
+}  // namespace UE::Cameras
+
 /**
  * A blend node that creates a camera cut (i.e. it doesn't blend at all).
  */
diff --git a/Engine/Plugins/Cameras/GameplayCameras/Source/GameplayCameras/Public/Nodes/Blends/SimpleBlendCameraNode.h b/Engine/Plugins/Cameras/GameplayCameras/Source/GameplayCameras/Public/Nodes/Blends/SimpleBlendCameraNode.h
--- a/Engine/Plugins/Cameras/GameplayCameras/Source/GameplayCameras/Public/Nodes/Blends/SimpleBlendCameraNode.h
+++ b/Engine/Plugins/Cameras/GameplayCameras/Source/GameplayCameras/Public/Nodes/Blends/SimpleBlendCameraNode.h
@@ -50,6 +50,17 @@ public:
 	/** Gets the last evaluated blend factor. */
 	float GetBlendFactor() const { return BlendFactor; }
 
+	/** Gets whether the blend is currently at 100%. */
+	bool IsBlendFull() const { return BlendFactor >= 1.f; }
+
+	/** 
+	 * Gets whether the blend was flagged as finished.
+	 * A simple blend typically finishes when it reaches 100%, but in rare cases it may
+	 * want to continue running after that, possibly going back below 100% before going
+	 * back up.
+	 */
+	bool IsBlendFinished() const { return bIsBlendFinished; }
+
 protected:
 
 	GAMEPLAYCAMERAS_API virtual void OnRun(const FCameraNodeEvaluationParams& Params, FCameraNodeEvaluationResult& OutResult) override;
