Cameras: fix initial orientation service not correctly setting yaw/pitch anymore

diff --git a/Engine/Plugins/Cameras/GameplayCameras/Source/GameplayCameras/Private/Services/OrientationInitializationService.cpp b/Engine/Plugins/Cameras/GameplayCameras/Source/GameplayCameras/Private/Services/OrientationInitializationService.cpp
--- a/Engine/Plugins/Cameras/GameplayCameras/Source/GameplayCameras/Private/Services/OrientationInitializationService.cpp
+++ b/Engine/Plugins/Cameras/GameplayCameras/Source/GameplayCameras/Private/Services/OrientationInitializationService.cpp
@@ -166,7 +166,9 @@ void FOrientationInitializationService::TryInitializeYawPitch(const FCameraRigEv
 	{
 		Operation.Pitch = FConsumableDouble::Absolute(Pitch.GetValue());
 	}
-	CameraRigInfo.RootEvaluator->ExecuteOperation(OperationParams, Operation);
+
+	FCameraNodeEvaluatorHierarchy CameraRigHierarchy(CameraRigInfo.RootEvaluator);
+	CameraRigHierarchy.CallExecuteOperation(OperationParams, Operation);
 }
 
 void FOrientationInitializationService::TryPreserveTarget(const FCameraRigEvaluationInfo& CameraRigInfo, bool bUseRelativeTarget)
