Cameras: shorten the cameras debugger tab's title to prevent an ugly truncation of it by the editor

diff --git a/Engine/Plugins/Cameras/GameplayCameras/Source/GameplayCamerasEditor/Private/Debugger/SGameplayCamerasDebugger.cpp b/Engine/Plugins/Cameras/GameplayCameras/Source/GameplayCamerasEditor/Private/Debugger/SGameplayCamerasDebugger.cpp
--- a/Engine/Plugins/Cameras/GameplayCameras/Source/GameplayCamerasEditor/Private/Debugger/SGameplayCamerasDebugger.cpp
+++ b/Engine/Plugins/Cameras/GameplayCameras/Source/GameplayCamerasEditor/Private/Debugger/SGameplayCamerasDebugger.cpp
@@ -46,8 +46,8 @@ void SGameplayCamerasDebugger::RegisterTabSpawners()
 		SGameplayCamerasDebugger::WindowName,
 		FOnSpawnTab::CreateStatic(&SGameplayCamerasDebugger::SpawnGameplayCamerasDebugger)
 	)
-	.SetDisplayName(LOCTEXT("TabDisplayName", "Cameras Debugger"))
-	.SetTooltipText(LOCTEXT("TabTooltipText", "Open the Cameras Debugger tab."))
+	.SetDisplayName(LOCTEXT("TabDisplayName", "Camera Debugger"))
+	.SetTooltipText(LOCTEXT("TabTooltipText", "Open the Gameplay Cameras Debugger tab."))
 	.SetIcon(FSlateIcon(CamerasEditorStyle->GetStyleSetName(), "Debugger.TabIcon"))
 	.SetGroup(WorkspaceMenu::GetMenuStructure().GetDeveloperToolsDebugCategory())
 	.SetCanSidebarTab(false);
@@ -65,7 +65,7 @@ TSharedRef<SDockTab> SGameplayCamerasDebugger::SpawnGameplayCamerasDebugger(cons
 {
 	auto NomadTab = SNew(SDockTab)
 		.TabRole(ETabRole::NomadTab)
-		.Label(LOCTEXT("TabTitle", "Cameras Debugger"));
+		.Label(LOCTEXT("TabTitle", "Camera Debugger"));
 
 	TSharedRef<SWidget> MainWidget = SNew(SGameplayCamerasDebugger);
 	NomadTab->SetContent(MainWidget);
