Cameras: fix LOCTEXT conflicts

diff --git a/Engine/Plugins/Cameras/GameplayCameras/Source/GameplayCamerasEditor/Private/Toolkits/CameraRigInterfaceParametersToolkit.cpp b/Engine/Plugins/Cameras/GameplayCameras/Source/GameplayCamerasEditor/Private/Toolkits/CameraRigInterfaceParametersToolkit.cpp
--- a/Engine/Plugins/Cameras/GameplayCameras/Source/GameplayCamerasEditor/Private/Toolkits/CameraRigInterfaceParametersToolkit.cpp
+++ b/Engine/Plugins/Cameras/GameplayCameras/Source/GameplayCamerasEditor/Private/Toolkits/CameraRigInterfaceParametersToolkit.cpp
@@ -101,7 +101,7 @@ protected:
 
 	void OnParameterNameTextCommitted(const FText& Text, ETextCommit::Type CommitType)
 	{
-		const FScopedTransaction Transaction(LOCTEXT("RenameInterfaceParameter", "Rename Interface Parameter"));
+		const FScopedTransaction Transaction(LOCTEXT("RenameInterfaceParameterTransaction", "Rename Interface Parameter"));
 
 		Item->Modify();
 		Item->InterfaceParameterName = Text.ToString();
@@ -205,8 +205,8 @@ protected:
 		TypeTree.Add(MakeShared<FPinTypeTreeInfo>(UEdGraphSchema_K2::PC_Struct, TBaseStructure<FVector4>::Get(), LOCTEXT("Vector4dType", "A double precision 4D vector")));
 		TypeTree.Add(MakeShared<FPinTypeTreeInfo>(UEdGraphSchema_K2::PC_Struct, TVariantStructure<FRotator3f>::Get(), LOCTEXT("Rotator3fType", "A 3D rotation")));
 		TypeTree.Add(MakeShared<FPinTypeTreeInfo>(UEdGraphSchema_K2::PC_Struct, TBaseStructure<FRotator>::Get(), LOCTEXT("Rotator3dType", "A double precision 3D rotation")));
-		TypeTree.Add(MakeShared<FPinTypeTreeInfo>(UEdGraphSchema_K2::PC_Struct, TVariantStructure<FTransform3f>::Get(), LOCTEXT("TransformType", "A 3D transformation")));
-		TypeTree.Add(MakeShared<FPinTypeTreeInfo>(UEdGraphSchema_K2::PC_Struct, TBaseStructure<FTransform>::Get(), LOCTEXT("TransformType", "A double precision 3D transformation")));
+		TypeTree.Add(MakeShared<FPinTypeTreeInfo>(UEdGraphSchema_K2::PC_Struct, TVariantStructure<FTransform3f>::Get(), LOCTEXT("Transform3fType", "A 3D transformation")));
+		TypeTree.Add(MakeShared<FPinTypeTreeInfo>(UEdGraphSchema_K2::PC_Struct, TBaseStructure<FTransform>::Get(), LOCTEXT("Transform3dType", "A double precision 3D transformation")));
 	}
 
 	FEdGraphPinType GetBlendableParameterPinType() const
