Gameplay Cameras: Initialize

diff --git a/Engine/Plugins/Cameras/GameplayCameras/Source/GameplayCameras/Public/Core/CameraContextDataAllocationInfo.h b/Engine/Plugins/Cameras/GameplayCameras/Source/GameplayCameras/Public/Core/CameraContextDataAllocationInfo.h
--- a/Engine/Plugins/Cameras/GameplayCameras/Source/GameplayCameras/Public/Core/CameraContextDataAllocationInfo.h
+++ b/Engine/Plugins/Cameras/GameplayCameras/Source/GameplayCameras/Public/Core/CameraContextDataAllocationInfo.h
@@ -25,7 +25,7 @@ struct FCameraContextDataDefinition
 
 	/** The type of the data. */
 	UPROPERTY()
-	ECameraContextDataType DataType;
+	ECameraContextDataType DataType = ECameraContextDataType::Name;
 
 	/** An extra type object for the data. */
 	UPROPERTY()
diff --git a/Engine/Plugins/Cameras/GameplayCameras/Source/GameplayCameras/Public/Core/ICustomCameraNodeParameterProvider.h b/Engine/Plugins/Cameras/GameplayCameras/Source/GameplayCameras/Public/Core/ICustomCameraNodeParameterProvider.h
--- a/Engine/Plugins/Cameras/GameplayCameras/Source/GameplayCameras/Public/Core/ICustomCameraNodeParameterProvider.h
+++ b/Engine/Plugins/Cameras/GameplayCameras/Source/GameplayCameras/Public/Core/ICustomCameraNodeParameterProvider.h
@@ -30,7 +30,7 @@ struct FCustomCameraNodeBlendableParameter
 
 	/** The type of the parameter. */
 	UPROPERTY()
-	ECameraVariableType ParameterType;
+	ECameraVariableType ParameterType = ECameraVariableType::Boolean;
 
 	/** An optional camera variable that is dynamically driving the parameter's value. */
 	UPROPERTY()
@@ -51,7 +51,7 @@ struct FCustomCameraNodeDataParameter
 
 	/** The type of the parameter. */
 	UPROPERTY()
-	ECameraContextDataType ParameterType;
+	ECameraContextDataType ParameterType = ECameraContextDataType::Name;
 
 	/** An extra type object for the parameter. */
 	UPROPERTY()
