Exposed default post-process settings getter.

diff --git a/Engine/Plugins/Cameras/GameplayCameras/Source/GameplayCameras/Private/Core/PostProcessSettingsCollection.cpp b/Engine/Plugins/Cameras/GameplayCameras/Source/GameplayCameras/Private/Core/PostProcessSettingsCollection.cpp
--- a/Engine/Plugins/Cameras/GameplayCameras/Source/GameplayCameras/Private/Core/PostProcessSettingsCollection.cpp
+++ b/Engine/Plugins/Cameras/GameplayCameras/Source/GameplayCameras/Private/Core/PostProcessSettingsCollection.cpp
@@ -9,9 +9,7 @@ namespace UE::Cameras
 
 void FPostProcessSettingsCollection::Reset()
 {
-	static const FPostProcessSettings DefaultPostProcessSettings;
-
-	PostProcessSettings = DefaultPostProcessSettings;
+	PostProcessSettings = FPostProcessSettings::GetDefault();
 	bHasAnySetting = false;
 }
 
@@ -56,10 +54,8 @@ void FPostProcessSettingsCollection::InternalLerp(const FPostProcessSettings& To
 
 void FPostProcessSettingsCollection::Serialize(FArchive& Ar)
 {
-	static const FPostProcessSettings DefaultPostProcessSettings;
-
 	UScriptStruct* PostProcessSettingsStruct = FPostProcessSettings::StaticStruct();
-	PostProcessSettingsStruct->SerializeItem(Ar, &PostProcessSettings, &DefaultPostProcessSettings);
+	PostProcessSettingsStruct->SerializeItem(Ar, &PostProcessSettings, &FPostProcessSettings::GetDefault());
 
 	Ar << bHasAnySetting;
 }
diff --git a/Engine/Source/Runtime/Engine/Classes/Engine/Scene.h b/Engine/Source/Runtime/Engine/Classes/Engine/Scene.h
--- a/Engine/Source/Runtime/Engine/Classes/Engine/Scene.h
+++ b/Engine/Source/Runtime/Engine/Classes/Engine/Scene.h
@@ -2458,6 +2458,9 @@ PRAGMA_ENABLE_DEPRECATION_WARNINGS
 
 	// Default number of blade of the diaphragm to simulate in depth of field.
 	static constexpr int32 kDefaultDepthOfFieldBladeCount = 5;
+
+	// Get the default post-process settings
+	static ENGINE_API const FPostProcessSettings& GetDefault();
 	
 #if WITH_EDITORONLY_DATA
 	bool Serialize(FArchive& Ar);
diff --git a/Engine/Source/Runtime/Engine/Private/Scene.cpp b/Engine/Source/Runtime/Engine/Private/Scene.cpp
--- a/Engine/Source/Runtime/Engine/Private/Scene.cpp
+++ b/Engine/Source/Runtime/Engine/Private/Scene.cpp
@@ -920,6 +920,13 @@ void FPostProcessSettings::PostSerialize(const FArchive& Ar)
 }
 #endif
 
+const FPostProcessSettings& FPostProcessSettings::GetDefault()
+{
+	static const FPostProcessSettings DefaultSettings;
+
+	return DefaultSettings;
+}
+
 UScene::UScene(const FObjectInitializer& ObjectInitializer)
 	: Super(ObjectInitializer)
 {
diff --git a/Engine/Source/Runtime/Renderer/Private/PostProcess/PostProcessCombineLUTs.cpp b/Engine/Source/Runtime/Renderer/Private/PostProcess/PostProcessCombineLUTs.cpp
--- a/Engine/Source/Runtime/Renderer/Private/PostProcess/PostProcessCombineLUTs.cpp
+++ b/Engine/Source/Runtime/Renderer/Private/PostProcess/PostProcessCombineLUTs.cpp
@@ -212,13 +212,11 @@ struct FCachedLUTSettings
 		check(Textures);
 		check(Weights);
 
-		static const FPostProcessSettings DefaultSettings;
-
 		const FSceneViewFamily& ViewFamily = *(View.Family);
 		const FPostProcessSettings& Settings = (
 			ViewFamily.EngineShowFlags.ColorGrading && CVarColorGrading.GetValueOnRenderThread() != 0)
 			? View.FinalPostProcessSettings
-			: DefaultSettings;
+			: FPostProcessSettings::GetDefault();
 
 		for (uint32 BlendIndex = 0; BlendIndex < BlendCount; ++BlendIndex)
 		{
