Cameras: correctly free Rewind Debugger data and unregister debug delegate on clear

diff --git a/Engine/Plugins/Cameras/GameplayCameras/Source/GameplayCamerasEditor/Private/Trace/CameraSystemRewindDebuggerExtension.cpp b/Engine/Plugins/Cameras/GameplayCameras/Source/GameplayCamerasEditor/Private/Trace/CameraSystemRewindDebuggerExtension.cpp
--- a/Engine/Plugins/Cameras/GameplayCameras/Source/GameplayCamerasEditor/Private/Trace/CameraSystemRewindDebuggerExtension.cpp
+++ b/Engine/Plugins/Cameras/GameplayCameras/Source/GameplayCamerasEditor/Private/Trace/CameraSystemRewindDebuggerExtension.cpp
@@ -84,6 +84,16 @@ void FCameraSystemRewindDebuggerExtension::RecordingStopped(IRewindDebugger* Rew
 	UE::Trace::ToggleChannel(*FCameraSystemTrace::ChannelName, false);
 }
 
+void FCameraSystemRewindDebuggerExtension::Clear(IRewindDebugger* RewindDebugger)
+{
+	EnsureDebugDrawDelegate(false);
+
+	VisualizedWorld = nullptr;
+
+	DebugBlockStorage.DestroyDebugBlocks(true);
+	RootDebugBlock = nullptr;
+}
+
 void FCameraSystemRewindDebuggerExtension::EnsureDebugDrawDelegate(bool bIsRegistered)
 {
 	if (bIsRegistered && !DebugDrawDelegateHandle.IsValid())
diff --git a/Engine/Plugins/Cameras/GameplayCameras/Source/GameplayCamerasEditor/Private/Trace/CameraSystemRewindDebuggerExtension.h b/Engine/Plugins/Cameras/GameplayCameras/Source/GameplayCamerasEditor/Private/Trace/CameraSystemRewindDebuggerExtension.h
--- a/Engine/Plugins/Cameras/GameplayCameras/Source/GameplayCamerasEditor/Private/Trace/CameraSystemRewindDebuggerExtension.h
+++ b/Engine/Plugins/Cameras/GameplayCameras/Source/GameplayCamerasEditor/Private/Trace/CameraSystemRewindDebuggerExtension.h
@@ -17,7 +17,6 @@ class UCanvas;
 namespace UE::Cameras
 {
 
-class FDebugCanvasDeferredRenderer;
 class FRootCameraDebugBlock;
 
 /**
@@ -30,9 +29,11 @@ public:
 	FCameraSystemRewindDebuggerExtension();
 	virtual ~FCameraSystemRewindDebuggerExtension();
 
+	// IRewindDebuggerExtension interface.
 	virtual void RecordingStarted(IRewindDebugger* RewindDebugger) override;
 	virtual void Update(float DeltaTime, IRewindDebugger* RewindDebugger) override;
 	virtual void RecordingStopped(IRewindDebugger* RewindDebugger) override;
+	virtual void Clear(IRewindDebugger* RewindDebugger) override;
 
 private:
 
@@ -48,8 +49,6 @@ private:
 
 	FCameraDebugBlockStorage DebugBlockStorage;
 	FRootCameraDebugBlock* RootDebugBlock = nullptr;
-
-	TSharedPtr<FDebugCanvasDeferredRenderer> DeferredRenderer;
 };
 
 }  // namespace UE::Cameras
