Cameras: add support for forcing activation and/or transition to use on activation

diff --git a/Engine/Plugins/Cameras/GameplayCameras/Source/GameplayCameras/Private/Core/BlendStackCameraNode.cpp b/Engine/Plugins/Cameras/GameplayCameras/Source/GameplayCameras/Private/Core/BlendStackCameraNode.cpp
--- a/Engine/Plugins/Cameras/GameplayCameras/Source/GameplayCameras/Private/Core/BlendStackCameraNode.cpp
+++ b/Engine/Plugins/Cameras/GameplayCameras/Source/GameplayCameras/Private/Core/BlendStackCameraNode.cpp
@@ -848,6 +848,12 @@ void FTransientBlendStackCameraNodeEvaluator::InternalPostBlendExecute(TArrayVie
 
 const UCameraRigTransition* FTransientBlendStackCameraNodeEvaluator::FindTransition(const FBlendStackCameraPushParams& Params) const
 {
+	// If we are forced to use a specific transition, our search is over.
+	if (Params.TransitionOverride)
+	{
+		return Params.TransitionOverride;
+	}
+
 	// Find a transition that works for blending towards ToCameraRig.
 	// If the stack isn't empty, we need to find a transition that works between the previous and 
 	// next camera rigs. If the stack is empty, we blend the new camera rig in from nothing if
diff --git a/Engine/Plugins/Cameras/GameplayCameras/Source/GameplayCameras/Private/Core/CameraDirectorEvaluator.cpp b/Engine/Plugins/Cameras/GameplayCameras/Source/GameplayCameras/Private/Core/CameraDirectorEvaluator.cpp
--- a/Engine/Plugins/Cameras/GameplayCameras/Source/GameplayCameras/Private/Core/CameraDirectorEvaluator.cpp
+++ b/Engine/Plugins/Cameras/GameplayCameras/Source/GameplayCameras/Private/Core/CameraDirectorEvaluator.cpp
@@ -4,6 +4,7 @@
 
 #include "Core/CameraDirector.h"
 #include "Core/CameraEvaluationContext.h"
+#include "Core/CameraRigTransition.h"
 
 namespace UE::Cameras
 {
@@ -48,6 +49,35 @@ void FCameraDirectorEvaluator::Deactivate(const FCameraDirectorDeactivateParams&
 void FCameraDirectorEvaluator::Run(const FCameraDirectorEvaluationParams& Params, FCameraDirectorEvaluationResult& OutResult)
 {
 	OnRun(Params, OutResult);
+
+	if (NextActivationTransitionOverride)
+	{
+		OutResult.TransitionOverride = NextActivationTransitionOverride;
+	}
+	NextActivationTransitionOverride = nullptr;
+
+	if (bNextActivationForce)
+	{
+		OutResult.bForceActivateCameraRigs = true;
+	}
+	bNextActivationForce = false;
+}
+
+void FCameraDirectorEvaluator::OnEndCameraSystemUpdate()
+{
+	// Also clear these here, in case we didn't run this frame.
+	NextActivationTransitionOverride = nullptr;
+	bNextActivationForce = false;
+}
+
+void FCameraDirectorEvaluator::OverrideNextActivationTransition(const UCameraRigTransition* TransitionOverride)
+{
+	NextActivationTransitionOverride = TransitionOverride;
+}
+
+void FCameraDirectorEvaluator::ForceNextActivation()
+{
+	bNextActivationForce = true;
 }
 
 bool FCameraDirectorEvaluator::AddChildEvaluationContext(TSharedRef<FCameraEvaluationContext> InContext)
diff --git a/Engine/Plugins/Cameras/GameplayCameras/Source/GameplayCameras/Private/Core/CameraEvaluationContext.cpp b/Engine/Plugins/Cameras/GameplayCameras/Source/GameplayCameras/Private/Core/CameraEvaluationContext.cpp
--- a/Engine/Plugins/Cameras/GameplayCameras/Source/GameplayCameras/Private/Core/CameraEvaluationContext.cpp
+++ b/Engine/Plugins/Cameras/GameplayCameras/Source/GameplayCameras/Private/Core/CameraEvaluationContext.cpp
@@ -87,6 +87,16 @@ void FCameraEvaluationContext::OnEndCameraSystemUpdate()
 		Result.VariableTable.ClearAllWrittenThisFrameFlags();
 		Result.ContextDataTable.ClearAllWrittenThisFrameFlags();
 	}
+
+	if (DirectorEvaluator)
+	{
+		DirectorEvaluator->OnEndCameraSystemUpdate();
+	}
+
+	for (TSharedPtr<FCameraEvaluationContext> ChildContext : ChildrenContexts)
+	{
+		ChildContext->OnEndCameraSystemUpdate();
+	}
 }
 
 void FCameraEvaluationContext::AutoCreateDirectorEvaluator()
diff --git a/Engine/Plugins/Cameras/GameplayCameras/Source/GameplayCameras/Private/Core/CameraSystemEvaluator.cpp b/Engine/Plugins/Cameras/GameplayCameras/Source/GameplayCameras/Private/Core/CameraSystemEvaluator.cpp
--- a/Engine/Plugins/Cameras/GameplayCameras/Source/GameplayCameras/Private/Core/CameraSystemEvaluator.cpp
+++ b/Engine/Plugins/Cameras/GameplayCameras/Source/GameplayCameras/Private/Core/CameraSystemEvaluator.cpp
@@ -234,6 +234,7 @@ void FCameraSystemEvaluator::Update(const FCameraSystemEvaluationParams& Params)
 			FActivateCameraRigParams CameraRigParams;
 			CameraRigParams.EvaluationContext = ActiveCameraRig.EvaluationContext;
 			CameraRigParams.CameraRig = ActiveCameraRig.CameraRig;
+			CameraRigParams.TransitionOverride = DirectorResult.TransitionOverride;
 			CameraRigParams.bForceActivate = DirectorResult.bForceActivateCameraRigs;
 			RootEvaluator->ActivateCameraRig(CameraRigParams);
 		}
@@ -266,6 +267,7 @@ void FCameraSystemEvaluator::Update(const FCameraSystemEvaluationParams& Params)
 			FActivateCameraRigParams CameraRigParams;
 			CameraRigParams.EvaluationContext = CommonContext;
 			CameraRigParams.CameraRig = CombinedCameraRig;
+			CameraRigParams.TransitionOverride = DirectorResult.TransitionOverride;
 			CameraRigParams.bForceActivate = DirectorResult.bForceActivateCameraRigs;
 			RootEvaluator->ActivateCameraRig(CameraRigParams);
 		}
diff --git a/Engine/Plugins/Cameras/GameplayCameras/Source/GameplayCameras/Private/Core/DefaultRootCameraNode.cpp b/Engine/Plugins/Cameras/GameplayCameras/Source/GameplayCameras/Private/Core/DefaultRootCameraNode.cpp
--- a/Engine/Plugins/Cameras/GameplayCameras/Source/GameplayCameras/Private/Core/DefaultRootCameraNode.cpp
+++ b/Engine/Plugins/Cameras/GameplayCameras/Source/GameplayCameras/Private/Core/DefaultRootCameraNode.cpp
@@ -86,6 +86,7 @@ void FDefaultRootCameraNodeEvaluator::OnActivateCameraRig(const FActivateCameraR
 		FBlendStackCameraPushParams PushParams;
 		PushParams.EvaluationContext = Params.EvaluationContext;
 		PushParams.CameraRig = Params.CameraRig;
+		PushParams.TransitionOverride = Params.TransitionOverride;
 		PushParams.bForcePush = Params.bForceActivate;
 		MainLayer->Push(PushParams);
 	}
@@ -109,6 +110,7 @@ void FDefaultRootCameraNodeEvaluator::OnActivateCameraRig(const FActivateCameraR
 			FBlendStackCameraInsertParams InsertParams;
 			InsertParams.EvaluationContext = Params.EvaluationContext;
 			InsertParams.CameraRig = Params.CameraRig;
+			InsertParams.TransitionOverride = Params.TransitionOverride;
 			InsertParams.bForceInsert = Params.bForceActivate;
 			TargetLayer->Insert(InsertParams);
 		}
diff --git a/Engine/Plugins/Cameras/GameplayCameras/Source/GameplayCameras/Public/Core/BlendStackCameraNode.h b/Engine/Plugins/Cameras/GameplayCameras/Source/GameplayCameras/Public/Core/BlendStackCameraNode.h
--- a/Engine/Plugins/Cameras/GameplayCameras/Source/GameplayCameras/Public/Core/BlendStackCameraNode.h
+++ b/Engine/Plugins/Cameras/GameplayCameras/Source/GameplayCameras/Public/Core/BlendStackCameraNode.h
@@ -239,6 +239,9 @@ struct FBlendStackCameraPushParams
 	/** The source camera rig asset to instantiate and push on the blend stack. */
 	TObjectPtr<const UCameraRigAsset> CameraRig;
 
+	/** A transition to use, instead of looking one up. */
+	TObjectPtr<const UCameraRigTransition> TransitionOverride;
+
 	/** Whether to force pushing a new instance of the camera rig, even if it is currently active. */
 	bool bForcePush = false;
 };
@@ -311,6 +314,9 @@ struct FBlendStackCameraInsertParams
 	/** The source camera rig asset to instantiate and push on the blend stack. */
 	TObjectPtr<const UCameraRigAsset> CameraRig;
 
+	/** A transition to use, instead of looking one up. */
+	TObjectPtr<const UCameraRigTransition> TransitionOverride;
+
 	/** Whether to force insert a new instance of the camera rig, even if there is already one in the stack. */
 	bool bForceInsert = false;
 };
diff --git a/Engine/Plugins/Cameras/GameplayCameras/Source/GameplayCameras/Public/Core/CameraDirectorEvaluator.h b/Engine/Plugins/Cameras/GameplayCameras/Source/GameplayCameras/Public/Core/CameraDirectorEvaluator.h
--- a/Engine/Plugins/Cameras/GameplayCameras/Source/GameplayCameras/Public/Core/CameraDirectorEvaluator.h
+++ b/Engine/Plugins/Cameras/GameplayCameras/Source/GameplayCameras/Public/Core/CameraDirectorEvaluator.h
@@ -10,6 +10,7 @@
 
 class UCameraDirector;
 class UCameraRigAsset;
+class UCameraRigTransition;
 
 namespace UE::Cameras
 {
@@ -85,6 +86,9 @@ struct FCameraDirectorEvaluationResult
 	/** The camera rig(s) that the director says should be active this frame in the main layer. */
 	FActiveCameraRigInfos ActiveCameraRigs;
 
+	/** A transition to use for the next camera rig activation, instead of looking up one according to usual rules. */
+	TObjectPtr<const UCameraRigTransition> TransitionOverride;
+
 	/** Whether to force creating a new instance of the camera rigs, even if they are already active in the main layer. */
 	bool bForceActivateCameraRigs = false;
 
@@ -155,9 +159,6 @@ public:
 
 	/** Deactivates the camera director evaluator. */
 	GAMEPLAYCAMERAS_API void Deactivate(const FCameraDirectorDeactivateParams& Params);
-	
-	/** Runs the camera director to determine what camera rig(s) should be active this frame. */
-	GAMEPLAYCAMERAS_API void Run(const FCameraDirectorEvaluationParams& Params, FCameraDirectorEvaluationResult& OutResult);
 
 	/** Gets the camera director. */
 	const UCameraDirector* GetCameraDirector() const { return PrivateCameraDirector; }
@@ -169,6 +170,19 @@ public:
 		return Cast<CameraDirectorType>(PrivateCameraDirector);
 	}
 
+public:
+
+	/** Runs the camera director to determine what camera rig(s) should be active this frame. */
+	GAMEPLAYCAMERAS_API void Run(const FCameraDirectorEvaluationParams& Params, FCameraDirectorEvaluationResult& OutResult);
+
+	/** Request that the next camera rig activation use the provided transition. */
+	GAMEPLAYCAMERAS_API void OverrideNextActivationTransition(const UCameraRigTransition* TransitionOverride);
+
+	/** Request that the next camera rig activation be forced. */
+	GAMEPLAYCAMERAS_API void ForceNextActivation();
+
+public:
+
 	GAMEPLAYCAMERAS_API bool AddChildEvaluationContext(TSharedRef<FCameraEvaluationContext> InContext);
 	GAMEPLAYCAMERAS_API bool RemoveChildEvaluationContext(TSharedRef<FCameraEvaluationContext> InContext);
 
@@ -179,6 +193,8 @@ public:
 	// Internal API.
 	void SetPrivateCameraDirector(const UCameraDirector* InCameraDirector);
 
+	void OnEndCameraSystemUpdate();
+
 protected:
 
 	enum class EChildContextManipulationResult
@@ -224,6 +240,12 @@ private:
 
 	/** The camera director this evaluator is running. */
 	TObjectPtr<const UCameraDirector> PrivateCameraDirector;
+
+	/** A forced transition to use on the next update. Cleared after every system update. */
+	TObjectPtr<const UCameraRigTransition> NextActivationTransitionOverride;
+
+	/** Whether to force camera rig activation on the next update. Cleared after every system update. */
+	bool bNextActivationForce = false;
 };
 
 template<typename EvaluatorType, typename ...ArgTypes>
diff --git a/Engine/Plugins/Cameras/GameplayCameras/Source/GameplayCameras/Public/Core/RootCameraNode.h b/Engine/Plugins/Cameras/GameplayCameras/Source/GameplayCameras/Public/Core/RootCameraNode.h
--- a/Engine/Plugins/Cameras/GameplayCameras/Source/GameplayCameras/Public/Core/RootCameraNode.h
+++ b/Engine/Plugins/Cameras/GameplayCameras/Source/GameplayCameras/Public/Core/RootCameraNode.h
@@ -9,6 +9,7 @@
 #include "RootCameraNode.generated.h"
 
 class UCameraRigAsset;
+class UCameraRigTransition;
 
 /**
  * Defines evaluation layers for camera rigs.
@@ -51,6 +52,9 @@ struct FActivateCameraRigParams
 	/** The source camera rig asset that will be instantiated. */
 	TObjectPtr<const UCameraRigAsset> CameraRig;
 
+	/** A transition to use, instead of looking one up. */
+	TObjectPtr<const UCameraRigTransition> TransitionOverride;
+
 	/** The evaluation layer on which to instantiate the camera rig. */
 	ECameraRigLayer Layer = ECameraRigLayer::Main;
 
