Cameras: make the GameplayCameraComponent call AddReferencedObjects on its evaluation context

If the evaluation context is on the stack, it would already be handled by the camera system, but this covers the case when it is not on the stack any more.

diff --git a/Engine/Plugins/Cameras/GameplayCameras/Source/GameplayCameras/Private/GameFramework/GameplayCameraComponent.cpp b/Engine/Plugins/Cameras/GameplayCameras/Source/GameplayCameras/Private/GameFramework/GameplayCameraComponent.cpp
--- a/Engine/Plugins/Cameras/GameplayCameras/Source/GameplayCameras/Private/GameFramework/GameplayCameraComponent.cpp
+++ b/Engine/Plugins/Cameras/GameplayCameras/Source/GameplayCameras/Private/GameFramework/GameplayCameraComponent.cpp
@@ -47,6 +47,22 @@ void UGameplayCameraComponent::PostLoad()
 	}
 }
 
+void UGameplayCameraComponent::AddReferencedObjects(UObject* InThis, FReferenceCollector& Collector)
+{
+	Super::AddReferencedObjects(InThis, Collector);
+
+	UGameplayCameraComponent* This = CastChecked<UGameplayCameraComponent>(InThis);
+	if (This->EvaluationContext.IsValid())
+	{
+		This->EvaluationContext->AddReferencedObjects(Collector);
+	}
+}
+
+TSharedPtr<const UE::Cameras::FCameraEvaluationContext> UGameplayCameraComponent::GetEvaluationContext() const
+{
+	return EvaluationContext;
+}
+
 TSharedPtr<UE::Cameras::FCameraEvaluationContext> UGameplayCameraComponent::GetEvaluationContext()
 {
 	return EvaluationContext;
diff --git a/Engine/Plugins/Cameras/GameplayCameras/Source/GameplayCameras/Public/GameFramework/GameplayCameraComponent.h b/Engine/Plugins/Cameras/GameplayCameras/Source/GameplayCameras/Public/GameFramework/GameplayCameraComponent.h
--- a/Engine/Plugins/Cameras/GameplayCameras/Source/GameplayCameras/Public/GameFramework/GameplayCameraComponent.h
+++ b/Engine/Plugins/Cameras/GameplayCameras/Source/GameplayCameras/Public/GameFramework/GameplayCameraComponent.h
@@ -37,6 +37,9 @@ public:
 	/** Create a new camera component. */
 	GAMEPLAYCAMERAS_API UGameplayCameraComponent(const FObjectInitializer& ObjectInit);
 
+	/** Get the camera evaluation context used by this component. */
+	GAMEPLAYCAMERAS_API TSharedPtr<const UE::Cameras::FCameraEvaluationContext> GetEvaluationContext() const;
+
 	/** Get the camera evaluation context used by this component. */
 	GAMEPLAYCAMERAS_API TSharedPtr<UE::Cameras::FCameraEvaluationContext> GetEvaluationContext();
 
@@ -95,6 +98,8 @@ public:
 	// UObject interface.
 	void PostLoad() override;
 
+	static void AddReferencedObjects(UObject* InThis, FReferenceCollector& Collector);
+
 private:
 
 	void ActivateCameraEvaluationContext(int32 PlayerIndex);
