Cameras: fix input parameter on SetCameraPose Blueprint function

diff --git a/Engine/Plugins/Cameras/GameplayCameras/Source/GameplayCameras/Public/GameFramework/BlueprintCameraNodeEvaluationResult.h b/Engine/Plugins/Cameras/GameplayCameras/Source/GameplayCameras/Public/GameFramework/BlueprintCameraNodeEvaluationResult.h
--- a/Engine/Plugins/Cameras/GameplayCameras/Source/GameplayCameras/Public/GameFramework/BlueprintCameraNodeEvaluationResult.h
+++ b/Engine/Plugins/Cameras/GameplayCameras/Source/GameplayCameras/Public/GameFramework/BlueprintCameraNodeEvaluationResult.h
@@ -76,6 +76,6 @@ public:
 
 	/** Sets the camera pose. */
 	UFUNCTION(BlueprintCallable, Category="Camera")
-	static void SetCameraPose(FBlueprintCameraNodeEvaluationResult& Data, const FBlueprintCameraPose& CameraPose);
+	static void SetCameraPose(UPARAM(Ref) FBlueprintCameraNodeEvaluationResult& Data, const FBlueprintCameraPose& CameraPose);
 };
 
