Cameras: add support for blendable structures in camera rig parameters, and simplify how blendable parameters are built

- Make a "parameter definition" struct that contains everything we need to more easily handle exposed parameters.
- Remove the concept of "private variables", and instead just keep an ID like with context data parameters.
- Make context data IDs like variable IDs and instead keep a "debug name" in the data definition, again like with camera variables and blendable parameters.
- Add an extra type of blendable parameter which is a "blendable struct", based on some blend function registered on the Gameplay Cameras module.
- Property bags built for holding parameter values now get their default values set correctly.
- Camera asset and camera rig asset references now reflect those default values, and support overriding with camera paremeters, for nested prefab rigs.
- Remove hand-written equality operators, use the compiler for that.

Data changes:
- Remove "private variables" set on camera parameters during loading.

[RN] ue-cameras

diff --git a/Engine/Plugins/Cameras/GameplayCameras/Source/GameplayCameras/Private/Core/BlendStackRootCameraNode.cpp b/Engine/Plugins/Cameras/GameplayCameras/Source/GameplayCameras/Private/Core/BlendStackRootCameraNode.cpp
--- a/Engine/Plugins/Cameras/GameplayCameras/Source/GameplayCameras/Private/Core/BlendStackRootCameraNode.cpp
+++ b/Engine/Plugins/Cameras/GameplayCameras/Source/GameplayCameras/Private/Core/BlendStackRootCameraNode.cpp
@@ -6,6 +6,7 @@
 #include "Core/BlendCameraNode.h"
 #include "Core/CameraRigAsset.h"
 #include "Core/CameraRigAssetReference.h"
+#include "Core/CameraVariableTableAllocationInfo.h"
 #include "Debug/CameraDebugBlock.h"
 #include "Debug/CameraDebugBlockBuilder.h"
 #include "Debug/CameraDebugRenderer.h"
@@ -187,13 +188,13 @@ void FBlendStackRootCameraNodeEvaluator::InitializeBlendedParameterOverridesStac
 		{
 			continue;
 		}
-		if (!BlendableParameter->PrivateVariable)
+		if (!BlendableParameter->PrivateVariableID)
 		{
 			continue;
 		}
 
-		BlendedParameterOverridesTableAllocationInfo.VariableDefinitions.Add(
-				BlendableParameter->PrivateVariable->GetVariableDefinition());
+		FCameraVariableDefinition Definition = BlendableParameter->GetVariableDefinition();
+		BlendedParameterOverridesTableAllocationInfo.VariableDefinitions.Add(Definition);
 	}
 
 	FCameraRigCameraNodeEvaluator* RootPrefabNodeEvaluator = RootEvaluator->CastThisChecked<FCameraRigCameraNodeEvaluator>();
@@ -225,23 +226,28 @@ void FBlendStackRootCameraNodeEvaluator::OnRun(const FCameraNodeEvaluationParams
 
 void FBlendStackRootCameraNodeEvaluator::SetDefaultBlendableParameterValues(FCameraVariableTable& OutVariableTable)
 {
-	for (const UCameraRigBlendableParameter* BlendableParameter : BlendablePrefabCameraRig->Interface.BlendableParameters)
+	const FInstancedPropertyBag& DefaultParameters = BlendablePrefabCameraRig->GetDefaultParameters();
+	const uint8* RawDefaultParametersContainer = DefaultParameters.GetValue().GetMemory();
+
+	for (const FCameraRigParameterDefinition& Definition : BlendablePrefabCameraRig->GetParameterDefinitions())
 	{
-		if (!ensure(BlendableParameter))
+		if (Definition.ParameterType != ECameraRigInterfaceParameterType::Blendable)
+		{
+			continue;
+		}
+		if (!Definition.VariableID.IsValid())
 		{
 			continue;
 		}
 
-		const UCameraVariableAsset* PrivateVariable = BlendableParameter->PrivateVariable;
-		if (!PrivateVariable)
+		const FPropertyBagPropertyDesc* PropertyDesc = DefaultParameters.FindPropertyDescByID(Definition.ParameterGuid);
+		if (!ensure(PropertyDesc && PropertyDesc->CachedProperty))
 		{
 			continue;
 		}
 
-		OutVariableTable.SetValue(
-				PrivateVariable->GetVariableID(), 
-				PrivateVariable->GetVariableType(), 
-				PrivateVariable->GetDefaultValuePtr());
+		const void* RawValuePtr = PropertyDesc->CachedProperty->ContainerPtrToValuePtr<void>(RawDefaultParametersContainer);
+		OutVariableTable.SetValue(Definition.VariableID, Definition.VariableType, Definition.BlendableStructType, (const uint8*)RawValuePtr);
 	}
 }
 
diff --git a/Engine/Plugins/Cameras/GameplayCameras/Source/GameplayCameras/Private/Core/CameraAssetBuilder.cpp b/Engine/Plugins/Cameras/GameplayCameras/Source/GameplayCameras/Private/Core/CameraAssetBuilder.cpp
--- a/Engine/Plugins/Cameras/GameplayCameras/Source/GameplayCameras/Private/Core/CameraAssetBuilder.cpp
+++ b/Engine/Plugins/Cameras/GameplayCameras/Source/GameplayCameras/Private/Core/CameraAssetBuilder.cpp
@@ -7,6 +7,7 @@
 #include "Core/CameraParameters.h"
 #include "Core/CameraRigAsset.h"
 #include "Core/CameraRigAssetBuilder.h"
+#include "Core/CameraRigParameterDefinition.h"
 #include "GameplayCamerasDelegates.h"
 #include "Logging/TokenizedMessage.h"
 
@@ -75,32 +76,45 @@ void FCameraAssetBuilder::BuildCameraImpl()
 		CameraRigBuilder.BuildCameraRig(CameraRig);
 	}
 
-	// Get the list of all the camera rigs' interface parameters, and rebuild our
-	// parameters property bag.
-	int32 NextParameterPropertyIndex = 0;
-	TArray<FPropertyBagPropertyDesc> DefaultParameterProperties;
-	TArray<TWeakObjectPtr<const UCameraRigAsset>> ParameterOwners;
+	// Get the list of all the camera rigs' interface parameters, and cache some information
+	// about them.
+	TArray<FCameraRigParameterDefinition> ParameterDefinitions;
+	TArray<TObjectPtr<const UCameraRigAsset>> ParameterOwners;
 
 	for (const UCameraRigAsset* CameraRig : CameraRigs)
 	{
-		FCameraRigAssetBuilder::AppendDefaultParameters(CameraRig->Interface, DefaultParameterProperties);
-
-		for (int32 Index = NextParameterPropertyIndex; Index < DefaultParameterProperties.Num(); ++Index)
+		for (const FCameraRigParameterDefinition& Definition : CameraRig->GetParameterDefinitions())
 		{
+			ParameterDefinitions.Add(Definition);
 			ParameterOwners.Add(CameraRig);
 		}
-		NextParameterPropertyIndex = DefaultParameterProperties.Num();
 	}
 
+	if (ParameterDefinitions != CameraAsset->ParameterDefinitions || ParameterOwners != CameraAsset->ParameterOwners)
+	{
+		CameraAsset->Modify();
+		CameraAsset->ParameterDefinitions = ParameterDefinitions;
+		CameraAsset->ParameterOwners = ParameterOwners;
+	}
+
+	// Get the list of all the camera rigs' interface parameters, and rebuild our
+	// parameters property bag.
+	TArray<FPropertyBagPropertyDesc> DefaultParameterProperties;
+	for (const UCameraRigAsset* CameraRig : CameraRigs)
+	{
+		FCameraRigParameterBuilder::AppendDefaultParameterProperties(CameraRig, DefaultParameterProperties);
+	}
 	FInstancedPropertyBag DefaultParameters;
 	DefaultParameters.AddProperties(DefaultParameterProperties);
+	for (const UCameraRigAsset* CameraRig : CameraRigs)
+	{
+		FCameraRigParameterBuilder::SetDefaultParameterValues(CameraRig, DefaultParameters);
+	}
 
-	if (DefaultParameters.GetPropertyBagStruct() != CameraAsset->DefaultParameters.GetPropertyBagStruct() ||
-			ParameterOwners != CameraAsset->ParameterOwners)
+	if (!DefaultParameters.Identical(&CameraAsset->DefaultParameters, 0))
 	{
 		CameraAsset->Modify();
 		CameraAsset->DefaultParameters = DefaultParameters;
-		CameraAsset->ParameterOwners = ParameterOwners;
 	}
 
 	// Accumulate all the camera rigs' allocation infos and store that on the asset.
diff --git a/Engine/Plugins/Cameras/GameplayCameras/Source/GameplayCameras/Private/Core/CameraAssetReference.cpp b/Engine/Plugins/Cameras/GameplayCameras/Source/GameplayCameras/Private/Core/CameraAssetReference.cpp
--- a/Engine/Plugins/Cameras/GameplayCameras/Source/GameplayCameras/Private/Core/CameraAssetReference.cpp
+++ b/Engine/Plugins/Cameras/GameplayCameras/Source/GameplayCameras/Private/Core/CameraAssetReference.cpp
@@ -24,12 +24,12 @@ void FCameraAssetReference::ApplyParameterOverrides(UE::Cameras::FCameraNodeEval
 	OverrideEvaluator.ApplyParameterOverrides(OutResult.VariableTable, OutResult.ContextDataTable, bDrivenOverridesOnly);
 }
 
-bool FCameraAssetReference::IsParameterOverriden(const FGuid PropertyID) const
+bool FCameraAssetReference::IsParameterOverridden(const FGuid PropertyID) const
 {
 	return ParameterOverrideGuids.Contains(PropertyID);
 }
 
-void FCameraAssetReference::SetParameterOverriden(const FGuid PropertyID, bool bIsOverridden)
+void FCameraAssetReference::SetParameterOverridden(const FGuid PropertyID, bool bIsOverridden)
 {
 	if (bIsOverridden)
 	{
diff --git a/Engine/Plugins/Cameras/GameplayCameras/Source/GameplayCameras/Private/Core/CameraContextDataAllocationInfo.cpp b/Engine/Plugins/Cameras/GameplayCameras/Source/GameplayCameras/Private/Core/CameraContextDataAllocationInfo.cpp
deleted file mode 100644
diff --git a/Engine/Plugins/Cameras/GameplayCameras/Source/GameplayCameras/Private/Core/CameraContextDataTable.cpp b/Engine/Plugins/Cameras/GameplayCameras/Source/GameplayCameras/Private/Core/CameraContextDataTable.cpp
--- a/Engine/Plugins/Cameras/GameplayCameras/Source/GameplayCameras/Private/Core/CameraContextDataTable.cpp
+++ b/Engine/Plugins/Cameras/GameplayCameras/Source/GameplayCameras/Private/Core/CameraContextDataTable.cpp
@@ -2,7 +2,7 @@
 
 #include "Core/CameraContextDataTable.h"
 
-#include "Core/CameraContextDataAllocationInfo.h"
+#include "Core/CameraContextDataTableAllocationInfo.h"
 
 namespace UE::Cameras
 {
@@ -46,7 +46,7 @@ void FCameraContextDataTable::AddReferencedObjects(FReferenceCollector& Referenc
 	}
 }
 
-void FCameraContextDataTable::Initialize(const FCameraContextDataAllocationInfo& AllocationInfo)
+void FCameraContextDataTable::Initialize(const FCameraContextDataTableAllocationInfo& AllocationInfo)
 {
 	// Reset any previous state.
 	DestroyBuffer();
@@ -71,7 +71,7 @@ void FCameraContextDataTable::Initialize(const FCameraContextDataAllocationInfo&
 		NewEntry.Offset = NewEntryOffset;
 		NewEntry.Flags = EEntryFlags::None;
 #if WITH_EDITORONLY_DATA
-		NewEntry.DebugName = DataDefinition.DataID.DataName;
+		NewEntry.DebugName = DataDefinition.DataName;
 #endif
 
 		Entries.Add(NewEntry);
@@ -120,7 +120,7 @@ void FCameraContextDataTable::AddData(const FCameraContextDataDefinition& DataDe
 	NewEntry.Offset = DataPtr - Memory;
 	NewEntry.Flags = EEntryFlags::None;
 #if WITH_EDITORONLY_DATA
-	NewEntry.DebugName = DataDefinition.DataID.DataName;
+	NewEntry.DebugName = DataDefinition.DataName;
 #endif
 	
 	Entries.Add(NewEntry);
@@ -150,8 +150,9 @@ bool FCameraContextDataTable::GetDataTypeAllocationInfo(ECameraContextDataType D
 				const UScriptStruct* StructType = CastChecked<const UScriptStruct>(DataTypeObject);
 				if (StructType)
 				{
-					OutSizeOf = StructType->PropertiesSize;
-					OutAlignOf = StructType->MinAlignment;
+					const UScriptStruct::ICppStructOps* StructOps = StructType->GetCppStructOps();
+					OutSizeOf = StructOps->GetSize();
+					OutAlignOf = StructOps->GetAlignment();
 				}
 			}
 			break;
@@ -417,8 +418,8 @@ const uint8* FCameraContextDataTable::GetData(FCameraContextDataID DataID, ECame
 	const uint8* Data = TryGetData(DataID, ExpectedDataType, ExpectedDataTypeObject);
 	ensureMsgf(
 			Data, 
-			TEXT("Can't get camera context data (ID '%s') because it doesn't exist in the table, or isn't of the expected data type."), 
-			*DataID.DataName.ToString());
+			TEXT("Can't get camera context data (ID '%d') because it doesn't exist in the table, or isn't of the expected data type."), 
+			DataID.GetValue());
 	return Data;
 }
 
@@ -439,7 +440,7 @@ const uint8* FCameraContextDataTable::TryGetData(FCameraContextDataID DataID, EC
 void FCameraContextDataTable::SetData(FCameraContextDataID DataID, ECameraContextDataType ExpectedDataType, const UObject* ExpectedDataTypeObject, const uint8* InRawDataPtr, bool bMarkAsWrittenThisFrame)
 {
 	const bool bDidSet = TrySetData(DataID, ExpectedDataType, ExpectedDataTypeObject, InRawDataPtr, bMarkAsWrittenThisFrame);
-	ensureMsgf(bDidSet, TEXT("Can't set camera context data (ID '%s') beacuse it doesn't exist in the table."), *DataID.DataName.ToString());
+	ensureMsgf(bDidSet, TEXT("Can't set camera context data (ID '%d') beacuse it doesn't exist in the table."), DataID.GetValue());
 }
 
 bool FCameraContextDataTable::TrySetData(FCameraContextDataID DataID, ECameraContextDataType ExpectedDataType, const UObject* ExpectedDataTypeObject, const uint8* InRawDataPtr, bool bMarkAsWrittenThisFrame)
diff --git a/Engine/Plugins/Cameras/GameplayCameras/Source/GameplayCameras/Private/Core/CameraContextDataTableAllocationInfo.cpp b/Engine/Plugins/Cameras/GameplayCameras/Source/GameplayCameras/Private/Core/CameraContextDataTableAllocationInfo.cpp
new file mode 100644
--- /dev/null
+++ b/Engine/Plugins/Cameras/GameplayCameras/Source/GameplayCameras/Private/Core/CameraContextDataTableAllocationInfo.cpp
@@ -0,0 +1,28 @@
+// Copyright Epic Games, Inc. All Rights Reserved.
+
+#include "Core/CameraContextDataTableAllocationInfo.h"
+
+void FCameraContextDataTableAllocationInfo::Combine(const FCameraContextDataTableAllocationInfo& OtherInfo)
+{
+	TMap<FCameraContextDataID, int32> KnownNames;
+	for (auto It = DataDefinitions.CreateConstIterator(); It; ++It)
+	{
+		const FCameraContextDataDefinition& DataDefinition(*It);
+		KnownNames.Add(DataDefinition.DataID, It.GetIndex());
+	}
+
+	for (const FCameraContextDataDefinition& OtherDataDefinition : OtherInfo.DataDefinitions)
+	{
+		const int32 KnownIndex = KnownNames.FindRef(OtherDataDefinition.DataID, INDEX_NONE);
+		if (KnownIndex == INDEX_NONE)
+		{
+			DataDefinitions.Add(OtherDataDefinition);
+		}
+		else
+		{
+			const FCameraContextDataDefinition& KnownDataDefinition(DataDefinitions[KnownIndex]);
+			ensure(KnownDataDefinition == OtherDataDefinition);
+		}
+	}
+}
+
diff --git a/Engine/Plugins/Cameras/GameplayCameras/Source/GameplayCameras/Private/Core/CameraEvaluationContext.cpp b/Engine/Plugins/Cameras/GameplayCameras/Source/GameplayCameras/Private/Core/CameraEvaluationContext.cpp
--- a/Engine/Plugins/Cameras/GameplayCameras/Source/GameplayCameras/Private/Core/CameraEvaluationContext.cpp
+++ b/Engine/Plugins/Cameras/GameplayCameras/Source/GameplayCameras/Private/Core/CameraEvaluationContext.cpp
@@ -161,7 +161,7 @@ void FCameraEvaluationContext::Activate(const FCameraEvaluationContextActivatePa
 
 	AutoCreateDirectorEvaluator();
 
-	if (ensure(DirectorEvaluator))
+	if (DirectorEvaluator)
 	{
 		FCameraDirectorActivateParams DirectorParams;
 		DirectorParams.Evaluator = Params.Evaluator;
@@ -179,7 +179,7 @@ void FCameraEvaluationContext::Deactivate(const FCameraEvaluationContextDeactiva
 		return;
 	}
 
-	if (ensure(DirectorEvaluator))
+	if (DirectorEvaluator)
 	{
 		FCameraDirectorDeactivateParams DirectorParams;
 		DirectorParams.OwnerContext = SharedThis(this);
diff --git a/Engine/Plugins/Cameras/GameplayCameras/Source/GameplayCameras/Private/Core/CameraParameters.cpp b/Engine/Plugins/Cameras/GameplayCameras/Source/GameplayCameras/Private/Core/CameraParameters.cpp
--- a/Engine/Plugins/Cameras/GameplayCameras/Source/GameplayCameras/Private/Core/CameraParameters.cpp
+++ b/Engine/Plugins/Cameras/GameplayCameras/Source/GameplayCameras/Private/Core/CameraParameters.cpp
@@ -2,6 +2,7 @@
 
 #include "Core/CameraParameters.h"
 
+#include "Core/CameraRigAsset.h"
 #include "Math/Rotator.h"
 #include "Math/Transform.h"
 #include "Math/Vector.h"
@@ -204,3 +205,17 @@ bool FTransform3dCameraParameter::SerializeFromMismatchedTag(const FPropertyTag&
 	return false;
 }
 
+#define UE_CAMERA_VARIABLE_FOR_TYPE(ValueType, ValueName)\
+void F##ValueName##CameraParameter::PostSerialize(const FArchive& Ar)\
+{\
+	if (Ar.IsLoading())\
+	{\
+		if (Variable && Variable->GetOuter()->IsA<UCameraRigAsset>())\
+		{\
+			Variable = nullptr;\
+		}\
+	}\
+}
+UE_CAMERA_VARIABLE_FOR_ALL_TYPES()
+#undef UE_CAMERA_VARIABLE_FOR_TYPE
+
diff --git a/Engine/Plugins/Cameras/GameplayCameras/Source/GameplayCameras/Private/Core/CameraRigAsset.cpp b/Engine/Plugins/Cameras/GameplayCameras/Source/GameplayCameras/Private/Core/CameraRigAsset.cpp
--- a/Engine/Plugins/Cameras/GameplayCameras/Source/GameplayCameras/Private/Core/CameraRigAsset.cpp
+++ b/Engine/Plugins/Cameras/GameplayCameras/Source/GameplayCameras/Private/Core/CameraRigAsset.cpp
@@ -20,17 +20,10 @@ void FCameraRigAllocationInfo::Append(const FCameraRigAllocationInfo& OtherAlloc
 	const FCameraVariableTableAllocationInfo& OtherVariableTableInfo(OtherAllocationInfo.VariableTableInfo);
 	VariableTableInfo.VariableDefinitions.Append(OtherVariableTableInfo.VariableDefinitions);
 
-	const FCameraContextDataAllocationInfo& OtherContextDataTableInfo(OtherAllocationInfo.ContextDataTableInfo);
+	const FCameraContextDataTableAllocationInfo& OtherContextDataTableInfo(OtherAllocationInfo.ContextDataTableInfo);
 	ContextDataTableInfo.DataDefinitions.Append(OtherContextDataTableInfo.DataDefinitions);
 }
 
-bool operator==(const FCameraRigAllocationInfo& A, const FCameraRigAllocationInfo& B)
-{
-	return A.EvaluatorInfo == B.EvaluatorInfo
-		&& A.VariableTableInfo == B.VariableTableInfo
-		&& A.ContextDataTableInfo == B.ContextDataTableInfo;
-}
-
 #if WITH_EDITOR
 
 void UCameraRigInterfaceParameterBase::GetGraphNodePosition(FName InGraphName, int32& NodePosX, int32& NodePosY) const
@@ -80,6 +73,48 @@ void UCameraRigInterfaceParameterBase::PostDuplicate(EDuplicateMode::Type Duplic
 	}
 }
 
+FCameraVariableDefinition UCameraRigBlendableParameter::GetVariableDefinition() const
+{
+	FCameraVariableDefinition Definition;
+	Definition.VariableID = PrivateVariableID;
+	Definition.VariableType = ParameterType;
+	Definition.BlendableStructType = BlendableStructType;
+	Definition.bIsPrivate = true;
+	Definition.bIsInput = bIsPreBlended;
+#if WITH_EDITORONLY_DATA
+	Definition.VariableName = GetVariableName();
+#endif
+	return Definition;
+}
+
+FCameraContextDataDefinition UCameraRigDataParameter::GetDataDefinition() const
+{
+	FCameraContextDataDefinition Definition;
+	Definition.DataID = PrivateDataID;
+	Definition.DataType = DataType;
+	Definition.DataTypeObject = DataTypeObject;
+#if WITH_EDITORONLY_DATA
+	Definition.DataName = GetDataName();
+#endif
+	return Definition;
+}
+
+#if WITH_EDITORONLY_DATA
+
+FString UCameraRigBlendableParameter::GetVariableName() const
+{
+	const UCameraRigAsset* Owner = GetTypedOuter<UCameraRigAsset>();
+	return FString::Printf(TEXT("Override_%s_%s"), *GetNameSafe(Owner), *InterfaceParameterName);
+}
+
+FString UCameraRigDataParameter::GetDataName() const
+{
+	const UCameraRigAsset* Owner = GetTypedOuter<UCameraRigAsset>();
+	return FString::Printf(TEXT("Override_%s_%s"), *GetNameSafe(Owner), *InterfaceParameterName);
+}
+
+#endif  // WITH_EDITORONLY_DATA
+
 UCameraRigBlendableParameter* FCameraRigInterface::FindBlendableParameterByName(const FString& ParameterName) const
 {
 	const TObjectPtr<UCameraRigBlendableParameter>* FoundItem = BlendableParameters.FindByPredicate(
@@ -172,9 +207,10 @@ void UCameraRigAsset::PostLoad()
 	{
 		if (BlendableParameter->ParameterType == ECameraVariableType::Boolean)
 		{
-			if (BlendableParameter->PrivateVariable)
+			if (BlendableParameter->PrivateVariable_DEPRECATED)
 			{
-				BlendableParameter->ParameterType = BlendableParameter->PrivateVariable->GetVariableType();
+				BlendableParameter->ParameterType = BlendableParameter->PrivateVariable_DEPRECATED->GetVariableType();
+				BlendableParameter->PrivateVariable_DEPRECATED = nullptr;
 			}
 		}
 	}
diff --git a/Engine/Plugins/Cameras/GameplayCameras/Source/GameplayCameras/Private/Core/CameraRigAssetBuilder.cpp b/Engine/Plugins/Cameras/GameplayCameras/Source/GameplayCameras/Private/Core/CameraRigAssetBuilder.cpp
--- a/Engine/Plugins/Cameras/GameplayCameras/Source/GameplayCameras/Private/Core/CameraRigAssetBuilder.cpp
+++ b/Engine/Plugins/Cameras/GameplayCameras/Source/GameplayCameras/Private/Core/CameraRigAssetBuilder.cpp
@@ -22,108 +22,6 @@ namespace UE::Cameras
 namespace Internal
 {
 
-template<typename VariableAssetType, typename ValueType>
-void SetPrivateVariableDefaultValue(VariableAssetType* PrivateVariable, typename TCallTraits<ValueType>::ParamType Value)
-{
-	if (PrivateVariable->DefaultValue != Value)
-	{
-		PrivateVariable->Modify();
-		PrivateVariable->DefaultValue = Value;
-	}
-}
-
-template<>
-void SetPrivateVariableDefaultValue<UTransform3dCameraVariable, FTransform3d>(UTransform3dCameraVariable* PrivateVariable, const FTransform3d& Value)
-{
-	// Template overload because transforms don't have an operator!=.
-	if (!PrivateVariable->DefaultValue.Equals(Value, 0.f))
-	{
-		PrivateVariable->Modify();
-		PrivateVariable->DefaultValue = Value;
-	}
-}
-
-template<>
-void SetPrivateVariableDefaultValue<UTransform3fCameraVariable, FTransform3f>(UTransform3fCameraVariable* PrivateVariable, const FTransform3f& Value)
-{
-	// Template overload because transforms don't have an operator!=.
-	if (!PrivateVariable->DefaultValue.Equals(Value, 0.f))
-	{
-		PrivateVariable->Modify();
-		PrivateVariable->DefaultValue = Value;
-	}
-}
-
-template<>
-void SetPrivateVariableDefaultValue<UBooleanCameraVariable, bool>(UBooleanCameraVariable* PrivateVariable, bool bValue)
-{
-	// Template overload because boolean variables have a bDefaultValue field, not DefaultValue.
-	if (PrivateVariable->bDefaultValue != bValue)
-	{
-		PrivateVariable->Modify();
-		PrivateVariable->bDefaultValue = bValue;
-	}
-}
-
-FString MakePrivateVariableName(UCameraRigAsset* CameraRig, const UCameraRigBlendableParameter* BlendableParameter)
-{
-	return FString::Format(
-			TEXT("Override_{0}_{1}"), 
-			{ CameraRig->GetName(), BlendableParameter->InterfaceParameterName });
-}
-
-UCameraVariableAsset* CreatePrivateVariable(UCameraRigAsset* CameraRig, const UCameraRigBlendableParameter* BlendableParameter)
-{
-	TSubclassOf<UCameraVariableAsset> VariableClass;
-	switch (BlendableParameter->ParameterType)
-	{
-#define UE_CAMERA_VARIABLE_FOR_TYPE(VariableType, VariableName)\
-		case ECameraVariableType::VariableName:\
-			VariableClass = U##VariableName##CameraVariable::StaticClass();\
-			break;
-		UE_CAMERA_VARIABLE_FOR_ALL_TYPES()
-#undef UE_CAMERA_VARIABLE_FOR_TYPE
-	}
-	if (!ensure(VariableClass))
-	{
-		return nullptr;
-	}
-
-	const FString VariableName = MakePrivateVariableName(CameraRig, BlendableParameter);
-	UCameraVariableAsset* PrivateVariable = NewObject<UCameraVariableAsset>(
-			CameraRig, VariableClass, FName(*VariableName), RF_Transactional);
-
-	// Make sure it's a private variable.
-	PrivateVariable->bIsPrivate = true;
-	PrivateVariable->bAutoReset = false;
-
-	// Make it an input variable if the parameter is pre-blended.
-	PrivateVariable->bIsInput = BlendableParameter->bIsPreBlended;
-
-	return PrivateVariable;
-}
-
-void UpdatePrivateVariable(UCameraRigAsset* CameraRig, const UCameraRigBlendableParameter* BlendableParameter, UCameraVariableAsset* PrivateVariable)
-{
-	check(PrivateVariable);
-
-	ensure(PrivateVariable->bIsPrivate);
-	ensure(!PrivateVariable->bAutoReset);
-
-	const FString VariableName = MakePrivateVariableName(CameraRig, BlendableParameter);
-	if (PrivateVariable->GetName() != VariableName)
-	{
-		ensure(!PrivateVariable->HasAnyFlags(RF_Public));
-		PrivateVariable->Rename(*VariableName, nullptr);
-	}
-
-	if (PrivateVariable->bIsInput != BlendableParameter->bIsPreBlended)
-	{
-		PrivateVariable->Modify();
-		PrivateVariable->bIsInput = BlendableParameter->bIsPreBlended;
-	}
-}
-
 struct FInterfaceParameterBindingBuilder
 {
 	UCameraRigAsset* CameraRig;
@@ -144,102 +42,103 @@ struct FInterfaceParameterBindingBuilder
 		Owner.BuildLog.AddMessage(EMessageSeverity::Error, MoveTemp(ErrorMessage));
 	}
 
-	template<typename CameraParameterType>
-	void SetCameraParameterOverride(
+	template<typename CameraParameterOrVariableReferenceType>
+	void SetCameraParameterOrVariableReferenceOverride(
 			const UCameraRigBlendableParameter* BlendableParameter,
 			FStructProperty* TargetProperty,
-			CameraParameterType* CameraParameterPtr)
+			CameraParameterOrVariableReferenceType* CameraParameterOrVariableReferencePtr)
 	{
-		using VariableAssetType = typename CameraParameterType::VariableAssetType;
+		using VariableAssetType = typename CameraParameterOrVariableReferenceType::VariableAssetType;
 		using ValueType = typename VariableAssetType::ValueType;
 
-		ensure(BlendableParameter->PrivateVariable);
+		ensure(BlendableParameter->PrivateVariableID.IsValid());
 		ensure(BlendableParameter->TargetPropertyName == TargetProperty->GetFName());
 
-		VariableAssetType* NewVariable = EnsureCompatibleVariable<VariableAssetType>(BlendableParameter);
-		const bool bIsValid = CheckTwiceDrivenParameter(BlendableParameter, CameraParameterPtr);
-		if (!NewVariable || !bIsValid)
+		const bool bIsValid = CheckIfParameterCanBeOverridden(BlendableParameter, CameraParameterOrVariableReferencePtr);
+		if (!bIsValid)
 		{
 			return;
 		}
 
-		// Find the variable that was previously driving this parameter. If it has changed, flag the
-		// camera node as having been modified.
 		UCameraNode* TargetNode = BlendableParameter->Target;
-		VariableAssetType* PreviousVariable = FindOldDrivingVariable<VariableAssetType>(TargetProperty->GetFName(), TargetNode);
-		if (PreviousVariable != NewVariable)
+		FCameraVariableID PreviousVariableID = FindOldDrivingVariableID(TargetProperty->GetFName(), TargetNode);
+		if (PreviousVariableID != BlendableParameter->PrivateVariableID)
 		{
 			TargetNode->Modify();
 		}
 
-		SetPrivateVariableDefaultValue<VariableAssetType, ValueType>(NewVariable, CameraParameterPtr->Value);
-
-		CameraParameterPtr->Variable = NewVariable;
+		CameraParameterOrVariableReferencePtr->VariableID = BlendableParameter->PrivateVariableID;
 	}
 
-	template<typename VariableReferenceType>
-	void SetVariableReferenceOverride(
+	template<typename VariableAssetType>
+	void SetCustomBlendableParameterOverride(
 			const UCameraRigBlendableParameter* BlendableParameter,
-			FStructProperty* TargetProperty,
-			VariableReferenceType* VariableReferencePtr)
+			const FCustomCameraNodeParameterInfos::FBlendableParameterInfo& CustomParameter)
 	{
-		using VariableAssetType = typename VariableReferenceType::VariableAssetType;
 		using ValueType = typename VariableAssetType::ValueType;
 
-		ensure(BlendableParameter->PrivateVariable);
-		ensure(BlendableParameter->TargetPropertyName == TargetProperty->GetFName());
+		ensure(BlendableParameter->PrivateVariableID.IsValid());
+		ensure(BlendableParameter->TargetPropertyName == CustomParameter.ParameterName);
 
-		VariableAssetType* NewVariable = EnsureCompatibleVariable<VariableAssetType>(BlendableParameter);
-		const bool bIsValid = CheckTwiceDrivenParameter(BlendableParameter, VariableReferencePtr);
-		if (!NewVariable || !bIsValid)
+		const bool bIsValid = CheckIfParameterCanBeOverridden(BlendableParameter, CustomParameter);
+		if (!bIsValid)
 		{
 			return;
 		}
 
-		// Find the variable that was previously driving this reference. If it has changed, flag the
-		// camera node as having been modified.
 		UCameraNode* TargetNode = BlendableParameter->Target;
-		VariableAssetType* PreviousVariable = FindOldDrivingVariable<VariableAssetType>(TargetProperty->GetFName(), TargetNode);
-		if (PreviousVariable != NewVariable)
+		FCameraVariableID PreviousVariableID = FindOldDrivingVariableID(CustomParameter.ParameterName, TargetNode);
+		if (PreviousVariableID != BlendableParameter->PrivateVariableID)
 		{
 			TargetNode->Modify();
 		}
 
-		// No default value to set on the driving variable.
-
-		VariableReferencePtr->Variable = NewVariable;
+		(*CustomParameter.OverrideVariableID) = BlendableParameter->PrivateVariableID;
 	}
 
-	template<typename VariableAssetType>
-	void SetCustomBlendableParameterOverride(
+	void SetCustomBlendableStructParameterOverride(
 			const UCameraRigBlendableParameter* BlendableParameter,
 			const FCustomCameraNodeParameterInfos::FBlendableParameterInfo& CustomParameter)
 	{
-		using ValueType = typename VariableAssetType::ValueType;
-
-		ensure(BlendableParameter->PrivateVariable);
+		ensure(BlendableParameter->PrivateVariableID.IsValid());
 		ensure(BlendableParameter->TargetPropertyName == CustomParameter.ParameterName);
 
-		VariableAssetType* NewVariable = EnsureCompatibleVariable<VariableAssetType>(BlendableParameter);
-		const bool bIsValid = CheckTwiceDrivenParameter(BlendableParameter, *CustomParameter.OverrideVariable);
-		if (!NewVariable || !bIsValid)
+		const bool bIsValid = CheckIfParameterCanBeOverridden(BlendableParameter, CustomParameter.OverrideVariableID);
+		if (!bIsValid)
 		{
 			return;
 		}
 
-		// Find the variable that was previously driving this override. If it has changed, flag the
-		// camera node as having been modified.
+		// Also ensure the struct type is compatible.
+		if (CustomParameter.BlendableStructType != BlendableParameter->BlendableStructType)
+		{
+			ReportError(BlendableParameter->Target,
+					FText::Format(
+						LOCTEXT(
+							"IncompatibleBlendableStructType", 
+							"Invalid interface parameter '{0}', driving property '{1}' on '{2}': expected type {3} but was {4}"),
+						FText::FromString(BlendableParameter->InterfaceParameterName),
+						FText::FromName(BlendableParameter->TargetPropertyName),
+						FText::FromName(BlendableParameter->Target->GetFName()),
+#if WITH_EDITORONLY_DATA
+						CustomParameter.BlendableStructType->GetDisplayNameText(),
+						BlendableParameter->BlendableStructType->GetDisplayNameText()
+#else
+						FText::FromName(CustomParameter.BlendableStructType->GetFName()),
+						FText::FromName(BlendableParameter->BlendableStructType->GetFName())
+#endif
+						));
+			return;
+		}
+
 		UCameraNode* TargetNode = BlendableParameter->Target;
-		VariableAssetType* PreviousVariable = FindOldDrivingVariable<VariableAssetType>(CustomParameter.ParameterName, TargetNode);
-		if (PreviousVariable != NewVariable)
+		FCameraVariableID PreviousVariableID = FindOldDrivingVariableID(CustomParameter.ParameterName, TargetNode);
+		if (PreviousVariableID != BlendableParameter->PrivateVariableID)
 		{
 			TargetNode->Modify();
 		}
 
-		const ValueType* DefaultValue = reinterpret_cast<const ValueType*>(CustomParameter.DefaultValuePtr);
-		SetPrivateVariableDefaultValue<VariableAssetType, ValueType>(NewVariable, *DefaultValue);
-
-		(*CustomParameter.OverrideVariable) = NewVariable;
+		(*CustomParameter.OverrideVariableID) = BlendableParameter->PrivateVariableID;
 	}
 
 	void SetDataContextPropertyOverride(
@@ -249,7 +148,7 @@ struct FInterfaceParameterBindingBuilder
 	{
 		ensure(DataParameter->PrivateDataID);
 
-		const bool bIsValid = CheckTwiceDrivenParameter(DataParameter, *OverrideDataID);
+		const bool bIsValid = CheckIfParameterCanBeOverridden(DataParameter, OverrideDataID);
 		if (!bIsValid)
 		{
 			return;
@@ -272,7 +171,7 @@ struct FInterfaceParameterBindingBuilder
 		ensure(DataParameter->PrivateDataID);
 		ensure(DataParameter->TargetPropertyName == CustomParameter.ParameterName);
 
-		const bool bIsValid = CheckTwiceDrivenParameter(DataParameter, *CustomParameter.OverrideDataID);
+		const bool bIsValid = CheckIfParameterCanBeOverridden(DataParameter, CustomParameter.OverrideDataID);
 		if (!bIsValid)
 		{
 			return;
@@ -290,15 +189,14 @@ struct FInterfaceParameterBindingBuilder
 
 private:
 
-	template<typename ExpectedVariableAssetType>
-	ExpectedVariableAssetType* FindOldDrivingVariable(FName ForParameterName, UCameraNode* ForCameraNode)
+	FCameraVariableID FindOldDrivingVariableID(FName ForParameterName, UCameraNode* ForCameraNode)
 	{
 		using FDrivenParameterKey = FCameraRigAssetBuilder::FDrivenParameterKey;
 
 		FDrivenParameterKey ParameterKey{ ForParameterName, ForCameraNode };
-		TObjectPtr<UCameraVariableAsset> ReusedVariable;
-		Owner.OldDrivenBlendableParameters.RemoveAndCopyValue(ParameterKey, ReusedVariable);
-		return CastChecked<ExpectedVariableAssetType>(ReusedVariable, ECastCheckedType::NullAllowed);
+		FCameraVariableID ReusedVariableID;
+		Owner.OldDrivenBlendableParameters.RemoveAndCopyValue(ParameterKey, ReusedVariableID);
+		return ReusedVariableID;
 	}
 
 	FCameraContextDataID FindOldDrivingDataID(FName ForParameterName, UCameraNode* ForCameraNode)
@@ -311,91 +209,101 @@ private:
 		return ReusedDataID;
 	}
 
-	template<typename CameraParameterType>
-	bool CheckTwiceDrivenParameter(
+	template<typename CameraParameterOrVariableReferenceType>
+	bool CheckIfParameterCanBeOverridden(
 			const UCameraRigBlendableParameter* BlendableParameter, 
-			CameraParameterType* CameraParameter)
+			CameraParameterOrVariableReferenceType* CameraParameterOrVariableReference)
 	{
-		using VariableAssetType = typename CameraParameterType::VariableAssetType;
-		return CheckTwiceDrivenParameter<VariableAssetType>(BlendableParameter, CameraParameter->Variable);
+		if (CameraParameterOrVariableReference->Variable != nullptr)
+		{
+			ReportError(BlendableParameter->Target,
+					FText::Format(
+						LOCTEXT(
+							"BlendableParameterDrivenTwice", 
+							"Camera node parameter '{0}.{1}' is both exposed and driven by a variable!"),
+						FText::FromName(BlendableParameter->Target->GetFName()), 
+						FText::FromName(BlendableParameter->TargetPropertyName)));
+			return false;
+		}
+
+		return CheckIfParameterCanBeOverridden(BlendableParameter, &CameraParameterOrVariableReference->VariableID);
 	}
 
-	template<typename VariableAssetType>
-	bool CheckTwiceDrivenParameter(
+	bool CheckIfParameterCanBeOverridden(
 			const UCameraRigBlendableParameter* BlendableParameter, 
-			TObjectPtr<VariableAssetType> CameraParameterVariable)
+			const FCustomCameraNodeParameterInfos::FBlendableParameterInfo& CustomParameter)
 	{
-		if (CameraParameterVariable != nullptr)
+		if (CustomParameter.OverrideVariable != nullptr)
 		{
-			// We should have cleared all exposed parameters in GatherOldDrivenParameters, so the only variables
-			// left on camera parameters should be user-defined ones.
-			UObject* VariableOuter = CameraParameterVariable->GetOuter();
-			if (ensureMsgf(
-						VariableOuter != CameraRig, 
-						TEXT("Unexpected driving variable found: all exposed parameters should have been cleared before rebuilding.")))
-			{
-				// If this parameter is driven by a user-defined variable, emit an error.
-				ReportError(BlendableParameter->Target,
-						FText::Format(
-							LOCTEXT(
-								"BlendableParameterDrivenTwice", 
-								"Camera node parameter '{0}.{1}' is both exposed and driven by a variable!"),
-							FText::FromName(BlendableParameter->Target->GetFName()), 
-							FText::FromName(BlendableParameter->TargetPropertyName)));
-				return false;
-			}
+			ReportError(BlendableParameter->Target,
+					FText::Format(
+						LOCTEXT(
+							"BlendableParameterDrivenTwice", 
+							"Camera node parameter '{0}.{1}' is both exposed and driven by a variable!"),
+						FText::FromName(BlendableParameter->Target->GetFName()), 
+						FText::FromName(BlendableParameter->TargetPropertyName)));
+			return false;
 		}
-		return true;
+
+		return CheckIfParameterCanBeOverridden(BlendableParameter, CustomParameter.OverrideVariableID);
 	}
 
-	bool CheckTwiceDrivenParameter(
-			const UCameraRigDataParameter* DataParameter,
-			FCameraContextDataID DataID)
+	bool CheckIfParameterCanBeOverridden(
+			const UCameraRigBlendableParameter* BlendableParameter, 
+			FCameraVariableID* VariableID)
 	{
-		if (DataID.IsValid())
+		if (!VariableID)
 		{
-			ReportError(DataParameter->Target,
+			ReportError(BlendableParameter->Target,
 					FText::Format(
 						LOCTEXT(
-							"DataParameterDrivenTwice",
+							"BlendableParameterMissingOverrideID", 
+							"Camera node parameter '{0}.{1}' cannot be overriden by a parameter"),
+						FText::FromName(BlendableParameter->Target->GetFName()), 
+						FText::FromName(BlendableParameter->TargetPropertyName)));
+			return false;
+		}
+		if (VariableID->IsValid())
+		{
+			ReportError(BlendableParameter->Target,
+					FText::Format(
+						LOCTEXT(
+							"BlendableParameterDrivenTwice", 
 							"Camera node parameter '{0}.{1}' is somehow overriden twice!"),
-						FText::FromName(DataParameter->Target->GetFName()), 
-						FText::FromName(DataParameter->TargetPropertyName)));
+						FText::FromName(BlendableParameter->Target->GetFName()), 
+						FText::FromName(BlendableParameter->TargetPropertyName)));
 			return false;
 		}
 		return true;
 	}
 
-	template<typename VariableAssetType>
-	VariableAssetType* EnsureCompatibleVariable(const UCameraRigBlendableParameter* BlendableParameter)
+	bool CheckIfParameterCanBeOverridden(
+			const UCameraRigDataParameter* DataParameter,
+			FCameraContextDataID* DataID)
 	{
-		if (!ensure(BlendableParameter->PrivateVariable))
+		if (!DataID)
 		{
-			return nullptr;
+			ReportError(DataParameter->Target,
+					FText::Format(
+						LOCTEXT(
+							"DataParameterMissingOverrideID", 
+							"Camera node parameter '{0}.{1}' cannot be overriden by a parameter"),
+						FText::FromName(DataParameter->Target->GetFName()), 
+						FText::FromName(DataParameter->TargetPropertyName)));
+			return false;
 		}
-
-		VariableAssetType* Variable = Cast<VariableAssetType>(BlendableParameter->PrivateVariable);
-
-		if (!Variable)
+		if (DataID->IsValid())
 		{
-			const UClass* VariableAssetClass = VariableAssetType::StaticClass();
-			const UCameraVariableAsset* DefaultObject = VariableAssetClass->template GetDefaultObject<UCameraVariableAsset>();
-
-			ReportError(BlendableParameter->Target,
-				FText::Format(LOCTEXT(
-					"IncompatibleCameraNodeParameter",
-					"Invalid interface parameter '{0}', driving property '{1}' on '{2}': expected {3} camera parameter but was {4}"),
-					FText::FromString(BlendableParameter->InterfaceParameterName),
-					FText::FromName(BlendableParameter->TargetPropertyName),
-					FText::FromName(BlendableParameter->Target->GetFName()),
-					FText::FromName(UEnum::GetValueAsName(DefaultObject->GetVariableType())),
-					FText::FromName(UEnum::GetValueAsName(BlendableParameter->ParameterType))
-					)
-				);
-			return nullptr;
+			ReportError(DataParameter->Target,
+					FText::Format(
+						LOCTEXT(
+							"DataParameterDrivenTwice",
+							"Camera node parameter '{0}.{1}' is somehow overriden twice!"),
+						FText::FromName(DataParameter->Target->GetFName()), 
+						FText::FromName(DataParameter->TargetPropertyName)));
+			return false;
 		}
-
-		return Variable;
+		return true;
 	}
 
 private:
@@ -403,26 +311,30 @@ private:
 	FCameraRigAssetBuilder& Owner;
 };
 
-void TrashPrivateVariable(UCameraVariableAsset* VariableToTrash)
+void AddVariableToAllocationInfo(UCameraVariableAsset* Variable, FCameraVariableTableAllocationInfo& AllocationInfo)
 {
-	TStringBuilder<256> StringBuilder;
-	StringBuilder.Append("TRASH_");
-	StringBuilder.Append(VariableToTrash->GetName());
-
-	VariableToTrash->Modify();
-	VariableToTrash->Rename(StringBuilder.ToString());
+	if (Variable)
+	{
+		FCameraVariableDefinition VariableDefinition = Variable->GetVariableDefinition();
+		AllocationInfo.VariableDefinitions.Add(VariableDefinition);
+	}
 }
 
-void AddCameraVariableToAllocationInfo(UCameraVariableAsset* Variable, FCameraVariableTableAllocationInfo& AllocationInfo)
+void AddVariableToAllocationInfo(FCameraVariableID VariableID, ECameraVariableType VariableType, const UScriptStruct* BlendableStructType, FCameraVariableTableAllocationInfo& AllocationInfo)
 {
-	if (Variable)
+	if (VariableID)
 	{
-		FCameraVariableDefinition VariableDefinition = Variable->GetVariableDefinition();
+		FCameraVariableDefinition VariableDefinition;
+		VariableDefinition.VariableID = VariableID;
+		VariableDefinition.VariableType = VariableType;
+		VariableDefinition.BlendableStructType = BlendableStructType;
+		VariableDefinition.bIsPrivate = true;
+		VariableDefinition.bIsInput = true;
 		AllocationInfo.VariableDefinitions.Add(VariableDefinition);
 	}
 }
 
-void AddContextDataToAllocationInfo(FCameraContextDataID DataID, ECameraContextDataType DataType, const UObject* DataTypeObject, FCameraContextDataAllocationInfo& AllocationInfo)
+void AddContextDataToAllocationInfo(FCameraContextDataID DataID, ECameraContextDataType DataType, const UObject* DataTypeObject, FCameraContextDataTableAllocationInfo& AllocationInfo)
 {
 	if (DataID)
 	{
@@ -489,6 +401,7 @@ void FCameraRigAssetBuilder::BuildCameraRigImpl()
 
 	CallBuild();
 
+	BuildParameterDefinitions();
 	BuildDefaultParameters();
 }
 
@@ -526,9 +439,9 @@ void FCameraRigAssetBuilder::CallPreBuild()
 
 void FCameraRigAssetBuilder::GatherOldDrivenParameters()
 {
-	// Keep track of which camera parameters and variable references were previously driven by 
-	// private variables, and then clear those variables. This is because it's easier to rebuild 
-	// all this from a blank slate than trying to figure out what changed.
+	// Keep track of which blendable/data parameters were previously overriden with private IDs.
+	// Then clear those private IDs. This is because it's easier to rebuild all this from a blank 
+	// slate than trying to figure out what changed.
 	//
 	// As we rebuild things in BuildInterfaceParameterBindings, we compare to the old state to
 	// figure out if we need to flag anything as modified for the current transaction.
@@ -574,33 +487,23 @@ void FCameraRigAssetBuilder::GatherOldDrivenParameters()
 			if (StructProperty->Struct == F##ValueName##CameraParameter::StaticStruct())\
 			{\
 				auto* CameraParameterPtr = StructProperty->ContainerPtrToValuePtr<F##ValueName##CameraParameter>(CameraNode);\
-				if (CameraParameterPtr->Variable)\
+				if (CameraParameterPtr->VariableID.IsValid() && !CameraParameterPtr->Variable)\
 				{\
-					UObject* VariableOuter = CameraParameterPtr->Variable->GetOuter();\
-					if (VariableOuter == CameraRig)\
-					{\
-						OldDrivenBlendableParameters.Add(\
-								FDrivenParameterKey{ StructProperty->GetFName(), CameraNode },\
-								*reinterpret_cast<TObjectPtr<UCameraVariableAsset>*>(\
-									reinterpret_cast<FObjectPtr*>(&CameraParameterPtr->Variable)));\
-						CameraParameterPtr->Variable = nullptr;\
-					}\
+					OldDrivenBlendableParameters.Add(\
+							FDrivenParameterKey{ StructProperty->GetFName(), CameraNode },\
+							CameraParameterPtr->VariableID);\
+					CameraParameterPtr->VariableID = FCameraVariableID();\
 				}\
 			}\
 			else if (StructProperty->Struct == F##ValueName##CameraVariableReference::StaticStruct())\
 			{\
 				auto* VariableReferencePtr = StructProperty->ContainerPtrToValuePtr<F##ValueName##CameraVariableReference>(CameraNode);\
-				if (VariableReferencePtr->Variable)\
+				if (VariableReferencePtr->VariableID.IsValid() && !VariableReferencePtr->Variable)\
 				{\
-					UObject* VariableOuter = VariableReferencePtr->Variable->GetOuter();\
-					if (VariableOuter == CameraRig)\
-					{\
-						OldDrivenBlendableParameters.Add(\
-								FDrivenParameterKey{ StructProperty->GetFName(), CameraNode },\
-								*reinterpret_cast<TObjectPtr<UCameraVariableAsset>*>(\
-									reinterpret_cast<FObjectPtr*>(&VariableReferencePtr->Variable)));\
-						VariableReferencePtr->Variable = nullptr;\
-					}\
+					OldDrivenBlendableParameters.Add(\
+							FDrivenParameterKey{ StructProperty->GetFName(), CameraNode },\
+							VariableReferencePtr->VariableID);\
+					VariableReferencePtr->VariableID = FCameraVariableID();\
 				}\
 			}\
 			else
@@ -618,27 +521,23 @@ UE_CAMERA_VARIABLE_FOR_ALL_TYPES()
 
 			for (const FCustomCameraNodeParameterInfos::FBlendableParameterInfo& BlendableParameter : CustomParameters.BlendableParameters)
 			{
-				if (!ensure(BlendableParameter.OverrideVariable))
+				if (!BlendableParameter.OverrideVariableID)
 				{
 					continue;
 				}
 
-				if (UCameraVariableAsset* Variable = BlendableParameter.OverrideVariable->Get())
+				if (BlendableParameter.OverrideVariableID->IsValid())
 				{
-					UObject* VariableOuter = Variable->GetOuter();
-					if (VariableOuter == CameraRig)
-					{
-						OldDrivenBlendableParameters.Add(
-								FDrivenParameterKey{ BlendableParameter.ParameterName, CameraNode },
-								*BlendableParameter.OverrideVariable);
-						(*BlendableParameter.OverrideVariable) = nullptr;
-					}
+					OldDrivenBlendableParameters.Add(
+							FDrivenParameterKey{ BlendableParameter.ParameterName, CameraNode },
+							*BlendableParameter.OverrideVariableID);
+					(*BlendableParameter.OverrideVariableID) = FCameraVariableID();
 				}
 			}
 
 			for (const FCustomCameraNodeParameterInfos::FDataParameterInfo& DataParameter : CustomParameters.DataParameters)
 			{
-				if (!ensure(DataParameter.OverrideDataID))
+				if (!DataParameter.OverrideDataID)
 				{
 					continue;
 				}
@@ -686,28 +585,14 @@ void FCameraRigAssetBuilder::BuildInterfaceParameters()
 			continue;
 		}
 
-		// Create a new private variable for this interface parameter if it doesn't have on yet (e.g. it's a
-		// newly-created parameter), or if its type has changed.
-		UCameraVariableAsset* PrivateVariable = BlendableParameter->PrivateVariable;
-		if (PrivateVariable && PrivateVariable->GetVariableType() == BlendableParameter->ParameterType)
+		// Create a new private variable ID for this interface parameter. Flag the parameter as changed if
+		// the ID is different, generally when it's a new parameter.
+		FCameraVariableID VariableID = FCameraVariableID::FromHashValue(GetTypeHash(BlendableParameter->GetGuid()));
+		if (BlendableParameter->PrivateVariableID != VariableID)
 		{
-			UpdatePrivateVariable(CameraRig, BlendableParameter, PrivateVariable);
-			continue;
-		}
-
-		if (PrivateVariable)
-		{
-			// Trash it so its name is available for the new private variable.
-			TrashPrivateVariable(PrivateVariable);
-
-			// We supposedly have removed all references to private variables in GatherOldDrivenParameters,
-			// so once we null-out the reference on the interface parameter, this private variable should
-			// not be used anymore and should get GC'ed soon.
 			BlendableParameter->Modify();
-			BlendableParameter->PrivateVariable = nullptr;
+			BlendableParameter->PrivateVariableID = VariableID;
 		}
-
-		BlendableParameter->PrivateVariable = CreatePrivateVariable(CameraRig, BlendableParameter);
 	}
 
 	for (auto It = CameraRig->Interface.DataParameters.CreateIterator(); It; ++It)
@@ -737,13 +622,9 @@ void FCameraRigAssetBuilder::BuildInterfaceParameters()
 			continue;
 		}
 
-		// Create the data ID for this parameter. Flag the parameter as changed if the ID
-		// is different, such as if it was renamed.
-		const FString DataIDName = FString::Format(
-				TEXT("Override_{0}_{1}"), 
-				{ CameraRig->GetName(), DataParameter->InterfaceParameterName });
-
-		FCameraContextDataID DataID = FCameraContextDataID::FromName(FName(DataIDName));
+		// Create a new private data ID for this interface parameter. Flag the parameter as changed if
+		// the ID is different, generally when it's a new parameter.
+		FCameraContextDataID DataID = FCameraContextDataID::FromHashValue(GetTypeHash(DataParameter->GetGuid()));
 		if (DataParameter->PrivateDataID != DataID)
 		{
 			DataParameter->Modify();
@@ -802,11 +683,11 @@ void FCameraRigAssetBuilder::BuildInterfaceParameterBindings()
 
 		// See if this interface parameter is overriding a camera node parameter.
 		// Otherwise, maybe it's targeting a camera rig node's override for an inner rig interface parameter.
-		if (SetupCameraParameterOrVariableReferenceOverride(BlendableParameter))
+		if (SetupCustomBlendableParameterOverride(BlendableParameter))
 		{
 			// Implicit continue.
 		}
-		else if (SetupCustomBlendableParameterOverride(BlendableParameter))
+		else if (SetupCameraParameterOrVariableReferenceOverride(BlendableParameter))
 		{
 			// Implicit continue.
 		}
@@ -850,11 +731,11 @@ void FCameraRigAssetBuilder::BuildInterfaceParameterBindings()
 			continue;
 		}
 
-		if (SetupDataContextPropertyOverride(DataParameter))
+		if (SetupCustomDataParameterOverride(DataParameter))
 		{
 			// Implicit continue.
 		}
-		else if (SetupCustomDataParameterOverride(DataParameter))
+		else if (SetupDataContextPropertyOverride(DataParameter))
 		{
 			// Implicit continue.
 		}
@@ -913,14 +794,14 @@ bool FCameraRigAssetBuilder::SetupCameraParameterOrVariableReferenceOverride(con
 	if (TargetStructProperty->Struct == F##ValueName##CameraParameter::StaticStruct())\
 	{\
 		auto* CameraParameterPtr = TargetStructProperty->ContainerPtrToValuePtr<F##ValueName##CameraParameter>(Target);\
-		Builder.SetCameraParameterOverride<F##ValueName##CameraParameter>(\
+		Builder.SetCameraParameterOrVariableReferenceOverride<F##ValueName##CameraParameter>(\
 				BlendableParameter, TargetStructProperty, CameraParameterPtr\
 				);\
 	}\
 	else if (TargetStructProperty->Struct == F##ValueName##CameraVariableReference::StaticStruct())\
 	{\
 		auto* VariableReferencePtr = TargetStructProperty->ContainerPtrToValuePtr<F##ValueName##CameraVariableReference>(Target);\
-		Builder.SetVariableReferenceOverride<F##ValueName##CameraVariableReference>(\
+		Builder.SetCameraParameterOrVariableReferenceOverride<F##ValueName##CameraVariableReference>(\
 				BlendableParameter, TargetStructProperty, VariableReferencePtr\
 				);\
 	}\
@@ -978,6 +859,9 @@ bool FCameraRigAssetBuilder::SetupCustomBlendableParameterOverride(const UCamera
 			break;
 	UE_CAMERA_VARIABLE_FOR_ALL_TYPES()
 #undef UE_CAMERA_VARIABLE_FOR_TYPE
+		case ECameraVariableType::BlendableStruct:
+			Builder.SetCustomBlendableStructParameterOverride(BlendableParameter, *TargetCustomParameter);
+			break;
 	}
 
 	return true;
@@ -1056,7 +940,7 @@ void FCameraRigAssetBuilder::DiscardUnusedParameters()
 	// must be discarded. These are nodes and properties that used to be driven by
 	// variables and now aren't, so we need to flag them as modified.
 
-	for (TPair<FDrivenParameterKey, TObjectPtr<UCameraVariableAsset>> Pair : OldDrivenBlendableParameters)
+	for (TPair<FDrivenParameterKey, FCameraVariableID> Pair : OldDrivenBlendableParameters)
 	{
 		UCameraNode* Target = Pair.Key.Value;
 		Target->Modify();
@@ -1090,6 +974,9 @@ void FCameraRigAssetBuilder::CallBuild()
 		CallBuild(BuildContext, CameraNode);
 	}
 
+	// Add parameters to the allocation info.
+	BuildParametersAllocationInfo(BuildContext);
+
 	// Set the final allocation info on the camera rig asset.
 	if (CameraRig->AllocationInfo != BuildContext.AllocationInfo)
 	{
@@ -1103,8 +990,8 @@ void FCameraRigAssetBuilder::CallBuild(FCameraRigBuildContext& BuildContext, UCa
 	using namespace UE::Cameras::Internal;
 
 	// Look for properties that are camera parameters, and gather what camera variables they reference. 
-	// This is for both exposed rig parameters (which we just built in BuildNewDrivenParameters) and 
-	// for parameters driven by user-defined variables.
+	// This is only for user-defined variable overrides. We will do the same for exposed camera rig
+	// parameters later, in BuildParametersAllocationInfo.
 	UClass* CameraNodeClass = CameraNode->GetClass();
 	FCameraRigAllocationInfo& AllocationInfo = BuildContext.AllocationInfo;
 	for (TFieldIterator<FProperty> It(CameraNodeClass); It; ++It)
@@ -1119,12 +1006,12 @@ void FCameraRigAssetBuilder::CallBuild(FCameraRigBuildContext& BuildContext, UCa
 		if (StructProperty->Struct == F##ValueName##CameraParameter::StaticStruct())\
 		{\
 			auto* CameraParameterPtr = StructProperty->ContainerPtrToValuePtr<F##ValueName##CameraParameter>(CameraNode);\
-			AddCameraVariableToAllocationInfo(CameraParameterPtr->Variable, AllocationInfo.VariableTableInfo);\
+			AddVariableToAllocationInfo(CameraParameterPtr->Variable, AllocationInfo.VariableTableInfo);\
 		}\
 		else if (StructProperty->Struct == F##ValueName##CameraVariableReference::StaticStruct())\
 		{\
 			auto* CameraVariableReferencePtr = StructProperty->ContainerPtrToValuePtr<F##ValueName##CameraVariableReference>(CameraNode);\
-			AddCameraVariableToAllocationInfo(CameraVariableReferencePtr->Variable, AllocationInfo.VariableTableInfo);\
+			AddVariableToAllocationInfo(CameraVariableReferencePtr->Variable, AllocationInfo.VariableTableInfo);\
 		}\
 		else
 UE_CAMERA_VARIABLE_FOR_ALL_TYPES()
@@ -1134,9 +1021,7 @@ UE_CAMERA_VARIABLE_FOR_ALL_TYPES()
 		}
 	}
 
-	// Now do the same with custom parameters handled by the node itself. These custom parameters have 
-	// been hooked up to overrides in BuildInterfaceParameterBindings so we should be getting them back 
-	// with the overrides set here.
+	// Now do the same with custom parameters handled by the node itself.
 	if (ICustomCameraNodeParameterProvider* CustomParameterProvider = Cast<ICustomCameraNodeParameterProvider>(CameraNode))
 	{
 		FCustomCameraNodeParameterInfos CustomParameters;
@@ -1144,30 +1029,88 @@ UE_CAMERA_VARIABLE_FOR_ALL_TYPES()
 
 		for (const FCustomCameraNodeParameterInfos::FBlendableParameterInfo& BlendableParameter : CustomParameters.BlendableParameters)
 		{
-			if (ensure(BlendableParameter.OverrideVariable))
-			{
-				AddCameraVariableToAllocationInfo(BlendableParameter.OverrideVariable->Get(), AllocationInfo.VariableTableInfo);
-			}
+			AddVariableToAllocationInfo(BlendableParameter.OverrideVariable, AllocationInfo.VariableTableInfo);
 		}
+	}
+
+	// Let the camera node add any custom variables or extra memory.
+	CameraNode->Build(BuildContext);
+}
 
-		for (const FCustomCameraNodeParameterInfos::FDataParameterInfo& DataParameter : CustomParameters.DataParameters)
+void FCameraRigAssetBuilder::BuildParametersAllocationInfo(FCameraRigBuildContext& BuildContext)
+{
+	// The variables and context data definitions should have already been added by the camera nodes
+	// who have override variable IDs and data IDs set on them. 
+
+	for (const UCameraRigBlendableParameter* BlendableParameter : CameraRig->Interface.BlendableParameters)
+	{
+		if (!BlendableParameter->PrivateVariableID.IsValid())
 		{
-			if (ensure(DataParameter.OverrideDataID))
-			{
-				AddContextDataToAllocationInfo(*DataParameter.OverrideDataID, DataParameter.ParameterType, DataParameter.ParameterTypeObject, AllocationInfo.ContextDataTableInfo);
-			}
+			continue;
 		}
+
+		FCameraVariableDefinition Definition = BlendableParameter->GetVariableDefinition();
+		BuildContext.AllocationInfo.VariableTableInfo.VariableDefinitions.Add(Definition);
 	}
 
-	// Let the camera node add any custom variables or extra memory.
-	CameraNode->Build(BuildContext);
+	for (const UCameraRigDataParameter* DataParameter : CameraRig->Interface.DataParameters)
+	{
+		if (!DataParameter->PrivateDataID.IsValid())
+		{
+			continue;
+		}
+
+		FCameraContextDataDefinition Definition = DataParameter->GetDataDefinition();
+		BuildContext.AllocationInfo.ContextDataTableInfo.DataDefinitions.Add(Definition);
+	}
+}
+
+void FCameraRigAssetBuilder::BuildParameterDefinitions()
+{
+	TArray<FCameraRigParameterDefinition> ParameterDefinitions;
+
+	for (const UCameraRigBlendableParameter* BlendableParameter : CameraRig->Interface.BlendableParameters)
+	{
+		if (BlendableParameter && BlendableParameter->PrivateVariableID)
+		{
+			FCameraRigParameterDefinition Definition;
+			Definition.ParameterName = FName(BlendableParameter->InterfaceParameterName);
+			Definition.ParameterGuid = BlendableParameter->GetGuid();
+			Definition.ParameterType = ECameraRigInterfaceParameterType::Blendable;
+			Definition.VariableID = BlendableParameter->PrivateVariableID;
+			Definition.VariableType = BlendableParameter->ParameterType;
+			Definition.BlendableStructType = BlendableParameter->BlendableStructType;
+			ParameterDefinitions.Add(Definition);
+		}
+	}
+
+	for (const UCameraRigDataParameter* DataParameter : CameraRig->Interface.DataParameters)
+	{
+		if (DataParameter && DataParameter->PrivateDataID)
+		{
+			FCameraRigParameterDefinition Definition;
+			Definition.ParameterName = FName(DataParameter->InterfaceParameterName);
+			Definition.ParameterGuid = DataParameter->GetGuid();
+			Definition.ParameterType = ECameraRigInterfaceParameterType::Data;
+			Definition.DataID = DataParameter->PrivateDataID;
+			Definition.DataType = DataParameter->DataType;
+			Definition.DataTypeObject = DataParameter->DataTypeObject;
+			ParameterDefinitions.Add(Definition);
+		}
+	}
+
+	if (ParameterDefinitions != CameraRig->ParameterDefinitions)
+	{
+		CameraRig->Modify();
+		CameraRig->ParameterDefinitions = ParameterDefinitions;
+	}
 }
 
 void FCameraRigAssetBuilder::BuildDefaultParameters()
 {
 	FInstancedPropertyBag DefaultParameters;
-	BuildDefaultParameters(CameraRig, DefaultParameters);
-	if (DefaultParameters.GetPropertyBagStruct() != CameraRig->DefaultParameters.GetPropertyBagStruct())
+	FCameraRigParameterBuilder::BuildDefaultParameters(CameraRig, DefaultParameters);
+	if (!DefaultParameters.Identical(&CameraRig->DefaultParameters, 0))
 	{
 		CameraRig->Modify();
 		CameraRig->DefaultParameters = DefaultParameters;
@@ -1190,89 +1133,6 @@ void FCameraRigAssetBuilder::UpdateBuildStatus()
 	CameraRig->BuildStatus = BuildStatus;
 }
 
-void FCameraRigAssetBuilder::BuildDefaultParameters(UCameraRigAsset* CameraRigAsset, FInstancedPropertyBag& OutPropertyBag)
-{
-	TArray<FPropertyBagPropertyDesc> DefaultParameterProperties;
-	AppendDefaultParameters(CameraRigAsset->Interface, DefaultParameterProperties);
-	OutPropertyBag.AddProperties(DefaultParameterProperties);
-}
-
-void FCameraRigAssetBuilder::AppendDefaultParameters(const FCameraRigInterface& CameraRigInterface, TArray<FPropertyBagPropertyDesc>& OutProperties)
-{
-	for (const UCameraRigBlendableParameter* BlendableParameter : CameraRigInterface.BlendableParameters)
-	{
-		FName PropertyName(BlendableParameter->InterfaceParameterName);
-
-		EPropertyBagPropertyType PropertyType = EPropertyBagPropertyType::Struct;
-		const UObject* PropertyTypeObject = nullptr;
-
-		switch (BlendableParameter->ParameterType)
-		{
-#define UE_CAMERA_VARIABLE_FOR_TYPE(ValueType, ValueName)\
-			case ECameraVariableType::ValueName:\
-				PropertyTypeObject = F##ValueName##CameraParameter::StaticStruct();\
-				break;
-		UE_CAMERA_VARIABLE_FOR_ALL_TYPES()
-#undef UE_CAMERA_VARIABLE_FOR_TYPE
-			default:
-				ensure(false);
-				break;
-		}
-
-		FPropertyBagPropertyDesc NewProperty(PropertyName, PropertyType, PropertyTypeObject);
-		// Make the property bag match the camera interface parameter GUIDs.
-		NewProperty.ID = BlendableParameter->GetGuid();
-
-		OutProperties.Add(NewProperty);
-	}
-
-	for (const UCameraRigDataParameter* DataParameter : CameraRigInterface.DataParameters)
-	{
-		FName PropertyName(DataParameter->InterfaceParameterName);
-
-		bool bIsValid = true;
-		EPropertyBagPropertyType PropertyType = EPropertyBagPropertyType::Bool;
-		const UObject* PropertyTypeObject = DataParameter->DataTypeObject;
-
-		switch (DataParameter->DataType)
-		{
-			case ECameraContextDataType::Name:
-				PropertyType = EPropertyBagPropertyType::Name;
-				break;
-			case ECameraContextDataType::String:
-				PropertyType = EPropertyBagPropertyType::String;
-				break;
-			case ECameraContextDataType::Enum:
-				PropertyType = EPropertyBagPropertyType::Enum;
-				ensure(PropertyTypeObject && PropertyTypeObject->IsA<UEnum>());
-				break;
-			case ECameraContextDataType::Struct:
-				PropertyType = EPropertyBagPropertyType::Struct;
-				ensure(PropertyTypeObject && PropertyTypeObject->IsA<UScriptStruct>());
-				break;
-			case ECameraContextDataType::Object:
-				PropertyType = EPropertyBagPropertyType::Object;
-				break;
-			case ECameraContextDataType::Class:
-				PropertyType = EPropertyBagPropertyType::Class;
-				break;
-			default:
-				ensure(false);
-				bIsValid = false;
-				break;
-		}
-
-		if (bIsValid)
-		{
-			FPropertyBagPropertyDesc NewProperty(PropertyName, PropertyType, PropertyTypeObject);
-			// Make the property bag match the camera interface parameter GUIDs.
-			NewProperty.ID = DataParameter->GetGuid();
-
-			OutProperties.Add(NewProperty);
-		}
-	}
-}
-
 }  // namespace UE::Cameras
 
 #undef LOCTEXT_NAMESPACE
diff --git a/Engine/Plugins/Cameras/GameplayCameras/Source/GameplayCameras/Private/Core/CameraRigAssetReference.cpp b/Engine/Plugins/Cameras/GameplayCameras/Source/GameplayCameras/Private/Core/CameraRigAssetReference.cpp
--- a/Engine/Plugins/Cameras/GameplayCameras/Source/GameplayCameras/Private/Core/CameraRigAssetReference.cpp
+++ b/Engine/Plugins/Cameras/GameplayCameras/Source/GameplayCameras/Private/Core/CameraRigAssetReference.cpp
@@ -4,6 +4,7 @@
 
 #include "Core/CameraRigAsset.h"
 #include "Core/CameraNodeEvaluator.h"
+#include "Core/ICustomCameraNodeParameterProvider.h"
 #include "Helpers/CameraRigParameterOverrideEvaluator.h"
 
 #include UE_INLINE_GENERATED_CPP_BY_NAME(CameraRigAssetReference)
@@ -28,20 +29,50 @@ void FCameraRigAssetReference::ApplyParameterOverrides(UE::Cameras::FCameraNodeE
 	OverrideEvaluator.ApplyParameterOverrides(OutResult.VariableTable, OutResult.ContextDataTable, bDrivenOverridesOnly);
 }
 
-bool FCameraRigAssetReference::IsParameterOverriden(const FGuid PropertyID) const
+const FCameraRigAssetReferenceParameterMetaData* FCameraRigAssetReference::FindMetaData(const FGuid& PropertyID) const
 {
-	return ParameterOverrideGuids.Contains(PropertyID);
+	return ParameterMetaData.FindByPredicate(
+			[PropertyID](const FCameraRigAssetReferenceParameterMetaData& Item)
+			{
+				return Item.ParameterGuid == PropertyID;
+			});
+}
+
+FCameraRigAssetReferenceParameterMetaData* FCameraRigAssetReference::FindMetaData(const FGuid& PropertyID)
+{
+	return ParameterMetaData.FindByPredicate(
+			[PropertyID](FCameraRigAssetReferenceParameterMetaData& Item)
+			{
+				return Item.ParameterGuid == PropertyID;
+			});
 }
 
-void FCameraRigAssetReference::SetParameterOverriden(const FGuid PropertyID, bool bIsOverridden)
+void FCameraRigAssetReference::GenerateOverriddenParameterGuidArray(TArray<FGuid>& OutOverriddenIDs) const
 {
-	if (bIsOverridden)
+	for (const FCameraRigAssetReferenceParameterMetaData& MetaData : ParameterMetaData)
 	{
-		ParameterOverrideGuids.AddUnique(PropertyID);
+		if (MetaData.bIsOverridden)
+		{
+			OutOverriddenIDs.Add(MetaData.ParameterGuid);
+		}
 	}
-	else
+}
+
+bool FCameraRigAssetReference::IsParameterOverridden(const FGuid& PropertyID) const
+{
+	if (const FCameraRigAssetReferenceParameterMetaData* MetaData = FindMetaData(PropertyID))
+	{
+		return MetaData->bIsOverridden;
+	}
+	return false;
+}
+
+void FCameraRigAssetReference::SetParameterOverridden(const FGuid& PropertyID, bool bIsOverridden)
+{
+	FCameraRigAssetReferenceParameterMetaData* MetaData = FindMetaData(PropertyID);
+	if (ensure(MetaData))
 	{
-		ParameterOverrideGuids.Remove(PropertyID);
+		MetaData->bIsOverridden = bIsOverridden;
 	}
 }
 
@@ -79,23 +110,193 @@ void FCameraRigAssetReference::RebuildParameters()
 {
 	if (CameraRig)
 	{
+		TArray<FGuid> ParameterOverrideGuids;
+		GenerateOverriddenParameterGuidArray(ParameterOverrideGuids);
 		Parameters.MigrateToNewBagInstanceWithOverrides(CameraRig->GetDefaultParameters(), ParameterOverrideGuids);
-		
-		// Remove overrides for parameters that don't exist anymore.
+
 		if (const UPropertyBag* ParametersType = Parameters.GetPropertyBagStruct())
 		{
-			for (TArray<FGuid>::TIterator It = ParameterOverrideGuids.CreateIterator(); It; ++It)
+			// Remove metadata for parameters that don't exist anymore, and add default metadata for
+			// new parameters.
+			TSet<FGuid> ExistingMetaDataIDs;
+			for (TArray<FCameraRigAssetReferenceParameterMetaData>::TIterator It = ParameterMetaData.CreateIterator(); It; ++It)
+			{
+				ExistingMetaDataIDs.Add(It->ParameterGuid);
+			}
+
+			TSet<FGuid> WantedMetaDataIDs;
+			TMap<FGuid, ECameraRigInterfaceParameterType> WantedParameterTypes;
+			for (const FCameraRigParameterDefinition& Definition : CameraRig->GetParameterDefinitions())
 			{
-				if (!ParametersType->FindPropertyDescByID(*It))
+				WantedMetaDataIDs.Add(Definition.ParameterGuid);
+				WantedParameterTypes.Add(Definition.ParameterGuid, Definition.ParameterType);
+			}
+
+			TSet<FGuid> RemovedMetaDataIDs = ExistingMetaDataIDs.Difference(WantedMetaDataIDs);
+			for (TArray<FCameraRigAssetReferenceParameterMetaData>::TIterator It = ParameterMetaData.CreateIterator(); It; ++It)
+			{
+				if (RemovedMetaDataIDs.Contains(It->ParameterGuid))
 				{
 					It.RemoveCurrentSwap();
 				}
 			}
+
+			TSet<FGuid> AddedMetaDataIDs = WantedMetaDataIDs.Difference(ExistingMetaDataIDs);
+			for (const FGuid& ParameterGuid : AddedMetaDataIDs)
+			{
+				FCameraRigAssetReferenceParameterMetaData NewMetaData;
+				NewMetaData.ParameterGuid = ParameterGuid;
+				NewMetaData.ParameterType = WantedParameterTypes.FindChecked(ParameterGuid);
+				ParameterMetaData.Add(NewMetaData);
+			}
+		}
+		else
+		{
+			ParameterMetaData.Reset();
 		}
 	}
 	else
 	{
 		Parameters.Reset();
+		ParameterMetaData.Reset();
+	}
+}
+
+void FCameraRigAssetReference::GetCustomCameraNodeParameters(FCustomCameraNodeParameterInfos& OutParameterInfos)
+{
+	if (!CameraRig)
+	{
+		return;
+	}
+
+	RebuildParametersIfNeeded();
+
+	const UPropertyBag* ParametersStruct = Parameters.GetPropertyBagStruct();
+	const FInstancedPropertyBag& DefaultParameters = CameraRig->GetDefaultParameters();
+	if (!ensure(ParametersStruct && ParametersStruct == DefaultParameters.GetPropertyBagStruct()))
+	{
+		return;
+	}
+
+	for (const FCameraRigParameterDefinition& Definition : CameraRig->GetParameterDefinitions())
+	{
+		const FPropertyBagPropertyDesc* PropertyDesc = ParametersStruct->FindPropertyDescByID(Definition.ParameterGuid);
+		if (!ensure(PropertyDesc))
+		{
+			continue;
+		}
+
+		if (Definition.ParameterType == ECameraRigInterfaceParameterType::Blendable)
+		{
+			if (!ensure(PropertyDesc->ValueType == EPropertyBagPropertyType::Struct))
+			{
+				continue;
+			}
+
+			switch (Definition.VariableType)
+			{
+#define UE_CAMERA_VARIABLE_FOR_TYPE(ValueType, ValueName)\
+				case ECameraVariableType::ValueName:\
+					if (ensure(PropertyDesc->ValueTypeObject == F##ValueName##CameraParameter::StaticStruct()))\
+					{\
+						using CameraParameterType = F##ValueName##CameraParameter;\
+						TValueOrError<CameraParameterType*, EPropertyBagResult> PropertyValue =\
+							Parameters.GetValueStruct<CameraParameterType>(*PropertyDesc);\
+						if (ensure(PropertyValue.HasValue() && !PropertyValue.HasError()))\
+						{\
+							CameraParameterType* CameraParameter = PropertyValue.GetValue();\
+							check(CameraParameter);\
+							OutParameterInfos.AddBlendableParameter(\
+									Definition.ParameterName,\
+									Definition.VariableType,\
+									nullptr,\
+									reinterpret_cast<uint8*>(&CameraParameter->Value),\
+									&CameraParameter->VariableID);\
+						}\
+					}\
+					break;
+				UE_CAMERA_VARIABLE_FOR_ALL_TYPES()
+#undef UE_CAMERA_VARIABLE_FOR_TYPE
+				case ECameraVariableType::BlendableStruct:
+					{
+						TValueOrError<FStructView, EPropertyBagResult> PropertyValue =
+							Parameters.GetValueStruct(*PropertyDesc, Definition.BlendableStructType);
+						if (ensure(PropertyValue.HasValue() && !PropertyValue.HasError()))
+						{
+							FStructView& StructValue = PropertyValue.GetValue();
+							check(StructValue.IsValid());
+
+							FCameraRigAssetReferenceParameterMetaData* MetaData = FindMetaData(Definition.ParameterGuid);
+							check(MetaData);
+							ensure(MetaData->ParameterType == ECameraRigInterfaceParameterType::Blendable);
+
+							OutParameterInfos.AddBlendableParameter(
+									Definition.ParameterName,
+									Definition.VariableType,
+									Definition.BlendableStructType,
+									StructValue.GetMemory(),
+									&MetaData->OverrideVariableID);
+						}
+					}
+					break;
+			}
+		}
+		else if (Definition.ParameterType == ECameraRigInterfaceParameterType::Data)
+		{
+			FCameraRigAssetReferenceParameterMetaData* MetaData = FindMetaData(Definition.ParameterGuid);
+			check(MetaData);
+			ensure(MetaData->ParameterType == ECameraRigInterfaceParameterType::Data);
+
+			switch (Definition.DataType)
+			{
+				case ECameraContextDataType::Name:
+					if (ensure(PropertyDesc->ValueType == EPropertyBagPropertyType::Name))
+					{
+						const uint8* DefaultValue = reinterpret_cast<uint8*>(DefaultParameters.GetValueName(*PropertyDesc).TryGetValue());
+						OutParameterInfos.AddDataParameter(Definition.ParameterName, ECameraContextDataType::Name, nullptr, DefaultValue, &MetaData->OverrideDataID);
+					}
+					break;
+				case ECameraContextDataType::String:
+					if (ensure(PropertyDesc->ValueType == EPropertyBagPropertyType::String))
+					{
+						const uint8* DefaultValue = reinterpret_cast<uint8*>(DefaultParameters.GetValueString(*PropertyDesc).TryGetValue());
+						OutParameterInfos.AddDataParameter(Definition.ParameterName, ECameraContextDataType::String, nullptr, DefaultValue, &MetaData->OverrideDataID);
+					}
+					break;
+				case ECameraContextDataType::Enum:
+					if (ensure(PropertyDesc->ValueType == EPropertyBagPropertyType::Enum &&
+								PropertyDesc->ValueTypeObject == Definition.DataTypeObject))
+					{
+						const UEnum* EnumType = CastChecked<const UEnum>(Definition.DataTypeObject);
+						const uint8* DefaultValue = reinterpret_cast<uint8*>(DefaultParameters.GetValueEnum(*PropertyDesc, EnumType).TryGetValue());
+						OutParameterInfos.AddDataParameter(Definition.ParameterName, ECameraContextDataType::Enum, EnumType, DefaultValue, &MetaData->OverrideDataID);
+					}
+					break;
+				case ECameraContextDataType::Struct:
+					if (ensure(PropertyDesc->ValueType == EPropertyBagPropertyType::Struct &&
+								PropertyDesc->ValueTypeObject == Definition.DataTypeObject))
+					{
+						const UScriptStruct* DataType = CastChecked<const UScriptStruct>(Definition.DataTypeObject);
+						const uint8* DefaultValue = reinterpret_cast<uint8*>(DefaultParameters.GetValueStruct(*PropertyDesc, DataType).TryGetValue());
+						OutParameterInfos.AddDataParameter(Definition.ParameterName, ECameraContextDataType::Struct, DataType, DefaultValue, &MetaData->OverrideDataID);
+					}
+					break;
+				case ECameraContextDataType::Object:
+					if (ensure(PropertyDesc->ValueType == EPropertyBagPropertyType::Object))
+					{
+						const uint8* DefaultValue = reinterpret_cast<uint8*>(DefaultParameters.GetValueObject(*PropertyDesc).TryGetValue());
+						OutParameterInfos.AddDataParameter(Definition.ParameterName, ECameraContextDataType::Object, Definition.DataTypeObject, DefaultValue, &MetaData->OverrideDataID);
+					}
+					break;
+				case ECameraContextDataType::Class:
+					if (ensure(PropertyDesc->ValueType == EPropertyBagPropertyType::Class))
+					{
+						const uint8* DefaultValue = reinterpret_cast<uint8*>(DefaultParameters.GetValueClass(*PropertyDesc).TryGetValue());
+						OutParameterInfos.AddDataParameter(Definition.ParameterName, ECameraContextDataType::Class, Definition.DataTypeObject, DefaultValue, &MetaData->OverrideDataID);
+					}
+					break;
+			}
+		}
 	}
 }
 
@@ -118,7 +319,7 @@ void FCameraRigAssetReference::PostSerialize(const FArchive& Ar)
 	// Make a property bag with the legacy overrides, and then set the values in it.
 	bool bHasAnyLegacyOverride = false;
 	TArray<FPropertyBagPropertyDesc> LegacyParameterProperties;
-	TArray<FGuid> LegacyParameterOverrides;
+	TArray<FCameraRigAssetReferenceParameterMetaData> LegacyParameterMetaData;
 
 #define UE_CAMERA_VARIABLE_FOR_TYPE(ValueType, ValueName)\
 	for (F##ValueName##CameraRigParameterOverride& ParameterOverride : ParameterOverrides_DEPRECATED.ValueName##Overrides)\
@@ -129,7 +330,11 @@ void FCameraRigAssetReference::PostSerialize(const FArchive& Ar)
 		FPropertyBagPropertyDesc LegacyParameterProperty(PropertyName, PropertyType, PropertyTypeObject);\
 		LegacyParameterProperty.ID = ParameterOverride.InterfaceParameterGuid;\
 		LegacyParameterProperties.Add(LegacyParameterProperty);\
-		LegacyParameterOverrides.Add(LegacyParameterProperty.ID);\
+		FCameraRigAssetReferenceParameterMetaData MetaData;\
+		MetaData.ParameterGuid = ParameterOverride.InterfaceParameterGuid;\
+		MetaData.ParameterType = ECameraRigInterfaceParameterType::Blendable;\
+		MetaData.bIsOverridden = true;\
+		LegacyParameterMetaData.Add(MetaData);\
 		bHasAnyLegacyOverride = true;\
 	}
 	UE_CAMERA_VARIABLE_FOR_ALL_TYPES()
@@ -140,7 +345,7 @@ void FCameraRigAssetReference::PostSerialize(const FArchive& Ar)
 		Parameters = FInstancedPropertyBag();
 		Parameters.AddProperties(LegacyParameterProperties);
 
-		ParameterOverrideGuids = LegacyParameterOverrides;
+		ParameterMetaData = LegacyParameterMetaData;
 
 #define UE_CAMERA_VARIABLE_FOR_TYPE(ValueType, ValueName)\
 		for (F##ValueName##CameraRigParameterOverride& ParameterOverride : ParameterOverrides_DEPRECATED.ValueName##Overrides)\
diff --git a/Engine/Plugins/Cameras/GameplayCameras/Source/GameplayCameras/Private/Core/CameraRigJoints.cpp b/Engine/Plugins/Cameras/GameplayCameras/Source/GameplayCameras/Private/Core/CameraRigJoints.cpp
--- a/Engine/Plugins/Cameras/GameplayCameras/Source/GameplayCameras/Private/Core/CameraRigJoints.cpp
+++ b/Engine/Plugins/Cameras/GameplayCameras/Source/GameplayCameras/Private/Core/CameraRigJoints.cpp
@@ -2,6 +2,7 @@
 
 #include "Core/CameraRigJoints.h"
 
+#include "Core/CameraVariableTableAllocationInfo.h"
 #include "Core/BuiltInCameraVariables.h"
 
 namespace UE::Cameras
diff --git a/Engine/Plugins/Cameras/GameplayCameras/Source/GameplayCameras/Private/Core/CameraRigParameterDefinition.cpp b/Engine/Plugins/Cameras/GameplayCameras/Source/GameplayCameras/Private/Core/CameraRigParameterDefinition.cpp
new file mode 100644
--- /dev/null
+++ b/Engine/Plugins/Cameras/GameplayCameras/Source/GameplayCameras/Private/Core/CameraRigParameterDefinition.cpp
@@ -0,0 +1,286 @@
+// Copyright Epic Games, Inc. All Rights Reserved.
+
+#include "Core/CameraRigParameterDefinition.h"
+
+#include "Core/CameraNode.h"
+#include "Core/CameraParameters.h"
+#include "Core/CameraRigAsset.h"
+#include "Core/ICustomCameraNodeParameterProvider.h"
+#include "StructUtils/PropertyBag.h"
+
+namespace UE::Cameras
+{
+
+void FCameraRigParameterBuilder::BuildDefaultParameters(const UCameraRigAsset* CameraRig, FInstancedPropertyBag& OutPropertyBag)
+{
+	TArray<FPropertyBagPropertyDesc> DefaultParameterProperties;
+	AppendDefaultParameterProperties(CameraRig, DefaultParameterProperties);
+	OutPropertyBag.AddProperties(DefaultParameterProperties);
+	SetDefaultParameterValues(CameraRig, OutPropertyBag);
+}
+
+void FCameraRigParameterBuilder::AppendDefaultParameterProperties(const UCameraRigAsset* CameraRig, TArray<FPropertyBagPropertyDesc>& OutProperties)
+{
+	for (const FCameraRigParameterDefinition& Definition : CameraRig->GetParameterDefinitions())
+	{
+		bool bIsValidProperty = true;
+		EPropertyBagPropertyType PropertyType = EPropertyBagPropertyType::Struct;
+		const UObject* PropertyTypeObject = nullptr;
+
+		if (Definition.ParameterType == ECameraRigInterfaceParameterType::Blendable)
+		{
+			switch (Definition.VariableType)
+			{
+#define UE_CAMERA_VARIABLE_FOR_TYPE(ValueType, ValueName)\
+				case ECameraVariableType::ValueName:\
+					PropertyTypeObject = F##ValueName##CameraParameter::StaticStruct();\
+					break;
+				UE_CAMERA_VARIABLE_FOR_ALL_TYPES()
+#undef UE_CAMERA_VARIABLE_FOR_TYPE
+				case ECameraVariableType::BlendableStruct:
+					PropertyTypeObject = Definition.BlendableStructType;
+					break;
+				default:
+					ensure(false);
+					break;
+			}
+		}
+		else if (Definition.ParameterType == ECameraRigInterfaceParameterType::Data)
+		{
+			PropertyTypeObject = Definition.DataTypeObject;
+
+			switch (Definition.DataType)
+			{
+				case ECameraContextDataType::Name:
+					PropertyType = EPropertyBagPropertyType::Name;
+					break;
+				case ECameraContextDataType::String:
+					PropertyType = EPropertyBagPropertyType::String;
+					break;
+				case ECameraContextDataType::Enum:
+					PropertyType = EPropertyBagPropertyType::Enum;
+					ensure(PropertyTypeObject && PropertyTypeObject->IsA<UEnum>());
+					break;
+				case ECameraContextDataType::Struct:
+					PropertyType = EPropertyBagPropertyType::Struct;
+					ensure(PropertyTypeObject && PropertyTypeObject->IsA<UScriptStruct>());
+					break;
+				case ECameraContextDataType::Object:
+					PropertyType = EPropertyBagPropertyType::Object;
+					break;
+				case ECameraContextDataType::Class:
+					PropertyType = EPropertyBagPropertyType::Class;
+					break;
+				default:
+					bIsValidProperty = false;
+					break;
+			}
+		}
+		else
+		{
+			bIsValidProperty = false;
+		}
+
+		if (ensure(bIsValidProperty))
+		{
+			FPropertyBagPropertyDesc NewProperty(Definition.ParameterName, PropertyType, PropertyTypeObject);
+			// Make the property bag match the camera interface parameter GUIDs.
+			NewProperty.ID = Definition.ParameterGuid;
+
+			OutProperties.Add(NewProperty);
+		}
+	}
+}
+
+void FCameraRigParameterBuilder::SetDefaultParameterValues(const UCameraRigAsset* CameraRig, FInstancedPropertyBag& PropertyBag)
+{
+	uint8* PropertyBagValue = PropertyBag.GetMutableValue().GetMemory();
+	const UPropertyBag* PropertyBagStruct = PropertyBag.GetPropertyBagStruct();
+	if (!ensure(PropertyBagValue && PropertyBagStruct))
+	{
+		return;
+	}
+
+	for (const UCameraRigBlendableParameter* BlendableParameter : CameraRig->Interface.BlendableParameters)
+	{
+		if (!ensure(BlendableParameter))
+		{
+			continue;
+		}
+
+		UCameraNode* CameraNode = BlendableParameter->Target;
+		if (!CameraNode)
+		{
+			continue;
+		}
+
+		const uint8* RawSourceValuePtr = nullptr;
+
+		if (ICustomCameraNodeParameterProvider* CustomParameterProvider = Cast<ICustomCameraNodeParameterProvider>(CameraNode))
+		{
+			FCustomCameraNodeParameterInfos CustomParameters;
+			CustomParameterProvider->GetCustomCameraNodeParameters(CustomParameters);
+
+			FCustomCameraNodeParameterInfos::FBlendableParameterInfo* TargetCustomParameter = 
+				CustomParameters.BlendableParameters.FindByPredicate(
+						[BlendableParameter](FCustomCameraNodeParameterInfos::FBlendableParameterInfo& CustomParameter)
+						{
+							return CustomParameter.ParameterName == BlendableParameter->TargetPropertyName;
+						});
+			if (TargetCustomParameter)
+			{
+				RawSourceValuePtr = TargetCustomParameter->DefaultValue;
+			}
+		}
+
+		if (!RawSourceValuePtr)
+		{
+			const UClass* TargetClass = CameraNode->GetClass();
+			FStructProperty* StructProperty = CastField<FStructProperty>(
+					TargetClass->FindPropertyByName(BlendableParameter->TargetPropertyName));
+			if (StructProperty)
+			{
+				switch (BlendableParameter->ParameterType)
+				{
+#define UE_CAMERA_VARIABLE_FOR_TYPE(ValueType, ValueName)\
+					case ECameraVariableType::ValueName:\
+						{\
+							using CameraParameterType = F##ValueName##CameraParameter;\
+							if (StructProperty->Struct == CameraParameterType::StaticStruct())\
+							{\
+								auto* CameraParameterPtr = StructProperty->ContainerPtrToValuePtr<CameraParameterType>(CameraNode);\
+								RawSourceValuePtr = reinterpret_cast<uint8*>(&CameraParameterPtr->Value);\
+							}\
+						}\
+						break;
+					UE_CAMERA_VARIABLE_FOR_ALL_TYPES()
+#undef UE_CAMERA_VARIABLE_FOR_TYPE
+				}
+			}
+		}
+
+		if (!ensure(RawSourceValuePtr))
+		{
+			continue;
+		}
+
+		const FPropertyBagPropertyDesc* PropertyDesc = PropertyBagStruct->FindPropertyDescByID(BlendableParameter->GetGuid());
+		if (!ensure(PropertyDesc && PropertyDesc->CachedProperty))
+		{
+			continue;
+		}
+
+		void* RawDestinationValuePtr = PropertyDesc->CachedProperty->ContainerPtrToValuePtr<void>(PropertyBagValue);
+
+		if (!ensure(RawDestinationValuePtr))
+		{
+			continue;
+		}
+
+		switch (BlendableParameter->ParameterType)
+		{
+#define UE_CAMERA_VARIABLE_FOR_TYPE(ValueType, ValueName)\
+			case ECameraVariableType::ValueName:\
+				{\
+					using CameraParameterType = F##ValueName##CameraParameter;\
+					const ValueType* SourceValuePtr = reinterpret_cast<const ValueType*>(RawSourceValuePtr);\
+					CameraParameterType* DestinationParameter = reinterpret_cast<CameraParameterType*>(RawDestinationValuePtr);\
+					DestinationParameter->Value = *SourceValuePtr;\
+				}\
+				break;
+			UE_CAMERA_VARIABLE_FOR_ALL_TYPES()
+#undef UE_CAMERA_VARIABLE_FOR_TYPE
+			case ECameraVariableType::BlendableStruct:
+				if (ensure(BlendableParameter->BlendableStructType))
+				{
+					BlendableParameter->BlendableStructType->CopyScriptStruct(RawDestinationValuePtr, RawSourceValuePtr);
+				}
+				break;
+		}
+	}
+
+	for (const UCameraRigDataParameter* DataParameter : CameraRig->Interface.DataParameters)
+	{
+		if (!ensure(DataParameter && DataParameter->Target))
+		{
+			continue;
+		}
+
+		UCameraNode* CameraNode = DataParameter->Target;
+
+		const uint8* RawSourceValuePtr = nullptr;
+
+		if (ICustomCameraNodeParameterProvider* CustomParameterProvider = Cast<ICustomCameraNodeParameterProvider>(CameraNode))
+		{
+			FCustomCameraNodeParameterInfos CustomParameters;
+			CustomParameterProvider->GetCustomCameraNodeParameters(CustomParameters);
+
+			FCustomCameraNodeParameterInfos::FDataParameterInfo* TargetCustomParameter = 
+				CustomParameters.DataParameters.FindByPredicate(
+						[DataParameter](FCustomCameraNodeParameterInfos::FDataParameterInfo& CustomParameter)
+						{
+							return CustomParameter.ParameterName == DataParameter->TargetPropertyName;
+						});
+			if (TargetCustomParameter)
+			{
+				RawSourceValuePtr = TargetCustomParameter->DefaultValue;
+			}
+		}
+
+		if (!RawSourceValuePtr)
+		{
+			const UClass* TargetClass = DataParameter->Target->GetClass();
+			FProperty* Property = TargetClass->FindPropertyByName(DataParameter->TargetPropertyName);
+			if (Property)
+			{
+				RawSourceValuePtr = (uint8*)Property->ContainerPtrToValuePtr<void>(TargetClass);
+			}
+		}
+
+		if (!ensure(RawSourceValuePtr))
+		{
+			continue;
+		}
+
+		const FPropertyBagPropertyDesc* PropertyDesc = PropertyBagStruct->FindPropertyDescByID(DataParameter->GetGuid());
+		if (!ensure(PropertyDesc && PropertyDesc->CachedProperty))
+		{
+			continue;
+		}
+
+		void* RawDestinationValuePtr = PropertyDesc->CachedProperty->ContainerPtrToValuePtr<void>(PropertyBagValue);
+
+		if (!ensure(RawDestinationValuePtr))
+		{
+			continue;
+		}
+
+		switch (DataParameter->DataType)
+		{
+			case ECameraContextDataType::Name:
+				*((FName*)RawDestinationValuePtr) = *((FName*)RawSourceValuePtr);
+				break;
+			case ECameraContextDataType::String:
+				*((FString*)RawDestinationValuePtr) = *((FString*)RawSourceValuePtr);
+				break;
+			case ECameraContextDataType::Enum:
+				*((uint8*)RawDestinationValuePtr) = *((uint8*)RawSourceValuePtr);
+				break;
+			case ECameraContextDataType::Struct:
+				{
+					const UScriptStruct* StructType = CastChecked<const UScriptStruct>(DataParameter->DataTypeObject);
+					StructType->CopyScriptStruct(RawDestinationValuePtr, RawSourceValuePtr);
+				}
+				break;
+			case ECameraContextDataType::Object:
+				*((FObjectPtr*)RawDestinationValuePtr) = *((FObjectPtr*)RawSourceValuePtr);
+				break;
+			case ECameraContextDataType::Class:
+				*((FObjectPtr*)RawDestinationValuePtr) = *((FObjectPtr*)RawSourceValuePtr);
+				break;
+		}
+	}
+}
+
+}  // namespace UE::Cameras
+
diff --git a/Engine/Plugins/Cameras/GameplayCameras/Source/GameplayCameras/Private/Core/CameraVariableTable.cpp b/Engine/Plugins/Cameras/GameplayCameras/Source/GameplayCameras/Private/Core/CameraVariableTable.cpp
--- a/Engine/Plugins/Cameras/GameplayCameras/Source/GameplayCameras/Private/Core/CameraVariableTable.cpp
+++ b/Engine/Plugins/Cameras/GameplayCameras/Source/GameplayCameras/Private/Core/CameraVariableTable.cpp
@@ -3,6 +3,7 @@
 #include "Core/CameraVariableTable.h"
 
 #include "HAL/UnrealMemory.h"
+#include "IGameplayCamerasModule.h"
 #include "Math/UnrealMath.h"
 
 namespace UE::Cameras
@@ -25,6 +26,9 @@ bool IsVariableInMask(FCameraVariableID VariableID, const FCameraVariableTableFl
 
 }  // namespace Private
 
+TArray<FBlendableStructInfo> FCameraVariableTable::CachedBlendableStructs;
+bool FCameraVariableTable::bCachedBlendableStructs(false);
+
 FCameraVariableTable::FCameraVariableTable()
 {
 }
@@ -83,7 +87,7 @@ void FCameraVariableTable::Initialize(const FCameraVariableTableAllocationInfo&
 	uint32 CurSizeOf, CurAlignOf;
 	for (const FCameraVariableDefinition& VariableDefinition : AllocationInfo.VariableDefinitions)
 	{
-		GetVariableTypeAllocationInfo(VariableDefinition.VariableType, CurSizeOf, CurAlignOf);
+		GetVariableTypeAllocationInfo(VariableDefinition.VariableType, VariableDefinition.BlendableStructType, CurSizeOf, CurAlignOf);
 		const uint32 NewEntryOffset = Align(TotalSizeOf, CurAlignOf);
 		TotalSizeOf = NewEntryOffset + CurSizeOf;
 		MaxAlignOf = FMath::Max(MaxAlignOf, CurAlignOf);
@@ -91,6 +95,7 @@ void FCameraVariableTable::Initialize(const FCameraVariableTableAllocationInfo&
 		FEntry NewEntry;
 		NewEntry.ID = VariableDefinition.VariableID;
 		NewEntry.Type = VariableDefinition.VariableType;
+		NewEntry.StructType = VariableDefinition.BlendableStructType;
 		NewEntry.Offset = NewEntryOffset;
 		NewEntry.Flags = EEntryFlags::None;
 		if (VariableDefinition.bIsPrivate)
@@ -133,14 +138,46 @@ void FCameraVariableTable::Initialize(const FCameraVariableTableAllocationInfo&
 				break;
 UE_CAMERA_VARIABLE_FOR_ALL_TYPES()
 #undef UE_CAMERA_VARIABLE_FOR_TYPE
+			case ECameraVariableType::BlendableStruct:
+				{
+					check(Entry.StructType);
+					Entry.StructType->InitializeStruct(ValuePtr);
+				}
+				break;
+		}
+	}
+}
+
+void FCameraVariableTable::CacheBlendableStructs()
+{
+	if (bCachedBlendableStructs)
+	{
+		return;
+	}
+
+	IGameplayCamerasModule& GameplayCamerasModule = IGameplayCamerasModule::Get();
+	CachedBlendableStructs = GameplayCamerasModule.GetBlendableStructs();
+	bCachedBlendableStructs = true;
+}
+
+FBlendableStructTypeErasedInterpolator FCameraVariableTable::GetBlendableStructInterpolator(const UScriptStruct* StructType)
+{
+	CacheBlendableStructs();
+
+	for (const FBlendableStructInfo& Info : CachedBlendableStructs)
+	{
+		if (Info.StructType == StructType)
+		{
+			return Info.Interpolator;
 		}
 	}
+	return nullptr;
 }
 
 void FCameraVariableTable::AddVariable(const FCameraVariableDefinition& VariableDefinition)
 {
 	uint32 SizeOf, AlignOf;
-	GetVariableTypeAllocationInfo(VariableDefinition.VariableType, SizeOf, AlignOf);
+	GetVariableTypeAllocationInfo(VariableDefinition.VariableType, VariableDefinition.BlendableStructType, SizeOf, AlignOf);
 
 	uint8* VariablePtr = Align(Memory + Used, AlignOf);
 	uint32 NewUsed = (VariablePtr + SizeOf) - Memory;
@@ -157,6 +194,7 @@ void FCameraVariableTable::AddVariable(const FCameraVariableDefinition& Variable
 	FEntry NewEntry;
 	NewEntry.ID = VariableDefinition.VariableID;
 	NewEntry.Type = VariableDefinition.VariableType;
+	NewEntry.StructType = VariableDefinition.BlendableStructType;
 	NewEntry.Offset = VariablePtr - Memory;
 	NewEntry.Flags = EEntryFlags::None;
 	if (VariableDefinition.bIsPrivate)
@@ -215,7 +253,7 @@ const FCameraVariableTable::FEntry* FCameraVariableTable::FindEntry(FCameraVaria
 	return IndexPtr ? &Entries[*IndexPtr] : nullptr;
 }
 
-bool FCameraVariableTable::GetVariableTypeAllocationInfo(ECameraVariableType VariableType, uint32& OutSizeOf, uint32& OutAlignOf)
+bool FCameraVariableTable::GetVariableTypeAllocationInfo(ECameraVariableType VariableType, const UScriptStruct* StructType, uint32& OutSizeOf, uint32& OutAlignOf)
 {
 	switch (VariableType)
 	{
@@ -226,6 +264,14 @@ bool FCameraVariableTable::GetVariableTypeAllocationInfo(ECameraVariableType Var
 			return true;
 UE_CAMERA_VARIABLE_FOR_ALL_TYPES()
 #undef UE_CAMERA_VARIABLE_FOR_TYPE
+		case ECameraVariableType::BlendableStruct:
+			{
+				check(StructType);
+				const UScriptStruct::ICppStructOps* StructOps = StructType->GetCppStructOps();
+				OutSizeOf = StructOps->GetSize();
+				OutAlignOf = StructOps->GetAlignment();
+			}
+			return true;
 	}
 	return false;
 }
@@ -235,19 +281,25 @@ bool FCameraVariableTable::ContainsValue(FCameraVariableID VariableID) const
 	return EntryLookup.Contains(VariableID);
 }
 
-const uint8* FCameraVariableTable::GetValue(FCameraVariableID VariableID, ECameraVariableType ExpectedVariableType) const
+const uint8* FCameraVariableTable::GetValue(
+		FCameraVariableID VariableID,
+		ECameraVariableType ExpectedVariableType,
+		const UScriptStruct* ExpectedBlendableStructType) const
 {
-	const uint8* Value = TryGetValue(VariableID, ExpectedVariableType);
+	const uint8* Value = TryGetValue(VariableID, ExpectedVariableType, ExpectedBlendableStructType);
 	ensureMsgf(Value, TEXT("Can't get camera variable (ID '%d') because it doesn't exist in the table."), VariableID.GetValue());
 	return Value;
 }
 
-const uint8* FCameraVariableTable::TryGetValue(FCameraVariableID VariableID, ECameraVariableType ExpectedVariableType) const
+const uint8* FCameraVariableTable::TryGetValue(
+		FCameraVariableID VariableID,
+		ECameraVariableType ExpectedVariableType,
+		const UScriptStruct* ExpectedBlendableStructType) const
 {
 	const FEntry* Entry = FindEntry(VariableID);
 	if (Entry)
 	{
-		ensure(Entry->Type == ExpectedVariableType);
+		ensure(Entry->Type == ExpectedVariableType && Entry->StructType == ExpectedBlendableStructType);
 		return Memory + Entry->Offset;
 	}
 
@@ -257,15 +309,16 @@ const uint8* FCameraVariableTable::TryGetValue(FCameraVariableID VariableID, ECa
 void FCameraVariableTable::SetValue(
 		FCameraVariableID VariableID, 
 		ECameraVariableType ExpectedVariableType, 
+		const UScriptStruct* ExpectedBlendableStructType,
 		const uint8* InRawValuePtr,
 		bool bMarkAsWrittenThisFrame)
 {
 	FEntry* Entry = FindEntry(VariableID);
 	if (ensureMsgf(Entry, TEXT("Can't set camera variable (ID '%d') because it doesn't exist in the table."), VariableID.GetValue()))
 	{
-		check(ExpectedVariableType == Entry->Type);
+		check(ExpectedVariableType == Entry->Type && ExpectedBlendableStructType == Entry->StructType);
 		uint32 SizeOf, AlignOf;
-		GetVariableTypeAllocationInfo(Entry->Type, SizeOf, AlignOf);
+		GetVariableTypeAllocationInfo(Entry->Type, Entry->StructType, SizeOf, AlignOf);
 		uint8* ValuePtr = Memory + Entry->Offset;
 		FMemory::Memcpy(ValuePtr, InRawValuePtr, SizeOf);
 		Entry->Flags |= EEntryFlags::Written;
@@ -279,14 +332,15 @@ void FCameraVariableTable::SetValue(
 bool FCameraVariableTable::TrySetValue(
 		FCameraVariableID VariableID,
 		ECameraVariableType ExpectedVariableType,
+		const UScriptStruct* ExpectedBlendableStructType,
 		const uint8* InRawValuePtr,
 		bool bMarkAsWrittenThisFrame)
 {
 	if (FEntry* Entry = FindEntry(VariableID))
 	{
-		check(ExpectedVariableType == Entry->Type);
+		check(ExpectedVariableType == Entry->Type && ExpectedBlendableStructType == Entry->StructType);
 		uint32 SizeOf, AlignOf;
-		GetVariableTypeAllocationInfo(Entry->Type, SizeOf, AlignOf);
+		GetVariableTypeAllocationInfo(Entry->Type, Entry->StructType, SizeOf, AlignOf);
 		uint8* ValuePtr = Memory + Entry->Offset;
 		FMemory::Memcpy(ValuePtr, InRawValuePtr, SizeOf);
 		Entry->Flags |= EEntryFlags::Written;
@@ -358,6 +412,7 @@ bool FCameraVariableTable::TryGetVariableDefinition(FCameraVariableID VariableID
 	{
 		OutVariableDefinition.VariableID = Entry->ID;
 		OutVariableDefinition.VariableType = Entry->Type;
+		OutVariableDefinition.BlendableStructType = Entry->StructType;
 		OutVariableDefinition.bIsPrivate = EnumHasAnyFlags(Entry->Flags, EEntryFlags::Private);
 		OutVariableDefinition.bIsInput = EnumHasAnyFlags(Entry->Flags, EEntryFlags::Input);
 #if WITH_EDITORONLY_DATA
@@ -471,7 +526,7 @@ void FCameraVariableTable::InternalOverride(const FCameraVariableTable& OtherTab
 #else
 				const FString& DebugName = GUnavailableVariableDebugName;
 #endif
-				checkf(ThisEntry->Type == OtherEntry.Type, 
+				checkf(ThisEntry->Type == OtherEntry.Type && ThisEntry->StructType == OtherEntry.StructType, 
 						TEXT("Camera variable name collision! Expected '%d' (%s) to be of type '%s' but other table has type '%s'!"),
 						ThisEntry->ID.GetValue(), *DebugName,
 						*UEnum::GetValueAsString(ThisEntry->Type), *UEnum::GetValueAsString(OtherEntry.Type));
@@ -482,6 +537,7 @@ void FCameraVariableTable::InternalOverride(const FCameraVariableTable& OtherTab
 				FCameraVariableDefinition NewVariableDefinition;
 				NewVariableDefinition.VariableID = OtherEntry.ID;
 				NewVariableDefinition.VariableType = OtherEntry.Type;
+				NewVariableDefinition.BlendableStructType = OtherEntry.StructType;
 				NewVariableDefinition.bIsInput = EnumHasAllFlags(OtherEntry.Flags, EEntryFlags::Input);
 #if WITH_EDITORONLY_DATA
 				NewVariableDefinition.VariableName = OtherEntry.DebugName;
@@ -494,7 +550,7 @@ void FCameraVariableTable::InternalOverride(const FCameraVariableTable& OtherTab
 			if (ensure(ThisEntry))
 			{
 				uint32 ValueSize = 0, ValueAlignment = 0;
-				GetVariableTypeAllocationInfo(ThisEntry->Type, ValueSize, ValueAlignment);
+				GetVariableTypeAllocationInfo(ThisEntry->Type, ThisEntry->StructType, ValueSize, ValueAlignment);
 				check(ValueSize != 0);
 
 				uint8* ThisValuePtr = Memory + ThisEntry->Offset;
@@ -584,6 +640,16 @@ void FCameraVariableTable::InternalLerp(const FCameraVariableTable& ToTable, ECa
 						break;
 UE_CAMERA_VARIABLE_FOR_ALL_TYPES()
 #undef UE_CAMERA_VARIABLE_FOR_TYPE
+					case ECameraVariableType::BlendableStruct:
+						{
+							check(FromEntry->StructType);
+							FBlendableStructTypeErasedInterpolator Interpolator = GetBlendableStructInterpolator(FromEntry->StructType);
+							if (ensure(Interpolator))
+							{
+								Interpolator(FromValuePtr, ToValuePtr, Factor);
+							}
+						}
+						break;
 				}
 				// We consider this variable "written to this frame" if it was written in either variable tables this frame.
 				// If the value interpolates because the from/to values are different, but neither was written this frame, we
@@ -597,6 +663,7 @@ UE_CAMERA_VARIABLE_FOR_ALL_TYPES()
 				FCameraVariableDefinition NewVariableDefinition;
 				NewVariableDefinition.VariableID = ToEntry.ID;
 				NewVariableDefinition.VariableType = ToEntry.Type;
+				NewVariableDefinition.BlendableStructType = ToEntry.StructType;
 				NewVariableDefinition.bIsInput = EnumHasAllFlags(ToEntry.Flags, EEntryFlags::Input);
 #if WITH_EDITORONLY_DATA
 				NewVariableDefinition.VariableName = ToEntry.DebugName;
@@ -607,7 +674,7 @@ UE_CAMERA_VARIABLE_FOR_ALL_TYPES()
 				check(FromEntry);
 
 				uint32 ValueSize = 0, ValueAlignment = 0;
-				GetVariableTypeAllocationInfo(FromEntry->Type, ValueSize, ValueAlignment);
+				GetVariableTypeAllocationInfo(FromEntry->Type, FromEntry->StructType, ValueSize, ValueAlignment);
 				check(ValueSize != 0);
 
 				uint8* FromValuePtr = Memory + FromEntry->Offset;
@@ -626,44 +693,3 @@ UE_CAMERA_VARIABLE_FOR_ALL_TYPES()
 
 }  // namespace UE::Cameras
 
-bool operator==(const FCameraVariableDefinition& A, const FCameraVariableDefinition& B)
-{
-	return A.VariableID == B.VariableID
-		&& A.VariableType == B.VariableType
-		&& A.bIsPrivate == B.bIsPrivate
-		&& A.bIsInput == B.bIsInput
-#if WITH_EDITORONLY_DATA
-		&& A.VariableName == B.VariableName
-#endif
-		;
-}
-
-void FCameraVariableTableAllocationInfo::Combine(const FCameraVariableTableAllocationInfo& OtherInfo)
-{
-	TMap<FCameraVariableID, int32> KnownIDs;
-	for (auto It = VariableDefinitions.CreateConstIterator(); It; ++It)
-	{
-		const FCameraVariableDefinition& VariableDefinition(*It);
-		KnownIDs.Add(VariableDefinition.VariableID, It.GetIndex());
-	}
-
-	for (const FCameraVariableDefinition& OtherVariableDefinition : OtherInfo.VariableDefinitions)
-	{
-		const int32 KnownIndex = KnownIDs.FindRef(OtherVariableDefinition.VariableID, INDEX_NONE);
-		if (KnownIndex == INDEX_NONE)
-		{
-			VariableDefinitions.Add(OtherVariableDefinition);
-		}
-		else
-		{
-			const FCameraVariableDefinition& KnownVariableDefinition(VariableDefinitions[KnownIndex]);
-			ensure(KnownVariableDefinition == OtherVariableDefinition);
-		}
-	}
-}
-
-bool operator==(const FCameraVariableTableAllocationInfo& A, const FCameraVariableTableAllocationInfo& B)
-{
-	return A.VariableDefinitions == B.VariableDefinitions;
-}
-
diff --git a/Engine/Plugins/Cameras/GameplayCameras/Source/GameplayCameras/Private/Core/CameraVariableTableAllocationInfo.cpp b/Engine/Plugins/Cameras/GameplayCameras/Source/GameplayCameras/Private/Core/CameraVariableTableAllocationInfo.cpp
new file mode 100644
--- /dev/null
+++ b/Engine/Plugins/Cameras/GameplayCameras/Source/GameplayCameras/Private/Core/CameraVariableTableAllocationInfo.cpp
@@ -0,0 +1,28 @@
+// Copyright Epic Games, Inc. All Rights Reserved.
+
+#include "Core/CameraVariableTableAllocationInfo.h"
+
+void FCameraVariableTableAllocationInfo::Combine(const FCameraVariableTableAllocationInfo& OtherInfo)
+{
+	TMap<FCameraVariableID, int32> KnownIDs;
+	for (auto It = VariableDefinitions.CreateConstIterator(); It; ++It)
+	{
+		const FCameraVariableDefinition& VariableDefinition(*It);
+		KnownIDs.Add(VariableDefinition.VariableID, It.GetIndex());
+	}
+
+	for (const FCameraVariableDefinition& OtherVariableDefinition : OtherInfo.VariableDefinitions)
+	{
+		const int32 KnownIndex = KnownIDs.FindRef(OtherVariableDefinition.VariableID, INDEX_NONE);
+		if (KnownIndex == INDEX_NONE)
+		{
+			VariableDefinitions.Add(OtherVariableDefinition);
+		}
+		else
+		{
+			const FCameraVariableDefinition& KnownVariableDefinition(VariableDefinitions[KnownIndex]);
+			ensure(KnownVariableDefinition == OtherVariableDefinition);
+		}
+	}
+}
+
diff --git a/Engine/Plugins/Cameras/GameplayCameras/Source/GameplayCameras/Private/Core/ICustomCameraNodeParameterProvider.cpp b/Engine/Plugins/Cameras/GameplayCameras/Source/GameplayCameras/Private/Core/ICustomCameraNodeParameterProvider.cpp
--- a/Engine/Plugins/Cameras/GameplayCameras/Source/GameplayCameras/Private/Core/ICustomCameraNodeParameterProvider.cpp
+++ b/Engine/Plugins/Cameras/GameplayCameras/Source/GameplayCameras/Private/Core/ICustomCameraNodeParameterProvider.cpp
@@ -4,43 +4,44 @@
 
 #include "GameplayCamerasDelegates.h"
 
-bool operator==(const FCustomCameraNodeBlendableParameter& A, const FCustomCameraNodeBlendableParameter& B)
-{
-	return A.ParameterName == B.ParameterName &&
-		A.ParameterType == B.ParameterType &&
-		A.OverrideVariable == B.OverrideVariable;
-}
-
-bool operator==(const FCustomCameraNodeDataParameter& A, const FCustomCameraNodeDataParameter& B)
-{
-	return A.ParameterName == B.ParameterName &&
-		A.ParameterType == B.ParameterType &&
-		A.ParameterTypeObject == B.ParameterTypeObject &&
-		A.OverrideDataID == B.OverrideDataID;
-}
-
-bool operator==(const FCustomCameraNodeParameters& A, const FCustomCameraNodeParameters& B)
-{
-	return A.BlendableParameters == B.BlendableParameters &&
-		A.DataParameters == B.DataParameters;
-}
-
 void FCustomCameraNodeParameterInfos::AddBlendableParameter(
 		FName ParameterName, 
 		ECameraVariableType ParameterType, 
-		const uint8* DefaultValuePtr,
-		TObjectPtr<UCameraVariableAsset>* OverrideVariable)
+		const UScriptStruct* BlendableStructType,
+		const uint8* DefaultValue,
+		FCameraVariableID* OverrideVariableID)
+{
+	BlendableParameters.Add({ ParameterName, ParameterType, BlendableStructType, DefaultValue, OverrideVariableID });
+}
+
+void FCustomCameraNodeParameterInfos::AddBlendableParameter(FCustomCameraNodeBlendableParameter& Parameter, const uint8* DefaultValue)
 {
-	BlendableParameters.Add({ ParameterName, ParameterType, DefaultValuePtr, OverrideVariable });
+	AddBlendableParameter(
+			Parameter.ParameterName,
+			Parameter.ParameterType,
+			Parameter.BlendableStructType,
+			DefaultValue,
+			Parameter.OverrideVariable ? nullptr : &Parameter.OverrideVariableID);
 }
 
 void FCustomCameraNodeParameterInfos::AddDataParameter(
 		FName ParameterName, 
 		ECameraContextDataType ParameterType,
 		const UObject* ParameterTypeObject,
+		const uint8* DefaultValue,
 		FCameraContextDataID* OverrideDataID)
 {
-	DataParameters.Add({ ParameterName, ParameterType, ParameterTypeObject, OverrideDataID });
+	DataParameters.Add({ ParameterName, ParameterType, ParameterTypeObject, DefaultValue, OverrideDataID });
+}
+
+void FCustomCameraNodeParameterInfos::AddDataParameter(FCustomCameraNodeDataParameter& Parameter, const uint8* DefaultValue)
+{
+	AddDataParameter(
+			Parameter.ParameterName,
+			Parameter.ParameterType,
+			Parameter.ParameterTypeObject,
+			DefaultValue,
+			&Parameter.OverrideDataID);
 }
 
 void FCustomCameraNodeParameterInfos::GetBlendableParameters(TArray<FCustomCameraNodeBlendableParameter>& OutBlendableParameters) const
@@ -50,9 +51,11 @@ void FCustomCameraNodeParameterInfos::GetBlendableParameters(TArray<FCustomCamer
 		FCustomCameraNodeBlendableParameter& OutParameter = OutBlendableParameters.Emplace_GetRef();
 		OutParameter.ParameterName = BlendableParameter.ParameterName;
 		OutParameter.ParameterType = BlendableParameter.ParameterType;
-		if (BlendableParameter.OverrideVariable)
+		OutParameter.BlendableStructType = BlendableParameter.BlendableStructType;
+		OutParameter.OverrideVariable = BlendableParameter.OverrideVariable;
+		if (BlendableParameter.OverrideVariableID)
 		{
-			OutParameter.OverrideVariable = *BlendableParameter.OverrideVariable;
+			OutParameter.OverrideVariableID = *BlendableParameter.OverrideVariableID;
 		}
 	}
 }
@@ -80,10 +83,10 @@ bool FCustomCameraNodeParameterInfos::FindBlendableParameter(FName ParameterName
 		{
 			OutParameter.ParameterName = BlendableParameter.ParameterName;
 			OutParameter.ParameterType = BlendableParameter.ParameterType;
-			OutParameter.OverrideVariable = nullptr;
-			if (BlendableParameter.OverrideVariable)
+			OutParameter.BlendableStructType = BlendableParameter.BlendableStructType;
+			if (BlendableParameter.OverrideVariableID)
 			{
-				OutParameter.OverrideVariable = *BlendableParameter.OverrideVariable;
+				OutParameter.OverrideVariableID = *BlendableParameter.OverrideVariableID;
 			}
 			return true;
 		}
diff --git a/Engine/Plugins/Cameras/GameplayCameras/Source/GameplayCameras/Private/Debug/ContextDataTableDebugBlock.cpp b/Engine/Plugins/Cameras/GameplayCameras/Source/GameplayCameras/Private/Debug/ContextDataTableDebugBlock.cpp
--- a/Engine/Plugins/Cameras/GameplayCameras/Source/GameplayCameras/Private/Debug/ContextDataTableDebugBlock.cpp
+++ b/Engine/Plugins/Cameras/GameplayCameras/Source/GameplayCameras/Private/Debug/ContextDataTableDebugBlock.cpp
@@ -29,7 +29,10 @@ void FContextDataTableDebugBlock::Initialize(const FCameraContextDataTable& InCo
 {
 	for (const FCameraContextDataTable::FEntry& Entry : InContextDataTable.Entries)
 	{
-		const uint8* EntryData = InContextDataTable.Memory + Entry.Offset;
+		FString EntryName;
+#if WITH_EDITORONLY_DATA
+		EntryName = Entry.DebugName;
+#endif
 
 		FName EntryTypeName;
 		if (Entry.TypeObject)
@@ -38,6 +41,7 @@ void FContextDataTableDebugBlock::Initialize(const FCameraContextDataTable& InCo
 		}
 
 		FString EntryValueStr;
+		const uint8* EntryData = InContextDataTable.Memory + Entry.Offset;
 		switch (Entry.Type)
 		{
 			case ECameraContextDataType::Name:
@@ -70,7 +74,7 @@ void FContextDataTableDebugBlock::Initialize(const FCameraContextDataTable& InCo
 				break;
 		}
 
-		FEntryDebugInfo EntryDebugInfo{ Entry.ID.DataName, EntryTypeName, EntryValueStr};
+		FEntryDebugInfo EntryDebugInfo{ Entry.ID.GetValue(), EntryName, EntryTypeName, EntryValueStr};
 		EntryDebugInfo.bWritten = EnumHasAnyFlags(Entry.Flags, FCameraContextDataTable::EEntryFlags::Written);
 		EntryDebugInfo.bWrittenThisFrame = EnumHasAnyFlags(Entry.Flags, FCameraContextDataTable::EEntryFlags::WrittenThisFrame);
 		Entries.Add(EntryDebugInfo);
@@ -88,22 +92,38 @@ void FContextDataTableDebugBlock::OnDebugDraw(const FCameraDebugBlockDrawParams&
 
 	for (const FEntryDebugInfo& Entry : Entries)
 	{
-		Renderer.AddText(TEXT("%s [%s] "), *Entry.Name.ToString(), *Entry.TypeName.ToString());
-		if (!Entry.bWritten)
+#if WITH_EDITORONLY_DATA
+		if (!Entry.Name.IsEmpty())
 		{
-			Renderer.AddText("{cam_warning}[Uninitialized]");
+			Renderer.AddText(TEXT("%s [%s] "), *Entry.Name, *Entry.TypeName.ToString());
 		}
-		else if (Entry.bWrittenThisFrame)
+		else
 		{
-			Renderer.AddText(TEXT(" {cam_passive}[WrittenThisFrame]"));
+			Renderer.AddText(TEXT("<no name data> [%s] "), *Entry.TypeName.ToString());
 		}
-		Renderer.NewLine();
+#else
+		Renderer.AddText(TEXT("[%d] <no name data> : "), Entry.ID);
+#endif
+
+		if (Entry.bWritten)
+		{
+			if (Entry.bWrittenThisFrame)
+			{
+				Renderer.AddText(TEXT(" {cam_passive}[WrittenThisFrame]"));
+			}
+			Renderer.NewLine();
 
-		Renderer.AddIndent();
+			Renderer.AddIndent();
+			{
+				Renderer.AddText(Entry.Value);
+			}
+			Renderer.RemoveIndent();
+		}
+		else
 		{
-			Renderer.AddText(Entry.Value);
+			Renderer.AddText("{cam_warning}[Uninitialized]\n");
 		}
-		Renderer.RemoveIndent();
+
 		Renderer.SetTextColor(Colors.Default);
 	}
 }
diff --git a/Engine/Plugins/Cameras/GameplayCameras/Source/GameplayCameras/Private/Debug/VariableTableDebugBlock.cpp b/Engine/Plugins/Cameras/GameplayCameras/Source/GameplayCameras/Private/Debug/VariableTableDebugBlock.cpp
--- a/Engine/Plugins/Cameras/GameplayCameras/Source/GameplayCameras/Private/Debug/VariableTableDebugBlock.cpp
+++ b/Engine/Plugins/Cameras/GameplayCameras/Source/GameplayCameras/Private/Debug/VariableTableDebugBlock.cpp
@@ -34,21 +34,31 @@ void FVariableTableDebugBlock::Initialize(const FCameraVariableTable& InVariable
 #endif
 
 		FString EntryValueStr;
+		const uint8* RawValuePtr = InVariableTable.Memory + Entry.Offset;
+
+		switch (Entry.Type)
+		{
 #define UE_CAMERA_VARIABLE_FOR_TYPE(ValueType, ValueName)\
 			case ECameraVariableType::ValueName:\
 				if (EnumHasAnyFlags(Entry.Flags, FCameraVariableTable::EEntryFlags::Written))\
 				{\
-					const ValueType EntryValue = InVariableTable.GetValue<ValueType>(FCameraVariableID::FromHashValue(Entry.ID.GetValue()));\
-					EntryValueStr = ToDebugString(EntryValue);\
+					const ValueType* EntryValue = reinterpret_cast<const ValueType*>(RawValuePtr);\
+					EntryValueStr = ToDebugString(*EntryValue);\
 				}\
 				break;
-		switch (Entry.Type)
-		{
 			UE_CAMERA_VARIABLE_FOR_ALL_TYPES()
-		}
 #undef UE_CAMERA_VARIABLE_FOR_TYPE
+			case ECameraVariableType::BlendableStruct:
+				if (EnumHasAnyFlags(Entry.Flags, FCameraVariableTable::EEntryFlags::Written))
+				{
+					const UScriptStruct* StructType = CastChecked<const UScriptStruct>(Entry.StructType);
+					const int32 ExportFlags = PPF_Delimited | PPF_IncludeTransient | PPF_ExternalEditor;
+					StructType->ExportText(EntryValueStr, RawValuePtr, nullptr, nullptr, ExportFlags, nullptr);
+				}
+				break;
+		}
 
-		FEntryDebugInfo EntryDebugInfo{ Entry.ID.GetValue(), EntryName, EntryValueStr};
+		FEntryDebugInfo EntryDebugInfo{ Entry.ID.GetValue(), EntryName, EntryValueStr };
 		EntryDebugInfo.bIsInput = EnumHasAnyFlags(Entry.Flags, FCameraVariableTable::EEntryFlags::Input);
 		EntryDebugInfo.bIsPrivate = EnumHasAnyFlags(Entry.Flags, FCameraVariableTable::EEntryFlags::Private);
 		EntryDebugInfo.bWritten = EnumHasAnyFlags(Entry.Flags, FCameraVariableTable::EEntryFlags::Written);
diff --git a/Engine/Plugins/Cameras/GameplayCameras/Source/GameplayCameras/Private/GameFramework/CameraRigParameterInterop.cpp b/Engine/Plugins/Cameras/GameplayCameras/Source/GameplayCameras/Private/GameFramework/CameraRigParameterInterop.cpp
--- a/Engine/Plugins/Cameras/GameplayCameras/Source/GameplayCameras/Private/GameFramework/CameraRigParameterInterop.cpp
+++ b/Engine/Plugins/Cameras/GameplayCameras/Source/GameplayCameras/Private/GameFramework/CameraRigParameterInterop.cpp
@@ -17,41 +17,48 @@
 namespace UE::Cameras::Private
 {
 
-template<typename VariableAssetType>
-void SetCameraRigBlendableParameter(FBlueprintCameraVariableTable VariableTable, VariableAssetType* PrivateVariable, typename VariableAssetType::ValueType Value)
+template<typename ValueType>
+void SetCameraRigBlendableParameter(FBlueprintCameraVariableTable VariableTable, FCameraVariableID PrivateVariableID, typename TCallTraits<ValueType>::ParamType Value)
 {
 	if (!VariableTable.IsValid())
 	{
 		FFrame::KismetExecutionMessage(TEXT("Invalid camera variable table was passed."), ELogVerbosity::Error);
 		return;
 	}
-	if (PrivateVariable == nullptr)
+	if (!PrivateVariableID.IsValid())
 	{
 		FFrame::KismetExecutionMessage(TEXT("No camera rig was passed."), ELogVerbosity::Error);
 		return;
 	}
 
-	VariableTable.GetVariableTable()->SetValue(PrivateVariable, Value, true);
+	VariableTable.GetVariableTable()->TrySetValue<ValueType>(PrivateVariableID, Value);
 }
 
-UCameraVariableAsset* GetParameterPrivateVariable(UCameraRigAsset* CameraRig, const FString& ParameterName)
+FCameraVariableID GetParameterPrivateVariableID(UCameraRigAsset* CameraRig, const FString& ParameterName, const UScriptStruct* BlendableStructType = nullptr)
 {
 	UCameraRigBlendableParameter* BlendableParameter = CameraRig->Interface.FindBlendableParameterByName(ParameterName);
 	if (!BlendableParameter)
 	{
 		const FText Text = LOCTEXT("NoSuchBlendableParameter", "No parameter '{0}' found on camera rig '{1}'. Setting this camera variable table value will most probably accomplish nothing.");
 		FFrame::KismetExecutionMessage(*FText::Format(Text, FText::FromString(ParameterName), FText::FromString(CameraRig->GetPathName())).ToString(), ELogVerbosity::Warning);
-		return nullptr;
+		return FCameraVariableID();
 	}
 
-	if (!BlendableParameter->PrivateVariable)
+	if (!BlendableParameter->PrivateVariableID)
 	{
 		const FText Text = LOCTEXT("CameraRigNeedsBuilding", "Parameter '{0}' isn't built. Please build camera rig '{1}'.");
 		FFrame::KismetExecutionMessage(*FText::Format(Text, FText::FromString(ParameterName), FText::FromString(CameraRig->GetPathName())).ToString(), ELogVerbosity::Warning);
-		return nullptr;
+		return FCameraVariableID();
 	}
 
-	return BlendableParameter->PrivateVariable;
+	if (BlendableParameter->BlendableStructType != BlendableStructType)
+	{
+		const FText Text = LOCTEXT("InvalidParameterType", "Parameter '{0}' has an incorrect type. Please build camera rig '{1}'.");
+		FFrame::KismetExecutionMessage(*FText::Format(Text, FText::FromString(ParameterName), FText::FromString(CameraRig->GetPathName())).ToString(), ELogVerbosity::Warning);
+		return FCameraVariableID();
+	}
+
+	return BlendableParameter->PrivateVariableID;
 }
 
 bool ValidateSetCameraRigDataParameter(FBlueprintCameraContextDataTable& ContextDataTable, FCameraContextDataID DataID)
@@ -98,76 +105,95 @@ UCameraRigParameterInterop::UCameraRigParameterInterop(const FObjectInitializer&
 
 void UCameraRigParameterInterop::SetBooleanParameter(FBlueprintCameraNodeEvaluationResult& Result, UCameraRigAsset* CameraRig, const FString& ParameterName, bool ParameterValue)
 {
-	UE::Cameras::Private::SetCameraRigBlendableParameter(
+	UE::Cameras::Private::SetCameraRigBlendableParameter<bool>(
 			Result.GetVariableTable(), 
-			Cast<UBooleanCameraVariable>(UE::Cameras::Private::GetParameterPrivateVariable(CameraRig, ParameterName)), 
+			UE::Cameras::Private::GetParameterPrivateVariableID(CameraRig, ParameterName), 
 			ParameterValue);
 }
 
 void UCameraRigParameterInterop::SetIntegerParameter(FBlueprintCameraNodeEvaluationResult& Result, UCameraRigAsset* CameraRig, const FString& ParameterName, int32 ParameterValue)
 {
-	UE::Cameras::Private::SetCameraRigBlendableParameter(
+	UE::Cameras::Private::SetCameraRigBlendableParameter<int32>(
 			Result.GetVariableTable(), 
-			Cast<UInteger32CameraVariable>(UE::Cameras::Private::GetParameterPrivateVariable(CameraRig, ParameterName)), 
+			UE::Cameras::Private::GetParameterPrivateVariableID(CameraRig, ParameterName), 
 			ParameterValue);
 }
 
 void UCameraRigParameterInterop::SetFloatParameter(FBlueprintCameraNodeEvaluationResult& Result, UCameraRigAsset* CameraRig, const FString& ParameterName, double ParameterValue)
 {
-	UE::Cameras::Private::SetCameraRigBlendableParameter(
+	UE::Cameras::Private::SetCameraRigBlendableParameter<float>(
 			Result.GetVariableTable(), 
-			Cast<UFloatCameraVariable>(UE::Cameras::Private::GetParameterPrivateVariable(CameraRig, ParameterName)),
+			UE::Cameras::Private::GetParameterPrivateVariableID(CameraRig, ParameterName),
 			(float)ParameterValue);
 }
 
 void UCameraRigParameterInterop::SetDoubleParameter(FBlueprintCameraNodeEvaluationResult& Result, UCameraRigAsset* CameraRig, const FString& ParameterName, double ParameterValue)
 {
-	UE::Cameras::Private::SetCameraRigBlendableParameter(
+	UE::Cameras::Private::SetCameraRigBlendableParameter<double>(
 			Result.GetVariableTable(), 
-			Cast<UDoubleCameraVariable>(UE::Cameras::Private::GetParameterPrivateVariable(CameraRig, ParameterName)),
+			UE::Cameras::Private::GetParameterPrivateVariableID(CameraRig, ParameterName),
 			ParameterValue);
 }
 
 void UCameraRigParameterInterop::SetVector2Parameter(FBlueprintCameraNodeEvaluationResult& Result, UCameraRigAsset* CameraRig, const FString& ParameterName, FVector2D ParameterValue)
 {
-	UE::Cameras::Private::SetCameraRigBlendableParameter(
+	UE::Cameras::Private::SetCameraRigBlendableParameter<FVector2D>(
 			Result.GetVariableTable(), 
-			Cast<UVector2dCameraVariable>(UE::Cameras::Private::GetParameterPrivateVariable(CameraRig, ParameterName)),
+			UE::Cameras::Private::GetParameterPrivateVariableID(CameraRig, ParameterName),
 			ParameterValue);
 }
 
 void UCameraRigParameterInterop::SetVector3Parameter(FBlueprintCameraNodeEvaluationResult& Result, UCameraRigAsset* CameraRig, const FString& ParameterName, FVector ParameterValue)
 {
-	UE::Cameras::Private::SetCameraRigBlendableParameter(
+	UE::Cameras::Private::SetCameraRigBlendableParameter<FVector>(
 			Result.GetVariableTable(), 
-			Cast<UVector3dCameraVariable>(UE::Cameras::Private::GetParameterPrivateVariable(CameraRig, ParameterName)),
+			UE::Cameras::Private::GetParameterPrivateVariableID(CameraRig, ParameterName),
 			ParameterValue);
 }
 
 void UCameraRigParameterInterop::SetVector4Parameter(FBlueprintCameraNodeEvaluationResult& Result, UCameraRigAsset* CameraRig, const FString& ParameterName, FVector4 ParameterValue)
 {
-	UE::Cameras::Private::SetCameraRigBlendableParameter(
+	UE::Cameras::Private::SetCameraRigBlendableParameter<FVector4>(
 			Result.GetVariableTable(), 
-			Cast<UVector4dCameraVariable>(UE::Cameras::Private::GetParameterPrivateVariable(CameraRig, ParameterName)),
+			UE::Cameras::Private::GetParameterPrivateVariableID(CameraRig, ParameterName),
 			ParameterValue);
 }
 
 void UCameraRigParameterInterop::SetRotatorParameter(FBlueprintCameraNodeEvaluationResult& Result, UCameraRigAsset* CameraRig, const FString& ParameterName, FRotator ParameterValue)
 {
-	UE::Cameras::Private::SetCameraRigBlendableParameter(
+	UE::Cameras::Private::SetCameraRigBlendableParameter<FRotator>(
 			Result.GetVariableTable(), 
-			Cast<URotator3dCameraVariable>(UE::Cameras::Private::GetParameterPrivateVariable(CameraRig, ParameterName)),
+			UE::Cameras::Private::GetParameterPrivateVariableID(CameraRig, ParameterName),
 			ParameterValue);
 }
 
 void UCameraRigParameterInterop::SetTransformParameter(FBlueprintCameraNodeEvaluationResult& Result, UCameraRigAsset* CameraRig, const FString& ParameterName, FTransform ParameterValue)
 {
-	UE::Cameras::Private::SetCameraRigBlendableParameter(
+	UE::Cameras::Private::SetCameraRigBlendableParameter<FTransform>(
 			Result.GetVariableTable(), 
-			Cast<UTransform3dCameraVariable>(UE::Cameras::Private::GetParameterPrivateVariable(CameraRig, ParameterName)),
+			UE::Cameras::Private::GetParameterPrivateVariableID(CameraRig, ParameterName),
 			ParameterValue);
 }
 
+void UCameraRigParameterInterop::SetBlendableStructParameter(FBlueprintCameraNodeEvaluationResult& Result, UCameraRigAsset* CameraRig, const FString& ParameterName, const FInstancedStruct& ParameterValue)
+{
+	FBlueprintCameraVariableTable VariableTable = Result.GetVariableTable();
+	if (!VariableTable.IsValid())
+	{
+		FFrame::KismetExecutionMessage(TEXT("Invalid camera variable table was passed."), ELogVerbosity::Error);
+		return;
+	}
+
+	FCameraVariableID PrivateVariableID = UE::Cameras::Private::GetParameterPrivateVariableID(CameraRig, ParameterName, ParameterValue.GetScriptStruct());
+	if (!PrivateVariableID.IsValid())
+	{
+		FFrame::KismetExecutionMessage(TEXT("No camera rig was passed."), ELogVerbosity::Error);
+		return;
+	}
+
+	VariableTable.GetVariableTable()->TrySetValue(PrivateVariableID, ECameraVariableType::BlendableStruct, ParameterValue.GetScriptStruct(), ParameterValue.GetMemory());
+}
+
 void UCameraRigParameterInterop::SetNameParameter(UPARAM(Ref) FBlueprintCameraNodeEvaluationResult& Result, UCameraRigAsset* CameraRig, const FString& ParameterName, FName ParameterValue)
 {
 	const FCameraContextDataID DataID(UE::Cameras::Private::GetParameterPrivateDataID(CameraRig, ParameterName));
diff --git a/Engine/Plugins/Cameras/GameplayCameras/Source/GameplayCameras/Private/GameplayCamerasModule.cpp b/Engine/Plugins/Cameras/GameplayCameras/Source/GameplayCameras/Private/GameplayCamerasModule.cpp
--- a/Engine/Plugins/Cameras/GameplayCameras/Source/GameplayCameras/Private/GameplayCamerasModule.cpp
+++ b/Engine/Plugins/Cameras/GameplayCameras/Source/GameplayCameras/Private/GameplayCamerasModule.cpp
@@ -35,6 +35,42 @@ public:
 public:
 
 	// IGameplayCamerasModule interface
+	virtual void RegisterBlendableStruct(const UScriptStruct* StructType, UE::Cameras::FBlendableStructTypeErasedInterpolator Interpolator) override
+	{
+		using namespace UE::Cameras;
+
+		if (!ensure(EnumHasAllFlags(StructType->StructFlags, STRUCT_IsPlainOldData)))
+		{
+			return;
+		}
+
+		const bool bAlreadyRegistered = BlendableStructs.ContainsByPredicate([StructType](const FBlendableStructInfo& Item)
+				{
+					return Item.StructType == StructType;
+				});
+		if (ensure(!bAlreadyRegistered))
+		{
+			FBlendableStructInfo& NewInfo = BlendableStructs.Emplace_GetRef();
+			NewInfo.StructType = StructType;
+			NewInfo.Interpolator = Interpolator;
+		}
+	}
+
+	virtual TConstArrayView<UE::Cameras::FBlendableStructInfo> GetBlendableStructs() const override
+	{
+		return BlendableStructs;
+	}
+
+	virtual void UnregisterBlendableStruct(const UScriptStruct* StructType) override
+	{
+		using namespace UE::Cameras;
+
+		BlendableStructs.RemoveAll([StructType](const FBlendableStructInfo& Item)
+				{
+					return Item.StructType == StructType;
+				});
+	}
+
 #if WITH_EDITOR
 	virtual TSharedPtr<IGameplayCamerasLiveEditManager> GetLiveEditManager() const override
 	{
@@ -49,6 +85,8 @@ public:
 
 private:
 
+	TArray<UE::Cameras::FBlendableStructInfo> BlendableStructs;
+
 #if WITH_EDITOR
 	TSharedPtr<IGameplayCamerasLiveEditManager> LiveEditManager;
 #endif
diff --git a/Engine/Plugins/Cameras/GameplayCameras/Source/GameplayCameras/Private/Helpers/CameraAssetParameterOverrideEvaluator.cpp b/Engine/Plugins/Cameras/GameplayCameras/Source/GameplayCameras/Private/Helpers/CameraAssetParameterOverrideEvaluator.cpp
--- a/Engine/Plugins/Cameras/GameplayCameras/Source/GameplayCameras/Private/Helpers/CameraAssetParameterOverrideEvaluator.cpp
+++ b/Engine/Plugins/Cameras/GameplayCameras/Source/GameplayCameras/Private/Helpers/CameraAssetParameterOverrideEvaluator.cpp
@@ -4,6 +4,7 @@
 
 #include "Core/CameraAsset.h"
 #include "Core/CameraAssetReference.h"
+#include "Core/CameraRigAsset.h"
 #include "Helpers/CameraParameterOverrideHelper.h"
 
 namespace UE::Cameras
@@ -38,28 +39,33 @@ void FCameraAssetParameterOverrideEvaluator::ApplyParameterOverrides(FCameraVari
 		return;
 	}
 
-	TConstArrayView<FPropertyBagPropertyDesc> PropertyDescs = CameraParametersStruct->GetPropertyDescs();
-	const TArray<TWeakObjectPtr<const UCameraRigAsset>>& ParameterOwners = CameraAsset->ParameterOwners;
-	ensure(PropertyDescs.Num() == ParameterOwners.Num());
+	TConstArrayView<FCameraRigParameterDefinition> ParameterDefinitions = CameraAsset->GetParameterDefinitions();
+	const TArray<TObjectPtr<const UCameraRigAsset>>& ParameterOwners = CameraAsset->ParameterOwners;
+	ensure(ParameterDefinitions.Num() == ParameterOwners.Num());
 
 	FCameraParameterOverrideHelper Helper(*OutVariableTable, *OutContextDataTable);
 
-	for (int32 Index = 0, MaxIndex = FMath::Min(PropertyDescs.Num(), ParameterOwners.Num()); Index < MaxIndex; ++Index)
+	for (int32 Index = 0, MaxIndex = FMath::Min(ParameterDefinitions.Num(), ParameterOwners.Num()); Index < MaxIndex; ++Index)
 	{
-		const FPropertyBagPropertyDesc& PropertyDesc(PropertyDescs[Index]);
-		if (!CameraReference.IsParameterOverriden(PropertyDesc.ID))
+		const FCameraRigParameterDefinition& Definition(ParameterDefinitions[Index]);
+		if (!CameraReference.IsParameterOverridden(Definition.ParameterGuid))
 		{
 			continue;
 		}
 
-		const UCameraRigAsset* CameraRig = ParameterOwners[Index].Get();
+		const FPropertyBagPropertyDesc* PropertyDesc = CameraParametersStruct->FindPropertyDescByID(Definition.ParameterGuid);
+		if (!ensure(PropertyDesc))
+		{
+			continue;
+		}
+
+		const UCameraRigAsset* CameraRig = ParameterOwners[Index];
 		if (!ensure(CameraRig))
 		{
 			continue;
 		}
 
-		TValueOrError<FStructView, EPropertyBagResult> ParameterValueOrError = CameraParameters.GetValueStruct(PropertyDesc);
-		Helper.ApplyParameterOverride(CameraRig, PropertyDesc.ID, ParameterValueOrError, bDrivenOverridesOnly);
+		Helper.ApplyParameterOverride(CameraRig, Definition, CameraParameters, *PropertyDesc, bDrivenOverridesOnly);
 	}
 }
 
diff --git a/Engine/Plugins/Cameras/GameplayCameras/Source/GameplayCameras/Private/Helpers/CameraAssetParameterOverrideEvaluator.h b/Engine/Plugins/Cameras/GameplayCameras/Source/GameplayCameras/Private/Helpers/CameraAssetParameterOverrideEvaluator.h
--- a/Engine/Plugins/Cameras/GameplayCameras/Source/GameplayCameras/Private/Helpers/CameraAssetParameterOverrideEvaluator.h
+++ b/Engine/Plugins/Cameras/GameplayCameras/Source/GameplayCameras/Private/Helpers/CameraAssetParameterOverrideEvaluator.h
@@ -2,6 +2,8 @@
 
 #pragma once
 
+#include "UObject/UnrealNames.h"
+
 struct FCameraAssetReference;
 
 namespace UE::Cameras
diff --git a/Engine/Plugins/Cameras/GameplayCameras/Source/GameplayCameras/Private/Helpers/CameraParameterOverrideHelper.cpp b/Engine/Plugins/Cameras/GameplayCameras/Source/GameplayCameras/Private/Helpers/CameraParameterOverrideHelper.cpp
--- a/Engine/Plugins/Cameras/GameplayCameras/Source/GameplayCameras/Private/Helpers/CameraParameterOverrideHelper.cpp
+++ b/Engine/Plugins/Cameras/GameplayCameras/Source/GameplayCameras/Private/Helpers/CameraParameterOverrideHelper.cpp
@@ -5,7 +5,9 @@
 #include "Core/CameraContextDataTable.h"
 #include "Core/CameraParameters.h"
 #include "Core/CameraRigAsset.h"
+#include "Core/CameraRigParameterDefinition.h"
 #include "Core/CameraVariableTable.h"
+#include "IGameplayCamerasModule.h"
 
 namespace UE::Cameras
 {
@@ -16,19 +18,42 @@ namespace Internal
 template<typename ParameterType>
 void ApplyBlendableParameterOverride(
 		const UCameraRigAsset* CameraRig,
-		const UCameraRigBlendableParameter* BlendableParameter,
+		const FCameraRigParameterDefinition& ParameterDefinition,
 		const ParameterType& ParameterValue,
 		FCameraVariableTable& VariableTable,
 		bool bDrivenOverridesOnly)
 {
 	using ValueType = typename ParameterType::ValueType;
 
-	if (!ensure(BlendableParameter))
+	const FCameraVariableID ParameterVariableID(ParameterDefinition.VariableID);
+	if (ParameterValue.Variable != nullptr)
 	{
-		return;
+		// The override is driven by a variable... read its value and set it as the value for the
+		// prefab's variable. Basically, we forward the value from one variable to the next.
+		FCameraVariableDefinition OverrideDefinition(ParameterValue.Variable->GetVariableDefinition());
+
+		const ValueType OverrideValue = VariableTable.GetValue<ValueType>(
+				OverrideDefinition.VariableID, ParameterValue.Variable->GetDefaultValue());
+		VariableTable.SetValue<ValueType>(ParameterVariableID, OverrideValue);
 	}
+	else if (!bDrivenOverridesOnly)
+	{
+		// The override is a fixed value. Just set that on the prefab's variable.
+		VariableTable.SetValue<ValueType>(ParameterVariableID, ParameterValue.Value);
+	}
+}
 
-	if (!BlendableParameter->PrivateVariable)
+void ApplyBlendableParameterOverride(
+		const UCameraRigAsset* CameraRig,
+		const FCameraRigParameterDefinition& ParameterDefinition,
+		const FInstancedPropertyBag& PropertyBag,
+		const FPropertyBagPropertyDesc& PropertyBagPropertyDesc,
+		FCameraVariableTable& VariableTable,
+		bool bDrivenOverridesOnly)
+{
+	ensure(ParameterDefinition.ParameterType == ECameraRigInterfaceParameterType::Blendable);
+
+	if (!ParameterDefinition.VariableID)
 	{
 		// Ignore un-built parameter overrides in the editor since the user could have just added
 		// an override while PIE is running. They need to hit the Build button for the override
@@ -37,56 +62,152 @@ void ApplyBlendableParameterOverride(
 #if !WITH_EDITOR
 		UE_LOG(LogCameraSystem, Error,
 				TEXT("Invalid blendable parameter override '%s' in camera rig '%s'. Was it built/cooked?"),
-				*BlendableParameter->InterfaceParameterName,
+				*ParameterDefinition.ParameterName.ToString(),
 				*GetPathNameSafe(CameraRig));
 #endif
 		return;
 	}
 
-	FCameraVariableID ParameterVariableID(BlendableParameter->PrivateVariable->GetVariableID());
-
-	if (ParameterValue.Variable != nullptr)
+	TValueOrError<FStructView, EPropertyBagResult> ParameterValueOrError = PropertyBag.GetValueStruct(PropertyBagPropertyDesc);
+	if (!ensureMsgf(
+				ParameterValueOrError.HasValue() && !ParameterValueOrError.HasError(),
+				TEXT("Camera parameter has no valid value! Error: %s"),
+				*UEnum::GetValueAsString(ParameterValueOrError.GetError())))
 	{
-		// The override is driven by a variable... read its value and set it as the value for the
-		// prefab's variable. Basically, we forward the value from one variable to the next.
-		FCameraVariableDefinition OverrideDefinition(ParameterValue.Variable->GetVariableDefinition());
-
-		const ValueType OverrideValue = VariableTable.GetValue<ValueType>(
-				OverrideDefinition.VariableID, ParameterValue.Variable->GetDefaultValue());
-		VariableTable.SetValue<ValueType>(ParameterVariableID, OverrideValue);
+		return;
 	}
-	else if (!bDrivenOverridesOnly)
+
+	const FStructView& ParameterValue = ParameterValueOrError.GetValue();
+	const UScriptStruct* ParameterType = ParameterValue.GetScriptStruct();
+
+	switch (ParameterDefinition.VariableType)
 	{
-		// The override is a fixed value. Just set that on the prefab's variable.
-		VariableTable.SetValue<ValueType>(ParameterVariableID, ParameterValue.Value);
+#define UE_CAMERA_VARIABLE_FOR_TYPE(ValueType, ValueName)\
+		case ECameraVariableType::ValueName:\
+			{\
+				check(ParameterType == F##ValueName##CameraParameter::StaticStruct());\
+				const F##ValueName##CameraParameter& TypedParameterValue = ParameterValue.Get<F##ValueName##CameraParameter>();\
+				ApplyBlendableParameterOverride(CameraRig, ParameterDefinition, TypedParameterValue, VariableTable, bDrivenOverridesOnly);\
+			}\
+			break;
+UE_CAMERA_VARIABLE_FOR_ALL_TYPES()
+#undef UE_CAMERA_VARIABLE_FOR_TYPE
+		case ECameraVariableType::BlendableStruct:
+			{
+				const uint8* RawValuePtr = ParameterValue.GetMemory();
+				VariableTable.SetValue(ParameterDefinition.VariableID, ParameterDefinition.VariableType, ParameterDefinition.BlendableStructType, RawValuePtr);
+			}
+			break;
+		default:
+			ensure(false);
+			break;
 	}
 }
 
+template<typename ParameterType>
+void OverrideContextDataTableEntry(
+		const FCameraRigParameterDefinition& ParameterDefinition,
+		const ParameterType& ParameterValue,
+		FCameraContextDataTable& ContextDataTable)
+{
+	const uint8* RawParameterValue = reinterpret_cast<const uint8*>(&ParameterValue);
+	ContextDataTable.SetData(ParameterDefinition.DataID, ParameterDefinition.DataType, ParameterDefinition.DataTypeObject, RawParameterValue);
+}
+
+template<>
+void OverrideContextDataTableEntry<FStructView>(
+		const FCameraRigParameterDefinition& ParameterDefinition,
+		const FStructView& ParameterValue,
+		FCameraContextDataTable& ContextDataTable)
+{
+	const uint8* RawParameterValue = ParameterValue.GetMemory();
+	ContextDataTable.SetData(ParameterDefinition.DataID, ParameterDefinition.DataType, ParameterDefinition.DataTypeObject, RawParameterValue);
+}
+
+template<typename ParameterType>
 void ApplyDataParameterOverride(
 		const UCameraRigAsset* CameraRig,
-		const UCameraRigDataParameter* DataParameter,
-		const FStructView& ParameterValue,
+		const FCameraRigParameterDefinition& ParameterDefinition,
+		const TValueOrError<ParameterType, EPropertyBagResult>& ParameterValueOrError,
 		FCameraContextDataTable& ContextDataTable)
 {
-	if (!ensure(DataParameter))
+	if (!ensureMsgf(
+				ParameterValueOrError.HasValue() && !ParameterValueOrError.HasError(),
+				TEXT("Camera parameter has no valid value! Error: %s"),
+				*UEnum::GetValueAsString(ParameterValueOrError.GetError())))
 	{
 		return;
 	}
 
-	if (!DataParameter->PrivateDataID)
+	// Write the override value into the context data table.
+	const ParameterType& ParameterValue = ParameterValueOrError.GetValue();
+	OverrideContextDataTableEntry<ParameterType>(ParameterDefinition, ParameterValue, ContextDataTable);
+}
+
+void ApplyDataParameterOverride(
+		const UCameraRigAsset* CameraRig,
+		const FCameraRigParameterDefinition& ParameterDefinition,
+		const FInstancedPropertyBag& PropertyBag,
+		const FPropertyBagPropertyDesc& PropertyBagPropertyDesc,
+		FCameraContextDataTable& ContextDataTable)
+{
+	ensure(ParameterDefinition.ParameterType == ECameraRigInterfaceParameterType::Data);
+
+	if (!ParameterDefinition.DataID)
 	{
 #if !WITH_EDITOR
 		UE_LOG(LogCameraSystem, Error,
 				TEXT("Invalid data parameter override '%s' in camera rig '%s'. Was it built/cooked?"),
-				*DataParameter->InterfaceParameterName,
+				*ParameterDefinition.ParameterName.ToString(),
 				*GetPathNameSafe(CameraRig));
 		return;
 #endif
 	}
 
-	// Write the override value into the context data table.
-	FCameraContextDataID ParameterDataID = DataParameter->PrivateDataID;
-	ContextDataTable.SetStructViewData(ParameterDataID, ParameterValue);
+	switch (ParameterDefinition.DataType)
+	{
+		case ECameraContextDataType::Name:
+			{
+				TValueOrError<FName, EPropertyBagResult> ParameterValueOrError = PropertyBag.GetValueName(PropertyBagPropertyDesc);
+				ApplyDataParameterOverride(CameraRig, ParameterDefinition, ParameterValueOrError, ContextDataTable);
+			}
+			break;
+		case ECameraContextDataType::String:
+			{
+				TValueOrError<FString, EPropertyBagResult> ParameterValueOrError = PropertyBag.GetValueString(PropertyBagPropertyDesc);
+				ApplyDataParameterOverride(CameraRig, ParameterDefinition, ParameterValueOrError, ContextDataTable);
+			}
+			break;
+		case ECameraContextDataType::Enum:
+			{
+				const UEnum* EnumType = CastChecked<const UEnum>(ParameterDefinition.DataTypeObject);
+				TValueOrError<uint8, EPropertyBagResult> ParameterValueOrError = PropertyBag.GetValueEnum(PropertyBagPropertyDesc, EnumType);
+				ApplyDataParameterOverride(CameraRig, ParameterDefinition, ParameterValueOrError, ContextDataTable);
+			}
+			break;
+		case ECameraContextDataType::Struct:
+			{
+				const UScriptStruct* StructType = CastChecked<const UScriptStruct>(ParameterDefinition.DataTypeObject);
+				TValueOrError<FStructView, EPropertyBagResult> ParameterValueOrError = PropertyBag.GetValueStruct(PropertyBagPropertyDesc, StructType);
+				ApplyDataParameterOverride(CameraRig, ParameterDefinition, ParameterValueOrError, ContextDataTable);
+			}
+			break;
+		case ECameraContextDataType::Object:
+			{
+				TValueOrError<UObject*, EPropertyBagResult> ParameterValueOrError = PropertyBag.GetValueObject(PropertyBagPropertyDesc);
+				ApplyDataParameterOverride(CameraRig, ParameterDefinition, ParameterValueOrError, ContextDataTable);
+			}
+			break;
+		case ECameraContextDataType::Class:
+			{
+				TValueOrError<UClass*, EPropertyBagResult> ParameterValueOrError = PropertyBag.GetValueClass(PropertyBagPropertyDesc);
+				ApplyDataParameterOverride(CameraRig, ParameterDefinition, ParameterValueOrError, ContextDataTable);
+			}
+			break;
+		default:
+			ensure(false);
+			break;
+	}
 }
 
 }  // namespace Internal
@@ -95,39 +216,39 @@ FCameraParameterOverrideHelper::FCameraParameterOverrideHelper(FCameraVariableTa
 	: VariableTable(OutVariableTable)
 	, ContextDataTable(OutContextDataTable)
 {
+	IGameplayCamerasModule& GameplayCamerasModule = IGameplayCamerasModule::Get();
+	BlendableStructs = GameplayCamerasModule.GetBlendableStructs();
 }
 
 void FCameraParameterOverrideHelper::ApplyParameterOverride(
 		const UCameraRigAsset* CameraRig,
-		const FGuid& ParameterGuid,
-		TValueOrError<FStructView, EPropertyBagResult> ParameterValueOrError,
+		const FCameraRigParameterDefinition& ParameterDefinition,
+		const FInstancedPropertyBag& PropertyBag,
+		const FPropertyBagPropertyDesc& PropertyBagPropertyDesc,
 		bool bDrivenOverridesOnly)
 {
 	using namespace Internal;
 
-	if (ensureMsgf(
-				ParameterValueOrError.HasValue() && !ParameterValueOrError.HasError(),
-				TEXT("Camera parameter has no valid value! Error: %s"),
-				*UEnum::GetValueAsString(ParameterValueOrError.GetError())))
+	switch (ParameterDefinition.ParameterType)
 	{
-		const FStructView& ParameterValue = ParameterValueOrError.GetValue();
-		const UScriptStruct* ParameterType = ParameterValue.GetScriptStruct();
-
-		// Check if this is a blendable parameter or a data parameter.
-#define UE_CAMERA_VARIABLE_FOR_TYPE(ValueType, ValueName)\
-		if (ParameterType == F##ValueName##CameraParameter::StaticStruct())\
-		{\
-			const UCameraRigBlendableParameter* BlendableParameter = CameraRig->Interface.FindBlendableParameterByGuid(ParameterGuid);\
-			const F##ValueName##CameraParameter& TypedParameterValue = ParameterValue.Get<F##ValueName##CameraParameter>();\
-			ApplyBlendableParameterOverride(CameraRig, BlendableParameter, TypedParameterValue, VariableTable, bDrivenOverridesOnly);\
-		}\
-		else
-		UE_CAMERA_VARIABLE_FOR_ALL_TYPES()
-#undef UE_CAMERA_VARIABLE_FOR_TYPE
-		{
-			const UCameraRigDataParameter* DataParameter = CameraRig->Interface.FindDataParameterByGuid(ParameterGuid);
-			ApplyDataParameterOverride(CameraRig, DataParameter, ParameterValue, ContextDataTable);
-		}
+		case ECameraRigInterfaceParameterType::Blendable:
+			{
+				ApplyBlendableParameterOverride(
+						CameraRig, 
+						ParameterDefinition,
+						PropertyBag, PropertyBagPropertyDesc, 
+						VariableTable, bDrivenOverridesOnly);
+			}
+			break;
+		case ECameraRigInterfaceParameterType::Data:
+			{
+				ApplyDataParameterOverride(
+						CameraRig, 
+						ParameterDefinition,
+						PropertyBag, PropertyBagPropertyDesc, 
+						ContextDataTable);
+			}
+			break;
 	}
 }
 
diff --git a/Engine/Plugins/Cameras/GameplayCameras/Source/GameplayCameras/Private/Helpers/CameraParameterOverrideHelper.h b/Engine/Plugins/Cameras/GameplayCameras/Source/GameplayCameras/Private/Helpers/CameraParameterOverrideHelper.h
--- a/Engine/Plugins/Cameras/GameplayCameras/Source/GameplayCameras/Private/Helpers/CameraParameterOverrideHelper.h
+++ b/Engine/Plugins/Cameras/GameplayCameras/Source/GameplayCameras/Private/Helpers/CameraParameterOverrideHelper.h
@@ -2,11 +2,11 @@
 
 #pragma once
 
+#include "Core/CameraVariableTableFwd.h"
 #include "StructUtils/PropertyBag.h"
 
 class UCameraRigAsset;
-class UCameraRigBlendableParameter;
-class UCameraRigDataParameter;
+struct FCameraRigParameterDefinition;
 
 namespace UE::Cameras
 {
@@ -24,14 +24,17 @@ struct FCameraParameterOverrideHelper
 
 	void ApplyParameterOverride(
 			const UCameraRigAsset* CameraRig,
-			const FGuid& ParameterGuid,
-			TValueOrError<FStructView, EPropertyBagResult> ParameterValueOrError,
+			const FCameraRigParameterDefinition& ParameterDefinition,
+			const FInstancedPropertyBag& PropertyBag,
+			const FPropertyBagPropertyDesc& PropertyBagPropertyDesc,
 			bool bDrivenOverridesOnly);
 
 private:
 
 	FCameraVariableTable& VariableTable;
 	FCameraContextDataTable& ContextDataTable;
+
+	TArray<FBlendableStructInfo> BlendableStructs;
 };
 
 }  // namespace UE::Cameras
diff --git a/Engine/Plugins/Cameras/GameplayCameras/Source/GameplayCameras/Private/Helpers/CameraRigParameterOverrideEvaluator.cpp b/Engine/Plugins/Cameras/GameplayCameras/Source/GameplayCameras/Private/Helpers/CameraRigParameterOverrideEvaluator.cpp
--- a/Engine/Plugins/Cameras/GameplayCameras/Source/GameplayCameras/Private/Helpers/CameraRigParameterOverrideEvaluator.cpp
+++ b/Engine/Plugins/Cameras/GameplayCameras/Source/GameplayCameras/Private/Helpers/CameraRigParameterOverrideEvaluator.cpp
@@ -33,18 +33,29 @@ void FCameraRigParameterOverrideEvaluator::ApplyParameterOverrides(FCameraVariab
 	const UCameraRigAsset* CameraRig = CameraRigReference.GetCameraRig();
 	const FInstancedPropertyBag& CameraRigParameters = CameraRigReference.GetParameters();
 	const UPropertyBag* CameraRigParametersStruct = CameraRigParameters.GetPropertyBagStruct();
+	if (!CameraRigParametersStruct)
+	{
+		return;
+	}
+
+	TConstArrayView<FCameraRigParameterDefinition> ParameterDefinitions = CameraRig->GetParameterDefinitions();
 
 	FCameraParameterOverrideHelper Helper(*OutVariableTable, *OutContextDataTable);
 
-	for (const FPropertyBagPropertyDesc& PropertyDesc : CameraRigParametersStruct->GetPropertyDescs())
+	for (const FCameraRigParameterDefinition& Definition : ParameterDefinitions)
 	{
-		if (!CameraRigReference.IsParameterOverriden(PropertyDesc.ID))
+		if (!CameraRigReference.IsParameterOverridden(Definition.ParameterGuid))
+		{
+			continue;
+		}
+
+		const FPropertyBagPropertyDesc* PropertyDesc = CameraRigParametersStruct->FindPropertyDescByID(Definition.ParameterGuid);
+		if (!ensure(PropertyDesc))
 		{
 			continue;
 		}
 
-		TValueOrError<FStructView, EPropertyBagResult> ParameterValueOrError = CameraRigParameters.GetValueStruct(PropertyDesc);
-		Helper.ApplyParameterOverride(CameraRig, PropertyDesc.ID, ParameterValueOrError, bDrivenOverridesOnly);
+		Helper.ApplyParameterOverride(CameraRig, Definition, CameraRigParameters, *PropertyDesc, bDrivenOverridesOnly);
 	}
 }
 
diff --git a/Engine/Plugins/Cameras/GameplayCameras/Source/GameplayCameras/Private/Nodes/Collision/CollisionPushCameraNode.cpp b/Engine/Plugins/Cameras/GameplayCameras/Source/GameplayCameras/Private/Nodes/Collision/CollisionPushCameraNode.cpp
--- a/Engine/Plugins/Cameras/GameplayCameras/Source/GameplayCameras/Private/Nodes/Collision/CollisionPushCameraNode.cpp
+++ b/Engine/Plugins/Cameras/GameplayCameras/Source/GameplayCameras/Private/Nodes/Collision/CollisionPushCameraNode.cpp
@@ -3,11 +3,11 @@
 #include "Nodes/Collision/CollisionPushCameraNode.h"
 
 #include "CollisionQueryParams.h"
-#include "Core/BuiltInCameraVariables.h"
 #include "Core/CameraEvaluationContext.h"
 #include "Core/CameraParameterReader.h"
 #include "Core/CameraSystemEvaluator.h"
 #include "Core/CameraValueInterpolator.h"
+#include "Core/CameraVariableReferenceReader.h"
 #include "Debug/CameraDebugBlock.h"
 #include "Debug/CameraDebugBlockBuilder.h"
 #include "Debug/CameraDebugColors.h"
@@ -53,6 +53,9 @@ private:
 
 private:
 
+	TCameraVariableReferenceReader<bool> EnableCollisionReader;
+	TCameraVariableReferenceReader<FVector3d> CustomSafePositionReader;
+
 	TCameraParameterReader<float> CollisionSphereRadiusReader;
 	TCameraParameterReader<FVector3d> SafePositionOffsetReader;
 
@@ -99,6 +102,9 @@ void FCollisionPushCameraNodeEvaluator::OnInitialize(const FCameraNodeEvaluatorI
 {
 	const UCollisionPushCameraNode* CollisionPushNode = GetCameraNodeAs<UCollisionPushCameraNode>();
 
+	EnableCollisionReader.Initialize(CollisionPushNode->EnableCollision);
+	CustomSafePositionReader.Initialize(CollisionPushNode->CustomSafePosition);
+
 	CollisionSphereRadiusReader.Initialize(CollisionPushNode->CollisionSphereRadius);
 	SafePositionOffsetReader.Initialize(CollisionPushNode->SafePositionOffset);
 
@@ -132,12 +138,7 @@ void FCollisionPushCameraNodeEvaluator::OnRun(const FCameraNodeEvaluationParams&
 	}
 
 	// See if collision is enabled. If not, handle it as if we didn't collide with anything.
-	bool bEnableCollision = true;
-	const UCollisionPushCameraNode* ThisNode = GetCameraNodeAs<UCollisionPushCameraNode>();
-	if (const UBooleanCameraVariable* EnableCollisionVariable = ThisNode->EnableCollision.Get())
-	{
-		bEnableCollision = OutResult.VariableTable.GetValue(EnableCollisionVariable);
-	}
+	const bool bEnableCollision = EnableCollisionReader.Get(OutResult.VariableTable);
 #if UE_GAMEPLAY_CAMERAS_DEBUG
 	bDebugCollisionEnabled = bEnableCollision;
 #endif
@@ -235,16 +236,13 @@ TOptional<FVector3d> FCollisionPushCameraNodeEvaluator::GetSafePosition(const FC
 #endif
 
 	// See if we have a custom safe position.
-	if (UVector3dCameraVariable* CustomSafePositionVariable = ThisNode->CustomSafePosition.Get())
+	FVector3d SafePosition;
+	if (CustomSafePositionReader.TryGet(OutResult.VariableTable, SafePosition))
 	{
-		FVector3d SafePosition;
-		if (OutResult.VariableTable.TryGetValue(CustomSafePositionVariable, SafePosition))
-		{
 #if UE_GAMEPLAY_CAMERAS_DEBUG
-			bDebugGotSafePosition = false;
+		bDebugGotSafePosition = false;
 #endif
-			return TOptional<FVector3d>(SafePosition);
-		}
+		return TOptional<FVector3d>(SafePosition);
 	}
 	
 	// Compute the base safe position.
@@ -269,7 +267,6 @@ TOptional<FVector3d> FCollisionPushCameraNodeEvaluator::GetSafePosition(const FC
 			break;
 	}
 
-	FVector3d SafePosition;
 	const bool bGotSafePosition = FCameraNodeSpaceMath::GetCameraNodeOriginPosition(Params, OutResult, OriginPosition, SafePosition);
 
 #if UE_GAMEPLAY_CAMERAS_DEBUG
diff --git a/Engine/Plugins/Cameras/GameplayCameras/Source/GameplayCameras/Private/Nodes/Common/CameraRigCameraNode.cpp b/Engine/Plugins/Cameras/GameplayCameras/Source/GameplayCameras/Private/Nodes/Common/CameraRigCameraNode.cpp
--- a/Engine/Plugins/Cameras/GameplayCameras/Source/GameplayCameras/Private/Nodes/Common/CameraRigCameraNode.cpp
+++ b/Engine/Plugins/Cameras/GameplayCameras/Source/GameplayCameras/Private/Nodes/Common/CameraRigCameraNode.cpp
@@ -126,63 +126,7 @@ void UCameraRigCameraNode::OnBuild(FCameraRigBuildContext& BuildContext)
 
 void UCameraRigCameraNode::GetCustomCameraNodeParameters(FCustomCameraNodeParameterInfos& OutParameterInfos)
 {
-	const FInstancedPropertyBag& ParameterOverrides = CameraRigReference.GetParameters();
-	const UPropertyBag* ParameterOverridesStruct = ParameterOverrides.GetPropertyBagStruct();
-	if (!ParameterOverridesStruct)
-	{
-		return;
-	}
-
-	for (const FPropertyBagPropertyDesc& PropertyDesc : ParameterOverridesStruct->GetPropertyDescs())
-	{
-		switch (PropertyDesc.ValueType)
-		{
-			case EPropertyBagPropertyType::Struct:
-				{
-#define UE_CAMERA_VARIABLE_FOR_TYPE(ValueType, ValueName)\
-					if (PropertyDesc.ValueTypeObject == F##ValueName##CameraParameter::StaticStruct())\
-					{\
-						using CameraParameterType = F##ValueName##CameraParameter;\
-						TValueOrError<CameraParameterType*, EPropertyBagResult> PropertyValue =\
-						ParameterOverrides.GetValueStruct<CameraParameterType>(PropertyDesc);\
-						if (ensure(PropertyValue.HasValue() && !PropertyValue.HasError()))\
-						{\
-							CameraParameterType* CameraParameter = PropertyValue.GetValue();\
-							check(CameraParameter);\
-							const uint8* DefaultValuePtr = reinterpret_cast<const uint8*>(&CameraParameter->Value);\
-							OutParameterInfos.AddBlendableParameter(\
-									PropertyDesc.Name, ECameraVariableType::ValueName, DefaultValuePtr, &CameraParameter->Variable);\
-						}\
-					}\
-					else
-					UE_CAMERA_VARIABLE_FOR_ALL_TYPES()
-#undef UE_CAMERA_VARIABLE_FOR_TYPE
-					{
-						const UScriptStruct* DataType = CastChecked<const UScriptStruct>(PropertyDesc.ValueTypeObject);
-						OutParameterInfos.AddDataParameter(PropertyDesc.Name, ECameraContextDataType::Struct, DataType, nullptr);
-					}
-				}
-				break;
-			case EPropertyBagPropertyType::Name:
-				OutParameterInfos.AddDataParameter(PropertyDesc.Name, ECameraContextDataType::Name, nullptr, nullptr);
-				break;
-			case EPropertyBagPropertyType::String:
-				OutParameterInfos.AddDataParameter(PropertyDesc.Name, ECameraContextDataType::String, nullptr, nullptr);
-				break;
-			case EPropertyBagPropertyType::Enum:
-				{
-					const UEnum* EnumType = CastChecked<const UEnum>(PropertyDesc.ValueTypeObject);
-					OutParameterInfos.AddDataParameter(PropertyDesc.Name, ECameraContextDataType::Enum, EnumType, nullptr);
-				}
-				break;
-			case EPropertyBagPropertyType::Object:
-				OutParameterInfos.AddDataParameter(PropertyDesc.Name, ECameraContextDataType::Object, nullptr, nullptr);
-				break;
-			case EPropertyBagPropertyType::Class:
-				OutParameterInfos.AddDataParameter(PropertyDesc.Name, ECameraContextDataType::Object, nullptr, nullptr);
-				break;
-		}
-	}
+	CameraRigReference.GetCustomCameraNodeParameters(OutParameterInfos);
 }
 
 FCameraNodeEvaluatorPtr UCameraRigCameraNode::OnBuildEvaluator(FCameraNodeEvaluatorBuilder& Builder) const
diff --git a/Engine/Plugins/Cameras/GameplayCameras/Source/GameplayCameras/Private/Nodes/Framing/BaseFramingCameraNode.cpp b/Engine/Plugins/Cameras/GameplayCameras/Source/GameplayCameras/Private/Nodes/Framing/BaseFramingCameraNode.cpp
--- a/Engine/Plugins/Cameras/GameplayCameras/Source/GameplayCameras/Private/Nodes/Framing/BaseFramingCameraNode.cpp
+++ b/Engine/Plugins/Cameras/GameplayCameras/Source/GameplayCameras/Private/Nodes/Framing/BaseFramingCameraNode.cpp
@@ -71,10 +71,10 @@ void FBaseFramingCameraNodeEvaluator::OnInitialize(const FCameraNodeEvaluatorIni
 TOptional<FVector3d> FBaseFramingCameraNodeEvaluator::AcquireTargetLocation(const FCameraNodeEvaluationParams& Params, const FCameraNodeEvaluationResult& InResult)
 {
 	const UBaseFramingCameraNode* FramingNode = GetCameraNodeAs<UBaseFramingCameraNode>();
-	if (FramingNode->TargetLocation)
+	if (FramingNode->TargetLocation.IsValid())
 	{
 		FVector3d TargetLocation;
-		const bool bGotTargetLocation = InResult.VariableTable.TryGetValue(FramingNode->TargetLocation.Get(), TargetLocation);
+		const bool bGotTargetLocation = InResult.VariableTable.TryGetValue(FramingNode->TargetLocation.VariableID, TargetLocation);
 		return bGotTargetLocation ? TOptional<FVector3d>(TargetLocation) : TOptional<FVector3d>();
 	}
 	else if (APlayerController* PlayerController = Params.EvaluationContext->GetPlayerController())
diff --git a/Engine/Plugins/Cameras/GameplayCameras/Source/GameplayCameras/Private/Nodes/Input/CameraRigInput1DSlot.cpp b/Engine/Plugins/Cameras/GameplayCameras/Source/GameplayCameras/Private/Nodes/Input/CameraRigInput1DSlot.cpp
--- a/Engine/Plugins/Cameras/GameplayCameras/Source/GameplayCameras/Private/Nodes/Input/CameraRigInput1DSlot.cpp
+++ b/Engine/Plugins/Cameras/GameplayCameras/Source/GameplayCameras/Private/Nodes/Input/CameraRigInput1DSlot.cpp
@@ -100,9 +100,9 @@ void UCameraRigInput1DSlot::OnBuild(FCameraRigBuildContext& BuildContext)
 	{
 		VariableDefinition = FBuiltInCameraVariables::Get().GetDefinition(BuiltInVariable);
 	}
-	else if (CustomVariable.IsValid())
+	else if (CustomVariable.HasVariable())
 	{
-		VariableDefinition = CustomVariable.Get()->GetVariableDefinition();
+		VariableDefinition = CustomVariable.Variable->GetVariableDefinition();
 	}
 	else if (InputSlotParameters.bIsPreBlended)
 	{
diff --git a/Engine/Plugins/Cameras/GameplayCameras/Source/GameplayCameras/Private/Nodes/Input/CameraRigInput2DSlot.cpp b/Engine/Plugins/Cameras/GameplayCameras/Source/GameplayCameras/Private/Nodes/Input/CameraRigInput2DSlot.cpp
--- a/Engine/Plugins/Cameras/GameplayCameras/Source/GameplayCameras/Private/Nodes/Input/CameraRigInput2DSlot.cpp
+++ b/Engine/Plugins/Cameras/GameplayCameras/Source/GameplayCameras/Private/Nodes/Input/CameraRigInput2DSlot.cpp
@@ -107,9 +107,9 @@ void UCameraRigInput2DSlot::OnBuild(FCameraRigBuildContext& BuildContext)
 	{
 		VariableDefinition = FBuiltInCameraVariables::Get().GetDefinition(BuiltInVariable);
 	}
-	else if (CustomVariable.IsValid())
+	else if (CustomVariable.HasVariable())
 	{
-		VariableDefinition = CustomVariable.Get()->GetVariableDefinition();
+		VariableDefinition = CustomVariable.Variable->GetVariableDefinition();
 	}
 	else if (InputSlotParameters.bIsPreBlended)
 	{
diff --git a/Engine/Plugins/Cameras/GameplayCameras/Source/GameplayCameras/Private/Nodes/Utility/BlueprintCameraNode.cpp b/Engine/Plugins/Cameras/GameplayCameras/Source/GameplayCameras/Private/Nodes/Utility/BlueprintCameraNode.cpp
--- a/Engine/Plugins/Cameras/GameplayCameras/Source/GameplayCameras/Private/Nodes/Utility/BlueprintCameraNode.cpp
+++ b/Engine/Plugins/Cameras/GameplayCameras/Source/GameplayCameras/Private/Nodes/Utility/BlueprintCameraNode.cpp
@@ -96,17 +96,17 @@ void FBlueprintCameraNodeEvaluator::ApplyParameterOverrides(const FCameraVariabl
 	const UBlueprintCameraNode* BlueprintNode = GetCameraNodeAs<UBlueprintCameraNode>();
 	const FCustomCameraNodeParameters& Overrides = BlueprintNode->CameraNodeEvaluatorOverrides;
 
-	// Set the value of any properties driven by a camera variable.
 	for (const FCustomCameraNodeBlendableParameter& BlendableParameter : Overrides.BlendableParameters)
 	{
 		FProperty* Property = EvaluatorBlueprintClass->FindPropertyByName(BlendableParameter.ParameterName);
 		if (ensure(Property))
 		{
-			if (BlendableParameter.OverrideVariable)
+			if (BlendableParameter.OverrideVariableID)
 			{
 				const uint8* ValuePtr = VariableTable.GetValue(
-						BlendableParameter.OverrideVariable->GetVariableID(), 
-						BlendableParameter.ParameterType);
+						BlendableParameter.OverrideVariableID,
+						BlendableParameter.ParameterType,
+						BlendableParameter.BlendableStructType);
 				Property->SetValue_InContainer(EvaluatorBlueprint, ValuePtr);
 			}
 		}
@@ -284,12 +284,17 @@ void UBlueprintCameraNode::RebuildOverrides()
 
 	// Remember the overrides already present on parameters.
 	TMap<FName, UCameraVariableAsset*> OldOverrideVariableMap;
+	TMap<FName, FCameraVariableID> OldOverrideVariableIDMap;
 	for (const FCustomCameraNodeBlendableParameter& OldOverride : CameraNodeEvaluatorOverrides.BlendableParameters)
 	{
 		if (OldOverride.OverrideVariable)
 		{
 			OldOverrideVariableMap.Add(OldOverride.ParameterName, OldOverride.OverrideVariable);
 		}
+		if (OldOverride.OverrideVariableID)
+		{
+			OldOverrideVariableIDMap.Add(OldOverride.ParameterName, OldOverride.OverrideVariableID);
+		}
 	}
 	TMap<FName, FCameraContextDataID> OldOverrideDataIDMap;
 	for (const FCustomCameraNodeDataParameter& OldOverride : CameraNodeEvaluatorOverrides.DataParameters)
@@ -401,6 +406,7 @@ void UBlueprintCameraNode::RebuildOverrides()
 				bIsBlendableProperty = true;
 				BlendablePropertyType = ECameraVariableType::Transform3d;
 			}
+			// TODO: make blendable property if the struct is registered as blendable
 			else
 			{
 				bIsDataProperty = true;
@@ -449,6 +455,12 @@ void UBlueprintCameraNode::RebuildOverrides()
 			{
 				NewOverride.OverrideVariable = OldOverrideVariable;
 			}
+			FCameraVariableID OldOverrideVariableID;
+			OldOverrideVariableIDMap.RemoveAndCopyValue(NewOverride.ParameterName, OldOverrideVariableID);
+			if (OldOverrideVariableID)
+			{
+				NewOverride.OverrideVariableID = OldOverrideVariableID;
+			}
 
 			NewOverrides.BlendableParameters.Add(NewOverride);
 		}
@@ -525,11 +537,7 @@ void UBlueprintCameraNode::GetCustomCameraNodeParameters(FCustomCameraNodeParame
 		}
 
 		const void* DefaultValuePtr = BlendableProperty->ContainerPtrToValuePtr<void>(CameraNodeEvaluatorTemplate);
-		OutParameterInfos.AddBlendableParameter(
-				BlendableParameter.ParameterName,
-				BlendableParameter.ParameterType, 
-				reinterpret_cast<const uint8*>(DefaultValuePtr),
-				&BlendableParameter.OverrideVariable);
+		OutParameterInfos.AddBlendableParameter(BlendableParameter, (const uint8*)DefaultValuePtr);
 	}
 
 	for (FCustomCameraNodeDataParameter& DataParameter : CameraNodeEvaluatorOverrides.DataParameters)
@@ -540,11 +548,8 @@ void UBlueprintCameraNode::GetCustomCameraNodeParameters(FCustomCameraNodeParame
 			continue;
 		}
 
-		OutParameterInfos.AddDataParameter(
-				DataParameter.ParameterName,
-				DataParameter.ParameterType,
-				DataParameter.ParameterTypeObject,
-				&DataParameter.OverrideDataID);
+		const void* DefaultValuePtr = DataProperty->ContainerPtrToValuePtr<void>(CameraNodeEvaluatorTemplate);
+		OutParameterInfos.AddDataParameter(DataParameter, (const uint8*)DefaultValuePtr);
 	}
 }
 
diff --git a/Engine/Plugins/Cameras/GameplayCameras/Source/GameplayCameras/Private/Tests/CameraRigAssetBuilderTests.cpp b/Engine/Plugins/Cameras/GameplayCameras/Source/GameplayCameras/Private/Tests/CameraRigAssetBuilderTests.cpp
--- a/Engine/Plugins/Cameras/GameplayCameras/Source/GameplayCameras/Private/Tests/CameraRigAssetBuilderTests.cpp
+++ b/Engine/Plugins/Cameras/GameplayCameras/Source/GameplayCameras/Private/Tests/CameraRigAssetBuilderTests.cpp
@@ -74,9 +74,8 @@ bool FCameraRigAssetBuilderSimpleParameterTest::RunTest(const FString& Parameter
 
 	UCameraRigBlendableParameter* Parameter = CameraRig->Interface.BlendableParameters[0];
 	UTEST_EQUAL("Test parameter", Parameter->InterfaceParameterName, TEXT("Test"));
-	UTEST_NOT_NULL("Test parameter variable", Parameter->PrivateVariable.Get());
-	UTEST_EQUAL("Test parameter variable name", Parameter->PrivateVariable->GetName(), "Override_SimpleTest_Test");
-	UTEST_EQUAL("Test node parameter", (UCameraVariableAsset*)OffsetNode->TranslationOffset.Variable, Parameter->PrivateVariable.Get());
+	UTEST_TRUE("Test parameter variable ID", Parameter->PrivateVariableID.IsValid());
+	UTEST_EQUAL("Test node parameter", OffsetNode->TranslationOffset.VariableID, Parameter->PrivateVariableID);
 
 	return true;
 }
@@ -163,20 +162,13 @@ bool FCameraRigAssetBuilderDrivenOverridesTest::RunTest(const FString& Parameter
 	UCameraRigBlendableParameter* OffsetParam = InnerCameraRig->Interface.BlendableParameters[0];
 	UCameraRigBlendableParameter* FocalLengthParam = InnerCameraRig->Interface.BlendableParameters[1];
 
-	UTEST_EQUAL_EXPR(OffsetParam->PrivateVariable->GetName(), "Override_InnerCameraRig_OffsetParam");
-	UTEST_EQUAL_EXPR(FocalLengthParam->PrivateVariable->GetName(), "Override_InnerCameraRig_FocalLengthParam");
-
 	// Test that the inner nodes are driven by the interface parameters.
 	{
 		UOffsetCameraNode* OffsetNode = Registry->Get<UOffsetCameraNode>("Offset");
-		UTEST_EQUAL_EXPR((UCameraVariableAsset*)OffsetNode->TranslationOffset.Variable.Get(), OffsetParam->PrivateVariable.Get());
-		UTEST_EQUAL_EXPR(OffsetNode->TranslationOffset.Variable->DefaultValue, OffsetNode->TranslationOffset.Value);
-		UTEST_EQUAL_EXPR(OffsetNode->TranslationOffset.Variable->DefaultValue, FVector3d(10, 20, 30));
+		UTEST_EQUAL_EXPR(OffsetNode->TranslationOffset.VariableID, OffsetParam->PrivateVariableID);
 
 		ULensParametersCameraNode* LensNode = Registry->Get<ULensParametersCameraNode>("Lens");
-		UTEST_EQUAL_EXPR((UCameraVariableAsset*)LensNode->FocalLength.Variable.Get(), FocalLengthParam->PrivateVariable.Get());
-		UTEST_EQUAL_EXPR(LensNode->FocalLength.Variable->DefaultValue, LensNode->FocalLength.Value);
-		UTEST_EQUAL_EXPR(LensNode->FocalLength.Variable->DefaultValue, 20.f);
+		UTEST_EQUAL_EXPR(LensNode->FocalLength.VariableID, FocalLengthParam->PrivateVariableID);
 	}
 
 	// Test that the middle prefab node is driving the inner interface parameters, and that one of those
@@ -207,8 +199,7 @@ bool FCameraRigAssetBuilderDrivenOverridesTest::RunTest(const FString& Parameter
 
 		const FPropertyBagPropertyDesc* OffsetParamDesc = ParameterOverrides.FindPropertyDescByName("OffsetParam");
 		FVector3dCameraParameter* OffsetParamOverride = ParameterOverrides.GetValueStruct<FVector3dCameraParameter>("OffsetParam").GetValue();
-		UTEST_EQUAL_EXPR((UCameraVariableAsset*)OffsetParamOverride->Variable.Get(), MiddleOffsetParam->PrivateVariable.Get());
-		UTEST_EQUAL_EXPR(OffsetParamOverride->Variable->DefaultValue, FVector3d(15, 25, 35));
+		UTEST_EQUAL_EXPR(OffsetParamOverride->VariableID, MiddleOffsetParam->PrivateVariableID);
 	}
 
 	// Test that the outer prefab node is driving the middle interface parameters.
diff --git a/Engine/Plugins/Cameras/GameplayCameras/Source/GameplayCameras/Public/Core/BuiltInCameraVariables.h b/Engine/Plugins/Cameras/GameplayCameras/Source/GameplayCameras/Public/Core/BuiltInCameraVariables.h
--- a/Engine/Plugins/Cameras/GameplayCameras/Source/GameplayCameras/Public/Core/BuiltInCameraVariables.h
+++ b/Engine/Plugins/Cameras/GameplayCameras/Source/GameplayCameras/Public/Core/BuiltInCameraVariables.h
@@ -2,7 +2,7 @@
 
 #pragma once
 
-#include "Core/CameraVariableTableFwd.h"
+#include "Core/CameraVariableTableAllocationInfo.h"
 #include "CoreTypes.h"
 
 #include "BuiltInCameraVariables.generated.h"
diff --git a/Engine/Plugins/Cameras/GameplayCameras/Source/GameplayCameras/Public/Core/CameraAsset.h b/Engine/Plugins/Cameras/GameplayCameras/Source/GameplayCameras/Public/Core/CameraAsset.h
--- a/Engine/Plugins/Cameras/GameplayCameras/Source/GameplayCameras/Public/Core/CameraAsset.h
+++ b/Engine/Plugins/Cameras/GameplayCameras/Source/GameplayCameras/Public/Core/CameraAsset.h
@@ -3,10 +3,11 @@
 #pragma once
 
 #include "Core/CameraBuildStatus.h"
-#include "Core/CameraContextDataAllocationInfo.h"
+#include "Core/CameraContextDataTableAllocationInfo.h"
 #include "Core/CameraEventHandler.h"
+#include "Core/CameraRigParameterDefinition.h"
 #include "Core/CameraRigTransition.h"
-#include "Core/CameraVariableTableFwd.h"
+#include "Core/CameraVariableTableAllocationInfo.h"
 #include "Core/ObjectTreeGraphObject.h"
 #include "Core/ObjectTreeGraphRootObject.h"
 #include "CoreTypes.h"
@@ -56,7 +57,7 @@ struct FCameraAssetAllocationInfo
 
 	/** Combined context data table allocation info for all the camera rigs. */
 	UPROPERTY()
-	FCameraContextDataAllocationInfo ContextDataTableInfo;
+	FCameraContextDataTableAllocationInfo ContextDataTableInfo;
 
 	GAMEPLAYCAMERAS_API friend bool operator==(const FCameraAssetAllocationInfo& A, const FCameraAssetAllocationInfo& B);
 };
@@ -94,11 +95,16 @@ public:
 	/** Removes an exit transition. */
 	GAMEPLAYCAMERAS_API int32 RemoveExitTransition(UCameraRigTransition* InTransition);
 
+public:
+
 	/** Gets the default parameter values for all camera rigs. */
 	const FInstancedPropertyBag& GetDefaultParameters() const { return DefaultParameters; }
 	/** Gets the default parameter values for all camera rigs. */
 	FInstancedPropertyBag& GetDefaultParameters() { return DefaultParameters; }
 
+	/** Gets the definitions of parameters exposed on this camera asset. */
+	TConstArrayView<FCameraRigParameterDefinition> GetParameterDefinitions() const { return ParameterDefinitions; }
+
 public:
 
 	/**
@@ -180,9 +186,13 @@ private:
 	UPROPERTY()
 	FInstancedPropertyBag DefaultParameters;
 
-	/** Generated list of camera rigs owning each of the parameters. */
+	/** Generated list of all the camera rigs' parameters. */
+	UPROPERTY()
+	TArray<FCameraRigParameterDefinition> ParameterDefinitions;
+
+	/** Owner rigs for each parameter definition. */
 	UPROPERTY()
-	TArray<TWeakObjectPtr<const UCameraRigAsset>> ParameterOwners;
+	TArray<TObjectPtr<const UCameraRigAsset>> ParameterOwners;
 
 	/** Allocation info for the camera asset. */
 	UPROPERTY()
diff --git a/Engine/Plugins/Cameras/GameplayCameras/Source/GameplayCameras/Public/Core/CameraAssetReference.h b/Engine/Plugins/Cameras/GameplayCameras/Source/GameplayCameras/Public/Core/CameraAssetReference.h
--- a/Engine/Plugins/Cameras/GameplayCameras/Source/GameplayCameras/Public/Core/CameraAssetReference.h
+++ b/Engine/Plugins/Cameras/GameplayCameras/Source/GameplayCameras/Public/Core/CameraAssetReference.h
@@ -71,7 +71,7 @@ public:
 	}
 
 	/** Gets the IDs of the parameters with override values. */
-	TConstArrayView<FGuid> GetOverridenParameterGuids() const
+	TConstArrayView<FGuid> GetOverriddenParameterGuids() const
 	{
 		return ParameterOverrideGuids;
 	}
@@ -85,8 +85,8 @@ public:
 
 	bool SerializeFromMismatchedTag(FPropertyTag const& Tag, FStructuredArchive::FSlot Slot);
 
-	GAMEPLAYCAMERAS_API bool IsParameterOverriden(const FGuid PropertyID) const;
-	GAMEPLAYCAMERAS_API void SetParameterOverriden(const FGuid PropertyID, bool bIsOverridden);
+	GAMEPLAYCAMERAS_API bool IsParameterOverridden(const FGuid PropertyID) const;
+	GAMEPLAYCAMERAS_API void SetParameterOverridden(const FGuid PropertyID, bool bIsOverridden);
 
 	GAMEPLAYCAMERAS_API bool NeedsRebuildParameters() const;
 	GAMEPLAYCAMERAS_API bool RebuildParametersIfNeeded();
diff --git a/Engine/Plugins/Cameras/GameplayCameras/Source/GameplayCameras/Public/Core/CameraContextDataAllocationInfo.h b/Engine/Plugins/Cameras/GameplayCameras/Source/GameplayCameras/Public/Core/CameraContextDataAllocationInfo.h
deleted file mode 100644
diff --git a/Engine/Plugins/Cameras/GameplayCameras/Source/GameplayCameras/Public/Core/CameraContextDataReader.h b/Engine/Plugins/Cameras/GameplayCameras/Source/GameplayCameras/Public/Core/CameraContextDataReader.h
new file mode 100644
--- /dev/null
+++ b/Engine/Plugins/Cameras/GameplayCameras/Source/GameplayCameras/Public/Core/CameraContextDataReader.h
@@ -0,0 +1,59 @@
+// Copyright Epic Games, Inc. All Rights Reserved.
+
+#pragma once
+
+#include "Core/CameraContextDataTable.h"
+#include "CoreTypes.h"
+#include "Templates/UnrealTypeTraits.h"
+
+namespace UE::Cameras
+{
+
+template<typename DataType, ECameraContextDataType DataTypeEnum>
+struct TCameraContextDataReader
+{
+	TCameraContextDataReader() {}
+
+	template<typename ParameterType>
+	TCameraContextDataReader(const ParameterType& Parameter)
+	{
+		Initialize(Parameter);
+	}
+
+	template<typename ParameterType>
+	void Initialize(const ParameterType& Parameter)
+	{
+		static_assert(
+				std::is_same<DataType, typename ParameterType::DataType>(),
+				"The given parameter is of the wrong type for this reader! Data types must be the same.");
+
+		DefaultValuePtr = &Parameter.Value;
+		DataID = Parameter.DataID;
+
+		ensureMsgf(DefaultValuePtr, TEXT("The given parameter doesn't have a value!"));
+	}
+
+	DataType Get(const FCameraContextDataTable& ContextDataTable) const
+	{
+		if (!DataID.IsValid())
+		{
+			return *DefaultValuePtr;
+		}
+		else
+		{
+			if (const DataType* ActualValue = ContextDataTable.TryGetData<DataType>(DataID))
+			{
+				return *ActualValue;
+			}
+			return *DefaultValuePtr;
+		}
+	}
+
+private:
+
+	const DataType* DefaultValuePtr = nullptr;
+	FCameraContextDataID DataID;
+};
+
+}  // namespace UE::Cameras
+
diff --git a/Engine/Plugins/Cameras/GameplayCameras/Source/GameplayCameras/Public/Core/CameraContextDataTable.h b/Engine/Plugins/Cameras/GameplayCameras/Source/GameplayCameras/Public/Core/CameraContextDataTable.h
--- a/Engine/Plugins/Cameras/GameplayCameras/Source/GameplayCameras/Public/Core/CameraContextDataTable.h
+++ b/Engine/Plugins/Cameras/GameplayCameras/Source/GameplayCameras/Public/Core/CameraContextDataTable.h
@@ -9,9 +9,12 @@
 #include "Math/NumericLimits.h"
 #include "StructUtils/InstancedStruct.h"
 #include "StructUtils/StructView.h"
+#include "Templates/PointerIsConvertibleFromTo.h"
+#include "Templates/Requires.h"
+#include <type_traits>
 
-struct FCameraContextDataAllocationInfo;
 struct FCameraContextDataDefinition;
+struct FCameraContextDataTableAllocationInfo;
 
 namespace UE::Cameras
 {
@@ -23,9 +26,10 @@ class FContextDataTableDebugBlock;
 #endif  // UE_GAMEPLAY_CAMERAS_DEBUG
 
 template<typename DataType>
-struct TCameraContextDataReader
+struct TCameraContextDataTraits
 {
-	void Initialize(FCameraContextDataID InID, const FCameraContextDataTable& InContextDataTable);
+	static ECameraContextDataType GetDataType();
+	static const UObject* GetDataTypeObject();
 };
 
 /**
@@ -54,7 +58,7 @@ public:
 	~FCameraContextDataTable();
 
 	/** Initializes the context data table so that it fits the provided allocation info. */
-	void Initialize(const FCameraContextDataAllocationInfo& AllocationInfo);
+	void Initialize(const FCameraContextDataTableAllocationInfo& AllocationInfo);
 
 	/** Adds a data entry to the table. */
 	void AddData(const FCameraContextDataDefinition& DataDefinition);
@@ -83,6 +87,9 @@ public:
 	template<typename BaseClass>
 	TSubclassOf<BaseClass> GetClassData(FCameraContextDataID InID) const;
 
+	template<typename ValueType>
+	const ValueType* TryGetData(FCameraContextDataID InID) const;
+
 	// Setter methods.
 
 	void SetNameData(FCameraContextDataID InID, const FName& InData);
@@ -147,7 +154,7 @@ private:
 		uint32 Offset;
 		EEntryFlags Flags;
 #if WITH_EDITORONLY_DATA
-		FName DebugName;
+		FString DebugName;
 #endif
 	};
 
@@ -227,6 +234,18 @@ TSubclassOf<BaseClass> FCameraContextDataTable::GetClassData(FCameraContextDataI
 	return nullptr;
 }
 
+template<typename ValueType>
+const ValueType* FCameraContextDataTable::TryGetData(FCameraContextDataID InID) const
+{
+	ECameraContextDataType DataType = TCameraContextDataTraits<ValueType>::GetDataType();
+	const UObject* DataTypeObject = TCameraContextDataTraits<ValueType>::GetDataTypeObject();
+	if (const uint8* RawValue = TryGetData(InID, DataType, DataTypeObject))
+	{
+		return reinterpret_cast<const ValueType*>(RawValue);
+	}
+	return nullptr;
+}
+
 template<typename EnumType>
 void FCameraContextDataTable::SetEnumData(FCameraContextDataID InID, EnumType InData)
 {
@@ -265,5 +284,51 @@ bool FCameraContextDataTable::SetDataImpl(FCameraContextDataID InID, ECameraCont
 	return false;
 }
 
+template<typename DataType>
+ECameraContextDataType TCameraContextDataTraits<DataType>::GetDataType()
+{
+	if constexpr(std::is_same_v<DataType, FName>)
+	{
+		return ECameraContextDataType::Name;
+	}
+	else if constexpr(std::is_same_v<DataType, FString>)
+	{
+		return ECameraContextDataType::Name;
+	}
+	else if constexpr(std::is_enum_v<DataType>)
+	{
+		return ECameraContextDataType::Enum;
+	}
+	else if constexpr(TPointerIsConvertibleFromTo<DataType, UClass>::Value)
+	{
+		return ECameraContextDataType::Class;
+	}
+	else if constexpr(TPointerIsConvertibleFromTo<DataType, UObject>::Value)
+	{
+		return ECameraContextDataType::Object;
+	}
+	else if constexpr(std::is_same_v<decltype(DataType::StaticStruct), UScriptStruct*>)
+	{
+		return ECameraContextDataType::Struct;
+	}
+};
+
+template<typename DataType>
+const UObject* TCameraContextDataTraits<DataType>::GetDataTypeObject()
+{
+	if constexpr(std::is_enum_v<DataType>)
+	{
+		return StaticEnum<DataType>();
+	}
+	else if constexpr(std::is_same_v<decltype(DataType::StaticStruct), UScriptStruct*>)
+	{
+		return DataType::StaticStruct();
+	}
+	else
+	{
+		return nullptr;
+	}
+}
+
 }  // namespace UE::Cameras
 
diff --git a/Engine/Plugins/Cameras/GameplayCameras/Source/GameplayCameras/Public/Core/CameraContextDataTableAllocationInfo.h b/Engine/Plugins/Cameras/GameplayCameras/Source/GameplayCameras/Public/Core/CameraContextDataTableAllocationInfo.h
new file mode 100644
--- /dev/null
+++ b/Engine/Plugins/Cameras/GameplayCameras/Source/GameplayCameras/Public/Core/CameraContextDataTableAllocationInfo.h
@@ -0,0 +1,60 @@
+// Copyright Epic Games, Inc. All Rights Reserved.
+
+#pragma once
+
+#include "Core/CameraContextDataTableFwd.h"
+#include "UObject/Class.h"
+#include "UObject/ObjectMacros.h"
+#include "UObject/UnrealNames.h"
+
+#include "CameraContextDataTableAllocationInfo.generated.h"
+
+class UCameraNode;
+
+/**
+ * Definition for one entry in a camera rig's context data registry.
+ */
+USTRUCT()
+struct FCameraContextDataDefinition
+{
+	GENERATED_BODY()
+
+	/** The ID of data. */
+	UPROPERTY()
+	FCameraContextDataID DataID;
+
+	/** The type of the data. */
+	UPROPERTY()
+	ECameraContextDataType DataType;
+
+	/** An extra type object for the data. */
+	UPROPERTY()
+	TObjectPtr<const UObject> DataTypeObject;
+
+#if WITH_EDITORONLY_DATA
+	/** The name of the data, for debugging purposes. */
+	UPROPERTY()
+	FString DataName;
+#endif
+
+	bool operator==(const FCameraContextDataDefinition& Other) const = default;
+};
+
+/**
+ * Collection of context data entries for a camera rig.
+ */
+USTRUCT()
+struct FCameraContextDataTableAllocationInfo
+{
+	GENERATED_BODY()
+
+	/** The list of context data definitions. */
+	UPROPERTY()
+	TArray<FCameraContextDataDefinition> DataDefinitions;
+
+	/**Combines the given allocation info with this one. */
+	GAMEPLAYCAMERAS_API void Combine(const FCameraContextDataTableAllocationInfo& OtherInfo);
+
+	bool operator==(const FCameraContextDataTableAllocationInfo& Other) const = default;
+};
+
diff --git a/Engine/Plugins/Cameras/GameplayCameras/Source/GameplayCameras/Public/Core/CameraContextDataTableFwd.h b/Engine/Plugins/Cameras/GameplayCameras/Source/GameplayCameras/Public/Core/CameraContextDataTableFwd.h
--- a/Engine/Plugins/Cameras/GameplayCameras/Source/GameplayCameras/Public/Core/CameraContextDataTableFwd.h
+++ b/Engine/Plugins/Cameras/GameplayCameras/Source/GameplayCameras/Public/Core/CameraContextDataTableFwd.h
@@ -45,49 +45,54 @@ struct FCameraContextDataID
 
 public:
 
-	FCameraContextDataID() : DataName(NAME_None) {}
+	FCameraContextDataID() : Value(INVALID) {}
 
-	bool IsValid() const { return DataName != NAME_None; }
+	uint32 GetValue() const { return Value; }
+
+	bool IsValid() const { return Value != INVALID; }
 
 	explicit operator bool() const { return IsValid(); }
 
-	static FCameraContextDataID FromName(FName InName)
+	static FCameraContextDataID FromHashValue(uint32 InValue)
 	{
-		return FCameraContextDataID{ InName };
+		return FCameraContextDataID(InValue);
 	}
 
 public:
 
+	friend bool operator<(FCameraContextDataID A, FCameraContextDataID B)
+	{
+		return A.Value < B.Value;
+	}
+
 	friend bool operator==(FCameraContextDataID A, FCameraContextDataID B)
 	{
-		return A.DataName == B.DataName;
+		return A.Value == B.Value;
 	}
 
 	friend bool operator!=(FCameraContextDataID A, FCameraContextDataID B)
 	{
-		return !(A == B);
+		return A.Value != B.Value;
 	}
 
-	friend bool operator<(FCameraContextDataID A, FCameraContextDataID B)
+	friend uint32 GetTypeHash(FCameraContextDataID In)
 	{
-		return A.DataName.Compare(B.DataName) < 0;
+		return In.Value;
 	}
 
-	friend uint32 GetTypeHash(FCameraContextDataID In)
+	friend FArchive& operator<< (FArchive& Ar, FCameraContextDataID& In)
 	{
-		return GetTypeHash(In.DataName);
+		Ar << In.Value;
+		return Ar;
 	}
 
 private:
 
-	FCameraContextDataID(FName InName) : DataName(InName) {}
+	FCameraContextDataID(uint32 InValue) : Value(InValue) {}
 
-	UPROPERTY()
-	FName DataName = NAME_None;
+	static const uint32 INVALID = uint32(-1);
 
-	friend class UE::Cameras::FCameraContextDataTable;
-#if UE_GAMEPLAY_CAMERAS_DEBUG
-	friend class UE::Cameras::FContextDataTableDebugBlock;
-#endif  // UE_GAMEPLAY_CAMERAS_DEBUG
+	UPROPERTY()
+	uint32 Value;
 };
 
diff --git a/Engine/Plugins/Cameras/GameplayCameras/Source/GameplayCameras/Public/Core/CameraParameterReader.h b/Engine/Plugins/Cameras/GameplayCameras/Source/GameplayCameras/Public/Core/CameraParameterReader.h
--- a/Engine/Plugins/Cameras/GameplayCameras/Source/GameplayCameras/Public/Core/CameraParameterReader.h
+++ b/Engine/Plugins/Cameras/GameplayCameras/Source/GameplayCameras/Public/Core/CameraParameterReader.h
@@ -9,6 +9,9 @@
 namespace UE::Cameras
 {
 
+/**
+ * A utility class for reading the effective value of a camera parameter.
+ */
 template<typename ValueType>
 class TCameraParameterReader
 {
@@ -33,14 +36,9 @@ public:
 				"The given parameter is of the wrong type for this reader! Value types must be the same.");
 
 		DefaultValuePtr = &Parameter.Value;
-		ensureMsgf(DefaultValuePtr, TEXT("The given parameter doesn't have a value!"));
+		VariableID = Parameter.VariableID;
 
-		if (Parameter.Variable)
-		{
-			VariableID = Parameter.Variable->GetVariableID();
-			DefaultValuePtr = reinterpret_cast<const ValueType*>(Parameter.Variable->GetDefaultValuePtr());
-			ensureMsgf(DefaultValuePtr, TEXT("The given parameter's driving variable doesn't have a default value!"));
-		}
+		ensureMsgf(DefaultValuePtr, TEXT("The given parameter doesn't have a value!"));
 	}
 
 	/**
diff --git a/Engine/Plugins/Cameras/GameplayCameras/Source/GameplayCameras/Public/Core/CameraParameters.h b/Engine/Plugins/Cameras/GameplayCameras/Source/GameplayCameras/Public/Core/CameraParameters.h
--- a/Engine/Plugins/Cameras/GameplayCameras/Source/GameplayCameras/Public/Core/CameraParameters.h
+++ b/Engine/Plugins/Cameras/GameplayCameras/Source/GameplayCameras/Public/Core/CameraParameters.h
@@ -11,12 +11,25 @@
 #define UE_DEFINE_CAMERA_PARAMETER_VALUE_CONSTRUCTORS(ParameterClass)\
 	ParameterClass(typename TCallTraits<ParameterClass::ValueType>::ParamType InValue)\
 		: Value(InValue)\
-	{}
+	{}\
+	bool HasUserOverride() const { return Variable != nullptr; }\
+	bool HasNonUserOverride() const { return VariableID.IsValid() && (!Variable || Variable->GetVariableID() != VariableID); }\
+	void PostSerialize(const FArchive& Ar);
 
 #define UE_DEFINE_CAMERA_PARAMETER_ALL_CONSTRUCTORS(ParameterClass)\
 	ParameterClass() {}\
 	UE_DEFINE_CAMERA_PARAMETER_VALUE_CONSTRUCTORS(ParameterClass)
 
+// All camera parameters have:
+//
+// - Value: a value for the user to tweak. This is the "default" value.
+// - Variable: a variable chosen by the user to drive this parameter. 
+// - VariableID: the ID of the variable driving this parameter.
+//
+// When Variable is set, VariableID is the ID of that variable.
+// When Variable is not set, VariableID is the ID of something else
+// driving the parameter, such as a camera rig parameter override.
+
 /** Boolean camera parameter. */
 USTRUCT()
 struct GAMEPLAYCAMERAS_API FBooleanCameraParameter
@@ -29,6 +42,9 @@ struct GAMEPLAYCAMERAS_API FBooleanCameraParameter
 	UPROPERTY(EditAnywhere, Category=Common)
 	bool Value = false;
 
+	UPROPERTY()
+	FCameraVariableID VariableID;
+
 	UPROPERTY(EditAnywhere, Category=Common)
 	TObjectPtr<UBooleanCameraVariable> Variable;
 
@@ -49,6 +65,9 @@ struct GAMEPLAYCAMERAS_API FInteger32CameraParameter
 	UPROPERTY(EditAnywhere, Category=Common)
 	int32 Value = 0;
 
+	UPROPERTY()
+	FCameraVariableID VariableID;
+
 	UPROPERTY(EditAnywhere, Category=Common)
 	TObjectPtr<UInteger32CameraVariable> Variable;
 
@@ -69,6 +88,9 @@ struct FFloatCameraParameter
 	UPROPERTY(EditAnywhere, Category=Common)
 	float Value = 0.f;
 
+	UPROPERTY()
+	FCameraVariableID VariableID;
+
 	UPROPERTY(EditAnywhere, Category=Common)
 	TObjectPtr<UFloatCameraVariable> Variable;
 
@@ -89,6 +111,9 @@ struct GAMEPLAYCAMERAS_API FDoubleCameraParameter
 	UPROPERTY(EditAnywhere, Category=Common)
 	double Value = 0.0;
 
+	UPROPERTY()
+	FCameraVariableID VariableID;
+
 	UPROPERTY(EditAnywhere, Category=Common)
 	TObjectPtr<UDoubleCameraVariable> Variable;
 
@@ -109,6 +134,9 @@ struct GAMEPLAYCAMERAS_API FVector2fCameraParameter
 	UPROPERTY(EditAnywhere, Category=Common)
 	FVector2f Value;
 
+	UPROPERTY()
+	FCameraVariableID VariableID;
+
 	UPROPERTY(EditAnywhere, Category=Common)
 	TObjectPtr<UVector2fCameraVariable> Variable;
 
@@ -130,6 +158,9 @@ struct GAMEPLAYCAMERAS_API FVector2dCameraParameter
 	UPROPERTY(EditAnywhere, Category=Common)
 	FVector2D Value;
 
+	UPROPERTY()
+	FCameraVariableID VariableID;
+
 	UPROPERTY(EditAnywhere, Category=Common)
 	TObjectPtr<UVector2dCameraVariable> Variable;
 
@@ -151,6 +182,9 @@ struct GAMEPLAYCAMERAS_API FVector3fCameraParameter
 	UPROPERTY(EditAnywhere, Category=Common)
 	FVector3f Value;
 
+	UPROPERTY()
+	FCameraVariableID VariableID;
+
 	UPROPERTY(EditAnywhere, Category=Common)
 	TObjectPtr<UVector3fCameraVariable> Variable;
 
@@ -172,6 +206,9 @@ struct GAMEPLAYCAMERAS_API FVector3dCameraParameter
 	UPROPERTY(EditAnywhere, Category=Common)
 	FVector3d Value;
 
+	UPROPERTY()
+	FCameraVariableID VariableID;
+
 	UPROPERTY(EditAnywhere, Category=Common)
 	TObjectPtr<UVector3dCameraVariable> Variable;
 
@@ -193,6 +230,9 @@ struct GAMEPLAYCAMERAS_API FVector4fCameraParameter
 	UPROPERTY(EditAnywhere, Category=Common)
 	FVector4f Value;
 
+	UPROPERTY()
+	FCameraVariableID VariableID;
+
 	UPROPERTY(EditAnywhere, Category=Common)
 	TObjectPtr<UVector4fCameraVariable> Variable;
 
@@ -214,6 +254,9 @@ struct GAMEPLAYCAMERAS_API FVector4dCameraParameter
 	UPROPERTY(EditAnywhere, Category=Common)
 	FVector4d Value;
 
+	UPROPERTY()
+	FCameraVariableID VariableID;
+
 	UPROPERTY(EditAnywhere, Category=Common)
 	TObjectPtr<UVector4dCameraVariable> Variable;
 
@@ -235,6 +278,9 @@ struct GAMEPLAYCAMERAS_API FRotator3fCameraParameter
 	UPROPERTY(EditAnywhere, Category=Common)
 	FRotator3f Value;
 
+	UPROPERTY()
+	FCameraVariableID VariableID;
+
 	UPROPERTY(EditAnywhere, Category=Common)
 	TObjectPtr<URotator3fCameraVariable> Variable;
 
@@ -256,6 +302,9 @@ struct GAMEPLAYCAMERAS_API FRotator3dCameraParameter
 	UPROPERTY(EditAnywhere, Category=Common)
 	FRotator3d Value;
 
+	UPROPERTY()
+	FCameraVariableID VariableID;
+
 	UPROPERTY(EditAnywhere, Category=Common)
 	TObjectPtr<URotator3dCameraVariable> Variable;
 
@@ -277,6 +326,9 @@ struct GAMEPLAYCAMERAS_API FTransform3fCameraParameter
 	UPROPERTY(EditAnywhere, Category=Common)
 	FTransform3f Value;
 
+	UPROPERTY()
+	FCameraVariableID VariableID;
+
 	UPROPERTY(EditAnywhere, Category=Common)
 	TObjectPtr<UTransform3fCameraVariable> Variable;
 
@@ -297,6 +349,9 @@ struct GAMEPLAYCAMERAS_API FTransform3dCameraParameter
 	UPROPERTY(EditAnywhere, Category=Common)
 	FTransform3d Value;
 
+	UPROPERTY()
+	FCameraVariableID VariableID;
+
 	UPROPERTY(EditAnywhere, Category=Common)
 	TObjectPtr<UTransform3dCameraVariable> Variable;
 
@@ -335,7 +390,7 @@ inline bool CameraParameterValueEquals<FTransform3d>(const FTransform3d& A, cons
 	template<> struct TStructOpsTypeTraits<F##ValueName##CameraParameter>\
 		: public TStructOpsTypeTraitsBase2<F##ValueName##CameraParameter>\
 	{\
-		enum { WithStructuredSerializeFromMismatchedTag = true };\
+		enum { WithStructuredSerializeFromMismatchedTag = true, WithPostSerialize = true };\
 	};
 UE_CAMERA_VARIABLE_FOR_ALL_TYPES()
 #undef UE_CAMERA_VARIABLE_FOR_TYPE
diff --git a/Engine/Plugins/Cameras/GameplayCameras/Source/GameplayCameras/Public/Core/CameraRigAsset.h b/Engine/Plugins/Cameras/GameplayCameras/Source/GameplayCameras/Public/Core/CameraRigAsset.h
--- a/Engine/Plugins/Cameras/GameplayCameras/Source/GameplayCameras/Public/Core/CameraRigAsset.h
+++ b/Engine/Plugins/Cameras/GameplayCameras/Source/GameplayCameras/Public/Core/CameraRigAsset.h
@@ -3,12 +3,13 @@
 #pragma once
 
 #include "Core/CameraBuildStatus.h"
-#include "Core/CameraContextDataAllocationInfo.h"
+#include "Core/CameraContextDataTableAllocationInfo.h"
 #include "Core/CameraContextDataTableFwd.h"
 #include "Core/CameraEventHandler.h"
 #include "Core/CameraNodeEvaluatorFwd.h"
+#include "Core/CameraRigParameterDefinition.h"
 #include "Core/CameraRigTransition.h"
-#include "Core/CameraVariableTableFwd.h"
+#include "Core/CameraVariableTableAllocationInfo.h"
 #include "Core/ObjectTreeGraphObject.h"
 #include "Core/ObjectTreeGraphRootObject.h"
 #include "CoreTypes.h"
@@ -67,13 +68,13 @@ struct FCameraRigAllocationInfo
 
 	/** Allocation info for the context data table. */
 	UPROPERTY()
-	FCameraContextDataAllocationInfo ContextDataTableInfo;
+	FCameraContextDataTableAllocationInfo ContextDataTableInfo;
 
 public:
 
 	GAMEPLAYCAMERAS_API void Append(const FCameraRigAllocationInfo& OtherAllocationInfo);
 
-	GAMEPLAYCAMERAS_API friend bool operator==(const FCameraRigAllocationInfo& A, const FCameraRigAllocationInfo& B);
+	bool operator==(const FCameraRigAllocationInfo& Other) const = default;
 };
 
 template<>
@@ -86,16 +87,6 @@ struct TStructOpsTypeTraits<FCameraRigAllocationInfo> : public TStructOpsTypeTra
 	};
 };
 
-/**
- * The type of a camera rig parameter.
- */
-UENUM()
-enum class ECameraRigInterfaceParameterType : uint8
-{
-	Blendable,
-	Data
-};
-
 /**
  * Base class for interface parameters on a camera rig asset.
  */
@@ -183,6 +174,10 @@ public:
 	UPROPERTY()
 	ECameraVariableType ParameterType = ECameraVariableType::Boolean;
 
+	/** The struct type of this parameter if it is a blendable struct. */
+	UPROPERTY()
+	TObjectPtr<const UScriptStruct> BlendableStructType;
+
 	/**
 	 * Whether this parameter's value should be pre-blended.
 	 *
@@ -195,13 +190,24 @@ public:
 
 	// Built on save/cook.
 
-	/**
-	 * The private camera variable created to drive the target camera parameter on
-	 * the target camera node. This variable is created by the build method on the
-	 * camera rig.
-	 */
+	/** The ID to use to access the underlying variable value in the variable table. */
 	UPROPERTY()
-	TObjectPtr<UCameraVariableAsset> PrivateVariable;
+	FCameraVariableID PrivateVariableID;
+
+
+	// Deprecated.
+
+	UPROPERTY()
+	TObjectPtr<UCameraVariableAsset> PrivateVariable_DEPRECATED;
+
+public:
+
+	/** Gets the camera variable definition for this parameter. */
+	GAMEPLAYCAMERAS_API FCameraVariableDefinition GetVariableDefinition() const;
+
+#if WITH_EDITORONLY_DATA
+	GAMEPLAYCAMERAS_API FString GetVariableName() const;
+#endif
 };
 
 UCLASS(MinimalAPI)
@@ -221,9 +227,18 @@ public:
 
 	// Built on save/cook.
 
-	/** The reference to use to access the underlying data in the context data table. */
+	/** The ID to use to access the underlying data in the context data table. */
 	UPROPERTY()
 	FCameraContextDataID PrivateDataID;
+
+public:
+
+	/** Gets the camera context data definition for this parameter. */
+	GAMEPLAYCAMERAS_API FCameraContextDataDefinition GetDataDefinition() const;
+
+#if WITH_EDITORONLY_DATA
+	GAMEPLAYCAMERAS_API FString GetDataName() const;
+#endif
 };
 
 /**
@@ -236,10 +251,11 @@ struct FCameraRigInterface
 
 public:
 
-	/** The list of exposed parameters on the camera rig. */
+	/** The list of exposed blendable parameters on the camera rig. */
 	UPROPERTY(Instanced)
 	TArray<TObjectPtr<UCameraRigBlendableParameter>> BlendableParameters;
 
+	/** The list of exposed data parameters on the camera rig. */
 	UPROPERTY(Instanced)
 	TArray<TObjectPtr<UCameraRigDataParameter>> DataParameters;
 
@@ -322,6 +338,9 @@ public:
 	/** Gets the default values for the parameters exposed on this camera rig. */
 	FInstancedPropertyBag& GetDefaultParameters() { return DefaultParameters; }
 
+	/** Gets the definitions of parameters exposed on this camera rig. */
+	TConstArrayView<FCameraRigParameterDefinition> GetParameterDefinitions() const { return ParameterDefinitions; }
+
 public:
 
 	/** The current build state of this camera rig. */
@@ -397,6 +416,10 @@ private:
 	UPROPERTY()
 	FInstancedPropertyBag DefaultParameters;
 
+	/** The definitions of parameters exposed on this camera rig. */
+	UPROPERTY()
+	TArray<FCameraRigParameterDefinition> ParameterDefinitions;
+
 #if WITH_EDITORONLY_DATA
 
 	/** Position of the camera rig node in the node graph editor. */
diff --git a/Engine/Plugins/Cameras/GameplayCameras/Source/GameplayCameras/Public/Core/CameraRigAssetBuilder.h b/Engine/Plugins/Cameras/GameplayCameras/Source/GameplayCameras/Public/Core/CameraRigAssetBuilder.h
--- a/Engine/Plugins/Cameras/GameplayCameras/Source/GameplayCameras/Public/Core/CameraRigAssetBuilder.h
+++ b/Engine/Plugins/Cameras/GameplayCameras/Source/GameplayCameras/Public/Core/CameraRigAssetBuilder.h
@@ -15,7 +15,6 @@ class FStructProperty;
 class UCameraNode;
 class UCameraRigAsset;
 class UCameraRigCameraNode;
-class UCameraVariableAsset;
 struct FCameraVariableTableAllocationInfo;
 struct FInstancedPropertyBag;
 
@@ -63,8 +62,10 @@ private:
 
 	void CallBuild();
 	void CallBuild(FCameraRigBuildContext& BuildContext, UCameraNode* CameraNode);
+	void BuildParametersAllocationInfo(FCameraRigBuildContext& BuildContext);
 
 	void BuildDefaultParameters();
+	void BuildParameterDefinitions();
 
 	void UpdateBuildStatus();
 
@@ -76,13 +77,6 @@ private:
 	bool SetupDataContextPropertyOverride(const UCameraRigDataParameter* DataParameter);
 	bool SetupCustomDataParameterOverride(const UCameraRigDataParameter* DataParameter);
 
-public:
-
-	// Internal API.
-
-	static void BuildDefaultParameters(UCameraRigAsset* CameraRigAsset, FInstancedPropertyBag& OutPropertyBag);
-	static void AppendDefaultParameters(const FCameraRigInterface& CameraRigInterface, TArray<FPropertyBagPropertyDesc>& OutProperties);
-
 private:
 
 	FCameraBuildLog& BuildLog;
@@ -92,7 +86,7 @@ private:
 	FCameraNodeHierarchy CameraNodeHierarchy;
 
 	using FDrivenParameterKey = TTuple<FName, UCameraNode*>;
-	TMap<FDrivenParameterKey, TObjectPtr<UCameraVariableAsset>> OldDrivenBlendableParameters;
+	TMap<FDrivenParameterKey, FCameraVariableID> OldDrivenBlendableParameters;
 	TMap<FDrivenParameterKey, FCameraContextDataID> OldDrivenDataParameters;
 
 	friend struct Internal::FInterfaceParameterBindingBuilder;
diff --git a/Engine/Plugins/Cameras/GameplayCameras/Source/GameplayCameras/Public/Core/CameraRigAssetReference.h b/Engine/Plugins/Cameras/GameplayCameras/Source/GameplayCameras/Public/Core/CameraRigAssetReference.h
--- a/Engine/Plugins/Cameras/GameplayCameras/Source/GameplayCameras/Public/Core/CameraRigAssetReference.h
+++ b/Engine/Plugins/Cameras/GameplayCameras/Source/GameplayCameras/Public/Core/CameraRigAssetReference.h
@@ -2,7 +2,9 @@
 
 #pragma once
 
+#include "Core/CameraContextDataTableFwd.h"
 #include "Core/CameraParameters.h"
+#include "Core/CameraRigParameterDefinition.h"
 #include "StructUtils/PropertyBag.h"
 #include "UObject/ObjectPtr.h"
 
@@ -10,6 +12,7 @@
 
 class UCameraRigAsset;
 struct FCameraRigAssetReference;
+struct FCustomCameraNodeParameterInfos;
 struct FPropertyTag;
 
 namespace UE::Cameras
@@ -234,6 +237,35 @@ private:
 	friend struct FCameraRigAssetReference;
 };
 
+/**
+ * Metadata for a referenced camera rig's parameters.
+ */
+USTRUCT()
+struct FCameraRigAssetReferenceParameterMetaData
+{
+	GENERATED_BODY()
+
+	/** The GUID of the parameter. */
+	UPROPERTY()
+	FGuid ParameterGuid;
+
+	/** The type of parameter. */
+	UPROPERTY()
+	ECameraRigInterfaceParameterType ParameterType = ECameraRigInterfaceParameterType::Blendable;
+
+	/** The ID to use for overriding a blendable parameter. */
+	UPROPERTY()
+	FCameraVariableID OverrideVariableID;
+
+	/** The ID to use for overriding a data parameter. */
+	UPROPERTY()
+	FCameraContextDataID OverrideDataID;
+
+	/** Whether this parameter has an override value. */
+	UPROPERTY()
+	bool bIsOverridden = false;
+};
+
 /**
  * A structure holding a reference to a camera rig asset, along with the interface parameter
  * override values.
@@ -288,12 +320,6 @@ public:
 		return Parameters;
 	}
 
-	/** Gets the IDs of the parameters with override values. */
-	TConstArrayView<FGuid> GetOverridenParameterGuids() const
-	{
-		return ParameterOverrideGuids;
-	}
-
 	/** Applies the parameter override values to the given evaluation result. */
 	void ApplyParameterOverrides(UE::Cameras::FCameraNodeEvaluationResult& OutResult, bool bDrivenOverridesOnly);
 
@@ -304,13 +330,22 @@ public:
 	bool SerializeFromMismatchedTag(FPropertyTag const& Tag, FStructuredArchive::FSlot Slot);
 	void PostSerialize(const FArchive& Ar);
 
-	GAMEPLAYCAMERAS_API bool IsParameterOverriden(const FGuid PropertyID) const;
-	GAMEPLAYCAMERAS_API void SetParameterOverriden(const FGuid PropertyID, bool bIsOverridden);
+	GAMEPLAYCAMERAS_API bool IsParameterOverridden(const FGuid& PropertyID) const;
+	GAMEPLAYCAMERAS_API void SetParameterOverridden(const FGuid& PropertyID, bool bIsOverridden);
 
 	GAMEPLAYCAMERAS_API bool NeedsRebuildParameters() const;
 	GAMEPLAYCAMERAS_API bool RebuildParametersIfNeeded();
 	GAMEPLAYCAMERAS_API void RebuildParameters();
 
+	GAMEPLAYCAMERAS_API void GetCustomCameraNodeParameters(FCustomCameraNodeParameterInfos& OutParameterInfos);
+
+private:
+
+	const FCameraRigAssetReferenceParameterMetaData* FindMetaData(const FGuid& PropertyID) const;
+	FCameraRigAssetReferenceParameterMetaData* FindMetaData(const FGuid& PropertyID);
+	
+	void GenerateOverriddenParameterGuidArray(TArray<FGuid>& OutOverriddenIDs) const;
+
 private:
 
 	/** The referenced camera rig. */
@@ -321,9 +356,9 @@ private:
 	UPROPERTY(EditAnywhere, Category="", meta=(FixedLayout))
 	FInstancedPropertyBag Parameters;
 
-	/** The list of camera rig parameters with override values. */
-	UPROPERTY(EditAnywhere, Category="")
-	TArray<FGuid> ParameterOverrideGuids;
+	/** Metadata for the parameters. */
+	UPROPERTY()
+	TArray<FCameraRigAssetReferenceParameterMetaData> ParameterMetaData;
 
 
 	// Deprecated
diff --git a/Engine/Plugins/Cameras/GameplayCameras/Source/GameplayCameras/Public/Core/CameraRigJoints.h b/Engine/Plugins/Cameras/GameplayCameras/Source/GameplayCameras/Public/Core/CameraRigJoints.h
--- a/Engine/Plugins/Cameras/GameplayCameras/Source/GameplayCameras/Public/Core/CameraRigJoints.h
+++ b/Engine/Plugins/Cameras/GameplayCameras/Source/GameplayCameras/Public/Core/CameraRigJoints.h
@@ -5,6 +5,8 @@
 #include "Core/CameraVariableTableFwd.h"
 #include "CoreTypes.h"
 
+struct FCameraVariableDefinition;
+
 namespace UE::Cameras
 {
 
diff --git a/Engine/Plugins/Cameras/GameplayCameras/Source/GameplayCameras/Public/Core/CameraRigParameterDefinition.h b/Engine/Plugins/Cameras/GameplayCameras/Source/GameplayCameras/Public/Core/CameraRigParameterDefinition.h
new file mode 100644
--- /dev/null
+++ b/Engine/Plugins/Cameras/GameplayCameras/Source/GameplayCameras/Public/Core/CameraRigParameterDefinition.h
@@ -0,0 +1,121 @@
+// Copyright Epic Games, Inc. All Rights Reserved.
+
+#pragma once
+
+#include "Core/CameraContextDataTableFwd.h"
+#include "Core/CameraVariableTableFwd.h"
+#include "Misc/TVariant.h"
+
+#include "CameraRigParameterDefinition.generated.h"
+
+class UCameraRigAsset;
+struct FInstancedPropertyBag;
+struct FPropertyBagPropertyDesc;
+
+/**
+ * The type of a camera rig parameter.
+ */
+UENUM()
+enum class ECameraRigInterfaceParameterType : uint8
+{
+	Blendable,
+	Data
+};
+
+/**
+ * Information about a parameter exposed on a camera asset.
+ */
+USTRUCT()
+struct FCameraRigParameterDefinition
+{
+	GENERATED_BODY()
+
+	/** The name of the parameter. */
+	UPROPERTY()
+	FName ParameterName;
+
+	/**
+	 * The GUID of the parameter.
+	 * This matches the GUID on the corresponding UCameraRigBlendableParameter or
+	 * UCameraRigDataParameter object.
+	 */
+	UPROPERTY()
+	FGuid ParameterGuid;
+
+	/** The type of this parameter. */
+	UPROPERTY()
+	ECameraRigInterfaceParameterType ParameterType = ECameraRigInterfaceParameterType::Blendable;
+
+
+	// Blendable parameter properties.
+	// (only valid when ParameterType == Blendable)
+
+	/** The ID of the variable that drives this blendable parameter. */
+	UPROPERTY()
+	FCameraVariableID VariableID;
+
+	/** The type of the variable that drives this blendable parameter. */
+	UPROPERTY()
+	ECameraVariableType VariableType = ECameraVariableType::Boolean;
+
+	/** The type of the structure if VariableType is BlendableStruct. */
+	UPROPERTY()
+	TObjectPtr<const UScriptStruct> BlendableStructType;
+
+
+	// Data parameter properties.
+	// (only valid when ParameterType == Data)
+
+	/** The ID of the data that drives this blendable parameter. */
+	UPROPERTY()
+	FCameraContextDataID DataID;
+
+	/** The type of the data that drives this blendable parameter. */
+	UPROPERTY()
+	ECameraContextDataType DataType = ECameraContextDataType::Name;
+
+	/** The type object of the data that drives this blendable parameter. */
+	UPROPERTY()
+	TObjectPtr<const UObject> DataTypeObject;
+
+public:
+
+	bool operator==(const FCameraRigParameterDefinition& Other) const = default;
+};
+
+template<>
+struct TStructOpsTypeTraits<FCameraRigParameterDefinition> : public TStructOpsTypeTraitsBase2<FCameraRigParameterDefinition>
+{
+	enum
+	{
+		WithIdenticalViaEquality = true
+	};
+};
+
+namespace UE::Cameras
+{
+
+/**
+ * A helper class for building an FInstancedPropertyBag from a list of camera rig
+ * parameter definitions.
+ */
+class GAMEPLAYCAMERAS_API FCameraRigParameterBuilder
+{
+public:
+
+	/**
+	 * Builds a property bag that contains a property for each exposed parameter on the given camera rig.
+	 * Each property's value is set to the default value of the corresponding parameter.
+	 */
+	static void BuildDefaultParameters(const UCameraRigAsset* CameraRig, FInstancedPropertyBag& OutPropertyBag);
+
+public:
+
+	/** Generates property bag property descriptors for the parameters exposed by the given camera rig. */
+	static void AppendDefaultParameterProperties(const UCameraRigAsset* CameraRig, TArray<FPropertyBagPropertyDesc>& OutProperties);
+	/** Sets the default value of the property bag properties that correspond to parameters on the given camera rig. */
+	static void SetDefaultParameterValues(const UCameraRigAsset* CameraRig, FInstancedPropertyBag& PropertyBag);
+};
+
+}  // namespace UE::Cameras
+
diff --git a/Engine/Plugins/Cameras/GameplayCameras/Source/GameplayCameras/Public/Core/CameraVariableAssets.h b/Engine/Plugins/Cameras/GameplayCameras/Source/GameplayCameras/Public/Core/CameraVariableAssets.h
--- a/Engine/Plugins/Cameras/GameplayCameras/Source/GameplayCameras/Public/Core/CameraVariableAssets.h
+++ b/Engine/Plugins/Cameras/GameplayCameras/Source/GameplayCameras/Public/Core/CameraVariableAssets.h
@@ -3,7 +3,7 @@
 #pragma once
 
 #include "Containers/UnrealString.h"
-#include "Core/CameraVariableTableFwd.h"
+#include "Core/CameraVariableTableAllocationInfo.h"
 #include "CoreTypes.h"
 #include "Math/MathFwd.h"
 
diff --git a/Engine/Plugins/Cameras/GameplayCameras/Source/GameplayCameras/Public/Core/CameraVariableReferenceReader.h b/Engine/Plugins/Cameras/GameplayCameras/Source/GameplayCameras/Public/Core/CameraVariableReferenceReader.h
new file mode 100644
--- /dev/null
+++ b/Engine/Plugins/Cameras/GameplayCameras/Source/GameplayCameras/Public/Core/CameraVariableReferenceReader.h
@@ -0,0 +1,99 @@
+// Copyright Epic Games, Inc. All Rights Reserved.
+
+#pragma once
+
+#include "Core/CameraVariableTable.h"
+#include "CoreTypes.h"
+#include "Templates/UnrealTypeTraits.h"
+
+namespace UE::Cameras
+{
+
+/**
+ * A utility class for reading the effective value of variable reference.
+ */
+template<typename ValueType>
+class TCameraVariableReferenceReader
+{
+public:
+
+	TCameraVariableReferenceReader() {}
+
+	template<typename VariableReferenceType>
+	TCameraVariableReferenceReader(const VariableReferenceType& Reference)
+	{
+		Initialize(Reference);
+	}
+
+	/**
+	 * Initializes the reader around the given variable reference.
+	 */
+	template<typename VariableReferenceType>
+	void Initialize(const VariableReferenceType& Reference)
+	{
+		static_assert(
+				std::is_same<ValueType, typename VariableReferenceType::VariableAssetType::ValueType>(),
+				"The given variable reference is of the wrong type for this reader! Value types must be the same.");
+
+		if (Reference.Variable)
+		{
+			DefaultValue = Reference.Variable->GetDefaultValue();
+			VariableID = Reference.Variable->GetVariableID();
+		}
+		else
+		{
+			DefaultValue = ValueType();
+			VariableID = Reference.VariableID;
+		}
+	}
+
+	/**
+	 * Gets the actual value for the referenced variable.
+	 */
+	typename TCallTraits<ValueType>::ParamType Get(const FCameraVariableTable& VariableTable) const
+	{
+		if (!VariableID.IsValid())
+		{
+			return DefaultValue;
+		}
+		else
+		{
+			if (const ValueType* ActualValue = VariableTable.FindValue<ValueType>(VariableID))
+			{
+				return *ActualValue;
+			}
+			return DefaultValue;
+		}
+	}
+
+	/**
+	 * Gets the actual value for the referenced variable, or returns false if the variable doesn't
+	 * have a value yet.
+	 */
+	bool TryGet(const FCameraVariableTable& VariableTable, ValueType& OutValue) const
+	{
+		if (VariableID.IsValid())
+		{
+			return VariableTable.TryGetValue(VariableID, OutValue);
+		}
+		return false;
+	}
+
+	/**
+	 * Returns whether his reference points to a variable.
+	 */
+	bool IsDriven() const
+	{
+		return VariableID.IsValid();
+	}
+
+private:
+
+	/** The default value for the variable. */
+	ValueType DefaultValue;
+	/** The ID of the variable, if any. */
+	FCameraVariableID VariableID;
+};
+
+}  // namespace UE::Cameras
+
diff --git a/Engine/Plugins/Cameras/GameplayCameras/Source/GameplayCameras/Public/Core/CameraVariableReferences.h b/Engine/Plugins/Cameras/GameplayCameras/Source/GameplayCameras/Public/Core/CameraVariableReferences.h
--- a/Engine/Plugins/Cameras/GameplayCameras/Source/GameplayCameras/Public/Core/CameraVariableReferences.h
+++ b/Engine/Plugins/Cameras/GameplayCameras/Source/GameplayCameras/Public/Core/CameraVariableReferences.h
@@ -9,9 +9,17 @@
 #define UE_DEFINE_CAMERA_VARIABLE_REFERENCE(ValueName)\
 	F##ValueName##CameraVariableReference() {}\
 	F##ValueName##CameraVariableReference(VariableAssetType* InVariable) : Variable(InVariable) {}\
-	bool IsValid() const { return (bool)Variable; }\
-	U##ValueName##CameraVariable* Get() const { return Variable.Get(); }\
-	operator U##ValueName##CameraVariable* () const { return Variable.Get(); }
+	bool IsValid() const { return VariableID.IsValid(); }\
+	bool HasVariable() const { return Variable != nullptr; }
+
+// All camera references have:
+//
+// - Variable: a variable chosen by the user.
+// - VariableID: the ID of the variable to use for this reference.
+//
+// When Variable is set, VariableID is the ID of that variable.
+// When Variable is not set, VariableID is the ID of something else that
+// the caller code should use, such as a camera rig parameter override.
 
 USTRUCT()
 struct FBooleanCameraVariableReference
@@ -20,6 +28,9 @@ struct FBooleanCameraVariableReference
 
 	using VariableAssetType = UBooleanCameraVariable;
 
+	UPROPERTY()
+	FCameraVariableID VariableID;
+
 	UPROPERTY(EditAnywhere, Category="Variable")
 	TObjectPtr<UBooleanCameraVariable> Variable;
 
@@ -33,6 +44,9 @@ struct FInteger32CameraVariableReference
 
 	using VariableAssetType = UInteger32CameraVariable;
 
+	UPROPERTY()
+	FCameraVariableID VariableID;
+
 	UPROPERTY(EditAnywhere, Category="Variable")
 	TObjectPtr<UInteger32CameraVariable> Variable;
 
@@ -46,6 +60,9 @@ struct FFloatCameraVariableReference
 
 	using VariableAssetType = UFloatCameraVariable;
 
+	UPROPERTY()
+	FCameraVariableID VariableID;
+
 	UPROPERTY(EditAnywhere, Category="Variable")
 	TObjectPtr<UFloatCameraVariable> Variable;
 
@@ -59,6 +76,9 @@ struct FDoubleCameraVariableReference
 
 	using VariableAssetType = UDoubleCameraVariable;
 
+	UPROPERTY()
+	FCameraVariableID VariableID;
+
 	UPROPERTY(EditAnywhere, Category="Variable")
 	TObjectPtr<UDoubleCameraVariable> Variable;
 
@@ -72,6 +92,9 @@ struct FVector2fCameraVariableReference
 
 	using VariableAssetType = UVector2fCameraVariable;
 
+	UPROPERTY()
+	FCameraVariableID VariableID;
+
 	UPROPERTY(EditAnywhere, Category="Variable")
 	TObjectPtr<UVector2fCameraVariable> Variable;
 
@@ -85,6 +108,9 @@ struct FVector2dCameraVariableReference
 
 	using VariableAssetType = UVector2dCameraVariable;
 
+	UPROPERTY()
+	FCameraVariableID VariableID;
+
 	UPROPERTY(EditAnywhere, Category="Variable")
 	TObjectPtr<UVector2dCameraVariable> Variable;
 
@@ -98,6 +124,9 @@ struct FVector3fCameraVariableReference
 
 	using VariableAssetType = UVector3fCameraVariable;
 
+	UPROPERTY()
+	FCameraVariableID VariableID;
+
 	UPROPERTY(EditAnywhere, Category="Variable")
 	TObjectPtr<UVector3fCameraVariable> Variable;
 
@@ -111,6 +140,9 @@ struct FVector3dCameraVariableReference
 
 	using VariableAssetType = UVector3dCameraVariable;
 
+	UPROPERTY()
+	FCameraVariableID VariableID;
+
 	UPROPERTY(EditAnywhere, Category="Variable")
 	TObjectPtr<UVector3dCameraVariable> Variable;
 
@@ -124,6 +156,9 @@ struct FVector4fCameraVariableReference
 
 	using VariableAssetType = UVector4fCameraVariable;
 
+	UPROPERTY()
+	FCameraVariableID VariableID;
+
 	UPROPERTY(EditAnywhere, Category="Variable")
 	TObjectPtr<UVector4fCameraVariable> Variable;
 
@@ -137,6 +172,9 @@ struct FVector4dCameraVariableReference
 
 	using VariableAssetType = UVector4dCameraVariable;
 
+	UPROPERTY()
+	FCameraVariableID VariableID;
+
 	UPROPERTY(EditAnywhere, Category="Variable")
 	TObjectPtr<UVector4dCameraVariable> Variable;
 
@@ -150,6 +188,9 @@ struct FRotator3fCameraVariableReference
 
 	using VariableAssetType = URotator3fCameraVariable;
 
+	UPROPERTY()
+	FCameraVariableID VariableID;
+
 	UPROPERTY(EditAnywhere, Category="Variable")
 	TObjectPtr<URotator3fCameraVariable> Variable;
 
@@ -163,6 +204,9 @@ struct FRotator3dCameraVariableReference
 
 	using VariableAssetType = URotator3dCameraVariable;
 
+	UPROPERTY()
+	FCameraVariableID VariableID;
+
 	UPROPERTY(EditAnywhere, Category="Variable")
 	TObjectPtr<URotator3dCameraVariable> Variable;
 
@@ -176,6 +220,9 @@ struct FTransform3fCameraVariableReference
 
 	using VariableAssetType = UTransform3fCameraVariable;
 
+	UPROPERTY()
+	FCameraVariableID VariableID;
+
 	UPROPERTY(EditAnywhere, Category="Variable")
 	TObjectPtr<UTransform3fCameraVariable> Variable;
 
@@ -189,6 +236,9 @@ struct FTransform3dCameraVariableReference
 
 	using VariableAssetType = UTransform3dCameraVariable;
 
+	UPROPERTY()
+	FCameraVariableID VariableID;
+
 	UPROPERTY(EditAnywhere, Category="Variable")
 	TObjectPtr<UTransform3dCameraVariable> Variable;
 
diff --git a/Engine/Plugins/Cameras/GameplayCameras/Source/GameplayCameras/Public/Core/CameraVariableTable.h b/Engine/Plugins/Cameras/GameplayCameras/Source/GameplayCameras/Public/Core/CameraVariableTable.h
--- a/Engine/Plugins/Cameras/GameplayCameras/Source/GameplayCameras/Public/Core/CameraVariableTable.h
+++ b/Engine/Plugins/Cameras/GameplayCameras/Source/GameplayCameras/Public/Core/CameraVariableTable.h
@@ -5,17 +5,20 @@
 #include "Containers/Map.h"
 #include "Containers/Set.h"
 #include "Core/CameraVariableTableFwd.h"
+#include "Core/CameraVariableTableAllocationInfo.h"
 #include "CoreTypes.h"
 #include "GameplayCameras.h"
 #include "Misc/EnumClassFlags.h"
 #include "Templates/UnrealTypeTraits.h"
-#include "UObject/NameTypes.h"
 
 namespace UE::Cameras
 {
 
 template<typename ValueType>
-struct TCameraVariableTraits;
+struct TCameraVariableTraits
+{
+	static const ECameraVariableType Type = ECameraVariableType::BlendableStruct;
+};
 
 template<typename ValueType>
 struct TCameraVariableInterpolation;
@@ -152,21 +155,25 @@ public:
 
 	const uint8* GetValue(
 			FCameraVariableID VariableID,
-			ECameraVariableType ExpectedVariableType) const;
+			ECameraVariableType ExpectedVariableType,
+			const UScriptStruct* ExpectedBlendableStructType) const;
 
 	const uint8* TryGetValue(
 			FCameraVariableID VariableID,
-			ECameraVariableType ExpectedVariableType) const;
+			ECameraVariableType ExpectedVariableType,
+			const UScriptStruct* ExpectedBlendableStructType) const;
 
 	void SetValue(
 			FCameraVariableID VariableID, 
 			ECameraVariableType ExpectedVariableType, 
+			const UScriptStruct* ExpectedBlendableStructType,
 			const uint8* InRawValuePtr,
 			bool bMarkAsWrittenThisFrame = true);
 
 	bool TrySetValue(
 			FCameraVariableID VariableID, 
 			ECameraVariableType ExpectedVariableType, 
+			const UScriptStruct* ExpectedBlendableStructType,
 			const uint8* InRawValuePtr,
 			bool bMarkAsWrittenThisFrame = true);
 
@@ -187,7 +194,10 @@ private:
 
 	struct FEntry;
 
-	static bool GetVariableTypeAllocationInfo(ECameraVariableType VariableType, uint32& OutSizeOf, uint32& OutAlignOf);
+	static void CacheBlendableStructs();
+	static FBlendableStructTypeErasedInterpolator GetBlendableStructInterpolator(const UScriptStruct* StructType);
+
+	static bool GetVariableTypeAllocationInfo(ECameraVariableType VariableType, const UScriptStruct* StructType, uint32& OutSizeOf, uint32& OutAlignOf);
 
 	template<typename ValueType>
 	static bool CheckVariableType(ECameraVariableType InType)
@@ -220,6 +230,7 @@ private:
 	{
 		FCameraVariableID ID;
 		ECameraVariableType Type;
+		const UScriptStruct* StructType = nullptr;
 		uint32 Offset;
 		mutable EEntryFlags Flags;
 #if WITH_EDITORONLY_DATA
@@ -234,6 +245,9 @@ private:
 	uint32 Capacity = 0;
 	uint32 Used = 0;
 
+	static TArray<FBlendableStructInfo> CachedBlendableStructs;
+	static bool bCachedBlendableStructs;
+
 	template<typename T>
 	friend struct TCameraVariableInterpolation;
 
diff --git a/Engine/Plugins/Cameras/GameplayCameras/Source/GameplayCameras/Public/Core/CameraVariableTableAllocationInfo.h b/Engine/Plugins/Cameras/GameplayCameras/Source/GameplayCameras/Public/Core/CameraVariableTableAllocationInfo.h
new file mode 100644
--- /dev/null
+++ b/Engine/Plugins/Cameras/GameplayCameras/Source/GameplayCameras/Public/Core/CameraVariableTableAllocationInfo.h
@@ -0,0 +1,128 @@
+// Copyright Epic Games, Inc. All Rights Reserved.
+
+#pragma once
+
+#include "Core/CameraVariableTableFwd.h"
+#include "UObject/ObjectMacros.h"
+
+#include "CameraVariableTableAllocationInfo.generated.h"
+
+/**
+ * A structure that describes a camera variable.
+ */
+USTRUCT()
+struct FCameraVariableDefinition
+{
+	GENERATED_BODY()
+
+	/** The ID of the variable. */
+	UPROPERTY()
+	FCameraVariableID VariableID;
+
+	/** The type of the variable. */
+	UPROPERTY()
+	ECameraVariableType VariableType = ECameraVariableType::Boolean;
+
+	/** The type of a blendable struct (only valid when VariableType == BlendableStruct). */
+	UPROPERTY()
+	TObjectPtr<const UScriptStruct> BlendableStructType; 
+
+	/**
+	 * Whether the variable is private. 
+	 *
+	 * Private variables are not propagated from one table to another when
+	 * interpolating or overriding a table.
+	 */
+	UPROPERTY()
+	bool bIsPrivate = false;
+
+	/**
+	 * Whether the variable is an input variable.
+	 *
+	 * Input variables are blended during the pre-blend parameter update phase.
+	 */
+	UPROPERTY()
+	bool bIsInput = false;
+
+	/**
+	 * Whether the variable should auto-reset to an "unset" state after every
+	 * evaluation. Doing so resets it to its default value if bound to a
+	 * camera variable or camera parameter, or to a non-existent state otherwise.
+	 */
+	UPROPERTY()
+	bool bAutoReset = false;
+
+#if WITH_EDITORONLY_DATA
+	/** The name of the variable, for debugging purposes. */
+	UPROPERTY()
+	FString VariableName;
+#endif
+
+	/** Returns whether this definition has a valid variable ID. */
+	bool IsValid() const
+	{
+		return VariableID.IsValid();
+	}
+
+	/** Implicit conversion to a camera variable ID. */
+	operator FCameraVariableID() const
+	{
+		return VariableID;
+	}
+
+	/** Creates a variant of this camera variable definition. */
+	FCameraVariableDefinition CreateVariant(const FString& VariantID) const
+	{
+		FCameraVariableDefinition VariantDefinition(*this);
+		VariantDefinition.VariableID = FCameraVariableID::FromHashValue(
+				HashCombineFast(VariableID.GetValue(), GetTypeHash(VariantID)));
+#if WITH_EDITORONLY_DATA
+		if (!VariableName.IsEmpty())
+		{
+			VariantDefinition.VariableName += FString::Format(TEXT("_{0}Variant"), { VariantID });
+		}
+#endif
+		return VariantDefinition;
+	}
+
+	bool operator==(const FCameraVariableDefinition& Other) const = default;
+};
+
+template<>
+struct TStructOpsTypeTraits<FCameraVariableDefinition> : public TStructOpsTypeTraitsBase2<FCameraVariableDefinition>
+{
+	enum
+	{
+		WithCopy = true,
+		WithIdenticalViaEquality = true
+	};
+};
+
+/**
+ * A structure that describes the required camera variable table setup of a camera rig.
+ */
+USTRUCT()
+struct FCameraVariableTableAllocationInfo
+{
+	GENERATED_BODY()
+
+	/** The list of variables that should be allocated in a table. */
+	UPROPERTY()
+	TArray<FCameraVariableDefinition> VariableDefinitions;
+
+	/**Combines the given allocation info with this one. */
+	GAMEPLAYCAMERAS_API void Combine(const FCameraVariableTableAllocationInfo& OtherInfo);
+
+	bool operator==(const FCameraVariableTableAllocationInfo& Other) const = default;
+};
+
+template<>
+struct TStructOpsTypeTraits<FCameraVariableTableAllocationInfo> : public TStructOpsTypeTraitsBase2<FCameraVariableTableAllocationInfo>
+{
+	enum
+	{
+		WithCopy = true,
+		WithIdenticalViaEquality = true
+	};
+};
+
diff --git a/Engine/Plugins/Cameras/GameplayCameras/Source/GameplayCameras/Public/Core/CameraVariableTableFwd.h b/Engine/Plugins/Cameras/GameplayCameras/Source/GameplayCameras/Public/Core/CameraVariableTableFwd.h
--- a/Engine/Plugins/Cameras/GameplayCameras/Source/GameplayCameras/Public/Core/CameraVariableTableFwd.h
+++ b/Engine/Plugins/Cameras/GameplayCameras/Source/GameplayCameras/Public/Core/CameraVariableTableFwd.h
@@ -48,7 +48,8 @@ enum class ECameraVariableType
 	Rotator3f,
 	Rotator3d,
 	Transform3f,
-	Transform3d
+	Transform3d,
+	BlendableStruct
 };
 
 /**
@@ -112,118 +113,16 @@ private:
 	uint32 Value;
 };
 
-/**
- * A structure that describes a camera variable.
- */
-USTRUCT()
-struct FCameraVariableDefinition
+namespace UE::Cameras
 {
-	GENERATED_BODY()
-
-	/** The ID of the variable. */
-	UPROPERTY()
-	FCameraVariableID VariableID;
-
-	/** The type of the variable. */
-	UPROPERTY()
-	ECameraVariableType VariableType = ECameraVariableType::Boolean;
-
-	/**
-	 * Whether the variable is private. 
-	 *
-	 * Private variables are not propagated from one table to another when
-	 * interpolating or overriding a table.
-	 */
-	UPROPERTY()
-	bool bIsPrivate = false;
-
-	/**
-	 * Whether the variable is an input variable.
-	 *
-	 * Input variables are blended during the pre-blend parameter update phase.
-	 */
-	UPROPERTY()
-	bool bIsInput = false;
-
-	/**
-	 * Whether the variable should auto-reset to an "unset" state after every
-	 * evaluation. Doing so resets it to its default value if bound to a
-	 * camera variable or camera parameter, or to a non-existent state otherwise.
-	 */
-	UPROPERTY()
-	bool bAutoReset = false;
-
-#if WITH_EDITORONLY_DATA
-	/** The name of the variable, for debugging purposes. */
-	UPROPERTY()
-	FString VariableName;
-#endif
-
-	/** Returns whether this definition has a valid variable ID. */
-	bool IsValid() const
-	{
-		return VariableID.IsValid();
-	}
-
-	/** Implicit conversion to a camera variable ID. */
-	operator FCameraVariableID() const
-	{
-		return VariableID;
-	}
-
-	/** Creates a variant of this camera variable definition. */
-	FCameraVariableDefinition CreateVariant(const FString& VariantID) const
-	{
-		FCameraVariableDefinition VariantDefinition(*this);
-		VariantDefinition.VariableID = FCameraVariableID::FromHashValue(
-				HashCombineFast(VariableID.GetValue(), GetTypeHash(VariantID)));
-#if WITH_EDITORONLY_DATA
-		if (!VariableName.IsEmpty())
-		{
-			VariantDefinition.VariableName += FString::Format(TEXT("_{0}Variant"), { VariantID });
-		}
-#endif
-		return VariantDefinition;
-	}
-
-	GAMEPLAYCAMERAS_API friend bool operator==(const FCameraVariableDefinition& A, const FCameraVariableDefinition& B);
-};
 
-template<>
-struct TStructOpsTypeTraits<FCameraVariableDefinition> : public TStructOpsTypeTraitsBase2<FCameraVariableDefinition>
-{
-	enum
-	{
-		WithCopy = true,
-		WithIdenticalViaEquality = true
-	};
-};
+using FBlendableStructTypeErasedInterpolator = void(*)(uint8* From, const uint8* To, float Alpha);
 
-/**
- * A structure that describes the required camera variable table setup of a camera rig.
- */
-USTRUCT()
-struct FCameraVariableTableAllocationInfo
+struct FBlendableStructInfo
 {
-	GENERATED_BODY()
-
-	/** The list of variables that should be allocated in a table. */
-	UPROPERTY()
-	TArray<FCameraVariableDefinition> VariableDefinitions;
-
-	/**Combines the given allocation info with this one. */
-	GAMEPLAYCAMERAS_API void Combine(const FCameraVariableTableAllocationInfo& OtherInfo);
-
-	GAMEPLAYCAMERAS_API friend bool operator==(const FCameraVariableTableAllocationInfo& A, const FCameraVariableTableAllocationInfo& B);
+	TObjectPtr<const UScriptStruct> StructType;
+	FBlendableStructTypeErasedInterpolator Interpolator;
 };
 
-template<>
-struct TStructOpsTypeTraits<FCameraVariableTableAllocationInfo> : public TStructOpsTypeTraitsBase2<FCameraVariableTableAllocationInfo>
-{
-	enum
-	{
-		WithCopy = true,
-		WithIdenticalViaEquality = true
-	};
-};
+}  // namespace UE::Cameras
 
diff --git a/Engine/Plugins/Cameras/GameplayCameras/Source/GameplayCameras/Public/Core/ICustomCameraNodeParameterProvider.h b/Engine/Plugins/Cameras/GameplayCameras/Source/GameplayCameras/Public/Core/ICustomCameraNodeParameterProvider.h
--- a/Engine/Plugins/Cameras/GameplayCameras/Source/GameplayCameras/Public/Core/ICustomCameraNodeParameterProvider.h
+++ b/Engine/Plugins/Cameras/GameplayCameras/Source/GameplayCameras/Public/Core/ICustomCameraNodeParameterProvider.h
@@ -3,18 +3,18 @@
 #pragma once
 
 #include "Core/CameraVariableTableFwd.h"
-#include "Core/CameraContextDataAllocationInfo.h"
+#include "Core/CameraContextDataTableAllocationInfo.h"
 #include "UObject/Interface.h"
 #include "UObject/ObjectPtr.h"
 
 #include "ICustomCameraNodeParameterProvider.generated.h"
 
 class UCameraNode;
-class UCameraVariableAsset;
 
 namespace UE::Cameras
 {
 	class FCameraRigAssetBuilder;
+	class FCameraRigParameterBuilder;
 	namespace Internal { struct FInterfaceParameterBindingBuilder; }
 }
 
@@ -32,11 +32,19 @@ struct FCustomCameraNodeBlendableParameter
 	UPROPERTY()
 	ECameraVariableType ParameterType = ECameraVariableType::Boolean;
 
-	/** An optional camera variable that is dynamically driving the parameter's value. */
+	/** The struct type of a blendable struct. */
+	UPROPERTY()
+	TObjectPtr<const UScriptStruct> BlendableStructType;
+
+	/** An optional camera variable ID for dynamically driving the parameter's value. */
+	UPROPERTY()
+	FCameraVariableID OverrideVariableID;
+
+	/** An optional user-defined camera variable for dynamically driving the parameter's value. */
 	UPROPERTY()
 	TObjectPtr<UCameraVariableAsset> OverrideVariable;
 
-	GAMEPLAYCAMERAS_API friend bool operator==(const FCustomCameraNodeBlendableParameter& A, const FCustomCameraNodeBlendableParameter& B);
+	bool operator==(const FCustomCameraNodeBlendableParameter& Other) const = default;
 };
 
 /** Describes a custom camera data parameter. */
@@ -61,7 +69,7 @@ struct FCustomCameraNodeDataParameter
 	UPROPERTY()
 	FCameraContextDataID OverrideDataID;
 
-	GAMEPLAYCAMERAS_API friend bool operator==(const FCustomCameraNodeDataParameter& A, const FCustomCameraNodeDataParameter& B);
+	bool operator==(const FCustomCameraNodeDataParameter& Other) const = default;
 };
 
 /** Describes custom camera parameters. */
@@ -84,7 +92,7 @@ struct FCustomCameraNodeParameters
 	/** Removes all parameters from this structure. */
 	void Reset() { BlendableParameters.Reset(); DataParameters.Reset(); }
 
-	GAMEPLAYCAMERAS_API friend bool operator==(const FCustomCameraNodeParameters& A, const FCustomCameraNodeParameters& B);
+	bool operator==(const FCustomCameraNodeParameters& Other) const = default;
 };
 
 /**
@@ -95,36 +103,34 @@ struct FCustomCameraNodeParameterInfos
 	/** Returns whether there are any blendable or data parameters. */
 	bool HasAnyParameters() const { return !BlendableParameters.IsEmpty() || !DataParameters.IsEmpty(); }
 
-	/** Declares a blendable parameter. */
-	template<typename VariableAssetType>
-	void AddBlendableParameter(
-			FName ParameterName, 
-			ECameraVariableType ParameterType, 
-			const uint8* DefaultValuePtr,
-			TObjectPtr<VariableAssetType>* OverrideVariable)
-	{
-		FObjectPtr* OverrideVariablePtr = reinterpret_cast<FObjectPtr*>(OverrideVariable);
-		BlendableParameters.Add({ 
-				ParameterName, 
-				ParameterType, 
-				DefaultValuePtr,
-				reinterpret_cast<TObjectPtr<UCameraVariableAsset>*>(OverrideVariable) });
-	}
-
-	/** Declares a blendable parameter. */
+	/** 
+	 * Declares a blendable parameter. 
+	 * Pass a null pointer for OverrideVariableID if this parameter should not be overridable
+	 * by a camera rig parameter.
+	 */
 	GAMEPLAYCAMERAS_API void AddBlendableParameter(
 			FName ParameterName, 
 			ECameraVariableType ParameterType, 
-			const uint8* DefaultValuePtr,
-			TObjectPtr<UCameraVariableAsset>* OverrideVariable);
+			const UScriptStruct* BlendableStructType,
+			const uint8* DefaultValue,
+			FCameraVariableID* OverrideVariableID);
 
-	/** Declares a data parameter. */
+	GAMEPLAYCAMERAS_API void AddBlendableParameter(FCustomCameraNodeBlendableParameter& Parameter, const uint8* DefaultValue);
+
+	/** 
+	 * Declares a data parameter.
+	 * Pass a null pointer for OverrideDataID if this parameter should not be overridable
+	 * by a camera rig parameter.
+	 */
 	GAMEPLAYCAMERAS_API void AddDataParameter(
 			FName ParameterName, 
 			ECameraContextDataType ParameterType,
 			const UObject* ParameterTypeObject,
+			const uint8* DefaultValue,
 			FCameraContextDataID* OverrideDataID);
 
+	GAMEPLAYCAMERAS_API void AddDataParameter(FCustomCameraNodeDataParameter& Parameter, const uint8* DefaultValue);
+
 	/** Gets the list of blendable parameters. */
 	GAMEPLAYCAMERAS_API void GetBlendableParameters(TArray<FCustomCameraNodeBlendableParameter>& OutBlendableParameters) const;
 	/** Gets the list of data parameters. */
@@ -141,8 +147,10 @@ private:
 	{
 		FName ParameterName;
 		ECameraVariableType ParameterType;
-		const uint8* DefaultValuePtr = nullptr;
-		TObjectPtr<UCameraVariableAsset>* OverrideVariable = nullptr;
+		const UScriptStruct* BlendableStructType = nullptr;
+		const uint8* DefaultValue = nullptr;
+		FCameraVariableID* OverrideVariableID = nullptr;
+		UCameraVariableAsset* OverrideVariable = nullptr;
 	};
 
 	struct FDataParameterInfo
@@ -150,6 +158,7 @@ private:
 		FName ParameterName;
 		ECameraContextDataType ParameterType;
 		const UObject* ParameterTypeObject;
+		const uint8* DefaultValue = nullptr;
 		FCameraContextDataID* OverrideDataID = nullptr;
 	};
 
@@ -157,6 +166,7 @@ private:
 	TArray<FDataParameterInfo> DataParameters;
 
 	friend class UE::Cameras::FCameraRigAssetBuilder;
+	friend class UE::Cameras::FCameraRigParameterBuilder;
 	friend struct UE::Cameras::Internal::FInterfaceParameterBindingBuilder;
 };
 
diff --git a/Engine/Plugins/Cameras/GameplayCameras/Source/GameplayCameras/Public/Debug/ContextDataTableDebugBlock.h b/Engine/Plugins/Cameras/GameplayCameras/Source/GameplayCameras/Public/Debug/ContextDataTableDebugBlock.h
--- a/Engine/Plugins/Cameras/GameplayCameras/Source/GameplayCameras/Public/Debug/ContextDataTableDebugBlock.h
+++ b/Engine/Plugins/Cameras/GameplayCameras/Source/GameplayCameras/Public/Debug/ContextDataTableDebugBlock.h
@@ -36,7 +36,8 @@ private:
 
 	struct FEntryDebugInfo
 	{
-		FName Name;
+		uint32 ID;
+		FString Name;
 		FName TypeName;
 		FString Value;
 		bool bWritten;
diff --git a/Engine/Plugins/Cameras/GameplayCameras/Source/GameplayCameras/Public/GameFramework/CameraRigParameterInterop.h b/Engine/Plugins/Cameras/GameplayCameras/Source/GameplayCameras/Public/GameFramework/CameraRigParameterInterop.h
--- a/Engine/Plugins/Cameras/GameplayCameras/Source/GameplayCameras/Public/GameFramework/CameraRigParameterInterop.h
+++ b/Engine/Plugins/Cameras/GameplayCameras/Source/GameplayCameras/Public/GameFramework/CameraRigParameterInterop.h
@@ -61,6 +61,9 @@ public:
 	UFUNCTION(BlueprintCallable, Category="Camera", meta=(BlueprintInternalUseOnly="true"))
 	static void SetTransformParameter(UPARAM(Ref) FBlueprintCameraNodeEvaluationResult& Result, UCameraRigAsset* CameraRig, const FString& ParameterName, FTransform ParameterValue);
 
+	UFUNCTION(BlueprintCallable, Category="Camera", meta=(BlueprintInternalUseOnly="true"))
+	static void SetBlendableStructParameter(UPARAM(Ref) FBlueprintCameraNodeEvaluationResult& Result, UCameraRigAsset* CameraRig, const FString& ParameterName, const FInstancedStruct& ParameterValue);
+
 public:
 
 	UFUNCTION(BlueprintCallable, Category="Camera", meta=(BlueprintInternalUseOnly="true"))
diff --git a/Engine/Plugins/Cameras/GameplayCameras/Source/GameplayCameras/Public/IGameplayCamerasModule.h b/Engine/Plugins/Cameras/GameplayCameras/Source/GameplayCameras/Public/IGameplayCamerasModule.h
--- a/Engine/Plugins/Cameras/GameplayCameras/Source/GameplayCameras/Public/IGameplayCamerasModule.h
+++ b/Engine/Plugins/Cameras/GameplayCameras/Source/GameplayCameras/Public/IGameplayCamerasModule.h
@@ -2,8 +2,10 @@
 
 #pragma once
 
+#include "Core/CameraVariableTableFwd.h"
 #include "Modules/ModuleInterface.h"
 #include "Templates/SharedPointer.h"
+#include "UObject/WeakObjectPtr.h"
 
 namespace UE::Cameras
 {
@@ -25,6 +27,13 @@ public:
 	 */
 	GAMEPLAYCAMERAS_API static IGameplayCamerasModule& Get();
 
+	/** Register a new blendable structure. */
+	virtual void RegisterBlendableStruct(const UScriptStruct* StructType, UE::Cameras::FBlendableStructTypeErasedInterpolator Interpolator) = 0;
+	/** Gets the currently registered blendable structures. */
+	virtual TConstArrayView<UE::Cameras::FBlendableStructInfo> GetBlendableStructs() const = 0;
+	/** Unregisters a blendable structure. */
+	virtual void UnregisterBlendableStruct(const UScriptStruct* StructType) = 0;
+
 #if WITH_EDITOR
 	using IGameplayCamerasLiveEditManager = UE::Cameras::IGameplayCamerasLiveEditManager;
 
diff --git a/Engine/Plugins/Cameras/GameplayCameras/Source/GameplayCamerasEditor/Private/Customizations/CameraAssetReferenceDetailsCustomization.cpp b/Engine/Plugins/Cameras/GameplayCameras/Source/GameplayCamerasEditor/Private/Customizations/CameraAssetReferenceDetailsCustomization.cpp
--- a/Engine/Plugins/Cameras/GameplayCameras/Source/GameplayCamerasEditor/Private/Customizations/CameraAssetReferenceDetailsCustomization.cpp
+++ b/Engine/Plugins/Cameras/GameplayCameras/Source/GameplayCamerasEditor/Private/Customizations/CameraAssetReferenceDetailsCustomization.cpp
@@ -39,12 +39,12 @@ protected:
 		
 		virtual bool IsPropertyOverridden(const FGuid PropertyID) const override
 		{
-			return CameraAssetReference.IsParameterOverriden(PropertyID);
+			return CameraAssetReference.IsParameterOverridden(PropertyID);
 		}
 		
 		virtual void SetPropertyOverride(const FGuid PropertyID, const bool bIsOverridden) const override
 		{
-			CameraAssetReference.SetParameterOverriden(PropertyID, bIsOverridden);
+			CameraAssetReference.SetParameterOverridden(PropertyID, bIsOverridden);
 		}
 
 	private:
diff --git a/Engine/Plugins/Cameras/GameplayCameras/Source/GameplayCamerasEditor/Private/Customizations/CameraParameterDetailsCustomizations.cpp b/Engine/Plugins/Cameras/GameplayCameras/Source/GameplayCamerasEditor/Private/Customizations/CameraParameterDetailsCustomizations.cpp
--- a/Engine/Plugins/Cameras/GameplayCameras/Source/GameplayCamerasEditor/Private/Customizations/CameraParameterDetailsCustomizations.cpp
+++ b/Engine/Plugins/Cameras/GameplayCameras/Source/GameplayCamerasEditor/Private/Customizations/CameraParameterDetailsCustomizations.cpp
@@ -60,14 +60,10 @@ void FCameraParameterDetailsCustomization::CustomizeHeader(TSharedRef<IPropertyH
 	PropertyUtilities = CustomizationUtils.GetPropertyUtilities();
 	StructProperty = PropertyHandle;
 
-	// Figure out the parameter value and driving variable property names for the actual camera
-	// parameter type we are displaying. In theory these are always called "Value" and "Variable", but
-	// we do this in a statically-typed way via polyphormism, instead of in a duck-typing way.
-	FName ValuePropertyName, VariablePropertyName;
-	GetValueAndVariablePropertyNames(ValuePropertyName, VariablePropertyName);
-
-	ValueProperty = PropertyHandle->GetChildHandle(ValuePropertyName);
-	VariableProperty = PropertyHandle->GetChildHandle(VariablePropertyName);
+	// All camera parameters should have a "Value" and "Variable" property.
+	ValueProperty = PropertyHandle->GetChildHandle("Value");
+	VariableProperty = PropertyHandle->GetChildHandle("Variable");
+	ensure(ValueProperty && VariableProperty);
 
 	// Get the type of camera variable we need for this camera parameter (bool variable, float variable, etc.)
 	VariableClass = nullptr;
@@ -75,6 +71,7 @@ void FCameraParameterDetailsCustomization::CustomizeHeader(TSharedRef<IPropertyH
 	{
 		VariableClass = VariableObjectProperty->PropertyClass;
 	}
+	ensure(VariableClass);
 
 	// Update our variable info once now. We will then update it every tick, since the UI needs it
 	// for various things.
@@ -169,21 +166,23 @@ void FCameraParameterDetailsCustomization::CustomizeHeader(TSharedRef<IPropertyH
 			.OnGetMenuContent(this, &FCameraParameterDetailsCustomization::BuildCameraVariableBrowser)
 		]
 	];
-
-	// Setup some custom reset-to-default behavior, if we are allowed to.
-	// (see the code in FCameraRigAssetReferenceDetailsCustomization that sets this metadata)
-	const bool bNoResetToDefault = StructProperty->GetBoolMetaData("NoCustomCameraParameterResetToDefault");
-	if (!bNoResetToDefault)
-	{
-		HeaderRow.OverrideResetToDefault(
-				FResetToDefaultOverride::Create(
-					FIsResetToDefaultVisible::CreateSP(this, &FCameraParameterDetailsCustomization::IsResetToDefaultVisible),
-					FResetToDefaultHandler::CreateSP(this, &FCameraParameterDetailsCustomization::OnResetToDefault)));
-	}
 }
 
 void FCameraParameterDetailsCustomization::CustomizeChildren(TSharedRef<IPropertyHandle> PropertyHandle, IDetailChildrenBuilder& ChildBuilder, IPropertyTypeCustomizationUtils& CustomizationUtils)
 {
+	uint32 NumChildren = 0;
+	FPropertyAccess::Result Result = ValueProperty->GetNumChildren(NumChildren);
+	if (Result == FPropertyAccess::Success)
+	{
+		for (uint32 Index = 0; Index < NumChildren; ++Index)
+		{
+			TSharedPtr<IPropertyHandle> ChildProperty = ValueProperty->GetChildHandle(Index);
+			if (ChildProperty)
+			{
+				ChildBuilder.AddProperty(ChildProperty.ToSharedRef());
+			}
+		}
+	}
 }
 
 void FCameraParameterDetailsCustomization::Tick(float DeltaTime)
@@ -196,6 +195,13 @@ void FCameraParameterDetailsCustomization::UpdateVariableInfo()
 {
 	VariableInfo = FCameraVariableInfo();
 
+	StructProperty->EnumerateRawData(
+			[this](void* RawData, const int32 ValueIndex, const int32 NumValues)
+			{
+				VariableInfo.bHasNonUserOverride |= HasNonUserOverride(RawData);
+				return true;
+			});
+
 	UObject* VariableObject = nullptr;
 	FPropertyAccess::Result PropertyAccessResult = VariableProperty->GetValue(VariableObject);
 	if (PropertyAccessResult == FPropertyAccess::Success)
@@ -209,7 +215,6 @@ void FCameraParameterDetailsCustomization::UpdateVariableInfo()
 				VariableInfo.InfoText = Variable->DisplayName.IsEmpty() ?
 					FText::FromName(Variable->GetFName()) :
 					FText::FromString(Variable->DisplayName);
-				VariableInfo.bIsExposedParameterVariable = VariableInfo.CommonVariable->bIsPrivate;
 			}
 			else
 			{
@@ -217,7 +222,7 @@ void FCameraParameterDetailsCustomization::UpdateVariableInfo()
 				VariableInfo.ErrorText = LOCTEXT("InvalidVariableObject", "Invalid Variable");
 			}
 		}
-		// else: variable is not set
+		// else: variable is not set, leave info/error texts empty.
 	}
 	else if (PropertyAccessResult == FPropertyAccess::MultipleValues)
 	{
@@ -278,15 +283,15 @@ TSharedRef<SWidget> FCameraParameterDetailsCustomization::BuildCameraVariableBro
 bool FCameraParameterDetailsCustomization::IsValueEditorEnabled() const
 {
 	// The value widget is enabled (i.e. the user can change the value) if the parameter isn't driven by
-	// a variable, or if that variable is a private variable meant to expose the parameter on the rig interface.
-	return (VariableInfo.VariableValue == ECameraVariableValue::NotSet || VariableInfo.bIsExposedParameterVariable);
+	// a variable that was set by the user.
+	return VariableInfo.VariableValue == ECameraVariableValue::NotSet;
 }
 
 bool FCameraParameterDetailsCustomization::IsCameraVariableBrowserEnabled() const
 {
 	// The variable picker is enabled if the parameter isn't exposed to the rig interface via a private variable,
 	// since we can't drive a value with both an interface parameter and a user-defined variable.
-	return !VariableInfo.bIsExposedParameterVariable;
+	return !VariableInfo.bHasNonUserOverride;
 }
 
 FText FCameraParameterDetailsCustomization::GetVariableInfoText() const
@@ -296,7 +301,7 @@ FText FCameraParameterDetailsCustomization::GetVariableInfoText() const
 
 EVisibility FCameraParameterDetailsCustomization::GetVariableInfoTextVisibility() const
 {
-	const bool bShowVariableInfoText = !VariableInfo.InfoText.IsEmpty() && !VariableInfo.bIsExposedParameterVariable;
+	const bool bShowVariableInfoText = !VariableInfo.InfoText.IsEmpty();
 	return bShowVariableInfoText ? EVisibility::Visible : EVisibility::Collapsed;
 }
 
@@ -307,7 +312,7 @@ FOptionalSize FCameraParameterDetailsCustomization::GetVariableInfoTextMaxWidth(
 	// IMPORTANT: update this if the main layout changes inside Construct()
 	const float FixedSpace = 1.f + (2.f+ 16.f + 2.f) + (2.f + 16.f + 2.f) + 1.f;
 
-	const bool bShowVariableInfoText = !VariableInfo.InfoText.IsEmpty() && !VariableInfo.bIsExposedParameterVariable;
+	const bool bShowVariableInfoText = !VariableInfo.InfoText.IsEmpty();
 	const float LayoutBoxWidth = LayoutBox ? LayoutBox->GetPaintSpaceGeometry().GetLocalSize().X : 0.f;
 	return bShowVariableInfoText ? FOptionalSize((LayoutBoxWidth - FixedSpace) / 3.f) : FOptionalSize(0);
 }
@@ -319,7 +324,7 @@ FText FCameraParameterDetailsCustomization::GetVariableErrorText() const
 
 EVisibility FCameraParameterDetailsCustomization::GetVariableErrorTextVisibility() const
 {
-	const bool bShowVariableErrorText = !VariableInfo.ErrorText.IsEmpty() && !VariableInfo.bIsExposedParameterVariable;
+	const bool bShowVariableErrorText = !VariableInfo.ErrorText.IsEmpty();
 	return bShowVariableErrorText ? EVisibility::Visible : EVisibility::Collapsed;
 }
 
@@ -328,7 +333,7 @@ FOptionalSize FCameraParameterDetailsCustomization::GetVariableErrorTextMaxWidth
 	// See comments in GetVariableInfoTextMaxWidth.
 	const float FixedSpace = 1.f + (2.f+ 16.f + 2.f) + (2.f + 16.f + 2.f) + 1.f;
 
-	const bool bShowVariableErrorText = !VariableInfo.ErrorText.IsEmpty() && !VariableInfo.bIsExposedParameterVariable;
+	const bool bShowVariableErrorText = !VariableInfo.ErrorText.IsEmpty();
 	const float LayoutBoxWidth = LayoutBox ? LayoutBox->GetPaintSpaceGeometry().GetLocalSize().X : 0.f;
 	return bShowVariableErrorText ? FOptionalSize((LayoutBoxWidth - FixedSpace) / 3.f) : FOptionalSize(0);
 }
@@ -340,8 +345,7 @@ bool FCameraParameterDetailsCustomization::CanClearVariable() const
 
 void FCameraParameterDetailsCustomization::OnClearVariable()
 {
-	VariableProperty->ResetToDefault();
-	PropertyUtilities->RequestForceRefresh();
+	OnSetVariable(nullptr);
 }
 
 void FCameraParameterDetailsCustomization::OnSetVariable(UCameraVariableAsset* InVariable)
@@ -378,35 +382,31 @@ void FCameraParameterDetailsCustomization::OnSetVariable(UCameraVariableAsset* I
 bool FCameraParameterDetailsCustomization::IsResetToDefaultVisible(TSharedPtr<IPropertyHandle> InPropertyHandle) const
 {
 	// The user can reset the camera parameter to its default if the value is non-default, and/or the
-	// variable is a user-defined variable. In other words, if the variable is private because the parameter
-	// is exposed, then we don't want to reset that -- we only want to reset the value.
-	return 
-		ValueProperty->CanResetToDefault() || 
-		(VariableProperty->CanResetToDefault() && !VariableInfo.bIsExposedParameterVariable);
+	// variable is a user-defined variable. In other words, the VariableID property should not play a role
+	// in this.
+	return ValueProperty->CanResetToDefault() || VariableProperty->CanResetToDefault();
 }
 
 void FCameraParameterDetailsCustomization::OnResetToDefault(TSharedPtr<IPropertyHandle> InPropertyHandle)
 {
-	// As mentioned above, if the camera parameter is exposed publicly on the camera rig, we only want
-	// to reset the value to the default (and keep the private variable set on it).
-	// Otherwise, we can reset both the value and the variable.
-	if (VariableInfo.bIsExposedParameterVariable)
-	{
-		ValueProperty->ResetToDefault();
-	}
-	else
-	{
-		StructProperty->ResetToDefault();
-	}
+	// As mentioned above, we only reset the value and the variable, not the VariableID.
+	ValueProperty->ResetToDefault();
+	VariableProperty->ResetToDefault();
 
 	PropertyUtilities->RequestForceRefresh();
 }
 
 #define UE_CAMERA_VARIABLE_FOR_TYPE(ValueType, ValueName)\
+bool F##ValueName##CameraParameterDetailsCustomization::HasNonUserOverride(void* InRawData)\
+{\
+	F##ValueName##CameraParameter* TypedData = reinterpret_cast<F##ValueName##CameraParameter*>(InRawData);\
+	return TypedData->HasNonUserOverride();\
+}\
 void F##ValueName##CameraParameterDetailsCustomization::SetParameterVariable(void* InRawData, UCameraVariableAsset* InVariable)\
 {\
 	F##ValueName##CameraParameter* TypedData = reinterpret_cast<F##ValueName##CameraParameter*>(InRawData);\
 	TypedData->Variable = CastChecked<U##ValueName##CameraVariable>(InVariable, ECastCheckedType::NullAllowed);\
+	TypedData->VariableID = InVariable ? InVariable->GetVariableID() : FCameraVariableID();\
 }
 UE_CAMERA_VARIABLE_FOR_ALL_TYPES()
 #undef UE_CAMERA_VARIABLE_FOR_TYPE
diff --git a/Engine/Plugins/Cameras/GameplayCameras/Source/GameplayCamerasEditor/Private/Customizations/CameraRigAssetReferenceDetailsCustomization.cpp b/Engine/Plugins/Cameras/GameplayCameras/Source/GameplayCamerasEditor/Private/Customizations/CameraRigAssetReferenceDetailsCustomization.cpp
--- a/Engine/Plugins/Cameras/GameplayCameras/Source/GameplayCamerasEditor/Private/Customizations/CameraRigAssetReferenceDetailsCustomization.cpp
+++ b/Engine/Plugins/Cameras/GameplayCameras/Source/GameplayCamerasEditor/Private/Customizations/CameraRigAssetReferenceDetailsCustomization.cpp
@@ -39,12 +39,12 @@ protected:
 		
 		virtual bool IsPropertyOverridden(const FGuid PropertyID) const override
 		{
-			return CameraRigAssetReference.IsParameterOverriden(PropertyID);
+			return CameraRigAssetReference.IsParameterOverridden(PropertyID);
 		}
 		
 		virtual void SetPropertyOverride(const FGuid PropertyID, const bool bIsOverridden) const override
 		{
-			CameraRigAssetReference.SetParameterOverriden(PropertyID, bIsOverridden);
+			CameraRigAssetReference.SetParameterOverridden(PropertyID, bIsOverridden);
 		}
 
 	private:
diff --git a/Engine/Plugins/Cameras/GameplayCameras/Source/GameplayCamerasEditor/Private/Customizations/CameraVariableReferenceDetailsCustomizations.cpp b/Engine/Plugins/Cameras/GameplayCameras/Source/GameplayCamerasEditor/Private/Customizations/CameraVariableReferenceDetailsCustomizations.cpp
--- a/Engine/Plugins/Cameras/GameplayCameras/Source/GameplayCamerasEditor/Private/Customizations/CameraVariableReferenceDetailsCustomizations.cpp
+++ b/Engine/Plugins/Cameras/GameplayCameras/Source/GameplayCamerasEditor/Private/Customizations/CameraVariableReferenceDetailsCustomizations.cpp
@@ -207,17 +207,34 @@ bool FCameraVariableReferenceDetailsCustomization::CanClearVariable() const
 
 void FCameraVariableReferenceDetailsCustomization::OnClearVariable()
 {
-	VariableProperty->ResetToDefault();
-	PropertyUtilities->RequestForceRefresh();
+	OnSetVariable(nullptr);
 }
 
 void FCameraVariableReferenceDetailsCustomization::OnSetVariable(UCameraVariableAsset* InVariable)
 {
+	TArray<void*> RawData;
+	StructProperty->AccessRawData(RawData);
+
+	TArray<UObject*> OuterObjects;
+	StructProperty->GetOuterObjects(OuterObjects);
+
+	check(!OuterObjects.Num() || OuterObjects.Num() == RawData.Num());
+
 	{
 		FScopedTransaction Transaction(FText::Format(
 					LOCTEXT("SetPropertyValue", "Set {0}"), StructProperty->GetPropertyDisplayName()));
 
-		VariableProperty->SetValue(InVariable);
+		StructProperty->NotifyPreChange();
+
+		for (int32 ValueIndex = 0; ValueIndex < RawData.Num(); ++ValueIndex)
+		{
+			SetReferenceVariable(RawData[ValueIndex], InVariable);
+		}
+
+		StructProperty->NotifyPostChange(EPropertyChangeType::ValueSet);
+
+		FPropertyChangedEvent ChangeEvent(StructProperty->GetProperty(), EPropertyChangeType::ValueSet, OuterObjects);
+		PropertyUtilities->NotifyFinishedChangingProperties(ChangeEvent);
 	}
 
 	PropertyUtilities->RequestForceRefresh();
@@ -229,6 +246,16 @@ void FCameraVariableReferenceDetailsCustomization::OnResetToDefault()
 	PropertyUtilities->RequestForceRefresh();
 }
 
+#define UE_CAMERA_VARIABLE_FOR_TYPE(ValueType, ValueName)\
+void F##ValueName##CameraVariableReferenceDetailsCustomization::SetReferenceVariable(void* InRawData, UCameraVariableAsset* InVariable)\
+{\
+	F##ValueName##CameraVariableReference* TypedData = reinterpret_cast<F##ValueName##CameraVariableReference*>(InRawData);\
+	TypedData->Variable = CastChecked<U##ValueName##CameraVariable>(InVariable, ECastCheckedType::NullAllowed);\
+	TypedData->VariableID = InVariable ? InVariable->GetVariableID() : FCameraVariableID();\
+}
+UE_CAMERA_VARIABLE_FOR_ALL_TYPES()
+#undef UE_CAMERA_VARIABLE_FOR_TYPE
+
 }  // namespace UE::Cameras
 
 #undef LOCTEXT_NAMESPACE
diff --git a/Engine/Plugins/Cameras/GameplayCameras/Source/GameplayCamerasEditor/Private/Editors/CameraNodeGraphNode.cpp b/Engine/Plugins/Cameras/GameplayCameras/Source/GameplayCamerasEditor/Private/Editors/CameraNodeGraphNode.cpp
--- a/Engine/Plugins/Cameras/GameplayCameras/Source/GameplayCamerasEditor/Private/Editors/CameraNodeGraphNode.cpp
+++ b/Engine/Plugins/Cameras/GameplayCameras/Source/GameplayCamerasEditor/Private/Editors/CameraNodeGraphNode.cpp
@@ -129,6 +129,9 @@ UE_CAMERA_VARIABLE_FOR_ALL_TYPES()
 					break;
 				UE_CAMERA_VARIABLE_FOR_ALL_TYPES()
 #undef UE_CAMERA_VARIABLE_FOR_TYPE
+				case ECameraVariableType::BlendableStruct:
+					PinType.PinSubCategoryObject = const_cast<UScriptStruct*>(BlendableParameter.BlendableStructType.Get());
+					break;
 			}
 
 			UEdGraphPin* ParameterPin = CreatePin(EGPD_Input, PinType, BlendableParameter.ParameterName);
diff --git a/Engine/Plugins/Cameras/GameplayCameras/Source/GameplayCamerasEditor/Private/Editors/CameraNodeGraphPinColors.cpp b/Engine/Plugins/Cameras/GameplayCameras/Source/GameplayCamerasEditor/Private/Editors/CameraNodeGraphPinColors.cpp
--- a/Engine/Plugins/Cameras/GameplayCameras/Source/GameplayCamerasEditor/Private/Editors/CameraNodeGraphPinColors.cpp
+++ b/Engine/Plugins/Cameras/GameplayCameras/Source/GameplayCamerasEditor/Private/Editors/CameraNodeGraphPinColors.cpp
@@ -31,6 +31,7 @@ void FCameraNodeGraphPinColors::Initialize()
 	PinColors.Add(TypeEnum->GetNameByValue((int64)ECameraVariableType::Rotator3d), Settings->RotatorPinTypeColor);
 	PinColors.Add(TypeEnum->GetNameByValue((int64)ECameraVariableType::Transform3f), Settings->TransformPinTypeColor);
 	PinColors.Add(TypeEnum->GetNameByValue((int64)ECameraVariableType::Transform3d), Settings->TransformPinTypeColor);
+	PinColors.Add(TypeEnum->GetNameByValue((int64)ECameraVariableType::BlendableStruct), Settings->StructPinTypeColor);
 
 	PinColors.Add(FBooleanCameraParameter::StaticStruct()->GetFName(), Settings->BooleanPinTypeColor);
 	PinColors.Add(FInteger32CameraParameter::StaticStruct()->GetFName(), Settings->IntPinTypeColor);
diff --git a/Engine/Plugins/Cameras/GameplayCameras/Source/GameplayCamerasEditor/Private/Editors/CameraNodeGraphSchema.cpp b/Engine/Plugins/Cameras/GameplayCameras/Source/GameplayCamerasEditor/Private/Editors/CameraNodeGraphSchema.cpp
--- a/Engine/Plugins/Cameras/GameplayCameras/Source/GameplayCamerasEditor/Private/Editors/CameraNodeGraphSchema.cpp
+++ b/Engine/Plugins/Cameras/GameplayCameras/Source/GameplayCamerasEditor/Private/Editors/CameraNodeGraphSchema.cpp
@@ -218,21 +218,26 @@ void UCameraNodeGraphSchema::GetGraphContextActions(FGraphContextMenuBuilder& Co
 						FText::GetEmpty(),
 						LOCTEXT("NewInterfaceParameterAction", "Camera Rig Parameter"),
 						LOCTEXT("NewInterfaceParameterActionToolTip", "Exposes this parameter on the camera rig"));
-			Action->Target = Cast<UCameraNode>(CameraNodeNode->GetObject());
-			Action->TargetPropertyName = DraggedPin->PinName;
 
 			if (DraggedPin->PinType.PinCategory == PC_CameraParameter ||
 				DraggedPin->PinType.PinCategory == PC_CameraVariableReference)
 			{
-				ECameraVariableType ParameterType;
+				ECameraVariableType VariableType;
+				const UScriptStruct* BlendableStructType = nullptr;
 
-				if (FStructProperty* StructProperty = CastField<FStructProperty>(Property))
+				FCustomCameraNodeBlendableParameter BlendableParameter;
+				if (CustomParameters.FindBlendableParameter(DraggedPin->PinName, BlendableParameter))
+				{
+					VariableType = BlendableParameter.ParameterType;
+					BlendableStructType = BlendableParameter.BlendableStructType;
+				}
+				else if (FStructProperty* StructProperty = CastField<FStructProperty>(Property))
 				{
 #define UE_CAMERA_VARIABLE_FOR_TYPE(ValueType, ValueName)\
 					if ((StructProperty->Struct == F##ValueName##CameraParameter::StaticStruct()) ||\
 						(StructProperty->Struct == F##ValueName##CameraVariableReference::StaticStruct()))\
 					{\
-						ParameterType = ECameraVariableType::ValueName;\
+						VariableType = ECameraVariableType::ValueName;\
 					}\
 					else
 					UE_CAMERA_VARIABLE_FOR_ALL_TYPES()
@@ -246,22 +251,16 @@ void UCameraNodeGraphSchema::GetGraphContextActions(FGraphContextMenuBuilder& Co
 				}
 				else
 				{
-					FCustomCameraNodeBlendableParameter BlendableParameter;
-					if (CustomParameters.FindBlendableParameter(DraggedPin->PinName, BlendableParameter))
-					{
-						ParameterType = BlendableParameter.ParameterType;
-					}
-					else
-					{
-						// Unexpected: a parameter pin was created, but we found no property or custom
-						// parameter for it!
-						ensure(false);
-						return;
-					}
+					// Unexpected as per previous comments.
+					ensure(false);
+					return;
 				}
 
-				Action->NewNodeType = EInterfaceParameterCreateNodeType::BlendableParameter;
-				Action->BlendableParameterType = ParameterType;
+				FCameraRigParameterDefinition NewParameterDefinition;
+				NewParameterDefinition.ParameterType = ECameraRigInterfaceParameterType::Blendable;
+				NewParameterDefinition.VariableType = VariableType;
+				NewParameterDefinition.BlendableStructType = BlendableStructType;
+				Action->ParameterDefinition = NewParameterDefinition;
 			}
 			else if (DraggedPin->PinType.PinCategory == PC_CameraContextData)
 			{
@@ -321,9 +320,11 @@ void UCameraNodeGraphSchema::GetGraphContextActions(FGraphContextMenuBuilder& Co
 					}
 				}
 
-				Action->NewNodeType = EInterfaceParameterCreateNodeType::DataParameter;
-				Action->DataParameterType = DataType;
-				Action->DataParameterTypeObject = DataTypeObject;
+				FCameraRigParameterDefinition NewParameterDefinition;
+				NewParameterDefinition.ParameterType = ECameraRigInterfaceParameterType::Data;
+				NewParameterDefinition.DataType = DataType;
+				NewParameterDefinition.DataTypeObject = DataTypeObject;
+				Action->ParameterDefinition = NewParameterDefinition;
 			}
 
 			ContextMenuBuilder.AddAction(StaticCastSharedPtr<FEdGraphSchemaAction>(Action.ToSharedPtr()));
@@ -581,25 +582,19 @@ UEdGraphNode* FCameraNodeGraphSchemaAction_NewInterfaceParameterNode::PerformAct
 
 	// Create a new interface parameter and set it up based on the pin we're creating it from, if any.
 	UCameraRigInterfaceParameterBase* NewInterfaceParameter = nullptr;
-	if (NewNodeType == EInterfaceParameterCreateNodeType::BlendableParameter)
+	if (ParameterDefinition.ParameterType == ECameraRigInterfaceParameterType::Blendable)
 	{
 		UCameraRigBlendableParameter* NewBlendableParameter = NewObject<UCameraRigBlendableParameter>(CameraRig, NAME_None, RF_Transactional);
-		if (FromPin)
-		{
-			NewBlendableParameter->ParameterType = (ECameraVariableType)StaticEnum<ECameraVariableType>()->GetValueByName(FromPin->PinType.PinSubCategory);
-		}
+		NewBlendableParameter->ParameterType = ParameterDefinition.VariableType;
+		NewBlendableParameter->BlendableStructType = ParameterDefinition.BlendableStructType;
 		CameraRig->Interface.BlendableParameters.Add(NewBlendableParameter);
 		NewInterfaceParameter = NewBlendableParameter;
 	}
-	else if (NewNodeType == EInterfaceParameterCreateNodeType::DataParameter)
+	else if (ParameterDefinition.ParameterType == ECameraRigInterfaceParameterType::Data)
 	{
 		UCameraRigDataParameter* NewDataParameter = NewObject<UCameraRigDataParameter>(CameraRig, NAME_None, RF_Transactional);
-		if (FromPin)
-		{
-			const UEnum* DataTypeEnum = StaticEnum<ECameraContextDataType>();
-			NewDataParameter->DataType = (ECameraContextDataType)DataTypeEnum->GetValueByName(FromPin->PinType.PinSubCategory);
-			NewDataParameter->DataTypeObject = FromPin->PinType.PinSubCategoryObject.Get();
-		}
+		NewDataParameter->DataType = ParameterDefinition.DataType;
+		NewDataParameter->DataTypeObject = ParameterDefinition.DataTypeObject;
 		CameraRig->Interface.DataParameters.Add(NewDataParameter);
 		NewInterfaceParameter = NewDataParameter;
 	}
diff --git a/Engine/Plugins/Cameras/GameplayCameras/Source/GameplayCamerasEditor/Private/Toolkits/CameraRigInterfaceParametersToolkit.cpp b/Engine/Plugins/Cameras/GameplayCameras/Source/GameplayCamerasEditor/Private/Toolkits/CameraRigInterfaceParametersToolkit.cpp
--- a/Engine/Plugins/Cameras/GameplayCameras/Source/GameplayCamerasEditor/Private/Toolkits/CameraRigInterfaceParametersToolkit.cpp
+++ b/Engine/Plugins/Cameras/GameplayCameras/Source/GameplayCamerasEditor/Private/Toolkits/CameraRigInterfaceParametersToolkit.cpp
@@ -7,6 +7,7 @@
 #include "Editor.h"
 #include "Editors/CameraNodeGraphDragDropOp.h"
 #include "Framework/MultiBox/MultiBoxBuilder.h"
+#include "IGameplayCamerasModule.h"
 #include "ScopedTransaction.h"
 #include "SPinTypeSelector.h"
 #include "Styles/GameplayCamerasEditorStyle.h"
@@ -157,7 +158,6 @@ protected:
 
 			return SNew(SBox)
 				.MinDesiredWidth(125.f)
-				.MaxDesiredWidth(400.f)
 				.Padding(this, &SCameraRigBlendableParameterTableRow::GetPinTypeSelectorPadding)
 				[
 					SNew(SPinTypeSelector, FGetPinTypeTree::CreateSP(this, &SCameraRigBlendableParameterTableRow::GetPinTypeTreeInfos))
@@ -181,7 +181,7 @@ protected:
 	FMargin GetPinTypeSelectorPadding() const
 	{
 		// Add some horizontal margin when the pin type selector is read-only, so that the transition on hover is seamless.
-		return IsHovered() ? FMargin(0) : FMargin(8, 0);
+		return IsHovered() ? FMargin(0) : FMargin(7, 0);
 	}
 
 	void GetPinTypeTreeInfos(TArray<FPinTypeTreeItem>& TypeTree, ETypeTreeFilter TypeTreeFilter) const
@@ -207,6 +207,28 @@ protected:
 		TypeTree.Add(MakeShared<FPinTypeTreeInfo>(UEdGraphSchema_K2::PC_Struct, TBaseStructure<FRotator>::Get(), LOCTEXT("Rotator3dType", "A double precision 3D rotation")));
 		TypeTree.Add(MakeShared<FPinTypeTreeInfo>(UEdGraphSchema_K2::PC_Struct, TVariantStructure<FTransform3f>::Get(), LOCTEXT("Transform3fType", "A 3D transformation")));
 		TypeTree.Add(MakeShared<FPinTypeTreeInfo>(UEdGraphSchema_K2::PC_Struct, TBaseStructure<FTransform>::Get(), LOCTEXT("Transform3dType", "A double precision 3D transformation")));
+
+		TSharedPtr<FPinTypeTreeInfo> Structs = MakeShared<FPinTypeTreeInfo>(
+				LOCTEXT("BlendableStructPinTypeLabel", "Blendable Structures"),
+				UEdGraphSchema_K2::PC_Struct,
+				K2Schema,
+				LOCTEXT("BlendableStructPinTypeToolTip", "Blendable structure types"),
+				true);
+		IGameplayCamerasModule& GameplayCamerasModule = IGameplayCamerasModule::Get();
+		for (const FBlendableStructInfo& BlendableStruct : GameplayCamerasModule.GetBlendableStructs())
+		{
+			if (const UScriptStruct* StructType = BlendableStruct.StructType)
+			{
+				Structs->Children.Add(
+						MakeShared<UEdGraphSchema_K2::FPinTypeTreeInfo>(
+							UEdGraphSchema_K2::PC_Struct,
+							const_cast<UScriptStruct*>(StructType),
+							StructType->GetToolTipText(),
+							false,
+							(uint8)EObjectReferenceType::NotAnObject));
+			}
+		}
+		TypeTree.Add(Structs);
 	}
 
 	FEdGraphPinType GetBlendableParameterPinType() const
@@ -266,6 +288,10 @@ protected:
 				PinType.PinCategory = UEdGraphSchema_K2::PC_Struct;
 				PinType.PinSubCategoryObject = TBaseStructure<FTransform>::Get();
 				break;
+			case ECameraVariableType::BlendableStruct:
+				PinType.PinCategory = UEdGraphSchema_K2::PC_Struct;
+				PinType.PinSubCategoryObject = const_cast<UScriptStruct*>(Item->BlendableStructType.Get());
+				break;
 		}
 		return PinType;
 	}
@@ -431,7 +457,6 @@ protected:
 
 			return SNew(SBox)
 				.MinDesiredWidth(125.f)
-				.MaxDesiredWidth(400.f)
 				.Padding(this, &SCameraRigDataParameterTableRow::GetPinTypeSelectorPadding)
 				[
 					SNew(SPinTypeSelector, FGetPinTypeTree::CreateUObject(K2Schema, &UEdGraphSchema_K2::GetVariableTypeTree))
@@ -450,7 +475,7 @@ protected:
 	FMargin GetPinTypeSelectorPadding() const
 	{
 		// Add some horizontal margin when the pin type selector is read-only, so that the transition on hover is seamless.
-		return IsHovered() ? FMargin(0) : FMargin(8, 0);
+		return IsHovered() ? FMargin(0) : FMargin(7, 0);
 	}
 
 	FEdGraphPinType GetDataParameterPinType() const
diff --git a/Engine/Plugins/Cameras/GameplayCameras/Source/GameplayCamerasEditor/Public/Customizations/CameraParameterDetailsCustomizations.h b/Engine/Plugins/Cameras/GameplayCameras/Source/GameplayCamerasEditor/Public/Customizations/CameraParameterDetailsCustomizations.h
--- a/Engine/Plugins/Cameras/GameplayCameras/Source/GameplayCamerasEditor/Public/Customizations/CameraParameterDetailsCustomizations.h
+++ b/Engine/Plugins/Cameras/GameplayCameras/Source/GameplayCamerasEditor/Public/Customizations/CameraParameterDetailsCustomizations.h
@@ -45,7 +45,7 @@ public:
 
 protected:
 
-	virtual void GetValueAndVariablePropertyNames(FName& OutValueProperty, FName& OutVariableProperty) = 0;
+	virtual bool HasNonUserOverride(void* InRawData) = 0;
 	virtual void SetParameterVariable(void* InRawData, UCameraVariableAsset* InVariable) = 0;
 
 private:
@@ -62,7 +62,7 @@ private:
 	{
 		UCameraVariableAsset* CommonVariable = nullptr;
 		ECameraVariableValue VariableValue = ECameraVariableValue::NotSet;
-		bool bIsExposedParameterVariable = false;
+		bool bHasNonUserOverride = false;
 
 		FText InfoText;
 		FText ErrorText;
@@ -112,11 +112,7 @@ protected:
 class F##ValueName##CameraParameterDetailsCustomization : public FCameraParameterDetailsCustomization\
 {\
 protected:\
-	virtual void GetValueAndVariablePropertyNames(FName& OutValueProperty, FName& OutVariableProperty) override\
-	{\
-		OutValueProperty = GET_MEMBER_NAME_CHECKED(F##ValueName##CameraParameter, Value);\
-		OutVariableProperty = GET_MEMBER_NAME_CHECKED(F##ValueName##CameraParameter, Variable);\
-	}\
+	virtual bool HasNonUserOverride(void* InRawData) override;\
 	virtual void SetParameterVariable(void* InRawData, UCameraVariableAsset* InVariable) override;\
 };
 UE_CAMERA_VARIABLE_FOR_ALL_TYPES()
diff --git a/Engine/Plugins/Cameras/GameplayCameras/Source/GameplayCamerasEditor/Public/Customizations/CameraVariableReferenceDetailsCustomizations.h b/Engine/Plugins/Cameras/GameplayCameras/Source/GameplayCamerasEditor/Public/Customizations/CameraVariableReferenceDetailsCustomizations.h
--- a/Engine/Plugins/Cameras/GameplayCameras/Source/GameplayCamerasEditor/Public/Customizations/CameraVariableReferenceDetailsCustomizations.h
+++ b/Engine/Plugins/Cameras/GameplayCameras/Source/GameplayCamerasEditor/Public/Customizations/CameraVariableReferenceDetailsCustomizations.h
@@ -2,7 +2,7 @@
 
 #pragma once
 
-#include "Core/CameraVariableTableFwd.h"
+#include "Core/CameraVariableReferences.h"
 #include "IPropertyTypeCustomization.h"
 
 class FPropertyEditorModule;
@@ -31,6 +31,10 @@ public:
 	virtual void CustomizeHeader(TSharedRef<IPropertyHandle> PropertyHandle, FDetailWidgetRow& HeaderRow, IPropertyTypeCustomizationUtils& CustomizationUtils) override;
 	virtual void CustomizeChildren(TSharedRef<IPropertyHandle> PropertyHandle, IDetailChildrenBuilder& ChildBuilder, IPropertyTypeCustomizationUtils& CustomizationUtils) override;
 
+protected:
+
+	virtual void SetReferenceVariable(void* InRawData, UCameraVariableAsset* InVariable) = 0;
+
 private:
 
 	TSharedRef<SWidget> BuildCameraVariableBrowser();
@@ -59,6 +63,8 @@ protected:
 #define UE_CAMERA_VARIABLE_FOR_TYPE(ValueType, ValueName)\
 class F##ValueName##CameraVariableReferenceDetailsCustomization : public FCameraVariableReferenceDetailsCustomization\
 {\
+protected:\
+	virtual void SetReferenceVariable(void* InRawData, UCameraVariableAsset* InVariable) override;\
 };
 UE_CAMERA_VARIABLE_FOR_ALL_TYPES()
 #undef UE_CAMERA_VARIABLE_FOR_TYPE
diff --git a/Engine/Plugins/Cameras/GameplayCameras/Source/GameplayCamerasEditor/Public/Editors/CameraNodeGraphSchema.h b/Engine/Plugins/Cameras/GameplayCameras/Source/GameplayCamerasEditor/Public/Editors/CameraNodeGraphSchema.h
--- a/Engine/Plugins/Cameras/GameplayCameras/Source/GameplayCamerasEditor/Public/Editors/CameraNodeGraphSchema.h
+++ b/Engine/Plugins/Cameras/GameplayCameras/Source/GameplayCamerasEditor/Public/Editors/CameraNodeGraphSchema.h
@@ -3,6 +3,7 @@
 #pragma once
 
 #include "Core/CameraContextDataTableFwd.h"
+#include "Core/CameraRigParameterDefinition.h"
 #include "Core/CameraVariableTableFwd.h"
 #include "EdGraph/EdGraphPin.h"
 #include "Editors/CameraNodeGraphPinColors.h"
@@ -57,13 +58,6 @@ private:
 	UE::Cameras::FCameraNodeGraphPinColors PinColors;
 };
 
-UENUM()
-enum class EInterfaceParameterCreateNodeType
-{
-	BlendableParameter,
-	DataParameter
-};
-
 /**
  * Graph editor action for adding a new camera rig parameter node.
  */
@@ -74,25 +68,9 @@ struct FCameraNodeGraphSchemaAction_NewInterfaceParameterNode : public FEdGraphS
 
 public:
 
-	/** The camera node being driven by the camera rig parameter. */
-	UPROPERTY()
-	TObjectPtr<UCameraNode> Target;
-
-	/** The property on the target camera node being driven by the camera rig parameter. */
-	UPROPERTY()
-	FName TargetPropertyName;
-
-	UPROPERTY()
-	EInterfaceParameterCreateNodeType NewNodeType = EInterfaceParameterCreateNodeType::BlendableParameter;
-
-	UPROPERTY()
-	ECameraVariableType BlendableParameterType;
-
-	UPROPERTY()
-	ECameraContextDataType DataParameterType;
-
+	/** The new parameter's definition. */
 	UPROPERTY()
-	TObjectPtr<const UObject> DataParameterTypeObject;
+	FCameraRigParameterDefinition ParameterDefinition;
 
 public:
 
diff --git a/Engine/Plugins/Cameras/GameplayCameras/Source/GameplayCamerasUncookedOnly/Private/BlueprintGraph/K2Node_CameraRigBase.cpp b/Engine/Plugins/Cameras/GameplayCameras/Source/GameplayCamerasUncookedOnly/Private/BlueprintGraph/K2Node_CameraRigBase.cpp
--- a/Engine/Plugins/Cameras/GameplayCameras/Source/GameplayCamerasUncookedOnly/Private/BlueprintGraph/K2Node_CameraRigBase.cpp
+++ b/Engine/Plugins/Cameras/GameplayCameras/Source/GameplayCamerasUncookedOnly/Private/BlueprintGraph/K2Node_CameraRigBase.cpp
@@ -14,6 +14,7 @@
 #include "GameplayCamerasDelegates.h"
 #include "Kismet2/BlueprintEditorUtils.h"
 #include "KismetCompiler.h"
+#include "Nodes/Framing/CameraFramingZone.h"
 #include "Subsystems/AssetEditorSubsystem.h"
 
 #define LOCTEXT_NAMESPACE "K2Node_CameraRigBase"
@@ -132,10 +133,10 @@ void UK2Node_CameraRigBase::OnCameraRigAssetBuilt(const UCameraRigAsset* InBuilt
 
 FEdGraphPinType UK2Node_CameraRigBase::MakeBlendableParameterPinType(const UCameraRigBlendableParameter* BlendableParameter)
 {
-	return MakeBlendableParameterPinType(BlendableParameter->PrivateVariable->GetVariableType());
+	return MakeBlendableParameterPinType(BlendableParameter->ParameterType, BlendableParameter->BlendableStructType);
 }
 
-FEdGraphPinType UK2Node_CameraRigBase::MakeBlendableParameterPinType(ECameraVariableType CameraVariableType)
+FEdGraphPinType UK2Node_CameraRigBase::MakeBlendableParameterPinType(ECameraVariableType CameraVariableType, const UScriptStruct* BlendableStructType)
 {
 	FName PinCategory;
 	FName PinSubCategory;
@@ -177,6 +178,10 @@ FEdGraphPinType UK2Node_CameraRigBase::MakeBlendableParameterPinType(ECameraVari
 			PinCategory = UEdGraphSchema_K2::PC_Struct;
 			PinSubCategoryObject = TBaseStructure<FTransform>::Get();
 			break;
+		case ECameraVariableType::BlendableStruct:
+			PinCategory = UEdGraphSchema_K2::PC_Struct;
+			PinSubCategoryObject = const_cast<UScriptStruct*>(BlendableStructType);
+			break;
 	}
 
 	FEdGraphPinType PinType;
@@ -225,9 +230,9 @@ FEdGraphPinType UK2Node_CameraRigBase::MakeDataParameterPinType(ECameraContextDa
 
 FName UK2Node_CameraRigBase::GetBlendableParameterInteropSettingFunctionName(const UCameraRigBlendableParameter* BlendableParameter)
 {
-	if (ensure(BlendableParameter && BlendableParameter->PrivateVariable))
+	if (ensure(BlendableParameter && BlendableParameter->PrivateVariableID))
 	{
-		return GetBlendableParameterInteropSettingFunctionName(BlendableParameter->PrivateVariable->GetVariableType());
+		return GetBlendableParameterInteropSettingFunctionName(BlendableParameter->ParameterType);
 	}
 	return NAME_None;
 }
@@ -265,6 +270,9 @@ FName UK2Node_CameraRigBase::GetBlendableParameterInteropSettingFunctionName(ECa
 		case ECameraVariableType::Transform3d:
 			CallSetParameterFuncName = GET_FUNCTION_NAME_CHECKED(UCameraRigParameterInterop, SetTransformParameter);
 			break;
+		case ECameraVariableType::BlendableStruct:
+			CallSetParameterFuncName = GET_FUNCTION_NAME_CHECKED(UCameraRigParameterInterop, SetBlendableStructParameter);
+			break;
 	}
 
 	return CallSetParameterFuncName;
diff --git a/Engine/Plugins/Cameras/GameplayCameras/Source/GameplayCamerasUncookedOnly/Private/BlueprintGraph/K2Node_CameraRigBase.h b/Engine/Plugins/Cameras/GameplayCameras/Source/GameplayCamerasUncookedOnly/Private/BlueprintGraph/K2Node_CameraRigBase.h
--- a/Engine/Plugins/Cameras/GameplayCameras/Source/GameplayCamerasUncookedOnly/Private/BlueprintGraph/K2Node_CameraRigBase.h
+++ b/Engine/Plugins/Cameras/GameplayCameras/Source/GameplayCamerasUncookedOnly/Private/BlueprintGraph/K2Node_CameraRigBase.h
@@ -45,7 +45,7 @@ public:
 public:
 
 	static FEdGraphPinType MakeBlendableParameterPinType(const UCameraRigBlendableParameter* BlendableParameter);
-	static FEdGraphPinType MakeBlendableParameterPinType(ECameraVariableType CameraVariableType);
+	static FEdGraphPinType MakeBlendableParameterPinType(ECameraVariableType CameraVariableType, const UScriptStruct* BlendableStructType);
 	static FEdGraphPinType MakeDataParameterPinType(const UCameraRigDataParameter* DataParameter);
 	static FEdGraphPinType MakeDataParameterPinType(ECameraContextDataType CameraContextDataType, const UObject* CameraContextDataTypeObject);
 
diff --git a/Engine/Plugins/Cameras/GameplayCameras/Source/GameplayCamerasUncookedOnly/Private/BlueprintGraph/K2Node_SetCameraRigParameter.cpp b/Engine/Plugins/Cameras/GameplayCameras/Source/GameplayCamerasUncookedOnly/Private/BlueprintGraph/K2Node_SetCameraRigParameter.cpp
--- a/Engine/Plugins/Cameras/GameplayCameras/Source/GameplayCamerasUncookedOnly/Private/BlueprintGraph/K2Node_SetCameraRigParameter.cpp
+++ b/Engine/Plugins/Cameras/GameplayCameras/Source/GameplayCamerasUncookedOnly/Private/BlueprintGraph/K2Node_SetCameraRigParameter.cpp
@@ -30,7 +30,7 @@ void UK2Node_SetCameraRigParameter::Initialize(const FAssetData& UnloadedCameraR
 	{
 		if (const UCameraRigBlendableParameter* BlendableParameter = LoadedCameraRig->Interface.FindBlendableParameterByName(InCameraParameterName))
 		{
-			Initialize(LoadedCameraRig, InCameraParameterName, BlendableParameter->ParameterType);
+			Initialize(LoadedCameraRig, InCameraParameterName, BlendableParameter->ParameterType, BlendableParameter->BlendableStructType);
 		}
 		else if (const UCameraRigDataParameter* DataParameter = LoadedCameraRig->Interface.FindDataParameterByName(InCameraParameterName))
 		{
@@ -40,12 +40,13 @@ void UK2Node_SetCameraRigParameter::Initialize(const FAssetData& UnloadedCameraR
 	}
 }
 
-void UK2Node_SetCameraRigParameter::Initialize(UCameraRigAsset* InCameraRig, const FString& InCameraParameterName, ECameraVariableType InCameraVariableType)
+void UK2Node_SetCameraRigParameter::Initialize(UCameraRigAsset* InCameraRig, const FString& InCameraParameterName, ECameraVariableType InCameraVariableType, const UScriptStruct* InBlendableStructType)
 {
 	CameraRig = InCameraRig;
 	CameraParameterName = InCameraParameterName;
 	CameraParameterType = EK2Node_CameraParameterType::Blendable;
 	BlendableCameraParameterType = InCameraVariableType;
+	BlendableStructType = InBlendableStructType;
 }
 
 void UK2Node_SetCameraRigParameter::Initialize(UCameraRigAsset* InCameraRig, const FString& InCameraParameterName, ECameraContextDataType InCameraContextDataType, const UObject* InCameraContextDataTypeObject)
@@ -66,7 +67,7 @@ void UK2Node_SetCameraRigParameter::AllocateDefaultPins()
 	switch (CameraParameterType)
 	{
 		case EK2Node_CameraParameterType::Blendable:
-			PinType = UK2Node_CameraRigBase::MakeBlendableParameterPinType(BlendableCameraParameterType);
+			PinType = UK2Node_CameraRigBase::MakeBlendableParameterPinType(BlendableCameraParameterType, BlendableStructType);
 			break;
 		case EK2Node_CameraParameterType::Data:
 			PinType = UK2Node_CameraRigBase::MakeDataParameterPinType(DataCameraParameterType, DataCameraParameterTypeObject);
diff --git a/Engine/Plugins/Cameras/GameplayCameras/Source/GameplayCamerasUncookedOnly/Private/BlueprintGraph/K2Node_SetCameraRigParameter.h b/Engine/Plugins/Cameras/GameplayCameras/Source/GameplayCamerasUncookedOnly/Private/BlueprintGraph/K2Node_SetCameraRigParameter.h
--- a/Engine/Plugins/Cameras/GameplayCameras/Source/GameplayCamerasUncookedOnly/Private/BlueprintGraph/K2Node_SetCameraRigParameter.h
+++ b/Engine/Plugins/Cameras/GameplayCameras/Source/GameplayCamerasUncookedOnly/Private/BlueprintGraph/K2Node_SetCameraRigParameter.h
@@ -31,7 +31,7 @@ public:
 	UK2Node_SetCameraRigParameter(const FObjectInitializer& ObjectInit);
 
 	void Initialize(const FAssetData& UnloadedCameraRig, const FString& InCameraParameterName);
-	void Initialize(UCameraRigAsset* InCameraRig, const FString& InCameraParameterName, ECameraVariableType InCameraVariableType);
+	void Initialize(UCameraRigAsset* InCameraRig, const FString& InCameraParameterName, ECameraVariableType InCameraVariableType, const UScriptStruct* InBlendableStructType);
 	void Initialize(UCameraRigAsset* InCameraRig, const FString& InCameraParameterName, ECameraContextDataType InCameraContextDataType, const UObject* InCameraContextDataTypeObject);
 
 public:
@@ -59,6 +59,9 @@ private:
 	UPROPERTY()
 	ECameraVariableType BlendableCameraParameterType;
 
+	UPROPERTY()
+	TObjectPtr<const UScriptStruct> BlendableStructType;
+
 	UPROPERTY()
 	ECameraContextDataType DataCameraParameterType;
 
diff --git a/Engine/Plugins/Cameras/GameplayCameras/Source/GameplayCamerasUncookedOnly/Private/BlueprintGraph/K2Node_SetCameraRigParameters.cpp b/Engine/Plugins/Cameras/GameplayCameras/Source/GameplayCamerasUncookedOnly/Private/BlueprintGraph/K2Node_SetCameraRigParameters.cpp
--- a/Engine/Plugins/Cameras/GameplayCameras/Source/GameplayCamerasUncookedOnly/Private/BlueprintGraph/K2Node_SetCameraRigParameters.cpp
+++ b/Engine/Plugins/Cameras/GameplayCameras/Source/GameplayCamerasUncookedOnly/Private/BlueprintGraph/K2Node_SetCameraRigParameters.cpp
@@ -56,7 +56,7 @@ void UK2Node_SetCameraRigParameters::ReallocatePinsDuringReconstruction(TArray<U
 			PreloadObject(BlendableParameter);
 			if (BlendableParameter)
 			{
-				PreloadObject(BlendableParameter->PrivateVariable);
+				PreloadObject(const_cast<UScriptStruct*>(BlendableParameter->BlendableStructType.Get()));
 			}
 		}
 		for (UCameraRigDataParameter* DataParameter : CameraRig->Interface.DataParameters)
@@ -152,7 +152,7 @@ void UK2Node_SetCameraRigParameters::ExpandNode(FKismetCompilerContext& Compiler
 			continue;
 		}
 
-		if (!BlendableParameter->PrivateVariable)
+		if (!BlendableParameter->PrivateVariableID)
 		{
 			CompilerContext.MessageLog.Error(*LOCTEXT("ErrorMissingParameterVariable", "SetCameraRigParameters node @@ needs camera rig @@ to be built.").ToString(), this, CameraRig);
 			continue;
@@ -330,7 +330,7 @@ void UK2Node_SetCameraRigParameters::CreateParameterPins()
 			continue;
 		}
 
-		if (!BlendableParameter->PrivateVariable)
+		if (!BlendableParameter->PrivateVariableID)
 		{
 			// Camera rig isn't fully built.
 			continue;
