Cameras: move debug drawing to the CameraSystemHost object

This is in preparation of adding a new PlayerCameraManager that also needs debug drawing.
Also remove unnecessary Become/EndViewTarget methods.

diff --git a/Engine/Plugins/Cameras/GameplayCameras/Source/GameplayCameras/Private/GameFramework/GameplayCameraSystemActor.cpp b/Engine/Plugins/Cameras/GameplayCameras/Source/GameplayCameras/Private/GameFramework/GameplayCameraSystemActor.cpp
--- a/Engine/Plugins/Cameras/GameplayCameras/Source/GameplayCameras/Private/GameFramework/GameplayCameraSystemActor.cpp
+++ b/Engine/Plugins/Cameras/GameplayCameras/Source/GameplayCameras/Private/GameFramework/GameplayCameraSystemActor.cpp
@@ -20,25 +20,11 @@ AGameplayCameraSystemActor::AGameplayCameraSystemActor(const FObjectInitializer&
 	RootComponent = CameraSystemComponent;
 }
 
-void AGameplayCameraSystemActor::BecomeViewTarget(APlayerController* PC)
-{
-	Super::BecomeViewTarget(PC);
-
-	CameraSystemComponent->OnBecomeViewTarget();
-}
-
 void AGameplayCameraSystemActor::CalcCamera(float DeltaTime, struct FMinimalViewInfo& OutResult)
 {
 	CameraSystemComponent->GetCameraView(DeltaTime, OutResult);
 }
 
-void AGameplayCameraSystemActor::EndViewTarget(APlayerController* PC)
-{
-	CameraSystemComponent->OnEndViewTarget();
-
-	Super::EndViewTarget(PC);
-}
-
 AGameplayCameraSystemActor* AGameplayCameraSystemActor::GetAutoSpawnedCameraSystemActor(APlayerController* PlayerController, bool bForceSpawn)
 {
 	static const TCHAR* AutoSpawnedActorName = TEXT("AutoSpawnedGameplayCameraSystemActor");
diff --git a/Engine/Plugins/Cameras/GameplayCameras/Source/GameplayCameras/Private/GameFramework/GameplayCameraSystemComponent.cpp b/Engine/Plugins/Cameras/GameplayCameras/Source/GameplayCameras/Private/GameFramework/GameplayCameraSystemComponent.cpp
--- a/Engine/Plugins/Cameras/GameplayCameras/Source/GameplayCameras/Private/GameFramework/GameplayCameraSystemComponent.cpp
+++ b/Engine/Plugins/Cameras/GameplayCameras/Source/GameplayCameras/Private/GameFramework/GameplayCameraSystemComponent.cpp
@@ -4,7 +4,6 @@
 
 #include "Components/BillboardComponent.h"
 #include "Core/CameraSystemEvaluator.h"
-#include "Debug/DebugDrawService.h"
 #include "Engine/Canvas.h"
 #include "Engine/Texture2D.h"
 #include "Engine/World.h"
@@ -66,21 +65,6 @@ void UGameplayCameraSystemComponent::OnRegister()
 #if WITH_EDITOR
 	CreateCameraSystemSpriteComponent();
 #endif  // WITH_EDITOR
-
-	AActor* OwnerActor = GetOwner();
-	if (!OwnerActor || OwnerActor->HasAnyFlags(RF_ClassDefaultObject))
-	{
-		return;
-	}
-
-#if UE_GAMEPLAY_CAMERAS_DEBUG
-	UWorld* World = GetWorld();
-	if (World && World->IsGameWorld())
-	{
-		DebugDrawDelegateHandle = UDebugDrawService::Register(
-				TEXT("Game"), FDebugDrawDelegate::CreateUObject(this, &UGameplayCameraSystemComponent::DebugDraw));
-	}
-#endif  // UE_GAMEPLAY_CAMERAS_DEBUG
 }
 
 #if WITH_EDITOR
@@ -227,48 +211,5 @@ void UGameplayCameraSystemComponent::EndPlay(const EEndPlayReason::Type EndPlayR
 	Super::EndPlay(EndPlayReason);
 }
 
-void UGameplayCameraSystemComponent::OnComponentDestroyed(bool bDestroyingHierarchy)
-{
-	Super::OnComponentDestroyed(bDestroyingHierarchy);
-
-#if UE_GAMEPLAY_CAMERAS_DEBUG
-	if (DebugDrawDelegateHandle.IsValid())
-	{
-		UDebugDrawService::Unregister(DebugDrawDelegateHandle);
-		DebugDrawDelegateHandle.Reset();
-	}
-#endif  // UE_GAMEPLAY_CAMERAS_DEBUG
-}
-
-void UGameplayCameraSystemComponent::OnBecomeViewTarget()
-{
-}
-
-void UGameplayCameraSystemComponent::OnEndViewTarget()
-{
-}
-
-#if UE_GAMEPLAY_CAMERAS_DEBUG
-
-void UGameplayCameraSystemComponent::DebugDraw(UCanvas* Canvas, APlayerController* PlayController)
-{
-	using namespace UE::Cameras;
-
-	TSharedPtr<FCameraSystemEvaluator> Evaluator = GetCameraSystemEvaluator(false);
-	if (Evaluator.IsValid())
-	{
-		// We're looking from the outside if we are not the view target, or if we don't have a player
-		// anymore (which happens in spectator mode like with the debug camera).
-		APlayerController* ActualPlayerController = WeakPlayerController.Get();
-		const bool bIsDebugCameraEnabled = (ActualPlayerController && (ActualPlayerController->GetViewTarget() != GetOwner() || !ActualPlayerController->Player));
-
-		FCameraSystemDebugUpdateParams DebugUpdateParams;
-		DebugUpdateParams.CanvasObject = Canvas;
-		DebugUpdateParams.bIsDebugCameraEnabled = bIsDebugCameraEnabled;
-		Evaluator->DebugUpdate(DebugUpdateParams);
-	}
-}
-#endif  // UE_GAMEPLAY_CAMERAS_DEBUG
-
 #undef LOCTEXT_NAMESPACE
 
diff --git a/Engine/Plugins/Cameras/GameplayCameras/Source/GameplayCameras/Private/GameFramework/GameplayCameraSystemHost.cpp b/Engine/Plugins/Cameras/GameplayCameras/Source/GameplayCameras/Private/GameFramework/GameplayCameraSystemHost.cpp
--- a/Engine/Plugins/Cameras/GameplayCameras/Source/GameplayCameras/Private/GameFramework/GameplayCameraSystemHost.cpp
+++ b/Engine/Plugins/Cameras/GameplayCameras/Source/GameplayCameras/Private/GameFramework/GameplayCameraSystemHost.cpp
@@ -3,6 +3,7 @@
 #include "GameFramework/GameplayCameraSystemHost.h"
 
 #include "Core/CameraSystemEvaluator.h"
+#include "Debug/DebugDrawService.h"
 #include "Engine/World.h"
 #include "GameFramework/GameplayCameraSystemActor.h"
 #include "GameFramework/GameplayCameraSystemComponent.h"
@@ -39,6 +40,8 @@ UGameplayCameraSystemHost* UGameplayCameraSystemHost::FindOrCreateHost(APlayerCo
 
 	UGameplayCameraSystemHost* NewHost = NewObject<UGameplayCameraSystemHost>(PlayerController, HostName);
 
+	NewHost->Initialize();
+
 	return NewHost;
 }
 
@@ -65,18 +68,36 @@ UGameplayCameraSystemHost* UGameplayCameraSystemHost::FindHost(APlayerController
 UGameplayCameraSystemHost::UGameplayCameraSystemHost(const FObjectInitializer& ObjectInitializer)
 	: Super(ObjectInitializer)
 {
-	if (!HasAnyFlags(RF_ClassDefaultObject))
+}
+
+void UGameplayCameraSystemHost::Initialize()
+{
+	Evaluator = MakeShared<FCameraSystemEvaluator>();
+	Evaluator->Initialize(this);
+
+#if UE_GAMEPLAY_CAMERAS_DEBUG
+	UWorld* World = GetWorld();
+	if (World && World->IsGameWorld())
 	{
-		Evaluator = MakeShared<FCameraSystemEvaluator>();
-		Evaluator->Initialize(this);
+		DebugDrawDelegateHandle = UDebugDrawService::Register(
+				TEXT("Game"), FDebugDrawDelegate::CreateUObject(this, &UGameplayCameraSystemHost::DebugDraw));
 	}
+#endif  // UE_GAMEPLAY_CAMERAS_DEBUG
 }
 
 void UGameplayCameraSystemHost::BeginDestroy()
 {
-	Super::BeginDestroy();
+#if UE_GAMEPLAY_CAMERAS_DEBUG
+	if (DebugDrawDelegateHandle.IsValid())
+	{
+		UDebugDrawService::Unregister(DebugDrawDelegateHandle);
+		DebugDrawDelegateHandle.Reset();
+	}
+#endif  // UE_GAMEPLAY_CAMERAS_DEBUG
 
 	Evaluator.Reset();
+
+	Super::BeginDestroy();
 }
 
 void UGameplayCameraSystemHost::AddReferencedObjects(UObject* Object, FReferenceCollector& Collector)
@@ -100,3 +121,24 @@ TSharedPtr<UE::Cameras::FCameraSystemEvaluator> UGameplayCameraSystemHost::GetCa
 	return Evaluator;
 }
 
+#if UE_GAMEPLAY_CAMERAS_DEBUG
+
+void UGameplayCameraSystemHost::DebugDraw(UCanvas* Canvas, APlayerController* PlayerController)
+{
+	using namespace UE::Cameras;
+
+	if (Evaluator.IsValid())
+	{
+		// We're looking from the outside if we are not the view target, or if we don't have a player
+		// anymore (which happens in spectator mode like with the debug camera).
+		APlayerController* ActualPlayerController = GetPlayerController();
+		const bool bIsDebugCameraEnabled = !ActualPlayerController || !ActualPlayerController->Player;
+
+		FCameraSystemDebugUpdateParams DebugUpdateParams;
+		DebugUpdateParams.CanvasObject = Canvas;
+		DebugUpdateParams.bIsDebugCameraEnabled = bIsDebugCameraEnabled;
+		Evaluator->DebugUpdate(DebugUpdateParams);
+	}
+}
+#endif  // UE_GAMEPLAY_CAMERAS_DEBUG
+
diff --git a/Engine/Plugins/Cameras/GameplayCameras/Source/GameplayCameras/Public/GameFramework/GameplayCameraSystemActor.h b/Engine/Plugins/Cameras/GameplayCameras/Source/GameplayCameras/Public/GameFramework/GameplayCameraSystemActor.h
--- a/Engine/Plugins/Cameras/GameplayCameras/Source/GameplayCameras/Public/GameFramework/GameplayCameraSystemActor.h
+++ b/Engine/Plugins/Cameras/GameplayCameras/Source/GameplayCameras/Public/GameFramework/GameplayCameraSystemActor.h
@@ -53,9 +53,7 @@ public:
 public:
 
 	// AActor interface.
-	virtual void BecomeViewTarget(APlayerController* PC) override;
 	virtual void CalcCamera(float DeltaTime, struct FMinimalViewInfo& OutResult) override;
-	virtual void EndViewTarget(APlayerController* PC) override;
 
 private:
 
diff --git a/Engine/Plugins/Cameras/GameplayCameras/Source/GameplayCameras/Public/GameFramework/GameplayCameraSystemComponent.h b/Engine/Plugins/Cameras/GameplayCameras/Source/GameplayCameras/Public/GameFramework/GameplayCameraSystemComponent.h
--- a/Engine/Plugins/Cameras/GameplayCameras/Source/GameplayCameras/Public/GameFramework/GameplayCameraSystemComponent.h
+++ b/Engine/Plugins/Cameras/GameplayCameras/Source/GameplayCameras/Public/GameFramework/GameplayCameraSystemComponent.h
@@ -2,7 +2,6 @@
 
 #pragma once
 
-#include "GameplayCameras.h"
 #include "UObject/ObjectMacros.h"
 #include "Components/SceneComponent.h"
 
@@ -61,20 +60,9 @@ public:
 	virtual void OnRegister() override;
 	virtual void BeginPlay() override;
 	virtual void EndPlay(const EEndPlayReason::Type EndPlayReason) override;
-	virtual void OnComponentDestroyed(bool bDestroyingHierarchy) override;
-
-public:
-
-	// Internal API
-	void OnBecomeViewTarget();
-	void OnEndViewTarget();
 
 private:
 
-#if UE_GAMEPLAY_CAMERAS_DEBUG
-	void DebugDraw(UCanvas* Canvas, APlayerController* PlayController);
-#endif  // UE_GAMEPLAY_CAMERAS_DEBUG
-
 #if WITH_EDITOR
 	void CreateCameraSystemSpriteComponent();
 #endif  // WITH_EDITOR
@@ -97,16 +85,12 @@ public:
 
 private:
 	
-	UPROPERTY()
+	UPROPERTY(Transient)
 	TObjectPtr<UGameplayCameraSystemHost> CameraSystemHost;
 
-	UPROPERTY()
+	UPROPERTY(Transient)
 	TWeakObjectPtr<APlayerController> WeakPlayerController;
 
-#if UE_GAMEPLAY_CAMERAS_DEBUG
-	FDelegateHandle DebugDrawDelegateHandle;
-#endif  // UE_GAMEPLAY_CAMERAS_DEBUG
-
 #if WITH_EDITORONLY_DATA
 
 	/** Sprite scaling for the editor. */
diff --git a/Engine/Plugins/Cameras/GameplayCameras/Source/GameplayCameras/Public/GameFramework/GameplayCameraSystemHost.h b/Engine/Plugins/Cameras/GameplayCameras/Source/GameplayCameras/Public/GameFramework/GameplayCameraSystemHost.h
--- a/Engine/Plugins/Cameras/GameplayCameras/Source/GameplayCameras/Public/GameFramework/GameplayCameraSystemHost.h
+++ b/Engine/Plugins/Cameras/GameplayCameras/Source/GameplayCameras/Public/GameFramework/GameplayCameraSystemHost.h
@@ -2,12 +2,14 @@
 
 #pragma once
 
-#include "UObject/Object.h"
+#include "GameplayCameras.h"
 #include "Templates/SharedPointerFwd.h"
+#include "UObject/Object.h"
 
 #include "GameplayCameraSystemHost.generated.h"
 
 class APlayerController;
+class UCanvas;
 
 namespace UE::Cameras
 {
@@ -49,6 +51,14 @@ protected:
 	// UObject interface.
 	GAMEPLAYCAMERAS_API virtual void BeginDestroy() override;
 
+private:
+
+	void Initialize();
+
+#if UE_GAMEPLAY_CAMERAS_DEBUG
+	void DebugDraw(UCanvas* Canvas, APlayerController* PlayerController);
+#endif  // UE_GAMEPLAY_CAMERAS_DEBUG
+
 public:
 
 	/**
@@ -70,5 +80,9 @@ private:
 
 	/** The camera system evaluator. */
 	TSharedPtr<FCameraSystemEvaluator> Evaluator;
+
+#if UE_GAMEPLAY_CAMERAS_DEBUG
+	FDelegateHandle DebugDrawDelegateHandle;
+#endif  // UE_GAMEPLAY_CAMERAS_DEBUG
 };
 
