Cameras: add more debug information about the active player controller

diff --git a/Engine/Plugins/Cameras/GameplayCameras/Source/GameplayCameras/Private/Debug/PlayerControllersDebugBlock.cpp b/Engine/Plugins/Cameras/GameplayCameras/Source/GameplayCameras/Private/Debug/PlayerControllersDebugBlock.cpp
--- a/Engine/Plugins/Cameras/GameplayCameras/Source/GameplayCameras/Private/Debug/PlayerControllersDebugBlock.cpp
+++ b/Engine/Plugins/Cameras/GameplayCameras/Source/GameplayCameras/Private/Debug/PlayerControllersDebugBlock.cpp
@@ -5,6 +5,8 @@
 #include "Camera/PlayerCameraManager.h"
 #include "Debug/CameraDebugBlockBuilder.h"
 #include "Debug/CameraDebugRenderer.h"
+#include "Debug/DebugTextRenderer.h"
+#include "Engine/LocalPlayer.h"
 #include "Engine/World.h"
 #include "GameFramework/PlayerController.h"
 
@@ -41,12 +43,30 @@ void FPlayerControllersDebugBlock::Initialize(UWorld* World)
 		}
 
 		APlayerCameraManager* CameraManager = PlayerController->PlayerCameraManager;
-		AActor* ActiveViewTarget = CameraManager ? CameraManager->GetViewTarget() : nullptr;
+		if (!CameraManager)
+		{
+			continue;
+		}
+
+		AActor* ActiveViewTarget = CameraManager->GetViewTarget();
+		const FMinimalViewInfo& ViewTargetPOV = CameraManager->GetCameraCacheView();
 
 		FPlayerControllerDebugInfo PlayerControllerInfo;
 		PlayerControllerInfo.PlayerControllerName = *GetNameSafe(PlayerController);
 		PlayerControllerInfo.CameraManagerName = *GetNameSafe(CameraManager);
-		PlayerControllerInfo.ActiveViewTargetName = *GetNameSafe(ActiveViewTarget);
+		PlayerControllerInfo.ViewTargetName = *GetNameSafe(ActiveViewTarget);
+
+		if (const ULocalPlayer* LocalPlayer = PlayerController->GetLocalPlayer())
+		{
+			PlayerControllerInfo.LocalPlayerName = LocalPlayer->GetName();
+			PlayerControllerInfo.DefaultAspectRatioAxisConstraint = LocalPlayer->AspectRatioAxisConstraint;
+		}
+
+		PlayerControllerInfo.ViewTargetLocation = ViewTargetPOV.Location;
+		PlayerControllerInfo.ViewTargetRotation = ViewTargetPOV.Rotation;
+		PlayerControllerInfo.ViewTargetFOV = ViewTargetPOV.FOV;
+		PlayerControllerInfo.ViewTargetAspectRatio = ViewTargetPOV.AspectRatio;
+
 		PlayerControllers.Add(PlayerControllerInfo);
 	}
 }
@@ -64,8 +84,25 @@ void FPlayerControllersDebugBlock::OnDebugDraw(const FCameraDebugBlockDrawParams
 				Renderer.AddText(TEXT("- {cam_notice}%s{cam_default}"), *PlayerController.PlayerControllerName);
 				Renderer.AddIndent();
 				{
+					Renderer.AddText(TEXT("Local player: {cam_notice}%s{cam_default}\n"), *PlayerController.LocalPlayerName);
 					Renderer.AddText(TEXT("Camera manager: {cam_notice}%s{cam_default}\n"), *PlayerController.CameraManagerName);
-					Renderer.AddText(TEXT("View target: {cam_notice}%s{cam_default}"), *PlayerController.ActiveViewTargetName);
+					Renderer.AddText(TEXT("View target: {cam_notice}%s{cam_default}"), *PlayerController.ViewTargetName);
+					Renderer.AddIndent();
+					{
+						Renderer.AddText(TEXT("Location  : %s\n"), *ToDebugString(PlayerController.ViewTargetLocation));
+						Renderer.AddText(TEXT("Rotation  : %s\n"), *ToDebugString(PlayerController.ViewTargetRotation));
+						Renderer.AddText(TEXT("FOV  : %s\n"), *ToDebugString(PlayerController.ViewTargetFOV));
+						Renderer.AddText(TEXT("AspectRatio  : %s\n"), *ToDebugString(PlayerController.ViewTargetAspectRatio));
+						if (PlayerController.DefaultAspectRatioAxisConstraint.IsSet())
+						{
+							Renderer.AddText(TEXT("DefaultAspectRatioAxisConstraint  : %s\n"), *ToDebugString(PlayerController.DefaultAspectRatioAxisConstraint.GetValue()));
+						}
+						else
+						{
+							Renderer.AddText(TEXT("DefaultAspectRatioAxisConstraint  : <no local player>\n"));
+						}
+					}
+					Renderer.RemoveIndent();
 				}
 				Renderer.RemoveIndent();
 			}
@@ -88,7 +125,11 @@ FArchive& operator<< (FArchive& Ar, FPlayerControllersDebugBlock::FPlayerControl
 {
 	Ar << PlayerControllerInfo.PlayerControllerName;
 	Ar << PlayerControllerInfo.CameraManagerName;
-	Ar << PlayerControllerInfo.ActiveViewTargetName;
+	Ar << PlayerControllerInfo.ViewTargetName;
+	Ar << PlayerControllerInfo.ViewTargetLocation;
+	Ar << PlayerControllerInfo.ViewTargetRotation;
+	Ar << PlayerControllerInfo.ViewTargetFOV;
+	Ar << PlayerControllerInfo.ViewTargetAspectRatio;
 	return Ar;
 }
 
diff --git a/Engine/Plugins/Cameras/GameplayCameras/Source/GameplayCameras/Public/Debug/PlayerControllersDebugBlock.h b/Engine/Plugins/Cameras/GameplayCameras/Source/GameplayCameras/Public/Debug/PlayerControllersDebugBlock.h
--- a/Engine/Plugins/Cameras/GameplayCameras/Source/GameplayCameras/Public/Debug/PlayerControllersDebugBlock.h
+++ b/Engine/Plugins/Cameras/GameplayCameras/Source/GameplayCameras/Public/Debug/PlayerControllersDebugBlock.h
@@ -3,6 +3,8 @@
 #pragma once
 
 #include "Debug/CameraDebugBlock.h"
+#include "Engine/EngineTypes.h"
+#include "Math/UnrealMath.h"
 
 #if UE_GAMEPLAY_CAMERAS_DEBUG
 
@@ -26,6 +28,7 @@ public:
 
 protected:
 
+	// FCameraDebugBlock interface.
 	virtual void OnDebugDraw(const FCameraDebugBlockDrawParams& Params, FCameraDebugRenderer& Renderer) override;
 	virtual void OnSerialize(FArchive& Ar) override;
 
@@ -35,7 +38,13 @@ private:
 	{
 		FString PlayerControllerName;
 		FString CameraManagerName;
-		FString ActiveViewTargetName;
+		FString ViewTargetName;
+		FString LocalPlayerName;
+		TOptional<EAspectRatioAxisConstraint> DefaultAspectRatioAxisConstraint;
+		FVector3d ViewTargetLocation;
+		FRotator3d ViewTargetRotation;
+		float ViewTargetFOV;
+		float ViewTargetAspectRatio;
 	};
 	TArray<FPlayerControllerDebugInfo> PlayerControllers;
 	bool bHadValidWorld = false;
