Cameras: remove expired evaluation contexts at the end of the system update

diff --git a/Engine/Plugins/Cameras/GameplayCameras/Source/GameplayCameras/Private/Core/CameraEvaluationContextStack.cpp b/Engine/Plugins/Cameras/GameplayCameras/Source/GameplayCameras/Private/Core/CameraEvaluationContextStack.cpp
--- a/Engine/Plugins/Cameras/GameplayCameras/Source/GameplayCameras/Private/Core/CameraEvaluationContextStack.cpp
+++ b/Engine/Plugins/Cameras/GameplayCameras/Source/GameplayCameras/Private/Core/CameraEvaluationContextStack.cpp
@@ -150,12 +150,16 @@ void FCameraEvaluationContextStack::OnEndCameraSystemUpdate()
 	// regardless of when, during next frame, they set their variables. This is because various 
 	// gameplay systems, Blueprint scripting, whatever, might set variables at any time.
 	TArray<TSharedPtr<FCameraEvaluationContext>> ContextsToVisit;
-	for (const FContextEntry& Entry : Entries)
+	for (auto It = Entries.CreateIterator(); It; ++It)
 	{
-		if (TSharedPtr<FCameraEvaluationContext> Context = Entry.WeakContext.Pin())
+		if (TSharedPtr<FCameraEvaluationContext> Context = It->WeakContext.Pin())
 		{
 			ContextsToVisit.Add(Context);
 		}
+		else
+		{
+			It.RemoveCurrent();
+		}
 	}
 	while (!ContextsToVisit.IsEmpty())
 	{
