Cameras: fix static analysis build

#rnx

diff --git a/Engine/Plugins/Cameras/GameplayCameras/Source/GameplayCameras/Public/Core/CameraContextDataTable.h b/Engine/Plugins/Cameras/GameplayCameras/Source/GameplayCameras/Public/Core/CameraContextDataTable.h
--- a/Engine/Plugins/Cameras/GameplayCameras/Source/GameplayCameras/Public/Core/CameraContextDataTable.h
+++ b/Engine/Plugins/Cameras/GameplayCameras/Source/GameplayCameras/Public/Core/CameraContextDataTable.h
@@ -311,6 +311,15 @@ ECameraContextDataType TCameraContextDataTraits<DataType>::GetDataType()
 	{
 		return ECameraContextDataType::Struct;
 	}
+	else
+	{
+		// Before C++23 we have to make the static assert dependent on the template argument,
+		// so we can't just do static_assert(false).
+		// The trick here (stolen from Raymond Chen's blog) is to rely on the fact that sizeof()
+		// is never zero, and pointers don't need complete types, so !sizeof(DataType*) will
+		// always be false and will satisfy the compiler.
+		static_assert(!sizeof(DataType*), "DataType must be FName, FString, an enum, or a UClass, UObject, or UStruct type.");
+	}
 };
 
 template<typename DataType>
