Cameras: implement a player camera manager that runs the camera system

The new player camera manager runs the FCameraSystemEvaluator directly from the camera syste host object. It is able to take over from the CameraSystemActor, and it is able to "steal" control from another player camera manager.

diff --git a/Engine/Plugins/Cameras/GameplayCameras/Source/GameplayCameras/Private/Core/CameraEvaluationContextStack.cpp b/Engine/Plugins/Cameras/GameplayCameras/Source/GameplayCameras/Private/Core/CameraEvaluationContextStack.cpp
--- a/Engine/Plugins/Cameras/GameplayCameras/Source/GameplayCameras/Private/Core/CameraEvaluationContextStack.cpp
+++ b/Engine/Plugins/Cameras/GameplayCameras/Source/GameplayCameras/Private/Core/CameraEvaluationContextStack.cpp
@@ -99,6 +99,34 @@ bool FCameraEvaluationContextStack::RemoveContext(TSharedRef<FCameraEvaluationCo
 	return false;
 }
 
+bool FCameraEvaluationContextStack::RemoveContextsOwnedBy(UObject* ContextOwner, bool bAlsoRemoveInnerOwners)
+{
+	bool bRemovedAny = false;
+
+	for (auto It = Entries.CreateIterator(); It; ++It)
+	{
+		FContextEntry& Entry = (*It);
+		if (TSharedPtr<FCameraEvaluationContext> Context = Entry.WeakContext.Pin())
+		{
+			if (!Context->GetOwner())
+			{
+				continue;
+			}
+
+			if (Context->GetOwner() == ContextOwner || (bAlsoRemoveInnerOwners && Context->GetOwner()->IsIn(ContextOwner)))
+			{
+				FCameraEvaluationContextDeactivateParams DeactivateParams;
+				Context->Deactivate(DeactivateParams);
+
+				It.RemoveCurrent();
+				bRemovedAny = true;
+			}
+		}
+	}
+
+	return bRemovedAny;
+}
+
 void FCameraEvaluationContextStack::PopContext()
 {
 	Entries.Pop();
diff --git a/Engine/Plugins/Cameras/GameplayCameras/Source/GameplayCameras/Private/Core/CameraSystemEvaluator.cpp b/Engine/Plugins/Cameras/GameplayCameras/Source/GameplayCameras/Private/Core/CameraSystemEvaluator.cpp
--- a/Engine/Plugins/Cameras/GameplayCameras/Source/GameplayCameras/Private/Core/CameraSystemEvaluator.cpp
+++ b/Engine/Plugins/Cameras/GameplayCameras/Source/GameplayCameras/Private/Core/CameraSystemEvaluator.cpp
@@ -342,6 +342,7 @@ void FCameraSystemEvaluator::GetEvaluatedCameraView(FMinimalViewInfo& DesiredVie
 	DesiredView.Location = CameraPose.GetLocation();
 	DesiredView.Rotation = CameraPose.GetRotation();
 	DesiredView.FOV = CameraPose.GetEffectiveFieldOfView();
+	DesiredView.DesiredFOV = DesiredView.FOV;
 
 	DesiredView.AspectRatio = CameraPose.GetSensorAspectRatio();
 	DesiredView.bConstrainAspectRatio = CameraPose.GetConstrainAspectRatio();
diff --git a/Engine/Plugins/Cameras/GameplayCameras/Source/GameplayCameras/Private/GameFramework/ActorCameraEvaluationContext.cpp b/Engine/Plugins/Cameras/GameplayCameras/Source/GameplayCameras/Private/GameFramework/ActorCameraEvaluationContext.cpp
new file mode 100644
--- /dev/null
+++ b/Engine/Plugins/Cameras/GameplayCameras/Source/GameplayCameras/Private/GameFramework/ActorCameraEvaluationContext.cpp
@@ -0,0 +1,144 @@
+// Copyright Epic Games, Inc. All Rights Reserved.
+
+#include "GameFramework/ActorCameraEvaluationContext.h"
+
+#include "Camera/CameraComponent.h"
+#include "Core/CameraAsset.h"
+#include "Core/CameraNodeEvaluator.h"
+#include "Core/CameraRigAsset.h"
+#include "Core/PostProcessSettingsCollection.h"
+#include "Directors/SingleCameraDirector.h"
+#include "GameFramework/Actor.h"
+
+namespace UE::Cameras
+{
+
+UE_DEFINE_CAMERA_EVALUATION_CONTEXT(FActorCameraEvaluationContext)
+
+FActorCameraEvaluationContext::FActorCameraEvaluationContext()
+{
+}
+
+FActorCameraEvaluationContext::FActorCameraEvaluationContext(UCameraComponent* InCameraComponent)
+{
+	FCameraEvaluationContextInitializeParams Params;
+	Params.Owner = InCameraComponent;
+	Params.CameraAsset = MakeCameraComponentCameraAsset(InCameraComponent);
+	Initialize(Params);
+}
+
+FActorCameraEvaluationContext::FActorCameraEvaluationContext(AActor* InActor)
+{
+	FCameraEvaluationContextInitializeParams Params;
+	Params.Owner = InActor;
+	Params.CameraAsset = MakeCalcCameraActorCameraAsset(InActor);
+	Initialize(Params);
+}
+
+UCameraAsset* FActorCameraEvaluationContext::MakeCameraComponentCameraAsset(UObject* OuterObject)
+{
+	UCameraComponentCameraNode* WrapperCameraNode = NewObject<UCameraComponentCameraNode>(OuterObject, NAME_None, RF_Transient);
+	return MakeSimpleCameraAsset(OuterObject, WrapperCameraNode);
+}
+
+UCameraAsset* FActorCameraEvaluationContext::MakeCalcCameraActorCameraAsset(UObject* OuterObject)
+{
+	UCalcCameraActorCameraNode* WrapperCameraNode = NewObject<UCalcCameraActorCameraNode>(OuterObject, NAME_None, RF_Transient);
+	return MakeSimpleCameraAsset(OuterObject, WrapperCameraNode);
+}
+
+UCameraAsset* FActorCameraEvaluationContext::MakeSimpleCameraAsset(UObject* OuterObject, UCameraNode* RootNode)
+{
+	UCameraRigAsset* CameraRig = NewObject<UCameraRigAsset>(OuterObject, NAME_None, RF_Transient);
+	CameraRig->RootNode = RootNode;
+
+	USingleCameraDirector* SingleDirector = NewObject<USingleCameraDirector>(OuterObject, NAME_None, RF_Transient);
+	SingleDirector->CameraRig = CameraRig;
+
+	UCameraAsset* CameraAsset = NewObject<UCameraAsset>(OuterObject, NAME_None, RF_Transient);
+	CameraAsset->SetCameraDirector(SingleDirector);
+
+	return CameraAsset;
+}
+
+void FActorCameraEvaluationContext::ApplyMinimalViewInfo(const FMinimalViewInfo& ViewInfo, FCameraNodeEvaluationResult& OutResult)
+{
+	FCameraPose& CameraPose = OutResult.CameraPose;
+	CameraPose.SetLocation(ViewInfo.Location);
+	CameraPose.SetRotation(ViewInfo.Rotation);
+	CameraPose.SetFieldOfView(ViewInfo.FOV);
+	CameraPose.SetNearClippingPlane(ViewInfo.PerspectiveNearClipPlane);
+	CameraPose.SetSensorWidth(CameraPose.GetSensorHeight() * ViewInfo.AspectRatio);
+	CameraPose.SetAspectRatioAxisConstraint(ViewInfo.AspectRatioAxisConstraint.Get(CameraPose.GetAspectRatioAxisConstraint()));
+	CameraPose.SetConstrainAspectRatio(ViewInfo.bConstrainAspectRatio);
+
+	if (ViewInfo.PostProcessBlendWeight > 0.f)
+	{
+		OutResult.PostProcessSettings.LerpAll(ViewInfo.PostProcessSettings, ViewInfo.PostProcessBlendWeight);
+	}
+}
+
+class FCameraComponentCameraNodeEvaluator : public FCameraNodeEvaluator
+{
+	UE_DECLARE_CAMERA_NODE_EVALUATOR(GAMEPLAYCAMERAS_API, FCameraComponentCameraNodeEvaluator)
+
+protected:
+
+	// FCameraNodeEvaluator interface.
+	virtual void OnRun(const FCameraNodeEvaluationParams& Params, FCameraNodeEvaluationResult& OutResult) override;
+};
+
+UE_DEFINE_CAMERA_NODE_EVALUATOR(FCameraComponentCameraNodeEvaluator)
+
+void FCameraComponentCameraNodeEvaluator::OnRun(const FCameraNodeEvaluationParams& Params, FCameraNodeEvaluationResult& OutResult)
+{
+	UCameraComponent* CameraComponent = Cast<UCameraComponent>(Params.EvaluationContext->GetOwner());
+	if (!ensure(CameraComponent))
+	{
+		return;
+	}
+
+	FMinimalViewInfo CameraView;
+	CameraComponent->GetCameraView(Params.DeltaTime, CameraView);
+	FActorCameraEvaluationContext::ApplyMinimalViewInfo(CameraView, OutResult);
+}
+
+class FCalcCameraActorCameraNodeEvaluator : public FCameraNodeEvaluator
+{
+	UE_DECLARE_CAMERA_NODE_EVALUATOR(GAMEPLAYCAMERAS_API, FCalcCameraActorCameraNodeEvaluator)
+
+protected:
+
+	// FCameraNodeEvaluator interface.
+	virtual void OnRun(const FCameraNodeEvaluationParams& Params, FCameraNodeEvaluationResult& OutResult) override;
+};
+
+UE_DEFINE_CAMERA_NODE_EVALUATOR(FCalcCameraActorCameraNodeEvaluator)
+
+void FCalcCameraActorCameraNodeEvaluator::OnRun(const FCameraNodeEvaluationParams& Params, FCameraNodeEvaluationResult& OutResult)
+{
+	AActor* Actor = Cast<AActor>(Params.EvaluationContext->GetOwner());
+	if (!ensure(Actor))
+	{
+		return;
+	}
+
+	FMinimalViewInfo CameraView;
+	Actor->CalcCamera(Params.DeltaTime, CameraView);
+	FActorCameraEvaluationContext::ApplyMinimalViewInfo(CameraView, OutResult);
+}
+
+}  // namespace UE::Cameras
+
+FCameraNodeEvaluatorPtr UCameraComponentCameraNode::OnBuildEvaluator(FCameraNodeEvaluatorBuilder& Builder) const
+{
+	using namespace UE::Cameras;
+	return Builder.BuildEvaluator<FCameraComponentCameraNodeEvaluator>();
+}
+
+FCameraNodeEvaluatorPtr UCalcCameraActorCameraNode::OnBuildEvaluator(FCameraNodeEvaluatorBuilder& Builder) const
+{
+	using namespace UE::Cameras;
+	return Builder.BuildEvaluator<FCalcCameraActorCameraNodeEvaluator>();
+}
+
diff --git a/Engine/Plugins/Cameras/GameplayCameras/Source/GameplayCameras/Private/GameFramework/GameplayCameraSystemActor.cpp b/Engine/Plugins/Cameras/GameplayCameras/Source/GameplayCameras/Private/GameFramework/GameplayCameraSystemActor.cpp
--- a/Engine/Plugins/Cameras/GameplayCameras/Source/GameplayCameras/Private/GameFramework/GameplayCameraSystemActor.cpp
+++ b/Engine/Plugins/Cameras/GameplayCameras/Source/GameplayCameras/Private/GameFramework/GameplayCameraSystemActor.cpp
@@ -6,6 +6,7 @@
 #include "Engine/World.h"
 #include "GameFramework/GameplayCameraSystemComponent.h"
 #include "GameFramework/GameplayCameraSystemHost.h"
+#include "GameFramework/GameplayCamerasPlayerCameraManager.h"
 #include "GameFramework/PlayerController.h"
 #include "GameplayCamerasSettings.h"
 
@@ -25,20 +26,20 @@ void AGameplayCameraSystemActor::CalcCamera(float DeltaTime, struct FMinimalView
 	CameraSystemComponent->GetCameraView(DeltaTime, OutResult);
 }
 
-AGameplayCameraSystemActor* AGameplayCameraSystemActor::GetAutoSpawnedCameraSystemActor(APlayerController* PlayerController, bool bForceSpawn)
+AGameplayCameraSystemActor* AGameplayCameraSystemActor::GetAutoSpawnedCameraSystemActor(APlayerController* PlayerController, bool bSpawnIfMissing)
 {
-	static const TCHAR* AutoSpawnedActorName = TEXT("AutoSpawnedGameplayCameraSystemActor");
+	bool bDidSpawn;
+	return GetAutoSpawnedCameraSystemActor(PlayerController, bSpawnIfMissing, &bDidSpawn);
+}
 
-	const UGameplayCamerasSettings* Settings = GetDefault<UGameplayCamerasSettings>();
-	if (!Settings->bAutoSpawnCameraSystemActor)
-	{
-		return nullptr;
-	}
+AGameplayCameraSystemActor* AGameplayCameraSystemActor::GetAutoSpawnedCameraSystemActor(APlayerController* PlayerController, bool bSpawnIfMissing, bool* bOutSpawned)
+{
+	static const TCHAR* AutoSpawnedActorName = TEXT("AutoSpawnedGameplayCameraSystemActor");
 
 	UGameplayCameraSystemHost* Host = UGameplayCameraSystemHost::FindHost(PlayerController);
 	if (!Host)
 	{
-		if (bForceSpawn)
+		if (bSpawnIfMissing)
 		{
 			Host = UGameplayCameraSystemHost::FindOrCreateHost(PlayerController);
 		}
@@ -51,10 +52,11 @@ AGameplayCameraSystemActor* AGameplayCameraSystemActor::GetAutoSpawnedCameraSyst
 		}
 	}
 
+	*bOutSpawned = false;
 	AGameplayCameraSystemActor* SpawnedActor = FindObject<AGameplayCameraSystemActor>(PlayerController, AutoSpawnedActorName);
 	if (!SpawnedActor)
 	{
-		if (bForceSpawn)
+		if (bSpawnIfMissing)
 		{
 			FActorSpawnParameters SpawnParams;
 			SpawnParams.Name = AutoSpawnedActorName;
@@ -63,12 +65,8 @@ AGameplayCameraSystemActor* AGameplayCameraSystemActor::GetAutoSpawnedCameraSyst
 			SpawnedActor = World->SpawnActor<AGameplayCameraSystemActor>(SpawnParams);
 
 			SpawnedActor->Rename(nullptr, PlayerController);
-			
-			UGameplayCameraSystemComponent* CameraSystemComponent = SpawnedActor->CameraSystemComponent;
-			if (ensure(CameraSystemComponent))
-			{
-				CameraSystemComponent->bSetPlayerControllerRotation = Settings->bAutoSpawnCameraSystemActorSetsControlRotation;
-			}
+
+			*bOutSpawned = true;
 		}
 		else
 		{
@@ -81,10 +79,31 @@ AGameplayCameraSystemActor* AGameplayCameraSystemActor::GetAutoSpawnedCameraSyst
 
 void AGameplayCameraSystemActor::AutoManageActiveViewTarget(APlayerController* PlayerController)
 {
-	AGameplayCameraSystemActor* SpawnedActor = GetAutoSpawnedCameraSystemActor(PlayerController, true);
+	const UGameplayCamerasSettings* Settings = GetDefault<UGameplayCamerasSettings>();
+	if (!Settings->bAutoSpawnCameraSystemActor)
+	{
+		return;
+	}
+
+	if (AGameplayCamerasPlayerCameraManager* PlayerCameraManager = Cast<AGameplayCamerasPlayerCameraManager>(PlayerController->PlayerCameraManager))
+	{
+		return;
+	}
+
+	bool bDidSpawn = false;
+	AGameplayCameraSystemActor* SpawnedActor = GetAutoSpawnedCameraSystemActor(PlayerController, true, &bDidSpawn);
 	if (SpawnedActor)
 	{
-		SpawnedActor->GetCameraSystemComponent()->ActivateCameraSystemForPlayerController(PlayerController);
+		UGameplayCameraSystemComponent* CameraSystemComponent = SpawnedActor->CameraSystemComponent;
+		if (ensure(CameraSystemComponent))
+		{
+			if (bDidSpawn)
+			{
+				CameraSystemComponent->bSetPlayerControllerRotation = Settings->bAutoSpawnCameraSystemActorSetsControlRotation;
+			}
+
+			CameraSystemComponent->ActivateCameraSystemForPlayerController(PlayerController);
+		}
 	}
 }
 
diff --git a/Engine/Plugins/Cameras/GameplayCameras/Source/GameplayCameras/Private/GameFramework/GameplayCamerasPlayerCameraManager.cpp b/Engine/Plugins/Cameras/GameplayCameras/Source/GameplayCameras/Private/GameFramework/GameplayCamerasPlayerCameraManager.cpp
new file mode 100644
--- /dev/null
+++ b/Engine/Plugins/Cameras/GameplayCameras/Source/GameplayCameras/Private/GameFramework/GameplayCamerasPlayerCameraManager.cpp
@@ -0,0 +1,358 @@
+// Copyright Epic Games, Inc. All Rights Reserved.
+
+#include "GameFramework/GameplayCamerasPlayerCameraManager.h"
+
+#include "Camera/CameraComponent.h"
+#include "Core/CameraDirectorEvaluator.h"
+#include "Core/CameraEvaluationContext.h"
+#include "Core/CameraEvaluationContextStack.h"
+#include "Core/CameraEvaluationService.h"
+#include "Core/CameraRigTransition.h"
+#include "Core/CameraSystemEvaluator.h"
+#include "Core/RootCameraNodeCameraRigEvent.h"
+#include "Engine/Canvas.h"
+#include "Engine/Engine.h"
+#include "GameFramework/ActorCameraEvaluationContext.h"
+#include "GameFramework/GameplayCameraComponent.h"
+#include "GameFramework/GameplayCameraSystemActor.h"
+#include "GameFramework/GameplayCameraSystemComponent.h"
+#include "GameFramework/GameplayCameraSystemHost.h"
+#include "GameFramework/PlayerController.h"
+#include "UObject/Package.h"
+#include "UObject/UObjectGlobals.h"
+
+namespace UE::Cameras
+{
+
+class FViewTargetTransitionParamsBlendCameraNodeEvaluator : public FSimpleBlendCameraNodeEvaluator
+{
+	UE_DECLARE_BLEND_CAMERA_NODE_EVALUATOR_EX(GAMEPLAYCAMERAS_API, FViewTargetTransitionParamsBlendCameraNodeEvaluator, FSimpleBlendCameraNodeEvaluator)
+
+protected:
+
+	// FSimpleBlendCameraNodeEvaluator interface.
+	virtual void OnRun(const FCameraNodeEvaluationParams& Params, FCameraNodeEvaluationResult& OutResult) override;
+	virtual void OnComputeBlendFactor(const FCameraNodeEvaluationParams& Params, FSimpleBlendCameraNodeEvaluationResult& OutResult) override;
+
+private:
+
+	float CurrentTime = 0.f;
+};
+
+void FViewTargetTransitionParamsBlendCameraNodeEvaluator::OnRun(const FCameraNodeEvaluationParams& Params, FCameraNodeEvaluationResult& OutResult)
+{
+	const UViewTargetTransitionParamsBlendCameraNode* TransitionParamsNode = GetCameraNodeAs<UViewTargetTransitionParamsBlendCameraNode>();
+	CurrentTime += Params.DeltaTime;
+	if (CurrentTime >= TransitionParamsNode->TransitionParams.BlendTime)
+	{
+		CurrentTime = TransitionParamsNode->TransitionParams.BlendTime;
+		SetBlendFinished();
+	}
+
+	FSimpleBlendCameraNodeEvaluator::OnRun(Params, OutResult);
+}
+
+void FViewTargetTransitionParamsBlendCameraNodeEvaluator::OnComputeBlendFactor(const FCameraNodeEvaluationParams& Params, FSimpleBlendCameraNodeEvaluationResult& OutResult)
+{
+	const UViewTargetTransitionParamsBlendCameraNode* TransitionParamsNode = GetCameraNodeAs<UViewTargetTransitionParamsBlendCameraNode>();
+	float TimeFactor = 1.f;
+	if (TransitionParamsNode->TransitionParams.BlendTime > 0.f)
+	{
+		TimeFactor = CurrentTime / TransitionParamsNode->TransitionParams.BlendTime;
+	}
+	OutResult.BlendFactor = TransitionParamsNode->TransitionParams.GetBlendAlpha(TimeFactor);
+}
+
+UE_DEFINE_BLEND_CAMERA_NODE_EVALUATOR(FViewTargetTransitionParamsBlendCameraNodeEvaluator)
+
+class FViewTargetContextReferencerService : public FCameraEvaluationService
+{
+public:
+
+	FViewTargetContextReferencerService()
+	{
+		SetEvaluationServiceFlags(ECameraEvaluationServiceFlags::NeedsRootCameraNodeEvents);
+	}
+
+	void AddViewTargetContext(AActor* InViewTarget, TSharedRef<FCameraEvaluationContext> InContext)
+	{
+		Entries.Add(FEntry{ InViewTarget, InContext });
+	}
+
+protected:
+
+	virtual void OnInitialize(const FCameraEvaluationServiceInitializeParams& Params) override
+	{
+		Evaluator = Params.Evaluator;
+	}
+
+	virtual void OnRootCameraNodeEvent(const FRootCameraNodeCameraRigEvent& InEvent) override
+	{
+		if (InEvent.EventType == ERootCameraNodeCameraRigEventType::Deactivated)
+		{
+			TSharedPtr<const FCameraEvaluationContext> Context = InEvent.CameraRigInfo.EvaluationContext;
+			if (Context)
+			{
+				for (auto It = Entries.CreateIterator(); It; ++It)
+				{
+					if (It->Context == Context)
+					{
+						It.RemoveCurrent();
+					}
+				}
+			}
+		}
+	}
+
+private:
+
+	struct FEntry
+	{
+		TWeakObjectPtr<AActor> WeakViewTarget;
+		TSharedPtr<FCameraEvaluationContext> Context;
+	};
+
+	FCameraSystemEvaluator* Evaluator = nullptr;
+	TArray<FEntry> Entries;
+};
+
+}  // namespace UE::Cameras
+
+AGameplayCamerasPlayerCameraManager::AGameplayCamerasPlayerCameraManager(const FObjectInitializer& ObjectInitializer)
+	: Super(ObjectInitializer)
+{
+}
+
+void AGameplayCamerasPlayerCameraManager::StealPlayerController(APlayerController* PlayerController)
+{
+	using namespace UE::Cameras;
+
+	if (!ensure(PlayerController->PlayerCameraManager != this))
+	{
+		return;
+	}
+
+	OriginalCameraManager = PlayerController->PlayerCameraManager;
+	AActor* OriginalViewTarget = OriginalCameraManager->GetViewTarget();
+
+	PlayerController->PlayerCameraManager = this;
+	InitializeFor(PlayerController);
+
+	if (WeakOriginalAutoCameraSystemActor != OriginalViewTarget)
+	{
+		// We had some arbitrary view target before. Let's keep it.
+		SetViewTarget(OriginalViewTarget);
+	}
+	else
+	{
+		// The original view target was our camera system, so nothing changes. We will restore 
+		// that actor as the view target on the original camera manager when we're done.
+		//
+		// Let's set the view target to whatever owns the active context, or fall back to the
+		// player controller if the active context owner isn't an actor. We need to avoid using
+		// SetViewTarget otherwise that would create another evaluation context. Instead we
+		// call AssignViewTarget directly.
+		TSharedPtr<FCameraEvaluationContext> ActiveContext = CameraSystemEvaluator->GetEvaluationContextStack().GetActiveContext();
+		if (ActiveContext && ActiveContext->GetOwner())
+		{
+			AActor* ActiveViewTarget = ActiveContext->GetOwner()->GetTypedOuter<AActor>();
+			AssignViewTarget(ActiveViewTarget, ViewTarget);
+		}
+		ViewTarget.CheckViewTarget(PlayerController);
+	}
+}
+
+void AGameplayCamerasPlayerCameraManager::ReleasePlayerController()
+{
+	if (!ensure(PCOwner && PCOwner->PlayerCameraManager == this))
+	{
+		return;
+	}
+
+	PCOwner->PlayerCameraManager = OriginalCameraManager;
+	if (AGameplayCameraSystemActor* OriginalAutoCameraSystemActor = WeakOriginalAutoCameraSystemActor.Get())
+	{
+		OriginalCameraManager->SetViewTarget(OriginalAutoCameraSystemActor);
+	}
+
+	ViewTarget.Target = nullptr;
+
+	OriginalCameraManager = nullptr;
+	WeakOriginalAutoCameraSystemActor = nullptr;
+
+	ReleaseCameraSystemHost();
+
+	PCOwner = nullptr;
+}
+
+void AGameplayCamerasPlayerCameraManager::InitializeFor(APlayerController* PlayerController)
+{
+	AcquireCameraSystemHost(PlayerController);
+	Super::InitializeFor(PlayerController);
+}
+
+void AGameplayCamerasPlayerCameraManager::SetViewTarget(AActor* NewViewTarget, FViewTargetTransitionParams TransitionParams)
+{
+	using namespace UE::Cameras;
+
+	// We want to keep our view target in sync with whatever is the active context owner in the camera system.
+	// If that context owner isn't an actor, and isn't inside an actor (like a component), we use the player
+	// controller as the view target.
+
+	if (NewViewTarget == nullptr)
+	{
+		FCameraEvaluationContextStack& ContextStack = CameraSystemEvaluator->GetEvaluationContextStack();
+		ContextStack.PopContext();
+
+		if (TSharedPtr<FCameraEvaluationContext> NewActiveContext = ContextStack.GetActiveContext())
+		{
+			if (UObject* NewActiveContextOwner = NewActiveContext->GetOwner())
+			{
+				NewViewTarget = Cast<AActor>(NewActiveContextOwner);
+				if (!NewViewTarget)
+				{
+					NewViewTarget = NewActiveContextOwner->GetTypedOuter<AActor>();
+				}
+			}
+		}
+	}
+
+	// We pass empty transition params here because we never want to use PendingViewTarget, just ViewTarget.
+	Super::SetViewTarget(NewViewTarget, FViewTargetTransitionParams());
+
+	if (!NewViewTarget)
+	{
+		return;
+	}
+
+	if (UGameplayCameraComponent* GameplayCameraComponent = NewViewTarget->FindComponentByClass<UGameplayCameraComponent>())
+	{
+		GameplayCameraComponent->ActivateCameraForPlayerController(PCOwner);
+	}
+	else if (UCameraComponent* CameraComponent = NewViewTarget->FindComponentByClass<UCameraComponent>())
+	{
+		TSharedRef<FActorCameraEvaluationContext> NewContext = MakeShared<FActorCameraEvaluationContext>(CameraComponent);
+		CameraSystemEvaluator->PushEvaluationContext(NewContext);
+	}
+	else
+	{
+		TSharedRef<FActorCameraEvaluationContext> NewContext = MakeShared<FActorCameraEvaluationContext>(NewViewTarget);
+		CameraSystemEvaluator->PushEvaluationContext(NewContext);
+	}
+
+	// If transition parameters were given, override the next activation for the new evaluation context.
+	TSharedPtr<FCameraEvaluationContext> NextContext = CameraSystemEvaluator->GetEvaluationContextStack().GetActiveContext();
+	if (NextContext)
+	{
+		ViewTargetContextReferencerService->AddViewTargetContext(NewViewTarget, NextContext.ToSharedRef());
+
+		if (TransitionParams.BlendTime > 0.f)
+		{
+			UViewTargetTransitionParamsBlendCameraNode* BlendNode = NewObject<UViewTargetTransitionParamsBlendCameraNode>(GetTransientPackage());
+			BlendNode->TransitionParams = TransitionParams;
+
+			UCameraRigTransition* Transition = NewObject<UCameraRigTransition>(GetTransientPackage());
+			Transition->Blend = BlendNode;
+
+			FCameraDirectorEvaluator* DirectorEvaluator = NextContext->GetDirectorEvaluator();
+			DirectorEvaluator->OverrideNextActivationTransition(Transition);
+		}
+	}
+}
+
+void AGameplayCamerasPlayerCameraManager::AcquireCameraSystemHost(APlayerController* PlayerController)
+{
+	using namespace UE::Cameras;
+
+	// Early-out if we already have acquired the host.
+	if (CameraSystemHost)
+	{
+		return;
+	}
+
+	// Find or create a camera system host and grab its camera system.
+	CameraSystemHost = UGameplayCameraSystemHost::FindOrCreateHost(PlayerController);
+	ensure(CameraSystemHost);
+	CameraSystemEvaluator = CameraSystemHost->GetCameraSystemEvaluator();
+
+	ViewTargetContextReferencerService = MakeShared<FViewTargetContextReferencerService>();
+	CameraSystemEvaluator->RegisterEvaluationService(ViewTargetContextReferencerService.ToSharedRef());
+
+	// See if there was an auto-spawned camera system actor. If so, remember to set it back.
+	WeakOriginalAutoCameraSystemActor = AGameplayCameraSystemActor::GetAutoSpawnedCameraSystemActor(PlayerController, false);
+}
+
+void AGameplayCamerasPlayerCameraManager::ReleaseCameraSystemHost()
+{
+	if (ViewTargetContextReferencerService)
+	{
+		CameraSystemEvaluator->UnregisterEvaluationService(ViewTargetContextReferencerService.ToSharedRef());
+	}
+
+	CameraSystemHost = nullptr;
+	CameraSystemEvaluator = nullptr;
+}
+
+void AGameplayCamerasPlayerCameraManager::ProcessViewRotation(float DeltaTime, FRotator& OutViewRotation, FRotator& OutDeltaRot)
+{
+	Super::ProcessViewRotation(DeltaTime, OutViewRotation, OutDeltaRot);
+}
+
+void AGameplayCamerasPlayerCameraManager::DoUpdateCamera(float DeltaTime)
+{
+	using namespace UE::Cameras;
+
+	Super::DoUpdateCamera(DeltaTime);
+
+	if (CameraSystemEvaluator.IsValid())
+	{
+		FillCameraCache(LastFrameDesiredView);
+
+		FCameraSystemEvaluationParams UpdateParams;
+		UpdateParams.DeltaTime = DeltaTime;
+		CameraSystemEvaluator->Update(UpdateParams);
+
+		FMinimalViewInfo DesiredView;
+		CameraSystemEvaluator->GetEvaluatedCameraView(DesiredView);
+
+		FillCameraCache(DesiredView);
+
+		LastFrameDesiredView = DesiredView;
+	}
+}
+
+void AGameplayCamerasPlayerCameraManager::BeginPlay()
+{
+	Super::BeginPlay();
+}
+
+void AGameplayCamerasPlayerCameraManager::EndPlay(const EEndPlayReason::Type EndPlayReason)
+{
+	// TODO: destroy camera system actor?
+
+	Super::EndPlay(EndPlayReason);
+}
+
+void AGameplayCamerasPlayerCameraManager::DisplayDebug(UCanvas* Canvas, const FDebugDisplayInfo& DebugDisplay, float& YL, float& YPos)
+{
+	int Indentation = 1;
+	int LineNumber = FMath::CeilToInt(YPos / YL);
+
+	UFont* DrawFont = GEngine->GetSmallFont();
+	Canvas->SetDrawColor(FColor::Yellow);
+	Canvas->DrawText(
+			DrawFont, 
+			FString::Printf(TEXT("Please use the Camera Debugger panel to inspect '%s'."), *GetNameSafe(this)),
+			Indentation * YL, (LineNumber++) * YL);
+
+	YPos = LineNumber * YL;
+
+	Super::DisplayDebug(Canvas, DebugDisplay, YL, YPos);
+}
+
+FCameraNodeEvaluatorPtr UViewTargetTransitionParamsBlendCameraNode::OnBuildEvaluator(FCameraNodeEvaluatorBuilder& Builder) const
+{
+	using namespace UE::Cameras;
+	return Builder.BuildEvaluator<FViewTargetTransitionParamsBlendCameraNodeEvaluator>();
+}
+
diff --git a/Engine/Plugins/Cameras/GameplayCameras/Source/GameplayCameras/Public/Core/CameraEvaluationContextStack.h b/Engine/Plugins/Cameras/GameplayCameras/Source/GameplayCameras/Public/Core/CameraEvaluationContextStack.h
--- a/Engine/Plugins/Cameras/GameplayCameras/Source/GameplayCameras/Public/Core/CameraEvaluationContextStack.h
+++ b/Engine/Plugins/Cameras/GameplayCameras/Source/GameplayCameras/Public/Core/CameraEvaluationContextStack.h
@@ -47,6 +47,9 @@ public:
 	/** Remove an existing context from the stack. */
 	bool RemoveContext(TSharedRef<FCameraEvaluationContext> Context);
 
+	/** Remove contexts owned by the given object. */
+	bool RemoveContextsOwnedBy(UObject* ContextOwner, bool bAlsoRemoveInnerOwners = false);
+
 	/** Pop the active (top) context. */
 	void PopContext();
 
diff --git a/Engine/Plugins/Cameras/GameplayCameras/Source/GameplayCameras/Public/GameFramework/ActorCameraEvaluationContext.h b/Engine/Plugins/Cameras/GameplayCameras/Source/GameplayCameras/Public/GameFramework/ActorCameraEvaluationContext.h
new file mode 100644
--- /dev/null
+++ b/Engine/Plugins/Cameras/GameplayCameras/Source/GameplayCameras/Public/GameFramework/ActorCameraEvaluationContext.h
@@ -0,0 +1,77 @@
+// Copyright Epic Games, Inc. All Rights Reserved.
+
+#pragma once
+
+#include "Core/CameraEvaluationContext.h"
+#include "Core/CameraNode.h"
+
+#include "ActorCameraEvaluationContext.generated.h"
+
+class AActor;
+class UCameraComponent;
+struct FMinimalViewInfo;
+
+namespace UE::Cameras
+{
+
+/**
+ * An evaluation context for an arbitrary actor.
+ * It procedurally creates a single camera director that runs a rig only composed of
+ * a camera node that simply copies the actor's properties.
+ */
+class FActorCameraEvaluationContext : public FCameraEvaluationContext
+{
+	UE_DECLARE_CAMERA_EVALUATION_CONTEXT(GAMEPLAYCAMERAS_API, FActorCameraEvaluationContext)
+
+public:
+
+	FActorCameraEvaluationContext();
+	FActorCameraEvaluationContext(UCameraComponent* InCameraComponent);
+	FActorCameraEvaluationContext(AActor* InActor);
+
+public:
+
+	static void ApplyMinimalViewInfo(const FMinimalViewInfo& ViewInfo, FCameraNodeEvaluationResult& OutResult);
+
+private:
+
+	static UCameraAsset* MakeCameraComponentCameraAsset(UObject* OuterObject);
+	static UCameraAsset* MakeCalcCameraActorCameraAsset(UObject* OuterObject);
+	static UCameraAsset* MakeSimpleCameraAsset(UObject* OuterObject, UCameraNode* RootNode);
+};
+
+}  // namespace UE::Cameras
+
+/**
+ * A simple camera node that copies the properties of a camera component and applies them to
+ * the evaluation output.
+ * Note that this node expects to run inside an evaluation context owned by the camera component's
+ * parent actor.
+ */
+UCLASS(MinimalAPI, Hidden)
+class UCameraComponentCameraNode : public UCameraNode
+{
+	GENERATED_BODY()
+
+protected:
+
+	// UCameraNode interface.
+	virtual FCameraNodeEvaluatorPtr OnBuildEvaluator(FCameraNodeEvaluatorBuilder& Builder) const override;
+};
+
+/**
+ * A simple camera node that calls an actor's CalcCamera functoin and copies the result to
+ * the evaluation output.
+ * Note that this node expects to run inside an evaluation context owned by that actor.
+ */
+UCLASS(MinimalAPI, Hidden)
+class UCalcCameraActorCameraNode : public UCameraNode
+{
+	GENERATED_BODY()
+
+protected:
+
+	// UCameraNode interface.
+	virtual FCameraNodeEvaluatorPtr OnBuildEvaluator(FCameraNodeEvaluatorBuilder& Builder) const override;
+};
+
diff --git a/Engine/Plugins/Cameras/GameplayCameras/Source/GameplayCameras/Public/GameFramework/GameplayCameraSystemActor.h b/Engine/Plugins/Cameras/GameplayCameras/Source/GameplayCameras/Public/GameFramework/GameplayCameraSystemActor.h
--- a/Engine/Plugins/Cameras/GameplayCameras/Source/GameplayCameras/Public/GameFramework/GameplayCameraSystemActor.h
+++ b/Engine/Plugins/Cameras/GameplayCameras/Source/GameplayCameras/Public/GameFramework/GameplayCameraSystemActor.h
@@ -40,7 +40,7 @@ public:
 	 * settings, or if no actor has been spawned yet.
 	 */
 	UFUNCTION(BlueprintCallable, Category=Camera)
-	static AGameplayCameraSystemActor* GetAutoSpawnedCameraSystemActor(APlayerController* PlayerController, bool bForceSpawn = false);
+	static AGameplayCameraSystemActor* GetAutoSpawnedCameraSystemActor(APlayerController* PlayerController, bool bSpawnIfMissing = false);
 
 	/**
 	 * Automatically sets a camera system actor as the view-target, spawning a unique
@@ -55,6 +55,10 @@ public:
 	// AActor interface.
 	virtual void CalcCamera(float DeltaTime, struct FMinimalViewInfo& OutResult) override;
 
+private:
+
+	static AGameplayCameraSystemActor* GetAutoSpawnedCameraSystemActor(APlayerController* PlayerController, bool bSpawnIfMissing, bool* bOutSpawned);
+
 private:
 
 	UPROPERTY(VisibleAnywhere, Category=Camera, BlueprintGetter="GetCameraSystemComponent", meta=(ExposeFunctionCategories="CameraSystem"))
diff --git a/Engine/Plugins/Cameras/GameplayCameras/Source/GameplayCameras/Public/GameFramework/GameplayCamerasPlayerCameraManager.h b/Engine/Plugins/Cameras/GameplayCameras/Source/GameplayCameras/Public/GameFramework/GameplayCamerasPlayerCameraManager.h
new file mode 100644
--- /dev/null
+++ b/Engine/Plugins/Cameras/GameplayCameras/Source/GameplayCameras/Public/GameFramework/GameplayCamerasPlayerCameraManager.h
@@ -0,0 +1,120 @@
+// Copyright Epic Games, Inc. All Rights Reserved.
+
+#pragma once
+
+#include "Camera/PlayerCameraManager.h"
+#include "Core/CameraEvaluationContext.h"
+#include "Nodes/Blends/SimpleBlendCameraNode.h"
+
+#include "GameplayCamerasPlayerCameraManager.generated.h"
+
+class AGameplayCameraSystemActor;
+class UGameplayCameraSystemHost;
+
+namespace UE::Cameras
+{
+
+class FCameraSystemEvaluator;
+class FViewTargetContextReferencerService;
+
+}  // namespace UE::Cameras
+
+/**
+ * A player camera manager that runs the GameplayCameras camera system.
+ *
+ * Setting the view target does the following:
+ * - Push a new evaluation context for the provided view target actor.
+ *    - If that actor contains a GameplayCameraComponent, use its evaluation context directly.
+ *    - If that actor contains a CameraComponent, make an evaluation context that wraps it
+ *      and runs by simply copying that camera's properties (see FCameraActorCameraEvaluationContext).
+ *    - For other actors, do as above, but convert the output of the actor's CalcCamera function.
+ * - The old view target's evaluation context is immediately removed from the evaluation stack.
+ *   For other way to handle evaluation contexts, call methods directly on the camera system
+ *   evaluator instead of going through the base APlayerCameraManager class.
+ *
+ * There is only ever one active view target, the "pending" view target isn't used. This is
+ * because we may be blending between more than two camera rigs that may belong to more than
+ * two actors.
+ */
+UCLASS(notplaceable, MinimalAPI)
+class AGameplayCamerasPlayerCameraManager : public APlayerCameraManager
+{
+	GENERATED_BODY()
+
+public:
+
+	AGameplayCamerasPlayerCameraManager(const FObjectInitializer& ObjectInitializer);
+
+	/** Gets the camera system host object. */
+	UGameplayCameraSystemHost* GetCameraSystemHost() const { return CameraSystemHost; }
+
+public:
+
+	/** Replace the camera manager currently set on the provided controller with this camera manager. */
+	UFUNCTION(BlueprintCallable, Category="Camera")
+	void StealPlayerController(APlayerController* PlayerController);
+
+	/** Restore an originally stolen camera manager (see StealPlayerController). */
+	UFUNCTION(BlueprintCallable, Category="Camera")
+	void ReleasePlayerController();
+
+public:
+
+	// APlayerCameraManager interface.
+	virtual void InitializeFor(APlayerController* PlayerController) override;
+	virtual void SetViewTarget(AActor* NewViewTarget, FViewTargetTransitionParams TransitionParams = FViewTargetTransitionParams()) override;
+	virtual void ProcessViewRotation(float DeltaTime, FRotator& OutViewRotation, FRotator& OutDeltaRot) override;
+
+	// AActor interface.
+	virtual void BeginPlay() override;
+	virtual void EndPlay(const EEndPlayReason::Type EndPlayReason) override;
+	virtual void DisplayDebug(UCanvas* Canvas, const FDebugDisplayInfo& DebugDisplay, float& YL, float& YPos) override;
+
+protected:
+
+	// APlayerCameraManager interface.
+	virtual void DoUpdateCamera(float DeltaTime) override;
+
+private:
+
+	void AcquireCameraSystemHost(APlayerController* PlayerController);
+	void ReleaseCameraSystemHost();
+
+private:
+
+	UPROPERTY(Transient)
+	TObjectPtr<UGameplayCameraSystemHost> CameraSystemHost;
+
+	UPROPERTY(Transient)
+	TObjectPtr<APlayerCameraManager> OriginalCameraManager;
+
+	UPROPERTY(Transient)
+	TWeakObjectPtr<AGameplayCameraSystemActor> WeakOriginalAutoCameraSystemActor;
+
+	TSharedPtr<UE::Cameras::FCameraSystemEvaluator> CameraSystemEvaluator;
+
+	TSharedPtr<UE::Cameras::FViewTargetContextReferencerService> ViewTargetContextReferencerService;
+
+	FMinimalViewInfo LastFrameDesiredView;
+};
+
+/**
+ * A blend node that implements the blend algorithms of the FViewTargetTransitionParams.
+ */
+UCLASS(MinimalAPI, Hidden)
+class UViewTargetTransitionParamsBlendCameraNode : public USimpleBlendCameraNode
+{
+	GENERATED_BODY()
+
+public:
+
+	/** The transition params to use. */
+	UPROPERTY()
+	FViewTargetTransitionParams TransitionParams;
+
+protected:
+
+	// UCameraNode interface.
+	virtual FCameraNodeEvaluatorPtr OnBuildEvaluator(FCameraNodeEvaluatorBuilder& Builder) const override;
+};
+
