Cameras: handle custom pins before property pins when making connections in the camera node graph

diff --git a/Engine/Plugins/Cameras/GameplayCameras/Source/GameplayCamerasEditor/Private/Editors/CameraNodeGraphSchema.cpp b/Engine/Plugins/Cameras/GameplayCameras/Source/GameplayCamerasEditor/Private/Editors/CameraNodeGraphSchema.cpp
--- a/Engine/Plugins/Cameras/GameplayCameras/Source/GameplayCamerasEditor/Private/Editors/CameraNodeGraphSchema.cpp
+++ b/Engine/Plugins/Cameras/GameplayCameras/Source/GameplayCamerasEditor/Private/Editors/CameraNodeGraphSchema.cpp
@@ -267,7 +267,13 @@ void UCameraNodeGraphSchema::GetGraphContextActions(FGraphContextMenuBuilder& Co
 				ECameraContextDataType DataType;
 				const UObject* DataTypeObject = nullptr;
 				
-				if (Property)
+				FCustomCameraNodeDataParameter DataParameter;
+				if (CustomParameters.FindDataParameter(DraggedPin->PinName, DataParameter))
+				{
+					DataType = DataParameter.ParameterType;
+					DataTypeObject = DataParameter.ParameterTypeObject;
+				}
+				else if (Property)
 				{
 					if (FNameProperty* NameProperty = CastField<FNameProperty>(Property))
 					{
@@ -306,18 +312,9 @@ void UCameraNodeGraphSchema::GetGraphContextActions(FGraphContextMenuBuilder& Co
 				}
 				else
 				{
-					FCustomCameraNodeDataParameter DataParameter;
-					if (CustomParameters.FindDataParameter(DraggedPin->PinName, DataParameter))
-					{
-						DataType = DataParameter.ParameterType;
-						DataTypeObject = DataParameter.ParameterTypeObject;
-					}
-					else
-					{
-						// Unexpected as per previous comments.
-						ensure(false);
-						return;
-					}
+					// Unexpected as per previous comments.
+					ensure(false);
+					return;
 				}
 
 				FCameraRigParameterDefinition NewParameterDefinition;
