Cameras: fix possible crash in camera parameter detail customization

diff --git a/Engine/Plugins/Cameras/GameplayCameras/Source/GameplayCamerasEditor/Private/Customizations/CameraParameterDetailsCustomizations.cpp b/Engine/Plugins/Cameras/GameplayCameras/Source/GameplayCamerasEditor/Private/Customizations/CameraParameterDetailsCustomizations.cpp
--- a/Engine/Plugins/Cameras/GameplayCameras/Source/GameplayCamerasEditor/Private/Customizations/CameraParameterDetailsCustomizations.cpp
+++ b/Engine/Plugins/Cameras/GameplayCameras/Source/GameplayCamerasEditor/Private/Customizations/CameraParameterDetailsCustomizations.cpp
@@ -195,12 +195,18 @@ void FCameraParameterDetailsCustomization::UpdateVariableInfo()
 {
 	VariableInfo = FCameraVariableInfo();
 
-	StructProperty->EnumerateRawData(
-			[this](void* RawData, const int32 ValueIndex, const int32 NumValues)
-			{
-				VariableInfo.bHasNonUserOverride |= HasNonUserOverride(RawData);
-				return true;
-			});
+	if (StructProperty->IsValidHandle())
+	{
+		StructProperty->EnumerateRawData(
+				[this](void* RawData, const int32 ValueIndex, const int32 NumValues)
+				{
+					if (RawData)
+					{
+						VariableInfo.bHasNonUserOverride |= HasNonUserOverride(RawData);
+					}
+					return true;
+				});
+	}
 
 	UObject* VariableObject = nullptr;
 	FPropertyAccess::Result PropertyAccessResult = VariableProperty->GetValue(VariableObject);
