Cameras: add support for creating a debug renderer with an FCanvas instead of a UCanvas

diff --git a/Engine/Plugins/Cameras/GameplayCameras/Source/GameplayCameras/Private/Debug/CameraDebugRenderer.cpp b/Engine/Plugins/Cameras/GameplayCameras/Source/GameplayCameras/Private/Debug/CameraDebugRenderer.cpp
--- a/Engine/Plugins/Cameras/GameplayCameras/Source/GameplayCameras/Private/Debug/CameraDebugRenderer.cpp
+++ b/Engine/Plugins/Cameras/GameplayCameras/Source/GameplayCameras/Private/Debug/CameraDebugRenderer.cpp
@@ -2,7 +2,6 @@
 
 #include "Debug/CameraDebugRenderer.h"
 
-#include "Algo/Find.h"
 #include "CanvasItem.h"
 #include "CanvasTypes.h"
 #include "Components/LineBatchComponent.h"
@@ -16,7 +15,7 @@
 #include "Engine/World.h"
 #include "HAL/IConsoleManager.h"
 #include "Math/Box2D.h"
-#include "Misc/TVariant.h"
+#include "SceneView.h"
 
 #if UE_GAMEPLAY_CAMERAS_DEBUG
 
@@ -95,12 +94,61 @@ static FAutoConsoleVariableRef CVarGameplayCamerasDebugDefaultCoordinateSystemAx
 	GGameplayCamerasDebugDefaultCoordinateSystemAxesLength,
 	TEXT("Default: 100. The default length of coordinate system axes."));
 
+bool GGameplayCamerasDebugDrawBackground = true;
+static FAutoConsoleVariableRef CVarGameplayCamerasDebugDrawBackground(
+	TEXT("GameplayCameras.Debug.DrawBackground"),
+	GGameplayCamerasDebugDrawBackground,
+	TEXT(""));
+
+float GGameplayCamerasDebugBackgroundOpacity = 0.6f;
+static FAutoConsoleVariableRef CVarGameplayCamerasDebugBackgroundOpacity(
+	TEXT("GameplayCameras.Debug.BackgroundOpacity"),
+	GGameplayCamerasDebugBackgroundOpacity,
+	TEXT(""));
+
+FString GGameplayCamerasDebugColorScheme = TEXT("SolarizedDark");
+static FAutoConsoleVariableRef CVarGameplayCamerasDebugColorScheme(
+	TEXT("GameplayCameras.Debug.ColorScheme"),
+	GGameplayCamerasDebugColorScheme,
+	TEXT(""));
+
 FCameraDebugRenderer::FCameraDebugRenderer(UWorld* InWorld, UCanvas* InCanvasObject, bool bInIsExternalRendering)
-	: World(InWorld)
-	, CanvasObject(InCanvasObject)
-	, bIsExternalRendering(bInIsExternalRendering)
-	, DrawColor(FColor::White)
 {
+	Canvas = nullptr;
+	CanvasSize = FVector2D(EForceInit::ForceInit);
+	if (InCanvasObject)
+	{
+		Canvas = InCanvasObject->Canvas;
+		SceneView = InCanvasObject->SceneView;
+
+		CanvasSize = FVector2d(InCanvasObject->SizeX, InCanvasObject->SizeY);
+	}
+
+	Initialize(InWorld, bInIsExternalRendering);
+}
+
+FCameraDebugRenderer::FCameraDebugRenderer(UWorld* InWorld, const FSceneView* InSceneView, FCanvas* InCanvas, bool bInIsExternalRendering)
+{
+	Canvas = nullptr;
+	CanvasSize = FVector2D(EForceInit::ForceInit);
+	if (InCanvas)
+	{
+		Canvas = InCanvas;
+		SceneView = InSceneView;
+
+		const FIntRect ViewRect = InCanvas->GetViewRect();
+		CanvasSize = FVector2d(ViewRect.Width(), ViewRect.Height());
+	}
+
+	Initialize(InWorld, bInIsExternalRendering);
+}
+
+void FCameraDebugRenderer::Initialize(UWorld* InWorld, bool bInIsExternalRendering)
+{
+	World = InWorld;
+	bIsExternalRendering = bInIsExternalRendering;
+	DrawColor = FColor::White;
+
 	RenderFont = GEngine->GetSmallFont();
 	MaxCharHeight = RenderFont->GetMaxCharHeight();
 
@@ -108,10 +156,10 @@ FCameraDebugRenderer::FCameraDebugRenderer(UWorld* InWorld, UCanvas* InCanvasObj
 
 	NextCardPosition = FVector2f::ZeroVector;
 	NextCardColumn = 0;
-	if (CanvasObject)
+	if (Canvas)
 	{
 		NextCardPosition = FVector2f{ 
-			CanvasObject->SizeX - (float)GGameplayCamerasDebugCardWidth - (float)GGameplayCamerasDebugRightMargin,
+			(float)CanvasSize.X - (float)GGameplayCamerasDebugCardWidth - (float)GGameplayCamerasDebugRightMargin,
 			(float)GGameplayCamerasDebugTopMargin };
 	}
 }
@@ -121,19 +169,19 @@ FCameraDebugRenderer::~FCameraDebugRenderer()
 	FlushText();
 }
 
-FCanvas* FCameraDebugRenderer::GetCanvas() const
+void FCameraDebugRenderer::BeginDrawing()
 {
-	return CanvasObject ? CanvasObject->Canvas : nullptr;
+	// Update the color scheme in case it changed.
+	FCameraDebugColors::Set(GGameplayCamerasDebugColorScheme);
 }
 
-FVector2D FCameraDebugRenderer::GetCanvasSize() const
+void FCameraDebugRenderer::EndDrawing()
 {
-	if (CanvasObject)
+	// Render a translucent background to help readability.
+	if (GGameplayCamerasDebugDrawBackground)
 	{
-		FIntPoint ParentSize = CanvasObject->Canvas->GetParentCanvasSize();
-		return FVector2D(ParentSize.X, ParentSize.Y);
+		DrawTextBackgroundTile(GGameplayCamerasDebugBackgroundOpacity);
 	}
-	return FVector2D::ZeroVector;
 }
 
 void FCameraDebugRenderer::AddText(const FString& InString)
@@ -199,10 +247,10 @@ void FCameraDebugRenderer::FlushText()
 {
 	if (LineBuilder.Len() > 0)
 	{
-		int32 ViewHeight = GetCanvasSize().Y;
+		int32 ViewHeight = CanvasSize.Y;
 		if (NextDrawPosition.Y < ViewHeight)
 		{
-			FDebugTextRenderer TextRenderer(GetCanvas(), DrawColor, RenderFont);
+			FDebugTextRenderer TextRenderer(Canvas, DrawColor, RenderFont);
 			TextRenderer.LeftMargin = GetIndentMargin();
 			TextRenderer.RenderText(NextDrawPosition, LineBuilder.ToView());
 
@@ -247,8 +295,8 @@ void FCameraDebugRenderer::RemoveIndent()
 void FCameraDebugRenderer::DrawTextBackgroundTile(float Opacity)
 {
 	const float IndentMargin = GetIndentMargin();
-	const bool bIsLineEmpty = FMath::IsNearlyEqual(NextDrawPosition.X, IndentMargin);
-	const float TextBottom = bIsLineEmpty ? NextDrawPosition.Y : NextDrawPosition.Y + MaxCharHeight;
+	const bool bIsLastLineEmpty = FMath::IsNearlyEqual(NextDrawPosition.X, IndentMargin);
+	const float TextBottom = bIsLastLineEmpty ? NextDrawPosition.Y : NextDrawPosition.Y + MaxCharHeight;
 
 	const float InnerMargin = GGameplayCamerasDebugInnerMargin;
 	const FVector2D TopLeft(GGameplayCamerasDebugLeftMargin - InnerMargin, GGameplayCamerasDebugTopMargin - InnerMargin);
@@ -257,8 +305,8 @@ void FCameraDebugRenderer::DrawTextBackgroundTile(float Opacity)
 
 	const FColor BackgroundColor = FCameraDebugColors::Get().Background.WithAlpha((uint8)(Opacity * 255));
 
-	// Draw the background behind the text.
-	if (FCanvas* Canvas = GetCanvas())
+	// Draw the background behind the text, if any.
+	if (Canvas && TextBottom > GGameplayCamerasDebugTopMargin)
 	{
 		Canvas->PushDepthSortKey(GGameplayCamerasDebugBackgroundDepthSortKey);
 		{
@@ -276,7 +324,7 @@ void FCameraDebugRenderer::DrawClock(FCameraDebugClock& InClock, const FText& In
 	DrawParams.ClockName = InClockName;
 	DrawParams.ClockPosition = GetNextCardPosition();
 	DrawParams.ClockSize = FVector2f(GGameplayCamerasDebugCardWidth, GGameplayCamerasDebugCardHeight);
-	InClock.Draw(GetCanvas(), DrawParams);
+	InClock.Draw(Canvas, DrawParams);
 }
 
 void FCameraDebugRenderer::DrawCameraPose(const FCameraPose& InCameraPose, const FLinearColor& LineColor, float CameraSize)
@@ -298,7 +346,7 @@ FVector2f FCameraDebugRenderer::GetNextCardPosition()
 		// We went over the number of columns we're supposed to stick to.
 		// Place the next card below the previous cards, at the right-side edge of the canvas.
 		NextCardColumn = 0;
-		NextCardPosition.X = CanvasObject->SizeX - (float)GGameplayCamerasDebugCardWidth - (float)GGameplayCamerasDebugRightMargin;
+		NextCardPosition.X = CanvasSize.X - (float)GGameplayCamerasDebugCardWidth - (float)GGameplayCamerasDebugRightMargin;
 		NextCardPosition.Y += GGameplayCamerasDebugCardHeight + GGameplayCamerasDebugCardGap;
 	}
 	else
@@ -319,7 +367,7 @@ void FCameraDebugRenderer::GetNextDrawGraphParams(FCameraDebugGraphDrawParams& O
 
 void FCameraDebugRenderer::Draw2DPointCross(const FVector2D& Location, float CrossSize, const FLinearColor& LineColor, float LineThickness)
 {
-	if (FCanvas* Canvas = GetCanvas())
+	if (Canvas)
 	{
 		const float HalfCrossSize = CrossSize / 2.f;
 
@@ -337,7 +385,7 @@ void FCameraDebugRenderer::Draw2DPointCross(const FVector2D& Location, float Cro
 
 void FCameraDebugRenderer::Draw2DLine(const FVector2D& Start, const FVector2D& End, const FLinearColor& LineColor, float LineThickness)
 {
-	if (FCanvas* Canvas = GetCanvas())
+	if (Canvas)
 	{
 		FCanvasLineItem LineItem(Start, End);
 		LineItem.SetColor(LineColor);
@@ -348,7 +396,7 @@ void FCameraDebugRenderer::Draw2DLine(const FVector2D& Start, const FVector2D& E
 
 void FCameraDebugRenderer::Draw2DBox(const FBox2D& Box, const FLinearColor& LineColor, float LineThickness)
 {
-	if (FCanvas* Canvas = GetCanvas())
+	if (Canvas)
 	{
 		FCanvasBoxItem BoxItem(Box.Min, Box.GetSize());
 		BoxItem.SetColor(LineColor);
@@ -359,7 +407,7 @@ void FCameraDebugRenderer::Draw2DBox(const FBox2D& Box, const FLinearColor& Line
 
 void FCameraDebugRenderer::Draw2DBox(const FVector2D& BoxPosition, const FVector2D& BoxSize, const FLinearColor& LineColor, float LineThickness)
 {
-	if (FCanvas* Canvas = GetCanvas())
+	if (Canvas)
 	{
 		FCanvasBoxItem BoxItem(BoxPosition, BoxSize);
 		BoxItem.SetColor(LineColor);
@@ -405,6 +453,14 @@ void FCameraDebugRenderer::DrawLine(const FVector3d& Start, const FVector3d& End
 	}
 }
 
+void FCameraDebugRenderer::DrawBox(const FVector3d& Center, const FVector3d& Size, const FLinearColor& LineColor, float LineThickness)
+{
+	if (ULineBatchComponent* LineBatcher = GetDebugLineBatcher())
+	{
+		LineBatcher->DrawBox(Center, Size, LineColor, 0.f, SDPG_Foreground, LineThickness);
+	}
+}
+
 void FCameraDebugRenderer::DrawSphere(const FVector3d& Center, float Radius, int32 Segments, const FLinearColor& LineColor, float LineThickness)
 {
 	if (ULineBatchComponent* LineBatcher = GetDebugLineBatcher())
@@ -527,16 +583,18 @@ void FCameraDebugRenderer::DrawText(const FVector3d& WorldPosition, const FStrin
 
 void FCameraDebugRenderer::DrawText(const FVector3d& WorldPosition, const FVector2d& ScreenOffset, const FString& Text, const FLinearColor& TextColor, UFont* TextFont)
 {
-	if (CanvasObject)
+	if (Canvas && SceneView)
 	{
-		const FColor PreviousColor = CanvasObject->DrawColor;
-		const FVector3d ScreenPosition = CanvasObject->Project(WorldPosition);
 		UFont* ActualTextFont = TextFont ? TextFont : GEngine->GetSmallFont();
-		CanvasObject->DrawColor = TextColor.ToFColor(true);
-		CanvasObject->DrawText(
-				ActualTextFont, Text, 
-				ScreenPosition.X + ScreenOffset.X, ScreenPosition.Y + ScreenOffset.Y);
-		CanvasObject->DrawColor = PreviousColor;
+
+		const FVector3d ScreenPosition = SceneView->Project(WorldPosition);
+		FCanvasTextStringViewItem TextItem(
+				FVector2D(ScreenPosition.X + ScreenOffset.Y, ScreenPosition.Y + ScreenOffset.Y),
+				Text,
+				ActualTextFont,
+				TextColor);
+		TextItem.BlendMode = SE_BLEND_Translucent;
+		Canvas->DrawItem(TextItem);	
 	}
 }
 
diff --git a/Engine/Plugins/Cameras/GameplayCameras/Source/GameplayCameras/Private/Debug/RootCameraDebugBlock.cpp b/Engine/Plugins/Cameras/GameplayCameras/Source/GameplayCameras/Private/Debug/RootCameraDebugBlock.cpp
--- a/Engine/Plugins/Cameras/GameplayCameras/Source/GameplayCameras/Private/Debug/RootCameraDebugBlock.cpp
+++ b/Engine/Plugins/Cameras/GameplayCameras/Source/GameplayCameras/Private/Debug/RootCameraDebugBlock.cpp
@@ -51,24 +51,6 @@ static FAutoConsoleVariableRef CVarGameplayCamerasDebugPoseStatsShowVariableIDs(
 	GGameplayCamerasDebugPoseStatsShowVariableIDs,
 	TEXT(""));
 
-bool GGameplayCamerasDebugDrawBackground = true;
-static FAutoConsoleVariableRef CVarGameplayCamerasDebugDrawBackground(
-	TEXT("GameplayCameras.Debug.DrawBackground"),
-	GGameplayCamerasDebugDrawBackground,
-	TEXT(""));
-
-float GGameplayCamerasDebugBackgroundOpacity = 0.6f;
-static FAutoConsoleVariableRef CVarGameplayCamerasDebugBackgroundOpacity(
-	TEXT("GameplayCameras.Debug.BackgroundOpacity"),
-	GGameplayCamerasDebugBackgroundOpacity,
-	TEXT(""));
-
-FString GGameplayCamerasDebugColorScheme = TEXT("SolarizedDark");
-static FAutoConsoleVariableRef CVarGameplayCamerasDebugColorScheme(
-	TEXT("GameplayCameras.Debug.ColorScheme"),
-	GGameplayCamerasDebugColorScheme,
-	TEXT(""));
-
 UE_DEFINE_CAMERA_DEBUG_BLOCK(FRootCameraDebugBlock)
 
 void FRootCameraDebugBlock::BuildDebugBlocks(const FCameraSystemEvaluator& CameraSystem, const FCameraDebugBlockBuildParams& Params, FCameraDebugBlockBuilder& Builder)
@@ -168,17 +150,10 @@ void FRootCameraDebugBlock::RootDebugDraw(FCameraDebugRenderer& Renderer)
 		Params.ActiveCategories.Add(FString(CategoryView));
 	}
 
-	// Update the color scheme in case it changed.
-	FCameraDebugColors::Set(GGameplayCamerasDebugColorScheme);
-
 	// Do the drawing!
+	Renderer.BeginDrawing();
 	FCameraDebugBlock::DebugDraw(Params, Renderer);
-
-	// Render a translucent background to help readability.
-	if (GGameplayCamerasDebugDrawBackground)
-	{
-		Renderer.DrawTextBackgroundTile(GGameplayCamerasDebugBackgroundOpacity);
-	}
+	Renderer.EndDrawing();
 }
 
 }  // namespace UE::Cameras
diff --git a/Engine/Plugins/Cameras/GameplayCameras/Source/GameplayCameras/Public/Debug/CameraDebugRenderer.h b/Engine/Plugins/Cameras/GameplayCameras/Source/GameplayCameras/Public/Debug/CameraDebugRenderer.h
--- a/Engine/Plugins/Cameras/GameplayCameras/Source/GameplayCameras/Public/Debug/CameraDebugRenderer.h
+++ b/Engine/Plugins/Cameras/GameplayCameras/Source/GameplayCameras/Public/Debug/CameraDebugRenderer.h
@@ -11,6 +11,7 @@
 #include "Misc/StringBuilder.h"
 
 class FCanvas;
+class FSceneView;
 class UCanvas;
 class UFont;
 class ULineBatchComponent;
@@ -43,12 +44,17 @@ public:
 
 	/** Creates a new debug renderer. */
 	GAMEPLAYCAMERAS_API FCameraDebugRenderer(UWorld* InWorld, UCanvas* InCanvasObject, bool bInIsExternalRendering = false);
+	/** Creates a new debug renderer. */
+	GAMEPLAYCAMERAS_API FCameraDebugRenderer(UWorld* InWorld, const FSceneView* InSceneView, FCanvas* InCanvas, bool bInIsExternalRendering = false);
 	/** Destroys the debug renderer. */
 	GAMEPLAYCAMERAS_API ~FCameraDebugRenderer();
 
 	/** Whether we are looking at the camera system from the "outside". */
 	bool IsExternalRendering() const { return bIsExternalRendering; }
 
+	GAMEPLAYCAMERAS_API void BeginDrawing();
+	GAMEPLAYCAMERAS_API void EndDrawing();
+
 	/** Adds text to the text wall. */
 	GAMEPLAYCAMERAS_API void AddText(const FString& InString);
 	GAMEPLAYCAMERAS_API void AddText(const TCHAR* Fmt, ...);
@@ -107,6 +113,8 @@ public:
 	GAMEPLAYCAMERAS_API void DrawPoint(const FVector3d& Location, float PointSize, const FLinearColor& LineColor, float LineThickness = 1.f);
 	/** Draws a 3D line. */
 	GAMEPLAYCAMERAS_API void DrawLine(const FVector3d& Start, const FVector3d& End, const FLinearColor& LineColor, float LineThickness = 1.f);
+	/** Draws a 3D box. */
+	GAMEPLAYCAMERAS_API void DrawBox(const FVector3d& Center, const FVector3d& Size, const FLinearColor& LineColor, float LineThickness = 1.f);
 	/** Draws a 3D sphere. */
 	GAMEPLAYCAMERAS_API void DrawSphere(const FVector3d& Center, float Radius, int32 Segments, const FLinearColor& LineColor, float LineThickness = 1.f);
 	/** Draws a 3D arrow. */
@@ -138,13 +146,13 @@ public:
 public:
 
 	/** Gets the drawing canvas. */
-	FCanvas* GetCanvas() const;
+	FCanvas* GetCanvas() const { return Canvas; }
 
 	/** Gets the size of the canvas. */
-	FVector2D GetCanvasSize() const;
+	FVector2D GetCanvasSize() const { return CanvasSize; }
 
 	/** Returns whether this renderer has a valid canvas to draw upon. */
-	bool HasCanvas() const { return CanvasObject != nullptr; }
+	bool HasCanvas() const { return Canvas != nullptr; }
 
 public:
 
@@ -153,6 +161,8 @@ public:
 
 private:
 
+	void Initialize(UWorld* InWorld, bool bInIsExternalRendering);
+
 	void AddTextFmtImpl(const TCHAR* Fmt, va_list Args);
 	void AddTextImpl(const TCHAR* Buffer);
 
@@ -167,9 +177,13 @@ private:
 private:
 
 	/** The world in which we might draw debug primitives. */
-	UWorld* World;
+	UWorld* World = nullptr;
 	/** The canvas used to draw the text wall. */
-	UCanvas* CanvasObject;
+	FCanvas* Canvas = nullptr;
+
+	const FSceneView* SceneView = nullptr;
+	/** The size of the canvas. */
+	FVector2D CanvasSize;
 	/** Whether we are looking from the "outside" of the camera system. */
 	bool bIsExternalRendering = false;
 
