Cameras: move viewfinder rendering code to a helper class for future reuse

diff --git a/Engine/Plugins/Cameras/GameplayCameras/Source/GameplayCameras/Private/Debug/ViewfinderDebugBlock.cpp b/Engine/Plugins/Cameras/GameplayCameras/Source/GameplayCameras/Private/Debug/ViewfinderDebugBlock.cpp
--- a/Engine/Plugins/Cameras/GameplayCameras/Source/GameplayCameras/Private/Debug/ViewfinderDebugBlock.cpp
+++ b/Engine/Plugins/Cameras/GameplayCameras/Source/GameplayCameras/Private/Debug/ViewfinderDebugBlock.cpp
@@ -2,71 +2,14 @@
 
 #include "Debug/ViewfinderDebugBlock.h"
 
-#include "CanvasItem.h"
-#include "CanvasTypes.h"
 #include "Debug/CameraDebugCategories.h"
-#include "Debug/CameraDebugColors.h"
-#include "Debug/CameraDebugRenderer.h"
-#include "DrawDebugHelpers.h"
-#include "HAL/IConsoleManager.h"
+#include "Debug/ViewfinderRenderer.h"
 
 #if UE_GAMEPLAY_CAMERAS_DEBUG
 
 namespace UE::Cameras
 {
 
-float GGameplayCamerasDebugViewfinderReticleSizeFactor = 0.27f;
-static FAutoConsoleVariableRef CVarGameplayCamerasDebugViewfinderReticleSizeFactor(
-	TEXT("GameplayCameras.Debug.Viewfinder.ReticleSizeFactor"),
-	GGameplayCamerasDebugViewfinderReticleSizeFactor,
-	TEXT("Default: 0.1. The size of the viewfinder reticle, as a factor of the screen's vertical size."));
-
-float GGameplayCamerasDebugViewfinderReticleInnerSizeFactor = 0.7f;
-static FAutoConsoleVariableRef CVarGameplayCamerasDebugViewfinderReticleInnerSizeFactor(
-	TEXT("GameplayCameras.Debug.Viewfinder.ReticleInnerSizeFactor"),
-	GGameplayCamerasDebugViewfinderReticleInnerSizeFactor,
-	TEXT(""));
-
-int32 GGameplayCamerasDebugViewfinderReticleNumSides = 60;
-static FAutoConsoleVariableRef CVarGameplayCamerasDebugViewfinderReticleNumSides(
-	TEXT("GameplayCameras.Debug.Viewfinder.ReticleNumSides"),
-	GGameplayCamerasDebugViewfinderReticleNumSides,
-	TEXT(""));
-
-float GGameplayCamerasDebugViewfinderGuidesGapFactor = 0.02f;
-static FAutoConsoleVariableRef CVarGameplayCamerasDebugViewfinderGuidesGapFactor(
-	TEXT("GameplayCameras.Debug.Viewfinder.GuidesGapFactor"),
-	GGameplayCamerasDebugViewfinderGuidesGapFactor,
-	TEXT(""));
-
-namespace Private
-{
-
-void DrawCanvasLine(FCanvas* Canvas, const FVector2D& Start, const FVector2D& End, const FLinearColor& LineColor, float LineThickness = 1.f)
-{
-	FCanvasLineItem LineItem(Start, End);
-	LineItem.SetColor(LineColor);
-	LineItem.LineThickness = LineThickness;
-	LineItem.Draw(Canvas);
-}
-
-void DrawCanvasCircle(FCanvas* Canvas, const FVector2D& Center, float Radius, int32 NumSides, const FLinearColor& LineColor, float LineThickness = 1.f)
-{
-	const float	AngleDelta = 2.0f * UE_PI / NumSides;
-	FVector2D AxisX(1.f, 0.f);
-	FVector2D AxisY(0.f, -1.f);
-	FVector2D LastVertex = Center + AxisX * Radius;
-
-	for (int32 SideIndex = 0; SideIndex < NumSides; SideIndex++)
-	{
-		const FVector2D Vertex = Center + (AxisX * FMath::Cos(AngleDelta * (SideIndex + 1)) + AxisY * FMath::Sin(AngleDelta * (SideIndex + 1))) * Radius;
-		DrawCanvasLine(Canvas, LastVertex, Vertex, LineColor, LineThickness);
-		LastVertex = Vertex;
-	}
-}
-
-}  //  namespace Private 
-
 UE_DEFINE_CAMERA_DEBUG_BLOCK(FViewfinderDebugBlock)
 
 FViewfinderDebugBlock::FViewfinderDebugBlock()
@@ -75,78 +18,10 @@ FViewfinderDebugBlock::FViewfinderDebugBlock()
 
 void FViewfinderDebugBlock::OnDebugDraw(const FCameraDebugBlockDrawParams& Params, FCameraDebugRenderer& Renderer)
 {
-	if (!Params.IsCategoryActive(FCameraDebugCategories::Viewfinder))
+	if (Params.IsCategoryActive(FCameraDebugCategories::Viewfinder))
 	{
-		return;
+		FViewfinderRenderer::DrawViewfinder(Renderer, FViewfinderDrawElements::All);
 	}
-
-	FCanvas* Canvas = Renderer.GetCanvas();
-
-	const FVector2D CanvasSize = Renderer.GetCanvasSize();
-	const float CanvasSizeX = CanvasSize.X;
-	const float CanvasSizeY = CanvasSize.Y;
-
-	const FVector2d CanvasCenter(CanvasSizeX / 2.f, CanvasSizeY / 2.f);
-
-	// Draw the reticle.
-	const float ReticleRadius = CanvasSizeY * GGameplayCamerasDebugViewfinderReticleSizeFactor / 2.f;
-	const float ReticleInnerRadiusFactor = GGameplayCamerasDebugViewfinderReticleInnerSizeFactor;
-	const int32 ReticleNumSides = GGameplayCamerasDebugViewfinderReticleNumSides;
-	const FColor ReticleColor = FCameraDebugColors::Get().Passive;
-
-	// ...outer reticle circle.
-	Private::DrawCanvasCircle(Canvas, CanvasCenter, ReticleRadius, ReticleNumSides, ReticleColor);
-	// ...inner reticle circle.
-	const float ReticleInnerRadius = ReticleRadius * ReticleInnerRadiusFactor;
-	const int32 ReticleInnerNumSides = (int32)(ReticleNumSides * ReticleInnerRadiusFactor);
-	Private::DrawCanvasCircle(Canvas, CanvasCenter, ReticleInnerRadius, ReticleInnerNumSides, ReticleColor);
-	// ...horizontal line inside reticle.
-	Private::DrawCanvasLine(
-			Canvas, 
-			CanvasCenter - FVector2D(ReticleInnerRadius, 0), CanvasCenter + FVector2D(ReticleInnerRadius, 0), 
-			ReticleColor);
-
-	// Draw the rule-of-thirds guides.
-	const FColor GuideColor = FCameraDebugColors::Get().VeryPassive;
-	const float GuidesGap = CanvasSizeY * GGameplayCamerasDebugViewfinderGuidesGapFactor;
-	const FVector2D OneThird(CanvasSizeX / 3.f, CanvasSizeY / 3.f);
-	const FVector2D TwoThirds(CanvasSizeX / 1.5f, CanvasSizeY / 1.5f);
-	// ...top vertical guides
-	Private::DrawCanvasLine(
-			Canvas,
-			FVector2D(OneThird.X, 0), FVector2D(OneThird.X, OneThird.Y - GuidesGap),
-			ReticleColor, 2.f);
-	Private::DrawCanvasLine(
-			Canvas,
-			FVector2D(TwoThirds.X, 0), FVector2D(TwoThirds.X, OneThird.Y - GuidesGap),
-			ReticleColor, 2.f);
-	// ...bottom vertical guides
-	Private::DrawCanvasLine(
-			Canvas,
-			FVector2D(OneThird.X, TwoThirds.Y + GuidesGap), FVector2D(OneThird.X, CanvasSizeY),
-			ReticleColor, 2.f);
-	Private::DrawCanvasLine(
-			Canvas,
-			FVector2D(TwoThirds.X, TwoThirds.Y + GuidesGap), FVector2D(TwoThirds.X, CanvasSizeY),
-			ReticleColor, 2.f);
-	// ...left horizontal guides
-	Private::DrawCanvasLine(
-			Canvas,
-			FVector2D(0, OneThird.Y), FVector2D(OneThird.X - GuidesGap, OneThird.Y),
-			ReticleColor, 2.f);
-	Private::DrawCanvasLine(
-			Canvas,
-			FVector2D(0, TwoThirds.Y), FVector2D(OneThird.X - GuidesGap, TwoThirds.Y),
-			ReticleColor, 2.f);
-	// ...right horizontal guides
-	Private::DrawCanvasLine(
-			Canvas,
-			FVector2D(TwoThirds.X + GuidesGap, OneThird.Y), FVector2D(CanvasSizeX, OneThird.Y),
-			ReticleColor, 2.f);
-	Private::DrawCanvasLine(
-			Canvas,
-			FVector2D(TwoThirds.X + GuidesGap, TwoThirds.Y), FVector2D(CanvasSizeX, TwoThirds.Y),
-			ReticleColor, 2.f);
 }
 
 }  // namespace UE::Cameras
diff --git a/Engine/Plugins/Cameras/GameplayCameras/Source/GameplayCameras/Private/Debug/ViewfinderRenderer.cpp b/Engine/Plugins/Cameras/GameplayCameras/Source/GameplayCameras/Private/Debug/ViewfinderRenderer.cpp
new file mode 100644
--- /dev/null
+++ b/Engine/Plugins/Cameras/GameplayCameras/Source/GameplayCameras/Private/Debug/ViewfinderRenderer.cpp
@@ -0,0 +1,113 @@
+// Copyright Epic Games, Inc. All Rights Reserved.
+
+#include "Debug/ViewfinderRenderer.h"
+
+#include "CanvasItem.h"
+#include "CanvasTypes.h"
+#include "Debug/CameraDebugRenderer.h"
+#include "HAL/IConsoleManager.h"
+
+#if UE_GAMEPLAY_CAMERAS_DEBUG
+
+namespace UE::Cameras
+{
+
+float GGameplayCamerasDebugViewfinderReticleSizeFactor = 0.27f;
+static FAutoConsoleVariableRef CVarGameplayCamerasDebugViewfinderReticleSizeFactor(
+	TEXT("GameplayCameras.Debug.Viewfinder.ReticleSizeFactor"),
+	GGameplayCamerasDebugViewfinderReticleSizeFactor,
+	TEXT("Default: 0.1. The size of the viewfinder reticle, as a factor of the screen's vertical size."));
+
+float GGameplayCamerasDebugViewfinderReticleInnerSizeFactor = 0.7f;
+static FAutoConsoleVariableRef CVarGameplayCamerasDebugViewfinderReticleInnerSizeFactor(
+	TEXT("GameplayCameras.Debug.Viewfinder.ReticleInnerSizeFactor"),
+	GGameplayCamerasDebugViewfinderReticleInnerSizeFactor,
+	TEXT(""));
+
+int32 GGameplayCamerasDebugViewfinderReticleNumSides = 60;
+static FAutoConsoleVariableRef CVarGameplayCamerasDebugViewfinderReticleNumSides(
+	TEXT("GameplayCameras.Debug.Viewfinder.ReticleNumSides"),
+	GGameplayCamerasDebugViewfinderReticleNumSides,
+	TEXT(""));
+
+float GGameplayCamerasDebugViewfinderGuidesGapFactor = 0.02f;
+static FAutoConsoleVariableRef CVarGameplayCamerasDebugViewfinderGuidesGapFactor(
+	TEXT("GameplayCameras.Debug.Viewfinder.GuidesGapFactor"),
+	GGameplayCamerasDebugViewfinderGuidesGapFactor,
+	TEXT(""));
+
+void FViewfinderRenderer::DrawViewfinder(FCameraDebugRenderer& Renderer, FViewfinderDrawElements Elements)
+{
+	FCanvas* Canvas = Renderer.GetCanvas();
+
+	const FVector2D CanvasSize = Renderer.GetCanvasSize();
+	const float CanvasSizeX = CanvasSize.X;
+	const float CanvasSizeY = CanvasSize.Y;
+
+	const FVector2d CanvasCenter(CanvasSizeX / 2.f, CanvasSizeY / 2.f);
+
+	if (EnumHasAnyFlags(Elements, FViewfinderDrawElements::FocusReticle))
+	{
+		// Draw the reticle.
+		const float ReticleRadius = CanvasSizeY * GGameplayCamerasDebugViewfinderReticleSizeFactor / 2.f;
+		const float ReticleInnerRadiusFactor = GGameplayCamerasDebugViewfinderReticleInnerSizeFactor;
+		const int32 ReticleNumSides = GGameplayCamerasDebugViewfinderReticleNumSides;
+		const FColor ReticleColor = FCameraDebugColors::Get().Passive;
+
+		// ...outer reticle circle.
+		Renderer.Draw2DCircle(CanvasCenter, ReticleRadius, ReticleColor, 1.f, ReticleNumSides);
+		// ...inner reticle circle.
+		const float ReticleInnerRadius = ReticleRadius * ReticleInnerRadiusFactor;
+		const int32 ReticleInnerNumSides = (int32)(ReticleNumSides * ReticleInnerRadiusFactor);
+		Renderer.Draw2DCircle(CanvasCenter, ReticleInnerRadius, ReticleColor, 1.f, ReticleInnerNumSides);
+		// ...horizontal line inside reticle.
+		Renderer.Draw2DLine(
+				CanvasCenter - FVector2D(ReticleInnerRadius, 0), CanvasCenter + FVector2D(ReticleInnerRadius, 0), 
+				ReticleColor);
+	}
+
+	if (EnumHasAnyFlags(Elements, FViewfinderDrawElements::RuleOfThirds))
+	{
+		// Draw the rule-of-thirds guides.
+		const FColor GuideColor = FCameraDebugColors::Get().VeryPassive;
+		const float GuidesGap = CanvasSizeY * GGameplayCamerasDebugViewfinderGuidesGapFactor;
+		const FVector2D OneThird(CanvasSizeX / 3.f, CanvasSizeY / 3.f);
+		const FVector2D TwoThirds(CanvasSizeX / 1.5f, CanvasSizeY / 1.5f);
+		const FColor LineColor = FCameraDebugColors::Get().Passive;
+
+		// ...top vertical guides
+		Renderer.Draw2DLine(
+				FVector2D(OneThird.X, 0), FVector2D(OneThird.X, OneThird.Y - GuidesGap),
+				LineColor, 2.f);
+		Renderer.Draw2DLine(
+				FVector2D(TwoThirds.X, 0), FVector2D(TwoThirds.X, OneThird.Y - GuidesGap),
+				LineColor, 2.f);
+		// ...bottom vertical guides
+		Renderer.Draw2DLine(
+				FVector2D(OneThird.X, TwoThirds.Y + GuidesGap), FVector2D(OneThird.X, CanvasSizeY),
+				LineColor, 2.f);
+		Renderer.Draw2DLine(
+				FVector2D(TwoThirds.X, TwoThirds.Y + GuidesGap), FVector2D(TwoThirds.X, CanvasSizeY),
+				LineColor, 2.f);
+		// ...left horizontal guides
+		Renderer.Draw2DLine(
+				FVector2D(0, OneThird.Y), FVector2D(OneThird.X - GuidesGap, OneThird.Y),
+				LineColor, 2.f);
+		Renderer.Draw2DLine(
+				FVector2D(0, TwoThirds.Y), FVector2D(OneThird.X - GuidesGap, TwoThirds.Y),
+				LineColor, 2.f);
+		// ...right horizontal guides
+		Renderer.Draw2DLine(
+				FVector2D(TwoThirds.X + GuidesGap, OneThird.Y), FVector2D(CanvasSizeX, OneThird.Y),
+				LineColor, 2.f);
+		Renderer.Draw2DLine(
+				FVector2D(TwoThirds.X + GuidesGap, TwoThirds.Y), FVector2D(CanvasSizeX, TwoThirds.Y),
+				LineColor, 2.f);
+	}
+}
+
+}  // namespace UE::Cameras
+
+#endif  // UE_GAMEPLAY_CAMERAS_DEBUG
+
+
diff --git a/Engine/Plugins/Cameras/GameplayCameras/Source/GameplayCameras/Private/Debug/ViewfinderRenderer.h b/Engine/Plugins/Cameras/GameplayCameras/Source/GameplayCameras/Private/Debug/ViewfinderRenderer.h
new file mode 100644
--- /dev/null
+++ b/Engine/Plugins/Cameras/GameplayCameras/Source/GameplayCameras/Private/Debug/ViewfinderRenderer.h
@@ -0,0 +1,36 @@
+// Copyright Epic Games, Inc. All Rights Reserved.
+
+#pragma once
+
+#include "GameplayCameras.h"
+#include "Misc/EnumClassFlags.h"
+
+#if UE_GAMEPLAY_CAMERAS_DEBUG
+
+namespace UE::Cameras
+{
+
+class FCameraDebugRenderer;
+
+enum class FViewfinderDrawElements
+{
+	None = 0,
+	RuleOfThirds = 1 << 0,
+	FocusReticle = 1 << 1,
+
+	All = RuleOfThirds | FocusReticle
+};
+
+ENUM_CLASS_FLAGS(FViewfinderDrawElements)
+
+class FViewfinderRenderer
+{
+public:
+
+	static void DrawViewfinder(FCameraDebugRenderer& Renderer, FViewfinderDrawElements Elements);
+};
+
+}  // namespace UE::Cameras
+
+#endif  // UE_GAMEPLAY_CAMERAS_DEBUG
+
