Cameras: add APIs for getting mutable pointers into the variable table and context data table

diff --git a/Engine/Plugins/Cameras/GameplayCameras/Source/GameplayCameras/Private/Core/CameraContextDataTable.cpp b/Engine/Plugins/Cameras/GameplayCameras/Source/GameplayCameras/Private/Core/CameraContextDataTable.cpp
--- a/Engine/Plugins/Cameras/GameplayCameras/Source/GameplayCameras/Private/Core/CameraContextDataTable.cpp
+++ b/Engine/Plugins/Cameras/GameplayCameras/Source/GameplayCameras/Private/Core/CameraContextDataTable.cpp
@@ -413,7 +413,10 @@ FCameraContextDataTable::FEntry* FCameraContextDataTable::FindEntry(FCameraConte
 	return nullptr;
 }
 
-const uint8* FCameraContextDataTable::GetData(FCameraContextDataID DataID, ECameraContextDataType ExpectedDataType, const UObject* ExpectedDataTypeObject) const
+const uint8* FCameraContextDataTable::GetData(
+		FCameraContextDataID DataID,
+		ECameraContextDataType ExpectedDataType,
+		const UObject* ExpectedDataTypeObject) const
 {
 	const uint8* Data = TryGetData(DataID, ExpectedDataType, ExpectedDataTypeObject);
 	ensureMsgf(
@@ -423,7 +426,10 @@ const uint8* FCameraContextDataTable::GetData(FCameraContextDataID DataID, ECame
 	return Data;
 }
 
-const uint8* FCameraContextDataTable::TryGetData(FCameraContextDataID DataID, ECameraContextDataType ExpectedDataType, const UObject* ExpectedDataTypeObject) const
+const uint8* FCameraContextDataTable::TryGetData(
+		FCameraContextDataID DataID,
+		ECameraContextDataType ExpectedDataType,
+		const UObject* ExpectedDataTypeObject) const
 {
 	const FEntry* Entry = FindEntry(DataID);
 	if (Entry)
@@ -437,13 +443,40 @@ const uint8* FCameraContextDataTable::TryGetData(FCameraContextDataID DataID, EC
 	return nullptr;
 }
 
-void FCameraContextDataTable::SetData(FCameraContextDataID DataID, ECameraContextDataType ExpectedDataType, const UObject* ExpectedDataTypeObject, const uint8* InRawDataPtr, bool bMarkAsWrittenThisFrame)
+uint8* FCameraContextDataTable::TryGetMutableData(
+		FCameraContextDataID DataID,
+		ECameraContextDataType ExpectedDataType,
+		const UObject* ExpectedDataTypeObject)
+{
+	FEntry* Entry = FindEntry(DataID);
+	if (Entry)
+	{
+		if (Entry->Type == ExpectedDataType && Entry->TypeObject == ExpectedDataTypeObject)
+		{
+			return Memory + Entry->Offset;
+		}
+	}
+
+	return nullptr;
+}
+
+void FCameraContextDataTable::SetData(
+		FCameraContextDataID DataID,
+		ECameraContextDataType ExpectedDataType,
+		const UObject* ExpectedDataTypeObject,
+		const uint8* InRawDataPtr,
+		bool bMarkAsWrittenThisFrame)
 {
 	const bool bDidSet = TrySetData(DataID, ExpectedDataType, ExpectedDataTypeObject, InRawDataPtr, bMarkAsWrittenThisFrame);
 	ensureMsgf(bDidSet, TEXT("Can't set camera context data (ID '%d') beacuse it doesn't exist in the table."), DataID.GetValue());
 }
 
-bool FCameraContextDataTable::TrySetData(FCameraContextDataID DataID, ECameraContextDataType ExpectedDataType, const UObject* ExpectedDataTypeObject, const uint8* InRawDataPtr, bool bMarkAsWrittenThisFrame)
+bool FCameraContextDataTable::TrySetData(
+		FCameraContextDataID DataID,
+		ECameraContextDataType ExpectedDataType,
+		const UObject* ExpectedDataTypeObject,
+		const uint8* InRawDataPtr,
+		bool bMarkAsWrittenThisFrame)
 {
 	FEntry* Entry = FindEntry(DataID);
 	if (!Entry)
diff --git a/Engine/Plugins/Cameras/GameplayCameras/Source/GameplayCameras/Private/Core/CameraVariableTable.cpp b/Engine/Plugins/Cameras/GameplayCameras/Source/GameplayCameras/Private/Core/CameraVariableTable.cpp
--- a/Engine/Plugins/Cameras/GameplayCameras/Source/GameplayCameras/Private/Core/CameraVariableTable.cpp
+++ b/Engine/Plugins/Cameras/GameplayCameras/Source/GameplayCameras/Private/Core/CameraVariableTable.cpp
@@ -306,6 +306,21 @@ const uint8* FCameraVariableTable::TryGetValue(
 	return nullptr;
 }
 
+uint8* FCameraVariableTable::TryGetMutableValue(
+		FCameraVariableID VariableID,
+		ECameraVariableType ExpectedVariableType,
+		const UScriptStruct* ExpectedBlendableStructType)
+{
+	FEntry* Entry = FindEntry(VariableID);
+	if (Entry)
+	{
+		ensure(Entry->Type == ExpectedVariableType && Entry->StructType == ExpectedBlendableStructType);
+		return Memory + Entry->Offset;
+	}
+
+	return nullptr;
+}
+
 void FCameraVariableTable::SetValue(
 		FCameraVariableID VariableID, 
 		ECameraVariableType ExpectedVariableType, 
diff --git a/Engine/Plugins/Cameras/GameplayCameras/Source/GameplayCameras/Public/Core/CameraContextDataTable.h b/Engine/Plugins/Cameras/GameplayCameras/Source/GameplayCameras/Public/Core/CameraContextDataTable.h
--- a/Engine/Plugins/Cameras/GameplayCameras/Source/GameplayCameras/Public/Core/CameraContextDataTable.h
+++ b/Engine/Plugins/Cameras/GameplayCameras/Source/GameplayCameras/Public/Core/CameraContextDataTable.h
@@ -123,11 +123,34 @@ public:
 public:
 
 	// Low-level API.
-	const uint8* GetData(FCameraContextDataID DataID, ECameraContextDataType ExpectedDataType, const UObject* ExpectedDataTypeObject) const;
-	const uint8* TryGetData(FCameraContextDataID DataID, ECameraContextDataType ExpectedDataType, const UObject* ExpectedDataTypeObject) const;
-
-	void SetData(FCameraContextDataID DataID, ECameraContextDataType ExpectedDataType, const UObject* ExpectedDataTypeObject, const uint8* InRawDataPtr, bool bMarkAsWrittenThisFrame = true);
-	bool TrySetData(FCameraContextDataID DataID, ECameraContextDataType ExpectedDataType, const UObject* ExpectedDataTypeObject, const uint8* InRawDataPtr, bool bMarkAsWrittenThisFrame = true);
+	const uint8* GetData(
+			FCameraContextDataID DataID,
+			ECameraContextDataType ExpectedDataType,
+			const UObject* ExpectedDataTypeObject) const;
+
+	const uint8* TryGetData(
+			FCameraContextDataID DataID,
+			ECameraContextDataType ExpectedDataType,
+			const UObject* ExpectedDataTypeObject) const;
+
+	uint8* TryGetMutableData(
+			FCameraContextDataID DataID,
+			ECameraContextDataType ExpectedDataType,
+			const UObject* ExpectedDataTypeObject);
+
+	void SetData(
+			FCameraContextDataID DataID,
+			ECameraContextDataType ExpectedDataType,
+			const UObject* ExpectedDataTypeObject,
+			const uint8* InRawDataPtr,
+			bool bMarkAsWrittenThisFrame = true);
+
+	bool TrySetData(
+			FCameraContextDataID DataID,
+			ECameraContextDataType ExpectedDataType,
+			const UObject* ExpectedDataTypeObject,
+			const uint8* InRawDataPtr,
+			bool bMarkAsWrittenThisFrame = true);
 
 	bool IsValueWritten(FCameraContextDataID InID) const;
 	void UnsetValue(FCameraContextDataID InID);
diff --git a/Engine/Plugins/Cameras/GameplayCameras/Source/GameplayCameras/Public/Core/CameraVariableTable.h b/Engine/Plugins/Cameras/GameplayCameras/Source/GameplayCameras/Public/Core/CameraVariableTable.h
--- a/Engine/Plugins/Cameras/GameplayCameras/Source/GameplayCameras/Public/Core/CameraVariableTable.h
+++ b/Engine/Plugins/Cameras/GameplayCameras/Source/GameplayCameras/Public/Core/CameraVariableTable.h
@@ -163,6 +163,11 @@ public:
 			ECameraVariableType ExpectedVariableType,
 			const UScriptStruct* ExpectedBlendableStructType) const;
 
+	uint8* TryGetMutableValue(
+			FCameraVariableID VariableID,
+			ECameraVariableType ExpectedVariableType,
+			const UScriptStruct* ExpectedBlendableStructType);
+
 	void SetValue(
 			FCameraVariableID VariableID, 
 			ECameraVariableType ExpectedVariableType, 
