Cameras: run camera rigs in editor for providing a preview to the user

Give the GameplayCameraComponent a child CineCameraComponent. This child acts as the "output" of the camera system, although in the future we may make it only receive the "output" of any running rig(s) that use the parent GameplayCameraComponent as context. This is needed for interoperability with systems that only support a camera component, such as Sequencer.

We don't need the camera-looking mesh preview anymore so this change removes it.

Add a camera system evaluator to the gameplay camera component for running in editor. There's an option to enable/disable it.

Add a component visualizer for displaying a few things in editor.

diff --git a/Engine/Plugins/Cameras/GameplayCameras/Source/GameplayCameras/GameplayCameras.Build.cs b/Engine/Plugins/Cameras/GameplayCameras/Source/GameplayCameras/GameplayCameras.Build.cs
--- a/Engine/Plugins/Cameras/GameplayCameras/Source/GameplayCameras/GameplayCameras.Build.cs
+++ b/Engine/Plugins/Cameras/GameplayCameras/Source/GameplayCameras/GameplayCameras.Build.cs
@@ -24,5 +24,13 @@ public class GameplayCameras : ModuleRules
 				"TraceLog"
 			}
 		);
+
+		if (Target.bBuildEditor)
+		{
+			PrivateDependencyModuleNames.AddRange(new string[]
+				{
+					"UnrealEd"
+				});
+		}
 	}
 }
diff --git a/Engine/Plugins/Cameras/GameplayCameras/Source/GameplayCameras/Private/Core/CameraNodeEvaluator.cpp b/Engine/Plugins/Cameras/GameplayCameras/Source/GameplayCameras/Private/Core/CameraNodeEvaluator.cpp
--- a/Engine/Plugins/Cameras/GameplayCameras/Source/GameplayCameras/Private/Core/CameraNodeEvaluator.cpp
+++ b/Engine/Plugins/Cameras/GameplayCameras/Source/GameplayCameras/Private/Core/CameraNodeEvaluator.cpp
@@ -177,6 +177,24 @@ void FCameraNodeEvaluator::Serialize(const FCameraNodeEvaluatorSerializeParams&
 	OnSerialize(Params, Ar);
 }
 
+#if WITH_EDITOR
+
+void FCameraNodeEvaluator::DrawEditorPreview(const FCameraEditorPreviewDrawParams& Params, FCameraDebugRenderer& Renderer)
+{
+	OnDrawEditorPreview(Params, Renderer);
+
+	FCameraNodeEvaluatorChildrenView ChildrenView(GetChildren());
+	for (FCameraNodeEvaluator* Child : ChildrenView)
+	{
+		if (Child)
+		{
+			Child->DrawEditorPreview(Params, Renderer);
+		}
+	}
+}
+
+#endif  // WITH_EDITOR
+
 #if UE_GAMEPLAY_CAMERAS_DEBUG
 
 void FCameraNodeEvaluator::BuildDebugBlocks(const FCameraDebugBlockBuildParams& Params, FCameraDebugBlockBuilder& Builder)
diff --git a/Engine/Plugins/Cameras/GameplayCameras/Source/GameplayCameras/Private/Core/CameraSystemEvaluator.cpp b/Engine/Plugins/Cameras/GameplayCameras/Source/GameplayCameras/Private/Core/CameraSystemEvaluator.cpp
--- a/Engine/Plugins/Cameras/GameplayCameras/Source/GameplayCameras/Private/Core/CameraSystemEvaluator.cpp
+++ b/Engine/Plugins/Cameras/GameplayCameras/Source/GameplayCameras/Private/Core/CameraSystemEvaluator.cpp
@@ -2,7 +2,6 @@
 
 #include "Core/CameraSystemEvaluator.h"
 
-#include "Algo/Transform.h"
 #include "Camera/CameraTypes.h"
 #include "Core/CameraDirectorEvaluator.h"
 #include "Core/CameraEvaluationContext.h"
@@ -15,6 +14,7 @@
 #include "Debug/CameraSystemTrace.h"
 #include "Debug/RootCameraDebugBlock.h"
 #include "GameplayCamerasSettings.h"
+#include "Math/ColorList.h"
 #include "Services/OrientationInitializationService.h"
 #include "UObject/Package.h"
 #include "UObject/UObjectGlobals.h"
@@ -192,6 +192,11 @@ void FCameraSystemEvaluator::NotifyRootCameraNodeEvent(const FRootCameraNodeCame
 }
 
 void FCameraSystemEvaluator::Update(const FCameraSystemEvaluationParams& Params)
+{
+	UpdateImpl(Params.DeltaTime, ECameraNodeEvaluationType::Standard);
+}
+
+void FCameraSystemEvaluator::UpdateImpl(float DeltaTime, ECameraNodeEvaluationType EvaluationType)
 {
 	SCOPE_CYCLE_COUNTER(CameraSystemEval_Total);
 
@@ -204,7 +209,7 @@ void FCameraSystemEvaluator::Update(const FCameraSystemEvaluationParams& Params)
 	RootNodeResult.VariableTable.AutoResetValues();
 
 	// Pre-update all services.
-	PreUpdateServices(Params.DeltaTime, ECameraEvaluationServiceFlags::None);
+	PreUpdateServices(DeltaTime, ECameraEvaluationServiceFlags::None);
 
 	// Get the active evaluation context.
 	TSharedPtr<FCameraEvaluationContext> ActiveContext = ContextStack.GetActiveContext();
@@ -219,7 +224,7 @@ void FCameraSystemEvaluator::Update(const FCameraSystemEvaluationParams& Params)
 	if (ActiveDirectorEvaluator)
 	{
 		FCameraDirectorEvaluationParams DirectorParams;
-		DirectorParams.DeltaTime = Params.DeltaTime;
+		DirectorParams.DeltaTime = DeltaTime;
 		DirectorParams.OwnerContext = ActiveContext;
 
 		FCameraDirectorEvaluationResult DirectorResult;
@@ -277,7 +282,8 @@ void FCameraSystemEvaluator::Update(const FCameraSystemEvaluationParams& Params)
 		// Setup the params/result for running the root camera node.
 		FCameraNodeEvaluationParams NodeParams;
 		NodeParams.Evaluator = this;
-		NodeParams.DeltaTime = Params.DeltaTime;
+		NodeParams.DeltaTime = DeltaTime;
+		NodeParams.EvaluationType = EvaluationType;
 
 		RootNodeResult.Reset();
 
@@ -288,7 +294,7 @@ void FCameraSystemEvaluator::Update(const FCameraSystemEvaluationParams& Params)
 	}
 
 	// Post-update all services.
-	PostUpdateServices(Params.DeltaTime, ECameraEvaluationServiceFlags::None);
+	PostUpdateServices(DeltaTime, ECameraEvaluationServiceFlags::None);
 
 	// Harvest the result.
 	Result.Reset(RootNodeResult);
@@ -359,6 +365,30 @@ void FCameraSystemEvaluator::GetEvaluatedCameraView(FMinimalViewInfo& DesiredVie
 	CameraPose.ApplyPhysicalCameraSettings(DesiredView.PostProcessSettings, false);
 }
 
+#if WITH_EDITOR
+
+void FCameraSystemEvaluator::EditorPreviewUpdate(const FCameraSystemEvaluationParams& Params)
+{
+	UpdateImpl(Params.DeltaTime, ECameraNodeEvaluationType::EditorPreview);
+}
+
+void FCameraSystemEvaluator::DrawEditorPreview(const FCameraSystemEditorPreviewParams& Params)
+{	
+	FCameraEditorPreviewDrawParams NodeParams;
+
+	UObject* Owner = WeakOwner.Get();
+	UWorld* OwnerWorld = (Owner && Params.bDrawWorldDebug) ? Owner->GetWorld() : nullptr;
+	FCameraDebugRenderer Renderer(OwnerWorld, Params.SceneView, Params.Canvas, !Params.bIsLockedToCamera);
+
+	Renderer.BeginDrawing();
+	{
+		RootEvaluator->DrawEditorPreview(NodeParams, Renderer);
+	}
+	Renderer.EndDrawing();
+}
+
+#endif  // WITH_EDITOR
+
 #if UE_GAMEPLAY_CAMERAS_DEBUG
 
 bool FCameraSystemEvaluator::IsDebugTraceEnabled()
@@ -410,7 +440,7 @@ void FCameraSystemEvaluator::DebugUpdate(const FCameraSystemDebugUpdateParams& P
 		FCameraSystemTrace::TraceEvaluation(OwnerWorld, Result, *RootDebugBlock);
 	}
 #endif
-	
+
 	FCameraDebugRenderer Renderer(OwnerWorld, Params.CanvasObject, Params.bIsDebugCameraEnabled);
 	RootDebugBlock->RootDebugDraw(Renderer);
 }
diff --git a/Engine/Plugins/Cameras/GameplayCameras/Source/GameplayCameras/Private/Core/RootCameraNode.cpp b/Engine/Plugins/Cameras/GameplayCameras/Source/GameplayCameras/Private/Core/RootCameraNode.cpp
--- a/Engine/Plugins/Cameras/GameplayCameras/Source/GameplayCameras/Private/Core/RootCameraNode.cpp
+++ b/Engine/Plugins/Cameras/GameplayCameras/Source/GameplayCameras/Private/Core/RootCameraNode.cpp
@@ -31,6 +31,13 @@ void FRootCameraNodeEvaluator::GetActiveCameraRigInfo(FCameraRigEvaluationInfo&
 	OnGetActiveCameraRigInfo(OutCameraRigInfo);
 }
 
+bool FRootCameraNodeEvaluator::HasAnyActiveCameraRig() const
+{
+	FCameraRigEvaluationInfo RigInfo;
+	GetActiveCameraRigInfo(RigInfo);
+	return RigInfo.RootEvaluator != nullptr;
+}
+
 void FRootCameraNodeEvaluator::BuildSingleCameraRigHierarchy(const FSingleCameraRigHierarchyBuildParams& Params, FCameraNodeEvaluatorHierarchy& OutHierarchy)
 {
 	OnBuildSingleCameraRigHierarchy(Params, OutHierarchy);
diff --git a/Engine/Plugins/Cameras/GameplayCameras/Source/GameplayCameras/Private/Debug/CameraDebugRenderer.cpp b/Engine/Plugins/Cameras/GameplayCameras/Source/GameplayCameras/Private/Debug/CameraDebugRenderer.cpp
--- a/Engine/Plugins/Cameras/GameplayCameras/Source/GameplayCameras/Private/Debug/CameraDebugRenderer.cpp
+++ b/Engine/Plugins/Cameras/GameplayCameras/Source/GameplayCameras/Private/Debug/CameraDebugRenderer.cpp
@@ -114,36 +114,23 @@ static FAutoConsoleVariableRef CVarGameplayCamerasDebugColorScheme(
 
 FCameraDebugRenderer::FCameraDebugRenderer(UWorld* InWorld, UCanvas* InCanvasObject, bool bInIsExternalRendering)
 {
-	Canvas = nullptr;
-	CanvasSize = FVector2D(EForceInit::ForceInit);
+	FCanvas* InCanvas = nullptr;
+	const FSceneView* InSceneView = nullptr;
 	if (InCanvasObject)
 	{
-		Canvas = InCanvasObject->Canvas;
-		SceneView = InCanvasObject->SceneView;
-
-		CanvasSize = FVector2d(InCanvasObject->SizeX, InCanvasObject->SizeY);
+		InCanvas = InCanvasObject->Canvas;
+		InSceneView = InCanvasObject->SceneView;
 	}
 
-	Initialize(InWorld, bInIsExternalRendering);
+	Initialize(InWorld, InSceneView, InCanvas, bInIsExternalRendering);
 }
 
 FCameraDebugRenderer::FCameraDebugRenderer(UWorld* InWorld, const FSceneView* InSceneView, FCanvas* InCanvas, bool bInIsExternalRendering)
 {
-	Canvas = nullptr;
-	CanvasSize = FVector2D(EForceInit::ForceInit);
-	if (InCanvas)
-	{
-		Canvas = InCanvas;
-		SceneView = InSceneView;
-
-		const FIntRect ViewRect = InCanvas->GetViewRect();
-		CanvasSize = FVector2d(ViewRect.Width(), ViewRect.Height());
-	}
-
-	Initialize(InWorld, bInIsExternalRendering);
+	Initialize(InWorld, InSceneView, InCanvas, bInIsExternalRendering);
 }
 
-void FCameraDebugRenderer::Initialize(UWorld* InWorld, bool bInIsExternalRendering)
+void FCameraDebugRenderer::Initialize(UWorld* InWorld, const FSceneView* InSceneView, FCanvas* InCanvas, bool bInIsExternalRendering)
 {
 	World = InWorld;
 	bIsExternalRendering = bInIsExternalRendering;
@@ -156,8 +143,22 @@ void FCameraDebugRenderer::Initialize(UWorld* InWorld, bool bInIsExternalRenderi
 
 	NextCardPosition = FVector2f::ZeroVector;
 	NextCardColumn = 0;
-	if (Canvas)
+
+	CanvasSize = FVector2D(EForceInit::ForceInit);
+
+	if (InSceneView && InCanvas)
 	{
+		SceneView = InSceneView;
+		Canvas = InCanvas;
+
+		FIntRect ViewRect = InCanvas->GetViewRect();
+		if (ViewRect.Width() == 0 || ViewRect.Height() == 0)
+		{
+			ViewRect = InSceneView->UnconstrainedViewRect;
+		}
+
+		CanvasSize = FVector2d(ViewRect.Width(), ViewRect.Height());
+
 		NextCardPosition = FVector2f{ 
 			(float)CanvasSize.X - (float)GGameplayCamerasDebugCardWidth - (float)GGameplayCamerasDebugRightMargin,
 			(float)GGameplayCamerasDebugTopMargin };
diff --git a/Engine/Plugins/Cameras/GameplayCameras/Source/GameplayCameras/Private/GameFramework/GameplayCameraActor.cpp b/Engine/Plugins/Cameras/GameplayCameras/Source/GameplayCameras/Private/GameFramework/GameplayCameraActor.cpp
--- a/Engine/Plugins/Cameras/GameplayCameras/Source/GameplayCameras/Private/GameFramework/GameplayCameraActor.cpp
+++ b/Engine/Plugins/Cameras/GameplayCameras/Source/GameplayCameras/Private/GameFramework/GameplayCameraActor.cpp
@@ -2,6 +2,7 @@
 
 #include "GameFramework/GameplayCameraActor.h"
 
+#include "CineCameraComponent.h"
 #include "Engine/EngineTypes.h"
 #include "Engine/World.h"
 #include "GameFramework/GameplayCameraComponent.h"
@@ -24,5 +25,20 @@ USceneComponent* AGameplayCameraActor::GetDefaultAttachComponent() const
 	return CameraComponent;
 }
 
+void AGameplayCameraActor::Tick(float DeltaTime)
+{
+	Super::Tick(DeltaTime);
+}
+
+bool AGameplayCameraActor::ShouldTickIfViewportsOnly() const
+{
+	return true;
+}
+
+void AGameplayCameraActor::DisplayDebug(UCanvas* Canvas, const FDebugDisplayInfo& DebugDisplay, float& YL, float& YPos)
+{
+	Super::DisplayDebug(Canvas, DebugDisplay, YL, YPos);
+}
+
 #undef LOCTEXT_NAMESPACE
 
diff --git a/Engine/Plugins/Cameras/GameplayCameras/Source/GameplayCameras/Private/GameFramework/GameplayCameraComponent.cpp b/Engine/Plugins/Cameras/GameplayCameras/Source/GameplayCameras/Private/GameFramework/GameplayCameraComponent.cpp
--- a/Engine/Plugins/Cameras/GameplayCameras/Source/GameplayCameras/Private/GameFramework/GameplayCameraComponent.cpp
+++ b/Engine/Plugins/Cameras/GameplayCameras/Source/GameplayCameras/Private/GameFramework/GameplayCameraComponent.cpp
@@ -2,38 +2,48 @@
 
 #include "GameFramework/GameplayCameraComponent.h"
 
-#include "Components/StaticMeshComponent.h"
+#include "CineCameraComponent.h"
 #include "Core/CameraAsset.h"
 #include "Core/CameraAssetBuilder.h"
 #include "Core/CameraBuildLog.h"
 #include "Core/CameraSystemEvaluator.h"
+#include "Core/RootCameraNode.h"
+#include "Debug/CameraDebugRenderer.h"
+#include "Engine/Canvas.h"
 #include "Engine/EngineTypes.h"
-#include "Engine/StaticMesh.h"
 #include "Engine/World.h"
 #include "GameFramework/GameplayCameraSystemActor.h"
 #include "GameFramework/GameplayCameraSystemHost.h"
-#include "GameplayCameras.h"
+#include "GameplayCamerasDelegates.h"
 #include "Kismet/GameplayStatics.h"
-#include "UObject/ConstructorHelpers.h"
+#include "Misc/AssertionMacros.h"
+#include "PrimitiveDrawInterface.h"
+#include "SceneView.h"
+#include "ShowFlags.h"
 
 #include UE_INLINE_GENERATED_CPP_BY_NAME(GameplayCameraComponent)
 
 #define LOCTEXT_NAMESPACE "GameplayCameraComponent"
 
+namespace UE::Cameras
+{
+
+#if UE_GAMEPLAY_CAMERAS_DEBUG
+extern bool GGameplayCamerasDebugEnable;
+#endif  // UE_GAMEPLAY_CAMERAS_DEBUG
+
+}  // namespace UE::Cameras
+
 UGameplayCameraComponent::UGameplayCameraComponent(const FObjectInitializer& ObjectInit)
 	: Super(ObjectInit)
 {
+	bTickInEditor = true;
 	bWantsOnUpdateTransform = true;
+
 	PrimaryComponentTick.bCanEverTick = true;
 
-#if WITH_EDITORONLY_DATA
-	if (GIsEditor && !IsRunningCommandlet())
-	{
-		static ConstructorHelpers::FObjectFinder<UStaticMesh> EditorCameraMesh(
-				TEXT("/Engine/EditorMeshes/Camera/SM_CineCam.SM_CineCam"));
-		PreviewMesh = EditorCameraMesh.Object;
-	}
-#endif  // WITH_EDITORONLY_DATA
+	OutputCameraComponent = ObjectInit.CreateDefaultSubobject<UCineCameraComponent>(this, TEXT("OutputCameraComponent"), true);
+	OutputCameraComponent->SetupAttachment(this);
 }
 
 void UGameplayCameraComponent::PostLoad()
@@ -56,6 +66,15 @@ void UGameplayCameraComponent::AddReferencedObjects(UObject* InThis, FReferenceC
 	{
 		This->EvaluationContext->AddReferencedObjects(Collector);
 	}
+
+#if WITH_EDITOR
+
+	if (This->EditorPreviewEvaluator)
+	{
+		This->EditorPreviewEvaluator->AddReferencedObjects(Collector);
+	}
+
+#endif  // WITH_EDITOR
 }
 
 TSharedPtr<const UE::Cameras::FCameraEvaluationContext> UGameplayCameraComponent::GetEvaluationContext() const
@@ -97,14 +116,6 @@ void UGameplayCameraComponent::ActivateCameraEvaluationContext(int32 PlayerIndex
 	DeactivateCameraEvaluationContext();
 
 	APlayerController* PlayerController = UGameplayStatics::GetPlayerController(this, PlayerIndex);
-	if (!PlayerController)
-	{
-		FFrame::KismetExecutionMessage(
-				TEXT("Can't activate gameplay camera: no player controller found!"),
-				ELogVerbosity::Error);
-		return;
-	}
-
 	ActivateCameraEvaluationContext(PlayerController);
 }
 
@@ -123,6 +134,11 @@ void UGameplayCameraComponent::DeactivateCameraEvaluationContext()
 		Evaluator->RemoveEvaluationContext(EvaluationContext.ToSharedRef());
 	}
 
+	if (OutputCameraComponent)
+	{
+		OutputCameraComponent->SetRelativeTransform(FTransform());
+	}
+
 	// Don't deactivate the component: we still need to update our evaluation context while any
 	// running camera rigs blend out.
 }
@@ -158,6 +174,19 @@ void UGameplayCameraComponent::ActivateCameraEvaluationContext(APlayerController
 
 	AGameplayCameraSystemActor::AutoManageActiveViewTarget(PlayerController);
 
+	EnsureCameraEvaluationContextCreated(PlayerController);
+
+	TSharedPtr<FCameraSystemEvaluator> CameraSystemEvaluator = CameraSystemHost->GetCameraSystemEvaluator();
+	CameraSystemEvaluator->PushEvaluationContext(EvaluationContext.ToSharedRef());
+
+	// Make sure the component is active so it receives tick updates to maintain the evaluation context.
+	Activate();
+}
+
+void UGameplayCameraComponent::EnsureCameraEvaluationContextCreated(APlayerController* PlayerController)
+{
+	using namespace UE::Cameras;
+
 	if (!EvaluationContext.IsValid())
 	{
 		EvaluationContext = MakeShared<FGameplayCameraComponentEvaluationContext>();
@@ -170,12 +199,6 @@ void UGameplayCameraComponent::ActivateCameraEvaluationContext(APlayerController
 
 		UpdateCameraEvaluationContext(true);
 	}
-
-	TSharedPtr<FCameraSystemEvaluator> CameraSystemEvaluator = CameraSystemHost->GetCameraSystemEvaluator();
-	CameraSystemEvaluator->PushEvaluationContext(EvaluationContext.ToSharedRef());
-
-	// Make sure the component is active so it receives tick updates to maintain the evaluation context.
-	Activate();
 }
 
 #define UE_PRIVATE_GAMEPLAY_CAMERA_COMPONENT_VALIDATE_EVALUATION_CONTEXT(ErrorMsg, ErrorResult)\
@@ -234,42 +257,46 @@ FBlueprintCameraVariableTable UGameplayCameraComponent::GetInitialVariableTable(
 
 void UGameplayCameraComponent::OnRegister()
 {
+	using namespace UE::Cameras;
+
 	Super::OnRegister();
 
-#if WITH_EDITORONLY_DATA
-	if (PreviewMesh && !PreviewMeshComponent)
-	{
-		PreviewMeshComponent = NewObject<UStaticMeshComponent>(this, NAME_None, RF_Transactional | RF_TextExportTransient);
-		PreviewMeshComponent->SetupAttachment(this);
-		PreviewMeshComponent->SetIsVisualizationComponent(true);
-		PreviewMeshComponent->SetStaticMesh(PreviewMesh);
-		PreviewMeshComponent->SetCollisionProfileName(UCollisionProfile::NoCollision_ProfileName);
-		PreviewMeshComponent->bHiddenInGame = true;
-		PreviewMeshComponent->CastShadow = false;
-		PreviewMeshComponent->CreationMethod = CreationMethod;
-		PreviewMeshComponent->RegisterComponentWithWorld(GetWorld());
-	}
+#if WITH_EDITOR
+
+	UWorld* World = GetWorld();
+	bIsEditorWorld = (World && (World->WorldType == EWorldType::Editor || World->WorldType == EWorldType::EditorPreview));
+
+	FGameplayCamerasDelegates::OnCameraAssetBuilt().AddUObject(this, &UGameplayCameraComponent::OnCameraAssetBuilt);
 
-	UpdatePreviewMeshTransform();
-#endif	// WITH_EDITORONLY_DATA
+	const TCHAR* ShowFlagName = TEXT("GameplayCameras");
+	CustomShowFlag = FEngineShowFlags::FindIndexByName(ShowFlagName);
+
+#endif  // WITH_EDITOR
 }
 
 void UGameplayCameraComponent::BeginPlay()
 {
+	using namespace UE::Cameras;
+
 	Super::BeginPlay();
 
 #if WITH_EDITOR
+
 	if (CameraReference.IsValid())
 	{
-		// Auto-build the camera asset on begin play to make sure we've got the latest user edits.
-		using namespace UE::Cameras;
-		FCameraBuildLog BuildLog;
-		FCameraAssetBuilder Builder(BuildLog);
-		Builder.BuildCamera(CameraReference.GetCameraAsset());
+		UWorld* World = GetWorld();
+		if (World && World->WorldType == EWorldType::PIE)
+		{
+			// Auto-build the camera asset on begin play to make sure we've got the latest user edits.
+			FCameraBuildLog BuildLog;
+			FCameraAssetBuilder Builder(BuildLog);
+			Builder.BuildCamera(CameraReference.GetCameraAsset());
+		}
 
 		CameraReference.RebuildParametersIfNeeded();
 	}
-#endif
+
+#endif  // WITH_EDITOR
 
 	if (IsActive() && AutoActivateForPlayer != EAutoReceiveInput::Disabled && GetNetMode() != NM_DedicatedServer)
 	{
@@ -285,13 +312,43 @@ void UGameplayCameraComponent::EndPlay(const EEndPlayReason::Type EndPlayReason)
 	Super::EndPlay(EndPlayReason);
 }
 
-void UGameplayCameraComponent::TickComponent(float DeltaTime, enum ELevelTick TickType, FActorComponentTickFunction *ThisTickFunction)
+void UGameplayCameraComponent::OnUnregister()
 {
+	using namespace UE::Cameras;
+
+#if WITH_EDITOR
+
+	FGameplayCamerasDelegates::OnCameraAssetBuilt().RemoveAll(this);
+
+#endif  // WITH_EDITOR
+
+	Super::OnUnregister();
+}
+
+void UGameplayCameraComponent::TickComponent(float DeltaTime, ELevelTick TickType, FActorComponentTickFunction *ThisTickFunction)
+{
+	using namespace UE::Cameras;
+
 	Super::TickComponent(DeltaTime, TickType, ThisTickFunction);
 
+#if WITH_EDITOR
+
+	// Make sure things are setup (or not) if we want to run the camera logic in editor (or not).
+	AutoManageEditorPreviewEvaluator();
+
+#endif  // WITH_EDITOR
+
 	if (EvaluationContext)
 	{
 		UpdateCameraEvaluationContext(false);
+
+#if WITH_EDITOR
+		
+		UpdateEditorPreviewEvaluator(DeltaTime);
+
+#endif  // WITH_EDITOR
+
+		UpdateOutputCameraComponent();
 	}
 }
 
@@ -318,16 +375,61 @@ void UGameplayCameraComponent::UpdateCameraEvaluationContext(bool bApplyParamete
 	}
 }
 
-void UGameplayCameraComponent::OnComponentDestroyed(bool bDestroyingHierarchy)
+void UGameplayCameraComponent::UpdateOutputCameraComponent()
 {
-	Super::OnComponentDestroyed(bDestroyingHierarchy);
+	using namespace UE::Cameras;
 
-#if WITH_EDITORONLY_DATA
-	if (PreviewMeshComponent)
+	if (!OutputCameraComponent)
 	{
-		PreviewMeshComponent->DestroyComponent();
+		return;
+	}
+
+	TSharedPtr<FCameraSystemEvaluator> CameraSystemEvaluator;
+
+	if (CameraSystemHost)
+	{
+		CameraSystemEvaluator = CameraSystemHost->GetCameraSystemEvaluator();
 	}
-#endif  // WITH_EDITORONLY_DATA
+#if WITH_EDITOR
+	else if (EditorPreviewEvaluator)
+	{
+		CameraSystemEvaluator = EditorPreviewEvaluator;
+	}
+#endif  // WITH_EDITOR
+
+	bool bGotValidTransform = false;
+	if (CameraSystemEvaluator)
+	{
+		FRootCameraNodeEvaluator* RootNodeEvaluator = CameraSystemEvaluator->GetRootNodeEvaluator();
+		if (RootNodeEvaluator && RootNodeEvaluator->HasAnyActiveCameraRig())
+		{
+			const FCameraSystemEvaluationResult& Result = CameraSystemEvaluator->GetEvaluatedResult();
+
+			OutputCameraComponent->SetWorldTransform(Result.CameraPose.GetTransform());
+			OutputCameraComponent->SetFieldOfView(Result.CameraPose.GetEffectiveFieldOfView());
+			OutputCameraComponent->CurrentAperture = Result.CameraPose.GetAperture();
+			OutputCameraComponent->Filmback.SensorWidth = Result.CameraPose.GetSensorWidth();
+			OutputCameraComponent->Filmback.SensorHeight = Result.CameraPose.GetSensorHeight();
+			OutputCameraComponent->bConstrainAspectRatio = Result.CameraPose.GetConstrainAspectRatio();
+			OutputCameraComponent->bOverrideAspectRatioAxisConstraint = Result.CameraPose.GetOverrideAspectRatioAxisConstraint();
+			OutputCameraComponent->AspectRatioAxisConstraint = Result.CameraPose.GetAspectRatioAxisConstraint();
+
+			OutputCameraComponent->FocusSettings.ManualFocusDistance = Result.CameraPose.GetFocusDistance();
+			OutputCameraComponent->FocusSettings.FocusMethod = (Result.CameraPose.GetEnablePhysicalCamera() ? ECameraFocusMethod::Manual : ECameraFocusMethod::Disable);
+
+			bGotValidTransform = true;
+		}
+	}
+	
+	if (!bGotValidTransform)
+	{
+		OutputCameraComponent->SetRelativeTransform(FTransform());
+	}
+}
+
+void UGameplayCameraComponent::OnComponentDestroyed(bool bDestroyingHierarchy)
+{
+	Super::OnComponentDestroyed(bDestroyingHierarchy);
 }
 
 void UGameplayCameraComponent::OnUpdateTransform(EUpdateTransformFlags UpdateTransformFlags, ETeleportType Teleport)
@@ -338,33 +440,150 @@ void UGameplayCameraComponent::OnUpdateTransform(EUpdateTransformFlags UpdateTra
 	{
 		bIsCameraCutNextFrame = true;
 	}
+
+#if WITH_EDITOR
+
+	if (bIsEditorWorld && EvaluationContext)
+	{
+		UpdateCameraEvaluationContext(false);
+	}
+
+#endif  // WITH_EDITOR
 }
 
-#if WITH_EDITORONLY_DATA
+#if WITH_EDITOR
 
-void UGameplayCameraComponent::UpdatePreviewMeshTransform()
+void UGameplayCameraComponent::AutoManageEditorPreviewEvaluator()
 {
-	if (PreviewMeshComponent)
+	using namespace UE::Cameras;
+
+	if (!bIsEditorWorld)
+	{
+		return;
+	}
+
+	if (bRunInEditor && !(EditorPreviewEvaluator && EvaluationContext))
 	{
-		// CineCam mesh is wrong, adjust like UCineCameraComponent
-		PreviewMeshComponent->SetRelativeRotation(FRotator(0.f, 90.f, 0.f));
-		PreviewMeshComponent->SetRelativeLocation(FVector(-46.f, 0, -24.f));
-		PreviewMeshComponent->SetRelativeScale3D(FVector::OneVector);
+		// We want to run the camera logic in the editor but we haven't set things up for that.
+		// Let's create the preview evaluator and the evaluation context.
+		if (!EditorPreviewEvaluator)
+		{
+			EditorPreviewEvaluator = MakeShared<FCameraSystemEvaluator>();
+			EditorPreviewEvaluator->Initialize(this);
+		}
+		if (!EvaluationContext)
+		{
+			EnsureCameraEvaluationContextCreated(nullptr);
+			EditorPreviewEvaluator->PushEvaluationContext(EvaluationContext.ToSharedRef());
+		}
+	}
+	else if (!bRunInEditor && (EditorPreviewEvaluator || EvaluationContext))
+	{
+		// We don't want to run the camera logic in the editor anymore. Let's tear things down.
+		EditorPreviewEvaluator = nullptr;
+		EvaluationContext = nullptr;
 	}
 }
 
-#endif  // WITH_EDITORONLY_DATA
+void UGameplayCameraComponent::OnCameraAssetBuilt(const UCameraAsset* InCameraAsset)
+{
+	using namespace UE::Cameras;
 
-#if WITH_EDITOR
+	if (InCameraAsset != CameraReference.GetCameraAsset())
+	{
+		return;
+	}
+
+	// If our camera asset was just built, it may have some new parameters. We need to rebuild
+	// our variable table and context data table, and re-apply overrides.
+	if (bRunInEditor && EditorPreviewEvaluator && EvaluationContext)
+	{
+		const FCameraAssetAllocationInfo& AllocationInfo = InCameraAsset->GetAllocationInfo();
+		FCameraNodeEvaluationResult& InitialResult = EvaluationContext->GetInitialResult();
+		InitialResult.VariableTable.Initialize(AllocationInfo.VariableTableInfo);
+		InitialResult.ContextDataTable.Initialize(AllocationInfo.ContextDataTableInfo);
+
+		CameraReference.RebuildParametersIfNeeded();
+
+		UpdateCameraEvaluationContext(true);
+	}
+}
 
 bool UGameplayCameraComponent::GetEditorPreviewInfo(float DeltaTime, FMinimalViewInfo& ViewOut)
 {
-	// TODO: in the future, run the camera asset in a private camera system evaluator, with a UI
-	//		 to pick which camera rig to preview.
-	const FTransform3d& ComponentTransform = GetComponentTransform();
-	ViewOut.Location = ComponentTransform.GetLocation();
-	ViewOut.Rotation = ComponentTransform.Rotator();
-	return true;
+	if (OutputCameraComponent)
+	{
+		OutputCameraComponent->GetEditorPreviewInfo(DeltaTime, ViewOut);
+		return true;
+	}
+	return false;
+}
+
+void UGameplayCameraComponent::UpdateEditorPreviewEvaluator(float DeltaTime)
+{
+	using namespace UE::Cameras;
+
+	if (EditorPreviewEvaluator)
+	{
+		FCameraSystemEvaluationParams Params;
+		Params.DeltaTime = DeltaTime;
+		EditorPreviewEvaluator->Update(Params);
+	}
+}
+
+void UGameplayCameraComponent::OnDrawVisualizationHUD(const FViewport* Viewport, const FSceneView* SceneView, FCanvas* Canvas) const
+{
+	using namespace UE::Cameras;
+
+	const bool bHasShowFlag = SceneView->Family->EngineShowFlags.GetSingleFlag(CustomShowFlag);
+	if (bHasShowFlag && bRunInEditor && EditorPreviewEvaluator && EvaluationContext)
+	{
+		const AActor* OwnerActor = GetOwner();
+
+		const AActor* ViewActor = SceneView->ViewActor.Get();
+		const bool bIsLockedToCamera = (ViewActor == OwnerActor);
+
+		FCameraSystemEditorPreviewParams Params;
+		Params.Canvas = Canvas;
+		Params.SceneView = SceneView;
+		Params.bIsLockedToCamera = bIsLockedToCamera;
+		Params.bDrawWorldDebug = false;
+
+		EditorPreviewEvaluator->DrawEditorPreview(Params);
+	}
+}
+
+void UGameplayCameraComponent::PostEditChangeProperty( struct FPropertyChangedEvent& PropertyChangedEvent)
+{
+	using namespace UE::Cameras;
+
+	Super::PostEditChangeProperty(PropertyChangedEvent);
+
+	const FName MemberPropertyName = PropertyChangedEvent.GetMemberPropertyName();
+	if (MemberPropertyName == GET_MEMBER_NAME_CHECKED(UGameplayCameraComponent, CameraReference))
+	{
+		if (EditorPreviewEvaluator && EvaluationContext)
+		{
+			if (EvaluationContext->GetCameraAsset() != CameraReference.GetCameraAsset())
+			{
+				// The camera asset has changed! Recreate the context.
+				EditorPreviewEvaluator->RemoveEvaluationContext(EvaluationContext.ToSharedRef());
+				EvaluationContext = nullptr;
+
+				EnsureCameraEvaluationContextCreated(nullptr);
+				EditorPreviewEvaluator->PushEvaluationContext(EvaluationContext.ToSharedRef());
+			}
+			else
+			{
+				// Otherwise, maybe one of the parameter overrides has changed. Re-apply them.
+				UpdateCameraEvaluationContext(true);
+			}
+		}
+	}
+	else if (MemberPropertyName == GET_MEMBER_NAME_CHECKED(UGameplayCameraComponent, bRunInEditor))
+	{
+		AutoManageEditorPreviewEvaluator();
+	}
 }
 
 #endif  // WITH_EDITOR
diff --git a/Engine/Plugins/Cameras/GameplayCameras/Source/GameplayCameras/Private/GameFramework/GameplayCameraSystemHost.cpp b/Engine/Plugins/Cameras/GameplayCameras/Source/GameplayCameras/Private/GameFramework/GameplayCameraSystemHost.cpp
--- a/Engine/Plugins/Cameras/GameplayCameras/Source/GameplayCameras/Private/GameFramework/GameplayCameraSystemHost.cpp
+++ b/Engine/Plugins/Cameras/GameplayCameras/Source/GameplayCameras/Private/GameFramework/GameplayCameraSystemHost.cpp
@@ -4,6 +4,7 @@
 
 #include "Core/CameraSystemEvaluator.h"
 #include "Debug/DebugDrawService.h"
+#include "Engine/Canvas.h"
 #include "Engine/World.h"
 #include "GameFramework/GameplayCameraSystemActor.h"
 #include "GameFramework/GameplayCameraSystemComponent.h"
@@ -38,7 +39,7 @@ UGameplayCameraSystemHost* UGameplayCameraSystemHost::FindOrCreateHost(APlayerCo
 		HostName = DefaultHostName;
 	}
 
-	UGameplayCameraSystemHost* NewHost = NewObject<UGameplayCameraSystemHost>(PlayerController, HostName);
+	UGameplayCameraSystemHost* NewHost = NewObject<UGameplayCameraSystemHost>(PlayerController, HostName, RF_Transient);
 
 	NewHost->Initialize();
 
@@ -65,6 +66,17 @@ UGameplayCameraSystemHost* UGameplayCameraSystemHost::FindHost(APlayerController
 	return Host;
 }
 
+#if WITH_EDITOR
+
+UGameplayCameraSystemHost* UGameplayCameraSystemHost::CreateEditorPreviewHost(UObject* Outer, const TCHAR* HostName)
+{
+	UGameplayCameraSystemHost* NewHost = NewObject<UGameplayCameraSystemHost>(Outer, HostName, RF_Transient);
+	NewHost->Initialize();
+	return NewHost;
+}
+
+#endif  // WITH_EDITOR
+
 UGameplayCameraSystemHost::UGameplayCameraSystemHost(const FObjectInitializer& ObjectInitializer)
 	: Super(ObjectInitializer)
 {
diff --git a/Engine/Plugins/Cameras/GameplayCameras/Source/GameplayCameras/Private/GameplayCamerasModule.cpp b/Engine/Plugins/Cameras/GameplayCameras/Source/GameplayCameras/Private/GameplayCamerasModule.cpp
--- a/Engine/Plugins/Cameras/GameplayCameras/Source/GameplayCameras/Private/GameplayCamerasModule.cpp
+++ b/Engine/Plugins/Cameras/GameplayCameras/Source/GameplayCameras/Private/GameplayCamerasModule.cpp
@@ -6,11 +6,19 @@
 #include "GameplayCameras.h"
 #include "Logging/MessageLog.h"
 #include "Modules/ModuleManager.h"
+#include "ShowFlags.h"
 
 #define LOCTEXT_NAMESPACE "GameplayCamerasModule"
 
 DEFINE_LOG_CATEGORY(LogCameraSystem);
 
+namespace UE::Cameras
+{
+
+TCustomShowFlag<> GameplayCamerasShowFlag(TEXT("GameplayCameras"), true, SFG_Developer, LOCTEXT("ShowFlagDisplayName", "Gameplay Cameras"));
+
+}  // namespace UE::Cameras
+
 IGameplayCamerasModule& IGameplayCamerasModule::Get()
 {
 	return FModuleManager::LoadModuleChecked<IGameplayCamerasModule>("GameplayCameras");
diff --git a/Engine/Plugins/Cameras/GameplayCameras/Source/GameplayCameras/Public/Core/CameraNodeEvaluator.h b/Engine/Plugins/Cameras/GameplayCameras/Source/GameplayCameras/Public/Core/CameraNodeEvaluator.h
--- a/Engine/Plugins/Cameras/GameplayCameras/Source/GameplayCameras/Public/Core/CameraNodeEvaluator.h
+++ b/Engine/Plugins/Cameras/GameplayCameras/Source/GameplayCameras/Public/Core/CameraNodeEvaluator.h
@@ -134,7 +134,12 @@ enum class ECameraNodeEvaluationType
 	/** Normal evaluation. */
 	Standard,
 	/** Evaluation for IK aiming. */
-	IK
+	IK,
+
+#if WITH_EDITOR
+	/** Evaluation for editor preview. */
+	EditorPreview
+#endif  // WITH_EDITOR
 };
 
 /**
@@ -218,6 +223,17 @@ struct GAMEPLAYCAMERAS_API FCameraNodeEvaluatorSerializeParams
 {
 };
 
+#if WITH_EDITOR
+
+/**
+ * Parameter structure for drawing editor preview.
+ */
+struct FCameraEditorPreviewDrawParams
+{
+};
+
+#endif  // WITH_EDITOR
+
 /**
  * Base class for objects responsible for running a camera node.
  */
@@ -267,6 +283,10 @@ public:
 		return Cast<CameraNodeType>(PrivateCameraNode);
 	}
 
+#if WITH_EDITOR
+	void DrawEditorPreview(const FCameraEditorPreviewDrawParams& Params, FCameraDebugRenderer& Renderer);
+#endif  // WITH_EDITOR
+
 #if UE_GAMEPLAY_CAMERAS_DEBUG
 	/** Called to create debug blocks for this node evaluator. */
 	void BuildDebugBlocks(const FCameraDebugBlockBuildParams& Params, FCameraDebugBlockBuilder& Builder);
@@ -322,6 +342,10 @@ protected:
 	 */
 	GAMEPLAYCAMERAS_API virtual void OnSerialize(const FCameraNodeEvaluatorSerializeParams& Params, FArchive& Ar) {}
 
+#if WITH_EDITOR
+	GAMEPLAYCAMERAS_API virtual void OnDrawEditorPreview(const FCameraEditorPreviewDrawParams& Params, FCameraDebugRenderer& Renderer) {}
+#endif  // WITH_EDITOR
+
 #if UE_GAMEPLAY_CAMERAS_DEBUG
 	/** Called to create debug blocks for this node evaluator. */
 	GAMEPLAYCAMERAS_API virtual void OnBuildDebugBlocks(const FCameraDebugBlockBuildParams& Params, FCameraDebugBlockBuilder& Builder);
diff --git a/Engine/Plugins/Cameras/GameplayCameras/Source/GameplayCameras/Public/Core/CameraSystemEvaluator.h b/Engine/Plugins/Cameras/GameplayCameras/Source/GameplayCameras/Public/Core/CameraSystemEvaluator.h
--- a/Engine/Plugins/Cameras/GameplayCameras/Source/GameplayCameras/Public/Core/CameraSystemEvaluator.h
+++ b/Engine/Plugins/Cameras/GameplayCameras/Source/GameplayCameras/Public/Core/CameraSystemEvaluator.h
@@ -14,6 +14,8 @@
 #include "Templates/SharedPointer.h"
 #include "UObject/GCObject.h"
 
+class FCanvas;
+class FSceneView;
 class UCameraDirector;
 class UCameraRigAsset;
 class UCanvas;
@@ -28,6 +30,7 @@ class FCameraEvaluationService;
 class FCameraRigCombinationRegistry;
 class FRootCameraNodeEvaluator;
 enum class ECameraEvaluationServiceFlags;
+enum class ECameraNodeEvaluationType;
 struct FRootCameraNodeCameraRigEvent;
 
 #if UE_GAMEPLAY_CAMERAS_DEBUG
@@ -83,13 +86,50 @@ public:
 };
 
 #if UE_GAMEPLAY_CAMERAS_DEBUG
+
+/**
+ * Parameter structure for running the debug pass of the camera system.
+ */
 struct FCameraSystemDebugUpdateParams
 {
+	/** The canvas to draw upon. */
 	UCanvas* CanvasObject = nullptr;
+
+	/** Whether the debug camera is enabled, giving an "outside" view of camera system. */
 	bool bIsDebugCameraEnabled = false;
 };
+
 #endif  // UE_GAMEPLAY_CAMERAS_DEBUG
 
+#if WITH_EDITOR
+
+/**
+ * Parameter structure for rendering the state of the camera system in editor.
+ */
+struct FCameraSystemEditorPreviewParams
+{
+	/** The canvas to draw upon. */
+	FCanvas* Canvas = nullptr;
+
+	/** The scene view being renderered. */
+	const FSceneView* SceneView = nullptr;
+
+	/**
+	 * Whether debug drawing is done from "inside" the camera system.
+	 * When false, it is acceptable to draw camera frustrums and axes.
+	 * When true, these things shouldn't be drawn as they would clutter the view
+	 * and obfuscate the scene.
+	 */
+	bool bIsLockedToCamera = true;
+
+	/**
+	 * Whether to enable drawing debug lines and other 3D primitives in the world.
+	 */
+	bool bDrawWorldDebug = true;
+};
+
+#endif  // WITH_EDITOR
+
 /**
  * The main camera system evaluator class.
  */
@@ -167,12 +207,22 @@ public:
 	/** Collect reference objects for the garbage collector. */
 	GAMEPLAYCAMERAS_API void AddReferencedObjects(FReferenceCollector& Collector);
 
+#if WITH_EDITOR
+	/** Run an update suitable for an in-editor preview. */
+	GAMEPLAYCAMERAS_API void EditorPreviewUpdate(const FCameraSystemEvaluationParams& Params);
+
+	/** Render information about the state of the camera system in-editor. */
+	GAMEPLAYCAMERAS_API void DrawEditorPreview(const FCameraSystemEditorPreviewParams& Params);
+#endif  // WITH_EDITOR
+
 #if UE_GAMEPLAY_CAMERAS_DEBUG
 	GAMEPLAYCAMERAS_API void DebugUpdate(const FCameraSystemDebugUpdateParams& Params);
 #endif  // UE_GAMEPLAY_CAMERAS_DEBUG
 
 private:
 
+	void UpdateImpl(float DeltaTime, ECameraNodeEvaluationType EvaluationType);
+
 	void PreUpdateServices(float DeltaTime, ECameraEvaluationServiceFlags ExtraFlags);
 	void PostUpdateServices(float DeltaTime, ECameraEvaluationServiceFlags ExtraFlags);
 
diff --git a/Engine/Plugins/Cameras/GameplayCameras/Source/GameplayCameras/Public/Core/RootCameraNode.h b/Engine/Plugins/Cameras/GameplayCameras/Source/GameplayCameras/Public/Core/RootCameraNode.h
--- a/Engine/Plugins/Cameras/GameplayCameras/Source/GameplayCameras/Public/Core/RootCameraNode.h
+++ b/Engine/Plugins/Cameras/GameplayCameras/Source/GameplayCameras/Public/Core/RootCameraNode.h
@@ -119,6 +119,9 @@ public:
 	/** Gets information about the active camera rig in the main layer. */
 	void GetActiveCameraRigInfo(FCameraRigEvaluationInfo& OutCameraRigInfo) const;
 
+	/** Gets whether any camera rig is running. */
+	bool HasAnyActiveCameraRig() const;
+
 	/**
 	 * Builds the hierarchy of the system for a given single camera rig.
 	 * This is expected to return the nodes of all the layers, except for the main layer which
diff --git a/Engine/Plugins/Cameras/GameplayCameras/Source/GameplayCameras/Public/Debug/CameraDebugRenderer.h b/Engine/Plugins/Cameras/GameplayCameras/Source/GameplayCameras/Public/Debug/CameraDebugRenderer.h
--- a/Engine/Plugins/Cameras/GameplayCameras/Source/GameplayCameras/Public/Debug/CameraDebugRenderer.h
+++ b/Engine/Plugins/Cameras/GameplayCameras/Source/GameplayCameras/Public/Debug/CameraDebugRenderer.h
@@ -161,7 +161,7 @@ public:
 
 private:
 
-	void Initialize(UWorld* InWorld, bool bInIsExternalRendering);
+	void Initialize(UWorld* InWorld, const FSceneView* InSceneView, FCanvas* InCanvas, bool bInIsExternalRendering);
 
 	void AddTextFmtImpl(const TCHAR* Fmt, va_list Args);
 	void AddTextImpl(const TCHAR* Buffer);
diff --git a/Engine/Plugins/Cameras/GameplayCameras/Source/GameplayCameras/Public/GameFramework/GameplayCameraActor.h b/Engine/Plugins/Cameras/GameplayCameras/Source/GameplayCameras/Public/GameFramework/GameplayCameraActor.h
--- a/Engine/Plugins/Cameras/GameplayCameras/Source/GameplayCameras/Public/GameFramework/GameplayCameraActor.h
+++ b/Engine/Plugins/Cameras/GameplayCameras/Source/GameplayCameras/Public/GameFramework/GameplayCameraActor.h
@@ -10,6 +10,7 @@
 #include "GameplayCameraActor.generated.h"
 
 class APlayerController;
+class UCineCameraComponent;
 class UGameplayCameraComponent;
 
 /**
@@ -34,6 +35,9 @@ public:
 
 	// AActor interface.
 	virtual USceneComponent* GetDefaultAttachComponent() const override;
+	virtual void Tick(float DeltaTime) override;
+	virtual bool ShouldTickIfViewportsOnly() const override;
+	virtual void DisplayDebug(UCanvas* Canvas, const FDebugDisplayInfo& DebugDisplay, float& YL, float& YPos) override;
 
 private:
 
diff --git a/Engine/Plugins/Cameras/GameplayCameras/Source/GameplayCameras/Public/GameFramework/GameplayCameraComponent.h b/Engine/Plugins/Cameras/GameplayCameras/Source/GameplayCameras/Public/GameFramework/GameplayCameraComponent.h
--- a/Engine/Plugins/Cameras/GameplayCameras/Source/GameplayCameras/Public/GameFramework/GameplayCameraComponent.h
+++ b/Engine/Plugins/Cameras/GameplayCameras/Source/GameplayCameras/Public/GameFramework/GameplayCameraComponent.h
@@ -13,7 +13,12 @@
 #include "GameplayCameraComponent.generated.h"
 
 class APlayerController;
+class FPrimitiveDrawInterface;
+class FSceneView;
+class FViewport;
 class UCameraAsset;
+class UCanvas;
+class UCineCameraComponent;
 class UGameplayCameraSystemHost;
 
 namespace UE::Cameras
@@ -48,6 +53,10 @@ public:
 
 public:
 
+	/** Gets the child camera component used as the "output" for the gameplay/procedural camera. */
+	UFUNCTION(BlueprintGetter, Category=Camera)
+	UCineCameraComponent* GetOutputCameraComponent() const { return OutputCameraComponent; }
+
 	/** Activates the camera for the given player. */
 	UFUNCTION(BlueprintCallable, Category=Camera)
 	GAMEPLAYCAMERAS_API void ActivateCameraForPlayerIndex(int32 PlayerIndex);
@@ -86,7 +95,8 @@ public:
 	virtual void OnRegister() override;
 	virtual void BeginPlay() override;
 	virtual void EndPlay(const EEndPlayReason::Type EndPlayReason) override;
-	virtual void TickComponent(float DeltaTime, enum ELevelTick TickType, FActorComponentTickFunction *ThisTickFunction) override;
+	virtual void OnUnregister() override;
+	virtual void TickComponent(float DeltaTime, ELevelTick TickType, FActorComponentTickFunction *ThisTickFunction) override;
 	virtual void OnComponentDestroyed(bool bDestroyingHierarchy) override;
 #if WITH_EDITOR
 	virtual bool GetEditorPreviewInfo(float DeltaTime, FMinimalViewInfo& ViewOut) override;
@@ -96,22 +106,34 @@ public:
 	virtual void OnUpdateTransform(EUpdateTransformFlags UpdateTransformFlags, ETeleportType Teleport) override;
 
 	// UObject interface.
-	void PostLoad() override;
+	virtual void PostLoad() override;
+#if WITH_EDITOR
+	virtual void PostEditChangeProperty( struct FPropertyChangedEvent& PropertyChangedEvent) override;
+#endif
 
 	static void AddReferencedObjects(UObject* InThis, FReferenceCollector& Collector);
 
+public:
+
+#if WITH_EDITOR
+	GAMEPLAYCAMERAS_API void OnDrawVisualizationHUD(const FViewport* Viewport, const FSceneView* SceneView, FCanvas* Canvas) const;
+#endif
+
 private:
 
 	void ActivateCameraEvaluationContext(int32 PlayerIndex);
 	void ActivateCameraEvaluationContext(APlayerController* PlayerController);
+	void EnsureCameraEvaluationContextCreated(APlayerController* PlayerController);
 	void UpdateCameraEvaluationContext(bool bApplyParameterOverrides);
+	void UpdateOutputCameraComponent();
 	void DeactivateCameraEvaluationContext();
 
-#if WITH_EDITORONLY_DATA
-
-	void UpdatePreviewMeshTransform();
+#if WITH_EDITOR
+	void AutoManageEditorPreviewEvaluator();
+	void OnCameraAssetBuilt(const UCameraAsset* InCameraAsset);
 
-#endif
+	void UpdateEditorPreviewEvaluator(float DeltaTime);
+#endif  // WITH_EDITOR
 
 public:
 
@@ -126,6 +148,21 @@ public:
 	UPROPERTY(BlueprintReadWrite, EditAnywhere, Category=Activation, meta=(EditCondition="bAutoActivate"))
 	TEnumAsByte<EAutoReceiveInput::Type> AutoActivateForPlayer;
 
+#if WITH_EDITORONLY_DATA
+
+	UPROPERTY(EditAnywhere, Category=Camera)
+	bool bRunInEditor = true;
+
+#endif  // WITH_EDITORONLY_DATA
+
+protected:
+
+	UPROPERTY(Transient)
+	TObjectPtr<UCineCameraComponent> OutputCameraComponent;
+
+	UPROPERTY(Transient)
+	TObjectPtr<UGameplayCameraSystemHost> CameraSystemHost;
+
 protected:
 
 	using FGameplayCameraComponentEvaluationContext = UE::Cameras::FGameplayCameraComponentEvaluationContext;
@@ -134,18 +171,15 @@ protected:
 
 	bool bIsCameraCutNextFrame = false;
 
-	UPROPERTY()
-	TObjectPtr<UGameplayCameraSystemHost> CameraSystemHost;
-
-#if WITH_EDITORONLY_DATA
+#if WITH_EDITOR
+	
+	TSharedPtr<UE::Cameras::FCameraSystemEvaluator> EditorPreviewEvaluator;
 
-	UPROPERTY(Transient)
-	TObjectPtr<UStaticMesh> PreviewMesh;
+	bool bIsEditorWorld = false;
 
-	UPROPERTY(Transient)
-	TObjectPtr<UStaticMeshComponent> PreviewMeshComponent;
+	int32 CustomShowFlag = INDEX_NONE;
 
-#endif	// WITH_EDITORONLY_DATA
+#endif  // WITH_EDITOR
 
 private:
 
diff --git a/Engine/Plugins/Cameras/GameplayCameras/Source/GameplayCameras/Public/GameFramework/GameplayCameraSystemHost.h b/Engine/Plugins/Cameras/GameplayCameras/Source/GameplayCameras/Public/GameFramework/GameplayCameraSystemHost.h
--- a/Engine/Plugins/Cameras/GameplayCameras/Source/GameplayCameras/Public/GameFramework/GameplayCameraSystemHost.h
+++ b/Engine/Plugins/Cameras/GameplayCameras/Source/GameplayCameras/Public/GameFramework/GameplayCameraSystemHost.h
@@ -73,6 +73,12 @@ public:
 	 */
 	static GAMEPLAYCAMERAS_API UGameplayCameraSystemHost* FindHost(APlayerController* PlayerController, const TCHAR* HostName = nullptr, bool bAllowNull = true);
 
+#if WITH_EDITOR
+
+	static GAMEPLAYCAMERAS_API UGameplayCameraSystemHost* CreateEditorPreviewHost(UObject* Outer, const TCHAR* HostName = nullptr);
+
+#endif  // WITH_EDITOR
+
 private:
 
 	/** Default host name to use when creating a new host. */
diff --git a/Engine/Plugins/Cameras/GameplayCameras/Source/GameplayCamerasEditor/Private/ComponentVisualizers/GameplayCameraComponentVisualizer.cpp b/Engine/Plugins/Cameras/GameplayCameras/Source/GameplayCamerasEditor/Private/ComponentVisualizers/GameplayCameraComponentVisualizer.cpp
new file mode 100644
--- /dev/null
+++ b/Engine/Plugins/Cameras/GameplayCameras/Source/GameplayCamerasEditor/Private/ComponentVisualizers/GameplayCameraComponentVisualizer.cpp
@@ -0,0 +1,23 @@
+// Copyright Epic Games, Inc. All Rights Reserved.
+
+#include "ComponentVisualizers/GameplayCameraComponentVisualizer.h"
+
+#include "GameFramework/GameplayCameraComponent.h"
+
+namespace UE::Cameras
+{
+
+void FGameplayCameraComponentVisualizer::DrawVisualization(const UActorComponent* Component, const FSceneView* View, FPrimitiveDrawInterface* PDI)
+{
+}
+
+void FGameplayCameraComponentVisualizer::DrawVisualizationHUD(const UActorComponent* Component, const FViewport* Viewport, const FSceneView* View, FCanvas* Canvas)
+{
+	if (const UGameplayCameraComponent* GameplayCameraComponent = Cast<const UGameplayCameraComponent>(Component))
+	{
+		GameplayCameraComponent->OnDrawVisualizationHUD(Viewport, View, Canvas);
+	}
+}
+
+}  // namespace UE::Cameras
+
diff --git a/Engine/Plugins/Cameras/GameplayCameras/Source/GameplayCamerasEditor/Private/ComponentVisualizers/GameplayCameraComponentVisualizer.h b/Engine/Plugins/Cameras/GameplayCameras/Source/GameplayCamerasEditor/Private/ComponentVisualizers/GameplayCameraComponentVisualizer.h
new file mode 100644
--- /dev/null
+++ b/Engine/Plugins/Cameras/GameplayCameras/Source/GameplayCamerasEditor/Private/ComponentVisualizers/GameplayCameraComponentVisualizer.h
@@ -0,0 +1,24 @@
+// Copyright Epic Games, Inc. All Rights Reserved.
+
+#pragma once
+
+#include "ComponentVisualizer.h"
+
+class FPrimitiveDrawInterface;
+class FSceneView;
+class UActorComponent;
+
+namespace UE::Cameras
+{
+
+class FGameplayCameraComponentVisualizer : public FComponentVisualizer
+{
+public:
+
+	// FComponentVisualizer interface.
+	virtual void DrawVisualization(const UActorComponent* Component, const FSceneView* View, FPrimitiveDrawInterface* PDI) override;
+	virtual void DrawVisualizationHUD(const UActorComponent* Component, const FViewport* Viewport, const FSceneView* View, FCanvas* Canvas) override;
+};
+
+}  // namespace UE::Cameras
+
diff --git a/Engine/Plugins/Cameras/GameplayCameras/Source/GameplayCamerasEditor/Private/GameplayCamerasEditorModule.cpp b/Engine/Plugins/Cameras/GameplayCameras/Source/GameplayCamerasEditor/Private/GameplayCamerasEditorModule.cpp
--- a/Engine/Plugins/Cameras/GameplayCameras/Source/GameplayCamerasEditor/Private/GameplayCamerasEditorModule.cpp
+++ b/Engine/Plugins/Cameras/GameplayCameras/Source/GameplayCamerasEditor/Private/GameplayCamerasEditorModule.cpp
@@ -12,6 +12,7 @@
 #include "Commands/CameraVariableCollectionEditorCommands.h"
 #include "Commands/GameplayCamerasDebuggerCommands.h"
 #include "Commands/ObjectTreeGraphEditorCommands.h"
+#include "ComponentVisualizers/GameplayCameraComponentVisualizer.h"
 #include "Customizations/CameraAssetReferenceDetailsCustomization.h"
 #include "Customizations/CameraParameterDetailsCustomizations.h"
 #include "Customizations/CameraProxyTableDetailsCustomization.h"
@@ -27,9 +28,11 @@
 #include "Debugger/SGameplayCamerasDebugger.h"
 #include "Directors/BlueprintCameraDirector.h"
 #include "EdGraph/EdGraph.h"
+#include "Editor/UnrealEdEngine.h"
 #include "Editors/GameplayCamerasGraphPanelPinFactory.h"
 #include "Editors/SCameraVariablePicker.h"
 #include "Features/IModularFeatures.h"
+#include "GameFramework/GameplayCameraComponent.h"
 #include "GameplayCameras.h"
 #include "GameplayCamerasEditorSettings.h"
 #include "GameplayCamerasLiveEditManager.h"
@@ -50,6 +53,7 @@
 #include "Trace/CameraSystemRewindDebuggerExtension.h"
 #include "Trace/CameraSystemRewindDebuggerTrack.h"
 #include "Trace/CameraSystemTraceModule.h"
+#include "UnrealEdGlobals.h"
 
 #define LOCTEXT_NAMESPACE "GameplayCamerasEditor"
 
@@ -92,6 +96,7 @@ public:
 		RegisterRewindDebuggerFeatures();
 		RegisterDetailsCustomizations();
 		RegisterEdGraphUtilities();
+		RegisterComponentVisualizers();
 
 		InitializeLiveEditManager();
 
@@ -118,6 +123,7 @@ public:
 		UnregisterRewindDebuggerFeatures();
 		UnregisterDetailsCustomizations();
 		UnregisterEdGraphUtilities();
+		UnregisterComponentVisualizers();
 
 		TeardownLiveEditManager();
 
@@ -523,6 +529,24 @@ private:
 		FKismetEditorUtilities::UnregisterAutoBlueprintNodeCreation(this);
 	}
 
+	void RegisterComponentVisualizers()
+	{
+		using namespace UE::Cameras;
+
+		if (GUnrealEd)
+		{
+			GUnrealEd->RegisterComponentVisualizer(UGameplayCameraComponent::StaticClass()->GetFName(), MakeShared<FGameplayCameraComponentVisualizer>());
+		}
+	}
+
+	void UnregisterComponentVisualizers()
+	{
+		if (GUnrealEd)
+		{
+			GUnrealEd->UnregisterComponentVisualizer(UGameplayCameraComponent::StaticClass()->GetFName());
+		}
+	}
+
 	void InitializeLiveEditManager()
 	{
 		using namespace UE::Cameras;
