Cameras: add tracking of camera movements for drawing debug movement trail

diff --git a/Engine/Plugins/Cameras/GameplayCameras/Source/GameplayCameras/Private/Core/BlendStackCameraNode.cpp b/Engine/Plugins/Cameras/GameplayCameras/Source/GameplayCameras/Private/Core/BlendStackCameraNode.cpp
--- a/Engine/Plugins/Cameras/GameplayCameras/Source/GameplayCameras/Private/Core/BlendStackCameraNode.cpp
+++ b/Engine/Plugins/Cameras/GameplayCameras/Source/GameplayCameras/Private/Core/BlendStackCameraNode.cpp
@@ -61,6 +61,13 @@ UE_DEFINE_CAMERA_NODE_EVALUATOR(FBlendStackCameraNodeEvaluator)
 UE_DEFINE_CAMERA_NODE_EVALUATOR(FTransientBlendStackCameraNodeEvaluator)
 UE_DEFINE_CAMERA_NODE_EVALUATOR(FPersistentBlendStackCameraNodeEvaluator)
 
+FBlendStackCameraNodeEvaluator::FBlendStackCameraNodeEvaluator()
+{
+#if WITH_EDITOR || UE_GAMEPLAY_CAMERAS_DEBUG
+	bAutoCameraPoseMovementTrail = false;
+#endif  // WITH_EDITOR || UE_GAMEPLAY_CAMERAS_DEBUG
+}
+
 FBlendStackCameraNodeEvaluator::~FBlendStackCameraNodeEvaluator()
 {
 	// Pop all our entries to unregister the live-edit callbacks.
@@ -231,9 +238,7 @@ void FBlendStackCameraNodeEvaluator::ResolveEntries(TArray<FResolvedEntry>& OutR
 
 			// Reset this entry's flags for this frame.
 			FCameraNodeEvaluationResult& CurResult = Entry.Result;
-			CurResult.CameraPose.ClearAllChangedFlags();
-			CurResult.VariableTable.ClearAllWrittenThisFrameFlags();
-			CurResult.ContextDataTable.ClearAllWrittenThisFrameFlags();
+			CurResult.ResetFrameFlags();
 		}
 		// else: frozen entries may have null contexts or invalid initial results
 		//       because we're not going to update them anyway. We will however blend
@@ -247,7 +252,7 @@ void FBlendStackCameraNodeEvaluator::ResolveEntries(TArray<FResolvedEntry>& OutR
 	}
 }
 
-void FBlendStackCameraNodeEvaluator::OnRunFinished()
+void FBlendStackCameraNodeEvaluator::OnRunFinished(FCameraNodeEvaluationResult& OutResult)
 {
 	// Reset transient flags.
 	for (FCameraRigEntry& Entry : Entries)
@@ -257,6 +262,22 @@ void FBlendStackCameraNodeEvaluator::OnRunFinished()
 		Entry.bBlendRunThisFrame = false;
 		Entry.bForceCameraCut = false;
 	}
+
+#if WITH_EDITOR || UE_GAMEPLAY_CAMERAS_DEBUG
+
+	// Append the motion trail of the active entry so that we see all the steps it took
+	// to get to the end result. Also add an extra point for the actual final result,
+	// to represent the difference between the active result and the blended result.
+	// In theory, this extra segment should blend into nothingness over time.
+	if (Entries.Num() > 0)
+	{
+		const FCameraRigEntry& ActiveEntry = Entries.Last();
+		OutResult.AppendCameraPoseLocationTrail(ActiveEntry.Result);
+
+		OutResult.AddCameraPoseTrailPointIfNeeded();
+	}
+
+#endif  // WITH_EDITOR || UE_GAMEPLAY_CAMERAS_DEBUG
 }
 
 void FBlendStackCameraNodeEvaluator::PopEntry(int32 EntryIndex)
@@ -642,7 +663,7 @@ void FTransientBlendStackCameraNodeEvaluator::OnRun(const FCameraNodeEvaluationP
 	InternalPostBlendExecute(ResolvedEntries, Params, OutResult);
 
 	// Tidy up.
-	OnRunFinished();
+	OnRunFinished(OutResult);
 }
 
 void FTransientBlendStackCameraNodeEvaluator::InternalPreBlendPrepare(TArrayView<FResolvedEntry> ResolvedEntries, const FCameraNodeEvaluationParams& Params, FCameraNodeEvaluationResult& OutResult)
@@ -769,6 +790,7 @@ void FTransientBlendStackCameraNodeEvaluator::InternalUpdate(TArrayView<FResolve
 		FCameraNodeEvaluationResult& CurResult(Entry.Result);
 
 		// Start with the input given to us.
+		CurResult.Reset();
 		CurResult.CameraPose = OutResult.CameraPose;
 		CurResult.CameraRigJoints.OverrideAll(OutResult.CameraRigJoints);
 		CurResult.PostProcessSettings.OverrideAll(OutResult.PostProcessSettings);
@@ -787,6 +809,10 @@ void FTransientBlendStackCameraNodeEvaluator::InternalUpdate(TArrayView<FResolve
 		}
 		CurResult.bIsValid = true;
 
+#if WITH_EDITOR || UE_GAMEPLAY_CAMERAS_DEBUG
+		CurResult.AddCameraPoseTrailPointIfNeeded(ContextResult.CameraPose.GetLocation());
+#endif  // WITH_EDITOR || UE_GAMEPLAY_CAMERAS_DEBUG
+
 		// Run the camera rig's root node.
 		FCameraNodeEvaluator* RootEvaluator = Entry.RootEvaluator->GetRootEvaluator();
 		if (RootEvaluator)
@@ -1055,7 +1081,7 @@ void FPersistentBlendStackCameraNodeEvaluator::OnRun(const FCameraNodeEvaluation
 	InternalUpdate(ResolvedEntries, Params, OutResult);
 
 	// Tidy things up.
-	OnRunFinished();
+	OnRunFinished(OutResult);
 }
 
 void FPersistentBlendStackCameraNodeEvaluator::InternalUpdate(TArrayView<FResolvedEntry> ResolvedEntries, const FCameraNodeEvaluationParams& Params, FCameraNodeEvaluationResult& OutResult)
diff --git a/Engine/Plugins/Cameras/GameplayCameras/Source/GameplayCameras/Private/Core/CameraNodeEvaluator.cpp b/Engine/Plugins/Cameras/GameplayCameras/Source/GameplayCameras/Private/Core/CameraNodeEvaluator.cpp
--- a/Engine/Plugins/Cameras/GameplayCameras/Source/GameplayCameras/Private/Core/CameraNodeEvaluator.cpp
+++ b/Engine/Plugins/Cameras/GameplayCameras/Source/GameplayCameras/Private/Core/CameraNodeEvaluator.cpp
@@ -26,6 +26,17 @@ void FCameraNodeEvaluationResult::Reset()
 
 	bIsCameraCut = false;
 	bIsValid = false;
+
+#if WITH_EDITOR || UE_GAMEPLAY_CAMERAS_DEBUG
+	CameraPoseLocationTrail.Reset();
+#endif  // WITH_EDITOR || UE_GAMEPLAY_CAMERAS_DEBUG
+}
+
+void FCameraNodeEvaluationResult::ResetFrameFlags()
+{
+	CameraPose.ClearAllChangedFlags();
+	VariableTable.ClearAllWrittenThisFrameFlags();
+	ContextDataTable.ClearAllWrittenThisFrameFlags();
 }
 
 void FCameraNodeEvaluationResult::OverrideAll(const FCameraNodeEvaluationResult& OtherResult)
@@ -76,6 +87,34 @@ void FCameraNodeEvaluationResult::AddReferencedObjects(FReferenceCollector& Coll
 	ContextDataTable.AddReferencedObjects(Collector);
 }
 
+#if WITH_EDITOR || UE_GAMEPLAY_CAMERAS_DEBUG
+
+void FCameraNodeEvaluationResult::AddCameraPoseTrailPointIfNeeded()
+{
+	AddCameraPoseTrailPointIfNeeded(CameraPose.GetLocation());
+}
+
+void FCameraNodeEvaluationResult::AddCameraPoseTrailPointIfNeeded(const FVector3d& Point)
+{
+	if (CameraPoseLocationTrail.IsEmpty() || !CameraPoseLocationTrail.Last().Equals(Point))
+	{
+		CameraPoseLocationTrail.Add(Point);
+	}
+}
+
+void FCameraNodeEvaluationResult::AppendCameraPoseLocationTrail(const FCameraNodeEvaluationResult& InResult)
+{
+	ensure(this != &InResult);
+	CameraPoseLocationTrail.Append(InResult.CameraPoseLocationTrail);
+}
+
+TConstArrayView<FVector3d> FCameraNodeEvaluationResult::GetCameraPoseLocationTrail() const
+{
+	return CameraPoseLocationTrail;
+}
+
+#endif  // WITH_EDITOR || UE_GAMEPLAY_CAMERAS_DEBUG
+
 FCameraNodeEvaluator* FCameraNodeEvaluatorBuildParams::BuildEvaluator(const UCameraNode* InNode) const
 {
 	if (InNode)
@@ -161,6 +200,13 @@ void FCameraNodeEvaluator::Run(const FCameraNodeEvaluationParams& Params, FCamer
 	if (!PrivateCameraNode || PrivateCameraNode->bIsEnabled)
 	{
 		OnRun(Params, OutResult);
+
+#if WITH_EDITOR || UE_GAMEPLAY_CAMERAS_DEBUG
+		if (bAutoCameraPoseMovementTrail)
+		{
+			OutResult.AddCameraPoseTrailPointIfNeeded();
+		}
+#endif  // WITH_EDITOR
 	}
 }
 
diff --git a/Engine/Plugins/Cameras/GameplayCameras/Source/GameplayCameras/Private/Core/CameraSystemEvaluator.cpp b/Engine/Plugins/Cameras/GameplayCameras/Source/GameplayCameras/Private/Core/CameraSystemEvaluator.cpp
--- a/Engine/Plugins/Cameras/GameplayCameras/Source/GameplayCameras/Private/Core/CameraSystemEvaluator.cpp
+++ b/Engine/Plugins/Cameras/GameplayCameras/Source/GameplayCameras/Private/Core/CameraSystemEvaluator.cpp
@@ -200,10 +200,9 @@ void FCameraSystemEvaluator::UpdateImpl(float DeltaTime, ECameraNodeEvaluationTy
 {
 	SCOPE_CYCLE_COUNTER(CameraSystemEval_Total);
 
-	// Reset our result' flags.
-	RootNodeResult.CameraPose.ClearAllChangedFlags();
-	RootNodeResult.VariableTable.ClearAllWrittenThisFrameFlags();
-	RootNodeResult.ContextDataTable.ClearAllWrittenThisFrameFlags();
+	// Reset our result' flags. Don't reset the result itself yet, since we want to return
+	// last frame's values when we don't have anything to run.
+	RootNodeResult.ResetFrameFlags();
 
 	// Reset variables and data.
 	RootNodeResult.VariableTable.AutoResetValues();
@@ -383,6 +382,18 @@ void FCameraSystemEvaluator::DrawEditorPreview(const FCameraSystemEditorPreviewP
 	Renderer.BeginDrawing();
 	{
 		RootEvaluator->DrawEditorPreview(NodeParams, Renderer);
+
+		if (!Params.bIsLockedToCamera)
+		{
+			const FLinearColor TrailColor(FColorList::LightBlue);
+			TConstArrayView<FVector3d> Trail = RootNodeResult.GetCameraPoseLocationTrail();
+			for (int32 Index = 1; Index < Trail.Num(); ++Index)
+			{
+				const FVector3d& PrevPoint(Trail[Index - 1]);
+				const FVector3d& NextPoint(Trail[Index]);
+				Renderer.DrawLine(PrevPoint, NextPoint, TrailColor, 1.f);
+			}
+		}
 	}
 	Renderer.EndDrawing();
 }
diff --git a/Engine/Plugins/Cameras/GameplayCameras/Source/GameplayCameras/Public/Core/BlendStackCameraNode.h b/Engine/Plugins/Cameras/GameplayCameras/Source/GameplayCameras/Public/Core/BlendStackCameraNode.h
--- a/Engine/Plugins/Cameras/GameplayCameras/Source/GameplayCameras/Public/Core/BlendStackCameraNode.h
+++ b/Engine/Plugins/Cameras/GameplayCameras/Source/GameplayCameras/Public/Core/BlendStackCameraNode.h
@@ -97,6 +97,7 @@ class FBlendStackCameraNodeEvaluator
 
 public:
 
+	FBlendStackCameraNodeEvaluator();
 	~FBlendStackCameraNodeEvaluator();
 
 	/** Returns information about the top (active) camera rig, if any. */
@@ -167,7 +168,7 @@ protected:
 	};
 
 	void ResolveEntries(TArray<FResolvedEntry>& OutResolvedEntries);
-	void OnRunFinished();
+	void OnRunFinished(FCameraNodeEvaluationResult& OutResult);
 
 protected:
 
diff --git a/Engine/Plugins/Cameras/GameplayCameras/Source/GameplayCameras/Public/Core/CameraNodeEvaluator.h b/Engine/Plugins/Cameras/GameplayCameras/Source/GameplayCameras/Public/Core/CameraNodeEvaluator.h
--- a/Engine/Plugins/Cameras/GameplayCameras/Source/GameplayCameras/Public/Core/CameraNodeEvaluator.h
+++ b/Engine/Plugins/Cameras/GameplayCameras/Source/GameplayCameras/Public/Core/CameraNodeEvaluator.h
@@ -190,6 +190,9 @@ public:
 	/** Reset this result to its default (non-valid) state.  */
 	void Reset();
 
+	/** Reset all written-this-frame flags on the camera pose and tables. */
+	void ResetFrameFlags();
+
 	/** Override this result with the given other result. */
 	void OverrideAll(const FCameraNodeEvaluationResult& OtherResult);
 
@@ -203,6 +206,23 @@ public:
 
 	/** Collects objects from the context data table. */
 	void AddReferencedObjects(FReferenceCollector& Collector);
+
+#if WITH_EDITOR || UE_GAMEPLAY_CAMERAS_DEBUG
+
+public:
+
+	// Internal API.
+
+	void AddCameraPoseTrailPointIfNeeded();
+	void AddCameraPoseTrailPointIfNeeded(const FVector3d& Point);
+	void AppendCameraPoseLocationTrail(const FCameraNodeEvaluationResult& InResult);
+	TConstArrayView<FVector3d> GetCameraPoseLocationTrail() const;
+
+private:
+
+	TArray<FVector3d> CameraPoseLocationTrail;
+
+#endif  // WITH_EDITOR || UE_GAMEPLAY_CAMERAS_DEBUG
 };
 
 /**
@@ -351,6 +371,15 @@ protected:
 	GAMEPLAYCAMERAS_API virtual void OnBuildDebugBlocks(const FCameraDebugBlockBuildParams& Params, FCameraDebugBlockBuilder& Builder);
 #endif  // UE_GAMEPLAY_CAMERAS_DEBUG
 
+#if WITH_EDITOR || UE_GAMEPLAY_CAMERAS_DEBUG
+
+protected:
+
+	/** Whether to automatically add points to the debug trail after this evaluation has run. */
+	bool bAutoCameraPoseMovementTrail = true;
+
+#endif  // WITH_EDITOR || UE_GAMEPLAY_CAMERAS_DEBUG
+
 private:
 
 	/** The camera node to run. */
