Cameras: fix context data type traits and size/alignment for structures

diff --git a/Engine/Plugins/Cameras/GameplayCameras/Source/GameplayCameras/Private/Core/CameraContextDataTable.cpp b/Engine/Plugins/Cameras/GameplayCameras/Source/GameplayCameras/Private/Core/CameraContextDataTable.cpp
--- a/Engine/Plugins/Cameras/GameplayCameras/Source/GameplayCameras/Private/Core/CameraContextDataTable.cpp
+++ b/Engine/Plugins/Cameras/GameplayCameras/Source/GameplayCameras/Private/Core/CameraContextDataTable.cpp
@@ -150,9 +150,8 @@ bool FCameraContextDataTable::GetDataTypeAllocationInfo(ECameraContextDataType D
 				const UScriptStruct* StructType = CastChecked<const UScriptStruct>(DataTypeObject);
 				if (StructType)
 				{
-					const UScriptStruct::ICppStructOps* StructOps = StructType->GetCppStructOps();
-					OutSizeOf = StructOps->GetSize();
-					OutAlignOf = StructOps->GetAlignment();
+					OutSizeOf = StructType->GetPropertiesSize();
+					OutAlignOf = StructType->GetMinAlignment();
 				}
 			}
 			break;
diff --git a/Engine/Plugins/Cameras/GameplayCameras/Source/GameplayCameras/Public/Core/CameraContextDataTable.h b/Engine/Plugins/Cameras/GameplayCameras/Source/GameplayCameras/Public/Core/CameraContextDataTable.h
--- a/Engine/Plugins/Cameras/GameplayCameras/Source/GameplayCameras/Public/Core/CameraContextDataTable.h
+++ b/Engine/Plugins/Cameras/GameplayCameras/Source/GameplayCameras/Public/Core/CameraContextDataTable.h
@@ -330,7 +330,7 @@ ECameraContextDataType TCameraContextDataTraits<DataType>::GetDataType()
 	{
 		return ECameraContextDataType::Object;
 	}
-	else if constexpr(std::is_same_v<decltype(DataType::StaticStruct), UScriptStruct*>)
+	else if constexpr(std::is_same_v<decltype(DataType::StaticStruct()), UScriptStruct*>)
 	{
 		return ECameraContextDataType::Struct;
 	}
@@ -352,7 +352,7 @@ const UObject* TCameraContextDataTraits<DataType>::GetDataTypeObject()
 	{
 		return StaticEnum<DataType>();
 	}
-	else if constexpr(std::is_same_v<decltype(DataType::StaticStruct), UScriptStruct*>)
+	else if constexpr(std::is_same_v<decltype(DataType::StaticStruct()), UScriptStruct*>)
 	{
 		return DataType::StaticStruct();
 	}
