Cameras: add utility methods on camera parameters and variable references

diff --git a/Engine/Plugins/Cameras/GameplayCameras/Source/GameplayCameras/Private/Core/CameraParameters.cpp b/Engine/Plugins/Cameras/GameplayCameras/Source/GameplayCameras/Private/Core/CameraParameters.cpp
--- a/Engine/Plugins/Cameras/GameplayCameras/Source/GameplayCameras/Private/Core/CameraParameters.cpp
+++ b/Engine/Plugins/Cameras/GameplayCameras/Source/GameplayCameras/Private/Core/CameraParameters.cpp
@@ -3,6 +3,7 @@
 #include "Core/CameraParameters.h"
 
 #include "Core/CameraRigAsset.h"
+#include "Core/CameraVariableTable.h"
 #include "Math/Rotator.h"
 #include "Math/Transform.h"
 #include "Math/Vector.h"
@@ -215,6 +216,21 @@ void F##ValueName##CameraParameter::PostSerialize(const FArchive& Ar)\
 			Variable = nullptr;\
 		}\
 	}\
+}\
+ValueType F##ValueName##CameraParameter::GetValue(const UE::Cameras::FCameraVariableTable& VariableTable) const\
+{\
+	if (!VariableID.IsValid())\
+	{\
+		return Value;\
+	}\
+	else\
+	{\
+		if (const ValueType* ActualValue = VariableTable.FindValue<ValueType>(VariableID))\
+		{\
+			return *ActualValue;\
+		}\
+		return Value;\
+	}\
 }
 UE_CAMERA_VARIABLE_FOR_ALL_TYPES()
 #undef UE_CAMERA_VARIABLE_FOR_TYPE
diff --git a/Engine/Plugins/Cameras/GameplayCameras/Source/GameplayCameras/Private/Core/CameraVariableReferences.cpp b/Engine/Plugins/Cameras/GameplayCameras/Source/GameplayCameras/Private/Core/CameraVariableReferences.cpp
new file mode 100644
--- /dev/null
+++ b/Engine/Plugins/Cameras/GameplayCameras/Source/GameplayCameras/Private/Core/CameraVariableReferences.cpp
@@ -0,0 +1,21 @@
+// Copyright Epic Games, Inc. All Rights Reserved.
+
+#include "Core/CameraVariableReferences.h"
+
+#include "Core/CameraVariableTable.h"
+
+#define UE_CAMERA_VARIABLE_FOR_TYPE(ValueType, ValueName)\
+const ValueType* F##ValueName##CameraVariableReference::GetValue(const UE::Cameras::FCameraVariableTable& VariableTable) const\
+{\
+	if (VariableID)\
+	{\
+		if (const ValueType* ActualValue = VariableTable.FindValue<ValueType>(VariableID))\
+		{\
+			return ActualValue;\
+		}\
+	}\
+	return nullptr;\
+}
+UE_CAMERA_VARIABLE_FOR_ALL_TYPES()
+#undef UE_CAMERA_VARIABLE_FOR_TYPE
+
diff --git a/Engine/Plugins/Cameras/GameplayCameras/Source/GameplayCameras/Public/Core/CameraParameters.h b/Engine/Plugins/Cameras/GameplayCameras/Source/GameplayCameras/Public/Core/CameraParameters.h
--- a/Engine/Plugins/Cameras/GameplayCameras/Source/GameplayCameras/Public/Core/CameraParameters.h
+++ b/Engine/Plugins/Cameras/GameplayCameras/Source/GameplayCameras/Public/Core/CameraParameters.h
@@ -8,12 +8,19 @@
 
 #include "CameraParameters.generated.h"
 
+namespace UE::Cameras
+{
+	class FCameraVariableTable;
+}
+
 #define UE_DEFINE_CAMERA_PARAMETER_VALUE_CONSTRUCTORS(ParameterClass)\
 	ParameterClass(typename TCallTraits<ParameterClass::ValueType>::ParamType InValue)\
 		: Value(InValue)\
 	{}\
+	bool HasOverride() const { return VariableID.IsValid(); }\
 	bool HasUserOverride() const { return Variable != nullptr; }\
 	bool HasNonUserOverride() const { return VariableID.IsValid() && (!Variable || Variable->GetVariableID() != VariableID); }\
+	ParameterClass::ValueType GetValue(const UE::Cameras::FCameraVariableTable& VariableTable) const;\
 	void PostSerialize(const FArchive& Ar);
 
 #define UE_DEFINE_CAMERA_PARAMETER_ALL_CONSTRUCTORS(ParameterClass)\
diff --git a/Engine/Plugins/Cameras/GameplayCameras/Source/GameplayCameras/Public/Core/CameraVariableReferences.h b/Engine/Plugins/Cameras/GameplayCameras/Source/GameplayCameras/Public/Core/CameraVariableReferences.h
--- a/Engine/Plugins/Cameras/GameplayCameras/Source/GameplayCameras/Public/Core/CameraVariableReferences.h
+++ b/Engine/Plugins/Cameras/GameplayCameras/Source/GameplayCameras/Public/Core/CameraVariableReferences.h
@@ -6,11 +6,17 @@
 
 #include "CameraVariableReferences.generated.h"
 
+namespace UE::Cameras
+{
+	class FCameraVariableTable;
+}
+
 #define UE_DEFINE_CAMERA_VARIABLE_REFERENCE(ValueName)\
 	F##ValueName##CameraVariableReference() {}\
 	F##ValueName##CameraVariableReference(VariableAssetType* InVariable) : Variable(InVariable) {}\
 	bool IsValid() const { return VariableID.IsValid(); }\
-	bool HasVariable() const { return Variable != nullptr; }
+	bool HasVariable() const { return Variable != nullptr; }\
+	const F##ValueName##CameraVariableReference::ValueType* GetValue(const UE::Cameras::FCameraVariableTable& VariableTable) const;
 
 // All camera references have:
 //
@@ -26,6 +32,7 @@ struct FBooleanCameraVariableReference
 {
 	GENERATED_BODY()
 
+	using ValueType = bool;
 	using VariableAssetType = UBooleanCameraVariable;
 
 	UPROPERTY()
@@ -42,6 +49,7 @@ struct FInteger32CameraVariableReference
 {
 	GENERATED_BODY()
 
+	using ValueType = int32;
 	using VariableAssetType = UInteger32CameraVariable;
 
 	UPROPERTY()
@@ -58,6 +66,7 @@ struct FFloatCameraVariableReference
 {
 	GENERATED_BODY()
 
+	using ValueType = float;
 	using VariableAssetType = UFloatCameraVariable;
 
 	UPROPERTY()
@@ -74,6 +83,7 @@ struct FDoubleCameraVariableReference
 {
 	GENERATED_BODY()
 
+	using ValueType = double;
 	using VariableAssetType = UDoubleCameraVariable;
 
 	UPROPERTY()
@@ -90,6 +100,7 @@ struct FVector2fCameraVariableReference
 {
 	GENERATED_BODY()
 
+	using ValueType = FVector2f;
 	using VariableAssetType = UVector2fCameraVariable;
 
 	UPROPERTY()
@@ -106,6 +117,7 @@ struct FVector2dCameraVariableReference
 {
 	GENERATED_BODY()
 
+	using ValueType = FVector2d;
 	using VariableAssetType = UVector2dCameraVariable;
 
 	UPROPERTY()
@@ -122,6 +134,7 @@ struct FVector3fCameraVariableReference
 {
 	GENERATED_BODY()
 
+	using ValueType = FVector3f;
 	using VariableAssetType = UVector3fCameraVariable;
 
 	UPROPERTY()
@@ -138,6 +151,7 @@ struct FVector3dCameraVariableReference
 {
 	GENERATED_BODY()
 
+	using ValueType = FVector3d;
 	using VariableAssetType = UVector3dCameraVariable;
 
 	UPROPERTY()
@@ -154,6 +168,7 @@ struct FVector4fCameraVariableReference
 {
 	GENERATED_BODY()
 
+	using ValueType = FVector4f;
 	using VariableAssetType = UVector4fCameraVariable;
 
 	UPROPERTY()
@@ -170,6 +185,7 @@ struct FVector4dCameraVariableReference
 {
 	GENERATED_BODY()
 
+	using ValueType = FVector4d;
 	using VariableAssetType = UVector4dCameraVariable;
 
 	UPROPERTY()
@@ -186,6 +202,7 @@ struct FRotator3fCameraVariableReference
 {
 	GENERATED_BODY()
 
+	using ValueType = FRotator3f;
 	using VariableAssetType = URotator3fCameraVariable;
 
 	UPROPERTY()
@@ -202,6 +219,7 @@ struct FRotator3dCameraVariableReference
 {
 	GENERATED_BODY()
 
+	using ValueType = FRotator3d;
 	using VariableAssetType = URotator3dCameraVariable;
 
 	UPROPERTY()
@@ -218,6 +236,7 @@ struct FTransform3fCameraVariableReference
 {
 	GENERATED_BODY()
 
+	using ValueType = FTransform3f;
 	using VariableAssetType = UTransform3fCameraVariable;
 
 	UPROPERTY()
@@ -234,6 +253,7 @@ struct FTransform3dCameraVariableReference
 {
 	GENERATED_BODY()
 
+	using ValueType = FTransform3d;
 	using VariableAssetType = UTransform3dCameraVariable;
 
 	UPROPERTY()
