Cameras: add "KnownOnly" filter to the Override/Lerp methods of the variable table

diff --git a/Engine/Plugins/Cameras/GameplayCameras/Source/GameplayCameras/Private/Core/CameraVariableTable.cpp b/Engine/Plugins/Cameras/GameplayCameras/Source/GameplayCameras/Private/Core/CameraVariableTable.cpp
--- a/Engine/Plugins/Cameras/GameplayCameras/Source/GameplayCameras/Private/Core/CameraVariableTable.cpp
+++ b/Engine/Plugins/Cameras/GameplayCameras/Source/GameplayCameras/Private/Core/CameraVariableTable.cpp
@@ -508,10 +508,11 @@ void FCameraVariableTable::InternalOverride(const FCameraVariableTable& OtherTab
 {
 	using namespace UE::Cameras::Private;
 
-	const bool bChangedOnly = EnumHasAnyFlags(Filter, ECameraVariableTableFilter::ChangedOnly);
 	const bool bInputs = EnumHasAnyFlags(Filter, ECameraVariableTableFilter::Input);
 	const bool bOutputs = EnumHasAnyFlags(Filter, ECameraVariableTableFilter::Output);
 	const bool bPrivates = EnumHasAnyFlags(Filter, ECameraVariableTableFilter::Private);
+	const bool bKnownOnly = EnumHasAnyFlags(Filter, ECameraVariableTableFilter::KnownOnly);
+	const bool bChangedOnly = EnumHasAnyFlags(Filter, ECameraVariableTableFilter::ChangedOnly);
 
 	for (const FEntry& OtherEntry : OtherTable.Entries)
 	{
@@ -546,7 +547,7 @@ void FCameraVariableTable::InternalOverride(const FCameraVariableTable& OtherTab
 						ThisEntry->ID.GetValue(), *DebugName,
 						*UEnum::GetValueAsString(ThisEntry->Type), *UEnum::GetValueAsString(OtherEntry.Type));
 			}
-			else
+			else if (!bKnownOnly)
 			{
 				// We don't have this variable in our table. Let's add it.
 				FCameraVariableDefinition NewVariableDefinition;
@@ -561,6 +562,12 @@ void FCameraVariableTable::InternalOverride(const FCameraVariableTable& OtherTab
 
 				ThisEntry = FindEntry(OtherEntry.ID);
 			}
+			else
+			{
+				// We don't have this variable in our table but we only want to override
+				// known entries.
+				continue;
+			}
 
 			if (ensure(ThisEntry))
 			{
@@ -602,10 +609,11 @@ void FCameraVariableTable::InternalLerp(const FCameraVariableTable& ToTable, ECa
 {
 	using namespace UE::Cameras::Private;
 
-	const bool bChangedOnly = EnumHasAnyFlags(Filter, ECameraVariableTableFilter::ChangedOnly);
 	const bool bInputs = EnumHasAnyFlags(Filter, ECameraVariableTableFilter::Input);
 	const bool bOutputs = EnumHasAnyFlags(Filter, ECameraVariableTableFilter::Output);
 	const bool bPrivates = EnumHasAnyFlags(Filter, ECameraVariableTableFilter::Private);
+	const bool bKnownOnly = EnumHasAnyFlags(Filter, ECameraVariableTableFilter::KnownOnly);
+	const bool bChangedOnly = EnumHasAnyFlags(Filter, ECameraVariableTableFilter::ChangedOnly);
 
 	for (const FEntry& ToEntry : ToTable.Entries)
 	{
@@ -640,6 +648,11 @@ void FCameraVariableTable::InternalLerp(const FCameraVariableTable& ToTable, ECa
 						FromEntry->ID.GetValue(), *DebugName,
 						*UEnum::GetValueAsString(FromEntry->Type), *UEnum::GetValueAsString(ToEntry.Type));
 
+				const EEntryFlags FromFlags = FromEntry->Flags;
+				ensureMsgf(EnumHasAllFlags(FromFlags, EEntryFlags::Written),
+						TEXT("Camera variable '%d' (%s) is LERP'ing from an uninitialized value!"),
+						FromEntry->ID.GetValue(), *DebugName);
+
 				uint8* FromValuePtr = Memory + FromEntry->Offset;
 				const uint8* ToValuePtr = ToTable.Memory + ToEntry.Offset;
 				switch (FromEntry->Type)
@@ -669,10 +682,9 @@ UE_CAMERA_VARIABLE_FOR_ALL_TYPES()
 				// We consider this variable "written to this frame" if it was written in either variable tables this frame.
 				// If the value interpolates because the from/to values are different, but neither was written this frame, we
 				// don't consider the interpolated value written this frame either.
-				const EEntryFlags FromFlags = FromEntry->Flags;
 				EnumAddFlags(FromEntry->Flags, EEntryFlags::Written | (FromFlags & EEntryFlags::WrittenThisFrame) | (ToFlags & EEntryFlags::WrittenThisFrame));
 			}
-			else
+			else if (!bKnownOnly)
 			{
 				// We don't have this variable in our table. Let's add it.
 				FCameraVariableDefinition NewVariableDefinition;
@@ -697,6 +709,12 @@ UE_CAMERA_VARIABLE_FOR_ALL_TYPES()
 				FMemory::Memcpy(FromValuePtr, ToValuePtr, ValueSize);
 				EnumAddFlags(FromEntry->Flags, EEntryFlags::Written | (ToFlags & EEntryFlags::WrittenThisFrame));
 			}
+			else
+			{
+				// We don't have this variable in our table but we only want to override
+				// known entries.
+				continue;
+			}
 
 			if (OutMask)
 			{
diff --git a/Engine/Plugins/Cameras/GameplayCameras/Source/GameplayCameras/Public/Core/CameraVariableTable.h b/Engine/Plugins/Cameras/GameplayCameras/Source/GameplayCameras/Public/Core/CameraVariableTable.h
--- a/Engine/Plugins/Cameras/GameplayCameras/Source/GameplayCameras/Public/Core/CameraVariableTable.h
+++ b/Engine/Plugins/Cameras/GameplayCameras/Source/GameplayCameras/Public/Core/CameraVariableTable.h
@@ -34,9 +34,11 @@ enum class ECameraVariableTableFilter
 	/** Include output variables (i.e. anything not an input). */
 	Output = 1 << 1,
 	/** Include private variables. */
-	Private = 1 << 3,
+	Private = 1 << 2,
+	/** Only include data that is common to both tables. */
+	KnownOnly = 1 << 3,
 	/** Only include variables that were written this frame. */
-	ChangedOnly = 1 << 2,
+	ChangedOnly = 1 << 4,
 
 	/** Both input and output variables. */
 	InputOutput = Input | Output,
