Cameras: fix crash when applying parameter overrides on a camera that doesn't have a valid asset reference

diff --git a/Engine/Plugins/Cameras/GameplayCameras/Source/GameplayCameras/Private/Helpers/CameraAssetParameterOverrideEvaluator.cpp b/Engine/Plugins/Cameras/GameplayCameras/Source/GameplayCameras/Private/Helpers/CameraAssetParameterOverrideEvaluator.cpp
--- a/Engine/Plugins/Cameras/GameplayCameras/Source/GameplayCameras/Private/Helpers/CameraAssetParameterOverrideEvaluator.cpp
+++ b/Engine/Plugins/Cameras/GameplayCameras/Source/GameplayCameras/Private/Helpers/CameraAssetParameterOverrideEvaluator.cpp
@@ -34,7 +34,7 @@ void FCameraAssetParameterOverrideEvaluator::ApplyParameterOverrides(FCameraVari
 	const UCameraAsset* CameraAsset = CameraReference.GetCameraAsset();
 	const FInstancedPropertyBag& CameraParameters = CameraReference.GetParameters();
 	const UPropertyBag* CameraParametersStruct = CameraParameters.GetPropertyBagStruct();
-	if (!CameraParametersStruct)
+	if (!CameraAsset || !CameraParametersStruct)
 	{
 		return;
 	}
diff --git a/Engine/Plugins/Cameras/GameplayCameras/Source/GameplayCameras/Private/Helpers/CameraRigParameterOverrideEvaluator.cpp b/Engine/Plugins/Cameras/GameplayCameras/Source/GameplayCameras/Private/Helpers/CameraRigParameterOverrideEvaluator.cpp
--- a/Engine/Plugins/Cameras/GameplayCameras/Source/GameplayCameras/Private/Helpers/CameraRigParameterOverrideEvaluator.cpp
+++ b/Engine/Plugins/Cameras/GameplayCameras/Source/GameplayCameras/Private/Helpers/CameraRigParameterOverrideEvaluator.cpp
@@ -33,7 +33,7 @@ void FCameraRigParameterOverrideEvaluator::ApplyParameterOverrides(FCameraVariab
 	const UCameraRigAsset* CameraRig = CameraRigReference.GetCameraRig();
 	const FInstancedPropertyBag& CameraRigParameters = CameraRigReference.GetParameters();
 	const UPropertyBag* CameraRigParametersStruct = CameraRigParameters.GetPropertyBagStruct();
-	if (!CameraRigParametersStruct)
+	if (!CameraRig || !CameraRigParametersStruct)
 	{
 		return;
 	}
