Cameras: properly pre-blend all camera rig parameters before updating the blend stack

Before we were only doing this properly for merged camera rigs. Now we're doing it for all camera rigs.

diff --git a/Engine/Plugins/Cameras/GameplayCameras/Source/GameplayCameras/Private/Core/BlendStackCameraNode.cpp b/Engine/Plugins/Cameras/GameplayCameras/Source/GameplayCameras/Private/Core/BlendStackCameraNode.cpp
--- a/Engine/Plugins/Cameras/GameplayCameras/Source/GameplayCameras/Private/Core/BlendStackCameraNode.cpp
+++ b/Engine/Plugins/Cameras/GameplayCameras/Source/GameplayCameras/Private/Core/BlendStackCameraNode.cpp
@@ -17,6 +17,7 @@
 #include "Debug/CameraPoseDebugBlock.h"
 #include "Debug/VariableTableDebugBlock.h"
 #include "HAL/IConsoleManager.h"
+#include "Helpers/CameraParameterOverrideHelper.h"
 #include "IGameplayCamerasLiveEditManager.h"
 #include "IGameplayCamerasModule.h"
 #include "Math/ColorList.h"
@@ -97,6 +98,9 @@ bool FBlendStackCameraNodeEvaluator::InitializeEntry(
 	NewEntry.Result.VariableTable.Initialize(CameraRig->AllocationInfo.VariableTableInfo);
 	NewEntry.Result.ContextDataTable.Initialize(CameraRig->AllocationInfo.ContextDataTableInfo);
 
+	// Set the default values.
+	FCameraParameterOverrideHelper::ApplyDefaultBlendableParameters(CameraRig, NewEntry.Result.VariableTable);
+
 	// Set all the data from the context.
 	const FCameraNodeEvaluationResult& ContextResult = EvaluationContext->GetInitialResult();
 	NewEntry.Result.VariableTable.Override(
@@ -668,8 +672,8 @@ void FTransientBlendStackCameraNodeEvaluator::OnRun(const FCameraNodeEvaluationP
 
 void FTransientBlendStackCameraNodeEvaluator::InternalPreBlendPrepare(TArrayView<FResolvedEntry> ResolvedEntries, const FCameraNodeEvaluationParams& Params, FCameraNodeEvaluationResult& OutResult)
 {
-	constexpr ECameraVariableTableFilter VariableTableFilter = ECameraVariableTableFilter::InputOutput | ECameraVariableTableFilter::Private | ECameraVariableTableFilter::ChangedOnly;
-	constexpr ECameraContextDataTableFilter ContextDataTableFilter = ECameraContextDataTableFilter::ChangedOnly;
+	constexpr ECameraVariableTableFilter VariableTableFilter = ECameraVariableTableFilter::InputOutput | ECameraVariableTableFilter::Private;
+	constexpr ECameraContextDataTableFilter ContextDataTableFilter = ECameraContextDataTableFilter::None;
 
 	for (FResolvedEntry& ResolvedEntry : ResolvedEntries)
 	{
@@ -740,6 +744,11 @@ void FTransientBlendStackCameraNodeEvaluator::InternalPreBlendPrepare(TArrayView
 
 void FTransientBlendStackCameraNodeEvaluator::InternalPreBlendExecute(TArrayView<FResolvedEntry> ResolvedEntries, const FCameraNodeEvaluationParams& Params, FCameraNodeEvaluationResult& OutResult)
 {
+	// Blend all the camera rigs' input variables (include private variables such as camera rig parameters).
+	PreBlendVariableTable.ClearAllWrittenThisFrameFlags();
+
+	const ECameraVariableTableFilter VariableTableFilter = ECameraVariableTableFilter::Input | ECameraVariableTableFilter::Private;
+
 	for (FResolvedEntry& ResolvedEntry : ResolvedEntries)
 	{
 		FCameraRigEntry& Entry(ResolvedEntry.Entry);
@@ -751,8 +760,9 @@ void FTransientBlendStackCameraNodeEvaluator::InternalPreBlendExecute(TArrayView
 			CurParams.EvaluationContext = ResolvedEntry.Context;
 			CurParams.bIsFirstFrame = Entry.bIsFirstFrame;
 			FCameraNodePreBlendParams PreBlendParams(CurParams, CurResult.CameraPose, CurResult.VariableTable);
+			PreBlendParams.VariableTableFilter = VariableTableFilter;
 
-			FCameraNodePreBlendResult PreBlendResult(OutResult.VariableTable);
+			FCameraNodePreBlendResult PreBlendResult(PreBlendVariableTable);
 
 			FBlendCameraNodeEvaluator* EntryBlendEvaluator = Entry.RootEvaluator->GetBlendEvaluator();
 			if (EntryBlendEvaluator)
@@ -761,13 +771,26 @@ void FTransientBlendStackCameraNodeEvaluator::InternalPreBlendExecute(TArrayView
 			}
 			else
 			{
-				OutResult.VariableTable.Override(CurResult.VariableTable, ECameraVariableTableFilter::Input);
+				PreBlendVariableTable.Override(CurResult.VariableTable, PreBlendParams.VariableTableFilter);
 			}
 		}
 		else
 		{
 			// Frozen entries still contribute to the blend using their last evaluated values.
-			OutResult.VariableTable.Override(CurResult.VariableTable, ECameraVariableTableFilter::Input);
+			PreBlendVariableTable.Override(CurResult.VariableTable, VariableTableFilter);
+		}
+	}
+
+	// Write the values back to each entry table, so that each of these camera rigs will run with
+	// the pre-blended values. We limit this writing to the variables each of them knows, since
+	// there's no need to add entries they don't use to their variable table.
+	for (FResolvedEntry& ResolvedEntry : ResolvedEntries)
+	{
+		FCameraRigEntry& Entry(ResolvedEntry.Entry);
+
+		if (!Entry.bIsFrozen)
+		{
+			Entry.Result.VariableTable.Override(PreBlendVariableTable, VariableTableFilter | ECameraVariableTableFilter::KnownOnly);
 		}
 	}
 }
@@ -807,6 +830,7 @@ void FTransientBlendStackCameraNodeEvaluator::InternalUpdate(TArrayView<FResolve
 				CurResult.bIsCameraCut |= ActiveOnlyResult->bIsCameraCut;
 			}
 		}
+		
 		CurResult.bIsValid = true;
 
 #if WITH_EDITOR || UE_GAMEPLAY_CAMERAS_DEBUG
diff --git a/Engine/Plugins/Cameras/GameplayCameras/Source/GameplayCameras/Private/Core/BlendStackRootCameraNode.cpp b/Engine/Plugins/Cameras/GameplayCameras/Source/GameplayCameras/Private/Core/BlendStackRootCameraNode.cpp
--- a/Engine/Plugins/Cameras/GameplayCameras/Source/GameplayCameras/Private/Core/BlendStackRootCameraNode.cpp
+++ b/Engine/Plugins/Cameras/GameplayCameras/Source/GameplayCameras/Private/Core/BlendStackRootCameraNode.cpp
@@ -10,6 +10,7 @@
 #include "Debug/CameraDebugBlock.h"
 #include "Debug/CameraDebugBlockBuilder.h"
 #include "Debug/CameraDebugRenderer.h"
+#include "Helpers/CameraParameterOverrideHelper.h"
 #include "Helpers/CameraRigParameterOverrideEvaluator.h"
 #include "Nodes/Common/CameraRigCameraNode.h"
 
@@ -171,7 +172,7 @@ void FBlendStackRootCameraNodeEvaluator::MergeCameraRig(const FCameraNodeEvaluat
 
 void FBlendStackRootCameraNodeEvaluator::InitializeBlendedParameterOverridesStack()
 {
-	if (!ensureMsgf(BlendablePrefabCameraRig, TEXT("The blended parameter overiddes stack has already been initialized.")))
+	if (!ensureMsgf(BlendablePrefabCameraRig, TEXT("The blended parameter overrides stack has already been initialized.")))
 	{
 		return;
 	}
@@ -224,33 +225,6 @@ void FBlendStackRootCameraNodeEvaluator::OnRun(const FCameraNodeEvaluationParams
 	}
 }
 
-void FBlendStackRootCameraNodeEvaluator::SetDefaultBlendableParameterValues(FCameraVariableTable& OutVariableTable)
-{
-	const FInstancedPropertyBag& DefaultParameters = BlendablePrefabCameraRig->GetDefaultParameters();
-	const uint8* RawDefaultParametersContainer = DefaultParameters.GetValue().GetMemory();
-
-	for (const FCameraRigParameterDefinition& Definition : BlendablePrefabCameraRig->GetParameterDefinitions())
-	{
-		if (Definition.ParameterType != ECameraRigInterfaceParameterType::Blendable)
-		{
-			continue;
-		}
-		if (!Definition.VariableID.IsValid())
-		{
-			continue;
-		}
-
-		const FPropertyBagPropertyDesc* PropertyDesc = DefaultParameters.FindPropertyDescByID(Definition.ParameterGuid);
-		if (!ensure(PropertyDesc && PropertyDesc->CachedProperty))
-		{
-			continue;
-		}
-
-		const void* RawValuePtr = PropertyDesc->CachedProperty->ContainerPtrToValuePtr<void>(RawDefaultParametersContainer);
-		OutVariableTable.SetValue(Definition.VariableID, Definition.VariableType, Definition.BlendableStructType, (const uint8*)RawValuePtr);
-	}
-}
-
 void FBlendStackRootCameraNodeEvaluator::RunBlendedParameterOverridesStack(const FCameraBlendedParameterUpdateParams& Params, FCameraBlendedParameterUpdateResult& OutResult)
 {
 	if (BlendedParameterOverridesStack.IsEmpty())
@@ -271,7 +245,7 @@ void FBlendStackRootCameraNodeEvaluator::RunBlendedParameterOverridesStack(const
 
 		// Start by setting the default values of all parameters. If we don't do this, parameter overrides
 		// wouldn't have a base value to blend from.
-		SetDefaultBlendableParameterValues(CurResult.VariableTable);
+		FCameraParameterOverrideHelper::ApplyDefaultBlendableParameters(BlendablePrefabCameraRig, CurResult.VariableTable);
 
 		// Next, override the defaults with the specific values of this entry.
 		FCameraRigParameterOverrideEvaluator OverrideEvaluator(BlendedParameterOverrides.PrefabNode->CameraRigReference);
@@ -283,7 +257,7 @@ void FBlendStackRootCameraNodeEvaluator::RunBlendedParameterOverridesStack(const
 			BlendedParameterOverrides.BlendEvaluator->Run(Params.EvaluationParams, CurResult);
 
 			FCameraNodePreBlendParams BlendParams(Params.EvaluationParams, Params.LastCameraPose, CurResult.VariableTable);
-			BlendParams.ExtraVariableTableFilter = ECameraVariableTableFilter::Private;
+			BlendParams.VariableTableFilter = ECameraVariableTableFilter::Input | ECameraVariableTableFilter::Private;
 			FCameraNodePreBlendResult BlendResult(OutResult.VariableTable);
 			BlendedParameterOverrides.BlendEvaluator->BlendParameters(BlendParams, BlendResult);
 
diff --git a/Engine/Plugins/Cameras/GameplayCameras/Source/GameplayCameras/Private/Core/BlendStackRootCameraNode.h b/Engine/Plugins/Cameras/GameplayCameras/Source/GameplayCameras/Private/Core/BlendStackRootCameraNode.h
--- a/Engine/Plugins/Cameras/GameplayCameras/Source/GameplayCameras/Private/Core/BlendStackRootCameraNode.h
+++ b/Engine/Plugins/Cameras/GameplayCameras/Source/GameplayCameras/Private/Core/BlendStackRootCameraNode.h
@@ -84,7 +84,6 @@ protected:
 private:
 
 	void InitializeBlendedParameterOverridesStack();
-	void SetDefaultBlendableParameterValues(FCameraVariableTable& OutVariableTable);
 	void RunBlendedParameterOverridesStack(const FCameraBlendedParameterUpdateParams& Params, FCameraBlendedParameterUpdateResult& OutResult);
 
 private:
diff --git a/Engine/Plugins/Cameras/GameplayCameras/Source/GameplayCameras/Private/Helpers/CameraParameterOverrideHelper.cpp b/Engine/Plugins/Cameras/GameplayCameras/Source/GameplayCameras/Private/Helpers/CameraParameterOverrideHelper.cpp
--- a/Engine/Plugins/Cameras/GameplayCameras/Source/GameplayCameras/Private/Helpers/CameraParameterOverrideHelper.cpp
+++ b/Engine/Plugins/Cameras/GameplayCameras/Source/GameplayCameras/Private/Helpers/CameraParameterOverrideHelper.cpp
@@ -7,7 +7,6 @@
 #include "Core/CameraRigAsset.h"
 #include "Core/CameraRigParameterDefinition.h"
 #include "Core/CameraVariableTable.h"
-#include "IGameplayCamerasModule.h"
 
 namespace UE::Cameras
 {
@@ -216,8 +215,6 @@ FCameraParameterOverrideHelper::FCameraParameterOverrideHelper(FCameraVariableTa
 	: VariableTable(OutVariableTable)
 	, ContextDataTable(OutContextDataTable)
 {
-	IGameplayCamerasModule& GameplayCamerasModule = IGameplayCamerasModule::Get();
-	BlendableStructs = GameplayCamerasModule.GetBlendableStructs();
 }
 
 void FCameraParameterOverrideHelper::ApplyParameterOverride(
@@ -252,5 +249,37 @@ void FCameraParameterOverrideHelper::ApplyParameterOverride(
 	}
 }
 
+void FCameraParameterOverrideHelper::ApplyDefaultBlendableParameters(const UCameraRigAsset* CameraRig, FCameraVariableTable& OutVariableTable)
+{
+	if (!ensure(CameraRig))
+	{
+		return;
+	}
+
+	const FInstancedPropertyBag& DefaultParameters = CameraRig->GetDefaultParameters();
+	const uint8* RawDefaultParametersContainer = DefaultParameters.GetValue().GetMemory();
+
+	for (const FCameraRigParameterDefinition& Definition : CameraRig->GetParameterDefinitions())
+	{
+		if (Definition.ParameterType != ECameraRigInterfaceParameterType::Blendable)
+		{
+			continue;
+		}
+		if (!Definition.VariableID.IsValid())
+		{
+			continue;
+		}
+
+		const FPropertyBagPropertyDesc* PropertyDesc = DefaultParameters.FindPropertyDescByID(Definition.ParameterGuid);
+		if (!ensure(PropertyDesc && PropertyDesc->CachedProperty))
+		{
+			continue;
+		}
+
+		const void* RawValuePtr = PropertyDesc->CachedProperty->ContainerPtrToValuePtr<void>(RawDefaultParametersContainer);
+		OutVariableTable.SetValue(Definition.VariableID, Definition.VariableType, Definition.BlendableStructType, (const uint8*)RawValuePtr);
+	}
+}
+
 }  // namespace UE::Cameras
 
diff --git a/Engine/Plugins/Cameras/GameplayCameras/Source/GameplayCameras/Private/Helpers/CameraParameterOverrideHelper.h b/Engine/Plugins/Cameras/GameplayCameras/Source/GameplayCameras/Private/Helpers/CameraParameterOverrideHelper.h
--- a/Engine/Plugins/Cameras/GameplayCameras/Source/GameplayCameras/Private/Helpers/CameraParameterOverrideHelper.h
+++ b/Engine/Plugins/Cameras/GameplayCameras/Source/GameplayCameras/Private/Helpers/CameraParameterOverrideHelper.h
@@ -2,7 +2,6 @@
 
 #pragma once
 
-#include "Core/CameraVariableTableFwd.h"
 #include "StructUtils/PropertyBag.h"
 
 class UCameraRigAsset;
@@ -29,12 +28,12 @@ struct FCameraParameterOverrideHelper
 			const FPropertyBagPropertyDesc& PropertyBagPropertyDesc,
 			bool bDrivenOverridesOnly);
 
+	static void ApplyDefaultBlendableParameters(const UCameraRigAsset* CameraRig, FCameraVariableTable& OutVariableTable);
+
 private:
 
 	FCameraVariableTable& VariableTable;
 	FCameraContextDataTable& ContextDataTable;
-
-	TArray<FBlendableStructInfo> BlendableStructs;
 };
 
 }  // namespace UE::Cameras
diff --git a/Engine/Plugins/Cameras/GameplayCameras/Source/GameplayCameras/Private/Nodes/Blends/PopBlendCameraNode.cpp b/Engine/Plugins/Cameras/GameplayCameras/Source/GameplayCameras/Private/Nodes/Blends/PopBlendCameraNode.cpp
--- a/Engine/Plugins/Cameras/GameplayCameras/Source/GameplayCameras/Private/Nodes/Blends/PopBlendCameraNode.cpp
+++ b/Engine/Plugins/Cameras/GameplayCameras/Source/GameplayCameras/Private/Nodes/Blends/PopBlendCameraNode.cpp
@@ -36,7 +36,7 @@ void FPopBlendCameraNodeEvaluator::OnBlendResults(const FCameraNodeBlendParams&
 void FPopBlendCameraNodeHelper::PopParameters(const UE::Cameras::FCameraNodePreBlendParams& Params, FCameraNodePreBlendResult& OutResult)
 {
 	const FCameraVariableTable& ChildVariableTable(Params.ChildVariableTable);
-	OutResult.VariableTable.Override(ChildVariableTable, ECameraVariableTableFilter::Input);
+	OutResult.VariableTable.Override(ChildVariableTable, Params.VariableTableFilter);
 
 	OutResult.bIsBlendFull = true;
 	OutResult.bIsBlendFinished = true;
diff --git a/Engine/Plugins/Cameras/GameplayCameras/Source/GameplayCameras/Private/Nodes/Blends/SimpleBlendCameraNode.cpp b/Engine/Plugins/Cameras/GameplayCameras/Source/GameplayCameras/Private/Nodes/Blends/SimpleBlendCameraNode.cpp
--- a/Engine/Plugins/Cameras/GameplayCameras/Source/GameplayCameras/Private/Nodes/Blends/SimpleBlendCameraNode.cpp
+++ b/Engine/Plugins/Cameras/GameplayCameras/Source/GameplayCameras/Private/Nodes/Blends/SimpleBlendCameraNode.cpp
@@ -29,7 +29,7 @@ void FSimpleBlendCameraNodeEvaluator::OnRun(const FCameraNodeEvaluationParams& P
 void FSimpleBlendCameraNodeEvaluator::OnBlendParameters(const FCameraNodePreBlendParams& Params, FCameraNodePreBlendResult& OutResult)
 {
 	const FCameraVariableTable& ChildVariableTable(Params.ChildVariableTable);
-	OutResult.VariableTable.Lerp(ChildVariableTable, ECameraVariableTableFilter::Input | Params.ExtraVariableTableFilter, BlendFactor);
+	OutResult.VariableTable.Lerp(ChildVariableTable, Params.VariableTableFilter, BlendFactor);
 
 	OutResult.bIsBlendFull = BlendFactor >= 1.f;
 	OutResult.bIsBlendFinished = bIsBlendFinished;
diff --git a/Engine/Plugins/Cameras/GameplayCameras/Source/GameplayCameras/Public/Core/BlendCameraNode.h b/Engine/Plugins/Cameras/GameplayCameras/Source/GameplayCameras/Public/Core/BlendCameraNode.h
--- a/Engine/Plugins/Cameras/GameplayCameras/Source/GameplayCameras/Public/Core/BlendCameraNode.h
+++ b/Engine/Plugins/Cameras/GameplayCameras/Source/GameplayCameras/Public/Core/BlendCameraNode.h
@@ -39,8 +39,8 @@ struct FCameraNodePreBlendParams
 	const FCameraPose& LastCameraPose;
 	/** The variable table of the node tree being blended. */
 	const FCameraVariableTable& ChildVariableTable;
-	/** Extra filter for variable table blending. */
-	ECameraVariableTableFilter ExtraVariableTableFilter = ECameraVariableTableFilter::None;
+	/** The filter to use for variable table blending. */
+	ECameraVariableTableFilter VariableTableFilter = ECameraVariableTableFilter::None;
 };
 
 /**
diff --git a/Engine/Plugins/Cameras/GameplayCameras/Source/GameplayCameras/Public/Core/BlendStackCameraNode.h b/Engine/Plugins/Cameras/GameplayCameras/Source/GameplayCameras/Public/Core/BlendStackCameraNode.h
--- a/Engine/Plugins/Cameras/GameplayCameras/Source/GameplayCameras/Public/Core/BlendStackCameraNode.h
+++ b/Engine/Plugins/Cameras/GameplayCameras/Source/GameplayCameras/Public/Core/BlendStackCameraNode.h
@@ -302,6 +302,10 @@ private:
 	void PushVariantEntry(const FBlendStackCameraPushParams& Params, const UCameraRigTransition* Transition);
 	void PushNewEntry(const FBlendStackCameraPushParams& Params, const UCameraRigTransition* Transition);
 
+private:
+
+	/** Variable table for pre-blending. */
+	FCameraVariableTable PreBlendVariableTable;
 };
 
 /**
