Cameras: framing nodes improvements

- The dead zone is a real dead zone now but there are disengage/reengage timers.
- Better interpolation between low/high damping factors.
- Use a target info structure for acquiring the target location.

diff --git a/Engine/Plugins/Cameras/GameplayCameras/Source/GameplayCameras/Private/GameplayCamerasModule.cpp b/Engine/Plugins/Cameras/GameplayCameras/Source/GameplayCameras/Private/GameplayCamerasModule.cpp
--- a/Engine/Plugins/Cameras/GameplayCameras/Source/GameplayCameras/Private/GameplayCamerasModule.cpp
+++ b/Engine/Plugins/Cameras/GameplayCameras/Source/GameplayCameras/Private/GameplayCamerasModule.cpp
@@ -6,7 +6,9 @@
 #include "GameplayCameras.h"
 #include "Logging/MessageLog.h"
 #include "Modules/ModuleManager.h"
+#include "Nodes/Framing/CameraFramingZone.h"
 #include "ShowFlags.h"
+#include "UObject/UObjectBase.h"
 
 #define LOCTEXT_NAMESPACE "GameplayCamerasModule"
 
@@ -31,6 +33,8 @@ public:
 	// IModuleInterface interface
 	virtual void StartupModule() override
 	{
+		RegisterBuiltInBlendableStructs();
+
 #if UE_GAMEPLAY_CAMERAS_DEBUG
 		UE::Cameras::FCameraDebugColors::RegisterBuiltinColorSchemes();
 #endif  // UE_GAMEPLAY_CAMERAS_DEBUG
@@ -38,6 +42,7 @@ public:
 
 	virtual void ShutdownModule() override
 	{
+		UnregisterBuiltInBlendableStructs();
 	}
 
 public:
@@ -91,6 +96,25 @@ public:
 	}
 #endif
 
+private:
+
+	void RegisterBuiltInBlendableStructs()
+	{
+		using namespace UE::Cameras;
+
+		RegisterBlendableStruct(FCameraFramingZone::StaticStruct(), &FCameraFramingZone::TypeErasedInterpolate);
+	}
+
+	void UnregisterBuiltInBlendableStructs()
+	{
+		using namespace UE::Cameras;
+
+		if (UObjectInitialized())
+		{
+			UnregisterBlendableStruct(FCameraFramingZone::StaticStruct());
+		}
+	}
+
 private:
 
 	TArray<UE::Cameras::FBlendableStructInfo> BlendableStructs;
diff --git a/Engine/Plugins/Cameras/GameplayCameras/Source/GameplayCameras/Private/Math/CameraFramingZoneMath.cpp b/Engine/Plugins/Cameras/GameplayCameras/Source/GameplayCameras/Private/Math/CameraFramingZoneMath.cpp
--- a/Engine/Plugins/Cameras/GameplayCameras/Source/GameplayCameras/Private/Math/CameraFramingZoneMath.cpp
+++ b/Engine/Plugins/Cameras/GameplayCameras/Source/GameplayCameras/Private/Math/CameraFramingZoneMath.cpp
@@ -8,6 +8,34 @@
 namespace UE::Cameras
 {
 
+FFramingZone FFramingZone::FromScreenMargins(const FCameraFramingZone& Margins)
+{
+	FFramingZone NewZone;
+
+	NewZone.LeftBound = Margins.Left;
+	NewZone.TopBound = Margins.Top;
+	NewZone.RightBound = 1.0 - Margins.Right;
+	NewZone.BottomBound = 1.0 - Margins.Bottom;
+
+	NewZone.ClampBounds();
+
+	return NewZone;
+}
+
+FFramingZone FFramingZone::FromRelativeMargins(const FVector2d& ScreenLocation, const FCameraFramingZone& Margins)
+{
+	FFramingZone NewZone;
+
+	NewZone.LeftBound = ScreenLocation.X - Margins.Left;
+	NewZone.TopBound = ScreenLocation.Y - Margins.Top;
+	NewZone.RightBound = ScreenLocation.X + Margins.Right;
+	NewZone.BottomBound = ScreenLocation.Y + Margins.Bottom;
+
+	NewZone.ClampBounds();
+
+	return NewZone;
+}
+
 FFramingZone::FFramingZone()
 {
 	LeftBound = 0;
@@ -16,16 +44,6 @@ FFramingZone::FFramingZone()
 	BottomBound = 1;
 }
 
-FFramingZone::FFramingZone(const FFramingZoneMargins& FramingZoneMargins)
-{
-	LeftBound = FramingZoneMargins.LeftMargin;
-	TopBound = FramingZoneMargins.TopMargin;
-	RightBound = 1.0 - FramingZoneMargins.RightMargin;
-	BottomBound = 1.0 - FramingZoneMargins.BottomMargin;
-
-	ClampBounds();
-}
-
 void FFramingZone::ClampBounds()
 {
 	LeftBound = FMath::Clamp(LeftBound, 0.0, 1.0);
diff --git a/Engine/Plugins/Cameras/GameplayCameras/Source/GameplayCameras/Private/Nodes/Framing/BaseFramingCameraNode.cpp b/Engine/Plugins/Cameras/GameplayCameras/Source/GameplayCameras/Private/Nodes/Framing/BaseFramingCameraNode.cpp
--- a/Engine/Plugins/Cameras/GameplayCameras/Source/GameplayCameras/Private/Nodes/Framing/BaseFramingCameraNode.cpp
+++ b/Engine/Plugins/Cameras/GameplayCameras/Source/GameplayCameras/Private/Nodes/Framing/BaseFramingCameraNode.cpp
@@ -15,17 +15,31 @@
 #include "HAL/IConsoleManager.h"
 #include "Math/CameraPoseMath.h"
 #include "Math/ColorList.h"
+#include "Misc/AssertionMacros.h"
+#include "Misc/CoreMiscDefines.h"
 
 #include UE_INLINE_GENERATED_CPP_BY_NAME(BaseFramingCameraNode)
 
 namespace UE::Cameras
 {
 
-float GFramingUnlockRadiusEpsilon = 1.e-4;
+float GFramingIdealReachedEpsilon = 0.001;
 static FAutoConsoleVariableRef CVarFramingUnlockRadiusEpsilon(
-	TEXT("GameplayCameras.Framing.UnlockRadiusEpsilon"),
-	GFramingUnlockRadiusEpsilon,
-	TEXT("(Default: 0.0001. The epsilon to determine whether we have reached the framing unlock circle."));
+	TEXT("GameplayCameras.Framing.IdealReachedEpsilon"),
+	GFramingIdealReachedEpsilon,
+	TEXT("(Default: 0.001. The epsilon to determine whether we have reached ideal screen framing."));
+
+float GFramingExpolationEpsilon = 0.001;
+static FAutoConsoleVariableRef CVarFramingExtrapolationEpsilon(
+	TEXT("GameplayCameras.Framing.ExtrapolationEpsilon"),
+	GFramingIdealReachedEpsilon,
+	TEXT("(Default: 0.001. The epsilon to determine whether target movement extrapolation should be included."));
+
+int32 GFramingNumTargetMovementSamples = 10;
+static FAutoConsoleVariableRef CVarFramingNumTargetMovementSamples(
+	TEXT("GameplayCameras.Framing.NumTargetMovementSamples"),
+	GFramingNumTargetMovementSamples,
+	TEXT("(Default: 10. The number of samples to use to extrapolate target movement."));
 
 UE_DECLARE_CAMERA_DEBUG_BLOCK_START(GAMEPLAYCAMERAS_API, FBaseFramingCameraDebugBlock)
 	UE_DECLARE_CAMERA_DEBUG_BLOCK_FIELD(FBaseFramingCameraNodeEvaluator::FState, State);
@@ -34,38 +48,25 @@ UE_DECLARE_CAMERA_DEBUG_BLOCK_END()
 
 UE_DEFINE_CAMERA_DEBUG_BLOCK_WITH_FIELDS(FBaseFramingCameraDebugBlock)
 
-void FCameraFramingZoneParameterReader::Initialize(const FCameraFramingZone& FramingZone)
-{
-	LeftMargin.Initialize(FramingZone.LeftMargin);
-	TopMargin.Initialize(FramingZone.TopMargin);
-	RightMargin.Initialize(FramingZone.RightMargin);
-	BottomMargin.Initialize(FramingZone.BottomMargin);
-}
-
-FFramingZoneMargins FCameraFramingZoneParameterReader::GetZoneMargins(const FCameraVariableTable& VariableTable) const
-{
-	FFramingZoneMargins Margins;
-	Margins.LeftMargin = LeftMargin.Get(VariableTable);
-	Margins.TopMargin = TopMargin.Get(VariableTable);
-	Margins.RightMargin = RightMargin.Get(VariableTable);
-	Margins.BottomMargin = BottomMargin.Get(VariableTable);
-	return Margins;
-}
-
 UE_DEFINE_CAMERA_NODE_EVALUATOR(FBaseFramingCameraNodeEvaluator)
 
 void FBaseFramingCameraNodeEvaluator::OnInitialize(const FCameraNodeEvaluatorInitializeParams& Params, FCameraNodeEvaluationResult& OutResult)
 {
 	const UBaseFramingCameraNode* BaseFramingNode = GetCameraNodeAs<UBaseFramingCameraNode>();
 
-	Readers.HorizontalFraming.Initialize(BaseFramingNode->HorizontalFraming);
-	Readers.VerticalFraming.Initialize(BaseFramingNode->VerticalFraming);
+	Readers.IdealFramingLocation.Initialize(BaseFramingNode->IdealFramingLocation);
+	Readers.SetTargetDistance.Initialize(BaseFramingNode->SetTargetDistance);
+
 	Readers.ReframeDampingFactor.Initialize(BaseFramingNode->ReframeDampingFactor);
 	Readers.LowReframeDampingFactor.Initialize(BaseFramingNode->LowReframeDampingFactor);
-	Readers.ReframeUnlockRadius.Initialize(BaseFramingNode->ReframeUnlockRadius);
+	Readers.ReengageTime.Initialize(BaseFramingNode->ReengageTime);
+	Readers.DisengageTime.Initialize(BaseFramingNode->DisengageTime);
+	Readers.TargetMovementAnticipationTime.Initialize(BaseFramingNode->TargetMovementAnticipationTime);
+
+	Readers.DeadZone.Initialize(BaseFramingNode->DeadZone);
+	Readers.SoftZone.Initialize(BaseFramingNode->SoftZone);
 
-	Readers.DeadZoneMargin.Initialize(BaseFramingNode->DeadZone);
-	Readers.SoftZoneMargin.Initialize(BaseFramingNode->SoftZone);
+	Anticipation.History.Reserve(GFramingNumTargetMovementSamples);
 }
 
 TOptional<FVector3d> FBaseFramingCameraNodeEvaluator::AcquireTargetLocation(const FCameraNodeEvaluationParams& Params, const FCameraNodeEvaluationResult& InResult)
@@ -73,9 +74,19 @@ TOptional<FVector3d> FBaseFramingCameraNodeEvaluator::AcquireTargetLocation(cons
 	const UBaseFramingCameraNode* FramingNode = GetCameraNodeAs<UBaseFramingCameraNode>();
 	if (FramingNode->TargetLocation.IsValid())
 	{
-		FVector3d TargetLocation;
-		const bool bGotTargetLocation = InResult.VariableTable.TryGetValue(FramingNode->TargetLocation.VariableID, TargetLocation);
-		return bGotTargetLocation ? TOptional<FVector3d>(TargetLocation) : TOptional<FVector3d>();
+		const FVector3d* TargetLocation = FramingNode->TargetLocation.GetValue(InResult.VariableTable);
+		return TargetLocation ? TOptional<FVector3d>(*TargetLocation) : TOptional<FVector3d>();
+	}
+	else if (FramingNode->TargetInfo.Value.Actor || FramingNode->TargetInfo.HasOverride())
+	{
+		const FCameraTargetInfo& TargetInfo = FramingNode->TargetInfo.GetValue(InResult.ContextDataTable);
+
+		FTransform3d TargetTransform;
+		FCameraTargetInfoReader Reader(TargetInfo);
+		if (Reader.GetTargetTransform(TargetTransform))
+		{
+			return TOptional<FVector3d>(TargetTransform.GetLocation());
+		}
 	}
 	else if (APlayerController* PlayerController = Params.EvaluationContext->GetPlayerController())
 	{
@@ -88,25 +99,50 @@ TOptional<FVector3d> FBaseFramingCameraNodeEvaluator::AcquireTargetLocation(cons
 }
 
 void FBaseFramingCameraNodeEvaluator::UpdateFramingState(const FCameraNodeEvaluationParams& Params, const FCameraNodeEvaluationResult& OutResult, const FVector3d& TargetLocation, const FTransform3d& LastFraming)
+{
+	ComputeCurrentState(Params, OutResult, TargetLocation, LastFraming);
+	ComputeDesiredState(Params.DeltaTime);
+}
+
+void FBaseFramingCameraNodeEvaluator::EndFramingUpdate(const FCameraNodeEvaluationParams& Params, FCameraNodeEvaluationResult& OutResult)
+{
+	if (Readers.SetTargetDistance.Get(OutResult.VariableTable))
+	{
+		FCameraPose& OutCameraPose = OutResult.CameraPose;
+		const double TargetDistance = FVector3d::Distance(State.WorldTarget, OutCameraPose.GetLocation());
+		OutCameraPose.SetTargetDistance(TargetDistance);
+	}
+}
+
+void FBaseFramingCameraNodeEvaluator::ComputeCurrentState(const FCameraNodeEvaluationParams& Params, const FCameraNodeEvaluationResult& OutResult, const FVector3d& TargetLocation, const FTransform3d& LastFraming)
 {
 	// Get screen-space coordinates of the ideal framing point. These are in 0..1 UI space.
-	State.IdealTarget.X = Readers.HorizontalFraming.Get(OutResult.VariableTable);
-	State.IdealTarget.Y = Readers.VerticalFraming.Get(OutResult.VariableTable);
+	State.IdealTarget = Readers.IdealFramingLocation.Get(OutResult.VariableTable);
 
-	// Update the damping factors and unlock radius in case they are driven by a variable.
+	// Update the damping factors and reengage/disengage times in case they are driven by a variable.
 	State.ReframeDampingFactor = Readers.ReframeDampingFactor.Get(OutResult.VariableTable);
 	State.LowReframeDampingFactor = Readers.LowReframeDampingFactor.Get(OutResult.VariableTable);
-	State.ReframeUnlockRadius = Readers.ReframeUnlockRadius.Get(OutResult.VariableTable);
+	State.ReengageTime = Readers.ReengageTime.Get(OutResult.VariableTable);
+	State.DisengageTime = Readers.DisengageTime.Get(OutResult.VariableTable);
+	State.TargetMovementAnticipationTime = Readers.TargetMovementAnticipationTime.Get(OutResult.VariableTable);
 
 	// Get the effective margins of the framing zones for this frame.
-	const FFramingZoneMargins DeadZoneMargins = Readers.DeadZoneMargin.GetZoneMargins(OutResult.VariableTable);
-	const FFramingZoneMargins SoftZoneMargins = Readers.SoftZoneMargin.GetZoneMargins(OutResult.VariableTable);
+	const FCameraFramingZone DeadZone = Readers.DeadZone.Get(OutResult.VariableTable);
+	const FCameraFramingZone SoftZone = Readers.SoftZone.Get(OutResult.VariableTable);
 
 	// Compute the UI space coordinates of the framing zones.
-	State.DeadZone = FFramingZone(DeadZoneMargins);
-	State.SoftZone = FFramingZone(SoftZoneMargins);
+	State.DeadZone = FFramingZone::FromRelativeMargins(State.IdealTarget, DeadZone);
+	State.SoftZone = FFramingZone::FromScreenMargins(SoftZone);
+
+	// Make sure our framing zones are hierarchically correct: soft zone contains the dead zone, and the
+	// dead zone contains the ideal target.
+	State.DeadZone.ClampBounds(State.IdealTarget);
+	State.SoftZone.ClampBounds(State.DeadZone);
 
-	// We are going to reframing things iteratively, so we'll use a temporary pose defined by last frame's
+	// See if we need to extrapolate where the target will be in "anticipation time" seconds.
+	const FVector3d EffectiveTargetLocation = ComputeAnticipatedScreenTarget(Params.DeltaTime, TargetLocation);
+
+	// We are going to reframe things iteratively, so we'll use a temporary pose defined by last frame's
 	// shot transform.
 	FCameraPose TempPose(OutResult.CameraPose);
 	TempPose.SetTransform(LastFraming);
@@ -114,34 +150,53 @@ void FBaseFramingCameraNodeEvaluator::UpdateFramingState(const FCameraNodeEvalua
 	// Get the target in screen-space.
 	APlayerController* PlayerController = Params.EvaluationContext->GetPlayerController();
 	const double AspectRatio = FCameraPoseMath::GetEffectiveAspectRatio(TempPose, PlayerController);
-	const TOptional<FVector2d> ScreenTarget = FCameraPoseMath::ProjectWorldToScreen(TempPose, AspectRatio, TargetLocation, true);
-	State.WorldTarget = TargetLocation;
+	const TOptional<FVector2d> ScreenTarget = FCameraPoseMath::ProjectWorldToScreen(TempPose, AspectRatio, EffectiveTargetLocation, true);
+	State.WorldTarget = EffectiveTargetLocation;
 	State.ScreenTarget = ScreenTarget.Get(FVector2d(0.5, 0.5));
 
+#if UE_GAMEPLAY_CAMERAS_DEBUG
+	State.DeadZoneEdgePoint = State.ScreenTarget;
+	State.HardZoneEdgePoint = State.ScreenTarget;
+#endif  // UE_GAMEPLAY_CAMERAS_DEBUG
+
 	// Update the reframe damper's damping factor.
 	if (State.LowReframeDampingFactor <= 0)
 	{
+		// There's no high/low factors, so just set the standard one at 100%.
 		State.ReframeDamper.SetW0(State.ReframeDampingFactor);
+		State.ReframeDampingFactorAlpha = 1.f;
 	}
 	else
 	{
+		// Make a line between the ideal target and the current target. Note how it intersects the 
+		// boundaries dead zone and the hard zone. We will interpolate the damping factors from the
+		// first intersection to the second intersection.
 		const FVector2d IdealToCurrent(State.ScreenTarget - State.IdealTarget);
-		const double UnlockEdgeToCurrent(IdealToCurrent.Length() - State.ReframeUnlockRadius);
+		const double IdealToCurrentDistance = IdealToCurrent.Length();
+		if (IdealToCurrentDistance > 0.0)
+		{
+			const FVector2d IdealToCurrentNormalized = (IdealToCurrent / IdealToCurrentDistance);
+			const FVector2d DeadEdgePoint = State.DeadZone.ComputeClosestIntersection(State.IdealTarget, IdealToCurrentNormalized, true);
+			const FVector2d HardEdgePoint = State.SoftZone.ComputeClosestIntersection(State.IdealTarget, IdealToCurrentNormalized, true);
 
-		const FVector2d HardZonePoint = State.SoftZone.ComputeClosestIntersection(State.IdealTarget, IdealToCurrent);
-		const double UnlockEdgeToHardZone = FMath::Max(
-				FVector2d::Distance(State.IdealTarget, HardZonePoint) - State.ReframeUnlockRadius,
-				UE_DOUBLE_SMALL_NUMBER);
+			const double DeadToCurrent = FVector2d::Distance(DeadEdgePoint, State.ScreenTarget);
+			const double DeadToHardEdge = FVector2d::Distance(DeadEdgePoint, HardEdgePoint);
 
-		const double Alpha = FMath::Clamp(UnlockEdgeToCurrent / UnlockEdgeToHardZone, 0.0, 1.1);
-		State.ReframeDamper.SetW0(FMath::Lerp(State.LowReframeDampingFactor, State.ReframeDampingFactor, Alpha));
-	}
+			const double Alpha = FMath::Clamp(DeadToCurrent / DeadToHardEdge, 0.0, 1.1);
+			State.ReframeDampingFactorAlpha = Alpha;
 
-	// Make sure our framing zones are hierarchically correct: soft zone contains the dead zone, which contains
-	// the ideal target.
-	State.DeadZone.ClampBounds(State.IdealTarget);
-	State.SoftZone.ClampBounds(State.DeadZone);
+#if UE_GAMEPLAY_CAMERAS_DEBUG
+			State.DeadZoneEdgePoint = DeadEdgePoint;
+			State.HardZoneEdgePoint = HardEdgePoint;
+#endif  // UE_GAMEPLAY_CAMERAS_DEBUG
+		}
+		else
+		{
+			State.ReframeDampingFactorAlpha = 0.f;
+		}
+	}
 
+	const bool bWasReframing = State.bIsReframingTarget;
 	const bool bIsInSoftZone = State.SoftZone.Contains(State.ScreenTarget);
 	const bool bIsInDeadZone = State.DeadZone.Contains(State.ScreenTarget);
 	if (!ScreenTarget.IsSet() || !bIsInSoftZone)
@@ -149,41 +204,131 @@ void FBaseFramingCameraNodeEvaluator::UpdateFramingState(const FCameraNodeEvalua
 		// Target is out of view or outside the soft zone -- it's therefore in the hard zone and we will
 		// do everything we can to put it back in the soft zone ASAP.
 		State.TargetFramingState = ETargetFramingState::InHardZone;
+		State.ToggleEngageTimeLeft = 0.f;
+		State.ToggleEngageAlpha = 1.f;
 		State.bIsReframingTarget = true;
 	}
 	else if (!bIsInDeadZone)
 	{
 		// Target is in the soft zone so we will gently reframe it towards the ideal framing.
+		const bool bWasInDeadZone = (State.TargetFramingState == ETargetFramingState::InDeadZone);
 		State.TargetFramingState = ETargetFramingState::InSoftZone;
+
+		// We just exited the dead zone, so start the reengage timer. If we were still disengaging
+		// inside the dead zone, restart the timer from an equivalent time.
+		if (bWasInDeadZone)
+		{
+			if (State.ToggleEngageTimeLeft > 0 && State.DisengageTime > 0.0)
+			{
+				// Example: if we had 60% of time left to disengage, we would want to restart at 60%
+				//			of reengagement, with 40% left to go.
+				const float PreviousAlpha = State.ToggleEngageAlpha;
+				const float TimeRatioLeft = FMath::Clamp(State.ToggleEngageTimeLeft / State.DisengageTime, 0.f, 1.f);
+				State.ToggleEngageTimeLeft = (1.f - TimeRatioLeft) * State.ReengageTime;
+				ensure(FMath::IsNearlyEqual(PreviousAlpha, 1.0f - State.ToggleEngageTimeLeft / State.ReengageTime, 1.e-4f));
+			}
+			else
+			{
+				State.ToggleEngageTimeLeft = State.ReengageTime;
+			}
+		}
+		// If we are reengaging, continue doing so.
+		// The ratio between ToggleEngageTimeLeft and ReengageTime will go from 1 to 0, but we we want
+		// to reengage, i.e. go from 0% to 100%, hence the "one minus..." below.
+		if (State.ToggleEngageTimeLeft > 0.0)
+		{
+			State.ToggleEngageTimeLeft = FMath::Max(0, State.ToggleEngageTimeLeft - Params.DeltaTime);
+			State.ToggleEngageAlpha = 1.0f - FMath::Clamp(State.ToggleEngageTimeLeft / State.ReengageTime, 0.f, 1.f);
+		}
+		else
+		{
+			State.ToggleEngageAlpha = 1.f;
+		}
+		
+		// Put us in reframing mode as soon as we get into the soft zone.
 		State.bIsReframingTarget = true;
 	}
 	else
 	{
 		// Target is in the dead zone.
+		const bool bWasInSoftZone = (State.TargetFramingState == ETargetFramingState::InSoftZone);
 		State.TargetFramingState = ETargetFramingState::InDeadZone;
 
-		// Even though the target is free to move inside the dead zone, we might still want to continue reframing
-		// it towards the ideal position... if we didn't do that, reframing from the soft zone would stop entirely 
-		// once we reach the edge of the dead zone, and we would never really ever see the target near the ideal
-		// position. So if we re-enter the dead zone from the soft zone, we keep reframing until we hit a smaller
-		// "unlock reframing" zone defined by the ReframeUnlockRadius.
-		//
-		// If reframing wasn't active in the first place (e.g. the target is freely moving inside the dead zone), 
-		// we don't do anything and let it be.
-		if (State.bIsReframingTarget)
+		// We just entered the dead zone, so start the disengage timer. If we were still reengaging
+		// inside the soft zone, restart the timer from an equivalent time.
+		if (bWasInSoftZone)
 		{
-			// Since screen-space positions are in 0..1 space, we need to modulate vertical coordinates by the
-			// aspect ratio, otherwise we end up comparing against a squished ellipse instead of a circle.
-			const double DistanceToIdeal = FVector2d::Distance(
-					FVector2d(State.ScreenTarget.X, (State.ScreenTarget.Y - 0.5) / AspectRatio + 0.5),
-					FVector2d(State.IdealTarget.X, (State.IdealTarget.Y - 0.5) / AspectRatio + 0.5));
-			// Add an epsilon to the comparison to avoid being stuck in reframing mode because of
-			// floating point precision issues.
-			if (DistanceToIdeal <= (State.ReframeUnlockRadius + GFramingUnlockRadiusEpsilon))
+			if (State.ToggleEngageTimeLeft > 0 && State.ReengageTime > 0.0)
 			{
-				State.bIsReframingTarget = false;
+				// Example: if we had 30% of time left to reengage, we would want to restart at 30% of 
+				//		    disengagement, with 70% left to go.
+				const float PreviousAlpha = State.ToggleEngageAlpha;
+				const float TimeRatioLeft = FMath::Clamp(State.ToggleEngageTimeLeft / State.ReengageTime, 0.f, 1.f);
+				State.ToggleEngageTimeLeft = (1.f - TimeRatioLeft) * State.DisengageTime;
+				ensure(FMath::IsNearlyEqual(PreviousAlpha, State.ToggleEngageTimeLeft / State.DisengageTime, 1.e-4f));
+			}
+			else
+			{
+				State.ToggleEngageTimeLeft = State.DisengageTime;
 			}
 		}
+		// If we are disengaging, continue doing so.
+		// The ratio between ToggleEngageTimeLeft and DisengageTime will go from 1 to 0, which is what
+		// we want since we want ramp down from 100% to 0%.
+		if (State.ToggleEngageTimeLeft > 0.0)
+		{
+			State.ToggleEngageTimeLeft = FMath::Max(0, State.ToggleEngageTimeLeft - Params.DeltaTime);
+			State.ToggleEngageAlpha = FMath::Clamp(State.ToggleEngageTimeLeft / State.DisengageTime, 0.f, 1.f);
+		}
+		else
+		{
+			State.ToggleEngageAlpha = 0.f;
+		}
+
+		// Only truly disengage reframing once we reach the end of the disengagement time.
+		State.bIsReframingTarget = (State.ToggleEngageAlpha > 0);
+	}
+}
+
+FVector3d FBaseFramingCameraNodeEvaluator::ComputeAnticipatedScreenTarget(float DeltaTime, const FVector3d& InTargetLocation)
+{
+	if (State.TargetMovementAnticipationTime <= 0)
+	{
+		return InTargetLocation;
+	}
+	if (DeltaTime <= 0)
+	{
+		return InTargetLocation;
+	}
+
+	using FHistoryEntry = TTuple<FVector3d, float>;
+
+	const int32 MaxNumSamples = Anticipation.History.Max();
+	if (Anticipation.History.Num() == MaxNumSamples)
+	{
+		Anticipation.History.RemoveAt(0, EAllowShrinking::No);
+	}
+	Anticipation.History.Add({ InTargetLocation, DeltaTime });
+
+	FVector3d AverageMoveStep(EForceInit::ForceInit);
+	for (int32 Index = 1; Index < Anticipation.History.Num(); ++Index)
+	{
+		const FHistoryEntry& Entry(Anticipation.History[Index]);
+		const FHistoryEntry& PrevEntry(Anticipation.History[Index - 1]);
+		const FVector3d NormalizedMoveStep = ((Entry.Key - PrevEntry.Key) / Entry.Value);
+		AverageMoveStep += NormalizedMoveStep;
+	}
+	AverageMoveStep /= Anticipation.History.Num();
+
+	if (AverageMoveStep.Length() >= GFramingExpolationEpsilon)
+	{
+		const FVector3d ExtrapolatedMovement = AverageMoveStep * State.TargetMovementAnticipationTime;
+		const FVector3d ExtrapolatedWorldTarget = InTargetLocation + ExtrapolatedMovement;
+		return ExtrapolatedWorldTarget;
+	}
+	else
+	{
+		return InTargetLocation;
 	}
 }
 
@@ -192,7 +337,7 @@ void FBaseFramingCameraNodeEvaluator::ComputeDesiredState(float DeltaTime)
 	// If we  don't have any reframing to do, bail out.
 	FVector2d IdealToTarget(State.ScreenTarget - State.IdealTarget);
 	double DistanceToGo = IdealToTarget.Length();
-	if (!State.bIsReframingTarget || DistanceToGo == 0.f)
+	if (!State.bIsReframingTarget || DistanceToGo <= GFramingIdealReachedEpsilon)
 	{
 		ensure(State.TargetFramingState == ETargetFramingState::InDeadZone);
 		Desired.ScreenTarget = State.ScreenTarget;
@@ -211,12 +356,20 @@ void FBaseFramingCameraNodeEvaluator::ComputeDesiredState(float DeltaTime)
 		DistanceToGo = IdealToTarget.Length();
 	}
 
-	// Move the target towards the ideal position using some damping.
-	// Remove the radius of the unlock zone from the distance we pass to the damper, 
-	// otherwise the damper won't ever get to smoothly ease out to zero.
-	double DampingDistanceToGo = DistanceToGo - State.ReframeUnlockRadius;
-	const double NewDampedDistanceToGo = State.ReframeDamper.Update(DampingDistanceToGo, DeltaTime);
-	const double NewDistanceToGo = NewDampedDistanceToGo + State.ReframeUnlockRadius;
+	// Figure out the damping factor for this frame. We might have interpolation between
+	// the low and normal (high) damping factors, and then we might have interpolation
+	// between that and 0 for disengaging or reengaging framing.
+	float EffectiveDampingFactor = State.ReframeDampingFactor;
+	if (State.LowReframeDampingFactor > 0)
+	{
+		EffectiveDampingFactor = FMath::Lerp(
+				State.LowReframeDampingFactor, State.ReframeDampingFactor, State.ReframeDampingFactorAlpha);
+	}
+	EffectiveDampingFactor = FMath::Lerp(0.1f, FMath::Max(0.1f, EffectiveDampingFactor), State.ToggleEngageAlpha);
+	State.ReframeDamper.SetW0(EffectiveDampingFactor);
+
+	// Move the target towards the ideal framing using damping.
+	const double NewDistanceToGo = State.ReframeDamper.Update(DistanceToGo, DeltaTime);
 
 	// Compute where we want the target this frame.
 	const FVector2d InvReframeDir(IdealToTarget / DistanceToGo);
@@ -233,83 +386,62 @@ FVector2d FBaseFramingCameraNodeEvaluator::GetHardReframeCoords() const
 	// Let's compute the diagonal between the target and the ideal framing point, and bring
 	// the target where that diagonal intersects the soft zone.
 	const FVector2d Diagonal(State.IdealTarget - State.ScreenTarget);
-	if (Diagonal.X == 0.f && Diagonal.Y == 0.f)
+	if (Diagonal.IsZero())
 	{
 		// Somehow we're already on the desired framing. This shouldn't happen, we're supposed
 		// to be in the hard zone right now...
 		ensure(false);
 		return State.ScreenTarget;
 	}
-	if (Diagonal.X == 0.f)
-	{
-		// The target is directly above/below the ideal position.
-		return FVector2d(
-				State.IdealTarget.X,
-				(Diagonal.Y > 0 ? State.SoftZone.TopBound : State.SoftZone.BottomBound)
-				);
-	}
-	else if (Diagonal.Y == 0.f)
-	{
-		// The target is directly to the left/right of the ideal position.
-		return FVector2d(
-				(Diagonal.X > 0 ? State.SoftZone.LeftBound : State.SoftZone.RightBound),
-				State.IdealTarget.Y
-				);
-	}
-	else
+
+	return State.SoftZone.ComputeClosestIntersection(State.ScreenTarget, Diagonal, false);
+}
+
+#if UE_GAMEPLAY_CAMERAS_DEBUG
+
+void FBaseFramingCameraNodeEvaluator::DrawFramingState(const FState& State, const FDesired& Desired, FCameraDebugRenderer& Renderer)
+{
+	if (!Renderer.IsExternalRendering() && Renderer.HasCanvas())
 	{
-		// The diagonal's equation is P = V*d + P0
-		//
-		//		V = the unit direction vector of the diagonal
-		//		d = the distance
-		//		P0 = a known reference point (we can use the ideal target for this)
-		//
-		// We want to find P where P is on the edge of the soft zone. This means we know
-		// either P.x or P.y depending on the edge. 
-		//
-		// Let's say we deal with a vertical edge... then we know about P.x, because we
-		// want P to be on that edge. We can therefor compute d:
-		//
-		//		P.x = V.x*d + P0.x
-		//		d = (P.x - P0.x) / V.x
-		//
-		// If it's a horizontal edge, we can do the same:
-		//
-		//		P.y = V.y*d + P0.y
-		//		d = (P.y - P0.y) / V.y
-		//
-		// Keep going around the zone until we're sure the point is on the edge.
-		//
-		const FVector2d P0(State.IdealTarget);
-		const FVector2d V(Diagonal.GetSafeNormal());
-		double D = 0.f;
-		FVector2d P(State.ScreenTarget);
-
-		if (P.X < State.SoftZone.LeftBound)
-		{
-			D = (State.SoftZone.LeftBound - P0.X) / V.X;
-			P = V * D + P0;
-		}
-		if (P.Y < State.SoftZone.TopBound)
-		{
-			D = (State.SoftZone.TopBound - P0.Y) / V.Y;
-			P = V * D + P0;
-		}
-		if (P.X > State.SoftZone.RightBound)
-		{
-			D = (State.SoftZone.RightBound - P0.X) / V.X;
-			P = V * D + P0;
-		}
-		if (P.X > State.SoftZone.BottomBound)
+		const FVector2D CanvasSize = Renderer.GetCanvasSize();
+
+		Renderer.Draw2DBox(
+				State.SoftZone.GetCanvasPosition(CanvasSize), 
+				State.SoftZone.GetCanvasSize(CanvasSize),
+				FLinearColor::Red,
+				1.f);
+		Renderer.Draw2DBox(
+				State.DeadZone.GetCanvasPosition(CanvasSize), 
+				State.DeadZone.GetCanvasSize(CanvasSize),
+				FLinearColor::Green,
+				1.f);
+
+		if (State.LowReframeDampingFactor > 0)
 		{
-			D = (State.SoftZone.BottomBound - P0.Y) / V.Y;
-			P = V * D + P0;
+			Renderer.Draw2DLine(State.DeadZoneEdgePoint, State.HardZoneEdgePoint, FLinearColor(FColorList::LightGrey), 1.f);
 		}
 
-		return P;
+		const FVector2D ScreenTarget(State.ScreenTarget.X * CanvasSize.X, State.ScreenTarget.Y * CanvasSize.Y);
+		const FVector2D NextScreenTarget(Desired.ScreenTarget.X * CanvasSize.X, Desired.ScreenTarget.Y * CanvasSize.Y);
+		const FVector2D IdealTarget(State.IdealTarget.X * CanvasSize.X, State.IdealTarget.Y * CanvasSize.Y);
+
+		Renderer.Draw2DLine(ScreenTarget, NextScreenTarget, FLinearColor(FColorList::Salmon), 1.f);
+		Renderer.Draw2DCircle(ScreenTarget, 2.f, FLinearColor(FColorList::Orange), 2.f);
+		Renderer.Draw2DCircle(IdealTarget, 2.f, FLinearColor::Green, 2.f);
 	}
 }
 
+#endif
+
+#if WITH_EDITOR
+
+void FBaseFramingCameraNodeEvaluator::OnDrawEditorPreview(const FCameraEditorPreviewDrawParams& Params, FCameraDebugRenderer& Renderer)
+{
+	DrawFramingState(State, Desired, Renderer);
+}
+
+#endif  // WITH_EDITOR
+
 #if UE_GAMEPLAY_CAMERAS_DEBUG
 
 void FBaseFramingCameraNodeEvaluator::OnBuildDebugBlocks(const FCameraDebugBlockBuildParams& Params, FCameraDebugBlockBuilder& Builder)
@@ -341,51 +473,49 @@ void FBaseFramingCameraDebugBlock::OnDebugDraw(const FCameraDebugBlockDrawParams
 	{
 		Renderer.AddText(TEXT("[REFRAMING]"));
 	}
+	Renderer.NewLine();
 
-	Renderer.AddText(TEXT(" (damping = %0.3f, factor = %0.1f)"), State.ReframeDamper.GetX0(), State.ReframeDamper.GetW0());
-
-	if (Renderer.HasCanvas())
+	Renderer.AddIndent();
 	{
-		const FVector2D CanvasSize = Renderer.GetCanvasSize();
-
-		const float ReframeUnlockRadius(State.ReframeUnlockRadius * CanvasSize.X);
-		const FVector2D IdealTarget(State.IdealTarget.X * CanvasSize.X, State.IdealTarget.Y * CanvasSize.Y);
+		if (Renderer.HasCanvas())
+		{
+			const FVector2D CanvasSize = Renderer.GetCanvasSize();
 
-		Renderer.Draw2DBox(
-				State.SoftZone.GetCanvasPosition(CanvasSize), 
-				State.SoftZone.GetCanvasSize(CanvasSize),
-				FLinearColor::Red,
-				1.f);
-		Renderer.Draw2DBox(
-				State.DeadZone.GetCanvasPosition(CanvasSize), 
-				State.DeadZone.GetCanvasSize(CanvasSize),
-				FLinearColor::Green,
-				1.f);
-		Renderer.Draw2DCircle(IdealTarget, ReframeUnlockRadius, FLinearColor(FColorList::PaleGreen), 1.f);
+			const FVector2D FramingCorrection(Desired.FramingCorrection.X * CanvasSize.X, Desired.FramingCorrection.Y * CanvasSize.Y);
+			Renderer.AddText(TEXT("correction (%0.1f ; %0.1f) "),
+					State.ReframeDamper.GetX0(),
+					FramingCorrection.X, FramingCorrection.Y);
 
-		const FVector2D FramingCorrection(Desired.FramingCorrection.X * CanvasSize.X, Desired.FramingCorrection.Y * CanvasSize.Y);
+			const FVector2D ScreenTarget(State.ScreenTarget.X * CanvasSize.X, State.ScreenTarget.Y * CanvasSize.Y);
+			Renderer.AddText(TEXT("target (%0.1f; %0.1f) "), ScreenTarget.X, ScreenTarget.Y);
+		}
 
-		Renderer.AddText(TEXT("  correction (%0.1f ; %0.1f)"),
-				State.ReframeDamper.GetX0(),
-				FramingCorrection.X, FramingCorrection.Y);
+		Renderer.AddText(TEXT("damping %0.3f (factor = %0.1f)\n"), State.ReframeDamper.GetX0(), State.ReframeDamper.GetW0());
 
-		const FVector2D ScreenTarget(State.ScreenTarget.X * CanvasSize.X, State.ScreenTarget.Y * CanvasSize.Y);
-		const FVector2D NextScreenTarget(Desired.ScreenTarget.X * CanvasSize.X, Desired.ScreenTarget.Y * CanvasSize.Y);
+		Renderer.AddText(TEXT("reengage/disengage time left %.3f (damping factor alpha = %.1f)\n"),
+				State.ToggleEngageTimeLeft,
+				State.ToggleEngageAlpha);
 
-		Renderer.AddText(TEXT("  target (%0.1f; %0.1f)"), ScreenTarget.X, ScreenTarget.Y);
-		Renderer.Draw2DLine(ScreenTarget, NextScreenTarget, FLinearColor(FColorList::Salmon), 1.f);
-		Renderer.Draw2DCircle(ScreenTarget, 2.f, FLinearColor(FColorList::Orange), 2.f);
-		Renderer.Draw2DCircle(IdealTarget, 2.f, FLinearColor::Green, 2.f);
+		Renderer.AddText(TEXT("interpolation = %.1f\n"), State.ReframeDampingFactorAlpha);
 	}
+	Renderer.RemoveIndent();
+
+	FBaseFramingCameraNodeEvaluator::DrawFramingState(State, Desired, Renderer);
 }
 
 #endif  // UE_GAMEPLAY_CAMERAS_DEBUG
 
 void FBaseFramingCameraNodeEvaluator::FState::Serialize(FArchive& Ar)
 {
+	Ar << WorldTarget;
 	Ar << IdealTarget;
 	Ar << ReframeDampingFactor;
-	Ar << ReframeUnlockRadius;
+	Ar << LowReframeDampingFactor;
+	Ar << ReframeDampingFactorAlpha;
+	Ar << ReengageTime;
+	Ar << DisengageTime;
+	Ar << ToggleEngageTimeLeft;
+	Ar << ToggleEngageAlpha;
 	Ar << DeadZone;
 	Ar << SoftZone;
 	Ar << ScreenTarget;
@@ -417,13 +547,38 @@ FArchive& operator <<(FArchive& Ar, FBaseFramingCameraNodeEvaluator::FDesired& D
 
 UBaseFramingCameraNode::UBaseFramingCameraNode(const FObjectInitializer& ObjectInit)
 	: Super(ObjectInit)
-	, DeadZone(0.45f)
-	, SoftZone(0.05f)
+	, DeadZone(0.2)
+	, SoftZone(0.05)
 {
-	HorizontalFraming.Value = 0.5f;
-	VerticalFraming.Value = 0.5f;
+	SetTargetDistance.Value = true;
+	IdealFramingLocation.Value = FVector2d(0.5, 0.5);
 	ReframeDampingFactor.Value = 10.f;
 	LowReframeDampingFactor.Value = -1.f;
-	ReframeUnlockRadius.Value = 0.005f;
+	ReengageTime.Value = 1.f;
+	DisengageTime.Value = 1.f;
+}
+
+void UBaseFramingCameraNode::GetCustomCameraNodeParameters(FCustomCameraNodeParameterInfos& OutParameterInfos)
+{
+	OutParameterInfos.AddDataParameter(
+			GET_MEMBER_NAME_CHECKED(UBaseFramingCameraNode, TargetInfo),
+			ECameraContextDataType::Struct,
+			FCameraTargetInfo::StaticStruct(),
+			reinterpret_cast<const uint8*>(&TargetInfo.Value),
+			&TargetInfo.DataID);
+
+	OutParameterInfos.AddBlendableParameter(
+			GET_MEMBER_NAME_CHECKED(UBaseFramingCameraNode, DeadZone),
+			ECameraVariableType::BlendableStruct,
+			FCameraFramingZone::StaticStruct(),
+			reinterpret_cast<const uint8*>(&DeadZone.Value),
+			&DeadZone.VariableID);
+
+	OutParameterInfos.AddBlendableParameter(
+			GET_MEMBER_NAME_CHECKED(UBaseFramingCameraNode, SoftZone),
+			ECameraVariableType::BlendableStruct,
+			FCameraFramingZone::StaticStruct(),
+			reinterpret_cast<const uint8*>(&SoftZone.Value),
+			&SoftZone.VariableID);
 }
 
diff --git a/Engine/Plugins/Cameras/GameplayCameras/Source/GameplayCameras/Private/Nodes/Framing/CameraFramingZone.cpp b/Engine/Plugins/Cameras/GameplayCameras/Source/GameplayCameras/Private/Nodes/Framing/CameraFramingZone.cpp
new file mode 100644
--- /dev/null
+++ b/Engine/Plugins/Cameras/GameplayCameras/Source/GameplayCameras/Private/Nodes/Framing/CameraFramingZone.cpp
@@ -0,0 +1,20 @@
+// Copyright Epic Games, Inc. All Rights Reserved.
+
+#include "Nodes/Framing/CameraFramingZone.h"
+
+void FCameraFramingZone::TypeErasedInterpolate(uint8* From, const uint8* To, float Alpha)
+{
+	FCameraFramingZone& FromZone = *reinterpret_cast<FCameraFramingZone*>(From);
+	const FCameraFramingZone& ToZone = *reinterpret_cast<const FCameraFramingZone*>(To);
+
+	FromZone.Left = FMath::Lerp(FromZone.Left, ToZone.Left, Alpha);
+	FromZone.Top = FMath::Lerp(FromZone.Top, ToZone.Top, Alpha);
+	FromZone.Right = FMath::Lerp(FromZone.Right, ToZone.Right, Alpha);
+	FromZone.Bottom = FMath::Lerp(FromZone.Bottom, ToZone.Bottom, Alpha);
+}
+
+FString FCameraFramingZone::ToString() const
+{
+	return FString::Printf(TEXT("<< %f ^^ %f >> %f vv %f"), Left, Top, Right, Bottom);
+}
+
diff --git a/Engine/Plugins/Cameras/GameplayCameras/Source/GameplayCameras/Private/Nodes/Framing/CameraTargetInfo.cpp b/Engine/Plugins/Cameras/GameplayCameras/Source/GameplayCameras/Private/Nodes/Framing/CameraTargetInfo.cpp
new file mode 100644
--- /dev/null
+++ b/Engine/Plugins/Cameras/GameplayCameras/Source/GameplayCameras/Private/Nodes/Framing/CameraTargetInfo.cpp
@@ -0,0 +1,49 @@
+// Copyright Epic Games, Inc. All Rights Reserved.
+
+#include "Nodes/Framing/CameraTargetInfo.h"
+
+#include "Components/SkeletalMeshComponent.h"
+#include "Core/CameraContextDataTable.h"
+#include "GameFramework/Actor.h"
+
+const FCameraTargetInfo& FCameraTargetInfoParameter::GetValue(const UE::Cameras::FCameraContextDataTable& ContextDataTable) const
+{
+	if (DataID.IsValid())
+	{
+		if (const FCameraTargetInfo* ActualValue = ContextDataTable.TryGetData<FCameraTargetInfo>(DataID))
+		{
+			return *ActualValue;
+		}
+	}
+	return Value;
+}
+
+namespace UE::Cameras
+{
+
+FCameraTargetInfoReader::FCameraTargetInfoReader(const FCameraTargetInfo& InTargetInfo)
+{
+	Initialize(InTargetInfo);
+}
+
+void FCameraTargetInfoReader::Initialize(const FCameraTargetInfo& InTargetInfo)
+{
+	TargetInfo = InTargetInfo;
+	if (TargetInfo.Actor && !TargetInfo.SocketName.IsNone())
+	{
+		CachedSkeletalMeshComponent = TargetInfo.Actor->FindComponentByClass<USkeletalMeshComponent>();
+	}
+}
+
+bool FCameraTargetInfoReader::GetTargetTransform(FTransform3d& OutTransform) const
+{
+	if (!TargetInfo.SocketName.IsNone() && CachedSkeletalMeshComponent)
+	{
+		OutTransform = CachedSkeletalMeshComponent->GetSocketTransform(TargetInfo.SocketName);
+		return true;
+	}
+	return false;
+}
+
+}  // namespace UE::Cameras
+
diff --git a/Engine/Plugins/Cameras/GameplayCameras/Source/GameplayCameras/Private/Nodes/Framing/DollyFramingCameraNode.cpp b/Engine/Plugins/Cameras/GameplayCameras/Source/GameplayCameras/Private/Nodes/Framing/DollyFramingCameraNode.cpp
--- a/Engine/Plugins/Cameras/GameplayCameras/Source/GameplayCameras/Private/Nodes/Framing/DollyFramingCameraNode.cpp
+++ b/Engine/Plugins/Cameras/GameplayCameras/Source/GameplayCameras/Private/Nodes/Framing/DollyFramingCameraNode.cpp
@@ -85,7 +85,6 @@ void FDollyFramingCameraNodeEvaluator::OnRun(const FCameraNodeEvaluationParams&
 	const FVector3d TargetLocation = OptTargetLocation.GetValue();
 	const FTransform3d LastShotTransform = BuildDollyShotTransform(OutResult.CameraPose);
 	UpdateFramingState(Params, OutResult, TargetLocation, LastShotTransform);
-	ComputeDesiredState(Params.DeltaTime);
 
 	// If we need to reframe the target this tick, figure out how much we need to move the dolly
 	// to accomplish that.
@@ -125,6 +124,8 @@ void FDollyFramingCameraNodeEvaluator::OnRun(const FCameraNodeEvaluationParams&
 
 	const FTransform3d NewShotTransform = BuildDollyShotTransform(OutResult.CameraPose);
 	OutResult.CameraPose.SetTransform(NewShotTransform);
+
+	EndFramingUpdate(Params, OutResult);
 }
 
 FTransform3d FDollyFramingCameraNodeEvaluator::BuildDollyShotTransform(const FCameraPose& CameraPose) const
diff --git a/Engine/Plugins/Cameras/GameplayCameras/Source/GameplayCameras/Private/Nodes/Framing/PanningFramingCameraNode.cpp b/Engine/Plugins/Cameras/GameplayCameras/Source/GameplayCameras/Private/Nodes/Framing/PanningFramingCameraNode.cpp
--- a/Engine/Plugins/Cameras/GameplayCameras/Source/GameplayCameras/Private/Nodes/Framing/PanningFramingCameraNode.cpp
+++ b/Engine/Plugins/Cameras/GameplayCameras/Source/GameplayCameras/Private/Nodes/Framing/PanningFramingCameraNode.cpp
@@ -79,7 +79,6 @@ void FPanningFramingCameraNodeEvaluator::OnRun(const FCameraNodeEvaluationParams
 	const FVector3d TargetLocation = OptTargetLocation.GetValue();
 	const FTransform3d LastShotTransform = BuildPanningShotTransform(OutResult.CameraPose);
 	UpdateFramingState(Params, OutResult, TargetLocation, LastShotTransform);
-	ComputeDesiredState(Params.DeltaTime);
 
 	// If we need to reframe the target this tick, figure out how much we need to rotate the camera
 	// to accomplish that.
@@ -117,6 +116,8 @@ void FPanningFramingCameraNodeEvaluator::OnRun(const FCameraNodeEvaluationParams
 
 	const FTransform3d NewShotTransform = BuildPanningShotTransform(OutResult.CameraPose);
 	OutResult.CameraPose.SetTransform(NewShotTransform);
+
+	EndFramingUpdate(Params, OutResult);
 }
 
 FTransform3d FPanningFramingCameraNodeEvaluator::BuildPanningShotTransform(const FCameraPose& CameraPose) const
diff --git a/Engine/Plugins/Cameras/GameplayCameras/Source/GameplayCameras/Public/Math/CameraFramingZoneMath.h b/Engine/Plugins/Cameras/GameplayCameras/Source/GameplayCameras/Public/Math/CameraFramingZoneMath.h
--- a/Engine/Plugins/Cameras/GameplayCameras/Source/GameplayCameras/Public/Math/CameraFramingZoneMath.h
+++ b/Engine/Plugins/Cameras/GameplayCameras/Source/GameplayCameras/Public/Math/CameraFramingZoneMath.h
@@ -6,6 +6,7 @@
 #include "Math/MathFwd.h"
 
 class FArchive;
+struct FCameraFramingZone;
 struct FCameraPose;
 
 namespace UE::Cameras
@@ -13,20 +14,9 @@ namespace UE::Cameras
 
 struct FCameraFieldsOfView;
 
-/**
- * Effective margins for a rectangular screen-space zone.
- */
-struct FFramingZoneMargins
-{
-	double LeftMargin = 0;
-	double TopMargin = 0;
-	double RightMargin = 0;
-	double BottomMargin = 0;
-};
-
 /**
  * Effective coordinates for a rectangular screen-space zone.
- * Unlike FFramingZoneMargins, which stores margin values from the screen's edges,
+ * Unlike FCameraFramingZone, whose values can mean anything (margins, offsets, etc)
  * this struct is expected to store actual screen coordinates in 0..1 UI space.
  */
 struct FFramingZone
@@ -39,9 +29,6 @@ struct FFramingZone
 	/** Builds an empty framing zone. */
 	FFramingZone();
 
-	/** Build a framing zone from a set of margins. */
-	FFramingZone(const FFramingZoneMargins& FramingZoneMargins);
-
 	/** Makes sure all the bounds have valid values between 0 and 1. */
 	void ClampBounds();
 
@@ -76,6 +63,14 @@ struct FFramingZone
 
 	void Serialize(FArchive& Ar);
 
+public:
+
+	/** Build a framing zone from a set of margins relative to the screen edge. */
+	static FFramingZone FromScreenMargins(const FCameraFramingZone& Margins);
+
+	/** Build a framing zone from a set of margins relative to a screen location. */
+	static FFramingZone FromRelativeMargins(const FVector2d& ScreenLocation, const FCameraFramingZone& Margins);
+
 private:
 
 	static double GetNormalizedBound(double Bound);
diff --git a/Engine/Plugins/Cameras/GameplayCameras/Source/GameplayCameras/Public/Nodes/Framing/BaseFramingCameraNode.h b/Engine/Plugins/Cameras/GameplayCameras/Source/GameplayCameras/Public/Nodes/Framing/BaseFramingCameraNode.h
--- a/Engine/Plugins/Cameras/GameplayCameras/Source/GameplayCameras/Public/Nodes/Framing/BaseFramingCameraNode.h
+++ b/Engine/Plugins/Cameras/GameplayCameras/Source/GameplayCameras/Public/Nodes/Framing/BaseFramingCameraNode.h
@@ -6,7 +6,9 @@
 #include "Core/CameraParameterReader.h"
 #include "Core/CameraParameters.h"
 #include "Core/CameraVariableReferences.h"
+#include "Core/ICustomCameraNodeParameterProvider.h"
 #include "Nodes/Framing/CameraFramingZone.h"
+#include "Nodes/Framing/CameraTargetInfo.h"
 #include "Math/CameraFramingZoneMath.h"
 #include "Math/CriticalDamper.h"
 
@@ -19,26 +21,33 @@ class UVector3dCameraVariable;
  * The base class for a standard scren-space framing camera node.
  */
 UCLASS(MinimalAPI, Abstract, meta=(CameraNodeCategories="Framing"))
-class UBaseFramingCameraNode : public UCameraNode
+class UBaseFramingCameraNode : public UCameraNode, public ICustomCameraNodeParameterProvider
 {
 	GENERATED_BODY()
 
 public:
 
-	/**
-	 * A camera variable providing the location of the target to frame. If unspecified,
-	 * the player pawn's location will be used by default.
+	/** 
+	 * A variable whose value is the desired target's location in world space.
+	 * If set, and if the variable has been set, the obtained value takes priority
+	 * over the TargetInfo property.
 	 */
 	UPROPERTY(EditAnywhere, Category="Target")
 	FVector3dCameraVariableReference TargetLocation;
 
-	/** The ideal horizontal screen-space position of the target. */
-	UPROPERTY(EditAnywhere, Category="Framing Target")
-	FDoubleCameraParameter HorizontalFraming;
+	UPROPERTY(EditAnywhere, Category="Target")
+	FCameraTargetInfoParameter TargetInfo;
 
-	/** The ideal vertical screen-space position of the target. */
+	/**
+	 * Whether the camera pose's target distance should be set to the distance between
+	 * its location and the effective target's location.
+	 */
+	UPROPERTY(EditAnywhere, Category="Target")
+	FBooleanCameraParameter SetTargetDistance;
+
+	/** The ideal horizontal and vertical screen-space position of the target. */
 	UPROPERTY(EditAnywhere, Category="Framing Target")
-	FDoubleCameraParameter VerticalFraming;
+	FVector2dCameraParameter IdealFramingLocation;
 
 	/** The damping factor for how fast the framing recenters on the target. */
 	UPROPERTY(EditAnywhere, Category="Framing Target")
@@ -53,33 +62,51 @@ public:
 	FFloatCameraParameter LowReframeDampingFactor;
 
 	/**
-	 * The distance from the ideal framing position at which we can disengage reframing.
-	 * This should be a very small value, but if it is too small the reframing will keep "chasing"
-	 * the target for a long time even if it stays in the dead zone.
+	 * The time spent ramping up the reframing after exiting the dead zone.
+	 * If set to zero or a negative value, reframing will immediately restart once the target
+	 * has exited the dead zone. Otherwise the ReframeDampingFactor will interpolate from zero to
+	 * its desired value over the specified amount of seconds.
+	 */
+	UPROPERTY(EditAnywhere, Category="Framing Target")
+	FFloatCameraParameter ReengageTime;
+
+	/**
+	 * The time spent ramping down the reframing after entering the dead zone.
+	 * If set to zero or a negative value, reframing will immediately stop once the target has 
+	 * entered the dead zone. Otherwise, the ReframeDampingFactor will interpolate towards zero
+	 * over the specified amount of seconds.
 	 */
 	UPROPERTY(EditAnywhere, Category="Framing Target")
-	FFloatCameraParameter ReframeUnlockRadius;
+	FFloatCameraParameter DisengageTime;
+
+	UPROPERTY(EditAnywhere, Category="Framing Target")
+	FFloatCameraParameter TargetMovementAnticipationTime;
 
 	/** 
-	 * The margins of the dead zone, i.e. the zone inside which the target can freely move.
-	 * Margins are expressed in screen percentages from the edges.
+	 * The size of the dead zone, i.e. the zone inside which the target can freely move.
+	 * Sizes are expressed screen percentages around the desired framing location.
 	 */
 	UPROPERTY(EditAnywhere, Category="Framing Zones")
-	FCameraFramingZone DeadZone;
+	FCameraFramingZoneParameter DeadZone;
 
 	/**
 	 * The margins of the soft zone, i.e. the zone inside which the reframing will engage, in order
 	 * to bring the target back towards the ideal framing position. If the target is outside of the
 	 * soft zone, it will be forcibly and immedialy brought back to its edges, so this zone also 
 	 * defines the "hard" or "safe" zone of framing.
-	 * Margins are expressed in screen percentages from the edges.
+	 * Sizes are expressed in screen percentages from the edges.
 	 */
 	UPROPERTY(EditAnywhere, Category="Framing Zones")
-	FCameraFramingZone SoftZone;
+	FCameraFramingZoneParameter SoftZone;
 
 public:
 
 	UBaseFramingCameraNode(const FObjectInitializer& ObjectInit);
+
+public:
+
+	// ICustomCameraNodeParameterProvider interface.
+	virtual void GetCustomCameraNodeParameters(FCustomCameraNodeParameterInfos& OutParameterInfos) override;
 };
 
 namespace UE::Cameras
@@ -87,46 +114,26 @@ namespace UE::Cameras
 
 class FCameraVariableTable;
 
-/**
- * Utility struct for reading a framing zone's margin parameters.
- */
-struct FCameraFramingZoneParameterReader
-{
-public:
-
-	TCameraParameterReader<double> LeftMargin;
-	TCameraParameterReader<double> TopMargin;
-	TCameraParameterReader<double> RightMargin;
-	TCameraParameterReader<double> BottomMargin;
-
-public:
-
-	void Initialize(const FCameraFramingZone& FramingZone);
-	FFramingZoneMargins GetZoneMargins(const FCameraVariableTable& VariableTable) const;
-};
-
 /**
  * The base class for a framing camera node evaluator.
  *
  * This evaluator does nothing per se but provides utility functions to be called in 
  * a sub-class' OnRun method. Namely:
  *
+ * - AcquireTargetLocation() : a default way to get the world location of the desired
+ *			target.
+ *
  * - UpdateFramingState() : computes the current state of the framing node. The result
  *			can be obtained from the State field.
+ *			The, compute the desired framing state for the current tick, including the 
+ *			desired framing correction. This can be obtained from the Desired field.
+ *			It is up to the sub-class to implement the necessary logic to honor this 
+ *			correction. For instance, a dolly shot would translate left/right (and maybe 
+ *			up/down too) to try and reframe things accordingly, whereas a panning shot 
+ *			would rotate the camera left/right/up/down to accomplish the same.
  *
- * - ComputeDesiredState() : one the current state has been written, this method computes
- *			the desired framing state for the current tick, including the desired framing
- *			correction. It is up to the sub-class to implement the necessary logic to
- *			honor this correction. For instance, a dolly shot would translate left/right 
- *			(and maybe up/down too) to try and reframe things accordingly, whereas a panning
- *			shot would rotate the camera left/right/up/down to accomplish the same.
- *
- * - RegisterNewFraming() : once the framing correction has been executed by the sub-class,
- *			it's important to register the new camera transform with RegisterNewFraming, 
- *			otherwise the reframing will always act only on the incoming camera pose! If this
- *			incoming camera pose is fixed (e.g. the previous nodes are only fixed offsets)
- *			then if RegisterNewFraming isn't called, the reframing will always do the same
- *			thing every frame!
+ * - EndFramingUpdate() : the sub-class should call near the end of its OnRun method.
+ *			This will for instance optionally set the target distance.
  */
 class FBaseFramingCameraNodeEvaluator : public FCameraNodeEvaluator
 {
@@ -136,23 +143,38 @@ protected:
 
 	virtual void OnInitialize(const FCameraNodeEvaluatorInitializeParams& Params, FCameraNodeEvaluationResult& OutResult) override;
 
+#if WITH_EDITOR
+	virtual void OnDrawEditorPreview(const FCameraEditorPreviewDrawParams& Params, FCameraDebugRenderer& Renderer) override;
+#endif  // WITH_EDITOR
+
 #if UE_GAMEPLAY_CAMERAS_DEBUG
 	virtual void OnBuildDebugBlocks(const FCameraDebugBlockBuildParams& Params, FCameraDebugBlockBuilder& Builder) override;
 #endif  // UE_GAMEPLAY_CAMERAS_DEBUG
 
 protected:
 
+	struct FState;
+	struct FDesired;
+
 	/** Gets the target location. */
 	TOptional<FVector3d> AcquireTargetLocation(const FCameraNodeEvaluationParams& Params, const FCameraNodeEvaluationResult& InResult);
 	/** Updates the framing state for the current tick, see State member field. */
 	void UpdateFramingState(const FCameraNodeEvaluationParams& Params, const FCameraNodeEvaluationResult& OutResult, const FVector3d& TargetLocation, const FTransform3d& LastFraming);
-	/** Computes the desired reframing for the current tick, see Desired member field. */
-	void ComputeDesiredState(float DeltaTime);
+	/** Wraps-up the update with optional operations. */
+	void EndFramingUpdate(const FCameraNodeEvaluationParams& Params, FCameraNodeEvaluationResult& OutResult);
 
 private:
 
+#if UE_GAMEPLAY_CAMERAS_DEBUG
+	static void DrawFramingState(const FState& State, const FDesired& Desired, FCameraDebugRenderer& Renderer);
+#endif  // UE_GAMEPLAY_CAMERAS_DEBUG
+
 	FVector2d GetHardReframeCoords() const;
 
+	void ComputeCurrentState(const FCameraNodeEvaluationParams& Params, const FCameraNodeEvaluationResult& OutResult, const FVector3d& TargetLocation, const FTransform3d& LastFraming);
+	FVector3d ComputeAnticipatedScreenTarget(float DeltaTime, const FVector3d& InTargetLocation);
+	void ComputeDesiredState(float DeltaTime);
+
 protected:
 
 	/** The current location of the target. */
@@ -178,14 +200,17 @@ protected:
 	/** Utility structure for all the parameter readers we need every frame. */
 	struct FReaders
 	{
-		TCameraParameterReader<double> HorizontalFraming;
-		TCameraParameterReader<double> VerticalFraming;
+		TCameraParameterReader<FVector2d> IdealFramingLocation;
+		TCameraParameterReader<bool> SetTargetDistance;
+
 		TCameraParameterReader<float> ReframeDampingFactor;
 		TCameraParameterReader<float> LowReframeDampingFactor;
-		TCameraParameterReader<float> ReframeUnlockRadius;
+		TCameraParameterReader<float> ReengageTime;
+		TCameraParameterReader<float> DisengageTime;
+		TCameraParameterReader<float> TargetMovementAnticipationTime;
 
-		FCameraFramingZoneParameterReader DeadZoneMargin;
-		FCameraFramingZoneParameterReader SoftZoneMargin;
+		TCameraParameterReader<FCameraFramingZone> DeadZone;
+		TCameraParameterReader<FCameraFramingZone> SoftZone;
 	};
 	FReaders Readers;
 
@@ -201,8 +226,18 @@ protected:
 		float ReframeDampingFactor;
 		/** Current low reframing damping factor. */
 		float LowReframeDampingFactor;
-		/** Current reframe unlock radius. */
-		float ReframeUnlockRadius;
+		/** Current alpha between reframing damping factors. */
+		float ReframeDampingFactorAlpha;
+		/** Current reengage time. */
+		float ReengageTime;
+		/** Current disengage time. */
+		float DisengageTime;
+		/** Current time spent disengaging or reengaging reframing */
+		float ToggleEngageTimeLeft;
+		/** Current reframing damping factor alpha due to engage toggle */
+		float ToggleEngageAlpha;
+		/** Current look-ahead time for anticipating target movement */
+		float TargetMovementAnticipationTime;
 		/** Current coordinates of the dead zone. */
 		FFramingZone DeadZone;
 		/** Current coordinates of the soft zone. */
@@ -218,6 +253,13 @@ protected:
 		/** The damper for reframing from the soft zone. */
 		FCriticalDamper ReframeDamper;
 
+#if UE_GAMEPLAY_CAMERAS_DEBUG
+		/** Intersection of the reframing vector with the dead zone box. */
+		FVector2d DeadZoneEdgePoint;
+		/** Intersection of the reframing vector with the hard zone box. */
+		FVector2d HardZoneEdgePoint;
+#endif  // UE_GAMEPLAY_CAMERAS_DEBUG
+
 		void Serialize(FArchive& Ar);
 	};
 	FState State;
@@ -245,6 +287,12 @@ protected:
 	};
 	FDesired Desired;
 
+	struct FWorldTargetAnticipation
+	{
+		TArray<TTuple<FVector3d, float>, TInlineAllocator<10>> History;
+	};
+	FWorldTargetAnticipation Anticipation;
+
 	friend class FBaseFramingCameraDebugBlock;
 	friend FArchive& operator <<(FArchive& Ar, FState& State);
 	friend FArchive& operator <<(FArchive& Ar, FDesired& Desired);
diff --git a/Engine/Plugins/Cameras/GameplayCameras/Source/GameplayCameras/Public/Nodes/Framing/CameraFramingZone.h b/Engine/Plugins/Cameras/GameplayCameras/Source/GameplayCameras/Public/Nodes/Framing/CameraFramingZone.h
--- a/Engine/Plugins/Cameras/GameplayCameras/Source/GameplayCameras/Public/Nodes/Framing/CameraFramingZone.h
+++ b/Engine/Plugins/Cameras/GameplayCameras/Source/GameplayCameras/Public/Nodes/Framing/CameraFramingZone.h
@@ -2,8 +2,8 @@
 
 #pragma once
 
+#include "Core/CameraVariableTableFwd.h"
 #include "CoreTypes.h"
-#include "Core/CameraParameters.h"
 
 #include "CameraFramingZone.generated.h"
 
@@ -13,57 +13,125 @@
  * All margins are defined in percentages of the screen's horizontal size. They are also 
  * all defined relative to their respective edges.
  */
-USTRUCT()
+USTRUCT(BlueprintType)
 struct FCameraFramingZone
 {
 	GENERATED_BODY()
 
 public:
 
-	UPROPERTY(EditAnywhere, Category="Framing")
-	FDoubleCameraParameter LeftMargin;
+	UPROPERTY(BlueprintReadWrite, EditAnywhere, Category="Framing")
+	double Left;
 
-	UPROPERTY(EditAnywhere, Category="Framing")
-	FDoubleCameraParameter TopMargin;
+	UPROPERTY(BlueprintReadWrite, EditAnywhere, Category="Framing")
+	double Top;
 
-	UPROPERTY(EditAnywhere, Category="Framing")
-	FDoubleCameraParameter RightMargin;
+	UPROPERTY(BlueprintReadWrite, EditAnywhere, Category="Framing")
+	double Right;
 
-	UPROPERTY(EditAnywhere, Category="Framing")
-	FDoubleCameraParameter BottomMargin;
+	UPROPERTY(BlueprintReadWrite, EditAnywhere, Category="Framing")
+	double Bottom;
 
 public:
 
 	FCameraFramingZone()
 	{
-		LeftMargin.Value = 0;
-		TopMargin.Value = 0;
-		RightMargin.Value = 0;
-		BottomMargin.Value = 0;
+		Left = 0;
+		Top = 0;
+		Right = 0;
+		Bottom = 0;
+	}
+
+	FCameraFramingZone(double UniformValue)
+	{
+		Left = UniformValue;
+		Top = UniformValue;
+		Right = UniformValue;
+		Bottom = UniformValue;
+	}
+
+	FCameraFramingZone(double Horizontal, double Vertical)
+	{
+		Left = Horizontal;
+		Top = Vertical;
+		Right = Horizontal;
+		Bottom = Vertical;
+	}
+
+	FCameraFramingZone(double InLeft, double InTop, double InRight, double InBottom)
+	{
+		Left = InLeft;
+		Top = InTop;
+		Right = InRight;
+		Bottom = InBottom;
+	}
+
+public:
+
+	FCameraFramingZone operator*(float Scale) const
+	{
+		return FCameraFramingZone(Left * Scale, Top * Scale, Right * Scale, Bottom * Scale);
 	}
 
-	FCameraFramingZone(double UniformMargin)
+	FCameraFramingZone operator*(const FCameraFramingZone& Scale) const
 	{
-		LeftMargin.Value = UniformMargin;
-		TopMargin.Value = UniformMargin;
-		RightMargin.Value = UniformMargin;
-		BottomMargin.Value = UniformMargin;
+		return FCameraFramingZone(
+				Left * Scale.Left, 
+				Top * Scale.Top,
+				Right * Scale.Right,
+				Bottom * Scale.Bottom);
 	}
 
-	FCameraFramingZone(double HorizontalMargin, double VerticalMargin)
+	FCameraFramingZone operator+(const FCameraFramingZone& Other) const
 	{
-		LeftMargin.Value = HorizontalMargin;
-		TopMargin.Value = VerticalMargin;
-		RightMargin.Value = HorizontalMargin;
-		BottomMargin.Value = VerticalMargin;
+		return FCameraFramingZone(
+				Left + Other.Left, 
+				Top + Other.Top,
+				Right + Other.Right,
+				Bottom + Other.Bottom);
 	}
 
-	FCameraFramingZone(double InLeftMargin, double InTopMargin, double InRightMargin, double InBottomMargin)
+	FCameraFramingZone operator-(const FCameraFramingZone& Other) const
 	{
-		LeftMargin.Value = InLeftMargin;
-		TopMargin.Value = InTopMargin;
-		RightMargin.Value = InRightMargin;
-		BottomMargin.Value = InBottomMargin;
+		return FCameraFramingZone(
+				Left - Other.Left, 
+				Top - Other.Top,
+				Right - Other.Right,
+				Bottom - Other.Bottom);
 	}
+
+	bool operator== (const FCameraFramingZone& Other) const = default;
+	bool operator!= (const FCameraFramingZone& Other) const = default;
+
+public:
+
+	GAMEPLAYCAMERAS_API static void TypeErasedInterpolate(uint8* From, const uint8* To, float Alpha);
+
+	GAMEPLAYCAMERAS_API FString ToString() const;
+};
+
+template <> struct TIsPODType<FCameraFramingZone>
+{
+	enum { Value = true };
+};
+
+/** Framing zone camera parameter. */
+USTRUCT()
+struct GAMEPLAYCAMERAS_API FCameraFramingZoneParameter
+{
+	GENERATED_BODY()
+
+	using ValueType = FCameraFramingZone;
+
+	UPROPERTY(EditAnywhere, Category=Common, meta=(ShowOnlyInnerProperties))
+	FCameraFramingZone Value;
+
+	UPROPERTY()
+	FCameraVariableID VariableID;
+
+	FCameraFramingZoneParameter() {}
+	FCameraFramingZoneParameter(const FCameraFramingZone& InValue)
+		: Value(InValue)
+	{}
 };
 
diff --git a/Engine/Plugins/Cameras/GameplayCameras/Source/GameplayCameras/Public/Nodes/Framing/CameraTargetInfo.h b/Engine/Plugins/Cameras/GameplayCameras/Source/GameplayCameras/Public/Nodes/Framing/CameraTargetInfo.h
new file mode 100644
--- /dev/null
+++ b/Engine/Plugins/Cameras/GameplayCameras/Source/GameplayCameras/Public/Nodes/Framing/CameraTargetInfo.h
@@ -0,0 +1,70 @@
+// Copyright Epic Games, Inc. All Rights Reserved.
+
+#pragma once
+
+#include "CoreTypes.h"
+#include "Core/CameraContextDataTableFwd.h"
+
+#include "CameraTargetInfo.generated.h"
+
+class AActor;
+
+namespace UE::Cameras
+{
+	class FCameraContextDataTable;
+}
+
+USTRUCT(BlueprintType)
+struct GAMEPLAYCAMERAS_API FCameraTargetInfo
+{
+	GENERATED_BODY()
+
+	UPROPERTY(EditAnywhere, Category="Target")
+	TObjectPtr<AActor> Actor;
+
+	UPROPERTY(EditAnywhere, Category="Target")
+	FName SocketName;
+};
+
+USTRUCT()
+struct GAMEPLAYCAMERAS_API FCameraTargetInfoParameter
+{
+	GENERATED_BODY()
+
+	using DataType = FCameraTargetInfo;
+
+	UPROPERTY(EditAnywhere, Category=Common, meta=(ShowOnlyInnerProperties))
+	FCameraTargetInfo Value;
+
+	UPROPERTY()
+	FCameraContextDataID DataID;
+
+	FCameraTargetInfoParameter() {}
+	FCameraTargetInfoParameter(const FCameraTargetInfo& InValue)
+		: Value(InValue)
+	{}
+
+	bool HasOverride() const { return DataID.IsValid(); }
+	const FCameraTargetInfo& GetValue(const UE::Cameras::FCameraContextDataTable& ContextDataTable) const;
+};
+
+namespace UE::Cameras
+{
+
+struct FCameraTargetInfoReader
+{
+	FCameraTargetInfoReader() {}
+	FCameraTargetInfoReader(const FCameraTargetInfo& InTargetInfo);
+
+	void Initialize(const FCameraTargetInfo& InTargetInfo);
+
+	bool GetTargetTransform(FTransform3d& OutTransform) const;
+
+private:
+
+	FCameraTargetInfo TargetInfo;
+	const USkeletalMeshComponent* CachedSkeletalMeshComponent = nullptr;
+};
+
+}  // namespace UE::Cameras
+
diff --git a/Engine/Plugins/Cameras/GameplayCameras/Source/GameplayCamerasEditor/Private/Customizations/CameraFramingZoneDetailsCustomization.cpp b/Engine/Plugins/Cameras/GameplayCameras/Source/GameplayCamerasEditor/Private/Customizations/CameraFramingZoneDetailsCustomization.cpp
new file mode 100644
--- /dev/null
+++ b/Engine/Plugins/Cameras/GameplayCameras/Source/GameplayCamerasEditor/Private/Customizations/CameraFramingZoneDetailsCustomization.cpp
@@ -0,0 +1,4 @@
+// Copyright Epic Games, Inc. All Rights Reserved.
+
+#include "Customizations/CameraFramingZoneDetailsCustomization.h"
+
diff --git a/Engine/Plugins/Cameras/GameplayCameras/Source/GameplayCamerasEditor/Public/Customizations/CameraFramingZoneDetailsCustomization.h b/Engine/Plugins/Cameras/GameplayCameras/Source/GameplayCamerasEditor/Public/Customizations/CameraFramingZoneDetailsCustomization.h
new file mode 100644
--- /dev/null
+++ b/Engine/Plugins/Cameras/GameplayCameras/Source/GameplayCamerasEditor/Public/Customizations/CameraFramingZoneDetailsCustomization.h
@@ -0,0 +1,4 @@
+// Copyright Epic Games, Inc. All Rights Reserved.
+
+#pragma once
+
