Cameras: add support for animating camera parameters in Sequencer

This change uses the newly added "custom handlers" for Sequencer property bindings. It registers a custom handler for camera parameters, and implements a track editor for showing those parameters in the UI.

diff --git a/Engine/Plugins/Cameras/GameplayCameras/Source/GameplayCameras/Private/GameplayCamerasModule.cpp b/Engine/Plugins/Cameras/GameplayCameras/Source/GameplayCameras/Private/GameplayCamerasModule.cpp
--- a/Engine/Plugins/Cameras/GameplayCameras/Source/GameplayCameras/Private/GameplayCamerasModule.cpp
+++ b/Engine/Plugins/Cameras/GameplayCameras/Source/GameplayCameras/Private/GameplayCamerasModule.cpp
@@ -6,6 +6,7 @@
 #include "GameplayCameras.h"
 #include "Logging/MessageLog.h"
 #include "Modules/ModuleManager.h"
+#include "MovieScene/GameplayCamerasTrackInstancePropertyBindings.h"
 #include "Nodes/Framing/CameraFramingZone.h"
 #include "ShowFlags.h"
 #include "UObject/UObjectBase.h"
@@ -34,7 +35,8 @@ public:
 	virtual void StartupModule() override
 	{
 		RegisterBuiltInBlendableStructs();
-
+		RegisterTrackInstanceCustomPropertyBindings();
+		
 #if UE_GAMEPLAY_CAMERAS_DEBUG
 		UE::Cameras::FCameraDebugColors::RegisterBuiltinColorSchemes();
 #endif  // UE_GAMEPLAY_CAMERAS_DEBUG
@@ -43,6 +45,7 @@ public:
 	virtual void ShutdownModule() override
 	{
 		UnregisterBuiltInBlendableStructs();
+		UnregisterTrackInstanceCustomPropertyBindings();
 	}
 
 public:
@@ -114,6 +117,16 @@ private:
 			UnregisterBlendableStruct(FCameraFramingZone::StaticStruct());
 		}
 	}
+	
+	void RegisterTrackInstanceCustomPropertyBindings()
+	{
+		UE::Cameras::FGameplayCamerasTrackInstancePropertyBindings::Register();
+	}
+
+	void UnregisterTrackInstanceCustomPropertyBindings()
+	{
+		UE::Cameras::FGameplayCamerasTrackInstancePropertyBindings::Unregister();
+	}
 
 private:
 
diff --git a/Engine/Plugins/Cameras/GameplayCameras/Source/GameplayCameras/Private/MovieScene/GameplayCamerasTrackInstancePropertyBindings.cpp b/Engine/Plugins/Cameras/GameplayCameras/Source/GameplayCameras/Private/MovieScene/GameplayCamerasTrackInstancePropertyBindings.cpp
new file mode 100644
--- /dev/null
+++ b/Engine/Plugins/Cameras/GameplayCameras/Source/GameplayCameras/Private/MovieScene/GameplayCamerasTrackInstancePropertyBindings.cpp
@@ -0,0 +1,237 @@
+// Copyright Epic Games, Inc. All Rights Reserved.
+
+#include "MovieScene/GameplayCamerasTrackInstancePropertyBindings.h"
+
+#include "Core/CameraAsset.h"
+#include "Core/CameraParameters.h"
+#include "Core/CameraRigParameterDefinition.h"
+#include "GameFramework/GameplayCameraComponent.h"
+#include "MovieSceneCommonHelpers.h"
+
+namespace UE::Cameras
+{
+
+/**
+ * A custom property binding handler for camera parameters.
+ */
+struct FGameplayCamerasTrackInstanceCustomPropertyBindingHandler : public ITrackInstanceCustomPropertyBindingHandler
+{
+public:
+
+	virtual bool FindProperty(void* BasePointer, const FString& PropertyName, FTrackInstanceCustomProperty& OutProperty) override
+	{
+		UGameplayCameraComponent* CameraComponent = reinterpret_cast<UGameplayCameraComponent*>(BasePointer);
+		check(CameraComponent);
+
+		// See if we have already resolved this parameter before.
+		FResolvedParameterKey Key{ CameraComponent, FName(PropertyName) };
+		if (FResolvedParameterInfo* ExistingInfo = ResolvedParameters.Find(Key))
+		{
+			OutProperty.ContainerType = ETrackInstancePropertyContainerType::None;
+			OutProperty.PropertyStruct = ExistingInfo->CustomPropertyStruct;
+			OutProperty.PropertyID = ExistingInfo->CustomPropertyID;
+			return true;
+		}
+
+		// Binding to a camera component that doesn't have a camera reference set is ignored.
+		const UCameraAsset* CameraAsset = CameraComponent->CameraReference.GetCameraAsset();
+		if (CameraAsset == nullptr)
+		{
+			return false;
+		}
+
+		TConstArrayView<FCameraRigParameterDefinition> ParameterDefinitions = CameraAsset->GetParameterDefinitions();
+		const FCameraRigParameterDefinition* ParameterDefinition = ParameterDefinitions.FindByPredicate(
+				[&PropertyName](const FCameraRigParameterDefinition& Item)
+				{
+					return Item.ParameterName == PropertyName;
+				});
+		if (!ParameterDefinition)
+		{
+			// We might not find the parameter definition if the asset is out of date and hasn't been built.
+			return false;
+		}
+
+		// See if we can find a default value for this parameter.
+		const FInstancedPropertyBag& DefaultParameters = CameraAsset->GetDefaultParameters();
+		const UPropertyBag* DefaultParametersStruct = DefaultParameters.GetPropertyBagStruct();
+		if (!DefaultParametersStruct)
+		{
+			return false;
+		}
+
+		const FPropertyBagPropertyDesc* PropertyDesc = DefaultParametersStruct->FindPropertyDescByID(ParameterDefinition->ParameterGuid);
+		if (!PropertyDesc || !PropertyDesc->CachedProperty)
+		{
+			return false;
+		}
+
+		const uint8* RawDefaultParameters = DefaultParameters.GetValue().GetMemory();
+		const void* RawDefaultValuePtr = PropertyDesc->CachedProperty->ContainerPtrToValuePtr<void>(RawDefaultParameters);
+		if (!RawDefaultValuePtr)
+		{
+			return false;
+		}
+
+		// Cache the info for this resolved parameter.
+		FResolvedParameterInfo& NewInfo = ResolvedParameters.Emplace(Key);
+		NewInfo.Definition = *ParameterDefinition;
+		NewInfo.RawDefaultValuePtr = RawDefaultValuePtr;
+		NewInfo.CustomPropertyID = NextCustomPropertyID++;
+		if (ParameterDefinition->ParameterType == ECameraRigInterfaceParameterType::Blendable)
+		{
+			switch (ParameterDefinition->VariableType)
+			{
+				case ECameraVariableType::Vector2f:
+					NewInfo.CustomPropertyStruct = TVariantStructure<FVector2f>::Get();
+					break;
+				case ECameraVariableType::Vector2d:
+					NewInfo.CustomPropertyStruct = TBaseStructure<FVector2D>::Get();
+					break;
+				case ECameraVariableType::Vector3f:
+					NewInfo.CustomPropertyStruct = TVariantStructure<FVector3f>::Get();
+					break;
+				case ECameraVariableType::Vector3d:
+					NewInfo.CustomPropertyStruct = TBaseStructure<FVector>::Get();
+					break;
+				case ECameraVariableType::Vector4f:
+					NewInfo.CustomPropertyStruct = TVariantStructure<FVector4f>::Get();
+					break;
+				case ECameraVariableType::Vector4d:
+					NewInfo.CustomPropertyStruct = TBaseStructure<FVector4>::Get();
+					break;
+				case ECameraVariableType::Rotator3f:
+					NewInfo.CustomPropertyStruct = TVariantStructure<FRotator3f>::Get();
+					break;
+				case ECameraVariableType::Rotator3d:
+					NewInfo.CustomPropertyStruct = TBaseStructure<FRotator>::Get();
+					break;
+				case ECameraVariableType::Transform3f:
+					NewInfo.CustomPropertyStruct = TVariantStructure<FTransform3f>::Get();
+					break;
+				case ECameraVariableType::Transform3d:
+					NewInfo.CustomPropertyStruct = TBaseStructure<FTransform>::Get();
+					break;
+				case ECameraVariableType::BlendableStruct:
+					NewInfo.CustomPropertyStruct = const_cast<UScriptStruct*>(ParameterDefinition->BlendableStructType.Get());
+					break;
+			}
+		}
+		else if (ParameterDefinition->ParameterType == ECameraRigInterfaceParameterType::Data)
+		{
+			NewInfo.CustomPropertyStruct = const_cast<UStruct*>(Cast<const UStruct>(ParameterDefinition->DataTypeObject));
+		}
+
+		CustomPropertyIDToParameterKey.Add(NewInfo.CustomPropertyID, Key);
+
+		OutProperty.ContainerType = ETrackInstancePropertyContainerType::None;
+		OutProperty.PropertyStruct = NewInfo.CustomPropertyStruct;
+		OutProperty.PropertyID = NewInfo.CustomPropertyID;
+
+		return true;
+	}
+
+	virtual void* ContainerPtrToValuePtr(void* BasePointer, const FTrackInstanceCustomProperty& CustomProperty, int32 ArrayIndex = INDEX_NONE) override
+	{
+		// Find a previous resolved camera parameter.
+		UGameplayCameraComponent* CameraComponent = reinterpret_cast<UGameplayCameraComponent*>(BasePointer);
+		if (!ensure(CameraComponent))
+		{
+			return nullptr;
+		}
+
+		TSharedPtr<FCameraEvaluationContext> EvaluationContext = CameraComponent->GetEvaluationContext();
+		if (!EvaluationContext)
+		{
+			return nullptr;
+		}
+
+		const FResolvedParameterKey* KeyPtr = CustomPropertyIDToParameterKey.Find(CustomProperty.PropertyID);
+		if (!ensure(KeyPtr))
+		{
+			return nullptr;
+		}
+
+		const FResolvedParameterInfo* ResolvedParameterPtr = ResolvedParameters.Find(*KeyPtr);
+		if (!ensure(ResolvedParameterPtr))
+		{
+			return nullptr;
+		}
+
+		// Get the pointer to the value in the variable table or context data table. Since we don't know if
+		// this pointer will be used for reading or writing, we need to make sure it points to something
+		// valid. So if the value is uninitialized, we write the default value into it.
+		const FCameraRigParameterDefinition& Definition = ResolvedParameterPtr->Definition;
+		if (Definition.ParameterType == ECameraRigInterfaceParameterType::Blendable)
+		{
+			FCameraVariableTable& VariableTable = EvaluationContext->GetInitialResult().VariableTable;
+			uint8* RawValuePtr = VariableTable.TryGetMutableValue(Definition.VariableID, Definition.VariableType, Definition.BlendableStructType);
+			if (!VariableTable.IsValueWritten(Definition.VariableID))
+			{
+				VariableTable.SetValue(Definition.VariableID, Definition.VariableType, Definition.BlendableStructType, (uint8*)ResolvedParameterPtr->RawDefaultValuePtr);
+			}
+			return RawValuePtr;
+		}
+		else if (Definition.ParameterType == ECameraRigInterfaceParameterType::Data)
+		{
+			FCameraContextDataTable& ContextDataTable = EvaluationContext->GetInitialResult().ContextDataTable;
+			uint8* RawValuePtr = ContextDataTable.TryGetMutableData(Definition.DataID, Definition.DataType, Definition.DataTypeObject);
+			if (!ContextDataTable.IsValueWritten(Definition.DataID))
+			{
+				ContextDataTable.SetData(Definition.DataID, Definition.DataType, Definition.DataTypeObject, (uint8*)ResolvedParameterPtr->RawDefaultValuePtr);
+			}
+			return RawValuePtr;
+		}
+
+		return nullptr;
+	}
+
+private:
+
+	struct FResolvedParameterKey
+	{
+		FObjectKey CameraComponentKey = nullptr;
+		FName ParameterName = NAME_None;
+
+		bool operator==(const FResolvedParameterKey& Other) const = default;
+
+		friend int32 GetTypeHash(const FResolvedParameterKey& Key)
+		{
+			return HashCombine(GetTypeHash(Key.CameraComponentKey), GetTypeHash(Key.ParameterName));
+		}
+	};
+
+	struct FResolvedParameterInfo
+	{
+		FCameraRigParameterDefinition Definition;
+		const void* RawDefaultValuePtr = nullptr;
+		UStruct* CustomPropertyStruct = nullptr;
+		int32 CustomPropertyID = 0;
+	};
+
+	TMap<FResolvedParameterKey, FResolvedParameterInfo> ResolvedParameters;
+	TMap<int32, FResolvedParameterKey> CustomPropertyIDToParameterKey;
+	int32 NextCustomPropertyID = 0;
+};
+
+void FGameplayCamerasTrackInstancePropertyBindings::Register()
+{
+	using namespace UE::MovieScene;
+
+	UClass* CameraComponentClass = UGameplayCameraComponent::StaticClass();
+	FTrackInstancePropertyBindings::RegisterCustomHandler(
+			CameraComponentClass, 
+			FOnCreateTrackInstanceCustomPropertyBindingHandler::CreateLambda(
+				[](){ return MakeShared<FGameplayCamerasTrackInstanceCustomPropertyBindingHandler>(); }));
+}
+
+void FGameplayCamerasTrackInstancePropertyBindings::Unregister()
+{
+	using namespace UE::MovieScene;
+
+	UClass* CameraComponentClass = UGameplayCameraComponent::StaticClass();
+	FTrackInstancePropertyBindings::UnregisterCustomHandler(CameraComponentClass);
+}
+
+}  // namespace UE::Cameras
+
diff --git a/Engine/Plugins/Cameras/GameplayCameras/Source/GameplayCameras/Private/MovieScene/GameplayCamerasTrackInstancePropertyBindings.h b/Engine/Plugins/Cameras/GameplayCameras/Source/GameplayCameras/Private/MovieScene/GameplayCamerasTrackInstancePropertyBindings.h
new file mode 100644
--- /dev/null
+++ b/Engine/Plugins/Cameras/GameplayCameras/Source/GameplayCameras/Private/MovieScene/GameplayCamerasTrackInstancePropertyBindings.h
@@ -0,0 +1,15 @@
+// Copyright Epic Games, Inc. All Rights Reserved.
+
+#pragma once
+
+namespace UE::Cameras
+{
+
+struct FGameplayCamerasTrackInstancePropertyBindings
+{
+	static void Register();
+	static void Unregister();
+};
+
+}  // namespace UE::Cameras
+
diff --git a/Engine/Plugins/Cameras/GameplayCameras/Source/GameplayCamerasEditor/GameplayCamerasEditor.Build.cs b/Engine/Plugins/Cameras/GameplayCameras/Source/GameplayCamerasEditor/GameplayCamerasEditor.Build.cs
--- a/Engine/Plugins/Cameras/GameplayCameras/Source/GameplayCamerasEditor/GameplayCamerasEditor.Build.cs
+++ b/Engine/Plugins/Cameras/GameplayCameras/Source/GameplayCamerasEditor/GameplayCamerasEditor.Build.cs
@@ -38,10 +38,13 @@ public class GameplayCamerasEditor : ModuleRules
 				"Kismet",
 				"KismetWidgets",
 				"LevelEditor",
+				"MovieScene",
+				"MovieSceneTracks",
 				"Projects",
 				"RenderCore",
 				"RewindDebuggerInterface",
 				"RHI",
+				"Sequencer",
 				"Slate",
 				"SlateCore",
 				"StructUtilsEditor",
diff --git a/Engine/Plugins/Cameras/GameplayCameras/Source/GameplayCamerasEditor/Private/GameplayCamerasEditorModule.cpp b/Engine/Plugins/Cameras/GameplayCameras/Source/GameplayCamerasEditor/Private/GameplayCamerasEditorModule.cpp
--- a/Engine/Plugins/Cameras/GameplayCameras/Source/GameplayCamerasEditor/Private/GameplayCamerasEditorModule.cpp
+++ b/Engine/Plugins/Cameras/GameplayCameras/Source/GameplayCamerasEditor/Private/GameplayCamerasEditorModule.cpp
@@ -40,10 +40,12 @@
 #include "IGameplayCamerasModule.h"
 #include "IRewindDebuggerExtension.h"
 #include "ISettingsModule.h"
+#include "ISequencerModule.h"
 #include "K2Node_Event.h"
 #include "Kismet2/KismetEditorUtilities.h"
 #include "Modules/ModuleManager.h"
 #include "PropertyEditorModule.h"
+#include "Sequencer/GameplayCameraComponentTrackEditor.h"
 #include "Styles/GameplayCamerasEditorStyle.h"
 #include "ToolMenus.h"
 #include "Toolkits/BlueprintCameraDirectorAssetEditorMode.h"
@@ -97,6 +99,7 @@ public:
 		RegisterDetailsCustomizations();
 		RegisterEdGraphUtilities();
 		RegisterComponentVisualizers();
+		RegisterSequencerTracks();
 
 		InitializeLiveEditManager();
 
@@ -124,6 +127,7 @@ public:
 		UnregisterDetailsCustomizations();
 		UnregisterEdGraphUtilities();
 		UnregisterComponentVisualizers();
+		UnregisterSequencerTracks();
 
 		TeardownLiveEditManager();
 
@@ -546,6 +550,18 @@ private:
 			GUnrealEd->UnregisterComponentVisualizer(UGameplayCameraComponent::StaticClass()->GetFName());
 		}
 	}
+	
+	void RegisterSequencerTracks()
+	{
+		ISequencerModule& SequencerModule = FModuleManager::Get().LoadModuleChecked<ISequencerModule>("Sequencer");
+		GameplayCameraComponentTrackCreateEditorHandle = SequencerModule.RegisterTrackEditor(FOnCreateTrackEditor::CreateStatic(&FGameplayCameraComponentTrackEditor::CreateTrackEditor));
+	}
+
+	void UnregisterSequencerTracks()
+	{
+		ISequencerModule& SequencerModule = FModuleManager::Get().LoadModuleChecked<ISequencerModule>("Sequencer");
+		SequencerModule.UnRegisterTrackEditor(GameplayCameraComponentTrackCreateEditorHandle);
+	}
 
 	void InitializeLiveEditManager()
 	{
@@ -578,6 +594,8 @@ private:
 
 	TMap<FString, UE::Cameras::FCameraDebugCategoryInfo> DebugCategoryInfos;
 	TMap<FString, FOnCreateDebugCategoryPanel> DebugCategoryPanelCreators;
+	
+	FDelegateHandle GameplayCameraComponentTrackCreateEditorHandle;
 
 #if UE_GAMEPLAY_CAMERAS_TRACE
 	TSharedPtr<UE::Cameras::FCameraSystemTraceModule> TraceModule;
diff --git a/Engine/Plugins/Cameras/GameplayCameras/Source/GameplayCamerasEditor/Private/Sequencer/GameplayCameraComponentTrackEditor.cpp b/Engine/Plugins/Cameras/GameplayCameras/Source/GameplayCamerasEditor/Private/Sequencer/GameplayCameraComponentTrackEditor.cpp
new file mode 100644
--- /dev/null
+++ b/Engine/Plugins/Cameras/GameplayCameras/Source/GameplayCamerasEditor/Private/Sequencer/GameplayCameraComponentTrackEditor.cpp
@@ -0,0 +1,308 @@
+// Copyright Epic Games, Inc. All Rights Reserved.
+
+#include "Sequencer/GameplayCameraComponentTrackEditor.h"
+
+#include "Core/CameraAsset.h"
+#include "Core/CameraRigParameterDefinition.h"
+#include "Framework/MultiBox/MultiBoxBuilder.h"
+#include "GameFramework/GameplayCameraComponent.h"
+#include "Styles/GameplayCamerasEditorStyle.h"
+#include "Tracks/MovieSceneActorReferenceTrack.h"
+#include "Tracks/MovieSceneBoolTrack.h"
+#include "Tracks/MovieSceneDoubleTrack.h"
+#include "Tracks/MovieSceneEnumTrack.h"
+#include "Tracks/MovieSceneFloatTrack.h"
+#include "Tracks/MovieSceneIntegerTrack.h"
+#include "Tracks/MovieSceneObjectPropertyTrack.h"
+#include "Tracks/MovieSceneRotatorTrack.h"
+#include "Tracks/MovieSceneStringTrack.h"
+#include "Tracks/MovieSceneTransformTrack.h"
+#include "Tracks/MovieSceneVectorTrack.h"
+
+#define LOCTEXT_NAMESPACE "GameplayCameraComponentTrackEditor"
+
+TSharedRef<ISequencerTrackEditor> FGameplayCameraComponentTrackEditor::CreateTrackEditor(TSharedRef<ISequencer> OwningSequencer)
+{
+	return MakeShareable(new FGameplayCameraComponentTrackEditor(OwningSequencer));
+}
+
+FGameplayCameraComponentTrackEditor::FGameplayCameraComponentTrackEditor(TSharedRef<ISequencer> InSequencer)
+	: FMovieSceneTrackEditor(InSequencer) 
+{
+}
+
+void FGameplayCameraComponentTrackEditor::BindCommands(TSharedRef<FUICommandList> SequencerCommandBindings)
+{
+}
+
+void FGameplayCameraComponentTrackEditor::BuildAddTrackMenu(FMenuBuilder& MenuBuilder)
+{
+}
+
+void FGameplayCameraComponentTrackEditor::BuildTrackContextMenu(FMenuBuilder& MenuBuilder, UMovieSceneTrack* Track)
+{
+}
+
+void FGameplayCameraComponentTrackEditor::ExtendObjectBindingTrackMenu(TSharedRef<FExtender> Extender, const TArray<FGuid>& ObjectBindings, const UClass* ObjectClass)
+{
+	if (ObjectClass && ObjectClass->IsChildOf<UGameplayCameraComponent>())
+	{
+		Extender->AddMenuExtension(
+				TEXT("Tracks"), EExtensionHook::After, nullptr, 
+				FMenuExtensionDelegate::CreateSP(this, &FGameplayCameraComponentTrackEditor::OnExtendObjectBindingTrackMenu, ObjectBindings));
+	}
+}
+
+void FGameplayCameraComponentTrackEditor::OnExtendObjectBindingTrackMenu(FMenuBuilder& MenuBuilder, TArray<FGuid> ObjectBindings)
+{
+	TMap<FName, FCameraRigParameterDefinition> DefinitionsByName;
+	TMap<FName, int32> NumEqualDefinitions;
+
+	for (const FGuid& ObjectBinding : ObjectBindings)
+	{
+		UGameplayCameraComponent* CameraComponent = GetCameraComponentForBinding(ObjectBinding);
+		if (CameraComponent && CameraComponent->CameraReference.GetCameraAsset())
+		{
+			const UCameraAsset* CameraAsset = CameraComponent->CameraReference.GetCameraAsset();
+			for (const FCameraRigParameterDefinition& Definition : CameraAsset->GetParameterDefinitions())
+			{
+				if (FCameraRigParameterDefinition* ExistingDefinition = DefinitionsByName.Find(Definition.ParameterName))
+				{
+					if (Definition == *ExistingDefinition)
+					{
+						++NumEqualDefinitions[Definition.ParameterName];
+					}
+				}
+				else
+				{
+					DefinitionsByName.Add(Definition.ParameterName, Definition);
+					NumEqualDefinitions.Add(Definition.ParameterName, 1);
+				}
+			}
+		}
+	}
+
+	const int32 NumObjectBindings = ObjectBindings.Num();
+	TArray<FCameraRigParameterDefinition> DefinitionsToAdd;
+
+	for (auto It = DefinitionsByName.CreateConstIterator(); It; ++It)
+	{
+		if (NumEqualDefinitions.FindRef(It.Key()) == NumObjectBindings)
+		{
+			DefinitionsToAdd.Add(It.Value());
+		}
+	}
+
+	DefinitionsToAdd.StableSort([](const FCameraRigParameterDefinition& A, const FCameraRigParameterDefinition& B)
+			{
+				return A.ParameterName.Compare(B.ParameterName) < 0;
+			});
+
+	if (DefinitionsToAdd.Num() > 0)
+	{
+		MenuBuilder.BeginSection(TEXT("CameraParameters"), LOCTEXT("AddCameraParametersMenuSection", "Camera Parameters"));
+
+		for (const FCameraRigParameterDefinition& Definition : DefinitionsToAdd)
+		{
+			MenuBuilder.AddMenuEntry(
+					FText::FromName(Definition.ParameterName),
+					LOCTEXT("AddCameraParameterMenuToolTip", "Adds a track for controlling this camera parameter."),
+					FSlateIcon(),
+					FUIAction(
+						FExecuteAction::CreateSP(this, &FGameplayCameraComponentTrackEditor::AddCameraParameterTrack, Definition, ObjectBindings),
+						FCanExecuteAction::CreateSP(this, &FGameplayCameraComponentTrackEditor::CanAddCameraParameterTrack, Definition, ObjectBindings[0]))
+					);
+		}
+
+		MenuBuilder.EndSection();
+	}
+}
+
+TSharedPtr<SWidget> FGameplayCameraComponentTrackEditor::BuildOutlinerEditWidget(const FGuid& ObjectBinding, UMovieSceneTrack* Track, const FBuildEditWidgetParams& Params)
+{
+	return FMovieSceneTrackEditor::BuildOutlinerEditWidget(ObjectBinding, Track, Params);
+}
+
+TSharedPtr<SWidget> FGameplayCameraComponentTrackEditor::BuildOutlinerColumnWidget(const FBuildColumnWidgetParams& Params, const FName& ColumnName)
+{
+	return FMovieSceneTrackEditor::BuildOutlinerColumnWidget(Params, ColumnName);
+}
+
+TSharedRef<ISequencerSection> FGameplayCameraComponentTrackEditor::MakeSectionInterface(UMovieSceneSection& SectionObject, UMovieSceneTrack& Track, FGuid ObjectBinding)
+{
+	return FMovieSceneTrackEditor::MakeSectionInterface(SectionObject, Track, ObjectBinding);
+}
+
+void FGameplayCameraComponentTrackEditor::OnRelease()
+{
+}
+
+bool FGameplayCameraComponentTrackEditor::SupportsType(TSubclassOf<UMovieSceneTrack> Type) const
+{
+	// This track editor doesn't support any track, it just extends object bindings.
+	return false;
+}
+
+void FGameplayCameraComponentTrackEditor::Tick(float DeltaTime)
+{
+}
+
+const FSlateBrush* FGameplayCameraComponentTrackEditor::GetIconBrush() const
+{
+	using namespace UE::Cameras;
+	TSharedRef<FGameplayCamerasEditorStyle> CamerasEditorStyle = FGameplayCamerasEditorStyle::Get();
+	return CamerasEditorStyle->GetBrush("Sequencer.CameraRigTrack");
+}
+
+bool FGameplayCameraComponentTrackEditor::OnAllowDrop(const FDragDropEvent& DragDropEvent, FSequencerDragDropParams& DragDropParams)
+{
+	return false;
+}
+
+FReply FGameplayCameraComponentTrackEditor::OnDrop(const FDragDropEvent& DragDropEvent, const FSequencerDragDropParams& DragDropParams)
+{
+	return FReply::Unhandled();
+}
+
+UGameplayCameraComponent* FGameplayCameraComponentTrackEditor::GetCameraComponentForBinding(const FGuid& ObjectBinding) const
+{
+	TSharedPtr<ISequencer> SequencerPtr = GetSequencer();
+	if (SequencerPtr.IsValid())
+	{
+		UObject* BoundObject = SequencerPtr->FindSpawnedObjectOrTemplate(ObjectBinding);
+		return Cast<UGameplayCameraComponent>(BoundObject);
+	}
+	return nullptr;
+}
+
+void FGameplayCameraComponentTrackEditor::AddCameraParameterTrack(FCameraRigParameterDefinition Definition, TArray<FGuid> ObjectBindings)
+{
+	TSubclassOf<UMovieSceneTrack> ParameterTrackType = GetParameterTrackFromDefinition(Definition);
+	if (ParameterTrackType)
+	{
+		const FScopedTransaction Transaction(LOCTEXT("AddCameraParameterTrack", "Add camera parameter track"));
+
+		for (FGuid ObjectBinding : ObjectBindings)
+		{
+			FFindOrCreateTrackResult Result = FindOrCreateTrackForObject(ObjectBinding, ParameterTrackType, Definition.ParameterName, true);
+			if (Result.bWasCreated)
+			{
+				UMovieScenePropertyTrack* NewTrack = CastChecked<UMovieScenePropertyTrack>(Result.Track);
+				NewTrack->SetPropertyNameAndPath(Definition.ParameterName, Definition.ParameterName.ToString());
+				InitializeNewTrack(NewTrack, Definition);
+			}
+		}
+
+		GetSequencer()->NotifyMovieSceneDataChanged(EMovieSceneDataChangeType::MovieSceneStructureItemAdded);
+	}
+}
+
+bool FGameplayCameraComponentTrackEditor::CanAddCameraParameterTrack(FCameraRigParameterDefinition Definition, FGuid ObjectBinding) const
+{
+	TSharedPtr<ISequencer> SequencerPtr = GetSequencer();
+	const UMovieScene* FocusedMovieScene = SequencerPtr->GetFocusedMovieSceneSequence()->GetMovieScene();
+
+	TSubclassOf<UMovieSceneTrack> ParameterTrackType = GetParameterTrackFromDefinition(Definition);
+	if (!ParameterTrackType)
+	{
+		return false;
+	}
+	
+	UMovieSceneTrack* ExistingParameterTrack = FocusedMovieScene->FindTrack(ParameterTrackType, ObjectBinding, Definition.ParameterName);
+	return ExistingParameterTrack == nullptr;
+}
+
+void FGameplayCameraComponentTrackEditor::InitializeNewTrack(UMovieScenePropertyTrack* NewTrack, const FCameraRigParameterDefinition& Definition) const
+{
+	if (Definition.ParameterType == ECameraRigInterfaceParameterType::Blendable)
+	{
+		switch (Definition.VariableType)
+		{
+			case ECameraVariableType::Vector2f:
+				CastChecked<UMovieSceneFloatVectorTrack>(NewTrack)->SetNumChannelsUsed(2);
+				break;
+			case ECameraVariableType::Vector3f:
+				CastChecked<UMovieSceneFloatVectorTrack>(NewTrack)->SetNumChannelsUsed(3);
+				break;
+			case ECameraVariableType::Vector4f:
+				CastChecked<UMovieSceneFloatVectorTrack>(NewTrack)->SetNumChannelsUsed(4);
+				break;
+			case ECameraVariableType::Vector2d:
+				CastChecked<UMovieSceneDoubleVectorTrack>(NewTrack)->SetNumChannelsUsed(2);
+				break;
+			case ECameraVariableType::Vector3d:
+				CastChecked<UMovieSceneDoubleVectorTrack>(NewTrack)->SetNumChannelsUsed(3);
+				break;
+			case ECameraVariableType::Vector4d:
+				CastChecked<UMovieSceneDoubleVectorTrack>(NewTrack)->SetNumChannelsUsed(4);
+				break;
+			default:
+				break;
+		}
+	}
+	else if (Definition.ParameterType == ECameraRigInterfaceParameterType::Data)
+	{
+		switch (Definition.DataType)
+		{
+			case ECameraContextDataType::Enum:
+				{
+					const UEnum* Enum = Cast<const UEnum>(Definition.DataTypeObject);
+					CastChecked<UMovieSceneEnumTrack>(NewTrack)->SetEnum(const_cast<UEnum*>(Enum));
+				}
+				break;
+			case ECameraContextDataType::Object:
+				{
+					const UClass* ObjectClass = Cast<const UClass>(Definition.DataTypeObject);
+					if (!ObjectClass || !ObjectClass->IsChildOf<AActor>())
+					{
+						CastChecked<UMovieSceneObjectPropertyTrack>(NewTrack)->PropertyClass = const_cast<UClass*>(ObjectClass);
+					}
+				}
+				break;
+			default:
+				break;
+		}
+	}
+}
+
+TSubclassOf<UMovieScenePropertyTrack> FGameplayCameraComponentTrackEditor::GetParameterTrackFromDefinition(const FCameraRigParameterDefinition& Definition) const
+{
+	if (Definition.ParameterType == ECameraRigInterfaceParameterType::Blendable)
+	{
+		switch (Definition.VariableType)
+		{
+			case ECameraVariableType::Boolean:
+				return UMovieSceneBoolTrack::StaticClass();
+			case ECameraVariableType::Integer32:
+				return UMovieSceneIntegerTrack::StaticClass();
+			case ECameraVariableType::Float:
+				return UMovieSceneFloatTrack::StaticClass();
+			case ECameraVariableType::Double:
+				return UMovieSceneDoubleTrack::StaticClass();
+			case ECameraVariableType::Vector2f:
+			case ECameraVariableType::Vector3f:
+			case ECameraVariableType::Vector4f:
+				return UMovieSceneFloatVectorTrack::StaticClass();
+			case ECameraVariableType::Vector2d:
+			case ECameraVariableType::Vector3d:
+			case ECameraVariableType::Vector4d:
+				return UMovieSceneDoubleVectorTrack::StaticClass();
+			case ECameraVariableType::Rotator3f:
+			case ECameraVariableType::Rotator3d:
+				return UMovieSceneRotatorTrack::StaticClass();
+			case ECameraVariableType::Transform3f:
+			case ECameraVariableType::Transform3d:
+				return UMovieSceneTransformTrack::StaticClass();
+			default:
+				break;
+		}
+	}
+	else if (Definition.ParameterType == ECameraRigInterfaceParameterType::Data)
+	{
+		// TODO
+	}
+	return nullptr;
+}
+
+#undef LOCTEXT_NAMESPACE
+
diff --git a/Engine/Plugins/Cameras/GameplayCameras/Source/GameplayCamerasEditor/Private/Sequencer/GameplayCameraComponentTrackEditor.h b/Engine/Plugins/Cameras/GameplayCameras/Source/GameplayCamerasEditor/Private/Sequencer/GameplayCameraComponentTrackEditor.h
new file mode 100644
--- /dev/null
+++ b/Engine/Plugins/Cameras/GameplayCameras/Source/GameplayCamerasEditor/Private/Sequencer/GameplayCameraComponentTrackEditor.h
@@ -0,0 +1,51 @@
+// Copyright Epic Games, Inc. All Rights Reserved.
+
+#pragma once
+
+#include "ISequencerTrackEditor.h"
+#include "MovieSceneTrackEditor.h"
+
+class ISequencer;
+class UGameplayCameraComponent;
+class UMovieScenePropertyTrack;
+struct FCameraRigParameterDefinition;
+
+class FGameplayCameraComponentTrackEditor : public FMovieSceneTrackEditor
+{
+public:
+
+	static TSharedRef<ISequencerTrackEditor> CreateTrackEditor(TSharedRef<ISequencer> OwningSequencer);
+
+	FGameplayCameraComponentTrackEditor(TSharedRef<ISequencer> InSequencer);
+
+protected:
+
+	// ISequencerTrackEditor interface.
+	virtual void BindCommands(TSharedRef<FUICommandList> SequencerCommandBindings) override;
+	virtual void BuildAddTrackMenu(FMenuBuilder& MenuBuilder) override;
+	virtual void BuildTrackContextMenu(FMenuBuilder& MenuBuilder, UMovieSceneTrack* Track) override;
+	virtual void ExtendObjectBindingTrackMenu(TSharedRef<FExtender> Extender, const TArray<FGuid>& ObjectBindings, const UClass* ObjectClass) override;
+	virtual TSharedPtr<SWidget> BuildOutlinerEditWidget(const FGuid& ObjectBinding, UMovieSceneTrack* Track, const FBuildEditWidgetParams& Params) override;
+	virtual TSharedPtr<SWidget> BuildOutlinerColumnWidget(const FBuildColumnWidgetParams& Params, const FName& ColumnName) override;
+	virtual TSharedRef<ISequencerSection> MakeSectionInterface(UMovieSceneSection& SectionObject, UMovieSceneTrack& Track, FGuid ObjectBinding) override;
+	virtual void OnRelease() override;
+	virtual bool SupportsType(TSubclassOf<UMovieSceneTrack> Type) const override;
+	virtual void Tick(float DeltaTime) override;
+	virtual const FSlateBrush* GetIconBrush() const override;
+	virtual bool OnAllowDrop(const FDragDropEvent& DragDropEvent, FSequencerDragDropParams& DragDropParams) override;
+	virtual FReply OnDrop(const FDragDropEvent& DragDropEvent, const FSequencerDragDropParams& DragDropParams) override;
+
+private:
+
+	UGameplayCameraComponent* GetCameraComponentForBinding(const FGuid& ObjectBinding) const;
+
+	void OnExtendObjectBindingTrackMenu(FMenuBuilder& MenuBuilder, TArray<FGuid> ObjectBindings);
+
+	void AddCameraParameterTrack(FCameraRigParameterDefinition Definition, TArray<FGuid> ObjectBindings);
+	bool CanAddCameraParameterTrack(FCameraRigParameterDefinition Definition, FGuid ObjectBinding) const;
+
+	void InitializeNewTrack(UMovieScenePropertyTrack* NewTrack, const FCameraRigParameterDefinition& Definition) const;
+
+	TSubclassOf<UMovieScenePropertyTrack> GetParameterTrackFromDefinition(const FCameraRigParameterDefinition& Definition) const;
+};
+
