Cameras: add more type support in sequencer integration code

- add support for blendable structs in the sequencer integration code
- add framing zone sequencer track editor
- handle actor references in context data as actor reference tracks

diff --git a/Engine/Plugins/Cameras/GameplayCameras/Source/GameplayCameras/Private/MovieScene/MovieSceneCameraFramingZonePropertySystem.cpp b/Engine/Plugins/Cameras/GameplayCameras/Source/GameplayCameras/Private/MovieScene/MovieSceneCameraFramingZonePropertySystem.cpp
new file mode 100644
--- /dev/null
+++ b/Engine/Plugins/Cameras/GameplayCameras/Source/GameplayCameras/Private/MovieScene/MovieSceneCameraFramingZonePropertySystem.cpp
@@ -0,0 +1,27 @@
+// Copyright Epic Games, Inc. All Rights Reserved.
+
+#include "MovieScene/MovieSceneCameraFramingZonePropertySystem.h"
+
+#include "MovieScene/MovieSceneGameplayCamerasComponentTypes.h"
+#include "Systems/DoubleChannelEvaluatorSystem.h"
+#include "Systems/MovieScenePiecewiseDoubleBlenderSystem.h"
+
+#include UE_INLINE_GENERATED_CPP_BY_NAME(MovieSceneCameraFramingZonePropertySystem)
+
+UMovieSceneCameraFramingZonePropertySystem::UMovieSceneCameraFramingZonePropertySystem(const FObjectInitializer& ObjInit)
+	: Super(ObjInit)
+{
+	BindToProperty(UE::Cameras::FMovieSceneGameplayCamerasComponentTypes::Get()->CameraFramingZone);
+
+	if (HasAnyFlags(RF_ClassDefaultObject))
+	{
+		DefineImplicitPrerequisite(UMovieScenePiecewiseDoubleBlenderSystem::StaticClass(), GetClass());
+		DefineImplicitPrerequisite(UDoubleChannelEvaluatorSystem::StaticClass(), GetClass());
+	}
+}
+
+void UMovieSceneCameraFramingZonePropertySystem::OnRun(FSystemTaskPrerequisites& InPrerequisites, FSystemSubsequentTasks& Subsequents)
+{
+	Super::OnRun(InPrerequisites, Subsequents);
+}
+
diff --git a/Engine/Plugins/Cameras/GameplayCameras/Source/GameplayCameras/Private/MovieScene/MovieSceneCameraFramingZonePropertySystem.h b/Engine/Plugins/Cameras/GameplayCameras/Source/GameplayCameras/Private/MovieScene/MovieSceneCameraFramingZonePropertySystem.h
new file mode 100644
--- /dev/null
+++ b/Engine/Plugins/Cameras/GameplayCameras/Source/GameplayCameras/Private/MovieScene/MovieSceneCameraFramingZonePropertySystem.h
@@ -0,0 +1,23 @@
+// Copyright Epic Games, Inc. All Rights Reserved.
+
+#pragma once
+
+#include "Nodes/Framing/CameraFramingZone.h"
+#include "Systems/MovieScenePropertySystem.h"
+
+#include "MovieSceneCameraFramingZonePropertySystem.generated.h"
+
+Expose_TNameOf(FCameraFramingZone);
+
+UCLASS(MinimalAPI)
+class UMovieSceneCameraFramingZonePropertySystem : public UMovieScenePropertySystem
+{
+public:
+
+	GENERATED_BODY()
+	
+	UMovieSceneCameraFramingZonePropertySystem(const FObjectInitializer& ObjInit);
+
+	virtual void OnRun(FSystemTaskPrerequisites& InPrerequisites, FSystemSubsequentTasks& Subsequents) override;
+};
+
diff --git a/Engine/Plugins/Cameras/GameplayCameras/Source/GameplayCameras/Private/MovieScene/MovieSceneCameraFramingZoneSection.cpp b/Engine/Plugins/Cameras/GameplayCameras/Source/GameplayCameras/Private/MovieScene/MovieSceneCameraFramingZoneSection.cpp
new file mode 100644
--- /dev/null
+++ b/Engine/Plugins/Cameras/GameplayCameras/Source/GameplayCameras/Private/MovieScene/MovieSceneCameraFramingZoneSection.cpp
@@ -0,0 +1,126 @@
+// Copyright Epic Games, Inc. All Rights Reserved.
+
+#include "MovieScene/MovieSceneCameraFramingZoneSection.h"
+
+#include "MovieScene/MovieSceneCameraFramingZoneTrack.h"
+#include "MovieScene/MovieSceneGameplayCamerasComponentTypes.h"
+#include "Channels/MovieSceneChannelProxy.h"
+#include "Compilation/MovieSceneTemplateInterrogation.h"
+#include "Evaluation/MovieSceneEvaluationTrack.h"
+#include "Evaluation/MovieSceneEvaluationTemplateInstance.h"
+#include "Evaluation/MovieScenePropertyTemplate.h"
+#include "Tracks/MovieScenePropertyTrack.h"
+
+#include UE_INLINE_GENERATED_CPP_BY_NAME(MovieSceneCameraFramingZoneSection)
+
+#define LOCTEXT_NAMESPACE "MovieSceneCameraFramingZoneSection"
+
+#if WITH_EDITOR
+
+struct FCameraFramingZoneSectionEditorData
+{
+	FCameraFramingZoneSectionEditorData()
+	{
+		MetaData[0].SetIdentifiers("Left", LOCTEXT("LeftText", "Left"));
+		MetaData[0].SubPropertyPath = MetaData[0].Name;
+		MetaData[0].SortOrder = 0;
+		MetaData[0].bCanCollapseToTrack = false;
+
+		MetaData[1].SetIdentifiers("Top", LOCTEXT("TopText", "Top"));
+		MetaData[1].SubPropertyPath = MetaData[1].Name;
+		MetaData[1].SortOrder = 1;
+		MetaData[1].bCanCollapseToTrack = false;
+
+		MetaData[2].SetIdentifiers("Right", LOCTEXT("RightText", "Right"));
+		MetaData[2].SubPropertyPath = MetaData[2].Name;
+		MetaData[2].SortOrder = 2;
+		MetaData[2].bCanCollapseToTrack = false;
+
+		MetaData[3].SetIdentifiers("Bottom", LOCTEXT("BottomText", "Bottom"));
+		MetaData[3].SubPropertyPath = MetaData[3].Name;
+		MetaData[3].SortOrder = 3;
+		MetaData[3].bCanCollapseToTrack = false;
+
+		ExternalValues[0].OnGetExternalValue = ExtractLeftChannel;
+		ExternalValues[1].OnGetExternalValue = ExtractTopChannel;
+		ExternalValues[2].OnGetExternalValue = ExtractRightChannel;
+		ExternalValues[3].OnGetExternalValue = ExtractBottomChannel;
+	}
+
+	static TOptional<double> ExtractLeftChannel(UObject& InObject, FTrackInstancePropertyBindings* Bindings)
+	{
+		return Bindings ? Bindings->GetCurrentValue<FCameraFramingZone>(InObject).Left : TOptional<double>();
+	}
+	static TOptional<double> ExtractTopChannel(UObject& InObject, FTrackInstancePropertyBindings* Bindings)
+	{
+		return Bindings ? Bindings->GetCurrentValue<FCameraFramingZone>(InObject).Top : TOptional<double>();
+	}
+	static TOptional<double> ExtractRightChannel(UObject& InObject, FTrackInstancePropertyBindings* Bindings)
+	{
+		return Bindings ? Bindings->GetCurrentValue<FCameraFramingZone>(InObject).Right : TOptional<double>();
+	}
+	static TOptional<double> ExtractBottomChannel(UObject& InObject, FTrackInstancePropertyBindings* Bindings)
+	{
+		return Bindings ? Bindings->GetCurrentValue<FCameraFramingZone>(InObject).Bottom : TOptional<double>();
+	}
+
+	FMovieSceneChannelMetaData      MetaData[4];
+	TMovieSceneExternalValue<double> ExternalValues[4];
+};
+
+#endif	// WITH_EDITOR
+
+UMovieSceneCameraFramingZoneSection::UMovieSceneCameraFramingZoneSection(const FObjectInitializer& ObjectInitializer)
+	: Super(ObjectInitializer)
+{
+	EvalOptions.EnableAndSetCompletionMode(EMovieSceneCompletionMode::ProjectDefault);
+
+	BlendType = EMovieSceneBlendType::Absolute;
+
+	FMovieSceneChannelProxyData Channels;
+	bSupportsInfiniteRange = true;
+
+#if WITH_EDITOR
+
+	static const FCameraFramingZoneSectionEditorData EditorData;
+	Channels.Add(LeftMarginCurve, EditorData.MetaData[0], EditorData.ExternalValues[0]);
+	Channels.Add(TopMarginCurve, EditorData.MetaData[1], EditorData.ExternalValues[1]);
+	Channels.Add(RightMarginCurve, EditorData.MetaData[2], EditorData.ExternalValues[2]);
+	Channels.Add(BottomMarginCurve, EditorData.MetaData[3], EditorData.ExternalValues[3]);
+
+#else
+
+	Channels.Add(LeftMarginCurve);
+	Channels.Add(TopMarginCurve);
+	Channels.Add(RightMarginCurve);
+	Channels.Add(BottomMarginCurve);
+
+#endif
+	
+	ChannelProxy = MakeShared<FMovieSceneChannelProxy>(MoveTemp(Channels));
+}
+
+bool UMovieSceneCameraFramingZoneSection::PopulateEvaluationFieldImpl(const TRange<FFrameNumber>& EffectiveRange, const FMovieSceneEvaluationFieldEntityMetaData& InMetaData, FMovieSceneEntityComponentFieldBuilder* OutFieldBuilder)
+{
+	FMovieScenePropertyTrackEntityImportHelper::PopulateEvaluationField(*this, EffectiveRange, InMetaData, OutFieldBuilder);
+	return true;
+}
+
+void UMovieSceneCameraFramingZoneSection::ImportEntityImpl(UMovieSceneEntitySystemLinker* EntityLinker, const FEntityImportParams& Params, FImportedEntity* OutImportedEntity)
+{
+	using namespace UE::Cameras;
+	using namespace UE::MovieScene;
+
+	const FBuiltInComponentTypes* Components = FBuiltInComponentTypes::Get();
+	const FMovieSceneGameplayCamerasComponentTypes* GameplayCamerasComponents = FMovieSceneGameplayCamerasComponentTypes::Get();
+
+	FPropertyTrackEntityImportHelper(GameplayCamerasComponents->CameraFramingZone)
+		.AddConditional(Components->DoubleChannel[0], &LeftMarginCurve, LeftMarginCurve.HasAnyData())
+		.AddConditional(Components->DoubleChannel[1], &TopMarginCurve, TopMarginCurve.HasAnyData())
+		.AddConditional(Components->DoubleChannel[2], &RightMarginCurve, RightMarginCurve.HasAnyData())
+		.AddConditional(Components->DoubleChannel[3], &BottomMarginCurve, BottomMarginCurve.HasAnyData())
+		.Commit(this, Params, OutImportedEntity);
+}
+
+#undef LOCTEXT_NAMESPACE
+
diff --git a/Engine/Plugins/Cameras/GameplayCameras/Source/GameplayCameras/Private/MovieScene/MovieSceneCameraFramingZoneTrack.cpp b/Engine/Plugins/Cameras/GameplayCameras/Source/GameplayCameras/Private/MovieScene/MovieSceneCameraFramingZoneTrack.cpp
new file mode 100644
--- /dev/null
+++ b/Engine/Plugins/Cameras/GameplayCameras/Source/GameplayCameras/Private/MovieScene/MovieSceneCameraFramingZoneTrack.cpp
@@ -0,0 +1,24 @@
+// Copyright Epic Games, Inc. All Rights Reserved.
+
+#include "MovieScene/MovieSceneCameraFramingZoneTrack.h"
+#include "MovieScene/MovieSceneCameraFramingZoneSection.h"
+
+#include UE_INLINE_GENERATED_CPP_BY_NAME(MovieSceneCameraFramingZoneTrack)
+
+UMovieSceneCameraFramingZoneTrack::UMovieSceneCameraFramingZoneTrack(const FObjectInitializer& Init)
+	: Super(Init)
+{
+	EvalOptions.bEvaluateNearestSection_DEPRECATED = EvalOptions.bCanEvaluateNearestSection = true;
+	SupportedBlendTypes = FMovieSceneBlendTypeField::All();
+}
+
+bool UMovieSceneCameraFramingZoneTrack::SupportsType(TSubclassOf<UMovieSceneSection> SectionClass) const
+{
+	return SectionClass == UMovieSceneCameraFramingZoneSection::StaticClass();
+}
+
+UMovieSceneSection* UMovieSceneCameraFramingZoneTrack::CreateNewSection()
+{
+	return NewObject<UMovieSceneCameraFramingZoneSection>(this, NAME_None, RF_Transactional);
+}
+
diff --git a/Engine/Plugins/Cameras/GameplayCameras/Source/GameplayCameras/Private/MovieScene/MovieSceneGameplayCamerasComponentTypes.cpp b/Engine/Plugins/Cameras/GameplayCameras/Source/GameplayCameras/Private/MovieScene/MovieSceneGameplayCamerasComponentTypes.cpp
new file mode 100644
--- /dev/null
+++ b/Engine/Plugins/Cameras/GameplayCameras/Source/GameplayCameras/Private/MovieScene/MovieSceneGameplayCamerasComponentTypes.cpp
@@ -0,0 +1,59 @@
+// Copyright Epic Games, Inc. All Rights Reserved.
+
+#include "MovieScene/MovieSceneGameplayCamerasComponentTypes.h"
+
+#include "EntitySystem/BuiltInComponentTypes.h"
+#include "EntitySystem/MovieSceneComponentRegistry.h"
+#include "EntitySystem/MovieSceneEntityFactoryTemplates.h"
+#include "EntitySystem/MovieSceneEntitySystemLinker.h"
+#include "EntitySystem/MovieScenePropertyComponentHandler.h"
+#include "MovieSceneTracksComponentTypes.h"
+#include "Systems/MovieScenePiecewiseDoubleBlenderSystem.h"
+
+namespace UE::Cameras
+{
+
+static bool GMovieSceneUMGComponentTypesDestroyed = false;
+static TUniquePtr<FMovieSceneGameplayCamerasComponentTypes> GMovieSceneUMGComponentTypes;
+
+FMovieSceneGameplayCamerasComponentTypes::FMovieSceneGameplayCamerasComponentTypes()
+{
+	using namespace UE::MovieScene;
+
+	FBuiltInComponentTypes* BuiltInComponents = FBuiltInComponentTypes::Get();
+	FComponentRegistry* ComponentRegistry = UMovieSceneEntitySystemLinker::GetComponents();
+
+	ComponentRegistry->NewPropertyType(CameraFramingZone, TEXT("Camera Framing Zone"));
+
+	BuiltInComponents->PropertyRegistry.DefineCompositeProperty(CameraFramingZone, TEXT("Apply FCameraFramingZone Properties"))
+	.AddComposite(BuiltInComponents->DoubleResult[0], &FCameraFramingZone::Left)
+	.AddComposite(BuiltInComponents->DoubleResult[1], &FCameraFramingZone::Top)
+	.AddComposite(BuiltInComponents->DoubleResult[2], &FCameraFramingZone::Right)
+	.AddComposite(BuiltInComponents->DoubleResult[3], &FCameraFramingZone::Bottom)
+	.SetBlenderSystem<UMovieScenePiecewiseDoubleBlenderSystem>()
+	.SetCustomAccessors(&CustomCameraFramingZoneAccessors)
+	.Commit();
+}
+
+FMovieSceneGameplayCamerasComponentTypes::~FMovieSceneGameplayCamerasComponentTypes()
+{
+}
+
+void FMovieSceneGameplayCamerasComponentTypes::Destroy()
+{
+	GMovieSceneUMGComponentTypes.Reset();
+	GMovieSceneUMGComponentTypesDestroyed = true;
+}
+
+FMovieSceneGameplayCamerasComponentTypes* FMovieSceneGameplayCamerasComponentTypes::Get()
+{
+	if (!GMovieSceneUMGComponentTypes.IsValid())
+	{
+		check(!GMovieSceneUMGComponentTypesDestroyed);
+		GMovieSceneUMGComponentTypes.Reset(new FMovieSceneGameplayCamerasComponentTypes);
+	}
+	return GMovieSceneUMGComponentTypes.Get();
+}
+
+} // namespace UE::Cameras
+
diff --git a/Engine/Plugins/Cameras/GameplayCameras/Source/GameplayCameras/Private/MovieScene/MovieSceneGameplayCamerasComponentTypes.h b/Engine/Plugins/Cameras/GameplayCameras/Source/GameplayCameras/Private/MovieScene/MovieSceneGameplayCamerasComponentTypes.h
new file mode 100644
--- /dev/null
+++ b/Engine/Plugins/Cameras/GameplayCameras/Source/GameplayCameras/Private/MovieScene/MovieSceneGameplayCamerasComponentTypes.h
@@ -0,0 +1,34 @@
+// Copyright Epic Games, Inc. All Rights Reserved.
+
+#pragma once
+
+#include "Nodes/Framing/CameraFramingZone.h"
+#include "EntitySystem/MovieSceneEntityIDs.h"
+#include "EntitySystem/MovieScenePropertySystemTypes.h"
+#include "EntitySystem/MovieScenePropertyTraits.h"
+#include "EntitySystem/MovieScenePropertyMetaDataTraits.h"
+
+namespace UE::Cameras
+{
+
+using FCameraFramingZoneTraits = UE::MovieScene::TDirectPropertyTraits<FCameraFramingZone>;
+
+struct FMovieSceneGameplayCamerasComponentTypes
+{
+	GAMEPLAYCAMERAS_API ~FMovieSceneGameplayCamerasComponentTypes();
+
+	UE::MovieScene::TPropertyComponents<FCameraFramingZoneTraits> CameraFramingZone;
+
+	UE::MovieScene::TCustomPropertyRegistration<FCameraFramingZoneTraits, 1> CustomCameraFramingZoneAccessors;
+
+	static GAMEPLAYCAMERAS_API void Destroy();
+
+	static GAMEPLAYCAMERAS_API FMovieSceneGameplayCamerasComponentTypes* Get();
+
+private:
+
+	FMovieSceneGameplayCamerasComponentTypes();
+};
+
+} // namespace UE::Cameras
+
diff --git a/Engine/Plugins/Cameras/GameplayCameras/Source/GameplayCameras/Public/MovieScene/MovieSceneCameraFramingZoneSection.h b/Engine/Plugins/Cameras/GameplayCameras/Source/GameplayCameras/Public/MovieScene/MovieSceneCameraFramingZoneSection.h
new file mode 100644
--- /dev/null
+++ b/Engine/Plugins/Cameras/GameplayCameras/Source/GameplayCameras/Public/MovieScene/MovieSceneCameraFramingZoneSection.h
@@ -0,0 +1,40 @@
+// Copyright Epic Games, Inc. All Rights Reserved.
+
+#pragma once
+
+#include "Channels/MovieSceneDoubleChannel.h"
+#include "Curves/KeyHandle.h"
+#include "Curves/RichCurve.h"
+#include "EntitySystem/IMovieSceneEntityProvider.h"
+#include "MovieSceneSection.h"
+#include "Nodes/Framing/CameraFramingZone.h"
+
+#include "MovieSceneCameraFramingZoneSection.generated.h"
+
+UCLASS(MinimalAPI)
+class UMovieSceneCameraFramingZoneSection 
+	: public UMovieSceneSection
+	, public IMovieSceneEntityProvider
+{
+	GENERATED_UCLASS_BODY()
+
+public:
+
+	UPROPERTY()
+	FMovieSceneDoubleChannel LeftMarginCurve;
+
+	UPROPERTY()
+	FMovieSceneDoubleChannel TopMarginCurve;
+
+	UPROPERTY()
+	FMovieSceneDoubleChannel RightMarginCurve;
+
+	UPROPERTY()
+	FMovieSceneDoubleChannel BottomMarginCurve;
+
+private:
+
+	virtual bool PopulateEvaluationFieldImpl(const TRange<FFrameNumber>& EffectiveRange, const FMovieSceneEvaluationFieldEntityMetaData& InMetaData, FMovieSceneEntityComponentFieldBuilder* OutFieldBuilder) override;
+	virtual void ImportEntityImpl(UMovieSceneEntitySystemLinker* EntityLinker, const FEntityImportParams& Params, FImportedEntity* OutImportedEntity) override;
+};
+
diff --git a/Engine/Plugins/Cameras/GameplayCameras/Source/GameplayCameras/Public/MovieScene/MovieSceneCameraFramingZoneTrack.h b/Engine/Plugins/Cameras/GameplayCameras/Source/GameplayCameras/Public/MovieScene/MovieSceneCameraFramingZoneTrack.h
new file mode 100644
--- /dev/null
+++ b/Engine/Plugins/Cameras/GameplayCameras/Source/GameplayCameras/Public/MovieScene/MovieSceneCameraFramingZoneTrack.h
@@ -0,0 +1,22 @@
+// Copyright Epic Games, Inc. All Rights Reserved.
+
+#pragma once
+
+#include "Tracks/MovieScenePropertyTrack.h"
+
+#include "MovieSceneCameraFramingZoneTrack.generated.h"
+
+UCLASS( MinimalAPI )
+class UMovieSceneCameraFramingZoneTrack : public UMovieScenePropertyTrack
+{
+	GENERATED_BODY()
+
+public:
+
+	UMovieSceneCameraFramingZoneTrack(const FObjectInitializer& Init);
+	
+	// UMovieSceneTrack interface.
+	virtual bool SupportsType(TSubclassOf<UMovieSceneSection> SectionClass) const override;
+	virtual UMovieSceneSection* CreateNewSection() override;
+};
+
diff --git a/Engine/Plugins/Cameras/GameplayCameras/Source/GameplayCamerasEditor/GameplayCamerasEditor.Build.cs b/Engine/Plugins/Cameras/GameplayCameras/Source/GameplayCamerasEditor/GameplayCamerasEditor.Build.cs
--- a/Engine/Plugins/Cameras/GameplayCameras/Source/GameplayCamerasEditor/GameplayCamerasEditor.Build.cs
+++ b/Engine/Plugins/Cameras/GameplayCameras/Source/GameplayCamerasEditor/GameplayCamerasEditor.Build.cs
@@ -39,6 +39,7 @@ public class GameplayCamerasEditor : ModuleRules
 				"KismetWidgets",
 				"LevelEditor",
 				"MovieScene",
+				"MovieSceneTools",
 				"MovieSceneTracks",
 				"Projects",
 				"RenderCore",
diff --git a/Engine/Plugins/Cameras/GameplayCameras/Source/GameplayCamerasEditor/Private/GameplayCamerasEditorModule.cpp b/Engine/Plugins/Cameras/GameplayCameras/Source/GameplayCamerasEditor/Private/GameplayCamerasEditorModule.cpp
--- a/Engine/Plugins/Cameras/GameplayCameras/Source/GameplayCamerasEditor/Private/GameplayCamerasEditorModule.cpp
+++ b/Engine/Plugins/Cameras/GameplayCameras/Source/GameplayCamerasEditor/Private/GameplayCamerasEditorModule.cpp
@@ -45,6 +45,7 @@
 #include "Kismet2/KismetEditorUtilities.h"
 #include "Modules/ModuleManager.h"
 #include "PropertyEditorModule.h"
+#include "Sequencer/CameraFramingZoneTrackEditor.h"
 #include "Sequencer/GameplayCameraComponentTrackEditor.h"
 #include "Styles/GameplayCamerasEditorStyle.h"
 #include "ToolMenus.h"
@@ -553,14 +554,18 @@ private:
 	
 	void RegisterSequencerTracks()
 	{
+		using namespace UE::Cameras;
+
 		ISequencerModule& SequencerModule = FModuleManager::Get().LoadModuleChecked<ISequencerModule>("Sequencer");
 		GameplayCameraComponentTrackCreateEditorHandle = SequencerModule.RegisterTrackEditor(FOnCreateTrackEditor::CreateStatic(&FGameplayCameraComponentTrackEditor::CreateTrackEditor));
+		CameraFramingZoneTrackCreateEditorHandle = SequencerModule.RegisterPropertyTrackEditor<FCameraFramingZoneTrackEditor>();
 	}
 
 	void UnregisterSequencerTracks()
 	{
 		ISequencerModule& SequencerModule = FModuleManager::Get().LoadModuleChecked<ISequencerModule>("Sequencer");
 		SequencerModule.UnRegisterTrackEditor(GameplayCameraComponentTrackCreateEditorHandle);
+		SequencerModule.UnRegisterTrackEditor(CameraFramingZoneTrackCreateEditorHandle);
 	}
 
 	void InitializeLiveEditManager()
@@ -596,6 +601,7 @@ private:
 	TMap<FString, FOnCreateDebugCategoryPanel> DebugCategoryPanelCreators;
 	
 	FDelegateHandle GameplayCameraComponentTrackCreateEditorHandle;
+	FDelegateHandle CameraFramingZoneTrackCreateEditorHandle;
 
 #if UE_GAMEPLAY_CAMERAS_TRACE
 	TSharedPtr<UE::Cameras::FCameraSystemTraceModule> TraceModule;
diff --git a/Engine/Plugins/Cameras/GameplayCameras/Source/GameplayCamerasEditor/Private/Sequencer/CameraFramingZoneTrackEditor.cpp b/Engine/Plugins/Cameras/GameplayCameras/Source/GameplayCamerasEditor/Private/Sequencer/CameraFramingZoneTrackEditor.cpp
new file mode 100644
--- /dev/null
+++ b/Engine/Plugins/Cameras/GameplayCameras/Source/GameplayCamerasEditor/Private/Sequencer/CameraFramingZoneTrackEditor.cpp
@@ -0,0 +1,53 @@
+// Copyright Epic Games, Inc. All Rights Reserved.
+
+#include "Sequencer/CameraFramingZoneTrackEditor.h"
+
+#include "Channels/MovieSceneDoubleChannel.h"
+#include "KeyPropertyParams.h"
+#include "Misc/Optional.h"
+#include "Nodes/Framing/CameraFramingZone.h"
+#include "PropertyPath.h"
+#include "UObject/NameTypes.h"
+#include "UObject/UnrealNames.h"
+#include "UObject/UnrealType.h"
+#include "UObject/WeakFieldPtr.h"
+
+class ISequencer;
+class ISequencerTrackEditor;
+class UMovieSceneSection;
+
+namespace UE::Cameras
+{
+
+FName FCameraFramingZoneTrackEditor::LeftName("Left");
+FName FCameraFramingZoneTrackEditor::TopName("Top");
+FName FCameraFramingZoneTrackEditor::RightName("Right");
+FName FCameraFramingZoneTrackEditor::BottomName("Bottom");
+
+TSharedRef<ISequencerTrackEditor> FCameraFramingZoneTrackEditor::CreateTrackEditor(TSharedRef<ISequencer> InSequencer)
+{
+	return MakeShareable(new FCameraFramingZoneTrackEditor(InSequencer));
+}
+
+void FCameraFramingZoneTrackEditor::GenerateKeysFromPropertyChanged(const FPropertyChangedParams& PropertyChangedParams, UMovieSceneSection* SectionToKey, FGeneratedTrackKeys& OutGeneratedKeys)
+{
+	using namespace UE::MovieScene;
+
+	FPropertyPath StructPath = PropertyChangedParams.StructPathToKey;
+	FName ChannelName = StructPath.GetNumProperties() != 0 ? StructPath.GetLeafMostProperty().Property->GetFName() : NAME_None;
+
+	FCameraFramingZone CameraFramingZone = PropertyChangedParams.GetPropertyValue<FCameraFramingZone>();
+
+	const bool bKeyLeft   = ChannelName == NAME_None || ChannelName == LeftName;
+	const bool bKeyTop    = ChannelName == NAME_None || ChannelName == TopName;
+	const bool bKeyRight  = ChannelName == NAME_None || ChannelName == RightName;
+	const bool bKeyBottom = ChannelName == NAME_None || ChannelName == BottomName;
+
+	OutGeneratedKeys.Add(FMovieSceneChannelValueSetter::Create<FMovieSceneDoubleChannel>(0, CameraFramingZone.Left, bKeyLeft));
+	OutGeneratedKeys.Add(FMovieSceneChannelValueSetter::Create<FMovieSceneDoubleChannel>(1, CameraFramingZone.Top, bKeyTop));
+	OutGeneratedKeys.Add(FMovieSceneChannelValueSetter::Create<FMovieSceneDoubleChannel>(2, CameraFramingZone.Right, bKeyRight));
+	OutGeneratedKeys.Add(FMovieSceneChannelValueSetter::Create<FMovieSceneDoubleChannel>(3, CameraFramingZone.Bottom, bKeyBottom));
+}
+
+}  // namespace UE::Cameras
+
diff --git a/Engine/Plugins/Cameras/GameplayCameras/Source/GameplayCamerasEditor/Private/Sequencer/CameraFramingZoneTrackEditor.h b/Engine/Plugins/Cameras/GameplayCameras/Source/GameplayCamerasEditor/Private/Sequencer/CameraFramingZoneTrackEditor.h
new file mode 100644
--- /dev/null
+++ b/Engine/Plugins/Cameras/GameplayCameras/Source/GameplayCamerasEditor/Private/Sequencer/CameraFramingZoneTrackEditor.h
@@ -0,0 +1,56 @@
+// Copyright Epic Games, Inc. All Rights Reserved.
+
+#pragma once
+
+#include "PropertyTrackEditor.h"
+
+#include "AnimatedPropertyKey.h"
+#include "Containers/Array.h"
+#include "Containers/ArrayView.h"
+#include "Containers/ContainerAllocationPolicies.h"
+#include "Delegates/Delegate.h"
+#include "HAL/PlatformCrt.h"
+#include "KeyframeTrackEditor.h"
+#include "MovieScene/MovieSceneCameraFramingZoneTrack.h"
+#include "Templates/SharedPointer.h"
+#include "Templates/UnrealTemplate.h"
+
+class FName;
+class FPropertyChangedParams;
+class ISequencer;
+class ISequencerTrackEditor;
+class UMovieSceneSection;
+
+namespace UE::Cameras
+{
+
+class FCameraFramingZoneTrackEditor : public FPropertyTrackEditor<UMovieSceneCameraFramingZoneTrack>
+{
+public:
+
+	FCameraFramingZoneTrackEditor(TSharedRef<ISequencer> InSequencer)
+		: FPropertyTrackEditor(InSequencer, GetAnimatedPropertyTypes())
+	{
+	}
+
+	static TArray<FAnimatedPropertyKey, TInlineAllocator<1>> GetAnimatedPropertyTypes()
+	{
+		return TArray<FAnimatedPropertyKey, TInlineAllocator<1>>({ FAnimatedPropertyKey::FromStructType("CameraFramingZone") });
+	}
+
+	static TSharedRef<ISequencerTrackEditor> CreateTrackEditor(TSharedRef<ISequencer> OwningSequencer);
+
+protected:
+
+	virtual void GenerateKeysFromPropertyChanged( const FPropertyChangedParams& PropertyChangedParams, UMovieSceneSection* SectionToKey, FGeneratedTrackKeys& OutGeneratedKeys) override;
+
+private:
+
+	static FName LeftName;
+	static FName TopName;
+	static FName RightName;
+	static FName BottomName;
+};
+
+}  // namespace UE::Cameras
+
diff --git a/Engine/Plugins/Cameras/GameplayCameras/Source/GameplayCamerasEditor/Private/Sequencer/GameplayCameraComponentTrackEditor.cpp b/Engine/Plugins/Cameras/GameplayCameras/Source/GameplayCamerasEditor/Private/Sequencer/GameplayCameraComponentTrackEditor.cpp
--- a/Engine/Plugins/Cameras/GameplayCameras/Source/GameplayCamerasEditor/Private/Sequencer/GameplayCameraComponentTrackEditor.cpp
+++ b/Engine/Plugins/Cameras/GameplayCameras/Source/GameplayCamerasEditor/Private/Sequencer/GameplayCameraComponentTrackEditor.cpp
@@ -5,7 +5,10 @@
 #include "Core/CameraAsset.h"
 #include "Core/CameraRigParameterDefinition.h"
 #include "Framework/MultiBox/MultiBoxBuilder.h"
+#include "GameFramework/Actor.h"
 #include "GameFramework/GameplayCameraComponent.h"
+#include "MovieScene/MovieSceneCameraFramingZoneTrack.h"
+#include "Nodes/Framing/CameraFramingZone.h"
 #include "Styles/GameplayCamerasEditorStyle.h"
 #include "Tracks/MovieSceneActorReferenceTrack.h"
 #include "Tracks/MovieSceneBoolTrack.h"
@@ -293,13 +296,35 @@ TSubclassOf<UMovieScenePropertyTrack> FGameplayCameraComponentTrackEditor::GetPa
 			case ECameraVariableType::Transform3f:
 			case ECameraVariableType::Transform3d:
 				return UMovieSceneTransformTrack::StaticClass();
+			case ECameraVariableType::BlendableStruct:
+				if (Definition.BlendableStructType == FCameraFramingZone::StaticStruct())
+				{
+					return UMovieSceneCameraFramingZoneTrack::StaticClass();
+				}
+				break;
 			default:
 				break;
 		}
 	}
 	else if (Definition.ParameterType == ECameraRigInterfaceParameterType::Data)
 	{
-		// TODO
+		switch (Definition.DataType)
+		{
+			case ECameraContextDataType::String:
+				return UMovieSceneStringTrack::StaticClass();
+			case ECameraContextDataType::Enum:
+				return UMovieSceneEnumTrack::StaticClass();
+			case ECameraContextDataType::Object:
+				if (Definition.DataTypeObject)
+				{
+					const UClass* ObjectClass = Cast<const UClass>(Definition.DataTypeObject);
+					if (ObjectClass && ObjectClass->IsChildOf<AActor>())
+					{
+						return UMovieSceneActorReferenceTrack::StaticClass();
+					}
+				}
+				return UMovieSceneObjectPropertyTrack::StaticClass();
+		}
 	}
 	return nullptr;
 }
