Cameras: don't crash when enum types or struct types aren't found for a context data table entry

diff --git a/Engine/Plugins/Cameras/GameplayCameras/Source/GameplayCameras/Private/Core/CameraContextDataTable.cpp b/Engine/Plugins/Cameras/GameplayCameras/Source/GameplayCameras/Private/Core/CameraContextDataTable.cpp
--- a/Engine/Plugins/Cameras/GameplayCameras/Source/GameplayCameras/Private/Core/CameraContextDataTable.cpp
+++ b/Engine/Plugins/Cameras/GameplayCameras/Source/GameplayCameras/Private/Core/CameraContextDataTable.cpp
@@ -26,8 +26,11 @@ void FCameraContextDataTable::AddReferencedObjects(FReferenceCollector& Referenc
 		{
 			case ECameraContextDataType::Struct:
 				{
-					const UScriptStruct* StructType = CastChecked<const UScriptStruct>(Entry.TypeObject);
-					ReferenceCollector.AddPropertyReferencesWithStructARO(StructType, RawData);
+					const UScriptStruct* StructType = Cast<const UScriptStruct>(Entry.TypeObject);
+					if (ensure(StructType))
+					{
+						ReferenceCollector.AddPropertyReferencesWithStructARO(StructType, RawData);
+					}
 				}
 				break;
 			case ECameraContextDataType::Object:
@@ -147,8 +150,8 @@ bool FCameraContextDataTable::GetDataTypeAllocationInfo(ECameraContextDataType D
 			break;
 		case ECameraContextDataType::Struct:
 			{
-				const UScriptStruct* StructType = CastChecked<const UScriptStruct>(DataTypeObject);
-				if (StructType)
+				const UScriptStruct* StructType = Cast<const UScriptStruct>(DataTypeObject);
+				if (ensure(StructType))
 				{
 					OutSizeOf = StructType->GetPropertiesSize();
 					OutAlignOf = StructType->GetMinAlignment();
@@ -181,10 +184,22 @@ bool FCameraContextDataTable::InitializeDefaultDataValue(ECameraContextDataType
 			new (DataPtr) FString();
 			break;
 		case ECameraContextDataType::Enum:
-			*reinterpret_cast<uint8*>(DataPtr) = (uint8)CastChecked<const UEnum>(DataTypeObject)->GetValueByIndex(0);
+			{
+				const UEnum* EnumType = Cast<const UEnum>(DataTypeObject);
+				if (ensure(EnumType))
+				{
+					*reinterpret_cast<uint8*>(DataPtr) = (uint8)EnumType->GetValueByIndex(0);
+				}
+			}
 			break;
 		case ECameraContextDataType::Struct:
-			CastChecked<const UScriptStruct>(DataTypeObject)->InitializeDefaultValue(DataPtr);
+			{
+				const UScriptStruct* StructType = Cast<const UScriptStruct>(DataTypeObject);
+				if (ensure(StructType))
+				{
+					StructType->InitializeDefaultValue(DataPtr);
+				}
+			}
 			break;
 		case ECameraContextDataType::Object:
 			new (DataPtr) TObjectPtr<UObject>();
@@ -248,7 +263,13 @@ void FCameraContextDataTable::DestroyBuffer()
 				// Nothing to do.
 				break;
 			case ECameraContextDataType::Struct:
-				CastChecked<const UScriptStruct>(Entry.TypeObject)->DestroyStruct(DataPtr);
+				{
+					const UScriptStruct* StructType = Cast<const UScriptStruct>(Entry.TypeObject);
+					if (ensure(StructType))
+					{
+						StructType->DestroyStruct(DataPtr);
+					}
+				}
 				break;
 			case ECameraContextDataType::Object:
 				reinterpret_cast<TObjectPtr<UObject>*>(DataPtr)->~TObjectPtr<UObject>();
@@ -514,7 +535,13 @@ bool FCameraContextDataTable::SetDataValue(ECameraContextDataType DataType, cons
 			*reinterpret_cast<uint8*>(DestDataPtr) = *reinterpret_cast<const uint8*>(SrcDataPtr);
 			break;
 		case ECameraContextDataType::Struct:
-			CastChecked<const UScriptStruct>(DataTypeObject)->CopyScriptStruct(DestDataPtr, SrcDataPtr);
+			{
+				const UScriptStruct* StructType = Cast<const UScriptStruct>(DataTypeObject);
+				if (ensure(StructType))
+				{
+					StructType->CopyScriptStruct(DestDataPtr, SrcDataPtr);
+				}
+			}
 			break;
 		case ECameraContextDataType::Object:
 			*reinterpret_cast<TObjectPtr<UObject>*>(DestDataPtr) = *reinterpret_cast<const TObjectPtr<UObject>*>(SrcDataPtr);
diff --git a/Engine/Plugins/Cameras/GameplayCameras/Source/GameplayCameras/Private/Helpers/CameraParameterOverrideHelper.cpp b/Engine/Plugins/Cameras/GameplayCameras/Source/GameplayCameras/Private/Helpers/CameraParameterOverrideHelper.cpp
--- a/Engine/Plugins/Cameras/GameplayCameras/Source/GameplayCameras/Private/Helpers/CameraParameterOverrideHelper.cpp
+++ b/Engine/Plugins/Cameras/GameplayCameras/Source/GameplayCameras/Private/Helpers/CameraParameterOverrideHelper.cpp
@@ -179,16 +179,22 @@ void ApplyDataParameterOverride(
 			break;
 		case ECameraContextDataType::Enum:
 			{
-				const UEnum* EnumType = CastChecked<const UEnum>(ParameterDefinition.DataTypeObject);
-				TValueOrError<uint8, EPropertyBagResult> ParameterValueOrError = PropertyBag.GetValueEnum(PropertyBagPropertyDesc, EnumType);
-				ApplyDataParameterOverride(CameraRig, ParameterDefinition, ParameterValueOrError, ContextDataTable);
+				const UEnum* EnumType = Cast<const UEnum>(ParameterDefinition.DataTypeObject);
+				if (ensure(EnumType))
+				{
+					TValueOrError<uint8, EPropertyBagResult> ParameterValueOrError = PropertyBag.GetValueEnum(PropertyBagPropertyDesc, EnumType);
+					ApplyDataParameterOverride(CameraRig, ParameterDefinition, ParameterValueOrError, ContextDataTable);
+				}
 			}
 			break;
 		case ECameraContextDataType::Struct:
 			{
-				const UScriptStruct* StructType = CastChecked<const UScriptStruct>(ParameterDefinition.DataTypeObject);
-				TValueOrError<FStructView, EPropertyBagResult> ParameterValueOrError = PropertyBag.GetValueStruct(PropertyBagPropertyDesc, StructType);
-				ApplyDataParameterOverride(CameraRig, ParameterDefinition, ParameterValueOrError, ContextDataTable);
+				const UScriptStruct* StructType = Cast<const UScriptStruct>(ParameterDefinition.DataTypeObject);
+				if (ensure(StructType))
+				{
+					TValueOrError<FStructView, EPropertyBagResult> ParameterValueOrError = PropertyBag.GetValueStruct(PropertyBagPropertyDesc, StructType);
+					ApplyDataParameterOverride(CameraRig, ParameterDefinition, ParameterValueOrError, ContextDataTable);
+				}
 			}
 			break;
 		case ECameraContextDataType::Object:
