Cameras: show debug HUD in editor when the Camera Debugger is enabled

diff --git a/Engine/Plugins/Cameras/GameplayCameras/Source/GameplayCameras/Private/Core/CameraSystemEvaluator.cpp b/Engine/Plugins/Cameras/GameplayCameras/Source/GameplayCameras/Private/Core/CameraSystemEvaluator.cpp
--- a/Engine/Plugins/Cameras/GameplayCameras/Source/GameplayCameras/Private/Core/CameraSystemEvaluator.cpp
+++ b/Engine/Plugins/Cameras/GameplayCameras/Source/GameplayCameras/Private/Core/CameraSystemEvaluator.cpp
@@ -380,6 +380,12 @@ void FCameraSystemEvaluator::DrawEditorPreview(const FCameraSystemEditorPreviewP
 	FCameraDebugRenderer Renderer(OwnerWorld, Params.SceneView, Params.Canvas, !Params.bIsLockedToCamera);
 
 	Renderer.BeginDrawing();
+
+	if (GGameplayCamerasDebugEnable && RootDebugBlock)
+	{
+		RootDebugBlock->RootDebugDraw(Renderer);
+	}
+	else
 	{
 		RootEvaluator->DrawEditorPreview(NodeParams, Renderer);
 
diff --git a/Engine/Plugins/Cameras/GameplayCameras/Source/GameplayCameras/Private/GameFramework/GameplayCameraComponent.cpp b/Engine/Plugins/Cameras/GameplayCameras/Source/GameplayCameras/Private/GameFramework/GameplayCameraComponent.cpp
--- a/Engine/Plugins/Cameras/GameplayCameras/Source/GameplayCameras/Private/GameFramework/GameplayCameraComponent.cpp
+++ b/Engine/Plugins/Cameras/GameplayCameras/Source/GameplayCameras/Private/GameFramework/GameplayCameraComponent.cpp
@@ -25,15 +25,6 @@
 
 #define LOCTEXT_NAMESPACE "GameplayCameraComponent"
 
-namespace UE::Cameras
-{
-
-#if UE_GAMEPLAY_CAMERAS_DEBUG
-extern bool GGameplayCamerasDebugEnable;
-#endif  // UE_GAMEPLAY_CAMERAS_DEBUG
-
-}  // namespace UE::Cameras
-
 UGameplayCameraComponent::UGameplayCameraComponent(const FObjectInitializer& ObjectInit)
 	: Super(ObjectInit)
 {
