Cameras: don't run actual camera directors in editor

Most camera directors need to query game state, which fails in editor, so run the editor preview with a custom director that lets the user pick what rig to run.

Inside the rigs, don't runs things like the Blueprint node since it's pretty unstable to instantiate and run BP classes in editor.

Future work should include a custom UI widget for choosing what rig to preview.

diff --git a/Engine/Plugins/Cameras/GameplayCameras/Source/GameplayCameras/Private/Core/CameraEvaluationContext.cpp b/Engine/Plugins/Cameras/GameplayCameras/Source/GameplayCameras/Private/Core/CameraEvaluationContext.cpp
--- a/Engine/Plugins/Cameras/GameplayCameras/Source/GameplayCameras/Private/Core/CameraEvaluationContext.cpp
+++ b/Engine/Plugins/Cameras/GameplayCameras/Source/GameplayCameras/Private/Core/CameraEvaluationContext.cpp
@@ -3,12 +3,63 @@
 #include "Core/CameraEvaluationContext.h"
 
 #include "Core/CameraAsset.h"
+#include "Core/CameraRigAsset.h"
 #include "Core/CameraDirector.h"
+#include "Core/CameraDirectorEvaluator.h"
+#include "Core/CameraSystemEvaluator.h"
 #include "GameFramework/PlayerController.h"
 
 namespace UE::Cameras
 {
 
+#if WITH_EDITOR
+
+class FEditorPreviewCameraDirectorEvaluator : public FCameraDirectorEvaluator
+{
+	UE_DECLARE_CAMERA_DIRECTOR_EVALUATOR(GAMEPLAYCAMERAS_API, FEditorPreviewCameraDirectorEvaluator)
+
+public:
+
+	FEditorPreviewCameraDirectorEvaluator()
+	{}
+
+	FEditorPreviewCameraDirectorEvaluator(TConstArrayView<const UCameraRigAsset*> InCameraRigs)
+		: CameraRigs(InCameraRigs)
+	{}
+
+	int32 GetCameraRigIndex() const { return PreviewIndex; }
+	void SetCameraRigIndex(int32 Index) { PreviewIndex = Index; }
+
+protected:
+
+	virtual void OnRun(const FCameraDirectorEvaluationParams& Params, FCameraDirectorEvaluationResult& OutResult) override
+	{
+		if (!CameraRigs.IsValidIndex(PreviewIndex))
+		{
+			PreviewIndex = (CameraRigs.IsEmpty() ? INDEX_NONE : 0);
+		}
+
+		if (PreviewIndex != INDEX_NONE)
+		{
+			OutResult.Add(Params.OwnerContext, CameraRigs[PreviewIndex]);
+		}
+	}
+
+	virtual void OnAddReferencedObjects(FReferenceCollector& Collector) override
+	{
+		Collector.AddReferencedObjects(CameraRigs);
+	}
+
+private:
+
+	TArray<TObjectPtr<const UCameraRigAsset>> CameraRigs;
+	int32 PreviewIndex = 0;
+};
+
+UE_DEFINE_CAMERA_DIRECTOR_EVALUATOR(FEditorPreviewCameraDirectorEvaluator)
+
+#endif  // WITH_EDITOR
+
 UE_GAMEPLAY_CAMERAS_DEFINE_RTTI(FCameraEvaluationContext)
 
 FCameraEvaluationContext::FCameraEvaluationContext()
@@ -124,6 +175,45 @@ void FCameraEvaluationContext::AutoCreateDirectorEvaluator()
 	}
 }
 
+#if WITH_EDITOR
+
+void FCameraEvaluationContext::AutoCreateEditorPreviewDirectorEvaluator(const FCameraEvaluationContextActivateParams& Params)
+{
+	if (ensure(Params.Evaluator) && Params.Evaluator->GetRole() == ECameraSystemEvaluatorRole::EditorPreview)
+	{
+		FCameraDirectorRigUsageInfo UsageInfo;
+
+		if (CameraAsset && CameraAsset->GetCameraDirector())
+		{
+			UCameraDirector* CameraDirector = CameraAsset->GetCameraDirector();
+			CameraDirector->GatherRigUsageInfo(UsageInfo);
+		}
+
+		FCameraDirectorEvaluatorBuilder DirectorBuilder(DirectorEvaluatorStorage);
+		DirectorEvaluator = DirectorBuilder.BuildEvaluator<FEditorPreviewCameraDirectorEvaluator>(UsageInfo.CameraRigs);
+
+		FCameraDirectorInitializeParams InitParams;
+		InitParams.OwnerContext = SharedThis(this);
+		DirectorEvaluator->Initialize(InitParams);
+	}
+}
+
+void FCameraEvaluationContext::SetEditorPreviewCameraRigIndex(int32 Index)
+{
+	if (!DirectorEvaluator)
+	{
+		return;
+	}
+
+	FEditorPreviewCameraDirectorEvaluator* EditorPreviewEvaluator = DirectorEvaluator->CastThis<FEditorPreviewCameraDirectorEvaluator>();
+	if (EditorPreviewEvaluator)
+	{
+		EditorPreviewEvaluator->SetCameraRigIndex(Index);
+	}
+}
+
+#endif  // WITH_EDITOR
+
 FCameraNodeEvaluationResult& FCameraEvaluationContext::GetOrAddConditionalResult(ECameraEvaluationDataCondition Condition)
 {
 	if (FCameraNodeEvaluationResult* ExistingResult = ConditionalResults.Find(Condition))
@@ -155,10 +245,12 @@ void FCameraEvaluationContext::Activate(const FCameraEvaluationContextActivatePa
 		return;
 	}
 
-	CameraSystemEvaluator = Params.Evaluator;
-
 	OnActivate(Params);
 
+#if WITH_EDITOR
+	AutoCreateEditorPreviewDirectorEvaluator(Params);
+#endif  // WITH_EDITOR
+
 	AutoCreateDirectorEvaluator();
 
 	if (DirectorEvaluator)
@@ -191,8 +283,6 @@ void FCameraEvaluationContext::Deactivate(const FCameraEvaluationContextDeactiva
 
 	OnDeactivate(Params);
 
-	CameraSystemEvaluator = nullptr;
-
 	bActivated = false;
 }
 
diff --git a/Engine/Plugins/Cameras/GameplayCameras/Source/GameplayCameras/Private/Core/CameraSystemEvaluator.cpp b/Engine/Plugins/Cameras/GameplayCameras/Source/GameplayCameras/Private/Core/CameraSystemEvaluator.cpp
--- a/Engine/Plugins/Cameras/GameplayCameras/Source/GameplayCameras/Private/Core/CameraSystemEvaluator.cpp
+++ b/Engine/Plugins/Cameras/GameplayCameras/Source/GameplayCameras/Private/Core/CameraSystemEvaluator.cpp
@@ -67,6 +67,8 @@ void FCameraSystemEvaluator::Initialize(const FCameraSystemEvaluatorCreateParams
 	}
 	WeakOwner = Owner;
 
+	Role = Params.Role;
+
 	if (Params.RootNodeFactory)
 	{
 		RootNode = Params.RootNodeFactory();
diff --git a/Engine/Plugins/Cameras/GameplayCameras/Source/GameplayCameras/Private/Debug/CameraDirectorTreeDebugBlock.cpp b/Engine/Plugins/Cameras/GameplayCameras/Source/GameplayCameras/Private/Debug/CameraDirectorTreeDebugBlock.cpp
--- a/Engine/Plugins/Cameras/GameplayCameras/Source/GameplayCameras/Private/Debug/CameraDirectorTreeDebugBlock.cpp
+++ b/Engine/Plugins/Cameras/GameplayCameras/Source/GameplayCameras/Private/Debug/CameraDirectorTreeDebugBlock.cpp
@@ -63,12 +63,13 @@ void FCameraDirectorTreeDebugBlock::InitializeEntry(TSharedPtr<FCameraEvaluation
 
 		const UObject* ContextOwner = Context->GetOwner();
 		const FCameraDirectorEvaluator* DirectorEvaluator = Context->GetDirectorEvaluator();
+		const UCameraDirector* CameraDirector = DirectorEvaluator->GetCameraDirector();
 
 		EntryDebugInfo.ContextClassName = ContextTypeName;
 		EntryDebugInfo.OwnerName = *GetPathNameSafe(ContextOwner);
 		EntryDebugInfo.OwnerClassName = ContextOwner ? ContextOwner->GetClass()->GetFName() : NAME_None;
 		EntryDebugInfo.CameraAssetName = GetNameSafe(Context->GetCameraAsset());
-		EntryDebugInfo.CameraDirectorClassName = DirectorEvaluator->GetCameraDirector()->GetFName();
+		EntryDebugInfo.CameraDirectorClassName = GetFNameSafe(CameraDirector);
 		EntryDebugInfo.InitialContextTransform = Context->GetInitialResult().CameraPose.GetTransform();
 		EntryDebugInfo.bIsValid = true;
 
diff --git a/Engine/Plugins/Cameras/GameplayCameras/Source/GameplayCameras/Private/GameFramework/GameplayCameraComponent.cpp b/Engine/Plugins/Cameras/GameplayCameras/Source/GameplayCameras/Private/GameFramework/GameplayCameraComponent.cpp
--- a/Engine/Plugins/Cameras/GameplayCameras/Source/GameplayCameras/Private/GameFramework/GameplayCameraComponent.cpp
+++ b/Engine/Plugins/Cameras/GameplayCameras/Source/GameplayCameras/Private/GameFramework/GameplayCameraComponent.cpp
@@ -48,6 +48,25 @@ void UGameplayCameraComponent::PostLoad()
 	}
 }
 
+void UGameplayCameraComponent::BeginDestroy()
+{
+	Super::BeginDestroy();
+
+#if WITH_EDITOR
+
+	if (EditorPreviewEvaluator)
+	{
+		EditorPreviewEvaluator.Reset();
+	}
+
+	if (EvaluationContext)
+	{
+		EvaluationContext.Reset();
+	}
+
+#endif  // WITH_EDITOR
+}
+
 void UGameplayCameraComponent::AddReferencedObjects(UObject* InThis, FReferenceCollector& Collector)
 {
 	Super::AddReferencedObjects(InThis, Collector);
@@ -452,7 +471,7 @@ void UGameplayCameraComponent::AutoManageEditorPreviewEvaluator()
 	{
 		return;
 	}
-
+	
 	if (bRunInEditor && !(EditorPreviewEvaluator && EvaluationContext))
 	{
 		// We want to run the camera logic in the editor but we haven't set things up for that.
@@ -460,12 +479,17 @@ void UGameplayCameraComponent::AutoManageEditorPreviewEvaluator()
 		if (!EditorPreviewEvaluator)
 		{
 			EditorPreviewEvaluator = MakeShared<FCameraSystemEvaluator>();
-			EditorPreviewEvaluator->Initialize(this);
+
+			FCameraSystemEvaluatorCreateParams CreateParams;
+			CreateParams.Owner = this;
+			CreateParams.Role = ECameraSystemEvaluatorRole::EditorPreview;
+			EditorPreviewEvaluator->Initialize(CreateParams);
 		}
 		if (!EvaluationContext)
 		{
 			EnsureCameraEvaluationContextCreated(nullptr);
 			EditorPreviewEvaluator->PushEvaluationContext(EvaluationContext.ToSharedRef());
+			EvaluationContext->SetEditorPreviewCameraRigIndex(EditorPreviewCameraRigIndex);
 		}
 	}
 	else if (!bRunInEditor && (EditorPreviewEvaluator || EvaluationContext))
@@ -476,6 +500,45 @@ void UGameplayCameraComponent::AutoManageEditorPreviewEvaluator()
 	}
 }
 
+void UGameplayCameraComponent::OnCameraAssetReferenceChanged()
+{
+	if (!bIsEditorWorld)
+	{
+		return;
+	}
+
+	if (bRunInEditor && EditorPreviewEvaluator && EvaluationContext)
+	{
+		if (EvaluationContext->GetCameraAsset() != CameraReference.GetCameraAsset())
+		{
+			// The camera asset has changed! Recreate the context.
+			EditorPreviewEvaluator->RemoveEvaluationContext(EvaluationContext.ToSharedRef());
+			EvaluationContext = nullptr;
+
+			EnsureCameraEvaluationContextCreated(nullptr);
+			EditorPreviewEvaluator->PushEvaluationContext(EvaluationContext.ToSharedRef());
+		}
+		else
+		{
+			// Otherwise, maybe one of the parameter overrides has changed. Re-apply them.
+			UpdateCameraEvaluationContext(true);
+		}
+	}
+}
+
+void UGameplayCameraComponent::OnEditorPreviewCameraRigIndexChanged()
+{
+	if (!bIsEditorWorld)
+	{
+		return;
+	}
+
+	if (bRunInEditor && EditorPreviewEvaluator && EvaluationContext)
+	{
+		EvaluationContext->SetEditorPreviewCameraRigIndex(EditorPreviewCameraRigIndex);
+	}
+}
+
 void UGameplayCameraComponent::OnCameraAssetBuilt(const UCameraAsset* InCameraAsset)
 {
 	using namespace UE::Cameras;
@@ -553,28 +616,16 @@ void UGameplayCameraComponent::PostEditChangeProperty( struct FPropertyChangedEv
 	const FName MemberPropertyName = PropertyChangedEvent.GetMemberPropertyName();
 	if (MemberPropertyName == GET_MEMBER_NAME_CHECKED(UGameplayCameraComponent, CameraReference))
 	{
-		if (EditorPreviewEvaluator && EvaluationContext)
-		{
-			if (EvaluationContext->GetCameraAsset() != CameraReference.GetCameraAsset())
-			{
-				// The camera asset has changed! Recreate the context.
-				EditorPreviewEvaluator->RemoveEvaluationContext(EvaluationContext.ToSharedRef());
-				EvaluationContext = nullptr;
-
-				EnsureCameraEvaluationContextCreated(nullptr);
-				EditorPreviewEvaluator->PushEvaluationContext(EvaluationContext.ToSharedRef());
-			}
-			else
-			{
-				// Otherwise, maybe one of the parameter overrides has changed. Re-apply them.
-				UpdateCameraEvaluationContext(true);
-			}
-		}
+		OnCameraAssetReferenceChanged();
 	}
 	else if (MemberPropertyName == GET_MEMBER_NAME_CHECKED(UGameplayCameraComponent, bRunInEditor))
 	{
 		AutoManageEditorPreviewEvaluator();
 	}
+	else if (MemberPropertyName == GET_MEMBER_NAME_CHECKED(UGameplayCameraComponent, EditorPreviewCameraRigIndex))
+	{
+		OnEditorPreviewCameraRigIndexChanged();
+	}
 }
 
 #endif  // WITH_EDITOR
diff --git a/Engine/Plugins/Cameras/GameplayCameras/Source/GameplayCameras/Private/Nodes/Utility/BlueprintCameraNode.cpp b/Engine/Plugins/Cameras/GameplayCameras/Source/GameplayCameras/Private/Nodes/Utility/BlueprintCameraNode.cpp
--- a/Engine/Plugins/Cameras/GameplayCameras/Source/GameplayCameras/Private/Nodes/Utility/BlueprintCameraNode.cpp
+++ b/Engine/Plugins/Cameras/GameplayCameras/Source/GameplayCameras/Private/Nodes/Utility/BlueprintCameraNode.cpp
@@ -8,6 +8,7 @@
 #include "Core/CameraEvaluationContext.h"
 #include "Core/CameraNodeEvaluator.h"
 #include "Core/CameraRigBuildContext.h"
+#include "Core/CameraSystemEvaluator.h"
 #include "Core/CameraVariableAssets.h"
 #include "Core/CameraVariableTable.h"
 #include "Engine/World.h"
@@ -58,6 +59,13 @@ void FBlueprintCameraNodeEvaluator::OnInitialize(const FCameraNodeEvaluatorIniti
 		return;
 	}
 
+#if WITH_EDITOR
+	if (Params.Evaluator->GetRole() == ECameraSystemEvaluatorRole::EditorPreview)
+	{
+		return;
+	}
+#endif  // WITH_EDITOR
+
 	if (const UBlueprintCameraNodeEvaluator* EvaluatorTemplate = BlueprintNode->GetCameraNodeEvaluatorTemplate())
 	{
 		UObject* Outer = Params.EvaluationContext->GetOwner();
diff --git a/Engine/Plugins/Cameras/GameplayCameras/Source/GameplayCameras/Public/Core/CameraEvaluationContext.h b/Engine/Plugins/Cameras/GameplayCameras/Source/GameplayCameras/Public/Core/CameraEvaluationContext.h
--- a/Engine/Plugins/Cameras/GameplayCameras/Source/GameplayCameras/Public/Core/CameraEvaluationContext.h
+++ b/Engine/Plugins/Cameras/GameplayCameras/Source/GameplayCameras/Public/Core/CameraEvaluationContext.h
@@ -94,9 +94,6 @@ public:
 	/** Gets the camera asset that is hosted in this context. */
 	const UCameraAsset* GetCameraAsset() const { return CameraAsset; }
 
-	/** Gets the camera system evaluator. Only valid if this evaluation context is active. */
-	FCameraSystemEvaluator* GetCameraSystemEvaluator() const { return CameraSystemEvaluator; }
-
 	/** Gets the initial evaluation result for all camera rigs in this context. */
 	const FCameraNodeEvaluationResult& GetInitialResult() const { return InitialResult; }
 
@@ -143,6 +140,10 @@ public:
 	bool RegisterChildContext(TSharedRef<FCameraEvaluationContext> ChildContext);
 	bool UnregisterChildContext(TSharedRef<FCameraEvaluationContext> ChildContext);
 
+#if WITH_EDITOR
+	void SetEditorPreviewCameraRigIndex(int32 Index);
+#endif  // WITH_EDITOR
+
 protected:
 
 	virtual void OnActivate(const FCameraEvaluationContextActivateParams& Params) {}
@@ -150,6 +151,10 @@ protected:
 
 	void AutoCreateDirectorEvaluator();
 
+#if WITH_EDITOR
+	void AutoCreateEditorPreviewDirectorEvaluator(const FCameraEvaluationContextActivateParams& Params);
+#endif  // WITH_EDITOR
+
 protected:
 
 	/** The owner of this context, if any. */
diff --git a/Engine/Plugins/Cameras/GameplayCameras/Source/GameplayCameras/Public/Core/CameraSystemEvaluator.h b/Engine/Plugins/Cameras/GameplayCameras/Source/GameplayCameras/Public/Core/CameraSystemEvaluator.h
--- a/Engine/Plugins/Cameras/GameplayCameras/Source/GameplayCameras/Public/Core/CameraSystemEvaluator.h
+++ b/Engine/Plugins/Cameras/GameplayCameras/Source/GameplayCameras/Public/Core/CameraSystemEvaluator.h
@@ -37,6 +37,18 @@ struct FRootCameraNodeCameraRigEvent;
 class FRootCameraDebugBlock;
 #endif  // UE_GAMEPLAY_CAMERAS_DEBUG
 
+/**
+ * An enumeration that describes what a camera system evaluator is used for.
+ */
+enum class ECameraSystemEvaluatorRole
+{
+	Game,
+
+#if WITH_EDITOR
+	EditorPreview
+#endif
+};
+
 /**
  * Parameter structure for initializing a new camera system evaluator.
  */
@@ -45,6 +57,9 @@ struct GAMEPLAYCAMERAS_API FCameraSystemEvaluatorCreateParams
 	/** The owner of the camera system, if any. */
 	TObjectPtr<UObject> Owner;
 
+	/** The role of the camera system. */
+	ECameraSystemEvaluatorRole Role = ECameraSystemEvaluatorRole::Game;
+
 	/** An optional factory for creating the root node. */
 	using FRootNodeFactory = TFunction<URootCameraNode*()>;
 	FRootNodeFactory RootNodeFactory;
@@ -152,6 +167,9 @@ public:
 	/** Gets the owner of this camera system, if any, and if still valid. */
 	UObject* GetOwner() const { return WeakOwner.Get(); }
 
+	/** Gets the role of this camera system. */
+	ECameraSystemEvaluatorRole GetRole() const { return Role; }
+
 public:
 
 	/** Push a new evaluation context on the stack. */
@@ -263,6 +281,9 @@ private:
 	/** The current overall result of the camera system. */
 	FCameraSystemEvaluationResult Result;
 
+	/** The role of this camera system. */
+	ECameraSystemEvaluatorRole Role = ECameraSystemEvaluatorRole::Game;
+
 #if UE_GAMEPLAY_CAMERAS_DEBUG
 	/** Storage for debug drawing blocks. */
 	FCameraDebugBlockStorage DebugBlockStorage;
diff --git a/Engine/Plugins/Cameras/GameplayCameras/Source/GameplayCameras/Public/GameFramework/GameplayCameraComponent.h b/Engine/Plugins/Cameras/GameplayCameras/Source/GameplayCameras/Public/GameFramework/GameplayCameraComponent.h
--- a/Engine/Plugins/Cameras/GameplayCameras/Source/GameplayCameras/Public/GameFramework/GameplayCameraComponent.h
+++ b/Engine/Plugins/Cameras/GameplayCameras/Source/GameplayCameras/Public/GameFramework/GameplayCameraComponent.h
@@ -107,6 +107,7 @@ public:
 
 	// UObject interface.
 	virtual void PostLoad() override;
+	virtual void BeginDestroy() override;
 #if WITH_EDITOR
 	virtual void PostEditChangeProperty( struct FPropertyChangedEvent& PropertyChangedEvent) override;
 #endif
@@ -130,6 +131,8 @@ private:
 
 #if WITH_EDITOR
 	void AutoManageEditorPreviewEvaluator();
+	void OnCameraAssetReferenceChanged();
+	void OnEditorPreviewCameraRigIndexChanged();
 	void OnCameraAssetBuilt(const UCameraAsset* InCameraAsset);
 
 	void UpdateEditorPreviewEvaluator(float DeltaTime);
@@ -153,6 +156,9 @@ public:
 	UPROPERTY(EditAnywhere, Category=Camera)
 	bool bRunInEditor = true;
 
+	UPROPERTY(EditAnywhere, Category=Camera, meta=(EditCondition="bRunInEditor"))
+	int32 EditorPreviewCameraRigIndex = 0;
+
 #endif  // WITH_EDITORONLY_DATA
 
 protected:
