Cameras: correctly upgrade camera rig parameter overrides to the new camera rig reference format

diff --git a/Engine/Plugins/Cameras/GameplayCameras/Source/GameplayCameras/Private/Core/CameraRigAssetReference.cpp b/Engine/Plugins/Cameras/GameplayCameras/Source/GameplayCameras/Private/Core/CameraRigAssetReference.cpp
--- a/Engine/Plugins/Cameras/GameplayCameras/Source/GameplayCameras/Private/Core/CameraRigAssetReference.cpp
+++ b/Engine/Plugins/Cameras/GameplayCameras/Source/GameplayCameras/Private/Core/CameraRigAssetReference.cpp
@@ -38,13 +38,21 @@ const FCameraRigAssetReferenceParameterMetaData* FCameraRigAssetReference::FindM
 			});
 }
 
-FCameraRigAssetReferenceParameterMetaData* FCameraRigAssetReference::FindMetaData(const FGuid& PropertyID)
+FCameraRigAssetReferenceParameterMetaData& FCameraRigAssetReference::FindOrAddMetaData(const FGuid& PropertyID)
 {
-	return ParameterMetaData.FindByPredicate(
+	FCameraRigAssetReferenceParameterMetaData* ExistingMetaData = ParameterMetaData.FindByPredicate(
 			[PropertyID](FCameraRigAssetReferenceParameterMetaData& Item)
 			{
 				return Item.ParameterGuid == PropertyID;
 			});
+	if (ExistingMetaData)
+	{
+		return *ExistingMetaData;
+	}
+
+	FCameraRigAssetReferenceParameterMetaData& NewMetaData = ParameterMetaData.Emplace_GetRef();
+	NewMetaData.ParameterGuid = PropertyID;
+	return NewMetaData;
 }
 
 void FCameraRigAssetReference::GenerateOverriddenParameterGuidArray(TArray<FGuid>& OutOverriddenIDs) const
@@ -69,11 +77,8 @@ bool FCameraRigAssetReference::IsParameterOverridden(const FGuid& PropertyID) co
 
 void FCameraRigAssetReference::SetParameterOverridden(const FGuid& PropertyID, bool bIsOverridden)
 {
-	FCameraRigAssetReferenceParameterMetaData* MetaData = FindMetaData(PropertyID);
-	if (ensure(MetaData))
-	{
-		MetaData->bIsOverridden = bIsOverridden;
-	}
+	FCameraRigAssetReferenceParameterMetaData& MetaData = FindOrAddMetaData(PropertyID);
+	MetaData.bIsOverridden = bIsOverridden;
 }
 
 bool FCameraRigAssetReference::NeedsRebuildParameters() const
@@ -146,7 +151,6 @@ void FCameraRigAssetReference::RebuildParameters()
 			{
 				FCameraRigAssetReferenceParameterMetaData NewMetaData;
 				NewMetaData.ParameterGuid = ParameterGuid;
-				NewMetaData.ParameterType = WantedParameterTypes.FindChecked(ParameterGuid);
 				ParameterMetaData.Add(NewMetaData);
 			}
 		}
@@ -226,16 +230,13 @@ void FCameraRigAssetReference::GetCustomCameraNodeParameters(FCustomCameraNodePa
 							FStructView& StructValue = PropertyValue.GetValue();
 							check(StructValue.IsValid());
 
-							FCameraRigAssetReferenceParameterMetaData* MetaData = FindMetaData(Definition.ParameterGuid);
-							check(MetaData);
-							ensure(MetaData->ParameterType == ECameraRigInterfaceParameterType::Blendable);
-
+							FCameraRigAssetReferenceParameterMetaData& MetaData = FindOrAddMetaData(Definition.ParameterGuid);
 							OutParameterInfos.AddBlendableParameter(
 									Definition.ParameterName,
 									Definition.VariableType,
 									Definition.BlendableStructType,
 									StructValue.GetMemory(),
-									&MetaData->OverrideVariableID);
+									&MetaData.OverrideVariableID);
 						}
 					}
 					break;
@@ -243,9 +244,7 @@ void FCameraRigAssetReference::GetCustomCameraNodeParameters(FCustomCameraNodePa
 		}
 		else if (Definition.ParameterType == ECameraRigInterfaceParameterType::Data)
 		{
-			FCameraRigAssetReferenceParameterMetaData* MetaData = FindMetaData(Definition.ParameterGuid);
-			check(MetaData);
-			ensure(MetaData->ParameterType == ECameraRigInterfaceParameterType::Data);
+			FCameraRigAssetReferenceParameterMetaData& MetaData = FindOrAddMetaData(Definition.ParameterGuid);
 
 			switch (Definition.DataType)
 			{
@@ -253,14 +252,14 @@ void FCameraRigAssetReference::GetCustomCameraNodeParameters(FCustomCameraNodePa
 					if (ensure(PropertyDesc->ValueType == EPropertyBagPropertyType::Name))
 					{
 						const uint8* DefaultValue = reinterpret_cast<uint8*>(DefaultParameters.GetValueName(*PropertyDesc).TryGetValue());
-						OutParameterInfos.AddDataParameter(Definition.ParameterName, ECameraContextDataType::Name, nullptr, DefaultValue, &MetaData->OverrideDataID);
+						OutParameterInfos.AddDataParameter(Definition.ParameterName, ECameraContextDataType::Name, nullptr, DefaultValue, &MetaData.OverrideDataID);
 					}
 					break;
 				case ECameraContextDataType::String:
 					if (ensure(PropertyDesc->ValueType == EPropertyBagPropertyType::String))
 					{
 						const uint8* DefaultValue = reinterpret_cast<uint8*>(DefaultParameters.GetValueString(*PropertyDesc).TryGetValue());
-						OutParameterInfos.AddDataParameter(Definition.ParameterName, ECameraContextDataType::String, nullptr, DefaultValue, &MetaData->OverrideDataID);
+						OutParameterInfos.AddDataParameter(Definition.ParameterName, ECameraContextDataType::String, nullptr, DefaultValue, &MetaData.OverrideDataID);
 					}
 					break;
 				case ECameraContextDataType::Enum:
@@ -269,7 +268,7 @@ void FCameraRigAssetReference::GetCustomCameraNodeParameters(FCustomCameraNodePa
 					{
 						const UEnum* EnumType = CastChecked<const UEnum>(Definition.DataTypeObject);
 						const uint8* DefaultValue = reinterpret_cast<uint8*>(DefaultParameters.GetValueEnum(*PropertyDesc, EnumType).TryGetValue());
-						OutParameterInfos.AddDataParameter(Definition.ParameterName, ECameraContextDataType::Enum, EnumType, DefaultValue, &MetaData->OverrideDataID);
+						OutParameterInfos.AddDataParameter(Definition.ParameterName, ECameraContextDataType::Enum, EnumType, DefaultValue, &MetaData.OverrideDataID);
 					}
 					break;
 				case ECameraContextDataType::Struct:
@@ -278,21 +277,21 @@ void FCameraRigAssetReference::GetCustomCameraNodeParameters(FCustomCameraNodePa
 					{
 						const UScriptStruct* DataType = CastChecked<const UScriptStruct>(Definition.DataTypeObject);
 						const uint8* DefaultValue = reinterpret_cast<uint8*>(DefaultParameters.GetValueStruct(*PropertyDesc, DataType).TryGetValue());
-						OutParameterInfos.AddDataParameter(Definition.ParameterName, ECameraContextDataType::Struct, DataType, DefaultValue, &MetaData->OverrideDataID);
+						OutParameterInfos.AddDataParameter(Definition.ParameterName, ECameraContextDataType::Struct, DataType, DefaultValue, &MetaData.OverrideDataID);
 					}
 					break;
 				case ECameraContextDataType::Object:
 					if (ensure(PropertyDesc->ValueType == EPropertyBagPropertyType::Object))
 					{
 						const uint8* DefaultValue = reinterpret_cast<uint8*>(DefaultParameters.GetValueObject(*PropertyDesc).TryGetValue());
-						OutParameterInfos.AddDataParameter(Definition.ParameterName, ECameraContextDataType::Object, Definition.DataTypeObject, DefaultValue, &MetaData->OverrideDataID);
+						OutParameterInfos.AddDataParameter(Definition.ParameterName, ECameraContextDataType::Object, Definition.DataTypeObject, DefaultValue, &MetaData.OverrideDataID);
 					}
 					break;
 				case ECameraContextDataType::Class:
 					if (ensure(PropertyDesc->ValueType == EPropertyBagPropertyType::Class))
 					{
 						const uint8* DefaultValue = reinterpret_cast<uint8*>(DefaultParameters.GetValueClass(*PropertyDesc).TryGetValue());
-						OutParameterInfos.AddDataParameter(Definition.ParameterName, ECameraContextDataType::Class, Definition.DataTypeObject, DefaultValue, &MetaData->OverrideDataID);
+						OutParameterInfos.AddDataParameter(Definition.ParameterName, ECameraContextDataType::Class, Definition.DataTypeObject, DefaultValue, &MetaData.OverrideDataID);
 					}
 					break;
 			}
@@ -332,7 +331,6 @@ void FCameraRigAssetReference::PostSerialize(const FArchive& Ar)
 		LegacyParameterProperties.Add(LegacyParameterProperty);\
 		FCameraRigAssetReferenceParameterMetaData MetaData;\
 		MetaData.ParameterGuid = ParameterOverride.InterfaceParameterGuid;\
-		MetaData.ParameterType = ECameraRigInterfaceParameterType::Blendable;\
 		MetaData.bIsOverridden = true;\
 		LegacyParameterMetaData.Add(MetaData);\
 		bHasAnyLegacyOverride = true;\
@@ -360,5 +358,18 @@ void FCameraRigAssetReference::PostSerialize(const FArchive& Ar)
 	}
 
 	PRAGMA_ENABLE_DEPRECATION_WARNINGS
+
+	if (!ParameterOverrideGuids_DEPRECATED.IsEmpty())
+	{
+		for (const FGuid& Guid : ParameterOverrideGuids_DEPRECATED)
+		{
+			FCameraRigAssetReferenceParameterMetaData MetaData;
+			MetaData.ParameterGuid = Guid;
+			MetaData.bIsOverridden = true;
+			ParameterMetaData.Add(MetaData);
+		}
+
+		ParameterOverrideGuids_DEPRECATED.Reset();
+	}
 }
 
diff --git a/Engine/Plugins/Cameras/GameplayCameras/Source/GameplayCameras/Public/Core/CameraRigAssetReference.h b/Engine/Plugins/Cameras/GameplayCameras/Source/GameplayCameras/Public/Core/CameraRigAssetReference.h
--- a/Engine/Plugins/Cameras/GameplayCameras/Source/GameplayCameras/Public/Core/CameraRigAssetReference.h
+++ b/Engine/Plugins/Cameras/GameplayCameras/Source/GameplayCameras/Public/Core/CameraRigAssetReference.h
@@ -249,10 +249,6 @@ struct FCameraRigAssetReferenceParameterMetaData
 	UPROPERTY()
 	FGuid ParameterGuid;
 
-	/** The type of parameter. */
-	UPROPERTY()
-	ECameraRigInterfaceParameterType ParameterType = ECameraRigInterfaceParameterType::Blendable;
-
 	/** The ID to use for overriding a blendable parameter. */
 	UPROPERTY()
 	FCameraVariableID OverrideVariableID;
@@ -342,7 +338,7 @@ public:
 private:
 
 	const FCameraRigAssetReferenceParameterMetaData* FindMetaData(const FGuid& PropertyID) const;
-	FCameraRigAssetReferenceParameterMetaData* FindMetaData(const FGuid& PropertyID);
+	FCameraRigAssetReferenceParameterMetaData& FindOrAddMetaData(const FGuid& PropertyID);
 	
 	void GenerateOverriddenParameterGuidArray(TArray<FGuid>& OutOverriddenIDs) const;
 
@@ -363,6 +359,9 @@ private:
 
 	// Deprecated
 
+	UPROPERTY()
+	TArray<FGuid> ParameterOverrideGuids_DEPRECATED;
+
 	PRAGMA_DISABLE_DEPRECATION_WARNINGS
 
 	UPROPERTY()
