Cameras: improve camera rig merging

Merging can now also happen if the two cameras rigs are exactly the same, or any combination of nested prefab nodes that eventually wrap the same underlying camera rig.

A few other changes:

- When merging, move any root prefab node evaluators to the merge blend stack. This means we only have the "pure" merged rig's evaluators afterwards.
- Switch the debug name of the camera rig for that entry in the main blend stack.
- Move all the stuff only needed by the transient/main blend stack into an array of "extra infos", so that the base blend stack class doesn't need to know about it.
- Simplify the variable table API for filters and make it more consistent with the context data table.

diff --git a/Engine/Plugins/Cameras/GameplayCameras/Source/GameplayCameras/Private/Core/BlendStackCameraNode.cpp b/Engine/Plugins/Cameras/GameplayCameras/Source/GameplayCameras/Private/Core/BlendStackCameraNode.cpp
--- a/Engine/Plugins/Cameras/GameplayCameras/Source/GameplayCameras/Private/Core/BlendStackCameraNode.cpp
+++ b/Engine/Plugins/Cameras/GameplayCameras/Source/GameplayCameras/Private/Core/BlendStackCameraNode.cpp
@@ -2,6 +2,7 @@
 
 #include "Core/BlendStackCameraNode.h"
 
+#include "Algo/AnyOf.h"
 #include "Core/BlendCameraNode.h"
 #include "Core/BlendStackCameraRigEvent.h"
 #include "Core/BlendStackRootCameraNode.h"
@@ -23,7 +24,6 @@
 #include "Math/ColorList.h"
 #include "Modules/ModuleManager.h"
 #include "Nodes/Blends/PopBlendCameraNode.h"
-#include "Nodes/Common/CameraRigCameraNode.h"
 
 #include UE_INLINE_GENERATED_CPP_BY_NAME(BlendStackCameraNode)
 
@@ -79,7 +79,8 @@ bool FBlendStackCameraNodeEvaluator::InitializeEntry(
 		FCameraRigEntry& NewEntry, 
 		const UCameraRigAsset* CameraRig,
 		TSharedPtr<const FCameraEvaluationContext> EvaluationContext,
-		UBlendStackRootCameraNode* EntryRootNode)
+		UBlendStackRootCameraNode* EntryRootNode,
+		bool bSetActiveResult)
 {
 	// Clear the evaluator hierarchy in case we are hot-reloading an entry.
 	NewEntry.EvaluatorHierarchy.Reset();
@@ -98,16 +99,25 @@ bool FBlendStackCameraNodeEvaluator::InitializeEntry(
 	NewEntry.Result.VariableTable.Initialize(CameraRig->AllocationInfo.VariableTableInfo);
 	NewEntry.Result.ContextDataTable.Initialize(CameraRig->AllocationInfo.ContextDataTableInfo);
 
-	// Set the default values.
+	// Set the default values, so that pre-blending from default values works.
 	FCameraParameterOverrideHelper::ApplyDefaultBlendableParameters(CameraRig, NewEntry.Result.VariableTable);
 
 	// Set all the data from the context.
 	const FCameraNodeEvaluationResult& ContextResult = EvaluationContext->GetInitialResult();
-	NewEntry.Result.VariableTable.Override(
-			ContextResult.VariableTable,
-			ECameraVariableTableFilter::InputOutput | ECameraVariableTableFilter::Private);
+	NewEntry.Result.VariableTable.OverrideAll(ContextResult.VariableTable, true);
 	NewEntry.Result.ContextDataTable.OverrideAll(ContextResult.ContextDataTable);
 
+	// Add some conditional result if necessary.
+	if (bSetActiveResult && EvaluationContext)
+	{
+		const FCameraNodeEvaluationResult* ActiveOnlyResult = EvaluationContext->GetConditionalResult(ECameraEvaluationDataCondition::ActiveCameraRig);
+		if (ActiveOnlyResult)
+		{
+			NewEntry.Result.VariableTable.OverrideAll(ActiveOnlyResult->VariableTable, true);
+			NewEntry.Result.ContextDataTable.OverrideAll(ActiveOnlyResult->ContextDataTable);
+		}
+	}
+
 	// Initialize the node evaluators.
 	FCameraNodeEvaluatorInitializeParams InitParams(&NewEntry.EvaluatorHierarchy);
 	InitParams.Evaluator = OwningEvaluator;
@@ -262,8 +272,6 @@ void FBlendStackCameraNodeEvaluator::OnRunFinished(FCameraNodeEvaluationResult&
 	for (FCameraRigEntry& Entry : Entries)
 	{
 		Entry.bIsFirstFrame = false;
-		Entry.bInputRunThisFrame = false;
-		Entry.bBlendRunThisFrame = false;
 		Entry.bForceCameraCut = false;
 	}
 
@@ -446,8 +454,6 @@ void FBlendStackCameraNodeEvaluator::OnSerialize(const FCameraNodeEvaluatorSeria
 	{
 		Entry.Result.Serialize(Ar);
 		Ar << Entry.bIsFirstFrame;
-		Ar << Entry.bInputRunThisFrame;
-		Ar << Entry.bBlendRunThisFrame;
 		Ar << Entry.bIsFrozen;
 #if UE_GAMEPLAY_CAMERAS_TRACE
 		Ar << Entry.bLogWarnings;
@@ -459,8 +465,9 @@ void FBlendStackCameraNodeEvaluator::OnSerialize(const FCameraNodeEvaluatorSeria
 
 void FBlendStackCameraNodeEvaluator::OnPostBuildAsset(const FGameplayCameraAssetBuildEvent& BuildEvent)
 {
-	for (FCameraRigEntry& Entry : Entries)
+	for (int32 Index = 0; Index < Entries.Num(); ++Index)
 	{
+		FCameraRigEntry& Entry(Entries[Index]);
 		const bool bRebuildEntry = Entry.ListenedPackages.Contains(BuildEvent.AssetPackage);
 		if (bRebuildEntry)
 		{
@@ -471,7 +478,8 @@ void FBlendStackCameraNodeEvaluator::OnPostBuildAsset(const FGameplayCameraAsset
 			Entry.RootNode->RootNode = Entry.CameraRig->RootNode;
 
 			// Remove the blend on the root node, since we don't want the reloaded camera rig to re-blend-in
-			// for no good reason.
+			// for no good reason. This might "pop" if we reloaded this entry while it was blending, but
+			// that's acceptable.
 			Entry.RootNode->Blend = NewObject<UPopBlendCameraNode>(Entry.RootNode, NAME_None);
 
 			// Rebuild the evaluator tree.
@@ -479,7 +487,8 @@ void FBlendStackCameraNodeEvaluator::OnPostBuildAsset(const FGameplayCameraAsset
 					Entry,
 					Entry.CameraRig,
 					Entry.EvaluationContext.Pin(),
-					Entry.RootNode);
+					Entry.RootNode,
+					Index == Entries.Num() - 1);
 			if (!bInitialized)
 			{
 				Entry.bIsFrozen = true;
@@ -510,15 +519,15 @@ void FTransientBlendStackCameraNodeEvaluator::Push(const FBlendStackCameraPushPa
 			}
 
 			// See if we can merge the new camera rig onto the active camera rig.
-			const EBlendStackEntryComparison Comparison = TopEntry.RootEvaluator->Compare(Params.CameraRig);
+			const ECameraRigMergingEligibility Eligibility = TopEntry.RootEvaluator->CompareCameraRigForMerging(Params.CameraRig);
 
-			if (!Params.bForcePush && Comparison == EBlendStackEntryComparison::Active)
+			if (!Params.bForcePush && Eligibility == ECameraRigMergingEligibility::Active)
 			{
 				// This camera rig is already the active one on the merged stack.
 				return;
 			}
 
-			if (Comparison == EBlendStackEntryComparison::EligibleForMerge)
+			if (Eligibility == ECameraRigMergingEligibility::EligibleForMerge)
 			{
 				// This camera rig can be merged with the one current running. However, we
 				// only do it if the transition explicitly allows it.
@@ -527,7 +536,7 @@ void FTransientBlendStackCameraNodeEvaluator::Push(const FBlendStackCameraPushPa
 
 				if (Transition && Transition->bAllowCameraRigMerging)
 				{
-					PushVariantEntry(Params, Transition);
+					PushMergedEntry(Params, Transition);
 					return;
 				}
 			}
@@ -577,11 +586,27 @@ void FTransientBlendStackCameraNodeEvaluator::PushNewEntry(const FBlendStackCame
 			NewEntry, 
 			Params.CameraRig,
 			Params.EvaluationContext,
-			EntryRootNode);
+			EntryRootNode,
+			true);
 	if (!bInitialized)
 	{
 		return;
 	}
+	
+	// Take a snapshot of the initial result for having base values during pre-blending.
+	FCameraRigEntryExtraInfo NewExtraInfo;
+	{
+		FCameraNodeEvaluationResult& UnblendedResult = NewExtraInfo.UnblendedResult;
+
+		if (Params.CameraRig)
+		{
+			const FCameraRigAllocationInfo& AllocationInfo = Params.CameraRig->AllocationInfo;
+			UnblendedResult.VariableTable.Initialize(AllocationInfo.VariableTableInfo);
+			UnblendedResult.ContextDataTable.Initialize(AllocationInfo.ContextDataTableInfo);
+		}
+
+		UnblendedResult.OverrideAll(NewEntry.Result, true);
+	}
 
 #if WITH_EDITOR
 	// Listen to changes to the packages inside which this camera rig is defined. We will hot-reload the
@@ -592,6 +617,7 @@ void FTransientBlendStackCameraNodeEvaluator::PushNewEntry(const FBlendStackCame
 	// Important: we need to move the new entry here because copying evaluator storage
 	// is disabled.
 	Entries.Add(MoveTemp(NewEntry));
+	EntryExtraInfos.Add(MoveTemp(NewExtraInfo));
 
 	if (OnCameraRigEventDelegate.IsBound())
 	{
@@ -599,15 +625,18 @@ void FTransientBlendStackCameraNodeEvaluator::PushNewEntry(const FBlendStackCame
 	}
 }
 
-void FTransientBlendStackCameraNodeEvaluator::PushVariantEntry(const FBlendStackCameraPushParams& PushParams, const UCameraRigTransition* Transition)
+void FTransientBlendStackCameraNodeEvaluator::PushMergedEntry(const FBlendStackCameraPushParams& PushParams, const UCameraRigTransition* Transition)
 {
-	const UCameraRigCameraNode* PrefabNode = Cast<const UCameraRigCameraNode>(PushParams.CameraRig->RootNode);
 	const UBlendCameraNode* Blend = Transition ? Transition->Blend : nullptr;
 
 	FCameraRigEntry& TopEntry = Entries.Top();
 	FCameraNodeEvaluatorBuilder Builder(TopEntry.EvaluatorStorage);
 	FCameraNodeEvaluatorBuildParams BuildParams(Builder);
-	TopEntry.RootEvaluator->MergeCameraRig(BuildParams, PrefabNode, Blend);
+	TopEntry.RootEvaluator->MergeCameraRig(BuildParams, PushParams.CameraRig, Blend);
+
+	// Merging a camera rig the first time changes the evaluator tree by removing any prefab nodes
+	// near the root... so we need to rebuild our hierarchy cache.
+	TopEntry.EvaluatorHierarchy.Build(TopEntry.RootEvaluator);
 
 	// Swap out the camera rig registered as "active" for this entry.
 #if WITH_EDITOR
@@ -647,6 +676,8 @@ void FTransientBlendStackCameraNodeEvaluator::FreezeAll(TSharedPtr<FCameraEvalua
 
 void FTransientBlendStackCameraNodeEvaluator::OnRun(const FCameraNodeEvaluationParams& Params, FCameraNodeEvaluationResult& OutResult)
 {
+	ensure(Entries.Num() == EntryExtraInfos.Num());
+
 	// Validate our entries and resolve evaluation context weak pointers.
 	TArray<FResolvedEntry> ResolvedEntries;
 	ResolveEntries(ResolvedEntries);
@@ -668,16 +699,18 @@ void FTransientBlendStackCameraNodeEvaluator::OnRun(const FCameraNodeEvaluationP
 
 	// Tidy up.
 	OnRunFinished(OutResult);
+	InternalRunFinished(OutResult);
 }
 
 void FTransientBlendStackCameraNodeEvaluator::InternalPreBlendPrepare(TArrayView<FResolvedEntry> ResolvedEntries, const FCameraNodeEvaluationParams& Params, FCameraNodeEvaluationResult& OutResult)
 {
-	constexpr ECameraVariableTableFilter VariableTableFilter = ECameraVariableTableFilter::InputOutput | ECameraVariableTableFilter::Private;
-	constexpr ECameraContextDataTableFilter ContextDataTableFilter = ECameraContextDataTableFilter::None;
+	constexpr ECameraVariableTableFilter VariableTableFilter = ECameraVariableTableFilter::ChangedOnly;
+	constexpr ECameraContextDataTableFilter ContextDataTableFilter = ECameraContextDataTableFilter::ChangedOnly;
 
 	for (FResolvedEntry& ResolvedEntry : ResolvedEntries)
 	{
 		FCameraRigEntry& Entry(ResolvedEntry.Entry);
+		FCameraRigEntryExtraInfo& EntryExtraInfo(EntryExtraInfos[ResolvedEntry.EntryIndex]);
 
 		if (UNLIKELY(Entry.bIsFrozen))
 		{
@@ -688,32 +721,35 @@ void FTransientBlendStackCameraNodeEvaluator::InternalPreBlendPrepare(TArrayView
 		CurParams.EvaluationContext = ResolvedEntry.Context;
 		CurParams.bIsFirstFrame = Entry.bIsFirstFrame;
 
-		FCameraNodeEvaluationResult& CurResult(Entry.Result);
+		FCameraNodeEvaluationResult& UnblendedResult(EntryExtraInfo.UnblendedResult);
+
+		// Reset frame flags to match what happend to the Entry.Result in ResolveEntries.
+		UnblendedResult.ResetFrameFlags();
 
 		// Start with the input given to us.
-		CurResult.VariableTable.OverrideAll(OutResult.VariableTable);
+		UnblendedResult.VariableTable.OverrideAll(OutResult.VariableTable);
 
-		// Override it with whatever the evaluation context has set on its result.
+		// Override it with whatever the evaluation context has set on its result this frame.
 		// Evaluation contexts may have private variables we need to pass along, such as when rig parameter
 		// overrides have been set on them, so include private variables in the filter.
 		const FCameraNodeEvaluationResult& ContextResult(ResolvedEntry.Context->GetInitialResult());
-		CurResult.VariableTable.Override(ContextResult.VariableTable, VariableTableFilter);
-		CurResult.ContextDataTable.Override(ContextResult.ContextDataTable, ContextDataTableFilter);
+		UnblendedResult.VariableTable.Override(ContextResult.VariableTable, VariableTableFilter);
+		UnblendedResult.ContextDataTable.Override(ContextResult.ContextDataTable, ContextDataTableFilter);
 		if (ResolvedEntry.bIsActiveEntry)
 		{
 			if (const FCameraNodeEvaluationResult* ActiveOnlyResult = ResolvedEntry.Context->GetConditionalResult(ECameraEvaluationDataCondition::ActiveCameraRig))
 			{
-				CurResult.VariableTable.Override(ActiveOnlyResult->VariableTable, VariableTableFilter);
-				CurResult.ContextDataTable.Override(ActiveOnlyResult->ContextDataTable, ContextDataTableFilter);
+				UnblendedResult.VariableTable.Override(ActiveOnlyResult->VariableTable, VariableTableFilter);
+				UnblendedResult.ContextDataTable.Override(ActiveOnlyResult->ContextDataTable, ContextDataTableFilter);
 			}
 		}
 
 		// Gather input parameters if needed (and remember if it was indeed needed).
-		if (!Entry.bInputRunThisFrame)
+		if (!EntryExtraInfo.bInputRunThisFrame)
 		{
 			bool bHasPreBlendedParameters = false;
-			FCameraBlendedParameterUpdateParams InputParams(CurParams, CurResult.CameraPose);
-			FCameraBlendedParameterUpdateResult InputResult(CurResult.VariableTable);
+			FCameraBlendedParameterUpdateParams InputParams(CurParams, UnblendedResult.CameraPose);
+			FCameraBlendedParameterUpdateResult InputResult(UnblendedResult.VariableTable);
 
 			Entry.EvaluatorHierarchy.ForEachEvaluator(ECameraNodeEvaluatorFlags::NeedsParameterUpdate,
 					[&bHasPreBlendedParameters, &InputParams, &InputResult](FCameraNodeEvaluator* ParameterEvaluator)
@@ -722,44 +758,72 @@ void FTransientBlendStackCameraNodeEvaluator::InternalPreBlendPrepare(TArrayView
 						bHasPreBlendedParameters = true;
 					});
 
-			ResolvedEntry.bHasPreBlendedParameters = bHasPreBlendedParameters;
-			Entry.bInputRunThisFrame = true;
+			EntryExtraInfo.bHasPreBlendedParameters = bHasPreBlendedParameters;
+			EntryExtraInfo.bInputRunThisFrame = true;
 		}
 
 		// Run blends.
 		// Note that we pass last frame's camera pose to the Run() method. This may change.
 		// Blends aren't expected to use the camera pose to do any logic until BlendResults().
-		if (!Entry.bBlendRunThisFrame)
+		if (!EntryExtraInfo.bBlendRunThisFrame)
 		{
 			FBlendCameraNodeEvaluator* EntryBlendEvaluator = Entry.RootEvaluator->GetBlendEvaluator();
 			if (EntryBlendEvaluator)
 			{
-				EntryBlendEvaluator->Run(CurParams, CurResult);
+				EntryBlendEvaluator->Run(CurParams, UnblendedResult);
 			}
 
-			Entry.bBlendRunThisFrame = true;
+			EntryExtraInfo.bBlendRunThisFrame = true;
 		}
+
+		// Assign the new unblended result back to the entry result. The goal is that:
+		//
+		// 1) We only write values from the context result when they have been written to.
+		//    That is, we don't always set all values every frame. When "left alone", a camera rig
+		//    should retain whatever values it was last set with. For instance, if a camera rig 
+		//    received special values when it was the active camera rig, we want to keep those 
+		//    values when it starts blending out, instead of immediately overwriting them with
+		//    the "common" values (i.e. the values from the initial result).
+		//
+		// 2) However, we do want to write all the values into the PreBlendVariableTable every
+		//    frame, so that we can have them blend together. The blended values are then reassigned
+		//    to Entry.Result in InternalPreBlendExecute, which means we lose each camera rig's
+		//    "intended" values. Getting them from the evaluation context each frame would violate
+		//    previous item (1) so that's why we use the UnblendedResult as a snapshot.
+		//
+		FCameraNodeEvaluationResult& CurResult(Entry.Result);
+		CurResult.OverrideAll(UnblendedResult, true);
 	}
 }
 
 void FTransientBlendStackCameraNodeEvaluator::InternalPreBlendExecute(TArrayView<FResolvedEntry> ResolvedEntries, const FCameraNodeEvaluationParams& Params, FCameraNodeEvaluationResult& OutResult)
 {
 	// Blend all the camera rigs' input variables (include private variables such as camera rig parameters).
+	// The result of this pre-blending goes into PreBlendVariableTable, which will then have a big mix
+	// of all the variables.
 	PreBlendVariableTable.ClearAllWrittenThisFrameFlags();
 
-	const ECameraVariableTableFilter VariableTableFilter = ECameraVariableTableFilter::Input | ECameraVariableTableFilter::Private;
+	if (!Algo::AnyOf(EntryExtraInfos, [](const FCameraRigEntryExtraInfo& Item) { return Item.bHasPreBlendedParameters; }))
+	{
+		return;
+	}
+
+	constexpr ECameraVariableTableFilter VariableTableFilter = ECameraVariableTableFilter::None;
 
 	for (FResolvedEntry& ResolvedEntry : ResolvedEntries)
 	{
 		FCameraRigEntry& Entry(ResolvedEntry.Entry);
-		FCameraNodeEvaluationResult& CurResult(Entry.Result);
+		FCameraRigEntryExtraInfo& EntryExtraInfo(EntryExtraInfos[ResolvedEntry.EntryIndex]);
+
+		FCameraNodeEvaluationResult& UnblendedResult(EntryExtraInfo.UnblendedResult);
 
 		if (!Entry.bIsFrozen)
 		{
 			FCameraNodeEvaluationParams CurParams(Params);
 			CurParams.EvaluationContext = ResolvedEntry.Context;
 			CurParams.bIsFirstFrame = Entry.bIsFirstFrame;
-			FCameraNodePreBlendParams PreBlendParams(CurParams, CurResult.CameraPose, CurResult.VariableTable);
+
+			FCameraNodePreBlendParams PreBlendParams(CurParams, UnblendedResult.CameraPose, UnblendedResult.VariableTable);
 			PreBlendParams.VariableTableFilter = VariableTableFilter;
 
 			FCameraNodePreBlendResult PreBlendResult(PreBlendVariableTable);
@@ -768,16 +832,18 @@ void FTransientBlendStackCameraNodeEvaluator::InternalPreBlendExecute(TArrayView
 			if (EntryBlendEvaluator)
 			{
 				EntryBlendEvaluator->BlendParameters(PreBlendParams, PreBlendResult);
+				EntryExtraInfo.bIsPreBlendFull = (PreBlendResult.bIsBlendFinished && PreBlendResult.bIsBlendFull);
 			}
 			else
 			{
-				PreBlendVariableTable.Override(CurResult.VariableTable, PreBlendParams.VariableTableFilter);
+				PreBlendVariableTable.Override(UnblendedResult.VariableTable, PreBlendParams.VariableTableFilter);
+				EntryExtraInfo.bIsPreBlendFull = true;
 			}
 		}
 		else
 		{
 			// Frozen entries still contribute to the blend using their last evaluated values.
-			PreBlendVariableTable.Override(CurResult.VariableTable, VariableTableFilter);
+			PreBlendVariableTable.Override(UnblendedResult.VariableTable, VariableTableFilter);
 		}
 	}
 
@@ -787,11 +853,14 @@ void FTransientBlendStackCameraNodeEvaluator::InternalPreBlendExecute(TArrayView
 	for (FResolvedEntry& ResolvedEntry : ResolvedEntries)
 	{
 		FCameraRigEntry& Entry(ResolvedEntry.Entry);
+		FCameraRigEntryExtraInfo& EntryExtraInfo(EntryExtraInfos[ResolvedEntry.EntryIndex]);
 
 		if (!Entry.bIsFrozen)
 		{
-			// TEMP: disabled temporarily as it breaks blending between different prefabs.
-			//Entry.Result.VariableTable.Override(PreBlendVariableTable, VariableTableFilter | ECameraVariableTableFilter::KnownOnly);
+			FCameraNodeEvaluationResult& UnblendedResult(EntryExtraInfo.UnblendedResult);
+			FCameraNodeEvaluationResult& CurResult(Entry.Result);
+
+			Entry.Result.VariableTable.Override(PreBlendVariableTable, ECameraVariableTableFilter::KnownOnly);
 		}
 	}
 }
@@ -801,18 +870,19 @@ void FTransientBlendStackCameraNodeEvaluator::InternalUpdate(TArrayView<FResolve
 	for (FResolvedEntry& ResolvedEntry : ResolvedEntries)
 	{
 		FCameraRigEntry& Entry(ResolvedEntry.Entry);
+		FCameraRigEntryExtraInfo& EntryExtraInfo(EntryExtraInfos[ResolvedEntry.EntryIndex]);
 
 		if (UNLIKELY(Entry.bIsFrozen))
 		{
 			continue;
 		}
 
+		FCameraNodeEvaluationResult& CurResult(Entry.Result);
+
 		FCameraNodeEvaluationParams CurParams(Params);
 		CurParams.EvaluationContext = ResolvedEntry.Context;
 		CurParams.bIsFirstFrame = Entry.bIsFirstFrame;
 
-		FCameraNodeEvaluationResult& CurResult(Entry.Result);
-
 		// Start with the input given to us.
 		CurResult.Reset();
 		CurResult.CameraPose = OutResult.CameraPose;
@@ -853,6 +923,8 @@ void FTransientBlendStackCameraNodeEvaluator::InternalPostBlendExecute(TArrayVie
 	for (FResolvedEntry& ResolvedEntry : ResolvedEntries)
 	{
 		FCameraRigEntry& Entry(ResolvedEntry.Entry);
+		FCameraRigEntryExtraInfo& EntryExtraInfo(EntryExtraInfos[ResolvedEntry.EntryIndex]);
+
 		FCameraNodeEvaluationResult& CurResult(Entry.Result);
 
 		if (!Entry.bIsFrozen)
@@ -887,6 +959,7 @@ void FTransientBlendStackCameraNodeEvaluator::InternalPostBlendExecute(TArrayVie
 
 			PopEntriesBelow = ResolvedEntry.EntryIndex;
 		}
+		// else, merged entries only contribute to the pre-blend.
 	}
 
 	// Pop out camera rigs that have been blended out.
@@ -894,6 +967,32 @@ void FTransientBlendStackCameraNodeEvaluator::InternalPostBlendExecute(TArrayVie
 	if (BlendStackNode->BlendStackType == ECameraBlendStackType::IsolatedTransient && PopEntriesBelow != INDEX_NONE)
 	{
 		PopEntries(PopEntriesBelow);
+
+		EntryExtraInfos.RemoveAt(0, PopEntriesBelow);
+	}
+}
+
+void FTransientBlendStackCameraNodeEvaluator::InternalRunFinished(FCameraNodeEvaluationResult& OutResult)
+{
+	for (FCameraRigEntryExtraInfo& ExtraInfo : EntryExtraInfos)
+	{
+		ExtraInfo.bInputRunThisFrame = false;
+		ExtraInfo.bBlendRunThisFrame = false;
+		ExtraInfo.bHasPreBlendedParameters = false;
+		ExtraInfo.bIsPreBlendFull = false;
+	}
+}
+
+void FTransientBlendStackCameraNodeEvaluator::OnSerialize(const FCameraNodeEvaluatorSerializeParams& Params, FArchive& Ar)
+{
+	Super::OnSerialize(Params, Ar);
+
+	for (FCameraRigEntryExtraInfo& ExtraInfo : EntryExtraInfos)
+	{
+		Ar << ExtraInfo.bInputRunThisFrame;
+		Ar << ExtraInfo.bBlendRunThisFrame;
+		Ar << ExtraInfo.bHasPreBlendedParameters;
+		Ar << ExtraInfo.bIsPreBlendFull;
 	}
 }
 
@@ -1065,7 +1164,8 @@ void FPersistentBlendStackCameraNodeEvaluator::Insert(const FBlendStackCameraIns
 			NewEntry, 
 			Params.CameraRig,
 			Params.EvaluationContext,
-			EntryRootNode);
+			EntryRootNode,
+			false);
 	if (!bInitialized)
 	{
 		return;
@@ -1111,7 +1211,7 @@ void FPersistentBlendStackCameraNodeEvaluator::OnRun(const FCameraNodeEvaluation
 
 void FPersistentBlendStackCameraNodeEvaluator::InternalUpdate(TArrayView<FResolvedEntry> ResolvedEntries, const FCameraNodeEvaluationParams& Params, FCameraNodeEvaluationResult& OutResult)
 {
-	constexpr ECameraVariableTableFilter VariableTableFilter = ECameraVariableTableFilter::InputOutput | ECameraVariableTableFilter::Private | ECameraVariableTableFilter::ChangedOnly;
+	constexpr ECameraVariableTableFilter VariableTableFilter = ECameraVariableTableFilter::ChangedOnly;
 	constexpr ECameraContextDataTableFilter ContextDataTableFilter = ECameraContextDataTableFilter::ChangedOnly;
 
 	for (FResolvedEntry& ResolvedEntry : ResolvedEntries)
@@ -1175,7 +1275,7 @@ void FPersistentBlendStackCameraNodeEvaluator::InternalUpdate(TArrayView<FResolv
 			}
 			else
 			{
-				OutResult.VariableTable.Override(CurResult.VariableTable, ECameraVariableTableFilter::Input);
+				OutResult.VariableTable.Override(CurResult.VariableTable, ECameraVariableTableFilter::None);
 			}
 
 			// Run the camera rig's root node.
@@ -1202,7 +1302,7 @@ void FPersistentBlendStackCameraNodeEvaluator::InternalUpdate(TArrayView<FResolv
 		{
 			FCameraNodeEvaluationResult& CurResult(Entry.Result);
 
-			OutResult.VariableTable.Override(CurResult.VariableTable, ECameraVariableTableFilter::Input);
+			OutResult.VariableTable.Override(CurResult.VariableTable, ECameraVariableTableFilter::None);
 			OutResult.OverrideAll(CurResult);
 		}
 	}
diff --git a/Engine/Plugins/Cameras/GameplayCameras/Source/GameplayCameras/Private/Core/BlendStackRootCameraNode.cpp b/Engine/Plugins/Cameras/GameplayCameras/Source/GameplayCameras/Private/Core/BlendStackRootCameraNode.cpp
--- a/Engine/Plugins/Cameras/GameplayCameras/Source/GameplayCameras/Private/Core/BlendStackRootCameraNode.cpp
+++ b/Engine/Plugins/Cameras/GameplayCameras/Source/GameplayCameras/Private/Core/BlendStackRootCameraNode.cpp
@@ -16,6 +16,12 @@
 
 #include UE_INLINE_GENERATED_CPP_BY_NAME(BlendStackRootCameraNode)
 
+UBlendStackRootCameraNode::UBlendStackRootCameraNode(const FObjectInitializer& ObjInit)
+	: Super(ObjInit)
+{
+	AddNodeFlags(ECameraNodeFlags::CustomGetChildren);
+}
+
 FCameraNodeChildrenView UBlendStackRootCameraNode::OnGetChildren()
 {
 	FCameraNodeChildrenView Children;
@@ -84,61 +90,41 @@ void FBlendStackRootCameraNodeEvaluator::OnInitialize(const FCameraNodeEvaluator
 	const UBlendStackRootCameraNode* RootNode = GetCameraNodeAs<UBlendStackRootCameraNode>();
 	if (RootNode->RootNode)
 	{
-		// See if the camera rig we are running is only made up of a prefab with some overrides.
-		InitialPrefabNode = Cast<const UCameraRigCameraNode>(RootNode->RootNode);
-		if (InitialPrefabNode)
-		{
-			BlendablePrefabCameraRig = InitialPrefabNode->CameraRigReference.GetCameraRig();
-		}
+		const UCameraRigAsset* CameraRig = RootNode->RootNode->GetTypedOuter<const UCameraRigAsset>();
+		BlendablePrefabCameraRig = FindInnermostCameraRigPrefab(CameraRig);
 
 #if UE_GAMEPLAY_CAMERAS_DEBUG
-		const UCameraRigAsset* CameraRig = RootNode->RootNode->GetTypedOuter<UCameraRigAsset>();
 		CameraRigAssetName = GetNameSafe(CameraRig);
 #endif  // UE_GAMEPLAY_CAMERAS_DEBUG
 	}
 }
 
-EBlendStackEntryComparison FBlendStackRootCameraNodeEvaluator::Compare(const UCameraRigAsset* CameraRig) const
+ECameraRigMergingEligibility FBlendStackRootCameraNodeEvaluator::CompareCameraRigForMerging(const UCameraRigAsset* CameraRig) const
 {
-	if (!BlendablePrefabCameraRig)
-	{
-		return EBlendStackEntryComparison::Different;
-	}
-	
-	const UCameraRigCameraNode* PrefabNode = Cast<const UCameraRigCameraNode>(CameraRig->RootNode);
-	if (!PrefabNode)
-	{
-		return EBlendStackEntryComparison::Different;
-	}
+	const UCameraRigAsset* NewCameraRigPrefab = FindInnermostCameraRigPrefab(CameraRig);
 
-	const UCameraRigAsset* Prefab = PrefabNode->CameraRigReference.GetCameraRig();
-	if (!Prefab)
+	if (NewCameraRigPrefab != BlendablePrefabCameraRig)
 	{
-		return EBlendStackEntryComparison::Different;
-	}
-
-	if (Prefab != BlendablePrefabCameraRig)
-	{
-		return EBlendStackEntryComparison::Different;
+		return ECameraRigMergingEligibility::Different;
 	}
 
 	if (BlendedParameterOverridesStack.IsEmpty())
 	{
-		return EBlendStackEntryComparison::EligibleForMerge;
+		return ECameraRigMergingEligibility::EligibleForMerge;
 	}
 	
 	const FBlendedParameterOverrides& TopEntry = BlendedParameterOverridesStack.Top();
-	if (TopEntry.PrefabNodeAsset == CameraRig)
+	if (TopEntry.CameraRig == CameraRig)
 	{
-		return EBlendStackEntryComparison::Active;
+		return ECameraRigMergingEligibility::Active;
 	}
 
-	return EBlendStackEntryComparison::EligibleForMerge;
+	return ECameraRigMergingEligibility::EligibleForMerge;
 }
 
-void FBlendStackRootCameraNodeEvaluator::MergeCameraRig(const FCameraNodeEvaluatorBuildParams& Params, const UCameraRigCameraNode* PrefabNode, const UBlendCameraNode* Blend)
+void FBlendStackRootCameraNodeEvaluator::MergeCameraRig(const FCameraNodeEvaluatorBuildParams& Params, const UCameraRigAsset* CameraRig, const UBlendCameraNode* Blend)
 {
-	if (!ensureMsgf(PrefabNode, TEXT("No prefab node given.")))
+	if (!ensureMsgf(CameraRig, TEXT("No camera rig given.")))
 	{
 		return;
 	}
@@ -149,63 +135,93 @@ void FBlendStackRootCameraNodeEvaluator::MergeCameraRig(const FCameraNodeEvaluat
 		return;
 	}
 
-	if (!ensureMsgf(PrefabNode->CameraRigReference.GetCameraRig() == BlendablePrefabCameraRig,
-				TEXT("Adding blended parameter overrides for a different camera rig.")))
-	{
-		return;
-	}
-
 	InitializeBlendedParameterOverridesStack();
 
 	FBlendedParameterOverrides BlendedParameterOverrides;
-	BlendedParameterOverrides.PrefabNodeAsset = PrefabNode->GetTypedOuter<UCameraRigAsset>();
-	BlendedParameterOverrides.PrefabNode = PrefabNode;
+	BlendedParameterOverrides.CameraRig = CameraRig;
 	BlendedParameterOverrides.Blend = Blend;
-	BlendedParameterOverrides.Result.VariableTable.Initialize(BlendedParameterOverridesTableAllocationInfo);
+
+	BuildNestedPrefabTrail(CameraRig, BlendedParameterOverrides.PrefabTrail);
+
+	const FCameraRigAllocationInfo& AllocationInfo = BlendablePrefabCameraRig->AllocationInfo;
+	BlendedParameterOverrides.Result.VariableTable.Initialize(AllocationInfo.VariableTableInfo);
+
 	if (Blend)
 	{
 		BlendedParameterOverrides.BlendEvaluator = Params.BuildEvaluatorAs<FBlendCameraNodeEvaluator>(Blend);
 	}
 
 	BlendedParameterOverridesStack.Add(MoveTemp(BlendedParameterOverrides));
+
+	// Show the new active camera rig in the debug info.
+#if UE_GAMEPLAY_CAMERAS_DEBUG
+	CameraRigAssetName = GetNameSafe(CameraRig);
+#endif  // UE_GAMEPLAY_CAMERAS_DEBUG
 }
 
-void FBlendStackRootCameraNodeEvaluator::InitializeBlendedParameterOverridesStack()
+const UCameraRigAsset* FBlendStackRootCameraNodeEvaluator::FindInnermostCameraRigPrefab(const UCameraRigAsset* CameraRig)
 {
-	if (!ensureMsgf(BlendablePrefabCameraRig, TEXT("The blended parameter overrides stack has already been initialized.")))
+	if (const UCameraRigCameraNode* PrefabNode = Cast<const UCameraRigCameraNode>(CameraRig->RootNode))
 	{
-		return;
+		if (const UCameraRigAsset* InnerCameraRig = PrefabNode->CameraRigReference.GetCameraRig())
+		{
+			return FindInnermostCameraRigPrefab(InnerCameraRig);
+		}
 	}
+	return CameraRig;
+}
 
+FCameraNodeEvaluator* FBlendStackRootCameraNodeEvaluator::FindInnermostCameraRigEvaluator(FCameraNodeEvaluator* CameraNodeEvaluator)
+{
+	if (FCameraRigCameraNodeEvaluator* PrefabNodeEvaluator = CameraNodeEvaluator->CastThis<FCameraRigCameraNodeEvaluator>())
+	{
+		if (FCameraNodeEvaluator* InnerNodeEvaluator = PrefabNodeEvaluator->GetCameraRigRootEvaluator())
+		{
+			return FindInnermostCameraRigEvaluator(InnerNodeEvaluator);
+		}
+	}
+	return CameraNodeEvaluator;
+}
+
+void FBlendStackRootCameraNodeEvaluator::InitializeBlendedParameterOverridesStack()
+{
 	if (!BlendedParameterOverridesStack.IsEmpty())
 	{
 		return;
 	}
 
-	// Build the allocation info for the variable tables we keep with each set of parameter overrides.
-	for (const UCameraRigBlendableParameter* BlendableParameter : BlendablePrefabCameraRig->Interface.BlendableParameters)
+	// Swap out the current root evaluator for the innermost rig one, because we want to apply parameter
+	// overrides ourselves from now on.
+	FCameraNodeEvaluator* InnermostRootEvaluator = FindInnermostCameraRigEvaluator(RootEvaluator);
+	RootEvaluator = InnermostRootEvaluator;
+
+	const UBlendStackRootCameraNode* ThisNode = GetCameraNodeAs<UBlendStackRootCameraNode>();
+	const UCameraRigAsset* OriginalCameraRig = ThisNode->RootNode->GetTypedOuter<const UCameraRigAsset>();
+
+	FBlendedParameterOverrides InitialParameterOverrides;
+	InitialParameterOverrides.CameraRig = OriginalCameraRig;
+	// No blend, the initial entry is always at 100%.
+
+	BuildNestedPrefabTrail(OriginalCameraRig, InitialParameterOverrides.PrefabTrail);
+
+	const FCameraRigAllocationInfo& AllocationInfo = BlendablePrefabCameraRig->AllocationInfo;
+	InitialParameterOverrides.Result.VariableTable.Initialize(AllocationInfo.VariableTableInfo);
+
+	BlendedParameterOverridesStack.Add(MoveTemp(InitialParameterOverrides));
+}
+
+void FBlendStackRootCameraNodeEvaluator::BuildNestedPrefabTrail(const UCameraRigAsset* CameraRig, TArray<TObjectPtr<const UCameraRigCameraNode>>& OutPrefabNodes)
+{
+	if (const UCameraRigCameraNode* PrefabNode = Cast<const UCameraRigCameraNode>(CameraRig->RootNode))
 	{
-		if (!ensure(BlendableParameter))
-		{
-			continue;
-		}
-		if (!BlendableParameter->PrivateVariableID)
+		if (const UCameraRigAsset* InnerCameraRig = PrefabNode->CameraRigReference.GetCameraRig())
 		{
-			continue;
+			OutPrefabNodes.Add(PrefabNode);
+			return BuildNestedPrefabTrail(InnerCameraRig, OutPrefabNodes);
 		}
-
-		FCameraVariableDefinition Definition = BlendableParameter->GetVariableDefinition();
-		BlendedParameterOverridesTableAllocationInfo.VariableDefinitions.Add(Definition);
 	}
 
-	FCameraRigCameraNodeEvaluator* RootPrefabNodeEvaluator = RootEvaluator->CastThisChecked<FCameraRigCameraNodeEvaluator>();
-	RootPrefabNodeEvaluator->SetApplyParameterOverrides(false);
-
-	FBlendedParameterOverrides InitialParameterOverrides;
-	InitialParameterOverrides.PrefabNodeAsset = InitialPrefabNode->GetTypedOuter<UCameraRigAsset>();
-	InitialParameterOverrides.PrefabNode = InitialPrefabNode;
-	InitialParameterOverrides.Result.VariableTable.Initialize(BlendedParameterOverridesTableAllocationInfo);
-	BlendedParameterOverridesStack.Add(MoveTemp(InitialParameterOverrides));
+	ensure(CameraRig == BlendablePrefabCameraRig);
 }
 
 void FBlendStackRootCameraNodeEvaluator::OnUpdateParameters(const FCameraBlendedParameterUpdateParams& Params, FCameraBlendedParameterUpdateResult& OutResult)
@@ -232,11 +248,6 @@ void FBlendStackRootCameraNodeEvaluator::RunBlendedParameterOverridesStack(const
 		return;
 	}
 
-	if (!ensure(BlendablePrefabCameraRig))
-	{
-		return;
-	}
-
 	int32 PopEntriesBelow = INDEX_NONE;
 	for (int32 EntryIndex = 0; EntryIndex < BlendedParameterOverridesStack.Num(); ++EntryIndex)
 	{
@@ -247,9 +258,13 @@ void FBlendStackRootCameraNodeEvaluator::RunBlendedParameterOverridesStack(const
 		// wouldn't have a base value to blend from.
 		FCameraParameterOverrideHelper::ApplyDefaultBlendableParameters(BlendablePrefabCameraRig, CurResult.VariableTable);
 
-		// Next, override the defaults with the specific values of this entry.
-		FCameraRigParameterOverrideEvaluator OverrideEvaluator(BlendedParameterOverrides.PrefabNode->CameraRigReference);
-		OverrideEvaluator.ApplyParameterOverrides(CurResult.VariableTable, false);
+		// Next, override the defaults with the specific values of this entry, applied bottoms up.
+		for (int32 Index = BlendedParameterOverrides.PrefabTrail.Num() - 1; Index >= 0; --Index)
+		{
+			const UCameraRigCameraNode* CurPrefabNode = BlendedParameterOverrides.PrefabTrail[Index];
+			FCameraRigParameterOverrideEvaluator OverrideEvaluator(CurPrefabNode->CameraRigReference);
+			OverrideEvaluator.ApplyParameterOverrides(CurResult.VariableTable, false);
+		}
 
 		// Finally, update the parameter overrides' blend, and apply it.
 		if (BlendedParameterOverrides.BlendEvaluator)
@@ -257,7 +272,7 @@ void FBlendStackRootCameraNodeEvaluator::RunBlendedParameterOverridesStack(const
 			BlendedParameterOverrides.BlendEvaluator->Run(Params.EvaluationParams, CurResult);
 
 			FCameraNodePreBlendParams BlendParams(Params.EvaluationParams, Params.LastCameraPose, CurResult.VariableTable);
-			BlendParams.VariableTableFilter = ECameraVariableTableFilter::Input | ECameraVariableTableFilter::Private;
+			BlendParams.VariableTableFilter = ECameraVariableTableFilter::InputOnly;
 			FCameraNodePreBlendResult BlendResult(OutResult.VariableTable);
 			BlendedParameterOverrides.BlendEvaluator->BlendParameters(BlendParams, BlendResult);
 
@@ -268,13 +283,12 @@ void FBlendStackRootCameraNodeEvaluator::RunBlendedParameterOverridesStack(const
 		}
 		else
 		{
-			OutResult.VariableTable.Override(
-					CurResult.VariableTable, ECameraVariableTableFilter::Input | ECameraVariableTableFilter::Private);
+			OutResult.VariableTable.Override(CurResult.VariableTable, ECameraVariableTableFilter::InputOnly);
 
 			PopEntriesBelow = EntryIndex;
 		}
 	}
-	if (PopEntriesBelow >= 0)
+	if (PopEntriesBelow > 0)
 	{
 		BlendedParameterOverridesStack.RemoveAt(0, PopEntriesBelow);
 	}
@@ -283,13 +297,13 @@ void FBlendStackRootCameraNodeEvaluator::RunBlendedParameterOverridesStack(const
 void FBlendStackRootCameraNodeEvaluator::OnAddReferencedObjects(FReferenceCollector& Collector)
 {
 	Collector.AddReferencedObject(BlendablePrefabCameraRig);
-	Collector.AddReferencedObject(InitialPrefabNode);
 
 	for (FBlendedParameterOverrides& BlendedParameterOverrides: BlendedParameterOverridesStack)
 	{
-		Collector.AddReferencedObject(BlendedParameterOverrides.PrefabNodeAsset);
-		Collector.AddReferencedObject(BlendedParameterOverrides.PrefabNode);
+		Collector.AddReferencedObject(BlendedParameterOverrides.CameraRig);
 		Collector.AddReferencedObject(BlendedParameterOverrides.Blend);
+		Collector.AddReferencedObjects(BlendedParameterOverrides.PrefabTrail);
+		BlendedParameterOverrides.Result.AddReferencedObjects(Collector);
 	}
 }
 
@@ -302,15 +316,7 @@ void FBlendStackRootCameraNodeEvaluator::OnBuildDebugBlocks(const FCameraDebugBl
 	Algo::Transform(BlendedParameterOverridesStack, DebugBlock.BlendedParameterOverridesEntries,
 			[](const FBlendedParameterOverrides& Item)
 			{
-				if (Item.PrefabNode)
-				{
-					const UCameraRigAsset* OuterCameraRig = Item.PrefabNode->GetTypedOuter<const UCameraRigAsset>();
-					if (OuterCameraRig)
-					{
-						return GetNameSafe(OuterCameraRig);
-					}
-				}
-				return FString(TEXT("<invalid camera rig>"));
+				return GetNameSafe(Item.CameraRig);
 			});
 
 	if (BlendEvaluator)
diff --git a/Engine/Plugins/Cameras/GameplayCameras/Source/GameplayCameras/Private/Core/BlendStackRootCameraNode.h b/Engine/Plugins/Cameras/GameplayCameras/Source/GameplayCameras/Private/Core/BlendStackRootCameraNode.h
--- a/Engine/Plugins/Cameras/GameplayCameras/Source/GameplayCameras/Private/Core/BlendStackRootCameraNode.h
+++ b/Engine/Plugins/Cameras/GameplayCameras/Source/GameplayCameras/Private/Core/BlendStackRootCameraNode.h
@@ -21,6 +21,10 @@ class UBlendStackRootCameraNode : public UCameraNode
 {
 	GENERATED_BODY()
 
+public:
+
+	UBlendStackRootCameraNode(const FObjectInitializer& ObjInit);
+
 protected:
 
 	// UCameraNode interface.
@@ -43,7 +47,7 @@ namespace UE::Cameras
 
 class FBlendCameraNodeEvaluator;
 
-enum class EBlendStackEntryComparison
+enum class ECameraRigMergingEligibility
 {
 	Different,
 	EligibleForMerge,
@@ -64,8 +68,9 @@ public:
 	FBlendCameraNodeEvaluator* GetBlendEvaluator() const { return BlendEvaluator; }
 	FCameraNodeEvaluator* GetRootEvaluator() const { return RootEvaluator; }
 
-	EBlendStackEntryComparison Compare(const UCameraRigAsset* CameraRig) const;
-	void MergeCameraRig(const FCameraNodeEvaluatorBuildParams& Params, const UCameraRigCameraNode* PrefabNode, const UBlendCameraNode* Blend);
+	ECameraRigMergingEligibility CompareCameraRigForMerging(const UCameraRigAsset* CameraRig) const;
+
+	void MergeCameraRig(const FCameraNodeEvaluatorBuildParams& Params, const UCameraRigAsset* CameraRig, const UBlendCameraNode* Blend);
 
 protected:
 
@@ -83,7 +88,13 @@ protected:
 
 private:
 
+	struct FBlendedParameterOverrides;
+
+	static const UCameraRigAsset* FindInnermostCameraRigPrefab(const UCameraRigAsset* CameraRig);
+	static FCameraNodeEvaluator* FindInnermostCameraRigEvaluator(FCameraNodeEvaluator* CameraNodeEvaluator);
+
 	void InitializeBlendedParameterOverridesStack();
+	void BuildNestedPrefabTrail(const UCameraRigAsset* CameraRig, TArray<TObjectPtr<const UCameraRigCameraNode>>& OutPrefabNodes);
 	void RunBlendedParameterOverridesStack(const FCameraBlendedParameterUpdateParams& Params, FCameraBlendedParameterUpdateResult& OutResult);
 
 private:
@@ -92,14 +103,13 @@ private:
 	FCameraNodeEvaluator* RootEvaluator = nullptr;
 
 	TObjectPtr<const UCameraRigAsset> BlendablePrefabCameraRig;
-	TObjectPtr<const UCameraRigCameraNode> InitialPrefabNode;
-	FCameraVariableTableAllocationInfo BlendedParameterOverridesTableAllocationInfo;
 
 	struct FBlendedParameterOverrides
 	{
-		TObjectPtr<const UCameraRigAsset> PrefabNodeAsset;
-		TObjectPtr<const UCameraRigCameraNode> PrefabNode;
+		TObjectPtr<const UCameraRigAsset> CameraRig;
 		TObjectPtr<const UBlendCameraNode> Blend;
+		TArray<TObjectPtr<const UCameraRigCameraNode>> PrefabTrail;
+
 		FBlendCameraNodeEvaluator* BlendEvaluator = nullptr;
 		FCameraNodeEvaluationResult Result;
 	};
diff --git a/Engine/Plugins/Cameras/GameplayCameras/Source/GameplayCameras/Private/Core/CameraIKAim.cpp b/Engine/Plugins/Cameras/GameplayCameras/Source/GameplayCameras/Private/Core/CameraIKAim.cpp
--- a/Engine/Plugins/Cameras/GameplayCameras/Source/GameplayCameras/Private/Core/CameraIKAim.cpp
+++ b/Engine/Plugins/Cameras/GameplayCameras/Source/GameplayCameras/Private/Core/CameraIKAim.cpp
@@ -162,9 +162,7 @@ void FCameraIKAim::RunRootCameraNode(const FCameraIKAimParams& Params, const FCa
 	ScratchResult.Reset();
 
 	// Make sure the camera rig will get its private variables, such as rig interface parameter overrides.
-	ScratchResult.VariableTable.Override(
-			CameraRigInfo.LastResult->VariableTable, 
-			ECameraVariableTableFilter::InputOutput | ECameraVariableTableFilter::Private);
+	ScratchResult.VariableTable.OverrideAll(CameraRigInfo.LastResult->VariableTable, true);
 
 	FRootCameraNodeEvaluator* RootEvaluator = Params.Evaluator->GetRootNodeEvaluator();
 
diff --git a/Engine/Plugins/Cameras/GameplayCameras/Source/GameplayCameras/Private/Core/CameraNodeEvaluator.cpp b/Engine/Plugins/Cameras/GameplayCameras/Source/GameplayCameras/Private/Core/CameraNodeEvaluator.cpp
--- a/Engine/Plugins/Cameras/GameplayCameras/Source/GameplayCameras/Private/Core/CameraNodeEvaluator.cpp
+++ b/Engine/Plugins/Cameras/GameplayCameras/Source/GameplayCameras/Private/Core/CameraNodeEvaluator.cpp
@@ -39,21 +39,22 @@ void FCameraNodeEvaluationResult::ResetFrameFlags()
 	ContextDataTable.ClearAllWrittenThisFrameFlags();
 }
 
-void FCameraNodeEvaluationResult::OverrideAll(const FCameraNodeEvaluationResult& OtherResult)
+void FCameraNodeEvaluationResult::OverrideAll(const FCameraNodeEvaluationResult& OtherResult, bool bIncludePrivateValues)
 {
 	CameraPose.OverrideAll(OtherResult.CameraPose);
-	VariableTable.OverrideAll(OtherResult.VariableTable);
+	VariableTable.OverrideAll(OtherResult.VariableTable, bIncludePrivateValues);
+	ContextDataTable.OverrideAll(OtherResult.ContextDataTable);
 	CameraRigJoints.OverrideAll(OtherResult.CameraRigJoints);
 	PostProcessSettings.OverrideAll(OtherResult.PostProcessSettings);
 	bIsCameraCut = OtherResult.bIsCameraCut;
 	bIsValid = OtherResult.bIsValid;
 }
 
-void FCameraNodeEvaluationResult::LerpAll(const FCameraNodeEvaluationResult& ToResult, float BlendFactor)
+void FCameraNodeEvaluationResult::LerpAll(const FCameraNodeEvaluationResult& ToResult, float BlendFactor, bool bIncludePrivateValues)
 {
 	// Blend all properties.
 	CameraPose.LerpAll(ToResult.CameraPose, BlendFactor);
-	VariableTable.LerpAll(ToResult.VariableTable, BlendFactor);
+	VariableTable.LerpAll(ToResult.VariableTable, BlendFactor, bIncludePrivateValues);
 
 	// Merge/blend the joints.
 	CameraRigJoints.LerpAll(ToResult.CameraRigJoints, BlendFactor);
diff --git a/Engine/Plugins/Cameras/GameplayCameras/Source/GameplayCameras/Private/Core/CameraVariableTable.cpp b/Engine/Plugins/Cameras/GameplayCameras/Source/GameplayCameras/Private/Core/CameraVariableTable.cpp
--- a/Engine/Plugins/Cameras/GameplayCameras/Source/GameplayCameras/Private/Core/CameraVariableTable.cpp
+++ b/Engine/Plugins/Cameras/GameplayCameras/Source/GameplayCameras/Private/Core/CameraVariableTable.cpp
@@ -488,9 +488,10 @@ void FCameraVariableTable::Serialize(FArchive& Ar)
 	}
 }
 
-void FCameraVariableTable::OverrideAll(const FCameraVariableTable& OtherTable)
+void FCameraVariableTable::OverrideAll(const FCameraVariableTable& OtherTable, bool bIncludePrivateValues)
 {
-	const ECameraVariableTableFilter Filter = ECameraVariableTableFilter::InputOutput;
+	const ECameraVariableTableFilter Filter = !bIncludePrivateValues ? 
+		ECameraVariableTableFilter::PublicOnly : ECameraVariableTableFilter::None;
 	InternalOverride(OtherTable, Filter, nullptr, false, nullptr);
 }
 
@@ -508,9 +509,8 @@ void FCameraVariableTable::InternalOverride(const FCameraVariableTable& OtherTab
 {
 	using namespace UE::Cameras::Private;
 
-	const bool bInputs = EnumHasAnyFlags(Filter, ECameraVariableTableFilter::Input);
-	const bool bOutputs = EnumHasAnyFlags(Filter, ECameraVariableTableFilter::Output);
-	const bool bPrivates = EnumHasAnyFlags(Filter, ECameraVariableTableFilter::Private);
+	const bool bPublicOnly = EnumHasAnyFlags(Filter, ECameraVariableTableFilter::PublicOnly);
+	const bool bInputOnly = EnumHasAnyFlags(Filter, ECameraVariableTableFilter::InputOnly);
 	const bool bKnownOnly = EnumHasAnyFlags(Filter, ECameraVariableTableFilter::KnownOnly);
 	const bool bChangedOnly = EnumHasAnyFlags(Filter, ECameraVariableTableFilter::ChangedOnly);
 
@@ -518,11 +518,10 @@ void FCameraVariableTable::InternalOverride(const FCameraVariableTable& OtherTab
 	{
 		// Look for entries in the other table that have been written to, and aren't private.
 		const EEntryFlags OtherFlags = OtherEntry.Flags;
-		const bool bOtherEntryIsInput = EnumHasAnyFlags(OtherFlags, EEntryFlags::Input);
 		if (EnumHasAnyFlags(OtherFlags, EEntryFlags::Written)
+				&& (!bPublicOnly || !EnumHasAnyFlags(OtherFlags, EEntryFlags::Private))
+				&& (!bInputOnly || EnumHasAnyFlags(OtherFlags, EEntryFlags::Input))
 				&& (!bChangedOnly || EnumHasAnyFlags(OtherFlags, EEntryFlags::WrittenThisFrame))
-				&& ((bInputs && bOtherEntryIsInput) || (bOutputs && !bOtherEntryIsInput))
-				&& (bPrivates || !EnumHasAnyFlags(OtherFlags, EEntryFlags::Private))
 				&& IsVariableInMask(OtherEntry.ID, InMask, bInvertMask))
 		{
 			// See if we know this variable.
@@ -589,9 +588,10 @@ void FCameraVariableTable::InternalOverride(const FCameraVariableTable& OtherTab
 	}
 }
 
-void FCameraVariableTable::LerpAll(const FCameraVariableTable& ToTable, float Factor)
+void FCameraVariableTable::LerpAll(const FCameraVariableTable& ToTable, float Factor, bool bIncludePrivateValues)
 {
-	const ECameraVariableTableFilter Filter = ECameraVariableTableFilter::InputOutput;
+	const ECameraVariableTableFilter Filter = !bIncludePrivateValues ? 
+		ECameraVariableTableFilter::PublicOnly : ECameraVariableTableFilter::None;
 	InternalLerp(ToTable, Filter, Factor, nullptr, false, nullptr);
 }
 
@@ -609,9 +609,8 @@ void FCameraVariableTable::InternalLerp(const FCameraVariableTable& ToTable, ECa
 {
 	using namespace UE::Cameras::Private;
 
-	const bool bInputs = EnumHasAnyFlags(Filter, ECameraVariableTableFilter::Input);
-	const bool bOutputs = EnumHasAnyFlags(Filter, ECameraVariableTableFilter::Output);
-	const bool bPrivates = EnumHasAnyFlags(Filter, ECameraVariableTableFilter::Private);
+	const bool bPublicOnly = EnumHasAnyFlags(Filter, ECameraVariableTableFilter::PublicOnly);
+	const bool bInputOnly = EnumHasAnyFlags(Filter, ECameraVariableTableFilter::InputOnly);
 	const bool bKnownOnly = EnumHasAnyFlags(Filter, ECameraVariableTableFilter::KnownOnly);
 	const bool bChangedOnly = EnumHasAnyFlags(Filter, ECameraVariableTableFilter::ChangedOnly);
 
@@ -619,11 +618,10 @@ void FCameraVariableTable::InternalLerp(const FCameraVariableTable& ToTable, ECa
 	{
 		// Look for entries in the other table that have been written to, and aren't private.
 		const EEntryFlags ToFlags = ToEntry.Flags;
-		const bool bToEntryIsInput = EnumHasAnyFlags(ToFlags, EEntryFlags::Input);
 		if (EnumHasAnyFlags(ToFlags, EEntryFlags::Written)
+				&& (!bPublicOnly || !EnumHasAnyFlags(ToFlags, EEntryFlags::Private))
+				&& (!bInputOnly || EnumHasAnyFlags(ToFlags, EEntryFlags::Input))
 				&& (!bChangedOnly || EnumHasAnyFlags(ToFlags, EEntryFlags::WrittenThisFrame))
-				&& ((bInputs && bToEntryIsInput) || (bOutputs && !bToEntryIsInput))
-				&& (bPrivates || !EnumHasAnyFlags(ToFlags, EEntryFlags::Private))
 				&& IsVariableInMask(ToEntry.ID, InMask, bInvertMask))
 		{
 			// See if we know this variable.
diff --git a/Engine/Plugins/Cameras/GameplayCameras/Source/GameplayCameras/Private/Nodes/Blends/OrbitBlendCameraNode.cpp b/Engine/Plugins/Cameras/GameplayCameras/Source/GameplayCameras/Private/Nodes/Blends/OrbitBlendCameraNode.cpp
--- a/Engine/Plugins/Cameras/GameplayCameras/Source/GameplayCameras/Private/Nodes/Blends/OrbitBlendCameraNode.cpp
+++ b/Engine/Plugins/Cameras/GameplayCameras/Source/GameplayCameras/Private/Nodes/Blends/OrbitBlendCameraNode.cpp
@@ -90,7 +90,7 @@ void FOrbitBlendCameraNodeEvaluator::OnBlendParameters(const FCameraNodePreBlend
 	else
 	{
 		const FCameraVariableTable& ChildVariableTable(Params.ChildVariableTable);
-		OutResult.VariableTable.Override(ChildVariableTable, ECameraVariableTableFilter::Input);
+		OutResult.VariableTable.Override(ChildVariableTable, Params.VariableTableFilter);
 
 		OutResult.bIsBlendFinished = true;
 		OutResult.bIsBlendFull = true;
diff --git a/Engine/Plugins/Cameras/GameplayCameras/Source/GameplayCameras/Private/Nodes/Common/CameraRigCameraNode.cpp b/Engine/Plugins/Cameras/GameplayCameras/Source/GameplayCameras/Private/Nodes/Common/CameraRigCameraNode.cpp
--- a/Engine/Plugins/Cameras/GameplayCameras/Source/GameplayCameras/Private/Nodes/Common/CameraRigCameraNode.cpp
+++ b/Engine/Plugins/Cameras/GameplayCameras/Source/GameplayCameras/Private/Nodes/Common/CameraRigCameraNode.cpp
@@ -61,34 +61,18 @@ void FCameraRigCameraNodeEvaluator::OnRun(const FCameraNodeEvaluationParams& Par
 
 void FCameraRigCameraNodeEvaluator::ApplyParameterOverrides(FCameraVariableTable& OutVariableTable, bool bDrivenOnly)
 {
-	if (bApplyParameterOverrides)
-	{
-		const UCameraRigCameraNode* PrefabNode = GetCameraNodeAs<UCameraRigCameraNode>();
+	const UCameraRigCameraNode* PrefabNode = GetCameraNodeAs<UCameraRigCameraNode>();
 
-		FCameraRigParameterOverrideEvaluator OverrideEvaluator(PrefabNode->CameraRigReference);
-		OverrideEvaluator.ApplyParameterOverrides(OutVariableTable, bDrivenOnly);
-	}
+	FCameraRigParameterOverrideEvaluator OverrideEvaluator(PrefabNode->CameraRigReference);
+	OverrideEvaluator.ApplyParameterOverrides(OutVariableTable, bDrivenOnly);
 }
 
 void FCameraRigCameraNodeEvaluator::ApplyParameterOverrides(FCameraVariableTable& OutVariableTable, FCameraContextDataTable& OutContextDataTable, bool bDrivenOnly)
 {
-	if (bApplyParameterOverrides)
-	{
-		const UCameraRigCameraNode* PrefabNode = GetCameraNodeAs<UCameraRigCameraNode>();
-
-		FCameraRigParameterOverrideEvaluator OverrideEvaluator(PrefabNode->CameraRigReference);
-		OverrideEvaluator.ApplyParameterOverrides(OutVariableTable, OutContextDataTable, bDrivenOnly);
-	}
-}
-
-bool FCameraRigCameraNodeEvaluator::IsApplyingParameterOverrides() const
-{
-	return bApplyParameterOverrides;
-}
+	const UCameraRigCameraNode* PrefabNode = GetCameraNodeAs<UCameraRigCameraNode>();
 
-void FCameraRigCameraNodeEvaluator::SetApplyParameterOverrides(bool bShouldApply)
-{
-	bApplyParameterOverrides = bShouldApply;
+	FCameraRigParameterOverrideEvaluator OverrideEvaluator(PrefabNode->CameraRigReference);
+	OverrideEvaluator.ApplyParameterOverrides(OutVariableTable, OutContextDataTable, bDrivenOnly);
 }
 
 }  // namespace UE::Cameras
diff --git a/Engine/Plugins/Cameras/GameplayCameras/Source/GameplayCameras/Public/Core/BlendStackCameraNode.h b/Engine/Plugins/Cameras/GameplayCameras/Source/GameplayCameras/Public/Core/BlendStackCameraNode.h
--- a/Engine/Plugins/Cameras/GameplayCameras/Source/GameplayCameras/Public/Core/BlendStackCameraNode.h
+++ b/Engine/Plugins/Cameras/GameplayCameras/Source/GameplayCameras/Public/Core/BlendStackCameraNode.h
@@ -137,7 +137,8 @@ protected:
 		FCameraRigEntry& NewEntry, 
 		const UCameraRigAsset* CameraRig,
 		TSharedPtr<const FCameraEvaluationContext> EvaluationContext,
-		UBlendStackRootCameraNode* EntryRootNode);
+		UBlendStackRootCameraNode* EntryRootNode,
+		bool bSetActiveResult);
 
 	void FreezeEntry(FCameraRigEntry& Entry);
 
@@ -164,7 +165,6 @@ protected:
 		TSharedPtr<const FCameraEvaluationContext> Context;
 		int32 EntryIndex = INDEX_NONE;
 		bool bIsActiveEntry = false;
-		bool bHasPreBlendedParameters = false;
 	};
 
 	void ResolveEntries(TArray<FResolvedEntry>& OutResolvedEntries);
@@ -192,12 +192,9 @@ protected:
 		bool bIsFirstFrame = false;
 		/** Whether the context's initial result was valid last frame. */
 		bool bWasContextInitialResultValid = false;
-		/** Whether input slots were run (possibly from a preview update). */
-		bool bInputRunThisFrame = false;
-		/** Whether the blend node was run (possibly from a preview update). */
-		bool bBlendRunThisFrame = false;
 		/** Whether to force a camera cut on this entry this frame. */
 		bool bForceCameraCut = false;
+
 		/** Whether this entry is frozen. */
 		bool bIsFrozen = false;
 
@@ -283,6 +280,7 @@ protected:
 
 	// FCameraNodeEvaluator interface.
 	virtual void OnRun(const FCameraNodeEvaluationParams& Params, FCameraNodeEvaluationResult& OutResult) override;
+	virtual void OnSerialize(const FCameraNodeEvaluatorSerializeParams& Params, FArchive& Ar) override;
 
 private:
 
@@ -291,6 +289,7 @@ private:
 	void InternalPreBlendExecute(TArrayView<FResolvedEntry> ResolvedEntries, const FCameraNodeEvaluationParams& Params, FCameraNodeEvaluationResult& OutResult);
 	void InternalUpdate(TArrayView<FResolvedEntry> ResolvedEntries, const FCameraNodeEvaluationParams& Params, FCameraNodeEvaluationResult& OutResult);
 	void InternalPostBlendExecute(TArrayView<FResolvedEntry> ResolvedEntries, const FCameraNodeEvaluationParams& Params, FCameraNodeEvaluationResult& OutResult);
+	void InternalRunFinished(FCameraNodeEvaluationResult& OutResult);
 
 	// Utility functions for finding an appropriate transition.
 	const UCameraRigTransition* FindTransition(const FBlendStackCameraPushParams& Params) const;
@@ -299,11 +298,29 @@ private:
 			const UCameraRigAsset* FromCameraRig, const UCameraAsset* FromCameraAsset, bool bFromFrozen,
 			const UCameraRigAsset* ToCameraRig, const UCameraAsset* ToCameraAsset) const;
 
-	void PushVariantEntry(const FBlendStackCameraPushParams& Params, const UCameraRigTransition* Transition);
 	void PushNewEntry(const FBlendStackCameraPushParams& Params, const UCameraRigTransition* Transition);
+	void PushMergedEntry(const FBlendStackCameraPushParams& Params, const UCameraRigTransition* Transition);
 
 private:
 
+	/** Extra blending-related info for each camera rig entry. */
+	struct FCameraRigEntryExtraInfo
+	{
+		/** Result without pre-blending. */
+		FCameraNodeEvaluationResult UnblendedResult;
+
+		/** Whether input slots were run. */
+		bool bInputRunThisFrame = false;
+		/** Whether the blend node was run. */
+		bool bBlendRunThisFrame = false;
+		/** Whether this camera rig has any parameters to pre-blend. */
+		bool bHasPreBlendedParameters = false;
+		/** Whether this camera rig's pre-blend was full and finished this frame. */
+		bool bIsPreBlendFull = false;
+	};
+
+	TArray<FCameraRigEntryExtraInfo> EntryExtraInfos;
+
 	/** Variable table for pre-blending. */
 	FCameraVariableTable PreBlendVariableTable;
 };
diff --git a/Engine/Plugins/Cameras/GameplayCameras/Source/GameplayCameras/Public/Core/CameraNodeEvaluator.h b/Engine/Plugins/Cameras/GameplayCameras/Source/GameplayCameras/Public/Core/CameraNodeEvaluator.h
--- a/Engine/Plugins/Cameras/GameplayCameras/Source/GameplayCameras/Public/Core/CameraNodeEvaluator.h
+++ b/Engine/Plugins/Cameras/GameplayCameras/Source/GameplayCameras/Public/Core/CameraNodeEvaluator.h
@@ -194,10 +194,10 @@ public:
 	void ResetFrameFlags();
 
 	/** Override this result with the given other result. */
-	void OverrideAll(const FCameraNodeEvaluationResult& OtherResult);
+	void OverrideAll(const FCameraNodeEvaluationResult& OtherResult, bool bIncludePrivateValues = false);
 
 	/** Interpolate this result towards the other given result. */
-	void LerpAll(const FCameraNodeEvaluationResult& ToResult, float BlendFactor);
+	void LerpAll(const FCameraNodeEvaluationResult& ToResult, float BlendFactor, bool bIncludePrivateValues = false);
 
 	/** Serializes this result to the given archive. */
 	void Serialize(FArchive& Ar);
diff --git a/Engine/Plugins/Cameras/GameplayCameras/Source/GameplayCameras/Public/Core/CameraVariableTable.h b/Engine/Plugins/Cameras/GameplayCameras/Source/GameplayCameras/Public/Core/CameraVariableTable.h
--- a/Engine/Plugins/Cameras/GameplayCameras/Source/GameplayCameras/Public/Core/CameraVariableTable.h
+++ b/Engine/Plugins/Cameras/GameplayCameras/Source/GameplayCameras/Public/Core/CameraVariableTable.h
@@ -29,19 +29,14 @@ struct TCameraVariableInterpolation;
 enum class ECameraVariableTableFilter
 {
 	None = 0,
-	/** Include input variables. */
-	Input = 1 << 0,
-	/** Include output variables (i.e. anything not an input). */
-	Output = 1 << 1,
-	/** Include private variables. */
-	Private = 1 << 2,
+	/** Don't include private variables. */
+	PublicOnly = 1 << 0,
+	/** Only include input variables. */
+	InputOnly = 1 << 1,
 	/** Only include data that is common to both tables. */
-	KnownOnly = 1 << 3,
+	KnownOnly = 1 << 2,
 	/** Only include variables that were written this frame. */
-	ChangedOnly = 1 << 4,
-
-	/** Both input and output variables. */
-	InputOutput = Input | Output,
+	ChangedOnly = 1 << 3
 };
 ENUM_CLASS_FLAGS(ECameraVariableTableFilter)
 
@@ -143,11 +138,11 @@ public:
 
 	// Interpolation.
 	
-	void OverrideAll(const FCameraVariableTable& OtherTable);
+	void OverrideAll(const FCameraVariableTable& OtherTable, bool bIncludePrivateValues = false);
 	void Override(const FCameraVariableTable& OtherTable, ECameraVariableTableFilter Filter);
 	void Override(const FCameraVariableTable& OtherTable, ECameraVariableTableFilter Filter, const FCameraVariableTableFlags& InMask, bool bInvertMask, FCameraVariableTableFlags& OutMask);
 
-	void LerpAll(const FCameraVariableTable& ToTable, float Factor);
+	void LerpAll(const FCameraVariableTable& ToTable, float Factor, bool bIncludePrivateValues = false);
 	void Lerp(const FCameraVariableTable& ToTable, ECameraVariableTableFilter Filter, float Factor);
 	void Lerp(const FCameraVariableTable& ToTable, ECameraVariableTableFilter Filter, float Factor, const FCameraVariableTableFlags& InMask, bool bInvertMask, FCameraVariableTableFlags& OutMask);
 
diff --git a/Engine/Plugins/Cameras/GameplayCameras/Source/GameplayCameras/Public/Nodes/Common/CameraRigCameraNode.h b/Engine/Plugins/Cameras/GameplayCameras/Source/GameplayCameras/Public/Nodes/Common/CameraRigCameraNode.h
--- a/Engine/Plugins/Cameras/GameplayCameras/Source/GameplayCameras/Public/Nodes/Common/CameraRigCameraNode.h
+++ b/Engine/Plugins/Cameras/GameplayCameras/Source/GameplayCameras/Public/Nodes/Common/CameraRigCameraNode.h
@@ -52,9 +52,7 @@ public:
 
 	FCameraRigCameraNodeEvaluator();
 
-	// Internal API.
-	bool IsApplyingParameterOverrides() const;
-	void SetApplyParameterOverrides(bool bShouldApply);
+	FCameraNodeEvaluator* GetCameraRigRootEvaluator() const { return CameraRigRootEvaluator; }
 
 protected:
 
@@ -73,7 +71,6 @@ private:
 private:
 
 	FCameraNodeEvaluator* CameraRigRootEvaluator = nullptr;
-	bool bApplyParameterOverrides = true;
 };
 
 }  // namesapce UE::Cameras
