Cameras: use the correct container pointer when applying default rig parameter values

diff --git a/Engine/Plugins/Cameras/GameplayCameras/Source/GameplayCameras/Private/Core/CameraRigParameterDefinition.cpp b/Engine/Plugins/Cameras/GameplayCameras/Source/GameplayCameras/Private/Core/CameraRigParameterDefinition.cpp
--- a/Engine/Plugins/Cameras/GameplayCameras/Source/GameplayCameras/Private/Core/CameraRigParameterDefinition.cpp
+++ b/Engine/Plugins/Cameras/GameplayCameras/Source/GameplayCameras/Private/Core/CameraRigParameterDefinition.cpp
@@ -233,7 +233,7 @@ void FCameraRigParameterBuilder::SetDefaultParameterValues(const UCameraRigAsset
 			FProperty* Property = TargetClass->FindPropertyByName(DataParameter->TargetPropertyName);
 			if (Property)
 			{
-				RawSourceValuePtr = (uint8*)Property->ContainerPtrToValuePtr<void>(TargetClass);
+				RawSourceValuePtr = (uint8*)Property->ContainerPtrToValuePtr<void>(DataParameter->Target);
 			}
 		}
 
