Cameras: use proper structure name and remove unnecessary header includes

diff --git a/Engine/Plugins/Cameras/GameplayCameras/Source/GameplayCamerasEditor/Private/Sequencer/CameraFramingZoneTrackEditor.cpp b/Engine/Plugins/Cameras/GameplayCameras/Source/GameplayCamerasEditor/Private/Sequencer/CameraFramingZoneTrackEditor.cpp
--- a/Engine/Plugins/Cameras/GameplayCameras/Source/GameplayCamerasEditor/Private/Sequencer/CameraFramingZoneTrackEditor.cpp
+++ b/Engine/Plugins/Cameras/GameplayCameras/Source/GameplayCamerasEditor/Private/Sequencer/CameraFramingZoneTrackEditor.cpp
@@ -4,8 +4,6 @@
 
 #include "Channels/MovieSceneDoubleChannel.h"
 #include "KeyPropertyParams.h"
-#include "Misc/Optional.h"
-#include "Nodes/Framing/CameraFramingZone.h"
 #include "PropertyPath.h"
 #include "UObject/NameTypes.h"
 #include "UObject/UnrealNames.h"
diff --git a/Engine/Plugins/Cameras/GameplayCameras/Source/GameplayCamerasEditor/Private/Sequencer/CameraFramingZoneTrackEditor.h b/Engine/Plugins/Cameras/GameplayCameras/Source/GameplayCamerasEditor/Private/Sequencer/CameraFramingZoneTrackEditor.h
--- a/Engine/Plugins/Cameras/GameplayCameras/Source/GameplayCamerasEditor/Private/Sequencer/CameraFramingZoneTrackEditor.h
+++ b/Engine/Plugins/Cameras/GameplayCameras/Source/GameplayCamerasEditor/Private/Sequencer/CameraFramingZoneTrackEditor.h
@@ -8,12 +8,10 @@
 #include "Containers/Array.h"
 #include "Containers/ArrayView.h"
 #include "Containers/ContainerAllocationPolicies.h"
-#include "Delegates/Delegate.h"
-#include "HAL/PlatformCrt.h"
 #include "KeyframeTrackEditor.h"
 #include "MovieScene/MovieSceneCameraFramingZoneTrack.h"
+#include "Nodes/Framing/CameraFramingZone.h"
 #include "Templates/SharedPointer.h"
-#include "Templates/UnrealTemplate.h"
 
 class FName;
 class FPropertyChangedParams;
@@ -35,7 +33,7 @@ public:
 
 	static TArray<FAnimatedPropertyKey, TInlineAllocator<1>> GetAnimatedPropertyTypes()
 	{
-		return TArray<FAnimatedPropertyKey, TInlineAllocator<1>>({ FAnimatedPropertyKey::FromStructType("CameraFramingZone") });
+		return TArray<FAnimatedPropertyKey, TInlineAllocator<1>>({ FAnimatedPropertyKey::FromStructType(FCameraFramingZone::StaticStruct()) });
 	}
 
 	static TSharedRef<ISequencerTrackEditor> CreateTrackEditor(TSharedRef<ISequencer> OwningSequencer);
