Cameras: change how camera parameters are animated in Sequencer

Instead of using custom binding handlers, use a standard property path that points to the parameter overrides on the camera component. This makes it less confusing for the user since the values visible in the Details View match the values in the Sequencer.

- Add a flag for camera parameters to be "animated". This is because while Sequencer has control over the parameter, it should be set every frame on the camera rig.
- Use decorator objects on MovieScene sections. They add extra ECS components that get picked up by a new custom system that will set camera parameters as animated (see previous point).
- Make the camera parameter track editor create property tracks in the same way as other property track editors. The property paths are all of the form `CameraReference.Parameters.Value.<Blah...>`
- Since Sequencer can (as of recently) animate properties inside property bags, this "just works", and allows animating sub-properties of camera parameters, especially data parameters.
- Fix possible null pointer dereferencing in the parameter override evaluators.

diff --git a/Engine/Plugins/Cameras/GameplayCameras/Source/GameplayCameras/Private/Core/CameraAssetReference.cpp b/Engine/Plugins/Cameras/GameplayCameras/Source/GameplayCameras/Private/Core/CameraAssetReference.cpp
--- a/Engine/Plugins/Cameras/GameplayCameras/Source/GameplayCameras/Private/Core/CameraAssetReference.cpp
+++ b/Engine/Plugins/Cameras/GameplayCameras/Source/GameplayCameras/Private/Core/CameraAssetReference.cpp
@@ -41,6 +41,23 @@ void FCameraAssetReference::SetParameterOverridden(const FGuid PropertyID, bool
 	}
 }
 
+bool FCameraAssetReference::IsParameterAnimated(const FGuid PropertyID) const
+{
+	return ParameterAnimatedGuids.Contains(PropertyID);
+}
+
+void FCameraAssetReference::SetParameterAnimated(const FGuid PropertyID, bool bIsAnimated)
+{
+	if (bIsAnimated)
+	{
+		ParameterAnimatedGuids.AddUnique(PropertyID);
+	}
+	else
+	{
+		ParameterAnimatedGuids.Remove(PropertyID);
+	}
+}
+
 bool FCameraAssetReference::NeedsRebuildParameters() const
 {
 	if ((!CameraAsset && Parameters.IsValid()) || (CameraAsset && !Parameters.IsValid()))
diff --git a/Engine/Plugins/Cameras/GameplayCameras/Source/GameplayCameras/Private/Core/CameraRigAssetReference.cpp b/Engine/Plugins/Cameras/GameplayCameras/Source/GameplayCameras/Private/Core/CameraRigAssetReference.cpp
--- a/Engine/Plugins/Cameras/GameplayCameras/Source/GameplayCameras/Private/Core/CameraRigAssetReference.cpp
+++ b/Engine/Plugins/Cameras/GameplayCameras/Source/GameplayCameras/Private/Core/CameraRigAssetReference.cpp
@@ -81,6 +81,21 @@ void FCameraRigAssetReference::SetParameterOverridden(const FGuid& PropertyID, b
 	MetaData.bIsOverridden = bIsOverridden;
 }
 
+bool FCameraRigAssetReference::IsParameterAnimated(const FGuid& PropertyID) const
+{
+	if (const FCameraRigAssetReferenceParameterMetaData* MetaData = FindMetaData(PropertyID))
+	{
+		return MetaData->bIsAnimated;
+	}
+	return false;
+}
+
+void FCameraRigAssetReference::SetParameterAnimated(const FGuid& PropertyID, bool bIsAnimated)
+{
+	FCameraRigAssetReferenceParameterMetaData& MetaData = FindOrAddMetaData(PropertyID);
+	MetaData.bIsAnimated = bIsAnimated;
+}
+
 bool FCameraRigAssetReference::NeedsRebuildParameters() const
 {
 	if ((!CameraRig && Parameters.IsValid()) || (CameraRig && !Parameters.IsValid()))
diff --git a/Engine/Plugins/Cameras/GameplayCameras/Source/GameplayCameras/Private/GameFramework/GameplayCameraComponent.cpp b/Engine/Plugins/Cameras/GameplayCameras/Source/GameplayCameras/Private/GameFramework/GameplayCameraComponent.cpp
--- a/Engine/Plugins/Cameras/GameplayCameras/Source/GameplayCameras/Private/GameFramework/GameplayCameraComponent.cpp
+++ b/Engine/Plugins/Cameras/GameplayCameras/Source/GameplayCameras/Private/GameFramework/GameplayCameraComponent.cpp
@@ -362,7 +362,7 @@ void UGameplayCameraComponent::TickComponent(float DeltaTime, ELevelTick TickTyp
 	}
 }
 
-void UGameplayCameraComponent::UpdateCameraEvaluationContext(bool bApplyParameterOverrides)
+void UGameplayCameraComponent::UpdateCameraEvaluationContext(bool bForceApplyParameterOverrides)
 {
 	using namespace UE::Cameras;
 
@@ -379,10 +379,8 @@ void UGameplayCameraComponent::UpdateCameraEvaluationContext(bool bApplyParamete
 		bIsCameraCutNextFrame = false;
 	}
 
-	if (bApplyParameterOverrides)
-	{
-		CameraReference.ApplyParameterOverrides(InitialResult, false);
-	}
+	const bool bApplyDrivenParametersOnly = !bForceApplyParameterOverrides;
+	CameraReference.ApplyParameterOverrides(InitialResult, bApplyDrivenParametersOnly);
 }
 
 void UGameplayCameraComponent::UpdateOutputCameraComponent()
diff --git a/Engine/Plugins/Cameras/GameplayCameras/Source/GameplayCameras/Private/GameplayCamerasModule.cpp b/Engine/Plugins/Cameras/GameplayCameras/Source/GameplayCameras/Private/GameplayCamerasModule.cpp
--- a/Engine/Plugins/Cameras/GameplayCameras/Source/GameplayCameras/Private/GameplayCamerasModule.cpp
+++ b/Engine/Plugins/Cameras/GameplayCameras/Source/GameplayCameras/Private/GameplayCamerasModule.cpp
@@ -6,7 +6,8 @@
 #include "GameplayCameras.h"
 #include "Logging/MessageLog.h"
 #include "Modules/ModuleManager.h"
-#include "MovieScene/GameplayCamerasTrackInstancePropertyBindings.h"
+#include "MovieScene/MovieSceneCameraParameterInstantiator.h"
+#include "MovieSceneSection.h"
 #include "Nodes/Framing/CameraFramingZone.h"
 #include "ShowFlags.h"
 #include "UObject/UObjectBase.h"
@@ -35,7 +36,7 @@ public:
 	virtual void StartupModule() override
 	{
 		RegisterBuiltInBlendableStructs();
-		RegisterTrackInstanceCustomPropertyBindings();
+		RegisterMovieSceneIntegration();
 		
 #if UE_GAMEPLAY_CAMERAS_DEBUG
 		UE::Cameras::FCameraDebugColors::RegisterBuiltinColorSchemes();
@@ -45,7 +46,7 @@ public:
 	virtual void ShutdownModule() override
 	{
 		UnregisterBuiltInBlendableStructs();
-		UnregisterTrackInstanceCustomPropertyBindings();
+		UnregisterMovieSceneIntegration();
 	}
 
 public:
@@ -117,21 +118,30 @@ private:
 			UnregisterBlendableStruct(FCameraFramingZone::StaticStruct());
 		}
 	}
-	
-	void RegisterTrackInstanceCustomPropertyBindings()
+
+	void RegisterMovieSceneIntegration()
 	{
-		UE::Cameras::FGameplayCamerasTrackInstancePropertyBindings::Register();
+#if WITH_EDITORONLY_DATA
+		OnMovieSceneSectionAddedToTrackHandle = UMovieSceneSection::GetOnSectionAddedToTrack().AddStatic(
+				&UMovieSceneCameraParameterInstantiator::OnMovieSceneSectionAddedToTrack);
+#endif
 	}
 
-	void UnregisterTrackInstanceCustomPropertyBindings()
+	void UnregisterMovieSceneIntegration()
 	{
-		UE::Cameras::FGameplayCamerasTrackInstancePropertyBindings::Unregister();
+#if WITH_EDITORONLY_DATA
+		UMovieSceneSection::GetOnSectionAddedToTrack().Remove(OnMovieSceneSectionAddedToTrackHandle);
+#endif
 	}
 
 private:
 
 	TArray<UE::Cameras::FBlendableStructInfo> BlendableStructs;
 
+#if WITH_EDITORONLY_DATA
+	FDelegateHandle OnMovieSceneSectionAddedToTrackHandle;
+#endif
+
 #if WITH_EDITOR
 	TSharedPtr<IGameplayCamerasLiveEditManager> LiveEditManager;
 #endif
diff --git a/Engine/Plugins/Cameras/GameplayCameras/Source/GameplayCameras/Private/Helpers/CameraAssetParameterOverrideEvaluator.cpp b/Engine/Plugins/Cameras/GameplayCameras/Source/GameplayCameras/Private/Helpers/CameraAssetParameterOverrideEvaluator.cpp
--- a/Engine/Plugins/Cameras/GameplayCameras/Source/GameplayCameras/Private/Helpers/CameraAssetParameterOverrideEvaluator.cpp
+++ b/Engine/Plugins/Cameras/GameplayCameras/Source/GameplayCameras/Private/Helpers/CameraAssetParameterOverrideEvaluator.cpp
@@ -27,8 +27,6 @@ void FCameraAssetParameterOverrideEvaluator::ApplyParameterOverrides(FCameraVari
 
 void FCameraAssetParameterOverrideEvaluator::ApplyParameterOverrides(FCameraVariableTable* OutVariableTable, FCameraContextDataTable* OutContextDataTable, bool bDrivenOverridesOnly)
 {
-	using namespace Internal;
-
 	check(OutVariableTable);
 
 	const UCameraAsset* CameraAsset = CameraReference.GetCameraAsset();
@@ -43,12 +41,19 @@ void FCameraAssetParameterOverrideEvaluator::ApplyParameterOverrides(FCameraVari
 	const TArray<TObjectPtr<const UCameraRigAsset>>& ParameterOwners = CameraAsset->ParameterOwners;
 	ensure(ParameterDefinitions.Num() == ParameterOwners.Num());
 
-	FCameraParameterOverrideHelper Helper(*OutVariableTable, *OutContextDataTable);
+	FCameraParameterOverrideHelper Helper(OutVariableTable, OutContextDataTable);
 
 	for (int32 Index = 0, MaxIndex = FMath::Min(ParameterDefinitions.Num(), ParameterOwners.Num()); Index < MaxIndex; ++Index)
 	{
 		const FCameraRigParameterDefinition& Definition(ParameterDefinitions[Index]);
-		if (!CameraReference.IsParameterOverridden(Definition.ParameterGuid))
+
+		if (!OutContextDataTable && Definition.ParameterType == ECameraRigInterfaceParameterType::Data)
+		{
+			continue;
+		}
+
+		const bool bIsAnimated = CameraReference.IsParameterAnimated(Definition.ParameterGuid);
+		if (!CameraReference.IsParameterOverridden(Definition.ParameterGuid) && !bIsAnimated)
 		{
 			continue;
 		}
@@ -65,7 +70,8 @@ void FCameraAssetParameterOverrideEvaluator::ApplyParameterOverrides(FCameraVari
 			continue;
 		}
 
-		Helper.ApplyParameterOverride(CameraRig, Definition, CameraParameters, *PropertyDesc, bDrivenOverridesOnly);
+		const bool bThisDrivenOnly = bDrivenOverridesOnly && !bIsAnimated;
+		Helper.ApplyParameterOverride(CameraRig, Definition, CameraParameters, *PropertyDesc, bThisDrivenOnly);
 	}
 }
 
diff --git a/Engine/Plugins/Cameras/GameplayCameras/Source/GameplayCameras/Private/Helpers/CameraParameterOverrideHelper.cpp b/Engine/Plugins/Cameras/GameplayCameras/Source/GameplayCameras/Private/Helpers/CameraParameterOverrideHelper.cpp
--- a/Engine/Plugins/Cameras/GameplayCameras/Source/GameplayCameras/Private/Helpers/CameraParameterOverrideHelper.cpp
+++ b/Engine/Plugins/Cameras/GameplayCameras/Source/GameplayCameras/Private/Helpers/CameraParameterOverrideHelper.cpp
@@ -217,7 +217,7 @@ void ApplyDataParameterOverride(
 
 }  // namespace Internal
 
-FCameraParameterOverrideHelper::FCameraParameterOverrideHelper(FCameraVariableTable& OutVariableTable, FCameraContextDataTable& OutContextDataTable)
+FCameraParameterOverrideHelper::FCameraParameterOverrideHelper(FCameraVariableTable* OutVariableTable, FCameraContextDataTable* OutContextDataTable)
 	: VariableTable(OutVariableTable)
 	, ContextDataTable(OutContextDataTable)
 {
@@ -232,24 +232,32 @@ void FCameraParameterOverrideHelper::ApplyParameterOverride(
 {
 	using namespace Internal;
 
+	ensure(ParameterDefinition.ParameterGuid == PropertyBagPropertyDesc.ID);
+
 	switch (ParameterDefinition.ParameterType)
 	{
 		case ECameraRigInterfaceParameterType::Blendable:
 			{
-				ApplyBlendableParameterOverride(
-						CameraRig, 
-						ParameterDefinition,
-						PropertyBag, PropertyBagPropertyDesc, 
-						VariableTable, bDrivenOverridesOnly);
+				if (ensure(VariableTable))
+				{
+					ApplyBlendableParameterOverride(
+							CameraRig, 
+							ParameterDefinition,
+							PropertyBag, PropertyBagPropertyDesc, 
+							*VariableTable, bDrivenOverridesOnly);
+				}
 			}
 			break;
 		case ECameraRigInterfaceParameterType::Data:
 			{
-				ApplyDataParameterOverride(
-						CameraRig, 
-						ParameterDefinition,
-						PropertyBag, PropertyBagPropertyDesc, 
-						ContextDataTable);
+				if (ensure(ContextDataTable))
+				{
+					ApplyDataParameterOverride(
+							CameraRig, 
+							ParameterDefinition,
+							PropertyBag, PropertyBagPropertyDesc, 
+							*ContextDataTable);
+				}
 			}
 			break;
 	}
diff --git a/Engine/Plugins/Cameras/GameplayCameras/Source/GameplayCameras/Private/Helpers/CameraParameterOverrideHelper.h b/Engine/Plugins/Cameras/GameplayCameras/Source/GameplayCameras/Private/Helpers/CameraParameterOverrideHelper.h
--- a/Engine/Plugins/Cameras/GameplayCameras/Source/GameplayCameras/Private/Helpers/CameraParameterOverrideHelper.h
+++ b/Engine/Plugins/Cameras/GameplayCameras/Source/GameplayCameras/Private/Helpers/CameraParameterOverrideHelper.h
@@ -19,7 +19,7 @@ class FCameraVariableTable;
  */
 struct FCameraParameterOverrideHelper
 {
-	FCameraParameterOverrideHelper(FCameraVariableTable& OutVariableTable, FCameraContextDataTable& OutContextDataTable);
+	FCameraParameterOverrideHelper(FCameraVariableTable* OutVariableTable, FCameraContextDataTable* OutContextDataTable);
 
 	void ApplyParameterOverride(
 			const UCameraRigAsset* CameraRig,
@@ -32,8 +32,8 @@ struct FCameraParameterOverrideHelper
 
 private:
 
-	FCameraVariableTable& VariableTable;
-	FCameraContextDataTable& ContextDataTable;
+	FCameraVariableTable* VariableTable;
+	FCameraContextDataTable* ContextDataTable;
 };
 
 }  // namespace UE::Cameras
diff --git a/Engine/Plugins/Cameras/GameplayCameras/Source/GameplayCameras/Private/Helpers/CameraRigParameterOverrideEvaluator.cpp b/Engine/Plugins/Cameras/GameplayCameras/Source/GameplayCameras/Private/Helpers/CameraRigParameterOverrideEvaluator.cpp
--- a/Engine/Plugins/Cameras/GameplayCameras/Source/GameplayCameras/Private/Helpers/CameraRigParameterOverrideEvaluator.cpp
+++ b/Engine/Plugins/Cameras/GameplayCameras/Source/GameplayCameras/Private/Helpers/CameraRigParameterOverrideEvaluator.cpp
@@ -40,11 +40,17 @@ void FCameraRigParameterOverrideEvaluator::ApplyParameterOverrides(FCameraVariab
 
 	TConstArrayView<FCameraRigParameterDefinition> ParameterDefinitions = CameraRig->GetParameterDefinitions();
 
-	FCameraParameterOverrideHelper Helper(*OutVariableTable, *OutContextDataTable);
+	FCameraParameterOverrideHelper Helper(OutVariableTable, OutContextDataTable);
 
 	for (const FCameraRigParameterDefinition& Definition : ParameterDefinitions)
 	{
-		if (!CameraRigReference.IsParameterOverridden(Definition.ParameterGuid))
+		if (!OutContextDataTable && Definition.ParameterType == ECameraRigInterfaceParameterType::Data)
+		{
+			continue;
+		}
+
+		const bool bIsAnimated = CameraRigReference.IsParameterAnimated(Definition.ParameterGuid);
+		if (!CameraRigReference.IsParameterOverridden(Definition.ParameterGuid) && !bIsAnimated)
 		{
 			continue;
 		}
@@ -55,7 +61,8 @@ void FCameraRigParameterOverrideEvaluator::ApplyParameterOverrides(FCameraVariab
 			continue;
 		}
 
-		Helper.ApplyParameterOverride(CameraRig, Definition, CameraRigParameters, *PropertyDesc, bDrivenOverridesOnly);
+		const bool bThisDrivenOnly = bDrivenOverridesOnly && !bIsAnimated;
+		Helper.ApplyParameterOverride(CameraRig, Definition, CameraRigParameters, *PropertyDesc, bThisDrivenOnly);
 	}
 }
 
diff --git a/Engine/Plugins/Cameras/GameplayCameras/Source/GameplayCameras/Private/MovieScene/GameplayCamerasTrackInstancePropertyBindings.cpp b/Engine/Plugins/Cameras/GameplayCameras/Source/GameplayCameras/Private/MovieScene/GameplayCamerasTrackInstancePropertyBindings.cpp
deleted file mode 100644
diff --git a/Engine/Plugins/Cameras/GameplayCameras/Source/GameplayCameras/Private/MovieScene/GameplayCamerasTrackInstancePropertyBindings.h b/Engine/Plugins/Cameras/GameplayCameras/Source/GameplayCameras/Private/MovieScene/GameplayCamerasTrackInstancePropertyBindings.h
deleted file mode 100644
diff --git a/Engine/Plugins/Cameras/GameplayCameras/Source/GameplayCameras/Private/MovieScene/MovieSceneCameraParameterInstantiator.cpp b/Engine/Plugins/Cameras/GameplayCameras/Source/GameplayCameras/Private/MovieScene/MovieSceneCameraParameterInstantiator.cpp
new file mode 100644
--- /dev/null
+++ b/Engine/Plugins/Cameras/GameplayCameras/Source/GameplayCameras/Private/MovieScene/MovieSceneCameraParameterInstantiator.cpp
@@ -0,0 +1,211 @@
+// Copyright Epic Games, Inc. All Rights Reserved.
+
+#include "MovieScene/MovieSceneCameraParameterInstantiator.h"
+
+#include "Core/CameraAsset.h"
+#include "Core/CameraAssetReference.h"
+#include "Core/CameraRigParameterDefinition.h"
+#include "EntitySystem/BuiltInComponentTypes.h"
+#include "EntitySystem/MovieSceneEntityManager.h"
+#include "EntitySystem/MovieSceneEntitySystemTask.h"
+#include "EntitySystem/MovieScenePropertyBinding.h"
+#include "EntitySystem/MovieSceneSequenceInstanceHandle.h"
+#include "Evaluation/PreAnimatedState/MovieScenePreAnimatedStateStorage.h"
+#include "Evaluation/PreAnimatedState/MovieScenePreAnimatedStorageID.inl"
+#include "GameFramework/GameplayCameraComponent.h"
+#include "MovieScene/MovieSceneGameplayCamerasComponentTypes.h"
+#include "Tracks/MovieScenePropertyTrack.h"
+
+#include UE_INLINE_GENERATED_CPP_BY_NAME(MovieSceneCameraParameterInstantiator)
+
+namespace UE::Cameras
+{
+
+struct FPreAnimatedCameraParameterStateTraits : UE::MovieScene::FPreAnimatedStateTraits
+{
+	using KeyType = TTuple<FObjectKey, FGuid>;
+	using StorageType = bool;
+
+	bool CachePreAnimatedValue(UGameplayCameraComponent* CameraComponent, const FGuid& ParameterGuid)
+	{
+		if (CameraComponent && ParameterGuid.IsValid())
+		{
+			return CameraComponent->CameraReference.IsParameterAnimated(ParameterGuid);
+		}
+		return false;
+	}
+
+	void RestorePreAnimatedValue(const KeyType& InKey, bool bWasAnimated, const UE::MovieScene::FRestoreStateParams& Params)
+	{
+		UGameplayCameraComponent* CameraComponent = Cast<UGameplayCameraComponent>(InKey.Get<0>().ResolveObjectPtr());
+		const FGuid& ParameterGuid = InKey.Get<1>();
+		if (CameraComponent && ParameterGuid.IsValid())
+		{
+			CameraComponent->CameraReference.SetParameterAnimated(ParameterGuid, bWasAnimated);
+		}
+	}
+};
+
+struct FPreAnimatedCameraParameterStorage : UE::MovieScene::TPreAnimatedStateStorage<FPreAnimatedCameraParameterStateTraits>
+{
+	static UE::MovieScene::TAutoRegisterPreAnimatedStorageID<FPreAnimatedCameraParameterStorage> StorageID;
+};
+
+UE::MovieScene::TAutoRegisterPreAnimatedStorageID<FPreAnimatedCameraParameterStorage> FPreAnimatedCameraParameterStorage::StorageID;
+
+struct FSetupCameraParameterOverrideTask
+{
+	TSharedPtr<FPreAnimatedCameraParameterStorage> PreAnimatedStorage;
+
+	FSetupCameraParameterOverrideTask(TSharedPtr<FPreAnimatedCameraParameterStorage> InPreAnimatedStorage)
+		: PreAnimatedStorage(InPreAnimatedStorage)
+	{}
+
+	void ForEachEntity(
+			UE::MovieScene::FMovieSceneEntityID EntityID, 
+			UE::MovieScene::FRootInstanceHandle RootInstanceHandle,
+			UObject* BoundObject, 
+			const FMovieScenePropertyBinding& PropertyBinding, 
+			FGuid& CameraParameterOverrideID)
+	{
+		using namespace UE::MovieScene;
+
+		UGameplayCameraComponent* CameraComponent = Cast<UGameplayCameraComponent>(BoundObject);
+		if (!ensure(CameraComponent))
+		{
+			return;
+		}
+
+		const FGuid ParameterGuid = GetParameterGuid(CameraComponent->CameraReference, PropertyBinding);
+		if (!ParameterGuid.IsValid())
+		{
+			return;
+		}
+
+		const bool bWasAnimated = CameraComponent->CameraReference.IsParameterAnimated(ParameterGuid);
+
+		using PreAnimatedKeyType = TTuple<FObjectKey, FGuid>;
+
+		PreAnimatedStorage->BeginTrackingEntity(EntityID, true, RootInstanceHandle, CameraComponent, ParameterGuid);
+		PreAnimatedStorage->CachePreAnimatedValue(
+				PreAnimatedKeyType(CameraComponent, ParameterGuid),
+				[bWasAnimated](const PreAnimatedKeyType& InKey) { return bWasAnimated; });
+
+		CameraParameterOverrideID = ParameterGuid;
+		CameraComponent->CameraReference.SetParameterAnimated(ParameterGuid, true);
+	}
+
+	FGuid GetParameterGuid(const FCameraAssetReference& CameraReference, const FMovieScenePropertyBinding& PropertyBinding)
+	{
+		const UCameraAsset* CameraAsset = CameraReference.GetCameraAsset();
+		if (!ensure(CameraAsset))
+		{
+			return FGuid();
+		}
+
+		// This isn't ideal but we know all camera parameters are bound to the "Parameters" property bag
+		// of the camera component, so use that to figure out the parameter name and find is ID.
+		const FString PropertyPath = PropertyBinding.PropertyPath.ToString();
+		if (!ensure(PropertyPath.StartsWith(TEXT("CameraReference.Parameters.Value"))))
+		{
+			return FGuid();
+		}
+		
+		TArray<FString> PropertyNames;
+		PropertyPath.ParseIntoArray(PropertyNames, TEXT("."), true);
+		if (!ensure(PropertyNames.Num() >= 4))
+		{
+			return FGuid();
+		}
+
+		const FString& ParameterName = PropertyNames[3];
+		for (const FCameraRigParameterDefinition& ParameterDefintion : CameraAsset->GetParameterDefinitions())
+		{
+			if (ParameterDefintion.ParameterName == ParameterName)
+			{
+				return ParameterDefintion.ParameterGuid;
+			}
+		}
+
+		return FGuid();
+	}
+};
+
+}  // namespace UE::Cameras
+
+void UMovieSceneCameraParameterDecoration::ExtendEntityImpl(UMovieSceneEntitySystemLinker* EntityLinker, const FEntityImportParams& Params, FImportedEntity* OutImportedEntity)
+{
+	using namespace UE::Cameras;
+	using namespace UE::MovieScene;
+
+	FMovieSceneGameplayCamerasComponentTypes* CameraComponents = FMovieSceneGameplayCamerasComponentTypes::Get();
+
+	OutImportedEntity->AddBuilder(
+			FEntityBuilder()
+			.Add(CameraComponents->CameraParameterOverrideID, FGuid()));
+}
+
+UMovieSceneCameraParameterInstantiator::UMovieSceneCameraParameterInstantiator(const FObjectInitializer& ObjInit)
+	: Super(ObjInit)
+{
+	using namespace UE::Cameras;
+	using namespace UE::MovieScene;
+
+	FBuiltInComponentTypes* BuiltInComponents = FBuiltInComponentTypes::Get();
+	FMovieSceneGameplayCamerasComponentTypes* CameraComponents = FMovieSceneGameplayCamerasComponentTypes::Get();
+	RelevantComponent = CameraComponents->CameraParameterOverrideID;
+
+	if (HasAnyFlags(RF_ClassDefaultObject))
+	{
+		DefineComponentConsumer(GetClass(), BuiltInComponents->BoundObject);
+	}
+}
+
+#if WITH_EDITORONLY_DATA
+
+void UMovieSceneCameraParameterInstantiator::OnMovieSceneSectionAddedToTrack(UMovieSceneTrack* Track, UMovieSceneSection* NewSection)
+{
+	UMovieScenePropertyTrack* PropertyTrack = Cast<UMovieScenePropertyTrack>(Track);
+	if (!PropertyTrack)
+	{
+		return;
+	}
+
+	// This isn't ideal but we know all camera parameters are bound to the "Parameters" property bag
+	// of the camera component, so use that to determine if this is a camera parameter track.
+	const FMovieScenePropertyBinding& PropertyBinding = PropertyTrack->GetPropertyBinding();
+	const FString& PropertyPath = PropertyBinding.PropertyPath.ToString();
+	if (PropertyPath.StartsWith(TEXT("CameraReference.Parameters.Value")))
+	{
+		UMovieSceneCameraParameterDecoration* Decoration = NewSection->GetOrCreateDecoration<UMovieSceneCameraParameterDecoration>();
+	}
+}
+
+#endif
+
+void UMovieSceneCameraParameterInstantiator::OnLink()
+{
+	using namespace UE::Cameras;
+
+	PreAnimatedStorage = Linker->PreAnimatedState.GetOrCreateStorage<FPreAnimatedCameraParameterStorage>();
+}
+
+void UMovieSceneCameraParameterInstantiator::OnRun(FSystemTaskPrerequisites& InPrerequisites, FSystemSubsequentTasks& Subsequents)
+{
+	using namespace UE::Cameras;
+	using namespace UE::MovieScene;
+
+	FBuiltInComponentTypes* BuiltInComponents = FBuiltInComponentTypes::Get();
+	FMovieSceneGameplayCamerasComponentTypes* CameraComponents = FMovieSceneGameplayCamerasComponentTypes::Get();
+
+	FEntityTaskBuilder()
+	.ReadEntityIDs()
+	.Read(BuiltInComponents->RootInstanceHandle)
+	.Read(BuiltInComponents->BoundObject)
+	.Read(BuiltInComponents->PropertyBinding)
+	.Write(CameraComponents->CameraParameterOverrideID)
+	.FilterAny({ BuiltInComponents->Tags.NeedsLink })
+	.FilterNone({ BuiltInComponents->Tags.NeedsUnlink, BuiltInComponents->Tags.Ignored })
+	.RunInline_PerEntity(&Linker->EntityManager, FSetupCameraParameterOverrideTask(PreAnimatedStorage));
+}
+
diff --git a/Engine/Plugins/Cameras/GameplayCameras/Source/GameplayCameras/Private/MovieScene/MovieSceneCameraParameterInstantiator.h b/Engine/Plugins/Cameras/GameplayCameras/Source/GameplayCameras/Private/MovieScene/MovieSceneCameraParameterInstantiator.h
new file mode 100644
--- /dev/null
+++ b/Engine/Plugins/Cameras/GameplayCameras/Source/GameplayCameras/Private/MovieScene/MovieSceneCameraParameterInstantiator.h
@@ -0,0 +1,62 @@
+// Copyright Epic Games, Inc. All Rights Reserved.
+
+#pragma once
+
+#include "EntitySystem/IMovieSceneEntityDecorator.h"
+#include "EntitySystem/MovieSceneEntityInstantiatorSystem.h"
+#include "UObject/ObjectMacros.h"
+#include "UObject/UObjectGlobals.h"
+
+#include "MovieSceneCameraParameterInstantiator.generated.h"
+
+class UMovieSceneSection;
+class UMovieSceneTrack;
+
+namespace UE::Cameras
+{
+	struct FPreAnimatedCameraParameterStorage;
+}
+
+/**
+ * Decorator for camera parameter sections. The decorator extends the manufactured ECS entities
+ * so that we flag the camera parameter as being "animated" and make the camera system apply
+ * its value every frame until Sequencer lets go of it.
+ */
+UCLASS(MinimalAPI)
+class UMovieSceneCameraParameterDecoration : public UObject, public IMovieSceneEntityDecorator
+{
+	GENERATED_BODY()
+
+protected:
+
+	virtual void ExtendEntityImpl(UMovieSceneEntitySystemLinker* EntityLinker, const FEntityImportParams& Params, FImportedEntity* OutImportedEntity) override;
+};
+
+/**
+ * Sequencer system for flagging camera parameters as "animated" and make the camera system
+ * apply their values every frame.
+ */
+UCLASS(MinimalAPI)
+class UMovieSceneCameraParameterInstantiator : public UMovieSceneEntityInstantiatorSystem
+{
+	GENERATED_BODY()
+
+public:
+
+	UMovieSceneCameraParameterInstantiator(const FObjectInitializer& ObjInit);
+
+#if WITH_EDITORONLY_DATA
+	static void OnMovieSceneSectionAddedToTrack(UMovieSceneTrack* Track, UMovieSceneSection* NewSection);
+#endif
+
+private:
+
+	virtual void OnLink() override;
+	virtual void OnRun(FSystemTaskPrerequisites& InPrerequisites, FSystemSubsequentTasks& Subsequents) override final;
+
+private:
+
+	TSharedPtr<UE::Cameras::FPreAnimatedCameraParameterStorage> PreAnimatedStorage;
+};
+
+
diff --git a/Engine/Plugins/Cameras/GameplayCameras/Source/GameplayCameras/Private/MovieScene/MovieSceneGameplayCamerasComponentTypes.cpp b/Engine/Plugins/Cameras/GameplayCameras/Source/GameplayCameras/Private/MovieScene/MovieSceneGameplayCamerasComponentTypes.cpp
--- a/Engine/Plugins/Cameras/GameplayCameras/Source/GameplayCameras/Private/MovieScene/MovieSceneGameplayCamerasComponentTypes.cpp
+++ b/Engine/Plugins/Cameras/GameplayCameras/Source/GameplayCameras/Private/MovieScene/MovieSceneGameplayCamerasComponentTypes.cpp
@@ -23,6 +23,8 @@ FMovieSceneGameplayCamerasComponentTypes::FMovieSceneGameplayCamerasComponentTyp
 	FBuiltInComponentTypes* BuiltInComponents = FBuiltInComponentTypes::Get();
 	FComponentRegistry* ComponentRegistry = UMovieSceneEntitySystemLinker::GetComponents();
 
+	CameraParameterOverrideID = ComponentRegistry->NewComponentType<FGuid>(TEXT("Camera Parameter ID"), EComponentTypeFlags::CopyToChildren);
+
 	ComponentRegistry->NewPropertyType(CameraFramingZone, TEXT("Camera Framing Zone"));
 
 	BuiltInComponents->PropertyRegistry.DefineCompositeProperty(CameraFramingZone, TEXT("Apply FCameraFramingZone Properties"))
diff --git a/Engine/Plugins/Cameras/GameplayCameras/Source/GameplayCameras/Private/MovieScene/MovieSceneGameplayCamerasComponentTypes.h b/Engine/Plugins/Cameras/GameplayCameras/Source/GameplayCameras/Private/MovieScene/MovieSceneGameplayCamerasComponentTypes.h
--- a/Engine/Plugins/Cameras/GameplayCameras/Source/GameplayCameras/Private/MovieScene/MovieSceneGameplayCamerasComponentTypes.h
+++ b/Engine/Plugins/Cameras/GameplayCameras/Source/GameplayCameras/Private/MovieScene/MovieSceneGameplayCamerasComponentTypes.h
@@ -17,6 +17,8 @@ struct FMovieSceneGameplayCamerasComponentTypes
 {
 	GAMEPLAYCAMERAS_API ~FMovieSceneGameplayCamerasComponentTypes();
 
+	UE::MovieScene::TComponentTypeID<FGuid> CameraParameterOverrideID;
+
 	UE::MovieScene::TPropertyComponents<FCameraFramingZoneTraits> CameraFramingZone;
 
 	UE::MovieScene::TCustomPropertyRegistration<FCameraFramingZoneTraits, 1> CustomCameraFramingZoneAccessors;
diff --git a/Engine/Plugins/Cameras/GameplayCameras/Source/GameplayCameras/Public/Core/CameraAssetReference.h b/Engine/Plugins/Cameras/GameplayCameras/Source/GameplayCameras/Public/Core/CameraAssetReference.h
--- a/Engine/Plugins/Cameras/GameplayCameras/Source/GameplayCameras/Public/Core/CameraAssetReference.h
+++ b/Engine/Plugins/Cameras/GameplayCameras/Source/GameplayCameras/Public/Core/CameraAssetReference.h
@@ -88,6 +88,9 @@ public:
 	GAMEPLAYCAMERAS_API bool IsParameterOverridden(const FGuid PropertyID) const;
 	GAMEPLAYCAMERAS_API void SetParameterOverridden(const FGuid PropertyID, bool bIsOverridden);
 
+	GAMEPLAYCAMERAS_API bool IsParameterAnimated(const FGuid PropertyID) const;
+	GAMEPLAYCAMERAS_API void SetParameterAnimated(const FGuid PropertyID, bool bIsAnimated);
+
 	GAMEPLAYCAMERAS_API bool NeedsRebuildParameters() const;
 	GAMEPLAYCAMERAS_API bool RebuildParametersIfNeeded();
 	GAMEPLAYCAMERAS_API void RebuildParameters();
@@ -99,13 +102,16 @@ private:
 	TObjectPtr<UCameraAsset> CameraAsset;
 
 	/** The camera asset's parameters. */
-	UPROPERTY(EditAnywhere, Category="", meta=(FixedLayout))
+	UPROPERTY(EditAnywhere, Category="", meta=(FixedLayout=true, InterpBagProperties=true, DisplayName="Camera Parameters"))
 	FInstancedPropertyBag Parameters;
 
 	/** The list of camera rig parameters with override values. */
 	UPROPERTY(EditAnywhere, Category="")
 	TArray<FGuid> ParameterOverrideGuids;
 
+	UPROPERTY()
+	TArray<FGuid> ParameterAnimatedGuids;
+
 	friend class UE::Cameras::FCameraAssetReferenceDetailsCustomization;
 };
 
diff --git a/Engine/Plugins/Cameras/GameplayCameras/Source/GameplayCameras/Public/Core/CameraParameters.h b/Engine/Plugins/Cameras/GameplayCameras/Source/GameplayCameras/Public/Core/CameraParameters.h
--- a/Engine/Plugins/Cameras/GameplayCameras/Source/GameplayCameras/Public/Core/CameraParameters.h
+++ b/Engine/Plugins/Cameras/GameplayCameras/Source/GameplayCameras/Public/Core/CameraParameters.h
@@ -38,7 +38,7 @@ namespace UE::Cameras
 // driving the parameter, such as a camera rig parameter override.
 
 /** Boolean camera parameter. */
-USTRUCT()
+USTRUCT(BlueprintType)
 struct GAMEPLAYCAMERAS_API FBooleanCameraParameter
 {
 	GENERATED_BODY()
@@ -46,7 +46,7 @@ struct GAMEPLAYCAMERAS_API FBooleanCameraParameter
 	using ValueType = bool;
 	using VariableAssetType = UBooleanCameraVariable;
 
-	UPROPERTY(EditAnywhere, Category=Common)
+	UPROPERTY(EditAnywhere, Interp, Category=Common, meta=(SequencerUseParentPropertyName=true))
 	bool Value = false;
 
 	UPROPERTY()
@@ -61,7 +61,7 @@ struct GAMEPLAYCAMERAS_API FBooleanCameraParameter
 };
 
 /** Integer camera parameter. */
-USTRUCT()
+USTRUCT(BlueprintType)
 struct GAMEPLAYCAMERAS_API FInteger32CameraParameter
 {
 	GENERATED_BODY()
@@ -69,7 +69,7 @@ struct GAMEPLAYCAMERAS_API FInteger32CameraParameter
 	using ValueType = int32;
 	using VariableAssetType = UInteger32CameraVariable;
 
-	UPROPERTY(EditAnywhere, Category=Common)
+	UPROPERTY(EditAnywhere, Interp, Category=Common, meta=(SequencerUseParentPropertyName=true))
 	int32 Value = 0;
 
 	UPROPERTY()
@@ -84,7 +84,7 @@ struct GAMEPLAYCAMERAS_API FInteger32CameraParameter
 };
 
 /** Float camera parameter. */
-USTRUCT()
+USTRUCT(BlueprintType)
 struct FFloatCameraParameter
 {
 	GENERATED_BODY()
@@ -92,7 +92,7 @@ struct FFloatCameraParameter
 	using ValueType = float;
 	using VariableAssetType = UFloatCameraVariable;
 
-	UPROPERTY(EditAnywhere, Category=Common)
+	UPROPERTY(EditAnywhere, Interp, Category=Common, meta=(SequencerUseParentPropertyName=true))
 	float Value = 0.f;
 
 	UPROPERTY()
@@ -107,7 +107,7 @@ struct FFloatCameraParameter
 };
 
 /** Double camera parameter. */
-USTRUCT()
+USTRUCT(BlueprintType)
 struct GAMEPLAYCAMERAS_API FDoubleCameraParameter
 {
 	GENERATED_BODY()
@@ -115,7 +115,7 @@ struct GAMEPLAYCAMERAS_API FDoubleCameraParameter
 	using ValueType = double;
 	using VariableAssetType = UDoubleCameraVariable;
 
-	UPROPERTY(EditAnywhere, Category=Common)
+	UPROPERTY(EditAnywhere, Interp, Category=Common, meta=(SequencerUseParentPropertyName=true))
 	double Value = 0.0;
 
 	UPROPERTY()
@@ -130,7 +130,7 @@ struct GAMEPLAYCAMERAS_API FDoubleCameraParameter
 };
 
 /** Vector2f camera parameter. */
-USTRUCT()
+USTRUCT(BlueprintType)
 struct GAMEPLAYCAMERAS_API FVector2fCameraParameter
 {
 	GENERATED_BODY()
@@ -138,7 +138,7 @@ struct GAMEPLAYCAMERAS_API FVector2fCameraParameter
 	using ValueType = FVector2f;
 	using VariableAssetType = UVector2fCameraVariable;
 
-	UPROPERTY(EditAnywhere, Category=Common)
+	UPROPERTY(EditAnywhere, Interp, Category=Common, meta=(SequencerUseParentPropertyName=true))
 	FVector2f Value;
 
 	UPROPERTY()
@@ -154,7 +154,7 @@ struct GAMEPLAYCAMERAS_API FVector2fCameraParameter
 };
 
 /** Vector2d camera parameter. */
-USTRUCT()
+USTRUCT(BlueprintType)
 struct GAMEPLAYCAMERAS_API FVector2dCameraParameter
 {
 	GENERATED_BODY()
@@ -162,7 +162,7 @@ struct GAMEPLAYCAMERAS_API FVector2dCameraParameter
 	using ValueType = FVector2D;
 	using VariableAssetType = UVector2dCameraVariable;
 
-	UPROPERTY(EditAnywhere, Category=Common)
+	UPROPERTY(EditAnywhere, Interp, Category=Common, meta=(SequencerUseParentPropertyName=true))
 	FVector2D Value;
 
 	UPROPERTY()
@@ -178,7 +178,7 @@ struct GAMEPLAYCAMERAS_API FVector2dCameraParameter
 };
 
 /** Vector3f camera parameter. */
-USTRUCT()
+USTRUCT(BlueprintType)
 struct GAMEPLAYCAMERAS_API FVector3fCameraParameter
 {
 	GENERATED_BODY()
@@ -186,7 +186,7 @@ struct GAMEPLAYCAMERAS_API FVector3fCameraParameter
 	using ValueType = FVector3f;
 	using VariableAssetType = UVector3fCameraVariable;
 
-	UPROPERTY(EditAnywhere, Category=Common)
+	UPROPERTY(EditAnywhere, Interp, Category=Common, meta=(SequencerUseParentPropertyName=true))
 	FVector3f Value;
 
 	UPROPERTY()
@@ -202,16 +202,16 @@ struct GAMEPLAYCAMERAS_API FVector3fCameraParameter
 };
 
 /** Vector3d camera parameter. */
-USTRUCT()
+USTRUCT(BlueprintType)
 struct GAMEPLAYCAMERAS_API FVector3dCameraParameter
 {
 	GENERATED_BODY()
 
-	using ValueType = FVector3d;
+	using ValueType = FVector;
 	using VariableAssetType = UVector3dCameraVariable;
 
-	UPROPERTY(EditAnywhere, Category=Common)
-	FVector3d Value;
+	UPROPERTY(EditAnywhere, Interp, Category=Common, meta=(SequencerUseParentPropertyName=true))
+	FVector Value;
 
 	UPROPERTY()
 	FCameraVariableID VariableID;
@@ -226,7 +226,7 @@ struct GAMEPLAYCAMERAS_API FVector3dCameraParameter
 };
 
 /** Vector4f camera parameter. */
-USTRUCT()
+USTRUCT(BlueprintType)
 struct GAMEPLAYCAMERAS_API FVector4fCameraParameter
 {
 	GENERATED_BODY()
@@ -234,7 +234,7 @@ struct GAMEPLAYCAMERAS_API FVector4fCameraParameter
 	using ValueType = FVector4f;
 	using VariableAssetType = UVector4fCameraVariable;
 
-	UPROPERTY(EditAnywhere, Category=Common)
+	UPROPERTY(EditAnywhere, Interp, Category=Common, meta=(SequencerUseParentPropertyName=true))
 	FVector4f Value;
 
 	UPROPERTY()
@@ -250,16 +250,16 @@ struct GAMEPLAYCAMERAS_API FVector4fCameraParameter
 };
 
 /** Vector4d camera parameter. */
-USTRUCT()
+USTRUCT(BlueprintType)
 struct GAMEPLAYCAMERAS_API FVector4dCameraParameter
 {
 	GENERATED_BODY()
 
-	using ValueType = FVector4d;
+	using ValueType = FVector4;
 	using VariableAssetType = UVector4dCameraVariable;
 
-	UPROPERTY(EditAnywhere, Category=Common)
-	FVector4d Value;
+	UPROPERTY(EditAnywhere, Interp, Category=Common, meta=(SequencerUseParentPropertyName=true))
+	FVector4 Value;
 
 	UPROPERTY()
 	FCameraVariableID VariableID;
@@ -274,7 +274,7 @@ struct GAMEPLAYCAMERAS_API FVector4dCameraParameter
 };
 
 /** Rotator3f camera parameter. */
-USTRUCT()
+USTRUCT(BlueprintType)
 struct GAMEPLAYCAMERAS_API FRotator3fCameraParameter
 {
 	GENERATED_BODY()
@@ -282,7 +282,7 @@ struct GAMEPLAYCAMERAS_API FRotator3fCameraParameter
 	using ValueType = FRotator3f;
 	using VariableAssetType = URotator3fCameraVariable;
 
-	UPROPERTY(EditAnywhere, Category=Common)
+	UPROPERTY(EditAnywhere, Interp, Category=Common, meta=(SequencerUseParentPropertyName=true))
 	FRotator3f Value;
 
 	UPROPERTY()
@@ -298,16 +298,16 @@ struct GAMEPLAYCAMERAS_API FRotator3fCameraParameter
 };
 
 /** Rotator3d camera parameter. */
-USTRUCT()
+USTRUCT(BlueprintType)
 struct GAMEPLAYCAMERAS_API FRotator3dCameraParameter
 {
 	GENERATED_BODY()
 
-	using ValueType = FRotator3d;
+	using ValueType = FRotator;
 	using VariableAssetType = URotator3dCameraVariable;
 
-	UPROPERTY(EditAnywhere, Category=Common)
-	FRotator3d Value;
+	UPROPERTY(EditAnywhere, Interp, Category=Common, meta=(SequencerUseParentPropertyName=true))
+	FRotator Value;
 
 	UPROPERTY()
 	FCameraVariableID VariableID;
@@ -322,7 +322,7 @@ struct GAMEPLAYCAMERAS_API FRotator3dCameraParameter
 };
 
 /** Transform3f camera parameter. */
-USTRUCT()
+USTRUCT(BlueprintType)
 struct GAMEPLAYCAMERAS_API FTransform3fCameraParameter
 {
 	GENERATED_BODY()
@@ -330,7 +330,7 @@ struct GAMEPLAYCAMERAS_API FTransform3fCameraParameter
 	using ValueType = FTransform3f;
 	using VariableAssetType = UTransform3fCameraVariable;
 
-	UPROPERTY(EditAnywhere, Category=Common)
+	UPROPERTY(EditAnywhere, Interp, Category=Common, meta=(SequencerUseParentPropertyName=true))
 	FTransform3f Value;
 
 	UPROPERTY()
@@ -345,16 +345,16 @@ struct GAMEPLAYCAMERAS_API FTransform3fCameraParameter
 };
 
 /** Transform3d camera parameter. */
-USTRUCT()
+USTRUCT(BlueprintType)
 struct GAMEPLAYCAMERAS_API FTransform3dCameraParameter
 {
 	GENERATED_BODY()
 
-	using ValueType = FTransform3d;
+	using ValueType = FTransform;
 	using VariableAssetType = UTransform3dCameraVariable;
 
-	UPROPERTY(EditAnywhere, Category=Common)
-	FTransform3d Value;
+	UPROPERTY(EditAnywhere, Interp, Category=Common, meta=(SequencerUseParentPropertyName=true))
+	FTransform Value;
 
 	UPROPERTY()
 	FCameraVariableID VariableID;
diff --git a/Engine/Plugins/Cameras/GameplayCameras/Source/GameplayCameras/Public/Core/CameraRigAssetReference.h b/Engine/Plugins/Cameras/GameplayCameras/Source/GameplayCameras/Public/Core/CameraRigAssetReference.h
--- a/Engine/Plugins/Cameras/GameplayCameras/Source/GameplayCameras/Public/Core/CameraRigAssetReference.h
+++ b/Engine/Plugins/Cameras/GameplayCameras/Source/GameplayCameras/Public/Core/CameraRigAssetReference.h
@@ -260,6 +260,10 @@ struct FCameraRigAssetReferenceParameterMetaData
 	/** Whether this parameter has an override value. */
 	UPROPERTY()
 	bool bIsOverridden = false;
+
+	/** Whether this parameter is being animated. */
+	UPROPERTY()
+	bool bIsAnimated = false;
 };
 
 /**
@@ -329,6 +333,9 @@ public:
 	GAMEPLAYCAMERAS_API bool IsParameterOverridden(const FGuid& PropertyID) const;
 	GAMEPLAYCAMERAS_API void SetParameterOverridden(const FGuid& PropertyID, bool bIsOverridden);
 
+	GAMEPLAYCAMERAS_API bool IsParameterAnimated(const FGuid& PropertyID) const;
+	GAMEPLAYCAMERAS_API void SetParameterAnimated(const FGuid& PropertyID, bool bIsAnimated);
+
 	GAMEPLAYCAMERAS_API bool NeedsRebuildParameters() const;
 	GAMEPLAYCAMERAS_API bool RebuildParametersIfNeeded();
 	GAMEPLAYCAMERAS_API void RebuildParameters();
diff --git a/Engine/Plugins/Cameras/GameplayCameras/Source/GameplayCameras/Public/GameFramework/GameplayCameraComponent.h b/Engine/Plugins/Cameras/GameplayCameras/Source/GameplayCameras/Public/GameFramework/GameplayCameraComponent.h
--- a/Engine/Plugins/Cameras/GameplayCameras/Source/GameplayCameras/Public/GameFramework/GameplayCameraComponent.h
+++ b/Engine/Plugins/Cameras/GameplayCameras/Source/GameplayCameras/Public/GameFramework/GameplayCameraComponent.h
@@ -125,7 +125,7 @@ private:
 	void ActivateCameraEvaluationContext(int32 PlayerIndex);
 	void ActivateCameraEvaluationContext(APlayerController* PlayerController);
 	void EnsureCameraEvaluationContextCreated(APlayerController* PlayerController);
-	void UpdateCameraEvaluationContext(bool bApplyParameterOverrides);
+	void UpdateCameraEvaluationContext(bool bForceApplyParameterOverrides);
 	void UpdateOutputCameraComponent();
 	void DeactivateCameraEvaluationContext();
 
@@ -141,7 +141,7 @@ private:
 public:
 
 	/** The camera asset to run. */
-	UPROPERTY(BlueprintReadWrite, EditAnywhere, Category=Camera)
+	UPROPERTY(BlueprintReadWrite, EditAnywhere, Category=Camera, meta=(SequencerHideProperty=true))
 	FCameraAssetReference CameraReference;
 
 	/**
diff --git a/Engine/Plugins/Cameras/GameplayCameras/Source/GameplayCamerasEditor/Private/Sequencer/GameplayCameraComponentTrackEditor.cpp b/Engine/Plugins/Cameras/GameplayCameras/Source/GameplayCamerasEditor/Private/Sequencer/GameplayCameraComponentTrackEditor.cpp
--- a/Engine/Plugins/Cameras/GameplayCameras/Source/GameplayCamerasEditor/Private/Sequencer/GameplayCameraComponentTrackEditor.cpp
+++ b/Engine/Plugins/Cameras/GameplayCameras/Source/GameplayCamerasEditor/Private/Sequencer/GameplayCameraComponentTrackEditor.cpp
@@ -3,27 +3,141 @@
 #include "Sequencer/GameplayCameraComponentTrackEditor.h"
 
 #include "Core/CameraAsset.h"
+#include "Core/CameraParameters.h"
 #include "Core/CameraRigParameterDefinition.h"
+#include "EventHandlers/MovieSceneDataEventContainer.h"
 #include "Framework/MultiBox/MultiBoxBuilder.h"
 #include "GameFramework/Actor.h"
 #include "GameFramework/GameplayCameraComponent.h"
-#include "MovieScene/MovieSceneCameraFramingZoneTrack.h"
-#include "Nodes/Framing/CameraFramingZone.h"
+#include "ISequencer.h"
+#include "KeyPropertyParams.h"
+#include "MVVM/ViewModels/ObjectBindingModel.h"
+#include "MVVM/ViewModels/SequenceModel.h"
+#include "StructUtils/PropertyBag.h"
 #include "Styles/GameplayCamerasEditorStyle.h"
-#include "Tracks/MovieSceneActorReferenceTrack.h"
-#include "Tracks/MovieSceneBoolTrack.h"
-#include "Tracks/MovieSceneDoubleTrack.h"
-#include "Tracks/MovieSceneEnumTrack.h"
-#include "Tracks/MovieSceneFloatTrack.h"
-#include "Tracks/MovieSceneIntegerTrack.h"
-#include "Tracks/MovieSceneObjectPropertyTrack.h"
-#include "Tracks/MovieSceneRotatorTrack.h"
-#include "Tracks/MovieSceneStringTrack.h"
-#include "Tracks/MovieSceneTransformTrack.h"
-#include "Tracks/MovieSceneVectorTrack.h"
+#include "Tracks/MovieScenePropertyTrack.h"
 
 #define LOCTEXT_NAMESPACE "GameplayCameraComponentTrackEditor"
 
+namespace UE::Cameras::Internal
+{
+
+// TODO: duplicated from GetKeyablePropertyPaths in ObjectBindingModel.cpp
+void GetKeyablePropertyPathsImpl(TSharedPtr<ISequencer> Sequencer, const UClass* BaseObjectClass, const UStruct* Struct, const void* StructValuePtr, FPropertyPath PropertyPath, TArray<FPropertyPath>& KeyablePropertyPaths)
+{
+	for (TFieldIterator<FProperty> PropertyIterator(Struct); PropertyIterator; ++PropertyIterator)
+	{
+		FProperty* Property = *PropertyIterator;
+
+		if (!Property || Property->HasAnyPropertyFlags(CPF_Deprecated))
+		{
+			continue;
+		}
+
+		PropertyPath.AddProperty(FPropertyInfo(Property));
+
+		if (FArrayProperty* ArrayProperty = CastField<FArrayProperty>(Property))
+		{
+			FScriptArrayHelper ArrayHelper(ArrayProperty, ArrayProperty->ContainerPtrToValuePtr<void>(StructValuePtr));
+			for (int32 Index = 0; Index < ArrayHelper.Num(); ++Index)
+			{
+				PropertyPath.AddProperty(FPropertyInfo(ArrayProperty->Inner, Index));
+
+				if (Sequencer->CanKeyProperty(FCanKeyPropertyParams(BaseObjectClass, PropertyPath)))
+				{
+					KeyablePropertyPaths.Add(PropertyPath);
+				}
+				else if (FStructProperty* StructProperty = CastField<FStructProperty>(ArrayProperty->Inner))
+				{
+					GetKeyablePropertyPathsImpl(Sequencer, BaseObjectClass, StructProperty->Struct, ArrayHelper.GetRawPtr(Index), PropertyPath, KeyablePropertyPaths);
+				}
+
+				PropertyPath = *PropertyPath.TrimPath(1);
+			}
+		}
+		else if (Sequencer->CanKeyProperty(FCanKeyPropertyParams(BaseObjectClass, PropertyPath)))
+		{
+			KeyablePropertyPaths.Add(PropertyPath);
+		}
+		else if (FStructProperty* StructProperty = CastField<FStructProperty>(Property))
+		{
+			GetKeyablePropertyPathsImpl(Sequencer, BaseObjectClass, StructProperty->Struct, StructProperty->ContainerPtrToValuePtr<void>(StructValuePtr), PropertyPath, KeyablePropertyPaths);
+		}
+
+		PropertyPath = *PropertyPath.TrimPath(1);
+	}
+}
+
+void GetKeyablePropertyPaths(TSharedPtr<ISequencer> Sequencer, const UGameplayCameraComponent* CameraComponent, TArray<FPropertyPath>& KeyablePropertyPaths)
+{
+	const UCameraAsset* CameraAsset = CameraComponent->CameraReference.GetCameraAsset();
+	if (!CameraAsset)
+	{
+		return;
+	}
+
+	const FInstancedPropertyBag& CameraParameters = CameraComponent->CameraReference.GetParameters();
+	const UPropertyBag* CameraParametersStruct = CameraParameters.GetPropertyBagStruct();
+	const uint8* CameraParametersMemory = CameraParameters.GetValue().GetMemory();
+	if (!CameraParametersStruct || !CameraParametersMemory)
+	{
+		return;
+	}
+
+	// Start us off with the property path of the parameters struct.
+	const UClass* ComponentClass = UGameplayCameraComponent::StaticClass();
+	const UStruct* CameraAssetReferenceStruct = FCameraAssetReference::StaticStruct();
+	const UStruct* PropertyBagStruct = FInstancedPropertyBag::StaticStruct();
+
+	FPropertyPath PropertyPath;
+	PropertyPath.AddProperty(FPropertyInfo(ComponentClass->FindPropertyByName(TEXT("CameraReference"))));
+	PropertyPath.AddProperty(FPropertyInfo(CameraAssetReferenceStruct->FindPropertyByName(TEXT("Parameters"))));
+	PropertyPath.AddProperty(FPropertyInfo(PropertyBagStruct->FindPropertyByName(TEXT("Value"))));
+
+	GetKeyablePropertyPathsImpl(Sequencer, ComponentClass, CameraParametersStruct, CameraParametersMemory, PropertyPath, KeyablePropertyPaths);
+}
+
+class FCameraParameterTrackSetupHandler : public UE::MovieScene::TIntrusiveEventHandler<UE::MovieScene::ISequenceDataEventHandler>
+{
+public:
+
+	FCameraParameterTrackSetupHandler(const FGuid& ObjectBindingID)
+		: MonitoredObjectBindingID(ObjectBindingID)
+	{
+	}
+
+	void SetDesiredTrackName(const FText& InDesiredTrackName)
+	{
+		DesiredTrackName = InDesiredTrackName;
+	}
+
+	virtual void OnTrackAddedToBinding(UMovieSceneTrack* Track, const FGuid& ObjectBindingID) override
+	{
+		if (ObjectBindingID != MonitoredObjectBindingID)
+		{
+			return;
+		}
+
+		UMovieScenePropertyTrack* PropertyTrack = Cast<UMovieScenePropertyTrack>(Track);
+		if (!PropertyTrack)
+		{
+			return;
+		}
+
+		if (!DesiredTrackName.IsEmpty())
+		{
+			PropertyTrack->SetDisplayName(DesiredTrackName);
+		}
+	}
+	
+private:
+
+	FGuid MonitoredObjectBindingID;
+	FText DesiredTrackName;
+};
+
+}  // namespace UE::Cameras::Internal
+
 TSharedRef<ISequencerTrackEditor> FGameplayCameraComponentTrackEditor::CreateTrackEditor(TSharedRef<ISequencer> OwningSequencer)
 {
 	return MakeShareable(new FGameplayCameraComponentTrackEditor(OwningSequencer));
@@ -58,64 +172,23 @@ void FGameplayCameraComponentTrackEditor::ExtendObjectBindingTrackMenu(TSharedRe
 
 void FGameplayCameraComponentTrackEditor::OnExtendObjectBindingTrackMenu(FMenuBuilder& MenuBuilder, TArray<FGuid> ObjectBindings)
 {
-	TMap<FName, FCameraRigParameterDefinition> DefinitionsByName;
-	TMap<FName, int32> NumEqualDefinitions;
-
-	for (const FGuid& ObjectBinding : ObjectBindings)
+	if (ObjectBindings.Num() != 1)
 	{
-		UGameplayCameraComponent* CameraComponent = GetCameraComponentForBinding(ObjectBinding);
-		if (CameraComponent && CameraComponent->CameraReference.GetCameraAsset())
-		{
-			const UCameraAsset* CameraAsset = CameraComponent->CameraReference.GetCameraAsset();
-			for (const FCameraRigParameterDefinition& Definition : CameraAsset->GetParameterDefinitions())
-			{
-				if (FCameraRigParameterDefinition* ExistingDefinition = DefinitionsByName.Find(Definition.ParameterName))
-				{
-					if (Definition == *ExistingDefinition)
-					{
-						++NumEqualDefinitions[Definition.ParameterName];
-					}
-				}
-				else
-				{
-					DefinitionsByName.Add(Definition.ParameterName, Definition);
-					NumEqualDefinitions.Add(Definition.ParameterName, 1);
-				}
-			}
-		}
+		return;
 	}
 
-	const int32 NumObjectBindings = ObjectBindings.Num();
-	TArray<FCameraRigParameterDefinition> DefinitionsToAdd;
-
-	for (auto It = DefinitionsByName.CreateConstIterator(); It; ++It)
+	const UGameplayCameraComponent* CameraComponent = GetCameraComponentForBinding(ObjectBindings[0]);
+	if (CameraComponent)
 	{
-		if (NumEqualDefinitions.FindRef(It.Key()) == NumObjectBindings)
-		{
-			DefinitionsToAdd.Add(It.Value());
-		}
-	}
+		using namespace UE::Cameras::Internal;
 
-	DefinitionsToAdd.StableSort([](const FCameraRigParameterDefinition& A, const FCameraRigParameterDefinition& B)
-			{
-				return A.ParameterName.Compare(B.ParameterName) < 0;
-			});
+		TArray<FPropertyPath> KeyablePropertyPaths;
+		GetKeyablePropertyPaths(GetSequencer(), CameraComponent, KeyablePropertyPaths);
 
-	if (DefinitionsToAdd.Num() > 0)
-	{
 		MenuBuilder.BeginSection(TEXT("CameraParameters"), LOCTEXT("AddCameraParametersMenuSection", "Camera Parameters"));
 
-		for (const FCameraRigParameterDefinition& Definition : DefinitionsToAdd)
-		{
-			MenuBuilder.AddMenuEntry(
-					FText::FromName(Definition.ParameterName),
-					LOCTEXT("AddCameraParameterMenuToolTip", "Adds a track for controlling this camera parameter."),
-					FSlateIcon(),
-					FUIAction(
-						FExecuteAction::CreateSP(this, &FGameplayCameraComponentTrackEditor::AddCameraParameterTrack, Definition, ObjectBindings),
-						FCanExecuteAction::CreateSP(this, &FGameplayCameraComponentTrackEditor::CanAddCameraParameterTrack, Definition, ObjectBindings[0]))
-					);
-		}
+		// Start at level 3 to skip adding menus for CameraReference > Parameters > Value
+		BuildAddParameterTrackMenuItems(ObjectBindings[0], MenuBuilder, KeyablePropertyPaths, 3);
 
 		MenuBuilder.EndSection();
 	}
@@ -178,155 +251,151 @@ UGameplayCameraComponent* FGameplayCameraComponentTrackEditor::GetCameraComponen
 	return nullptr;
 }
 
-void FGameplayCameraComponentTrackEditor::AddCameraParameterTrack(FCameraRigParameterDefinition Definition, TArray<FGuid> ObjectBindings)
+void FGameplayCameraComponentTrackEditor::AddCameraParameterTrack(FPropertyMenuData PropertyMenuData, FGuid ObjectBinding)
 {
-	TSubclassOf<UMovieSceneTrack> ParameterTrackType = GetParameterTrackFromDefinition(Definition);
-	if (ParameterTrackType)
+	using namespace UE::Cameras::Internal;
+	using namespace UE::MovieScene;
+
+	TSharedPtr<ISequencer> Sequencer = GetSequencer();
+	UObject* BoundObject = Sequencer->FindSpawnedObjectOrTemplate(ObjectBinding);
+
+	if (BoundObject != nullptr)
 	{
-		const FScopedTransaction Transaction(LOCTEXT("AddCameraParameterTrack", "Add camera parameter track"));
+		UMovieSceneSequence* FocusedSequence = Sequencer->GetFocusedMovieSceneSequence();
+		FCameraParameterTrackSetupHandler TrackSetupHandler(ObjectBinding);
+		FocusedSequence->GetMovieScene()->EventHandlers.Link(&TrackSetupHandler);
 
-		for (FGuid ObjectBinding : ObjectBindings)
+		if (PropertyMenuData.PropertyIndexForMenuName != INDEX_NONE)
 		{
-			FFindOrCreateTrackResult Result = FindOrCreateTrackForObject(ObjectBinding, ParameterTrackType, Definition.ParameterName, true);
-			if (Result.bWasCreated)
-			{
-				UMovieScenePropertyTrack* NewTrack = CastChecked<UMovieScenePropertyTrack>(Result.Track);
-				NewTrack->SetPropertyNameAndPath(Definition.ParameterName, Definition.ParameterName.ToString());
-				InitializeNewTrack(NewTrack, Definition);
-			}
+			const FPropertyInfo& PropertyInfoForTrackName = PropertyMenuData.PropertyPath.GetPropertyInfo(PropertyMenuData.PropertyIndexForMenuName);
+			const FProperty* PropertyForTrackName = PropertyInfoForTrackName.Property.Get();
+			TrackSetupHandler.SetDesiredTrackName(PropertyForTrackName->GetDisplayNameText());
 		}
 
-		GetSequencer()->NotifyMovieSceneDataChanged(EMovieSceneDataChangeType::MovieSceneStructureItemAdded);
+		TArray<UObject*> KeyableBoundObjects;
+		KeyableBoundObjects.Add(BoundObject);
+
+		ESequencerKeyMode KeyMode = Sequencer->GetAutoSetTrackDefaults() == false ? ESequencerKeyMode::ManualKeyForced : ESequencerKeyMode::ManualKey;
+
+		FKeyPropertyParams KeyPropertyParams(KeyableBoundObjects, PropertyMenuData.PropertyPath, KeyMode);
+
+		Sequencer->KeyProperty(KeyPropertyParams);
 	}
 }
 
-bool FGameplayCameraComponentTrackEditor::CanAddCameraParameterTrack(FCameraRigParameterDefinition Definition, FGuid ObjectBinding) const
+bool FGameplayCameraComponentTrackEditor::CanAddCameraParameterTrack(FPropertyMenuData PropertyMenuData, FGuid ObjectBinding) const
 {
-	TSharedPtr<ISequencer> SequencerPtr = GetSequencer();
-	const UMovieScene* FocusedMovieScene = SequencerPtr->GetFocusedMovieSceneSequence()->GetMovieScene();
+	return true;
+}
 
-	TSubclassOf<UMovieSceneTrack> ParameterTrackType = GetParameterTrackFromDefinition(Definition);
-	if (!ParameterTrackType)
-	{
-		return false;
-	}
-	
-	UMovieSceneTrack* ExistingParameterTrack = FocusedMovieScene->FindTrack(ParameterTrackType, ObjectBinding, Definition.ParameterName);
-	return ExistingParameterTrack == nullptr;
+// TODO: most of the below stuff is duplicated from ObjectBindingModel.cpp but without classifying in categories, 
+//		 and with special handling of camera parameters.
+
+void FGameplayCameraComponentTrackEditor::BuildAddParameterTrackMenuItem(FMenuBuilder& MenuBuilder, const FPropertyMenuData& KeyablePropertyMenuData, const FGuid& ObjectBinding)
+{
+	FUIAction AddTrackMenuAction(
+			FExecuteAction::CreateSP(this, &FGameplayCameraComponentTrackEditor::AddCameraParameterTrack, KeyablePropertyMenuData, ObjectBinding),
+			FCanExecuteAction::CreateSP(this, &FGameplayCameraComponentTrackEditor::CanAddCameraParameterTrack, KeyablePropertyMenuData, ObjectBinding));
+	MenuBuilder.AddMenuEntry(FText::FromString(KeyablePropertyMenuData.MenuName), FText(), FSlateIcon(), AddTrackMenuAction);
 }
 
-void FGameplayCameraComponentTrackEditor::InitializeNewTrack(UMovieScenePropertyTrack* NewTrack, const FCameraRigParameterDefinition& Definition) const
+void FGameplayCameraComponentTrackEditor::BuildAddParameterTrackMenuItems(const FGuid& ObjectBinding, FMenuBuilder& MenuBuilder, TArray<FPropertyPath> KeyablePropertyPaths, int32 PropertyNameIndexStart)
 {
-	if (Definition.ParameterType == ECameraRigInterfaceParameterType::Blendable)
+	if (KeyablePropertyPaths.IsEmpty())
 	{
-		switch (Definition.VariableType)
-		{
-			case ECameraVariableType::Vector2f:
-				CastChecked<UMovieSceneFloatVectorTrack>(NewTrack)->SetNumChannelsUsed(2);
-				break;
-			case ECameraVariableType::Vector3f:
-				CastChecked<UMovieSceneFloatVectorTrack>(NewTrack)->SetNumChannelsUsed(3);
-				break;
-			case ECameraVariableType::Vector4f:
-				CastChecked<UMovieSceneFloatVectorTrack>(NewTrack)->SetNumChannelsUsed(4);
-				break;
-			case ECameraVariableType::Vector2d:
-				CastChecked<UMovieSceneDoubleVectorTrack>(NewTrack)->SetNumChannelsUsed(2);
-				break;
-			case ECameraVariableType::Vector3d:
-				CastChecked<UMovieSceneDoubleVectorTrack>(NewTrack)->SetNumChannelsUsed(3);
-				break;
-			case ECameraVariableType::Vector4d:
-				CastChecked<UMovieSceneDoubleVectorTrack>(NewTrack)->SetNumChannelsUsed(4);
-				break;
-			default:
-				break;
-		}
+		return;
 	}
-	else if (Definition.ParameterType == ECameraRigInterfaceParameterType::Data)
+
+	// Create property menu data based on keyable property paths
+	TArray<FPropertyMenuData> KeyablePropertyMenuDatas;
+	for (const FPropertyPath& KeyablePropertyPath : KeyablePropertyPaths)
 	{
-		switch (Definition.DataType)
+		if (!ensure(KeyablePropertyPath.GetNumProperties() > PropertyNameIndexStart))
 		{
-			case ECameraContextDataType::Enum:
+			continue;
+		}
+
+		const FPropertyInfo& PropertyInfo = KeyablePropertyPath.GetPropertyInfo(PropertyNameIndexStart);
+		if (const FProperty* Property = PropertyInfo.Property.Get())
+		{
+			FPropertyMenuData KeyableMenuData;
+			KeyableMenuData.PropertyPath = KeyablePropertyPath;
+			if (PropertyInfo.ArrayIndex != INDEX_NONE)
+			{
+				KeyableMenuData.MenuName = FText::Format(LOCTEXT("PropertyMenuTextFormat", "{0} [{1}]"), Property->GetDisplayNameText(), FText::AsNumber(PropertyInfo.ArrayIndex)).ToString();
+			}
+			else
+			{
+				KeyableMenuData.MenuName = Property->GetDisplayNameText().ToString();
+			}
+
+			if (const FStructProperty* StructProperty = CastField<const FStructProperty>(Property))
+			{
+				bool bIsCameraParameter = false;
+#define UE_CAMERA_VARIABLE_FOR_TYPE(ValueType, ValueName)\
+				if (StructProperty->Struct == F##ValueName##CameraParameter::StaticStruct())\
+				{\
+					bIsCameraParameter = true;\
+				}\
+				else
+				UE_CAMERA_VARIABLE_FOR_ALL_TYPES()
+#undef UE_CAMERA_VARIABLE_FOR_TYPE
 				{
-					const UEnum* Enum = Cast<const UEnum>(Definition.DataTypeObject);
-					CastChecked<UMovieSceneEnumTrack>(NewTrack)->SetEnum(const_cast<UEnum*>(Enum));
 				}
-				break;
-			case ECameraContextDataType::Object:
+
+				if (bIsCameraParameter)
 				{
-					const UClass* ObjectClass = Cast<const UClass>(Definition.DataTypeObject);
-					if (!ObjectClass || !ObjectClass->IsChildOf<AActor>())
-					{
-						CastChecked<UMovieSceneObjectPropertyTrack>(NewTrack)->PropertyClass = const_cast<UClass*>(ObjectClass);
-					}
+					KeyableMenuData.PropertyIndexForMenuName = PropertyNameIndexStart;
 				}
-				break;
-			default:
-				break;
+			}
+
+			KeyablePropertyMenuDatas.Add(KeyableMenuData);
 		}
 	}
-}
 
-TSubclassOf<UMovieScenePropertyTrack> FGameplayCameraComponentTrackEditor::GetParameterTrackFromDefinition(const FCameraRigParameterDefinition& Definition) const
-{
-	if (Definition.ParameterType == ECameraRigInterfaceParameterType::Blendable)
+	KeyablePropertyMenuDatas.Sort();
+
+	for (int32 MenuDataIndex = 0; MenuDataIndex < KeyablePropertyMenuDatas.Num(); )
 	{
-		switch (Definition.VariableType)
+		// If this menu data only has one property name left in it, add the menu item
+		if (KeyablePropertyMenuDatas[MenuDataIndex].PropertyPath.GetNumProperties() == PropertyNameIndexStart + 1 || 
+				KeyablePropertyMenuDatas[MenuDataIndex].PropertyIndexForMenuName == PropertyNameIndexStart)
 		{
-			case ECameraVariableType::Boolean:
-				return UMovieSceneBoolTrack::StaticClass();
-			case ECameraVariableType::Integer32:
-				return UMovieSceneIntegerTrack::StaticClass();
-			case ECameraVariableType::Float:
-				return UMovieSceneFloatTrack::StaticClass();
-			case ECameraVariableType::Double:
-				return UMovieSceneDoubleTrack::StaticClass();
-			case ECameraVariableType::Vector2f:
-			case ECameraVariableType::Vector3f:
-			case ECameraVariableType::Vector4f:
-				return UMovieSceneFloatVectorTrack::StaticClass();
-			case ECameraVariableType::Vector2d:
-			case ECameraVariableType::Vector3d:
-			case ECameraVariableType::Vector4d:
-				return UMovieSceneDoubleVectorTrack::StaticClass();
-			case ECameraVariableType::Rotator3f:
-			case ECameraVariableType::Rotator3d:
-				return UMovieSceneRotatorTrack::StaticClass();
-			case ECameraVariableType::Transform3f:
-			case ECameraVariableType::Transform3d:
-				return UMovieSceneTransformTrack::StaticClass();
-			case ECameraVariableType::BlendableStruct:
-				if (Definition.BlendableStructType == FCameraFramingZone::StaticStruct())
-				{
-					return UMovieSceneCameraFramingZoneTrack::StaticClass();
-				}
-				break;
-			default:
-				break;
+			BuildAddParameterTrackMenuItem(MenuBuilder, KeyablePropertyMenuDatas[MenuDataIndex], ObjectBinding);
+			++MenuDataIndex;
 		}
-	}
-	else if (Definition.ParameterType == ECameraRigInterfaceParameterType::Data)
-	{
-		switch (Definition.DataType)
+		// Otherwise, look to the next menu data to gather up new data
+		else
 		{
-			case ECameraContextDataType::String:
-				return UMovieSceneStringTrack::StaticClass();
-			case ECameraContextDataType::Enum:
-				return UMovieSceneEnumTrack::StaticClass();
-			case ECameraContextDataType::Object:
-				if (Definition.DataTypeObject)
+			TArray<FPropertyPath> KeyableSubMenuPropertyPaths;
+			KeyableSubMenuPropertyPaths.Add(KeyablePropertyMenuDatas[MenuDataIndex].PropertyPath);
+
+			for (; MenuDataIndex < KeyablePropertyMenuDatas.Num() - 1; )
+			{
+				if (KeyablePropertyMenuDatas[MenuDataIndex].MenuName == KeyablePropertyMenuDatas[MenuDataIndex + 1].MenuName)
+				{	
+					++MenuDataIndex;
+					KeyableSubMenuPropertyPaths.Add(KeyablePropertyMenuDatas[MenuDataIndex].PropertyPath);
+				}
+				else
 				{
-					const UClass* ObjectClass = Cast<const UClass>(Definition.DataTypeObject);
-					if (ObjectClass && ObjectClass->IsChildOf<AActor>())
-					{
-						return UMovieSceneActorReferenceTrack::StaticClass();
-					}
+					break;
 				}
-				return UMovieSceneObjectPropertyTrack::StaticClass();
+			}
+
+			MenuBuilder.AddSubMenu(
+				FText::FromString(KeyablePropertyMenuDatas[MenuDataIndex].MenuName),
+				FText::GetEmpty(), 
+				FNewMenuDelegate::CreateSP(this, &FGameplayCameraComponentTrackEditor::BuildAddParameterTrackSubMenuItems, ObjectBinding, KeyableSubMenuPropertyPaths, PropertyNameIndexStart + 1));
+
+			++MenuDataIndex;
 		}
 	}
-	return nullptr;
+}
+
+void FGameplayCameraComponentTrackEditor::BuildAddParameterTrackSubMenuItems(FMenuBuilder& MenuBuilder, FGuid ObjectBinding, TArray<FPropertyPath> KeyablePropertyPaths, int32 PropertyNameIndexStart)
+{
+	BuildAddParameterTrackMenuItems(ObjectBinding, MenuBuilder, KeyablePropertyPaths, PropertyNameIndexStart);
 }
 
 #undef LOCTEXT_NAMESPACE
diff --git a/Engine/Plugins/Cameras/GameplayCameras/Source/GameplayCamerasEditor/Private/Sequencer/GameplayCameraComponentTrackEditor.h b/Engine/Plugins/Cameras/GameplayCameras/Source/GameplayCamerasEditor/Private/Sequencer/GameplayCameraComponentTrackEditor.h
--- a/Engine/Plugins/Cameras/GameplayCameras/Source/GameplayCamerasEditor/Private/Sequencer/GameplayCameraComponentTrackEditor.h
+++ b/Engine/Plugins/Cameras/GameplayCameras/Source/GameplayCamerasEditor/Private/Sequencer/GameplayCameraComponentTrackEditor.h
@@ -4,12 +4,19 @@
 
 #include "ISequencerTrackEditor.h"
 #include "MovieSceneTrackEditor.h"
+#include "Templates/SharedPointerFwd.h"
 
 class ISequencer;
 class UGameplayCameraComponent;
 class UMovieScenePropertyTrack;
 struct FCameraRigParameterDefinition;
 
+namespace UE::Sequencer
+{
+	class FObjectBindingModel;
+	class FObjectBindingPropertyMenuData;
+}
+
 class FGameplayCameraComponentTrackEditor : public FMovieSceneTrackEditor
 {
 public:
@@ -41,11 +48,26 @@ private:
 
 	void OnExtendObjectBindingTrackMenu(FMenuBuilder& MenuBuilder, TArray<FGuid> ObjectBindings);
 
-	void AddCameraParameterTrack(FCameraRigParameterDefinition Definition, TArray<FGuid> ObjectBindings);
-	bool CanAddCameraParameterTrack(FCameraRigParameterDefinition Definition, FGuid ObjectBinding) const;
+private:
+
+	struct FPropertyMenuData
+	{
+		FString MenuName;
+		FPropertyPath PropertyPath;
+		int32 PropertyIndexForMenuName = INDEX_NONE;
+
+		bool operator< (const FPropertyMenuData& Other) const
+		{
+			int32 CompareResult = MenuName.Compare(Other.MenuName);
+			return CompareResult < 0;
+		}
+	};
 
-	void InitializeNewTrack(UMovieScenePropertyTrack* NewTrack, const FCameraRigParameterDefinition& Definition) const;
+	void BuildAddParameterTrackMenuItems(const FGuid& ObjectBinding, FMenuBuilder& MenuBuilder, TArray<FPropertyPath> KeyablePropertyPaths, int32 PropertyNameIndexStart);
+	void BuildAddParameterTrackSubMenuItems(FMenuBuilder& MenuBuilder, FGuid ObjectBinding, TArray<FPropertyPath> KeyablePropertyPaths, int32 PropertyNameIndexStart);
+	void BuildAddParameterTrackMenuItem(FMenuBuilder& MenuBuilder, const FPropertyMenuData& KeyablePropertyMenuData, const FGuid& ObjectBinding);
 
-	TSubclassOf<UMovieScenePropertyTrack> GetParameterTrackFromDefinition(const FCameraRigParameterDefinition& Definition) const;
+	void AddCameraParameterTrack(FPropertyMenuData PropertyMenuData, FGuid ObjectBinding);
+	bool CanAddCameraParameterTrack(FPropertyMenuData PropertyMenuData, FGuid ObjectBinding) const;
 };
 
