Cameras: make camera parameter properties animatable by Sequencer

diff --git a/Engine/Plugins/Cameras/GameplayCameras/Source/GameplayCameras/Private/Core/CameraRigParameterDefinition.cpp b/Engine/Plugins/Cameras/GameplayCameras/Source/GameplayCameras/Private/Core/CameraRigParameterDefinition.cpp
--- a/Engine/Plugins/Cameras/GameplayCameras/Source/GameplayCameras/Private/Core/CameraRigParameterDefinition.cpp
+++ b/Engine/Plugins/Cameras/GameplayCameras/Source/GameplayCameras/Private/Core/CameraRigParameterDefinition.cpp
@@ -26,6 +26,7 @@ void FCameraRigParameterBuilder::AppendDefaultParameterProperties(const UCameraR
 		bool bIsValidProperty = true;
 		EPropertyBagPropertyType PropertyType = EPropertyBagPropertyType::Struct;
 		const UObject* PropertyTypeObject = nullptr;
+		EPropertyFlags PropertyFlags = CPF_None;
 
 		if (Definition.ParameterType == ECameraRigInterfaceParameterType::Blendable)
 		{
@@ -39,6 +40,7 @@ void FCameraRigParameterBuilder::AppendDefaultParameterProperties(const UCameraR
 #undef UE_CAMERA_VARIABLE_FOR_TYPE
 				case ECameraVariableType::BlendableStruct:
 					PropertyTypeObject = Definition.BlendableStructType;
+					PropertyFlags = CPF_Interp;
 					break;
 				default:
 					ensure(false);
@@ -53,13 +55,16 @@ void FCameraRigParameterBuilder::AppendDefaultParameterProperties(const UCameraR
 			{
 				case ECameraContextDataType::Name:
 					PropertyType = EPropertyBagPropertyType::Name;
+					PropertyFlags = CPF_Interp;
 					break;
 				case ECameraContextDataType::String:
 					PropertyType = EPropertyBagPropertyType::String;
+					PropertyFlags = CPF_Interp;
 					break;
 				case ECameraContextDataType::Enum:
 					PropertyType = EPropertyBagPropertyType::Enum;
 					ensure(PropertyTypeObject && PropertyTypeObject->IsA<UEnum>());
+					PropertyFlags = CPF_Interp;
 					break;
 				case ECameraContextDataType::Struct:
 					PropertyType = EPropertyBagPropertyType::Struct;
@@ -67,9 +72,11 @@ void FCameraRigParameterBuilder::AppendDefaultParameterProperties(const UCameraR
 					break;
 				case ECameraContextDataType::Object:
 					PropertyType = EPropertyBagPropertyType::Object;
+					PropertyFlags = CPF_Interp;
 					break;
 				case ECameraContextDataType::Class:
 					PropertyType = EPropertyBagPropertyType::Class;
+					PropertyFlags = CPF_Interp;
 					break;
 				default:
 					bIsValidProperty = false;
@@ -86,6 +93,7 @@ void FCameraRigParameterBuilder::AppendDefaultParameterProperties(const UCameraR
 			FPropertyBagPropertyDesc NewProperty(Definition.ParameterName, PropertyType, PropertyTypeObject);
 			// Make the property bag match the camera interface parameter GUIDs.
 			NewProperty.ID = Definition.ParameterGuid;
+			NewProperty.PropertyFlags |= PropertyFlags;
 
 			OutProperties.Add(NewProperty);
 		}
