Cameras: properly gather old data parameters defined using the magic DataID property suffix

diff --git a/Engine/Plugins/Cameras/GameplayCameras/Source/GameplayCameras/Private/Core/CameraRigAssetBuilder.cpp b/Engine/Plugins/Cameras/GameplayCameras/Source/GameplayCameras/Private/Core/CameraRigAssetBuilder.cpp
--- a/Engine/Plugins/Cameras/GameplayCameras/Source/GameplayCameras/Private/Core/CameraRigAssetBuilder.cpp
+++ b/Engine/Plugins/Cameras/GameplayCameras/Source/GameplayCameras/Private/Core/CameraRigAssetBuilder.cpp
@@ -477,40 +477,60 @@ void FCameraRigAssetBuilder::GatherOldDrivenParameters()
 		
 		for (TFieldIterator<FProperty> It(CameraNodeClass); It; ++It)
 		{
-			FStructProperty* StructProperty = CastField<FStructProperty>(*It);
-			if (!StructProperty)
-			{
-				continue;
-			}
+			FProperty* Property(*It);
 
+			// First look for some blendable camera parameters.
+			if (FStructProperty* StructProperty = CastField<FStructProperty>(Property))
+			{
+				bool bFoundCameraParameter = true;
 #define UE_CAMERA_VARIABLE_FOR_TYPE(ValueType, ValueName)\
-			if (StructProperty->Struct == F##ValueName##CameraParameter::StaticStruct())\
-			{\
-				auto* CameraParameterPtr = StructProperty->ContainerPtrToValuePtr<F##ValueName##CameraParameter>(CameraNode);\
-				if (CameraParameterPtr->VariableID.IsValid() && !CameraParameterPtr->Variable)\
+				if (StructProperty->Struct == F##ValueName##CameraParameter::StaticStruct())\
 				{\
-					OldDrivenBlendableParameters.Add(\
-							FDrivenParameterKey{ StructProperty->GetFName(), CameraNode },\
-							CameraParameterPtr->VariableID);\
-					CameraParameterPtr->VariableID = FCameraVariableID();\
+					auto* CameraParameterPtr = StructProperty->ContainerPtrToValuePtr<F##ValueName##CameraParameter>(CameraNode);\
+					if (CameraParameterPtr->VariableID.IsValid() && !CameraParameterPtr->Variable)\
+					{\
+						OldDrivenBlendableParameters.Add(\
+								FDrivenParameterKey{ StructProperty->GetFName(), CameraNode },\
+								CameraParameterPtr->VariableID);\
+						CameraParameterPtr->VariableID = FCameraVariableID();\
+					}\
 				}\
-			}\
-			else if (StructProperty->Struct == F##ValueName##CameraVariableReference::StaticStruct())\
-			{\
-				auto* VariableReferencePtr = StructProperty->ContainerPtrToValuePtr<F##ValueName##CameraVariableReference>(CameraNode);\
-				if (VariableReferencePtr->VariableID.IsValid() && !VariableReferencePtr->Variable)\
+				else if (StructProperty->Struct == F##ValueName##CameraVariableReference::StaticStruct())\
 				{\
-					OldDrivenBlendableParameters.Add(\
-							FDrivenParameterKey{ StructProperty->GetFName(), CameraNode },\
-							VariableReferencePtr->VariableID);\
-					VariableReferencePtr->VariableID = FCameraVariableID();\
+					auto* VariableReferencePtr = StructProperty->ContainerPtrToValuePtr<F##ValueName##CameraVariableReference>(CameraNode);\
+					if (VariableReferencePtr->VariableID.IsValid() && !VariableReferencePtr->Variable)\
+					{\
+						OldDrivenBlendableParameters.Add(\
+								FDrivenParameterKey{ StructProperty->GetFName(), CameraNode },\
+								VariableReferencePtr->VariableID);\
+						VariableReferencePtr->VariableID = FCameraVariableID();\
+					}\
 				}\
-			}\
-			else
+				else
 UE_CAMERA_VARIABLE_FOR_ALL_TYPES()
 #undef UE_CAMERA_VARIABLE_FOR_TYPE
+				{
+					// Other struct type...
+					bFoundCameraParameter = false;
+				}
+
+				if (bFoundCameraParameter)
+				{
+					continue;
+				}
+			}
+
+			// Then look for some data parameters.
+			const FName DataIDPropertyName = FName(It->GetName() + TEXT("DataID"));
+			FStructProperty* DataIDStructProperty = CastField<FStructProperty>(CameraNodeClass->FindPropertyByName(DataIDPropertyName));
+			if (DataIDStructProperty && DataIDStructProperty->Struct == FCameraContextDataID::StaticStruct())
 			{
-				// Other struct type...
+				FCameraContextDataID* ExistingDataID = DataIDStructProperty->ContainerPtrToValuePtr<FCameraContextDataID>(CameraNode);
+				if (ExistingDataID->IsValid())
+				{
+					OldDrivenDataParameters.Add(FDrivenParameterKey{ Property->GetFName(), CameraNode }, *ExistingDataID);
+					*ExistingDataID = FCameraContextDataID();
+				}
 			}
 		}
 
