Cameras: add more options for framing node targeting

- Add support for automatic and manual target bounds
- Add support for a skeleton bone for when a socket doesn't exist

diff --git a/Engine/Plugins/Cameras/GameplayCameras/Config/BaseGameplayCameras.ini b/Engine/Plugins/Cameras/GameplayCameras/Config/BaseGameplayCameras.ini
--- a/Engine/Plugins/Cameras/GameplayCameras/Config/BaseGameplayCameras.ini
+++ b/Engine/Plugins/Cameras/GameplayCameras/Config/BaseGameplayCameras.ini
@@ -7,5 +7,6 @@
 
 ; UE 5.6 redirects.
 +ClassRedirects=(OldName="/Script/GameplayCameras.CameraRigInterfaceParameter", NewName="CameraRigBlendableParameter")
++StructRedirects=(OldName="/Script/GameplayCameras.CameraTargetInfo",NewName="CameraActorTargetInfo")
 +FunctionRedirects=(OldName="BlueprintCameraDirectorEvaluator.ActivateCameraRigPrefab", NewName="ActivateCameraRig")
 +PropertyRedirects=(OldName="/Script/GameplayCameras.CameraRigInterface.InterfaceParameters", NewName="BlendableParameters")
diff --git a/Engine/Plugins/Cameras/GameplayCameras/Source/GameplayCameras/Private/Math/CameraFramingZoneMath.cpp b/Engine/Plugins/Cameras/GameplayCameras/Source/GameplayCameras/Private/Math/CameraFramingZoneMath.cpp
--- a/Engine/Plugins/Cameras/GameplayCameras/Source/GameplayCameras/Private/Math/CameraFramingZoneMath.cpp
+++ b/Engine/Plugins/Cameras/GameplayCameras/Source/GameplayCameras/Private/Math/CameraFramingZoneMath.cpp
@@ -36,6 +36,33 @@ FFramingZone FFramingZone::FromRelativeMargins(const FVector2d& ScreenLocation,
 	return NewZone;
 }
 
+FFramingZone FFramingZone::FromPoints(TConstArrayView<FVector2d> ScreenPoints)
+{
+	FFramingZone NewZone;
+
+	if (!ensure(ScreenPoints.Num() > 0))
+	{
+		return NewZone;
+	}
+
+	const FVector2d& FirstPoint = ScreenPoints[0];
+	NewZone.LeftBound = FirstPoint.X;
+	NewZone.TopBound = FirstPoint.Y;
+	NewZone.RightBound = FirstPoint.X;
+	NewZone.BottomBound = FirstPoint.Y;
+
+	for (int32 Index = 1; Index < ScreenPoints.Num(); ++Index)
+	{
+		const FVector2d& Point(ScreenPoints[Index]);
+		NewZone.LeftBound = FMath::Min(NewZone.LeftBound, Point.X);
+		NewZone.TopBound = FMath::Min(NewZone.TopBound, Point.Y);
+		NewZone.RightBound = FMath::Max(NewZone.RightBound, Point.X);
+		NewZone.BottomBound = FMath::Max(NewZone.BottomBound, Point.Y);
+	}
+
+	return NewZone;
+}
+
 FFramingZone::FFramingZone()
 {
 	LeftBound = 0;
@@ -44,6 +71,16 @@ FFramingZone::FFramingZone()
 	BottomBound = 1;
 }
 
+bool FFramingZone::IsValid() const
+{
+	return LeftBound >= 0 && LeftBound <= 1 &&
+		TopBound >= 0 && TopBound <= 1 &&
+		RightBound >= 0 && RightBound <= 1 &&
+		BottomBound >= 0 && BottomBound <= 1 &&
+		LeftBound <= RightBound &&
+		TopBound <= BottomBound;
+}
+
 void FFramingZone::ClampBounds()
 {
 	LeftBound = FMath::Clamp(LeftBound, 0.0, 1.0);
diff --git a/Engine/Plugins/Cameras/GameplayCameras/Source/GameplayCameras/Private/Nodes/Framing/BaseFramingCameraNode.cpp b/Engine/Plugins/Cameras/GameplayCameras/Source/GameplayCameras/Private/Nodes/Framing/BaseFramingCameraNode.cpp
--- a/Engine/Plugins/Cameras/GameplayCameras/Source/GameplayCameras/Private/Nodes/Framing/BaseFramingCameraNode.cpp
+++ b/Engine/Plugins/Cameras/GameplayCameras/Source/GameplayCameras/Private/Nodes/Framing/BaseFramingCameraNode.cpp
@@ -54,6 +54,8 @@ void FBaseFramingCameraNodeEvaluator::OnInitialize(const FCameraNodeEvaluatorIni
 {
 	const UBaseFramingCameraNode* BaseFramingNode = GetCameraNodeAs<UBaseFramingCameraNode>();
 
+	Readers.TargetInfo.Initialize(BaseFramingNode->TargetInfo, BaseFramingNode->TargetInfoDataID);
+
 	Readers.IdealFramingLocation.Initialize(BaseFramingNode->IdealFramingLocation);
 	Readers.SetTargetDistance.Initialize(BaseFramingNode->SetTargetDistance);
 
@@ -69,38 +71,48 @@ void FBaseFramingCameraNodeEvaluator::OnInitialize(const FCameraNodeEvaluatorIni
 	Anticipation.History.Reserve(GFramingNumTargetMovementSamples);
 }
 
-TOptional<FVector3d> FBaseFramingCameraNodeEvaluator::AcquireTargetLocation(const FCameraNodeEvaluationParams& Params, const FCameraNodeEvaluationResult& InResult)
+bool FBaseFramingCameraNodeEvaluator::AcquireTargetInfo(const FCameraNodeEvaluationParams& Params, const FCameraNodeEvaluationResult& InResult, FAcquiredTargetInfo& OutInfo)
 {
 	const UBaseFramingCameraNode* FramingNode = GetCameraNodeAs<UBaseFramingCameraNode>();
 	if (FramingNode->TargetLocation.IsValid())
 	{
 		const FVector3d* TargetLocation = FramingNode->TargetLocation.GetValue(InResult.VariableTable);
-		return TargetLocation ? TOptional<FVector3d>(*TargetLocation) : TOptional<FVector3d>();
+		if (TargetLocation)
+		{
+			OutInfo.WorldTransform = FTransform3d(*TargetLocation);
+			OutInfo.LocalBounds = FBoxSphereBounds3d(EForceInit::ForceInit);
+			return true;
+		}
 	}
-	else if (FramingNode->TargetInfo.Value.Actor || FramingNode->TargetInfo.HasOverride())
+	else if (FramingNode->TargetInfo.IsValid() || FramingNode->TargetInfoDataID.IsValid())
 	{
-		const FCameraTargetInfo& TargetInfo = FramingNode->TargetInfo.GetValue(InResult.ContextDataTable);
-
 		FTransform3d TargetTransform;
-		FCameraTargetInfoReader Reader(TargetInfo);
-		if (Reader.GetTargetTransform(TargetTransform))
+		FBoxSphereBounds3d TargetBounds;
+		if (Readers.TargetInfo.GetTargetInfo(InResult.ContextDataTable, TargetTransform, TargetBounds))
 		{
-			return TOptional<FVector3d>(TargetTransform.GetLocation());
+			OutInfo.WorldTransform = TargetTransform;
+			OutInfo.LocalBounds = TargetBounds;
+			return true;
 		}
 	}
 	else if (APlayerController* PlayerController = Params.EvaluationContext->GetPlayerController())
 	{
 		APawn* Pawn = PlayerController->GetPawn();
-		FVector3d TargetLocation = Pawn->GetActorLocation();
-		return TOptional<FVector3d>(TargetLocation);
+		OutInfo.WorldTransform = FTransform3d(Pawn->GetActorLocation());
+		OutInfo.LocalBounds = FBoxSphereBounds3d(EForceInit::ForceInit);
+		if (USceneComponent* RootComponent = Pawn->GetRootComponent())
+		{
+			OutInfo.LocalBounds = RootComponent->Bounds;
+		}
+		return true;
 	}
 
-	return TOptional<FVector3d>();
+	return false;
 }
 
-void FBaseFramingCameraNodeEvaluator::UpdateFramingState(const FCameraNodeEvaluationParams& Params, const FCameraNodeEvaluationResult& OutResult, const FVector3d& TargetLocation, const FTransform3d& LastFraming)
+void FBaseFramingCameraNodeEvaluator::UpdateFramingState(const FCameraNodeEvaluationParams& Params, const FCameraNodeEvaluationResult& OutResult, const FAcquiredTargetInfo& TargetInfo, const FTransform3d& LastFraming)
 {
-	ComputeCurrentState(Params, OutResult, TargetLocation, LastFraming);
+	ComputeCurrentState(Params, OutResult, TargetInfo, LastFraming);
 	ComputeDesiredState(Params.DeltaTime);
 }
 
@@ -114,7 +126,7 @@ void FBaseFramingCameraNodeEvaluator::EndFramingUpdate(const FCameraNodeEvaluati
 	}
 }
 
-void FBaseFramingCameraNodeEvaluator::ComputeCurrentState(const FCameraNodeEvaluationParams& Params, const FCameraNodeEvaluationResult& OutResult, const FVector3d& TargetLocation, const FTransform3d& LastFraming)
+void FBaseFramingCameraNodeEvaluator::ComputeCurrentState(const FCameraNodeEvaluationParams& Params, const FCameraNodeEvaluationResult& OutResult, const FAcquiredTargetInfo& TargetInfo, const FTransform3d& LastFraming)
 {
 	// Get screen-space coordinates of the ideal framing point. These are in 0..1 UI space.
 	State.IdealTarget = Readers.IdealFramingLocation.Get(OutResult.VariableTable);
@@ -140,7 +152,7 @@ void FBaseFramingCameraNodeEvaluator::ComputeCurrentState(const FCameraNodeEvalu
 	State.SoftZone.ClampBounds(State.DeadZone);
 
 	// See if we need to extrapolate where the target will be in "anticipation time" seconds.
-	const FVector3d EffectiveTargetLocation = ComputeAnticipatedScreenTarget(Params.DeltaTime, TargetLocation);
+	const FVector3d EffectiveTargetLocation = ComputeAnticipatedScreenTarget(Params.DeltaTime, TargetInfo.WorldTransform.GetLocation());
 
 	// We are going to reframe things iteratively, so we'll use a temporary pose defined by last frame's
 	// shot transform.
@@ -152,8 +164,33 @@ void FBaseFramingCameraNodeEvaluator::ComputeCurrentState(const FCameraNodeEvalu
 	const double AspectRatio = FCameraPoseMath::GetEffectiveAspectRatio(TempPose, PlayerController);
 	const TOptional<FVector2d> ScreenTarget = FCameraPoseMath::ProjectWorldToScreen(TempPose, AspectRatio, EffectiveTargetLocation, true);
 	State.WorldTarget = EffectiveTargetLocation;
+	State.LocalTargetBounds = TargetInfo.LocalBounds;
 	State.ScreenTarget = ScreenTarget.Get(FVector2d(0.5, 0.5));
 
+	// Get the target's bounds in screen-space.
+	FTransform3d EffectiveTargetTransform(TargetInfo.WorldTransform.GetRotation(), EffectiveTargetLocation);
+	State.ScreenTargetBounds = ComputeScreenTargetBounds(TempPose, AspectRatio, EffectiveTargetTransform, TargetInfo.LocalBounds);
+
+	// Compute the effective dead-zone, which is the subset of the dead-zone that encompasses as much
+	// of the target's bound as possible.
+	FFramingZone EffectiveDeadZone = State.DeadZone;
+	const FVector2d DeadZoneCenter = State.DeadZone.Center();
+	const double TargetBoundsHalfWidth = State.ScreenTargetBounds.Width() / 2.0;
+	const double TargetBoundsHalfHeight = State.ScreenTargetBounds.Height() / 2.0;
+	EffectiveDeadZone.LeftBound += TargetBoundsHalfWidth;
+	EffectiveDeadZone.TopBound += TargetBoundsHalfHeight;
+	EffectiveDeadZone.RightBound -= TargetBoundsHalfWidth;
+	EffectiveDeadZone.BottomBound -= TargetBoundsHalfHeight;
+	if (EffectiveDeadZone.LeftBound > EffectiveDeadZone.RightBound)
+	{
+		EffectiveDeadZone.LeftBound = EffectiveDeadZone.RightBound = DeadZoneCenter.X;
+	}
+	if (EffectiveDeadZone.TopBound > EffectiveDeadZone.BottomBound)
+	{
+		EffectiveDeadZone.TopBound = EffectiveDeadZone.BottomBound = DeadZoneCenter.Y;
+	}
+	State.EffectiveDeadZone = EffectiveDeadZone;
+
 #if UE_GAMEPLAY_CAMERAS_DEBUG
 	State.DeadZoneEdgePoint = State.ScreenTarget;
 	State.HardZoneEdgePoint = State.ScreenTarget;
@@ -198,7 +235,7 @@ void FBaseFramingCameraNodeEvaluator::ComputeCurrentState(const FCameraNodeEvalu
 
 	const bool bWasReframing = State.bIsReframingTarget;
 	const bool bIsInSoftZone = State.SoftZone.Contains(State.ScreenTarget);
-	const bool bIsInDeadZone = State.DeadZone.Contains(State.ScreenTarget);
+	const bool bIsInDeadZone = State.EffectiveDeadZone.Contains(State.ScreenTarget);
 	if (!ScreenTarget.IsSet() || !bIsInSoftZone)
 	{
 		// Target is out of view or outside the soft zone -- it's therefore in the hard zone and we will
@@ -332,6 +369,29 @@ FVector3d FBaseFramingCameraNodeEvaluator::ComputeAnticipatedScreenTarget(float
 	}
 }
 
+FFramingZone FBaseFramingCameraNodeEvaluator::ComputeScreenTargetBounds(const FCameraPose& CameraPose, float AspectRatio, const FTransform3d& TargetTransform, const FBoxSphereBounds3d& LocalBounds)
+{
+	const FVector3d BoxExtent = LocalBounds.BoxExtent;
+	FVector3d BoxCorners[8];
+	BoxCorners[0] = TargetTransform.TransformPositionNoScale({ BoxExtent.X, BoxExtent.Y, BoxExtent.Z });
+	BoxCorners[1] = TargetTransform.TransformPositionNoScale({ -BoxExtent.X, BoxExtent.Y, BoxExtent.Z });
+	BoxCorners[2] = TargetTransform.TransformPositionNoScale({ BoxExtent.X, -BoxExtent.Y, BoxExtent.Z });
+	BoxCorners[3] = TargetTransform.TransformPositionNoScale({ -BoxExtent.X, -BoxExtent.Y, BoxExtent.Z });
+	BoxCorners[4] = TargetTransform.TransformPositionNoScale({ BoxExtent.X, BoxExtent.Y, -BoxExtent.Z });
+	BoxCorners[5] = TargetTransform.TransformPositionNoScale({ -BoxExtent.X, BoxExtent.Y, -BoxExtent.Z });
+	BoxCorners[6] = TargetTransform.TransformPositionNoScale({ BoxExtent.X, -BoxExtent.Y, -BoxExtent.Z });
+	BoxCorners[7] = TargetTransform.TransformPositionNoScale({ -BoxExtent.X, -BoxExtent.Y, -BoxExtent.Z });
+
+	FVector2d ScreenBoxCorners[8];
+	for (int32 Index = 0; Index < 8; ++Index)
+	{
+		TOptional<FVector2d> ScreenCorner = FCameraPoseMath::ProjectWorldToScreen(CameraPose, AspectRatio, BoxCorners[Index], true);
+		ScreenBoxCorners[Index] = ScreenCorner.GetValue();
+	}
+
+	return FFramingZone::FromPoints(ScreenBoxCorners);
+}
+
 void FBaseFramingCameraNodeEvaluator::ComputeDesiredState(float DeltaTime)
 {
 	// If we  don't have any reframing to do, bail out.
@@ -415,6 +475,16 @@ void FBaseFramingCameraNodeEvaluator::DrawFramingState(const FState& State, cons
 				State.DeadZone.GetCanvasSize(CanvasSize),
 				FLinearColor::Green,
 				1.f);
+		Renderer.Draw2DBox(
+				State.EffectiveDeadZone.GetCanvasPosition(CanvasSize), 
+				State.EffectiveDeadZone.GetCanvasSize(CanvasSize),
+				FColorList::LightGrey,
+				1.f);
+		Renderer.Draw2DBox(
+				State.ScreenTargetBounds.GetCanvasPosition(CanvasSize),
+				State.ScreenTargetBounds.GetCanvasSize(CanvasSize),
+				FLinearColor::Gray,
+				1.f);
 
 		if (State.LowReframeDampingFactor > 0)
 		{
@@ -501,6 +571,8 @@ void FBaseFramingCameraDebugBlock::OnDebugDraw(const FCameraDebugBlockDrawParams
 	Renderer.RemoveIndent();
 
 	FBaseFramingCameraNodeEvaluator::DrawFramingState(State, Desired, Renderer);
+
+	Renderer.DrawBox(State.WorldTarget, State.LocalTargetBounds.BoxExtent, FLinearColor::Yellow, 0.5f);
 }
 
 #endif  // UE_GAMEPLAY_CAMERAS_DEBUG
@@ -560,13 +632,6 @@ UBaseFramingCameraNode::UBaseFramingCameraNode(const FObjectInitializer& ObjectI
 
 void UBaseFramingCameraNode::GetCustomCameraNodeParameters(FCustomCameraNodeParameterInfos& OutParameterInfos)
 {
-	OutParameterInfos.AddDataParameter(
-			GET_MEMBER_NAME_CHECKED(UBaseFramingCameraNode, TargetInfo),
-			ECameraContextDataType::Struct,
-			FCameraTargetInfo::StaticStruct(),
-			reinterpret_cast<const uint8*>(&TargetInfo.Value),
-			&TargetInfo.DataID);
-
 	OutParameterInfos.AddBlendableParameter(
 			GET_MEMBER_NAME_CHECKED(UBaseFramingCameraNode, DeadZone),
 			ECameraVariableType::BlendableStruct,
diff --git a/Engine/Plugins/Cameras/GameplayCameras/Source/GameplayCameras/Private/Nodes/Framing/CameraActorTargetInfo.cpp b/Engine/Plugins/Cameras/GameplayCameras/Source/GameplayCameras/Private/Nodes/Framing/CameraActorTargetInfo.cpp
new file mode 100644
--- /dev/null
+++ b/Engine/Plugins/Cameras/GameplayCameras/Source/GameplayCameras/Private/Nodes/Framing/CameraActorTargetInfo.cpp
@@ -0,0 +1,117 @@
+// Copyright Epic Games, Inc. All Rights Reserved.
+
+#include "Nodes/Framing/CameraActorTargetInfo.h"
+
+#include "Components/SkeletalMeshComponent.h"
+#include "Core/CameraContextDataTable.h"
+#include "GameFramework/Actor.h"
+
+namespace UE::Cameras
+{
+
+FCameraActorTargetInfoReader::FCameraActorTargetInfoReader(const FCameraActorTargetInfo& InTargetInfo, FCameraContextDataID InDataID)
+{
+	Initialize(InTargetInfo, InDataID);
+}
+
+void FCameraActorTargetInfoReader::Initialize(const FCameraActorTargetInfo& InTargetInfo, FCameraContextDataID InDataID)
+{
+	DefaultTargetInfo = InTargetInfo;
+	DataID = InDataID;
+
+	CacheTargetInfo(DefaultTargetInfo);
+}
+
+void FCameraActorTargetInfoReader::CacheTargetInfo(const FCameraActorTargetInfo& InTargetInfo)
+{
+	if (CachedTargetInfo != InTargetInfo)
+	{
+		CachedTargetInfo = InTargetInfo;
+
+		CachedSkeletalMeshComponent = nullptr;
+		if (InTargetInfo.Actor && (!InTargetInfo.SocketName.IsNone() || !InTargetInfo.BoneName.IsNone()))
+		{
+			CachedSkeletalMeshComponent = InTargetInfo.Actor->FindComponentByClass<USkeletalMeshComponent>();
+		}
+
+		CachedBoneName = NAME_None;
+		CachedParentBoneName = NAME_None;
+		if (CachedSkeletalMeshComponent)
+		{
+			CachedBoneName = InTargetInfo.BoneName;
+			if (!InTargetInfo.SocketName.IsNone())
+			{
+				CachedBoneName = CachedSkeletalMeshComponent->GetSocketBoneName(InTargetInfo.SocketName);
+			}
+			if (!CachedBoneName.IsNone())
+			{
+				CachedParentBoneName = CachedSkeletalMeshComponent->GetParentBone(CachedBoneName);
+			}
+		}
+	}
+}
+
+bool FCameraActorTargetInfoReader::GetTargetInfo(const FCameraContextDataTable& ContextDataTable, FTransform3d& OutTransform, FBoxSphereBounds3d& OutBounds)
+{
+	if (DataID.IsValid())
+	{
+		if (const FCameraActorTargetInfo* NewTargetInfo = ContextDataTable.TryGetData<FCameraActorTargetInfo>(DataID))
+		{
+			CacheTargetInfo(*NewTargetInfo);
+		}
+		else
+		{
+			CacheTargetInfo(DefaultTargetInfo);
+		}
+	}
+
+	if (CachedSkeletalMeshComponent && !CachedBoneName.IsNone())
+	{
+		OutTransform = CachedSkeletalMeshComponent->GetBoneTransform(CachedBoneName);
+		ComputeTargetBounds(OutTransform.GetLocation(), OutBounds);
+		return true;
+	}
+	else if (CachedTargetInfo.Actor)
+	{
+		OutTransform = CachedTargetInfo.Actor->GetTransform();
+		ComputeTargetBounds(OutTransform.GetLocation(), OutBounds);
+		return true;
+	}
+	return false;
+}
+
+void FCameraActorTargetInfoReader::ComputeTargetBounds(const FVector3d& TargetLocation, FBoxSphereBounds3d& OutBounds)
+{
+	switch (CachedTargetInfo.TargetShape)
+	{
+		case ECameraTargetShape::Point:
+			OutBounds = FBoxSphereBounds3d(EForceInit::ForceInit);
+			break;
+		case ECameraTargetShape::AutomaticBounds:
+			if (CachedSkeletalMeshComponent && !CachedParentBoneName.IsNone())
+			{
+				const FVector3d ParentBoneLocation = CachedSkeletalMeshComponent->GetBoneLocation(CachedParentBoneName);
+				const FVector3d ParentToBone((TargetLocation - ParentBoneLocation).GetAbs());
+				const float ParentToBoneLength(ParentToBone.Length());
+				OutBounds = FBoxSphereBounds3d(FVector3d::ZeroVector, ParentToBone, ParentToBoneLength);
+			}
+			else
+			{
+				OutBounds = FBoxSphereBounds3d(EForceInit::ForceInit);
+				if (USceneComponent* RootComponent = CachedTargetInfo.Actor->GetRootComponent())
+				{
+					OutBounds = RootComponent->Bounds;
+				}
+			}
+			break;
+		case ECameraTargetShape::ManualBounds:
+			{
+				float TargetSize = FMath::Max(0.f, CachedTargetInfo.TargetSize);
+				OutBounds = FBoxSphereBounds3d(FVector3d::ZeroVector, FVector3d(TargetSize), TargetSize);
+			}
+			break;
+	}
+}
+
+}  // namespace UE::Cameras
+
diff --git a/Engine/Plugins/Cameras/GameplayCameras/Source/GameplayCameras/Private/Nodes/Framing/CameraTargetInfo.cpp b/Engine/Plugins/Cameras/GameplayCameras/Source/GameplayCameras/Private/Nodes/Framing/CameraTargetInfo.cpp
deleted file mode 100644
diff --git a/Engine/Plugins/Cameras/GameplayCameras/Source/GameplayCameras/Private/Nodes/Framing/DollyFramingCameraNode.cpp b/Engine/Plugins/Cameras/GameplayCameras/Source/GameplayCameras/Private/Nodes/Framing/DollyFramingCameraNode.cpp
--- a/Engine/Plugins/Cameras/GameplayCameras/Source/GameplayCameras/Private/Nodes/Framing/DollyFramingCameraNode.cpp
+++ b/Engine/Plugins/Cameras/GameplayCameras/Source/GameplayCameras/Private/Nodes/Framing/DollyFramingCameraNode.cpp
@@ -75,16 +75,15 @@ void FDollyFramingCameraNodeEvaluator::OnInitialize(const FCameraNodeEvaluatorIn
 
 void FDollyFramingCameraNodeEvaluator::OnRun(const FCameraNodeEvaluationParams& Params, FCameraNodeEvaluationResult& OutResult)
 {
-	TOptional<FVector3d> OptTargetLocation = AcquireTargetLocation(Params, OutResult);
-	if (!OptTargetLocation.IsSet())
+	FAcquiredTargetInfo TargetInfo;
+	if (!AcquireTargetInfo(Params, OutResult, TargetInfo))
 	{
 		return;
 	}
 
 	// Let the base class figure out all the screen-space framing stuff.
-	const FVector3d TargetLocation = OptTargetLocation.GetValue();
 	const FTransform3d LastShotTransform = BuildDollyShotTransform(OutResult.CameraPose);
-	UpdateFramingState(Params, OutResult, TargetLocation, LastShotTransform);
+	UpdateFramingState(Params, OutResult, TargetInfo, LastShotTransform);
 
 	// If we need to reframe the target this tick, figure out how much we need to move the dolly
 	// to accomplish that.
diff --git a/Engine/Plugins/Cameras/GameplayCameras/Source/GameplayCameras/Private/Nodes/Framing/PanningFramingCameraNode.cpp b/Engine/Plugins/Cameras/GameplayCameras/Source/GameplayCameras/Private/Nodes/Framing/PanningFramingCameraNode.cpp
--- a/Engine/Plugins/Cameras/GameplayCameras/Source/GameplayCameras/Private/Nodes/Framing/PanningFramingCameraNode.cpp
+++ b/Engine/Plugins/Cameras/GameplayCameras/Source/GameplayCameras/Private/Nodes/Framing/PanningFramingCameraNode.cpp
@@ -69,16 +69,15 @@ void FPanningFramingCameraNodeEvaluator::OnInitialize(const FCameraNodeEvaluator
 
 void FPanningFramingCameraNodeEvaluator::OnRun(const FCameraNodeEvaluationParams& Params, FCameraNodeEvaluationResult& OutResult)
 {
-	TOptional<FVector3d> OptTargetLocation = AcquireTargetLocation(Params, OutResult);
-	if (!OptTargetLocation.IsSet())
+	FAcquiredTargetInfo TargetInfo;
+	if (!AcquireTargetInfo(Params, OutResult, TargetInfo))
 	{
 		return;
 	}
 
 	// Let the base class figure out all the screen-space framing stuff.
-	const FVector3d TargetLocation = OptTargetLocation.GetValue();
 	const FTransform3d LastShotTransform = BuildPanningShotTransform(OutResult.CameraPose);
-	UpdateFramingState(Params, OutResult, TargetLocation, LastShotTransform);
+	UpdateFramingState(Params, OutResult, TargetInfo, LastShotTransform);
 
 	// If we need to reframe the target this tick, figure out how much we need to rotate the camera
 	// to accomplish that.
diff --git a/Engine/Plugins/Cameras/GameplayCameras/Source/GameplayCameras/Public/Core/CameraContextDataReader.h b/Engine/Plugins/Cameras/GameplayCameras/Source/GameplayCameras/Public/Core/CameraContextDataReader.h
--- a/Engine/Plugins/Cameras/GameplayCameras/Source/GameplayCameras/Public/Core/CameraContextDataReader.h
+++ b/Engine/Plugins/Cameras/GameplayCameras/Source/GameplayCameras/Public/Core/CameraContextDataReader.h
@@ -4,12 +4,11 @@
 
 #include "Core/CameraContextDataTable.h"
 #include "CoreTypes.h"
-#include "Templates/UnrealTypeTraits.h"
 
 namespace UE::Cameras
 {
 
-template<typename DataType, ECameraContextDataType DataTypeEnum>
+template<typename DataType>
 struct TCameraContextDataReader
 {
 	TCameraContextDataReader() {}
diff --git a/Engine/Plugins/Cameras/GameplayCameras/Source/GameplayCameras/Public/Math/CameraFramingZoneMath.h b/Engine/Plugins/Cameras/GameplayCameras/Source/GameplayCameras/Public/Math/CameraFramingZoneMath.h
--- a/Engine/Plugins/Cameras/GameplayCameras/Source/GameplayCameras/Public/Math/CameraFramingZoneMath.h
+++ b/Engine/Plugins/Cameras/GameplayCameras/Source/GameplayCameras/Public/Math/CameraFramingZoneMath.h
@@ -2,8 +2,8 @@
 
 #pragma once
 
-#include "CoreTypes.h"
-#include "Math/MathFwd.h"
+#include "Containers/ContainersFwd.h"
+#include "Math/UnrealMath.h"
 
 class FArchive;
 struct FCameraFramingZone;
@@ -29,6 +29,21 @@ struct FFramingZone
 	/** Builds an empty framing zone. */
 	FFramingZone();
 
+	/** Gets the width of the framing zone. */
+	double Width() const { return FMath::Max(0.0, RightBound - LeftBound); }
+
+	/** Gets the height of the framing zone. */
+	double Height() const { return FMath::Max(0.0, BottomBound - TopBound); }
+
+	/** Gets the center of the framing zone. */
+	FVector2d Center() const { return FVector2d((LeftBound + RightBound) / 2.0, (TopBound + BottomBound) / 2.0); }
+
+	/** 
+	 * Returns whether this framing zone is well-formed (bounds in 0..1 UI space, 
+	 * with left/top bounds lesser than or equal to right/bottom bounds).
+	 */
+	bool IsValid() const;
+
 	/** Makes sure all the bounds have valid values between 0 and 1. */
 	void ClampBounds();
 
@@ -71,6 +86,9 @@ public:
 	/** Build a framing zone from a set of margins relative to a screen location. */
 	static FFramingZone FromRelativeMargins(const FVector2d& ScreenLocation, const FCameraFramingZone& Margins);
 
+	/** Build a framing zone that encompasses all the given points (in 0..1 UI space). */
+	static FFramingZone FromPoints(TConstArrayView<FVector2d> ScreenPoints);
+
 private:
 
 	static double GetNormalizedBound(double Bound);
diff --git a/Engine/Plugins/Cameras/GameplayCameras/Source/GameplayCameras/Public/Nodes/Framing/BaseFramingCameraNode.h b/Engine/Plugins/Cameras/GameplayCameras/Source/GameplayCameras/Public/Nodes/Framing/BaseFramingCameraNode.h
--- a/Engine/Plugins/Cameras/GameplayCameras/Source/GameplayCameras/Public/Nodes/Framing/BaseFramingCameraNode.h
+++ b/Engine/Plugins/Cameras/GameplayCameras/Source/GameplayCameras/Public/Nodes/Framing/BaseFramingCameraNode.h
@@ -8,7 +8,7 @@
 #include "Core/CameraVariableReferences.h"
 #include "Core/ICustomCameraNodeParameterProvider.h"
 #include "Nodes/Framing/CameraFramingZone.h"
-#include "Nodes/Framing/CameraTargetInfo.h"
+#include "Nodes/Framing/CameraActorTargetInfo.h"
 #include "Math/CameraFramingZoneMath.h"
 #include "Math/CriticalDamper.h"
 
@@ -35,8 +35,11 @@ public:
 	UPROPERTY(EditAnywhere, Category="Target")
 	FVector3dCameraVariableReference TargetLocation;
 
-	UPROPERTY(EditAnywhere, Category="Target")
-	FCameraTargetInfoParameter TargetInfo;
+	UPROPERTY(EditAnywhere, Category="Target", meta=(CameraContextData=true))
+	FCameraActorTargetInfo TargetInfo;
+
+	UPROPERTY()
+	FCameraContextDataID TargetInfoDataID;
 
 	/**
 	 * Whether the camera pose's target distance should be set to the distance between
@@ -155,11 +158,12 @@ protected:
 
 	struct FState;
 	struct FDesired;
+	struct FAcquiredTargetInfo;
 
 	/** Gets the target location. */
-	TOptional<FVector3d> AcquireTargetLocation(const FCameraNodeEvaluationParams& Params, const FCameraNodeEvaluationResult& InResult);
+	bool AcquireTargetInfo(const FCameraNodeEvaluationParams& Params, const FCameraNodeEvaluationResult& InResult, FAcquiredTargetInfo& OutInfo);
 	/** Updates the framing state for the current tick, see State member field. */
-	void UpdateFramingState(const FCameraNodeEvaluationParams& Params, const FCameraNodeEvaluationResult& OutResult, const FVector3d& TargetLocation, const FTransform3d& LastFraming);
+	void UpdateFramingState(const FCameraNodeEvaluationParams& Params, const FCameraNodeEvaluationResult& OutResult, const FAcquiredTargetInfo& TargetInfo, const FTransform3d& LastFraming);
 	/** Wraps-up the update with optional operations. */
 	void EndFramingUpdate(const FCameraNodeEvaluationParams& Params, FCameraNodeEvaluationResult& OutResult);
 
@@ -171,8 +175,9 @@ private:
 
 	FVector2d GetHardReframeCoords() const;
 
-	void ComputeCurrentState(const FCameraNodeEvaluationParams& Params, const FCameraNodeEvaluationResult& OutResult, const FVector3d& TargetLocation, const FTransform3d& LastFraming);
+	void ComputeCurrentState(const FCameraNodeEvaluationParams& Params, const FCameraNodeEvaluationResult& OutResult, const FAcquiredTargetInfo& TargetInfo, const FTransform3d& LastFraming);
 	FVector3d ComputeAnticipatedScreenTarget(float DeltaTime, const FVector3d& InTargetLocation);
+	FFramingZone ComputeScreenTargetBounds(const FCameraPose& CameraPose, float AspectRatio, const FTransform3d& TargetTransform, const FBoxSphereBounds3d& LocalBounds);
 	void ComputeDesiredState(float DeltaTime);
 
 protected:
@@ -197,9 +202,17 @@ protected:
 		InHardZone
 	};
 
+	struct FAcquiredTargetInfo
+	{
+		FTransform3d WorldTransform;
+		FBoxSphereBounds3d LocalBounds;
+	};
+
 	/** Utility structure for all the parameter readers we need every frame. */
 	struct FReaders
 	{
+		FCameraActorTargetInfoReader TargetInfo;
+
 		TCameraParameterReader<FVector2d> IdealFramingLocation;
 		TCameraParameterReader<bool> SetTargetDistance;
 
@@ -219,6 +232,8 @@ protected:
 	{
 		/** World position of the tracked target. */
 		FVector3d WorldTarget;
+		/** Bounds of the tracked target. */
+		FBoxSphereBounds3d LocalTargetBounds;
 
 		/** Screen-space position of the ideal framing position. */
 		FVector2d IdealTarget;
@@ -238,6 +253,10 @@ protected:
 		float ToggleEngageAlpha;
 		/** Current look-ahead time for anticipating target movement */
 		float TargetMovementAnticipationTime;
+		/** Current target bounds zone. */
+		FFramingZone ScreenTargetBounds;
+		/** Dead zone minus the screen target bounds. */
+		FFramingZone EffectiveDeadZone;
 		/** Current coordinates of the dead zone. */
 		FFramingZone DeadZone;
 		/** Current coordinates of the soft zone. */
diff --git a/Engine/Plugins/Cameras/GameplayCameras/Source/GameplayCameras/Public/Nodes/Framing/CameraActorTargetInfo.h b/Engine/Plugins/Cameras/GameplayCameras/Source/GameplayCameras/Public/Nodes/Framing/CameraActorTargetInfo.h
new file mode 100644
--- /dev/null
+++ b/Engine/Plugins/Cameras/GameplayCameras/Source/GameplayCameras/Public/Nodes/Framing/CameraActorTargetInfo.h
@@ -0,0 +1,95 @@
+// Copyright Epic Games, Inc. All Rights Reserved.
+
+#pragma once
+
+#include "CoreTypes.h"
+#include "Core/CameraContextDataTableFwd.h"
+
+#include "CameraActorTargetInfo.generated.h"
+
+class AActor;
+
+namespace UE::Cameras
+{
+	class FCameraContextDataTable;
+}
+
+UENUM(BlueprintType)
+enum class ECameraTargetShape : uint8
+{
+	/** Use a single-point target. */
+	Point,
+	/** Use the target's computed bounds. */
+	AutomaticBounds,
+	/** Use custom bounds. */
+	ManualBounds
+};
+
+/**
+ * Targeting information for a camera rig.
+ */
+USTRUCT(BlueprintType)
+struct GAMEPLAYCAMERAS_API FCameraActorTargetInfo
+{
+	GENERATED_BODY()
+
+	/** The actor to target. */
+	UPROPERTY(EditAnywhere, Interp, Category="Target")
+	TObjectPtr<AActor> Actor;
+
+	/** An optional socket to target on the actor. */
+	UPROPERTY(EditAnywhere, Interp, Category="Target")
+	FName SocketName;
+
+	/** An optional bone to target on the actor. */
+	UPROPERTY(EditAnywhere, Interp, Category="Target")
+	FName BoneName;
+
+	/** Specifies the shape of the target. */
+	UPROPERTY(EditAnywhere, Interp, Category="Target")
+	ECameraTargetShape TargetShape = ECameraTargetShape::Point;
+
+	/** Specifies the size of target. Only used if TargetShape is set to manual bounds. */
+	UPROPERTY(EditAnywhere, Interp, Category="Target")
+	float TargetSize = 10.f;
+
+	/** The weight of this target. Unused when only one target is used. */
+	UPROPERTY(EditAnywhere, Interp, Category="Target")
+	float Weight = 1.f;
+
+	bool IsValid() const { return Actor != nullptr; }
+
+	bool operator== (const FCameraActorTargetInfo& Other) const = default;
+};
+
+namespace UE::Cameras
+{
+
+/** A special reader class for targeting information. */
+struct FCameraActorTargetInfoReader
+{
+	FCameraActorTargetInfoReader() {}
+	FCameraActorTargetInfoReader(const FCameraActorTargetInfo& InTargetInfo, FCameraContextDataID InDataID);
+
+	void Initialize(const FCameraActorTargetInfo& InTargetInfo, FCameraContextDataID InDataID);
+
+	bool GetTargetInfo(const FCameraContextDataTable& ContextDataTable, FTransform3d& OutTransform, FBoxSphereBounds3d& OutBounds);
+
+private:
+
+	void CacheTargetInfo(const FCameraActorTargetInfo& InTargetInfo);
+	void ComputeTargetBounds(const FVector3d& TargetLocation, FBoxSphereBounds3d& OutBounds);
+
+private:
+
+	FCameraActorTargetInfo DefaultTargetInfo;
+	FCameraContextDataID DataID;
+
+	FCameraActorTargetInfo CachedTargetInfo;
+	const USkeletalMeshComponent* CachedSkeletalMeshComponent = nullptr;
+	FName CachedBoneName;
+	FName CachedParentBoneName;
+};
+
+}  // namespace UE::Cameras
+
diff --git a/Engine/Plugins/Cameras/GameplayCameras/Source/GameplayCameras/Public/Nodes/Framing/CameraTargetInfo.h b/Engine/Plugins/Cameras/GameplayCameras/Source/GameplayCameras/Public/Nodes/Framing/CameraTargetInfo.h
deleted file mode 100644
