Cameras: fix doc/comment on dampen position node

diff --git a/Engine/Plugins/Cameras/GameplayCameras/Source/GameplayCameras/Public/Nodes/Common/DampenPositionCameraNode.h b/Engine/Plugins/Cameras/GameplayCameras/Source/GameplayCameras/Public/Nodes/Common/DampenPositionCameraNode.h
--- a/Engine/Plugins/Cameras/GameplayCameras/Source/GameplayCameras/Public/Nodes/Common/DampenPositionCameraNode.h
+++ b/Engine/Plugins/Cameras/GameplayCameras/Source/GameplayCameras/Public/Nodes/Common/DampenPositionCameraNode.h
@@ -9,7 +9,7 @@
 #include "DampenPositionCameraNode.generated.h"
 
 /**
- * A camera node that offsets the location of the camera.
+ * A camera node that smoothes the movement of the camera.
  */
 UCLASS(MinimalAPI, meta=(CameraNodeCategories="Common,Transform"))
 class UDampenPositionCameraNode : public UCameraNode
