Cameras: add attachment node

diff --git a/Engine/Plugins/Cameras/GameplayCameras/Source/GameplayCameras/Private/Nodes/Attach/AttachToActorCameraNode.cpp b/Engine/Plugins/Cameras/GameplayCameras/Source/GameplayCameras/Private/Nodes/Attach/AttachToActorCameraNode.cpp
new file mode 100644
--- /dev/null
+++ b/Engine/Plugins/Cameras/GameplayCameras/Source/GameplayCameras/Private/Nodes/Attach/AttachToActorCameraNode.cpp
@@ -0,0 +1,69 @@
+// Copyright Epic Games, Inc. All Rights Reserved.
+
+#include "Nodes/Attach/AttachToActorCameraNode.h"
+
+#include "Core/CameraParameterReader.h"
+
+#include UE_INLINE_GENERATED_CPP_BY_NAME(AttachToActorCameraNode)
+
+namespace UE::Cameras
+{
+
+class FAttachToActorCameraNodeEvaluator : public FCameraNodeEvaluator
+{
+	UE_DECLARE_CAMERA_NODE_EVALUATOR(GAMEPLAYCAMERAS_API, FAttachToActorCameraNodeEvaluator)
+
+protected:
+
+	virtual void OnInitialize(const FCameraNodeEvaluatorInitializeParams& Params, FCameraNodeEvaluationResult& OutResult) override;
+	virtual void OnRun(const FCameraNodeEvaluationParams& Params, FCameraNodeEvaluationResult& OutResult) override;
+
+private:
+
+	FCameraActorAttachmentInfoReader AttachmentReader;
+	TCameraParameterReader<bool> AttachToLocationReader;
+	TCameraParameterReader<bool> AttachToRotationReader;
+};
+
+UE_DEFINE_CAMERA_NODE_EVALUATOR(FAttachToActorCameraNodeEvaluator)
+
+void FAttachToActorCameraNodeEvaluator::OnInitialize(const FCameraNodeEvaluatorInitializeParams& Params, FCameraNodeEvaluationResult& OutResult)
+{
+	const UAttachToActorCameraNode* AttachNode = GetCameraNodeAs<UAttachToActorCameraNode>();
+	AttachmentReader.Initialize(AttachNode->Attachment, AttachNode->AttachmentDataID);
+	AttachToLocationReader.Initialize(AttachNode->AttachToLocation);
+	AttachToRotationReader.Initialize(AttachNode->AttachToRotation);
+}
+
+void FAttachToActorCameraNodeEvaluator::OnRun(const FCameraNodeEvaluationParams& Params, FCameraNodeEvaluationResult& OutResult)
+{
+	FTransform3d AttachTransform;
+	if (!AttachmentReader.GetAttachmentTransform(OutResult.ContextDataTable, AttachTransform))
+	{
+		return;
+	}
+
+	const bool bAttachToLocation = AttachToLocationReader.Get(OutResult.VariableTable);
+	const bool bAttachToRotation = AttachToRotationReader.Get(OutResult.VariableTable);
+	
+	if (bAttachToLocation)
+	{
+		const FVector3d AttachLocation = AttachTransform.GetLocation();
+		OutResult.CameraPose.SetLocation(AttachLocation);
+	}
+
+	if (bAttachToRotation)
+	{
+		const FRotator3d AttachRotation = AttachTransform.GetRotation().Rotator();
+		OutResult.CameraPose.SetRotation(AttachRotation);
+	}
+}
+
+}  // namespace UE::Cameras
+
+FCameraNodeEvaluatorPtr UAttachToActorCameraNode::OnBuildEvaluator(FCameraNodeEvaluatorBuilder& Builder) const
+{
+	using namespace UE::Cameras;
+	return Builder.BuildEvaluator<FAttachToActorCameraNodeEvaluator>();
+}
+
diff --git a/Engine/Plugins/Cameras/GameplayCameras/Source/GameplayCameras/Private/Nodes/Attach/CameraActorAttachmentInfo.cpp b/Engine/Plugins/Cameras/GameplayCameras/Source/GameplayCameras/Private/Nodes/Attach/CameraActorAttachmentInfo.cpp
new file mode 100644
--- /dev/null
+++ b/Engine/Plugins/Cameras/GameplayCameras/Source/GameplayCameras/Private/Nodes/Attach/CameraActorAttachmentInfo.cpp
@@ -0,0 +1,106 @@
+// Copyright Epic Games, Inc. All Rights Reserved.
+
+#include "Nodes/Attach/CameraActorAttachmentInfo.h"
+
+#include "Algo/Accumulate.h"
+#include "Components/SkeletalMeshComponent.h"
+#include "Core/CameraContextDataTable.h"
+#include "GameFramework/Actor.h"
+
+namespace UE::Cameras
+{
+
+FCameraActorAttachmentInfoReader::FCameraActorAttachmentInfoReader(const FCameraActorAttachmentInfo& InAttachmentInfo, FCameraContextDataID InDataID)
+{
+	Initialize(InAttachmentInfo, InDataID);
+}
+
+void FCameraActorAttachmentInfoReader::Initialize(const FCameraActorAttachmentInfo& InAttachmentInfo, FCameraContextDataID InDataID)
+{
+	DefaultAttachmentInfo = InAttachmentInfo;
+	DataID = InDataID;
+
+	CacheAttachmentInfo(DefaultAttachmentInfo);
+}
+
+void FCameraActorAttachmentInfoReader::CacheAttachmentInfo(const FCameraActorAttachmentInfo& InAttachmentInfo)
+{
+	if (CachedAttachmentInfo != InAttachmentInfo)
+	{
+		CachedAttachmentInfo = InAttachmentInfo;
+
+		CachedSkeletalMeshComponent = nullptr;
+		if (InAttachmentInfo.Actor && !InAttachmentInfo.SocketName.IsNone())
+		{
+			CachedSkeletalMeshComponent = InAttachmentInfo.Actor->FindComponentByClass<USkeletalMeshComponent>();
+		}
+	}
+}
+
+bool FCameraActorAttachmentInfoReader::GetAttachmentTransform(const FCameraContextDataTable& ContextDataTable, FTransform3d& OutTransform)
+{
+	if (DataID.IsValid())
+	{
+		if (const FCameraActorAttachmentInfo* NewAttachmentInfo = ContextDataTable.TryGetData<FCameraActorAttachmentInfo>(DataID))
+		{
+			CacheAttachmentInfo(*NewAttachmentInfo);
+		}
+		else
+		{
+			CacheAttachmentInfo(DefaultAttachmentInfo);
+		}
+	}
+
+	if (!CachedAttachmentInfo.SocketName.IsNone() && CachedSkeletalMeshComponent)
+	{
+		OutTransform = CachedSkeletalMeshComponent->GetSocketTransform(CachedAttachmentInfo.SocketName);
+		return true;
+	}
+	else if (CachedAttachmentInfo.Actor)
+	{
+		OutTransform = CachedAttachmentInfo.Actor->GetTransform();
+		return true;
+	}
+	return false;
+}
+
+bool FCameraActorAttachmentInfoReader::GetAttachmentTransform(TArrayView<FCameraActorAttachmentInfoReader> Readers, const FCameraContextDataTable& ContextDataTable,FTransform3d& OutTransform)
+{
+	using FWeightedTransform = TTuple<FTransform3d, float>;
+	
+	if (Readers.IsEmpty())
+	{
+		return false;
+	}
+
+	TArray<FWeightedTransform> WeightedTransforms;
+
+	for (FCameraActorAttachmentInfoReader& Reader : Readers)
+	{
+		FTransform3d CurTransform;
+		if (Reader.GetAttachmentTransform(ContextDataTable, CurTransform))
+		{
+			WeightedTransforms.Emplace(CurTransform, Reader.CachedAttachmentInfo.Weight);
+		}
+	}
+
+	float TotalWeight = Algo::Accumulate(
+			WeightedTransforms, 
+			0.f,
+			[](float Cur, const FWeightedTransform& Item) { return Cur + Item.Get<1>(); });
+	if (TotalWeight == 0.f)
+	{
+		return false;
+	}
+
+	OutTransform = FTransform3d::Identity;
+	for (const FWeightedTransform& WeightedTransform : WeightedTransforms)
+	{
+		OutTransform.BlendWith(WeightedTransform.Get<0>(), WeightedTransform.Get<1>() / TotalWeight);
+	}
+
+	return true;
+}
+
+}  // namespace UE::Cameras
+
diff --git a/Engine/Plugins/Cameras/GameplayCameras/Source/GameplayCameras/Public/Nodes/Attach/AttachToActorCameraNode.h b/Engine/Plugins/Cameras/GameplayCameras/Source/GameplayCameras/Public/Nodes/Attach/AttachToActorCameraNode.h
new file mode 100644
--- /dev/null
+++ b/Engine/Plugins/Cameras/GameplayCameras/Source/GameplayCameras/Public/Nodes/Attach/AttachToActorCameraNode.h
@@ -0,0 +1,43 @@
+// Copyright Epic Games, Inc. All Rights Reserved.
+
+#pragma once
+
+#include "Core/CameraContextDataTableFwd.h"
+#include "Core/CameraNode.h"
+#include "Core/CameraParameters.h"
+#include "Nodes/Attach/CameraActorAttachmentInfo.h"
+
+#include "AttachToActorCameraNode.generated.h"
+
+/**
+ * A camera node that moves the camera to an actor.
+ */
+UCLASS(MinimalAPI, meta=(CameraNodeCategories="Attachment"))
+class UAttachToActorCameraNode : public UCameraNode
+{
+	GENERATED_BODY()
+
+protected:
+
+	// UCameraNode interface.
+	virtual FCameraNodeEvaluatorPtr OnBuildEvaluator(FCameraNodeEvaluatorBuilder& Builder) const override;
+
+public:
+
+	/** The actor to attach to. */
+	UPROPERTY(EditAnywhere, Category=Common, meta=(CameraContextData=true))
+	FCameraActorAttachmentInfo Attachment;
+
+	/** The data ID for the attachment info. */
+	UPROPERTY()
+	FCameraContextDataID AttachmentDataID;
+
+	/** Whether to move the camera to the actor's location. */
+	UPROPERTY(EditAnywhere, Category=Common)
+	FBooleanCameraParameter AttachToLocation = true;
+
+	/** Whether to align the camera rotation to the actor's orientation. */
+	UPROPERTY(EditAnywhere, Category=Common)
+	FBooleanCameraParameter AttachToRotation = false;
+};
+
diff --git a/Engine/Plugins/Cameras/GameplayCameras/Source/GameplayCameras/Public/Nodes/Attach/CameraActorAttachmentInfo.h b/Engine/Plugins/Cameras/GameplayCameras/Source/GameplayCameras/Public/Nodes/Attach/CameraActorAttachmentInfo.h
new file mode 100644
--- /dev/null
+++ b/Engine/Plugins/Cameras/GameplayCameras/Source/GameplayCameras/Public/Nodes/Attach/CameraActorAttachmentInfo.h
@@ -0,0 +1,64 @@
+// Copyright Epic Games, Inc. All Rights Reserved.
+
+#pragma once
+
+#include "CoreTypes.h"
+#include "Core/CameraContextDataTableFwd.h"
+
+#include "CameraActorAttachmentInfo.generated.h"
+
+class AActor;
+
+/**
+ * Attachment information for a camera rig.
+ */
+USTRUCT(BlueprintType)
+struct GAMEPLAYCAMERAS_API FCameraActorAttachmentInfo
+{
+	GENERATED_BODY()
+
+	/** The actor to attach to. */
+	UPROPERTY(EditAnywhere, Category="Attachment")
+	TObjectPtr<AActor> Actor;
+
+	/** An optional socket to attach to on the actor. */
+	UPROPERTY(EditAnywhere, Interp, Category="Attachment")
+	FName SocketName;
+
+	/** The weight of this attachment. Unused if only one attachment is used. */
+	UPROPERTY(EditAnywhere, Interp, Category="Attachment")
+	float Weight = 1.f;
+
+	bool operator== (const FCameraActorAttachmentInfo& Other) const = default;
+};
+
+namespace UE::Cameras
+{
+
+/** A special reader class for attachment information. */
+struct FCameraActorAttachmentInfoReader
+{
+	FCameraActorAttachmentInfoReader() {}
+	FCameraActorAttachmentInfoReader(const FCameraActorAttachmentInfo& InAttachmentInfo, FCameraContextDataID InDataID);
+
+	void Initialize(const FCameraActorAttachmentInfo& InAttachmentInfo, FCameraContextDataID InDataID);
+
+	bool GetAttachmentTransform(const FCameraContextDataTable& ContextDataTable, FTransform3d& OutTransform);
+
+	static bool GetAttachmentTransform(TArrayView<FCameraActorAttachmentInfoReader> Readers, const FCameraContextDataTable& ContextDataTable, FTransform3d& OutTransform);
+
+private:
+
+	void CacheAttachmentInfo(const FCameraActorAttachmentInfo& InAttachmentInfo);
+
+private:
+
+	FCameraActorAttachmentInfo DefaultAttachmentInfo;
+	FCameraContextDataID DataID;
+
+	FCameraActorAttachmentInfo CachedAttachmentInfo;
+	const USkeletalMeshComponent* CachedSkeletalMeshComponent = nullptr;
+};
+
+}  // namespace UE::Cameras
+
