Cameras: add support for array data parameters

diff --git a/Engine/Plugins/Cameras/GameplayCameras/Source/GameplayCameras/Private/Core/CameraContextDataTable.cpp b/Engine/Plugins/Cameras/GameplayCameras/Source/GameplayCameras/Private/Core/CameraContextDataTable.cpp
--- a/Engine/Plugins/Cameras/GameplayCameras/Source/GameplayCameras/Private/Core/CameraContextDataTable.cpp
+++ b/Engine/Plugins/Cameras/GameplayCameras/Source/GameplayCameras/Private/Core/CameraContextDataTable.cpp
@@ -21,30 +21,72 @@ void FCameraContextDataTable::AddReferencedObjects(FReferenceCollector& Referenc
 	{
 		ReferenceCollector.AddReferencedObject(Entry.TypeObject);
 
-		uint8* RawData = Memory + Entry.Offset;
-		switch (Entry.Type)
+		if (Entry.ContainerType == ECameraContextDataContainerType::None)
 		{
-			case ECameraContextDataType::Struct:
-				{
-					const UScriptStruct* StructType = Cast<const UScriptStruct>(Entry.TypeObject);
-					if (ensure(StructType))
+			uint8* RawData = Memory + Entry.Offset;
+
+			switch (Entry.Type)
+			{
+				case ECameraContextDataType::Struct:
 					{
-						ReferenceCollector.AddPropertyReferencesWithStructARO(StructType, RawData);
+						const UScriptStruct* StructType = Cast<const UScriptStruct>(Entry.TypeObject);
+						if (ensure(StructType))
+						{
+							ReferenceCollector.AddPropertyReferencesWithStructARO(StructType, RawData);
+						}
 					}
-				}
-				break;
-			case ECameraContextDataType::Object:
-				{
-					TObjectPtr<UObject>* TypedData = reinterpret_cast<TObjectPtr<UObject>*>(RawData);
-					ReferenceCollector.AddReferencedObject(*TypedData);
-				}
-				break;
-			case ECameraContextDataType::Class:
-				{
-					TObjectPtr<UClass>* TypedData = reinterpret_cast<TObjectPtr<UClass>*>(RawData);
-					ReferenceCollector.AddReferencedObject(*TypedData);
-				}
-				break;
+					break;
+				case ECameraContextDataType::Object:
+					{
+						TObjectPtr<UObject>* TypedData = reinterpret_cast<TObjectPtr<UObject>*>(RawData);
+						ReferenceCollector.AddReferencedObject(*TypedData);
+					}
+					break;
+				case ECameraContextDataType::Class:
+					{
+						TObjectPtr<UClass>* TypedData = reinterpret_cast<TObjectPtr<UClass>*>(RawData);
+						ReferenceCollector.AddReferencedObject(*TypedData);
+					}
+					break;
+			}
+		}
+		else if (Entry.ContainerType == ECameraContextDataContainerType::Array)
+		{
+			FArrayEntryHelper Helper(Entry, Memory);
+
+			switch (Entry.Type)
+			{
+				case ECameraContextDataType::Struct:
+					{
+						const UScriptStruct* StructType = Cast<const UScriptStruct>(Entry.TypeObject);
+						if (ensure(StructType))
+						{
+							for (int32 Index = 0; Index < Helper.Num(); ++Index)
+							{
+								ReferenceCollector.AddPropertyReferencesWithStructARO(StructType, Helper.GetRawPtr(Index));
+							}
+						}
+					}
+					break;
+				case ECameraContextDataType::Object:
+					{
+						for (int32 Index = 0; Index < Helper.Num(); ++Index)
+						{
+							TObjectPtr<UObject>* TypedData = reinterpret_cast<TObjectPtr<UObject>*>(Helper.GetRawPtr(Index));
+							ReferenceCollector.AddReferencedObject(*TypedData);
+						}
+					}
+					break;
+				case ECameraContextDataType::Class:
+					{
+						for (int32 Index = 0; Index < Helper.Num(); ++Index)
+						{
+							TObjectPtr<UClass>* TypedData = reinterpret_cast<TObjectPtr<UClass>*>(Helper.GetRawPtr(Index));
+							ReferenceCollector.AddReferencedObject(*TypedData);
+						}
+					}
+					break;
+			}
 		}
 	}
 }
@@ -63,13 +105,14 @@ void FCameraContextDataTable::Initialize(const FCameraContextDataTableAllocation
 
 	for (const FCameraContextDataDefinition& DataDefinition : AllocationInfo.DataDefinitions)
 	{
-		GetDataTypeAllocationInfo(DataDefinition.DataType, DataDefinition.DataTypeObject, CurSizeOf, CurAlignOf);
+		GetDataTypeAllocationInfo(DataDefinition.DataType, DataDefinition.DataContainerType, DataDefinition.DataTypeObject, CurSizeOf, CurAlignOf);
 		const uint32 NewEntryOffset = Align(TotalSizeOf, CurAlignOf);
 		TotalSizeOf = NewEntryOffset + CurSizeOf;
 
 		FEntry NewEntry;
 		NewEntry.ID = DataDefinition.DataID;
 		NewEntry.Type = DataDefinition.DataType;
+		NewEntry.ContainerType = DataDefinition.DataContainerType;
 		NewEntry.TypeObject = DataDefinition.DataTypeObject;
 		NewEntry.Offset = NewEntryOffset;
 		NewEntry.Flags = EEntryFlags::None;
@@ -90,7 +133,7 @@ void FCameraContextDataTable::Initialize(const FCameraContextDataTableAllocation
 	for (const FEntry& Entry : Entries)
 	{
 		uint8* DataPtr = Memory + Entry.Offset;
-		InitializeDefaultDataValue(Entry.Type, Entry.TypeObject, DataPtr);
+		ConstructDataValue(Entry.Type, Entry.ContainerType, Entry.TypeObject, DataPtr);
 	}
 }
 
@@ -102,7 +145,7 @@ void FCameraContextDataTable::AddData(const FCameraContextDataDefinition& DataDe
 	}
 
 	uint32 SizeOf, AlignOf;
-	GetDataTypeAllocationInfo(DataDefinition.DataType, DataDefinition.DataTypeObject, SizeOf, AlignOf);
+	GetDataTypeAllocationInfo(DataDefinition.DataType, DataDefinition.DataContainerType, DataDefinition.DataTypeObject, SizeOf, AlignOf);
 
 	uint8* DataPtr = Align(Memory + Used, AlignOf);
 	uint32 NewUsed = (DataPtr + SizeOf) - Memory;
@@ -119,6 +162,7 @@ void FCameraContextDataTable::AddData(const FCameraContextDataDefinition& DataDe
 	FEntry NewEntry;
 	NewEntry.ID = DataDefinition.DataID;
 	NewEntry.Type = DataDefinition.DataType;
+	NewEntry.ContainerType = DataDefinition.DataContainerType;
 	NewEntry.TypeObject = DataDefinition.DataTypeObject;
 	NewEntry.Offset = DataPtr - Memory;
 	NewEntry.Flags = EEntryFlags::None;
@@ -129,7 +173,7 @@ void FCameraContextDataTable::AddData(const FCameraContextDataDefinition& DataDe
 	Entries.Add(NewEntry);
 	EntryLookup.Add(DataDefinition.DataID, Entries.Num() - 1);
 
-	InitializeDefaultDataValue(NewEntry.Type, NewEntry.TypeObject, Memory + NewEntry.Offset);
+	ConstructDataValue(NewEntry.Type, NewEntry.ContainerType, NewEntry.TypeObject, Memory + NewEntry.Offset);
 }
 
 bool FCameraContextDataTable::GetDataTypeAllocationInfo(ECameraContextDataType DataType, const UObject* DataTypeObject, uint32& OutSizeOf, uint32& OutAlignOf)
@@ -173,7 +217,22 @@ bool FCameraContextDataTable::GetDataTypeAllocationInfo(ECameraContextDataType D
 	return true;
 }
 
-bool FCameraContextDataTable::InitializeDefaultDataValue(ECameraContextDataType DataType, const UObject* DataTypeObject, uint8* DataPtr)
+bool FCameraContextDataTable::GetDataTypeAllocationInfo(ECameraContextDataType DataType, ECameraContextDataContainerType DataContainerType, const UObject* DataTypeObject, uint32& OutSizeOf, uint32& OutAlignOf)
+{
+	if (DataContainerType == ECameraContextDataContainerType::None)
+	{
+		return GetDataTypeAllocationInfo(DataType, DataTypeObject, OutSizeOf, OutAlignOf);
+	}
+	else if (DataContainerType == ECameraContextDataContainerType::Array)
+	{
+		OutSizeOf = sizeof(FEntryScriptArray);
+		OutAlignOf = alignof(FEntryScriptArray);
+		return true;
+	}
+	return false;
+}
+
+bool FCameraContextDataTable::ConstructDataValue(ECameraContextDataType DataType, const UObject* DataTypeObject, uint8* DataPtr)
 {
 	switch (DataType)
 	{
@@ -214,6 +273,73 @@ bool FCameraContextDataTable::InitializeDefaultDataValue(ECameraContextDataType
 	return true;
 }
 
+bool FCameraContextDataTable::ConstructDataValue(ECameraContextDataType DataType, ECameraContextDataContainerType DataContainerType, const UObject* DataTypeObject, uint8* DataPtr)
+{
+	if (DataContainerType == ECameraContextDataContainerType::None)
+	{
+		return ConstructDataValue(DataType, DataTypeObject, DataPtr);
+	}
+	else if (DataContainerType == ECameraContextDataContainerType::Array)
+	{
+		new (DataPtr) FEntryScriptArray();
+	}
+	return false;
+}
+
+bool FCameraContextDataTable::DestroyDataValue(ECameraContextDataType DataType, const UObject* DataTypeObject, uint8* DataPtr)
+{
+	switch (DataType)
+	{
+		case ECameraContextDataType::Name:
+			((FName*)DataPtr)->~FName();
+			break;
+		case ECameraContextDataType::String:
+			((FString*)DataPtr)->~FString();
+			break;
+		case ECameraContextDataType::Enum:
+			// Nothing to do.
+			break;
+		case ECameraContextDataType::Struct:
+			{
+				const UScriptStruct* StructType = Cast<const UScriptStruct>(DataTypeObject);
+				if (ensure(StructType))
+				{
+					StructType->DestroyStruct(DataPtr);
+				}
+			}
+			break;
+		case ECameraContextDataType::Object:
+			((TObjectPtr<UObject>*)DataPtr)->~TObjectPtr<UObject>();
+			break;
+		case ECameraContextDataType::Class:
+			((TObjectPtr<UClass>*)DataPtr)->~TObjectPtr<UClass>();
+			break;
+		default:
+			ensure(false);
+			return false;
+	}
+	return true;
+}
+
+bool FCameraContextDataTable::DestroyDataValue(ECameraContextDataType DataType, ECameraContextDataContainerType DataContainerType, const UObject* DataTypeObject, uint8* DataPtr)
+{
+	if (DataContainerType == ECameraContextDataContainerType::None)
+	{
+		return DestroyDataValue(DataType, DataTypeObject, DataPtr);
+	}
+	else if (DataContainerType == ECameraContextDataContainerType::Array)
+	{
+		FArrayEntryHelper Helper(DataType, DataTypeObject, DataPtr);
+		for (int32 Index = 0; Index < Helper.Num(); ++Index)
+		{
+			uint8* RawElementPtr = Helper.GetRawPtr(Index);
+			DestroyDataValue(DataType, DataTypeObject, RawElementPtr);
+		}
+		((FEntryScriptArray*)DataPtr)->~FEntryScriptArray();
+	}
+	return false;
+}
+
 void FCameraContextDataTable::ReallocateBuffer(uint32 MinRequired)
 {
 	static const uint32 DefaultCapacity = 64u;
@@ -250,37 +376,7 @@ void FCameraContextDataTable::DestroyBuffer()
 	for (const FEntry& Entry : Entries)
 	{
 		uint8* DataPtr = Memory + Entry.Offset;
-
-		switch (Entry.Type)
-		{
-			case ECameraContextDataType::Name:
-				reinterpret_cast<FName*>(DataPtr)->~FName();
-				break;
-			case ECameraContextDataType::String:
-				reinterpret_cast<FString*>(DataPtr)->~FString();
-				break;
-			case ECameraContextDataType::Enum:
-				// Nothing to do.
-				break;
-			case ECameraContextDataType::Struct:
-				{
-					const UScriptStruct* StructType = Cast<const UScriptStruct>(Entry.TypeObject);
-					if (ensure(StructType))
-					{
-						StructType->DestroyStruct(DataPtr);
-					}
-				}
-				break;
-			case ECameraContextDataType::Object:
-				reinterpret_cast<TObjectPtr<UObject>*>(DataPtr)->~TObjectPtr<UObject>();
-				break;
-			case ECameraContextDataType::Class:
-				reinterpret_cast<TObjectPtr<UClass>*>(DataPtr)->~TObjectPtr<UClass>();
-				break;
-			default:
-				ensure(false);
-				break;
-		}
+		DestroyDataValue(Entry.Type, Entry.ContainerType, Entry.TypeObject, DataPtr);
 	}
 
 	FMemory::Free(Memory);
@@ -392,7 +488,10 @@ void FCameraContextDataTable::SetClassData(FCameraContextDataID InID, UClass* In
 void FCameraContextDataTable::SetStructViewData(FCameraContextDataID InID, const FStructView& InData)
 {
 	FEntry* Entry = FindEntry(InID);
-	if (ensure(Entry && Entry->Type == ECameraContextDataType::Struct && InData.GetScriptStruct() == Entry->TypeObject))
+	if (ensure(Entry && 
+				Entry->Type == ECameraContextDataType::Struct && 
+				Entry->ContainerType == ECameraContextDataContainerType::None &&
+				InData.GetScriptStruct() == Entry->TypeObject))
 	{
 		uint8* DataPtr = Memory + Entry->Offset;
 		const UScriptStruct* StructType = CastChecked<const UScriptStruct>(Entry->TypeObject);
@@ -404,7 +503,10 @@ void FCameraContextDataTable::SetStructViewData(FCameraContextDataID InID, const
 void FCameraContextDataTable::SetInstancedStructData(FCameraContextDataID InID, const FInstancedStruct& InData)
 {
 	FEntry* Entry = FindEntry(InID);
-	if (ensure(Entry && Entry->Type == ECameraContextDataType::Struct && InData.GetScriptStruct() == Entry->TypeObject))
+	if (ensure(Entry && 
+				Entry->Type == ECameraContextDataType::Struct && 
+				Entry->ContainerType == ECameraContextDataContainerType::None &&
+				InData.GetScriptStruct() == Entry->TypeObject))
 	{
 		uint8* DataPtr = Memory + Entry->Offset;
 		const UScriptStruct* StructType = CastChecked<const UScriptStruct>(Entry->TypeObject);
@@ -413,6 +515,79 @@ void FCameraContextDataTable::SetInstancedStructData(FCameraContextDataID InID,
 	}
 }
 
+void FCameraContextDataTable::SetNameArrayData(FCameraContextDataID InID, TConstArrayView<FName> InData)
+{
+	SetArrayDataImpl(InID, ECameraContextDataType::Name, nullptr, InData);
+}
+
+void FCameraContextDataTable::SetStringArrayData(FCameraContextDataID InID, TConstArrayView<FString> InData)
+{
+	SetArrayDataImpl(InID, ECameraContextDataType::String, nullptr, InData);
+}
+
+void FCameraContextDataTable::SetEnumArrayData(FCameraContextDataID InID, const UEnum* EnumType, TConstArrayView<uint8> InData)
+{
+	SetArrayDataImpl(InID, ECameraContextDataType::Enum, EnumType, InData);
+}
+
+void FCameraContextDataTable::SetObjectArrayData(FCameraContextDataID InID, TConstArrayView<UObject*> InData)
+{
+	TArray<TObjectPtr<UObject>> InDataPtrs(InData);
+	SetArrayDataImpl(InID, ECameraContextDataType::Object, nullptr, TConstArrayView<TObjectPtr<UObject>>(InDataPtrs));
+}
+
+void FCameraContextDataTable::SetClassArrayData(FCameraContextDataID InID, TConstArrayView<UClass*> InData)
+{
+	TArray<TObjectPtr<UClass>> InDataPtrs(InData);
+	SetArrayDataImpl(InID, ECameraContextDataType::Class, nullptr, TConstArrayView<TObjectPtr<UClass>>(InDataPtrs));
+}
+
+void FCameraContextDataTable::SetStructViewArrayData(FCameraContextDataID InID, TConstArrayView<FStructView> InData)
+{
+	FEntry* Entry = FindEntry(InID);
+	if (ensure(Entry && 
+				Entry->Type == ECameraContextDataType::Struct && 
+				Entry->ContainerType == ECameraContextDataContainerType::Array))
+	{
+		const UScriptStruct* StructType = CastChecked<const UScriptStruct>(Entry->TypeObject);
+
+		FArrayEntryHelper Helper(*Entry, Memory);
+		Helper.Resize(InData.Num());
+		for (int32 Index = 0; Index < InData.Num(); ++Index)
+		{
+			uint8* RawData = Helper.GetRawPtr(Index);
+			if (ensure(StructType == InData[Index].GetScriptStruct()))
+			{
+				StructType->CopyScriptStruct(RawData, InData[Index].GetMemory());
+			}
+		}
+		EnumAddFlags(Entry->Flags, EEntryFlags::Written | EEntryFlags::WrittenThisFrame);
+	}
+}
+
+void FCameraContextDataTable::SetInstancedStructArrayData(FCameraContextDataID InID, TConstArrayView<FInstancedStruct> InData)
+{
+	FEntry* Entry = FindEntry(InID);
+	if (ensure(Entry && 
+				Entry->Type == ECameraContextDataType::Struct && 
+				Entry->ContainerType == ECameraContextDataContainerType::Array))
+	{
+		const UScriptStruct* StructType = CastChecked<const UScriptStruct>(Entry->TypeObject);
+
+		FArrayEntryHelper Helper(*Entry, Memory);
+		Helper.Resize(InData.Num());
+		for (int32 Index = 0; Index < InData.Num(); ++Index)
+		{
+			uint8* RawData = Helper.GetRawPtr(Index);
+			if (ensure(StructType == InData[Index].GetScriptStruct()))
+			{
+				StructType->CopyScriptStruct(RawData, InData[Index].GetMemory());
+			}
+		}
+		EnumAddFlags(Entry->Flags, EEntryFlags::Written | EEntryFlags::WrittenThisFrame);
+	}
+}
+
 const FCameraContextDataTable::FEntry* FCameraContextDataTable::FindEntry(FCameraContextDataID InID) const
 {
 	const uint32 EntryIndex = EntryLookup.FindRef(InID, INDEX_NONE);
@@ -454,7 +629,9 @@ const uint8* FCameraContextDataTable::TryGetData(
 	const FEntry* Entry = FindEntry(DataID);
 	if (Entry)
 	{
-		if (Entry->Type == ExpectedDataType && Entry->TypeObject == ExpectedDataTypeObject)
+		if (Entry->Type == ExpectedDataType && 
+				Entry->ContainerType == ECameraContextDataContainerType::None && 
+				Entry->TypeObject == ExpectedDataTypeObject)
 		{
 			return Memory + Entry->Offset;
 		}
@@ -471,7 +648,9 @@ uint8* FCameraContextDataTable::TryGetMutableData(
 	FEntry* Entry = FindEntry(DataID);
 	if (Entry)
 	{
-		if (Entry->Type == ExpectedDataType && Entry->TypeObject == ExpectedDataTypeObject)
+		if (Entry->Type == ExpectedDataType && 
+				Entry->ContainerType == ECameraContextDataContainerType::None && 
+				Entry->TypeObject == ExpectedDataTypeObject)
 		{
 			return Memory + Entry->Offset;
 		}
@@ -480,6 +659,25 @@ uint8* FCameraContextDataTable::TryGetMutableData(
 	return nullptr;
 }
 
+const FCameraContextDataTable::FEntryScriptArray* FCameraContextDataTable::TryGetArrayData(
+		FCameraContextDataID DataID,
+		ECameraContextDataType ExpectedDataType,
+		const UObject* ExpectedDataTypeObject) const
+{
+	const FEntry* Entry = FindEntry(DataID);
+	if (Entry)
+	{
+		if (Entry->Type == ExpectedDataType && 
+				Entry->ContainerType == ECameraContextDataContainerType::Array && 
+				Entry->TypeObject == ExpectedDataTypeObject)
+		{
+			return (FEntryScriptArray*)(Memory + Entry->Offset);
+		}
+	}
+
+	return nullptr;
+}
+
 void FCameraContextDataTable::SetData(
 		FCameraContextDataID DataID,
 		ECameraContextDataType ExpectedDataType,
@@ -504,7 +702,9 @@ bool FCameraContextDataTable::TrySetData(
 		return false;
 	}
 
-	if (!ensure(ExpectedDataType == Entry->Type && ExpectedDataTypeObject == Entry->TypeObject))
+	if (!ensure(Entry->Type == ExpectedDataType && 
+				Entry->ContainerType == ECameraContextDataContainerType::None &&
+				Entry->TypeObject == ExpectedDataTypeObject))
 	{
 		return false;
 	}
@@ -521,6 +721,87 @@ bool FCameraContextDataTable::TrySetData(
 	return true;
 }
 
+bool FCameraContextDataTable::TrySetArrayDataNum(
+		FCameraContextDataID DataID, 
+		int32 Count,
+		bool bMarkAsWrittenThisFrame)
+{
+	FEntry* Entry = FindEntry(DataID);
+	if (!Entry)
+	{
+		return false;
+	}
+
+	if (!ensure(Entry->ContainerType == ECameraContextDataContainerType::Array))
+	{
+		return false;
+	}
+
+	FArrayEntryHelper Helper(*Entry, Memory);
+	Helper.Resize(Count);
+
+	Entry->Flags |= EEntryFlags::Written;
+	if (bMarkAsWrittenThisFrame)
+	{
+		Entry->Flags |= EEntryFlags::WrittenThisFrame;
+	}
+
+	return true;
+}
+
+bool FCameraContextDataTable::TrySetArrayData(
+		FCameraContextDataID DataID,
+		ECameraContextDataType ExpectedDataType,
+		const UObject* ExpectedDataTypeObject,
+		int32 Index,
+		const uint8* InRawDataPtr,
+		bool bMarkAsWrittenThisFrame)
+{
+	FEntry* Entry = FindEntry(DataID);
+	if (!Entry)
+	{
+		return false;
+	}
+
+	if (!ensure(Entry->ContainerType == ECameraContextDataContainerType::Array))
+	{
+		return false;
+	}
+
+	FArrayEntryHelper Helper(*Entry, Memory);
+	uint8* DataPtr = Helper.GetRawPtr(Index);
+	SetDataValue(Entry->Type, Entry->TypeObject, DataPtr, InRawDataPtr);
+
+	Entry->Flags |= EEntryFlags::Written;
+	if (bMarkAsWrittenThisFrame)
+	{
+		Entry->Flags |= EEntryFlags::WrittenThisFrame;
+	}
+	
+	return true;
+}
+
+bool FCameraContextDataTable::SetDataValue(ECameraContextDataType DataType, ECameraContextDataContainerType DataContainerType, const UObject* DataTypeObject, uint8* DestDataPtr, const uint8* SrcDataPtr)
+{
+	if (DataContainerType == ECameraContextDataContainerType::None)
+	{
+		return SetDataValue(DataType, DataTypeObject, DestDataPtr, SrcDataPtr);
+	}
+	else if (DataContainerType == ECameraContextDataContainerType::Array)
+	{
+		FArrayEntryHelper DestHelper(DataType, DataTypeObject, DestDataPtr);
+		FArrayEntryHelper SrcHelper(DataType, DataTypeObject, const_cast<uint8*>(SrcDataPtr));
+		DestHelper.Resize(SrcHelper.Num());
+		for (int32 Index = 0; Index < SrcHelper.Num(); ++Index)
+		{
+			uint8* DestElementPtr = DestHelper.GetRawPtr(Index);
+			const uint8* SrcElementPtr = SrcHelper.GetRawPtr(Index);
+			SetDataValue(DataType, DataTypeObject, DestElementPtr, SrcElementPtr);
+		}
+	}
+	return false;
+}
+
 bool FCameraContextDataTable::SetDataValue(ECameraContextDataType DataType, const UObject* DataTypeObject, uint8* DestDataPtr, const uint8* SrcDataPtr)
 {
 	switch (DataType)
@@ -635,6 +916,7 @@ void FCameraContextDataTable::InternalOverride(const FCameraContextDataTable& Ot
 				FCameraContextDataDefinition OtherEntryDefinition;
 				OtherEntryDefinition.DataID = OtherEntry.ID;
 				OtherEntryDefinition.DataType = OtherEntry.Type;
+				OtherEntryDefinition.DataContainerType = OtherEntry.ContainerType;
 				OtherEntryDefinition.DataTypeObject = OtherEntry.TypeObject;
 				AddData(OtherEntryDefinition);
 				ThisIndex = Entries.Num() - 1;
@@ -650,12 +932,72 @@ void FCameraContextDataTable::InternalOverride(const FCameraContextDataTable& Ot
 
 			uint8* ThisDataPtr = Memory + ThisEntry.Offset;
 			uint8* OtherDataPtr = OtherTable.Memory + OtherEntry.Offset;
-			SetDataValue(ThisEntry.Type, ThisEntry.TypeObject, ThisDataPtr, OtherDataPtr);
+			SetDataValue(ThisEntry.Type, ThisEntry.ContainerType, ThisEntry.TypeObject, ThisDataPtr, OtherDataPtr);
 
 			EnumAddFlags(ThisEntry.Flags, EEntryFlags::Written | (OtherFlags & EEntryFlags::WrittenThisFrame));
 		}
 	}
 }
 
+FCameraContextDataTable::FArrayEntryHelper::FArrayEntryHelper(const FEntry& Entry, uint8* TableMemory)
+	: FArrayEntryHelper(Entry.Type, Entry.TypeObject, TableMemory + Entry.Offset)
+{
+	check(Entry.ContainerType == ECameraContextDataContainerType::Array);
+}
+
+FCameraContextDataTable::FArrayEntryHelper::FArrayEntryHelper(ECameraContextDataType DataType, const UObject* DataTypeObject, uint8* RawPtr)
+{
+	ElementType = DataType;
+	ElementTypeObject = DataTypeObject;
+
+	uint32 SizeOf, AlignOf;
+	FCameraContextDataTable::GetDataTypeAllocationInfo(DataType, DataTypeObject, SizeOf, AlignOf);
+
+	// ElementSizeOf is the total size of an array entry, including padding.
+	ElementSizeOf = Align(SizeOf, AlignOf);
+	ElementAlignOf = AlignOf;
+
+	ScriptArray = (FEntryScriptArray*)RawPtr;
+}
+
+bool FCameraContextDataTable::FArrayEntryHelper::IsValidIndex(int32 Index) const
+{
+	return Index >= 0 && Index < Num();
+}
+
+int32 FCameraContextDataTable::FArrayEntryHelper::Num() const
+{
+	return ScriptArray->Num();
+}
+
+uint8* FCameraContextDataTable::FArrayEntryHelper::GetRawPtr(int32 Index)
+{
+	checkSlow(IsValidIndex(Index));
+	return (uint8*)ScriptArray->GetData() + Index * ElementSizeOf;
+}
+
+void FCameraContextDataTable::FArrayEntryHelper::Resize(int32 Count)
+{
+	int32 OldNum = Num();
+	if (Count > OldNum)
+	{
+		ScriptArray->Add(Count - OldNum, (int32)ElementSizeOf, ElementAlignOf);
+		for (int32 Index = OldNum; Index < Count; ++Index)
+		{
+			uint8* ElementPtr = GetRawPtr(Index);
+			FCameraContextDataTable::ConstructDataValue(ElementType, ElementTypeObject, ElementPtr);
+		}
+	}
+	else if (Count < OldNum)
+	{
+		for (int32 Index = Count; Index < OldNum; ++Index)
+		{
+			uint8* ElementPtr = GetRawPtr(Index);
+			FCameraContextDataTable::DestroyDataValue(ElementType, ElementTypeObject, ElementPtr);
+		}
+		ScriptArray->Remove(Count, OldNum - Count, (int32)ElementSizeOf, ElementAlignOf);
+	}
+}
+
 }  // namespace UE::Cameras
 
diff --git a/Engine/Plugins/Cameras/GameplayCameras/Source/GameplayCameras/Private/Core/CameraRigAsset.cpp b/Engine/Plugins/Cameras/GameplayCameras/Source/GameplayCameras/Private/Core/CameraRigAsset.cpp
--- a/Engine/Plugins/Cameras/GameplayCameras/Source/GameplayCameras/Private/Core/CameraRigAsset.cpp
+++ b/Engine/Plugins/Cameras/GameplayCameras/Source/GameplayCameras/Private/Core/CameraRigAsset.cpp
@@ -92,6 +92,7 @@ FCameraContextDataDefinition UCameraRigDataParameter::GetDataDefinition() const
 	FCameraContextDataDefinition Definition;
 	Definition.DataID = PrivateDataID;
 	Definition.DataType = DataType;
+	Definition.DataContainerType = DataContainerType;
 	Definition.DataTypeObject = DataTypeObject;
 #if WITH_EDITORONLY_DATA
 	Definition.DataName = GetDataName();
@@ -463,6 +464,10 @@ void UCameraRigAsset::GetAssetRegistryTags(FAssetRegistryTagsContext Context) co
 			FAssetRegistryTag DataParameterTag;
 			DataParameterTag.Name = FName(DataParameter->InterfaceParameterName);
 			DataParameterTag.Value = ContextDataTypeEnum->GetNameStringByValue((int64)DataParameter->DataType);
+			if (DataParameter->DataContainerType == ECameraContextDataContainerType::Array)
+			{
+				DataParameterTag.Value += TEXT("[]");
+			}
 			DataParameterTag.Type = FAssetRegistryTag::TT_Alphabetical;
 			Context.AddTag(DataParameterTag);
 		}
diff --git a/Engine/Plugins/Cameras/GameplayCameras/Source/GameplayCameras/Private/Core/CameraRigAssetBuilder.cpp b/Engine/Plugins/Cameras/GameplayCameras/Source/GameplayCameras/Private/Core/CameraRigAssetBuilder.cpp
--- a/Engine/Plugins/Cameras/GameplayCameras/Source/GameplayCameras/Private/Core/CameraRigAssetBuilder.cpp
+++ b/Engine/Plugins/Cameras/GameplayCameras/Source/GameplayCameras/Private/Core/CameraRigAssetBuilder.cpp
@@ -334,18 +334,6 @@ void AddVariableToAllocationInfo(FCameraVariableID VariableID, ECameraVariableTy
 	}
 }
 
-void AddContextDataToAllocationInfo(FCameraContextDataID DataID, ECameraContextDataType DataType, const UObject* DataTypeObject, FCameraContextDataTableAllocationInfo& AllocationInfo)
-{
-	if (DataID)
-	{
-		FCameraContextDataDefinition DataDefinition;
-		DataDefinition.DataID = DataID;
-		DataDefinition.DataType = DataType;
-		DataDefinition.DataTypeObject = DataTypeObject;
-		AllocationInfo.DataDefinitions.Add(DataDefinition);
-	}
-}
-
 }  // namespace Internal
 
 FCameraRigAssetBuilder::FCameraRigAssetBuilder(FCameraBuildLog& InBuildLog)
@@ -1114,6 +1102,7 @@ void FCameraRigAssetBuilder::BuildParameterDefinitions()
 			Definition.ParameterType = ECameraRigInterfaceParameterType::Data;
 			Definition.DataID = DataParameter->PrivateDataID;
 			Definition.DataType = DataParameter->DataType;
+			Definition.DataContainerType = DataParameter->DataContainerType;
 			Definition.DataTypeObject = DataParameter->DataTypeObject;
 			ParameterDefinitions.Add(Definition);
 		}
diff --git a/Engine/Plugins/Cameras/GameplayCameras/Source/GameplayCameras/Private/Core/CameraRigAssetReference.cpp b/Engine/Plugins/Cameras/GameplayCameras/Source/GameplayCameras/Private/Core/CameraRigAssetReference.cpp
--- a/Engine/Plugins/Cameras/GameplayCameras/Source/GameplayCameras/Private/Core/CameraRigAssetReference.cpp
+++ b/Engine/Plugins/Cameras/GameplayCameras/Source/GameplayCameras/Private/Core/CameraRigAssetReference.cpp
@@ -267,14 +267,14 @@ void FCameraRigAssetReference::GetCustomCameraNodeParameters(FCustomCameraNodePa
 					if (ensure(PropertyDesc->ValueType == EPropertyBagPropertyType::Name))
 					{
 						const uint8* DefaultValue = reinterpret_cast<uint8*>(DefaultParameters.GetValueName(*PropertyDesc).TryGetValue());
-						OutParameterInfos.AddDataParameter(Definition.ParameterName, ECameraContextDataType::Name, nullptr, DefaultValue, &MetaData.OverrideDataID);
+						OutParameterInfos.AddDataParameter(Definition.ParameterName, ECameraContextDataType::Name, ECameraContextDataContainerType::None, nullptr, DefaultValue, &MetaData.OverrideDataID);
 					}
 					break;
 				case ECameraContextDataType::String:
 					if (ensure(PropertyDesc->ValueType == EPropertyBagPropertyType::String))
 					{
 						const uint8* DefaultValue = reinterpret_cast<uint8*>(DefaultParameters.GetValueString(*PropertyDesc).TryGetValue());
-						OutParameterInfos.AddDataParameter(Definition.ParameterName, ECameraContextDataType::String, nullptr, DefaultValue, &MetaData.OverrideDataID);
+						OutParameterInfos.AddDataParameter(Definition.ParameterName, ECameraContextDataType::String, ECameraContextDataContainerType::None, nullptr, DefaultValue, &MetaData.OverrideDataID);
 					}
 					break;
 				case ECameraContextDataType::Enum:
@@ -283,7 +283,7 @@ void FCameraRigAssetReference::GetCustomCameraNodeParameters(FCustomCameraNodePa
 					{
 						const UEnum* EnumType = CastChecked<const UEnum>(Definition.DataTypeObject);
 						const uint8* DefaultValue = reinterpret_cast<uint8*>(DefaultParameters.GetValueEnum(*PropertyDesc, EnumType).TryGetValue());
-						OutParameterInfos.AddDataParameter(Definition.ParameterName, ECameraContextDataType::Enum, EnumType, DefaultValue, &MetaData.OverrideDataID);
+						OutParameterInfos.AddDataParameter(Definition.ParameterName, ECameraContextDataType::Enum, ECameraContextDataContainerType::None, EnumType, DefaultValue, &MetaData.OverrideDataID);
 					}
 					break;
 				case ECameraContextDataType::Struct:
@@ -292,21 +292,21 @@ void FCameraRigAssetReference::GetCustomCameraNodeParameters(FCustomCameraNodePa
 					{
 						const UScriptStruct* DataType = CastChecked<const UScriptStruct>(Definition.DataTypeObject);
 						const uint8* DefaultValue = reinterpret_cast<uint8*>(DefaultParameters.GetValueStruct(*PropertyDesc, DataType).TryGetValue());
-						OutParameterInfos.AddDataParameter(Definition.ParameterName, ECameraContextDataType::Struct, DataType, DefaultValue, &MetaData.OverrideDataID);
+						OutParameterInfos.AddDataParameter(Definition.ParameterName, ECameraContextDataType::Struct, ECameraContextDataContainerType::None, DataType, DefaultValue, &MetaData.OverrideDataID);
 					}
 					break;
 				case ECameraContextDataType::Object:
 					if (ensure(PropertyDesc->ValueType == EPropertyBagPropertyType::Object))
 					{
 						const uint8* DefaultValue = reinterpret_cast<uint8*>(DefaultParameters.GetValueObject(*PropertyDesc).TryGetValue());
-						OutParameterInfos.AddDataParameter(Definition.ParameterName, ECameraContextDataType::Object, Definition.DataTypeObject, DefaultValue, &MetaData.OverrideDataID);
+						OutParameterInfos.AddDataParameter(Definition.ParameterName, ECameraContextDataType::Object, ECameraContextDataContainerType::None, Definition.DataTypeObject, DefaultValue, &MetaData.OverrideDataID);
 					}
 					break;
 				case ECameraContextDataType::Class:
 					if (ensure(PropertyDesc->ValueType == EPropertyBagPropertyType::Class))
 					{
 						const uint8* DefaultValue = reinterpret_cast<uint8*>(DefaultParameters.GetValueClass(*PropertyDesc).TryGetValue());
-						OutParameterInfos.AddDataParameter(Definition.ParameterName, ECameraContextDataType::Class, Definition.DataTypeObject, DefaultValue, &MetaData.OverrideDataID);
+						OutParameterInfos.AddDataParameter(Definition.ParameterName, ECameraContextDataType::Class, ECameraContextDataContainerType::None, Definition.DataTypeObject, DefaultValue, &MetaData.OverrideDataID);
 					}
 					break;
 			}
diff --git a/Engine/Plugins/Cameras/GameplayCameras/Source/GameplayCameras/Private/Core/CameraRigParameterDefinition.cpp b/Engine/Plugins/Cameras/GameplayCameras/Source/GameplayCameras/Private/Core/CameraRigParameterDefinition.cpp
--- a/Engine/Plugins/Cameras/GameplayCameras/Source/GameplayCameras/Private/Core/CameraRigParameterDefinition.cpp
+++ b/Engine/Plugins/Cameras/GameplayCameras/Source/GameplayCameras/Private/Core/CameraRigParameterDefinition.cpp
@@ -7,6 +7,7 @@
 #include "Core/CameraRigAsset.h"
 #include "Core/ICustomCameraNodeParameterProvider.h"
 #include "StructUtils/PropertyBag.h"
+#include "UObject/UnrealType.h"
 
 namespace UE::Cameras
 {
@@ -25,6 +26,7 @@ void FCameraRigParameterBuilder::AppendDefaultParameterProperties(const UCameraR
 	{
 		bool bIsValidProperty = true;
 		EPropertyBagPropertyType PropertyType = EPropertyBagPropertyType::Struct;
+		EPropertyBagContainerType ContainerType = EPropertyBagContainerType::None;
 		const UObject* PropertyTypeObject = nullptr;
 		EPropertyFlags PropertyFlags = CPF_None;
 
@@ -82,6 +84,13 @@ void FCameraRigParameterBuilder::AppendDefaultParameterProperties(const UCameraR
 					bIsValidProperty = false;
 					break;
 			}
+
+			switch (Definition.DataContainerType)
+			{
+				case ECameraContextDataContainerType::Array:
+					ContainerType = EPropertyBagContainerType::Array;
+					break;
+			}
 		}
 		else
 		{
@@ -90,7 +99,7 @@ void FCameraRigParameterBuilder::AppendDefaultParameterProperties(const UCameraR
 
 		if (ensure(bIsValidProperty))
 		{
-			FPropertyBagPropertyDesc NewProperty(Definition.ParameterName, PropertyType, PropertyTypeObject);
+			FPropertyBagPropertyDesc NewProperty(Definition.ParameterName, ContainerType, PropertyType, PropertyTypeObject);
 			// Make the property bag match the camera interface parameter GUIDs.
 			NewProperty.ID = Definition.ParameterGuid;
 			NewProperty.PropertyFlags |= PropertyFlags;
@@ -256,39 +265,48 @@ void FCameraRigParameterBuilder::SetDefaultParameterValues(const UCameraRigAsset
 			continue;
 		}
 
-		void* RawDestinationValuePtr = PropertyDesc->CachedProperty->ContainerPtrToValuePtr<void>(PropertyBagValue);
-
-		if (!ensure(RawDestinationValuePtr))
+		if (DataParameter->DataContainerType == ECameraContextDataContainerType::None)
 		{
-			continue;
+			void* RawDestinationValuePtr = PropertyDesc->CachedProperty->ContainerPtrToValuePtr<void>(PropertyBagValue);
+			if (ensure(RawDestinationValuePtr))
+			{
+				SetDefaultParameterValue(DataParameter, RawDestinationValuePtr, RawSourceValuePtr);
+			}
 		}
-
-		switch (DataParameter->DataType)
+		else if (DataParameter->DataContainerType == ECameraContextDataContainerType::Array)
 		{
-			case ECameraContextDataType::Name:
-				*((FName*)RawDestinationValuePtr) = *((FName*)RawSourceValuePtr);
-				break;
-			case ECameraContextDataType::String:
-				*((FString*)RawDestinationValuePtr) = *((FString*)RawSourceValuePtr);
-				break;
-			case ECameraContextDataType::Enum:
-				*((uint8*)RawDestinationValuePtr) = *((uint8*)RawSourceValuePtr);
-				break;
-			case ECameraContextDataType::Struct:
-				{
-					const UScriptStruct* StructType = CastChecked<const UScriptStruct>(DataParameter->DataTypeObject);
-					StructType->CopyScriptStruct(RawDestinationValuePtr, RawSourceValuePtr);
-				}
-				break;
-			case ECameraContextDataType::Object:
-				*((FObjectPtr*)RawDestinationValuePtr) = *((FObjectPtr*)RawSourceValuePtr);
-				break;
-			case ECameraContextDataType::Class:
-				*((FObjectPtr*)RawDestinationValuePtr) = *((FObjectPtr*)RawSourceValuePtr);
-				break;
+			// Array properties are empty by default.
 		}
 	}
 }
 
+void FCameraRigParameterBuilder::SetDefaultParameterValue(const UCameraRigDataParameter* DataParameter, void* DestValuePtr, const void* SrcValuePtr)
+{
+	switch (DataParameter->DataType)
+	{
+		case ECameraContextDataType::Name:
+			*((FName*)DestValuePtr) = *((FName*)SrcValuePtr);
+			break;
+		case ECameraContextDataType::String:
+			*((FString*)DestValuePtr) = *((FString*)SrcValuePtr);
+			break;
+		case ECameraContextDataType::Enum:
+			*((uint8*)DestValuePtr) = *((uint8*)SrcValuePtr);
+			break;
+		case ECameraContextDataType::Struct:
+			{
+				const UScriptStruct* StructType = CastChecked<const UScriptStruct>(DataParameter->DataTypeObject);
+				StructType->CopyScriptStruct(DestValuePtr, SrcValuePtr);
+			}
+			break;
+		case ECameraContextDataType::Object:
+			*((FObjectPtr*)DestValuePtr) = *((FObjectPtr*)SrcValuePtr);
+			break;
+		case ECameraContextDataType::Class:
+			*((FObjectPtr*)DestValuePtr) = *((FObjectPtr*)SrcValuePtr);
+			break;
+	}
+}
+
 }  // namespace UE::Cameras
 
diff --git a/Engine/Plugins/Cameras/GameplayCameras/Source/GameplayCameras/Private/Core/ICustomCameraNodeParameterProvider.cpp b/Engine/Plugins/Cameras/GameplayCameras/Source/GameplayCameras/Private/Core/ICustomCameraNodeParameterProvider.cpp
--- a/Engine/Plugins/Cameras/GameplayCameras/Source/GameplayCameras/Private/Core/ICustomCameraNodeParameterProvider.cpp
+++ b/Engine/Plugins/Cameras/GameplayCameras/Source/GameplayCameras/Private/Core/ICustomCameraNodeParameterProvider.cpp
@@ -27,11 +27,12 @@ void FCustomCameraNodeParameterInfos::AddBlendableParameter(FCustomCameraNodeBle
 void FCustomCameraNodeParameterInfos::AddDataParameter(
 		FName ParameterName, 
 		ECameraContextDataType ParameterType,
+		ECameraContextDataContainerType ParameterContainerType,
 		const UObject* ParameterTypeObject,
 		const uint8* DefaultValue,
 		FCameraContextDataID* OverrideDataID)
 {
-	DataParameters.Add({ ParameterName, ParameterType, ParameterTypeObject, DefaultValue, OverrideDataID });
+	DataParameters.Add({ ParameterName, ParameterType, ParameterContainerType, ParameterTypeObject, DefaultValue, OverrideDataID });
 }
 
 void FCustomCameraNodeParameterInfos::AddDataParameter(FCustomCameraNodeDataParameter& Parameter, const uint8* DefaultValue)
@@ -39,6 +40,7 @@ void FCustomCameraNodeParameterInfos::AddDataParameter(FCustomCameraNodeDataPara
 	AddDataParameter(
 			Parameter.ParameterName,
 			Parameter.ParameterType,
+			Parameter.ParameterContainerType,
 			Parameter.ParameterTypeObject,
 			DefaultValue,
 			&Parameter.OverrideDataID);
@@ -67,6 +69,7 @@ void FCustomCameraNodeParameterInfos::GetDataParameters(TArray<FCustomCameraNode
 		FCustomCameraNodeDataParameter& OutParameter = OutDataParameters.Emplace_GetRef();
 		OutParameter.ParameterName = DataParameter.ParameterName;
 		OutParameter.ParameterType = DataParameter.ParameterType;
+		OutParameter.ParameterContainerType = DataParameter.ParameterContainerType;
 		OutParameter.ParameterTypeObject = DataParameter.ParameterTypeObject;
 		if (DataParameter.OverrideDataID)
 		{
@@ -102,6 +105,7 @@ bool FCustomCameraNodeParameterInfos::FindDataParameter(FName ParameterName, FCu
 		{
 			OutParameter.ParameterName = DataParameter.ParameterName;
 			OutParameter.ParameterType = DataParameter.ParameterType;
+			OutParameter.ParameterContainerType = DataParameter.ParameterContainerType;
 			OutParameter.ParameterTypeObject = DataParameter.ParameterTypeObject;
 			if (DataParameter.OverrideDataID)
 			{
diff --git a/Engine/Plugins/Cameras/GameplayCameras/Source/GameplayCameras/Private/Debug/ContextDataTableDebugBlock.cpp b/Engine/Plugins/Cameras/GameplayCameras/Source/GameplayCameras/Private/Debug/ContextDataTableDebugBlock.cpp
--- a/Engine/Plugins/Cameras/GameplayCameras/Source/GameplayCameras/Private/Debug/ContextDataTableDebugBlock.cpp
+++ b/Engine/Plugins/Cameras/GameplayCameras/Source/GameplayCameras/Private/Debug/ContextDataTableDebugBlock.cpp
@@ -40,39 +40,8 @@ void FContextDataTableDebugBlock::Initialize(const FCameraContextDataTable& InCo
 			EntryTypeName = Entry.TypeObject->GetFName();
 		}
 
-		FString EntryValueStr;
 		const uint8* EntryData = InContextDataTable.Memory + Entry.Offset;
-		switch (Entry.Type)
-		{
-			case ECameraContextDataType::Name:
-				EntryValueStr = reinterpret_cast<const FName*>(EntryData)->ToString();
-				break;
-			case ECameraContextDataType::String:
-				EntryValueStr = *reinterpret_cast<const FString*>(EntryData);
-				break;
-			case ECameraContextDataType::Enum:
-				{
-					const UEnum* EnumType = CastChecked<const UEnum>(Entry.TypeObject);
-					EntryValueStr = EnumType->GetNameStringByValue((int64)*reinterpret_cast<const uint8*>(EntryData));
-				}
-				break;
-			case ECameraContextDataType::Struct:
-				{
-					const UScriptStruct* StructType = CastChecked<const UScriptStruct>(Entry.TypeObject);
-					const int32 ExportFlags = PPF_Delimited | PPF_IncludeTransient | PPF_ExternalEditor;
-					StructType->ExportText(EntryValueStr, EntryData, nullptr, nullptr, ExportFlags, nullptr);
-				}
-				break;
-			case ECameraContextDataType::Object:
-				EntryValueStr = reinterpret_cast<const TObjectPtr<UObject>*>(EntryData)->GetPathName();
-				break;
-			case ECameraContextDataType::Class:
-				EntryValueStr = reinterpret_cast<const TObjectPtr<UClass>*>(EntryData)->GetPathName();
-				break;
-			default:
-				ensure(false);
-				break;
-		}
+		FString EntryValueStr = GetDebugValueString(Entry.Type, Entry.ContainerType, Entry.TypeObject, EntryData);
 
 		FEntryDebugInfo EntryDebugInfo{ Entry.ID.GetValue(), EntryName, EntryTypeName, EntryValueStr};
 		EntryDebugInfo.bWritten = EnumHasAnyFlags(Entry.Flags, FCameraContextDataTable::EEntryFlags::Written);
@@ -86,6 +55,70 @@ void FContextDataTableDebugBlock::Initialize(const FCameraContextDataTable& InCo
 			});
 }
 
+FString FContextDataTableDebugBlock::GetDebugValueString(
+		ECameraContextDataType DataType, 
+		ECameraContextDataContainerType DataContainerType,
+		const UObject* DataTypeObject,
+		const uint8* DataPtr)
+{
+	if (DataContainerType == ECameraContextDataContainerType::None)
+	{
+		return GetDebugValueString(DataType, DataTypeObject, DataPtr);
+	}
+	else if (DataContainerType == ECameraContextDataContainerType::Array)
+	{
+		FString ArrayStr;
+		FCameraContextDataTable::FArrayEntryHelper Helper(DataType, DataTypeObject, const_cast<uint8*>(DataPtr));
+		for (int32 Index = 0; Index < Helper.Num(); ++Index)
+		{
+			const uint8* ElementPtr = Helper.GetRawPtr(Index);
+			ArrayStr += FString::Printf(TEXT("[%d] %s\n"), Index, *GetDebugValueString(DataType, DataTypeObject, ElementPtr));
+		}
+		return ArrayStr;
+	}
+	return FString();
+}
+
+FString FContextDataTableDebugBlock::GetDebugValueString(
+		ECameraContextDataType DataType, 
+		const UObject* DataTypeObject,
+		const uint8* DataPtr)
+{
+	FString EntryValueStr;
+	switch (DataType)
+	{
+		case ECameraContextDataType::Name:
+			EntryValueStr = reinterpret_cast<const FName*>(DataPtr)->ToString();
+			break;
+		case ECameraContextDataType::String:
+			EntryValueStr = *reinterpret_cast<const FString*>(DataPtr);
+			break;
+		case ECameraContextDataType::Enum:
+			{
+				const UEnum* EnumType = CastChecked<const UEnum>(DataTypeObject);
+				EntryValueStr = EnumType->GetNameStringByValue((int64)*reinterpret_cast<const uint8*>(DataPtr));
+			}
+			break;
+		case ECameraContextDataType::Struct:
+			{
+				const UScriptStruct* StructType = CastChecked<const UScriptStruct>(DataTypeObject);
+				const int32 ExportFlags = PPF_Delimited | PPF_IncludeTransient | PPF_ExternalEditor;
+				StructType->ExportText(EntryValueStr, DataPtr, nullptr, nullptr, ExportFlags, nullptr);
+			}
+			break;
+		case ECameraContextDataType::Object:
+			EntryValueStr = reinterpret_cast<const TObjectPtr<UObject>*>(DataPtr)->GetPathName();
+			break;
+		case ECameraContextDataType::Class:
+			EntryValueStr = reinterpret_cast<const TObjectPtr<UClass>*>(DataPtr)->GetPathName();
+			break;
+		default:
+			ensure(false);
+			break;
+	}
+	return EntryValueStr;
+}
+
 void FContextDataTableDebugBlock::OnDebugDraw(const FCameraDebugBlockDrawParams& Params, FCameraDebugRenderer& Renderer)
 {
 	const FCameraDebugColors& Colors = FCameraDebugColors::Get();
diff --git a/Engine/Plugins/Cameras/GameplayCameras/Source/GameplayCameras/Private/GameFramework/CameraRigParameterInterop.cpp b/Engine/Plugins/Cameras/GameplayCameras/Source/GameplayCameras/Private/GameFramework/CameraRigParameterInterop.cpp
--- a/Engine/Plugins/Cameras/GameplayCameras/Source/GameplayCameras/Private/GameFramework/CameraRigParameterInterop.cpp
+++ b/Engine/Plugins/Cameras/GameplayCameras/Source/GameplayCameras/Private/GameFramework/CameraRigParameterInterop.cpp
@@ -254,5 +254,65 @@ void UCameraRigParameterInterop::SetClassParameter(FBlueprintCameraNodeEvaluatio
 	}
 }
 
+void UCameraRigParameterInterop::SetNameArrayParameter(UPARAM(Ref) FBlueprintCameraNodeEvaluationResult& Result, UCameraRigAsset* CameraRig, const FString& ParameterName, const TArray<FName>& ParameterValue)
+{
+	const FCameraContextDataID DataID(UE::Cameras::Private::GetParameterPrivateDataID(CameraRig, ParameterName));
+	FBlueprintCameraContextDataTable ContextDataTable(Result.GetContextDataTable());
+	if (UE::Cameras::Private::ValidateSetCameraRigDataParameter(ContextDataTable, DataID))
+	{
+		ContextDataTable.GetContextDataTable()->SetNameArrayData(DataID, ParameterValue);
+	}
+}
+
+void UCameraRigParameterInterop::SetStringArrayParameter(UPARAM(Ref) FBlueprintCameraNodeEvaluationResult& Result, UCameraRigAsset* CameraRig, const FString& ParameterName, const TArray<FString>& ParameterValue)
+{
+	const FCameraContextDataID DataID(UE::Cameras::Private::GetParameterPrivateDataID(CameraRig, ParameterName));
+	FBlueprintCameraContextDataTable ContextDataTable(Result.GetContextDataTable());
+	if (UE::Cameras::Private::ValidateSetCameraRigDataParameter(ContextDataTable, DataID))
+	{
+		ContextDataTable.GetContextDataTable()->SetStringArrayData(DataID, ParameterValue);
+	}
+}
+
+void UCameraRigParameterInterop::SetEnumArrayParameter(UPARAM(Ref) FBlueprintCameraNodeEvaluationResult& Result, UCameraRigAsset* CameraRig, const FString& ParameterName, const UEnum* EnumType, const TArray<uint8>& ParameterValue)
+{
+	const FCameraContextDataID DataID(UE::Cameras::Private::GetParameterPrivateDataID(CameraRig, ParameterName));
+	FBlueprintCameraContextDataTable ContextDataTable(Result.GetContextDataTable());
+	if (UE::Cameras::Private::ValidateSetCameraRigDataParameter(ContextDataTable, DataID))
+	{
+		ContextDataTable.GetContextDataTable()->SetEnumArrayData(DataID, EnumType, ParameterValue);
+	}
+}
+
+void UCameraRigParameterInterop::SetStructArrayParameter(UPARAM(Ref) FBlueprintCameraNodeEvaluationResult& Result, UCameraRigAsset* CameraRig, const FString& ParameterName, const TArray<FInstancedStruct>& ParameterValue)
+{
+	const FCameraContextDataID DataID(UE::Cameras::Private::GetParameterPrivateDataID(CameraRig, ParameterName));
+	FBlueprintCameraContextDataTable ContextDataTable(Result.GetContextDataTable());
+	if (UE::Cameras::Private::ValidateSetCameraRigDataParameter(ContextDataTable, DataID))
+	{
+		ContextDataTable.GetContextDataTable()->SetInstancedStructArrayData(DataID, ParameterValue);
+	}
+}
+
+void UCameraRigParameterInterop::SetObjectArrayParameter(UPARAM(Ref) FBlueprintCameraNodeEvaluationResult& Result, UCameraRigAsset* CameraRig, const FString& ParameterName, const TArray<UObject*>& ParameterValue)
+{
+	const FCameraContextDataID DataID(UE::Cameras::Private::GetParameterPrivateDataID(CameraRig, ParameterName));
+	FBlueprintCameraContextDataTable ContextDataTable(Result.GetContextDataTable());
+	if (UE::Cameras::Private::ValidateSetCameraRigDataParameter(ContextDataTable, DataID))
+	{
+		ContextDataTable.GetContextDataTable()->SetObjectArrayData(DataID, ParameterValue);
+	}
+}
+
+void UCameraRigParameterInterop::SetClassArrayParameter(UPARAM(Ref) FBlueprintCameraNodeEvaluationResult& Result, UCameraRigAsset* CameraRig, const FString& ParameterName, const TArray<UClass*>& ParameterValue)
+{
+	const FCameraContextDataID DataID(UE::Cameras::Private::GetParameterPrivateDataID(CameraRig, ParameterName));
+	FBlueprintCameraContextDataTable ContextDataTable(Result.GetContextDataTable());
+	if (UE::Cameras::Private::ValidateSetCameraRigDataParameter(ContextDataTable, DataID))
+	{
+		ContextDataTable.GetContextDataTable()->SetClassArrayData(DataID, ParameterValue);
+	}
+}
+
 #undef LOCTEXT_NAMESPACE
 
diff --git a/Engine/Plugins/Cameras/GameplayCameras/Source/GameplayCameras/Private/Helpers/CameraParameterOverrideHelper.cpp b/Engine/Plugins/Cameras/GameplayCameras/Source/GameplayCameras/Private/Helpers/CameraParameterOverrideHelper.cpp
--- a/Engine/Plugins/Cameras/GameplayCameras/Source/GameplayCameras/Private/Helpers/CameraParameterOverrideHelper.cpp
+++ b/Engine/Plugins/Cameras/GameplayCameras/Source/GameplayCameras/Private/Helpers/CameraParameterOverrideHelper.cpp
@@ -7,6 +7,7 @@
 #include "Core/CameraRigAsset.h"
 #include "Core/CameraRigParameterDefinition.h"
 #include "Core/CameraVariableTable.h"
+#include "StructUtils/PropertyBag.h"
 
 namespace UE::Cameras
 {
@@ -110,7 +111,7 @@ void OverrideContextDataTableEntry(
 		FCameraContextDataTable& ContextDataTable)
 {
 	const uint8* RawParameterValue = reinterpret_cast<const uint8*>(&ParameterValue);
-	ContextDataTable.SetData(ParameterDefinition.DataID, ParameterDefinition.DataType, ParameterDefinition.DataTypeObject, RawParameterValue);
+	ContextDataTable.TrySetData(ParameterDefinition.DataID, ParameterDefinition.DataType, ParameterDefinition.DataTypeObject, RawParameterValue);
 }
 
 template<>
@@ -120,7 +121,29 @@ void OverrideContextDataTableEntry<FStructView>(
 		FCameraContextDataTable& ContextDataTable)
 {
 	const uint8* RawParameterValue = ParameterValue.GetMemory();
-	ContextDataTable.SetData(ParameterDefinition.DataID, ParameterDefinition.DataType, ParameterDefinition.DataTypeObject, RawParameterValue);
+	ContextDataTable.TrySetData(ParameterDefinition.DataID, ParameterDefinition.DataType, ParameterDefinition.DataTypeObject, RawParameterValue);
+}
+
+template<typename ParameterType>
+void OverrideContextDataTableEntryElement(
+		const FCameraRigParameterDefinition& ParameterDefinition,
+		int32 Index,
+		const ParameterType& ParameterValue,
+		FCameraContextDataTable& ContextDataTable)
+{
+	const uint8* RawParameterValue = reinterpret_cast<const uint8*>(&ParameterValue);
+	ContextDataTable.TrySetArrayData(ParameterDefinition.DataID, ParameterDefinition.DataType, ParameterDefinition.DataTypeObject, Index, RawParameterValue);
+}
+
+template<>
+void OverrideContextDataTableEntryElement<FStructView>(
+		const FCameraRigParameterDefinition& ParameterDefinition,
+		int32 Index,
+		const FStructView& ParameterValue,
+		FCameraContextDataTable& ContextDataTable)
+{
+	const uint8* RawParameterValue = ParameterValue.GetMemory();
+	ContextDataTable.TrySetArrayData(ParameterDefinition.DataID, ParameterDefinition.DataType, ParameterDefinition.DataTypeObject, Index, RawParameterValue);
 }
 
 template<typename ParameterType>
@@ -143,26 +166,34 @@ void ApplyDataParameterOverride(
 	OverrideContextDataTableEntry<ParameterType>(ParameterDefinition, ParameterValue, ContextDataTable);
 }
 
-void ApplyDataParameterOverride(
+template<typename ParameterType>
+void ApplyDataParameterElementOverride(
 		const UCameraRigAsset* CameraRig,
 		const FCameraRigParameterDefinition& ParameterDefinition,
-		const FInstancedPropertyBag& PropertyBag,
-		const FPropertyBagPropertyDesc& PropertyBagPropertyDesc,
+		const TValueOrError<ParameterType, EPropertyBagResult>& ElementValueOrError,
+		int32 Index,
 		FCameraContextDataTable& ContextDataTable)
 {
-	ensure(ParameterDefinition.ParameterType == ECameraRigInterfaceParameterType::Data);
-
-	if (!ParameterDefinition.DataID)
+	if (!ensureMsgf(
+				ElementValueOrError.HasValue() && !ElementValueOrError.HasError(),
+				TEXT("Camera parameter has no valid value! Error: %s"),
+				*UEnum::GetValueAsString(ElementValueOrError.GetError())))
 	{
-#if !WITH_EDITOR
-		UE_LOG(LogCameraSystem, Error,
-				TEXT("Invalid data parameter override '%s' in camera rig '%s'. Was it built/cooked?"),
-				*ParameterDefinition.ParameterName.ToString(),
-				*GetPathNameSafe(CameraRig));
 		return;
-#endif
 	}
 
+	// Write the override value into the context data table's entry array.
+	const ParameterType& ParameterValue = ElementValueOrError.GetValue();
+	OverrideContextDataTableEntryElement<ParameterType>(ParameterDefinition, Index, ParameterValue, ContextDataTable);
+}
+
+void ApplyDataParameterSingleOverride(
+		const UCameraRigAsset* CameraRig,
+		const FCameraRigParameterDefinition& ParameterDefinition,
+		const FInstancedPropertyBag& PropertyBag,
+		const FPropertyBagPropertyDesc& PropertyBagPropertyDesc,
+		FCameraContextDataTable& ContextDataTable)
+{
 	switch (ParameterDefinition.DataType)
 	{
 		case ECameraContextDataType::Name:
@@ -215,6 +246,120 @@ void ApplyDataParameterOverride(
 	}
 }
 
+void ApplyDataParameterArrayOverride(
+		const UCameraRigAsset* CameraRig,
+		const FCameraRigParameterDefinition& ParameterDefinition,
+		const FInstancedPropertyBag& PropertyBag,
+		const FPropertyBagPropertyDesc& PropertyBagPropertyDesc,
+		FCameraContextDataTable& ContextDataTable)
+{
+	TValueOrError<const FPropertyBagArrayRef, EPropertyBagResult> ArrayOrError = PropertyBag.GetArrayRef(PropertyBagPropertyDesc);
+	if (!ensureMsgf(
+				ArrayOrError.HasValue() && !ArrayOrError.HasError(),
+				TEXT("Camera parameter has no valid value! Error: %s"),
+				*UEnum::GetValueAsString(ArrayOrError.GetError())))
+	{
+		return;
+	}
+
+	const FPropertyBagArrayRef& ArrayRef = ArrayOrError.GetValue();
+	const int32 ArrayNum = ArrayRef.Num();
+
+	const bool bSetNumSuccess = ContextDataTable.TrySetArrayDataNum(ParameterDefinition.DataID, ArrayNum);
+	if (!ensureMsgf(bSetNumSuccess, TEXT("Camera parameter array '%s' can't be resized!"), *ParameterDefinition.ParameterName.ToString()))
+	{
+		return;
+	}
+
+	switch (ParameterDefinition.DataType)
+	{
+		case ECameraContextDataType::Name:
+			for (int32 Index = 0; Index < ArrayNum; ++Index)
+			{
+				TValueOrError<FName, EPropertyBagResult> ElementValueOrError = ArrayRef.GetValueName(Index);
+				ApplyDataParameterElementOverride(CameraRig, ParameterDefinition, ElementValueOrError, Index, ContextDataTable);
+			}
+			break;
+		case ECameraContextDataType::String:
+			for (int32 Index = 0; Index < ArrayNum; ++Index)
+			{
+				TValueOrError<FString, EPropertyBagResult> ParameterValueOrError = ArrayRef.GetValueString(Index);
+				ApplyDataParameterElementOverride(CameraRig, ParameterDefinition, ParameterValueOrError, Index, ContextDataTable);
+			}
+			break;
+		case ECameraContextDataType::Enum:
+			for (int32 Index = 0; Index < ArrayNum; ++Index)
+			{
+				const UEnum* EnumType = Cast<const UEnum>(ParameterDefinition.DataTypeObject);
+				if (ensure(EnumType))
+				{
+					TValueOrError<uint8, EPropertyBagResult> ParameterValueOrError = ArrayRef.GetValueEnum(Index, EnumType);
+					ApplyDataParameterElementOverride(CameraRig, ParameterDefinition, ParameterValueOrError, Index, ContextDataTable);
+				}
+			}
+			break;
+		case ECameraContextDataType::Struct:
+			for (int32 Index = 0; Index < ArrayNum; ++Index)
+			{
+				const UScriptStruct* StructType = Cast<const UScriptStruct>(ParameterDefinition.DataTypeObject);
+				if (ensure(StructType))
+				{
+					TValueOrError<FStructView, EPropertyBagResult> ParameterValueOrError = ArrayRef.GetValueStruct(Index, StructType);
+					ApplyDataParameterElementOverride(CameraRig, ParameterDefinition, ParameterValueOrError, Index, ContextDataTable);
+				}
+			}
+			break;
+		case ECameraContextDataType::Object:
+			for (int32 Index = 0; Index < ArrayNum; ++Index)
+			{
+				TValueOrError<UObject*, EPropertyBagResult> ParameterValueOrError = ArrayRef.GetValueObject(Index);
+				ApplyDataParameterElementOverride(CameraRig, ParameterDefinition, ParameterValueOrError, Index, ContextDataTable);
+			}
+			break;
+		case ECameraContextDataType::Class:
+			for (int32 Index = 0; Index < ArrayNum; ++Index)
+			{
+				TValueOrError<UClass*, EPropertyBagResult> ParameterValueOrError = ArrayRef.GetValueClass(Index);
+				ApplyDataParameterElementOverride(CameraRig, ParameterDefinition, ParameterValueOrError, Index, ContextDataTable);
+			}
+			break;
+		default:
+			ensure(false);
+			break;
+	}
+}
+
+void ApplyDataParameterOverride(
+		const UCameraRigAsset* CameraRig,
+		const FCameraRigParameterDefinition& ParameterDefinition,
+		const FInstancedPropertyBag& PropertyBag,
+		const FPropertyBagPropertyDesc& PropertyBagPropertyDesc,
+		FCameraContextDataTable& ContextDataTable)
+{
+	ensure(ParameterDefinition.ParameterType == ECameraRigInterfaceParameterType::Data);
+
+	if (!ParameterDefinition.DataID)
+	{
+#if !WITH_EDITOR
+		UE_LOG(LogCameraSystem, Error,
+				TEXT("Invalid data parameter override '%s' in camera rig '%s'. Was it built/cooked?"),
+				*ParameterDefinition.ParameterName.ToString(),
+				*GetPathNameSafe(CameraRig));
+		return;
+#endif
+	}
+
+	if (ParameterDefinition.DataContainerType == ECameraContextDataContainerType::None)
+	{
+		ApplyDataParameterSingleOverride(CameraRig, ParameterDefinition, PropertyBag, PropertyBagPropertyDesc, ContextDataTable);
+	}
+	else if (ParameterDefinition.DataContainerType == ECameraContextDataContainerType::Array)
+	{
+		ApplyDataParameterArrayOverride(CameraRig, ParameterDefinition, PropertyBag, PropertyBagPropertyDesc, ContextDataTable);
+	}
+
+}
+
 }  // namespace Internal
 
 FCameraParameterOverrideHelper::FCameraParameterOverrideHelper(FCameraVariableTable* OutVariableTable, FCameraContextDataTable* OutContextDataTable)
@@ -250,7 +395,9 @@ void FCameraParameterOverrideHelper::ApplyParameterOverride(
 			break;
 		case ECameraRigInterfaceParameterType::Data:
 			{
-				if (ensure(ContextDataTable))
+				// Data parameters can't be driven by variables so only apply them if we are not
+				// just applying driven parameters.
+				if (ensure(ContextDataTable) && !bDrivenOverridesOnly)
 				{
 					ApplyDataParameterOverride(
 							CameraRig, 
diff --git a/Engine/Plugins/Cameras/GameplayCameras/Source/GameplayCameras/Private/Nodes/Utility/BlueprintCameraNode.cpp b/Engine/Plugins/Cameras/GameplayCameras/Source/GameplayCameras/Private/Nodes/Utility/BlueprintCameraNode.cpp
--- a/Engine/Plugins/Cameras/GameplayCameras/Source/GameplayCameras/Private/Nodes/Utility/BlueprintCameraNode.cpp
+++ b/Engine/Plugins/Cameras/GameplayCameras/Source/GameplayCameras/Private/Nodes/Utility/BlueprintCameraNode.cpp
@@ -340,6 +340,7 @@ void UBlueprintCameraNode::RebuildOverrides()
 
 		bool bIsDataProperty = false;
 		ECameraContextDataType DataPropertyType = ECameraContextDataType::Name;
+		ECameraContextDataContainerType DataPropertyContainerType = ECameraContextDataContainerType::None;
 		const UObject* DataPropertyTypeObject = nullptr;
 
 		if (FBoolProperty* BoolProperty = CastField<FBoolProperty>(Property))
diff --git a/Engine/Plugins/Cameras/GameplayCameras/Source/GameplayCameras/Public/Core/CameraContextDataTable.h b/Engine/Plugins/Cameras/GameplayCameras/Source/GameplayCameras/Public/Core/CameraContextDataTable.h
--- a/Engine/Plugins/Cameras/GameplayCameras/Source/GameplayCameras/Public/Core/CameraContextDataTable.h
+++ b/Engine/Plugins/Cameras/GameplayCameras/Source/GameplayCameras/Public/Core/CameraContextDataTable.h
@@ -5,12 +5,11 @@
 #include "Core/CameraContextDataTableFwd.h"
 #include "Containers/Array.h"
 #include "Containers/Map.h"
+#include "Containers/ScriptArray.h"
 #include "GameplayCameras.h"
-#include "Math/NumericLimits.h"
 #include "StructUtils/InstancedStruct.h"
 #include "StructUtils/StructView.h"
 #include "Templates/PointerIsConvertibleFromTo.h"
-#include "Templates/Requires.h"
 #include <type_traits>
 
 struct FCameraContextDataDefinition;
@@ -90,6 +89,15 @@ public:
 	template<typename ValueType>
 	const ValueType* TryGetData(FCameraContextDataID InID) const;
 
+	template<typename ValueType>
+	TConstArrayView<ValueType> TryGetArrayData(FCameraContextDataID InID) const;
+
+	template<typename ValueType>
+	bool TryGetArrayData(FCameraContextDataID InID, TConstArrayView<ValueType>& OutValues) const;
+
+	template<typename ValueType>
+	const ValueType* TryGetArrayData(FCameraContextDataID InID, int32 ArrayIndex) const;
+
 	// Setter methods.
 
 	void SetNameData(FCameraContextDataID InID, const FName& InData);
@@ -107,6 +115,21 @@ public:
 	void SetStructViewData(FCameraContextDataID InID, const FStructView& InData);
 	void SetInstancedStructData(FCameraContextDataID InID, const FInstancedStruct& InData);
 
+	void SetNameArrayData(FCameraContextDataID InID, TConstArrayView<FName> InData);
+	void SetStringArrayData(FCameraContextDataID InID, TConstArrayView<FString> InData);
+	void SetEnumArrayData(FCameraContextDataID InID, const UEnum* EnumType, TConstArrayView<uint8> InData);
+	void SetObjectArrayData(FCameraContextDataID InID, TConstArrayView<UObject*> InData);
+	void SetClassArrayData(FCameraContextDataID InID, TConstArrayView<UClass*> InData);
+
+	template<typename EnumType>
+	void SetEnumArrayData(FCameraContextDataID InID, TConstArrayView<EnumType> InData);
+
+	template<typename StructType>
+	void SetStructArrayData(FCameraContextDataID InID, TConstArrayView<StructType> InData);
+
+	void SetStructViewArrayData(FCameraContextDataID InID, TConstArrayView<FStructView> InData);
+	void SetInstancedStructArrayData(FCameraContextDataID InID, TConstArrayView<FInstancedStruct> InData);
+
 public:
 
 	// Overriding.
@@ -122,6 +145,9 @@ public:
 
 public:
 
+	/** Type of data for array entries. */
+	using FEntryScriptArray = FScriptArray;
+
 	// Low-level API.
 	const uint8* GetData(
 			FCameraContextDataID DataID,
@@ -138,6 +164,11 @@ public:
 			ECameraContextDataType ExpectedDataType,
 			const UObject* ExpectedDataTypeObject);
 
+	const FEntryScriptArray* TryGetArrayData(
+			FCameraContextDataID DataID,
+			ECameraContextDataType ExpectedDataType,
+			const UObject* ExpectedDataTypeObject) const;
+
 	void SetData(
 			FCameraContextDataID DataID,
 			ECameraContextDataType ExpectedDataType,
@@ -152,6 +183,19 @@ public:
 			const uint8* InRawDataPtr,
 			bool bMarkAsWrittenThisFrame = true);
 
+	bool TrySetArrayDataNum(
+			FCameraContextDataID DataID, 
+			int32 Count,
+			bool bMarkAsWrittenThisFrame = true);
+
+	bool TrySetArrayData(
+			FCameraContextDataID DataID, 
+			ECameraContextDataType ExpectedDataType,
+			const UObject* ExpectedDataTypeObject,
+			int32 Index,
+			const uint8* InRawDataPtr,
+			bool bMarkAsWrittenThisFrame = true);
+
 	bool IsValueWritten(FCameraContextDataID InID) const;
 	void UnsetValue(FCameraContextDataID InID);
 	void UnsetAllValues();
@@ -173,6 +217,7 @@ private:
 	{
 		FCameraContextDataID ID;
 		ECameraContextDataType Type;
+		ECameraContextDataContainerType ContainerType;
 		TObjectPtr<const UObject> TypeObject;
 		uint32 Offset;
 		EEntryFlags Flags;
@@ -181,8 +226,34 @@ private:
 #endif
 	};
 
+	struct FArrayEntryHelper
+	{
+		FArrayEntryHelper(const FEntry& Entry, uint8* TableMemory);
+		FArrayEntryHelper(ECameraContextDataType DataType, const UObject* DataTypeObject, uint8* RawPtr);
+
+		bool IsValidIndex(int32 Index) const;
+		int32 Num() const;
+		uint8* GetRawPtr(int32 Index);
+		void Resize(int32 Count);
+
+		ECameraContextDataType ElementType;
+		TObjectPtr<const UObject> ElementTypeObject;
+
+		uint32 ElementSizeOf;
+		uint32 ElementAlignOf;
+
+		FEntryScriptArray* ScriptArray;
+	};
+
 	static bool GetDataTypeAllocationInfo(ECameraContextDataType DataType, const UObject* DataTypeObject, uint32& OutSizeOf, uint32& OutAlignOf);
-	static bool InitializeDefaultDataValue(ECameraContextDataType DataType, const UObject* DataTypeObject, uint8* DataPtr);
+	static bool GetDataTypeAllocationInfo(ECameraContextDataType DataType, ECameraContextDataContainerType DataContainerType, const UObject* DataTypeObject, uint32& OutSizeOf, uint32& OutAlignOf);
+
+	static bool ConstructDataValue(ECameraContextDataType DataType, const UObject* DataTypeObject, uint8* DataPtr);
+	static bool ConstructDataValue(ECameraContextDataType DataType, ECameraContextDataContainerType DataContainerType, const UObject* DataTypeObject, uint8* DataPtr);
+	static bool DestroyDataValue(ECameraContextDataType DataType, const UObject* DataTypeObject, uint8* DataPtr);
+	static bool DestroyDataValue(ECameraContextDataType DataType, ECameraContextDataContainerType DataContainerType, const UObject* DataTypeObject, uint8* DataPtr);
+
+	static bool SetDataValue(ECameraContextDataType DataType, ECameraContextDataContainerType DataContainerType, const UObject* DataTypeObject, uint8* DestDataPtr, const uint8* SrcDataPtr);
 	static bool SetDataValue(ECameraContextDataType DataType, const UObject* DataTypeObject, uint8* DestDataPtr, const uint8* SrcDataPtr);
 
 	const FEntry* FindEntry(FCameraContextDataID InID) const;
@@ -199,6 +270,12 @@ private:
 	template<typename StorageType>
 	bool SetDataImpl(FCameraContextDataID InID, ECameraContextDataType DataType, const UObject* DataTypeObject, const StorageType& InData);
 
+	template<typename StorageType>
+	TConstArrayView<StorageType> GetArrayDataImpl(FCameraContextDataID InID, ECameraContextDataType DataType, const UObject* DataTypeObject) const;
+
+	template<typename StorageType>
+	bool SetArrayDataImpl(FCameraContextDataID InID, ECameraContextDataType DataType, const UObject* DataTypeObject, TConstArrayView<StorageType> InData);
+
 private:
 
 	TArray<FEntry> Entries;
@@ -269,6 +346,42 @@ const ValueType* FCameraContextDataTable::TryGetData(FCameraContextDataID InID)
 	return nullptr;
 }
 
+template<typename ValueType>
+TConstArrayView<ValueType> FCameraContextDataTable::TryGetArrayData(FCameraContextDataID InID) const
+{
+	TConstArrayView<ValueType> Values;
+	if (TryGetArrayData(InID, Values))
+	{
+		return Values;
+	}
+	return TConstArrayView<ValueType>();
+}
+
+template<typename ValueType>
+bool FCameraContextDataTable::TryGetArrayData(FCameraContextDataID InID, TConstArrayView<ValueType>& OutValues) const
+{
+	ECameraContextDataType DataType = TCameraContextDataTraits<ValueType>::GetDataType();
+	const UObject* DataTypeObject = TCameraContextDataTraits<ValueType>::GetDataTypeObject();
+	if (const FEntryScriptArray* Array = TryGetArrayData(InID, DataType, DataTypeObject))
+	{
+		const ValueType* ArrayData = reinterpret_cast<const ValueType*>(Array->GetData());
+		OutValues = TConstArrayView<ValueType>(ArrayData, Array->Num());
+		return true;
+	}
+	return false;
+}
+
+template<typename ValueType>
+const ValueType* FCameraContextDataTable::TryGetArrayData(FCameraContextDataID InID, int32 ArrayIndex) const
+{
+	TConstArrayView<ValueType> TypedArrayView = TryGetArrayData<ValueType>(InID);
+	if (TypedArrayView.IsValidIndex(ArrayIndex))
+	{
+		return &TypedArrayView[ArrayIndex];
+	}
+	return nullptr;
+}
+
 template<typename EnumType>
 void FCameraContextDataTable::SetEnumData(FCameraContextDataID InID, EnumType InData)
 {
@@ -281,11 +394,26 @@ void FCameraContextDataTable::SetStructData(FCameraContextDataID InID, const Str
 	SetDataImpl(InID, ECameraContextDataType::Struct, StructType::StaticStruct(), InData);
 }
 
+template<typename EnumType>
+void FCameraContextDataTable::SetEnumArrayData(FCameraContextDataID InID, TConstArrayView<EnumType> InData)
+{
+	SetArrayDataImpl(InID, ECameraContextDataType::Enum, StaticEnum<EnumType>(), InData);
+}
+
+template<typename StructType>
+void FCameraContextDataTable::SetStructArrayData(FCameraContextDataID InID, TConstArrayView<StructType> InData)
+{
+	SetArrayDataImpl(InID, ECameraContextDataType::Struct, StructType::StaticStruct(), InData);
+}
+
 template<typename StorageType>
 const StorageType* FCameraContextDataTable::GetDataImpl(FCameraContextDataID InID, ECameraContextDataType DataType, const UObject* DataTypeObject) const
 {
 	const FEntry* Entry = FindEntry(InID);
-	if (Entry && Entry->Type == DataType && Entry->TypeObject == DataTypeObject)
+	if (Entry && 
+			Entry->Type == DataType && 
+			Entry->ContainerType == ECameraContextDataContainerType::None && 
+			Entry->TypeObject == DataTypeObject)
 	{
 		const uint8* RawData = Memory + Entry->Offset;
 		return reinterpret_cast<const StorageType*>(RawData);
@@ -297,7 +425,10 @@ template<typename StorageType>
 bool FCameraContextDataTable::SetDataImpl(FCameraContextDataID InID, ECameraContextDataType DataType, const UObject* DataTypeObject, const StorageType& InData)
 {
 	FEntry* Entry = FindEntry(InID);
-	if (Entry && Entry->Type == DataType && Entry->TypeObject == DataTypeObject)
+	if (Entry && 
+			Entry->Type == DataType && 
+			Entry->ContainerType == ECameraContextDataContainerType::None && 
+			Entry->TypeObject == DataTypeObject)
 	{
 		uint8* RawData = Memory + Entry->Offset;
 		*reinterpret_cast<StorageType*>(RawData) = InData;
@@ -307,6 +438,44 @@ bool FCameraContextDataTable::SetDataImpl(FCameraContextDataID InID, ECameraCont
 	return false;
 }
 
+template<typename StorageType>
+TConstArrayView<StorageType> FCameraContextDataTable::GetArrayDataImpl(FCameraContextDataID InID, ECameraContextDataType DataType, const UObject* DataTypeObject) const
+{
+	const FEntry* Entry = FindEntry(InID);
+	if (Entry && 
+			Entry->Type == DataType && 
+			Entry->ContainerType == ECameraContextDataContainerType::Array && 
+			Entry->TypeObject == DataTypeObject)
+	{
+		FEntryScriptArray* Array = (FEntryScriptArray*)(Memory + Entry->Offset);
+		const StorageType* ArrayData = reinterpret_cast<const StorageType*>(Array->GetData());
+		return TConstArrayView<StorageType>(ArrayData, Array->Num());
+	}
+	return TConstArrayView<StorageType>();
+}
+
+template<typename StorageType>
+bool FCameraContextDataTable::SetArrayDataImpl(FCameraContextDataID InID, ECameraContextDataType DataType, const UObject* DataTypeObject, TConstArrayView<StorageType> InData)
+{
+	FEntry* Entry = FindEntry(InID);
+	if (Entry && 
+			Entry->Type == DataType && 
+			Entry->ContainerType == ECameraContextDataContainerType::Array && 
+			Entry->TypeObject == DataTypeObject)
+	{
+		FArrayEntryHelper Helper(*Entry, Memory);
+		Helper.Resize(InData.Num());
+		for (int32 Index = 0; Index < InData.Num(); ++Index)
+		{
+			uint8* RawData = Helper.GetRawPtr(Index);
+			*reinterpret_cast<StorageType*>(RawData) = InData[Index];
+		}
+		EnumAddFlags(Entry->Flags, EEntryFlags::Written | EEntryFlags::WrittenThisFrame);
+		return true;
+	}
+	return false;
+}
+
 template<typename DataType>
 ECameraContextDataType TCameraContextDataTraits<DataType>::GetDataType()
 {
diff --git a/Engine/Plugins/Cameras/GameplayCameras/Source/GameplayCameras/Public/Core/CameraContextDataTableAllocationInfo.h b/Engine/Plugins/Cameras/GameplayCameras/Source/GameplayCameras/Public/Core/CameraContextDataTableAllocationInfo.h
--- a/Engine/Plugins/Cameras/GameplayCameras/Source/GameplayCameras/Public/Core/CameraContextDataTableAllocationInfo.h
+++ b/Engine/Plugins/Cameras/GameplayCameras/Source/GameplayCameras/Public/Core/CameraContextDataTableAllocationInfo.h
@@ -27,6 +27,10 @@ struct FCameraContextDataDefinition
 	UPROPERTY()
 	ECameraContextDataType DataType = ECameraContextDataType::Name;
 
+	/** The type of container for the data. */
+	UPROPERTY()
+	ECameraContextDataContainerType DataContainerType = ECameraContextDataContainerType::None;
+
 	/** An extra type object for the data. */
 	UPROPERTY()
 	TObjectPtr<const UObject> DataTypeObject;
diff --git a/Engine/Plugins/Cameras/GameplayCameras/Source/GameplayCameras/Public/Core/CameraContextDataTableFwd.h b/Engine/Plugins/Cameras/GameplayCameras/Source/GameplayCameras/Public/Core/CameraContextDataTableFwd.h
--- a/Engine/Plugins/Cameras/GameplayCameras/Source/GameplayCameras/Public/Core/CameraContextDataTableFwd.h
+++ b/Engine/Plugins/Cameras/GameplayCameras/Source/GameplayCameras/Public/Core/CameraContextDataTableFwd.h
@@ -35,6 +35,16 @@ enum class ECameraContextDataType
 	Count UMETA(Hidden)
 };
 
+/**
+ * Supported container types for a camera node's context data.
+ */
+UENUM()
+enum class ECameraContextDataContainerType
+{
+	None,
+	Array
+};
+
 /**
  * The ID of a context data, used to refer to it in a camera context data table.
  */
diff --git a/Engine/Plugins/Cameras/GameplayCameras/Source/GameplayCameras/Public/Core/CameraRigAsset.h b/Engine/Plugins/Cameras/GameplayCameras/Source/GameplayCameras/Public/Core/CameraRigAsset.h
--- a/Engine/Plugins/Cameras/GameplayCameras/Source/GameplayCameras/Public/Core/CameraRigAsset.h
+++ b/Engine/Plugins/Cameras/GameplayCameras/Source/GameplayCameras/Public/Core/CameraRigAsset.h
@@ -219,7 +219,11 @@ public:
 
 	/** The type of this parameter. */
 	UPROPERTY()
-	ECameraContextDataType DataType;
+	ECameraContextDataType DataType = ECameraContextDataType::Name;
+
+	/** The type of container for this parameter. */
+	UPROPERTY()
+	ECameraContextDataContainerType DataContainerType = ECameraContextDataContainerType::None;
 
 	/** An additional type object for this parameter. */
 	UPROPERTY()
diff --git a/Engine/Plugins/Cameras/GameplayCameras/Source/GameplayCameras/Public/Core/CameraRigParameterDefinition.h b/Engine/Plugins/Cameras/GameplayCameras/Source/GameplayCameras/Public/Core/CameraRigParameterDefinition.h
--- a/Engine/Plugins/Cameras/GameplayCameras/Source/GameplayCameras/Public/Core/CameraRigParameterDefinition.h
+++ b/Engine/Plugins/Cameras/GameplayCameras/Source/GameplayCameras/Public/Core/CameraRigParameterDefinition.h
@@ -9,6 +9,7 @@
 #include "CameraRigParameterDefinition.generated.h"
 
 class UCameraRigAsset;
+class UCameraRigDataParameter;
 struct FInstancedPropertyBag;
 struct FPropertyBagPropertyDesc;
 
@@ -74,6 +75,10 @@ struct FCameraRigParameterDefinition
 	UPROPERTY()
 	ECameraContextDataType DataType = ECameraContextDataType::Name;
 
+	/** The type of container that drives this blendable parameter. */
+	UPROPERTY()
+	ECameraContextDataContainerType DataContainerType = ECameraContextDataContainerType::None;
+
 	/** The type object of the data that drives this blendable parameter. */
 	UPROPERTY()
 	TObjectPtr<const UObject> DataTypeObject;
@@ -115,6 +120,10 @@ public:
 	static void AppendDefaultParameterProperties(const UCameraRigAsset* CameraRig, TArray<FPropertyBagPropertyDesc>& OutProperties);
 	/** Sets the default value of the property bag properties that correspond to parameters on the given camera rig. */
 	static void SetDefaultParameterValues(const UCameraRigAsset* CameraRig, FInstancedPropertyBag& PropertyBag);
+
+private:
+
+	static void SetDefaultParameterValue(const UCameraRigDataParameter* DataParameter, void* DestValuePtr, const void* SrcValuePtr);
 };
 
 }  // namespace UE::Cameras
diff --git a/Engine/Plugins/Cameras/GameplayCameras/Source/GameplayCameras/Public/Core/ICustomCameraNodeParameterProvider.h b/Engine/Plugins/Cameras/GameplayCameras/Source/GameplayCameras/Public/Core/ICustomCameraNodeParameterProvider.h
--- a/Engine/Plugins/Cameras/GameplayCameras/Source/GameplayCameras/Public/Core/ICustomCameraNodeParameterProvider.h
+++ b/Engine/Plugins/Cameras/GameplayCameras/Source/GameplayCameras/Public/Core/ICustomCameraNodeParameterProvider.h
@@ -61,6 +61,10 @@ struct FCustomCameraNodeDataParameter
 	UPROPERTY()
 	ECameraContextDataType ParameterType = ECameraContextDataType::Name;
 
+	/** The type of the parameter container. */
+	UPROPERTY()
+	ECameraContextDataContainerType ParameterContainerType = ECameraContextDataContainerType::None;
+
 	/** An extra type object for the parameter. */
 	UPROPERTY()
 	TObjectPtr<const UObject> ParameterTypeObject;
@@ -125,6 +129,7 @@ struct FCustomCameraNodeParameterInfos
 	GAMEPLAYCAMERAS_API void AddDataParameter(
 			FName ParameterName, 
 			ECameraContextDataType ParameterType,
+			ECameraContextDataContainerType ParameterContainerType,
 			const UObject* ParameterTypeObject,
 			const uint8* DefaultValue,
 			FCameraContextDataID* OverrideDataID);
@@ -157,6 +162,7 @@ private:
 	{
 		FName ParameterName;
 		ECameraContextDataType ParameterType;
+		ECameraContextDataContainerType ParameterContainerType;
 		const UObject* ParameterTypeObject;
 		const uint8* DefaultValue = nullptr;
 		FCameraContextDataID* OverrideDataID = nullptr;
diff --git a/Engine/Plugins/Cameras/GameplayCameras/Source/GameplayCameras/Public/Debug/ContextDataTableDebugBlock.h b/Engine/Plugins/Cameras/GameplayCameras/Source/GameplayCameras/Public/Debug/ContextDataTableDebugBlock.h
--- a/Engine/Plugins/Cameras/GameplayCameras/Source/GameplayCameras/Public/Debug/ContextDataTableDebugBlock.h
+++ b/Engine/Plugins/Cameras/GameplayCameras/Source/GameplayCameras/Public/Debug/ContextDataTableDebugBlock.h
@@ -2,6 +2,7 @@
 
 #pragma once
 
+#include "Core/CameraContextDataTableFwd.h"
 #include "Debug/CameraDebugBlock.h"
 
 #if UE_GAMEPLAY_CAMERAS_DEBUG
@@ -32,6 +33,17 @@ protected:
 
 	void Initialize(const FCameraContextDataTable& InContextDataTable);
 
+	static FString GetDebugValueString(
+			ECameraContextDataType DataType, 
+			ECameraContextDataContainerType DataContainerType,
+			const UObject* DataTypeObject,
+			const uint8* DataPtr);
+
+	static FString GetDebugValueString(
+			ECameraContextDataType DataType, 
+			const UObject* DataTypeObject,
+			const uint8* DataPtr);
+
 private:
 
 	struct FEntryDebugInfo
diff --git a/Engine/Plugins/Cameras/GameplayCameras/Source/GameplayCameras/Public/GameFramework/CameraRigParameterInterop.h b/Engine/Plugins/Cameras/GameplayCameras/Source/GameplayCameras/Public/GameFramework/CameraRigParameterInterop.h
--- a/Engine/Plugins/Cameras/GameplayCameras/Source/GameplayCameras/Public/GameFramework/CameraRigParameterInterop.h
+++ b/Engine/Plugins/Cameras/GameplayCameras/Source/GameplayCameras/Public/GameFramework/CameraRigParameterInterop.h
@@ -83,5 +83,25 @@ public:
 
 	UFUNCTION(BlueprintCallable, Category="Camera", meta=(BlueprintInternalUseOnly="true"))
 	static void SetClassParameter(UPARAM(Ref) FBlueprintCameraNodeEvaluationResult& Result, UCameraRigAsset* CameraRig, const FString& ParameterName, UClass* ParameterValue);
+
+public:
+
+	UFUNCTION(BlueprintCallable, Category="Camera", meta=(BlueprintInternalUseOnly="true"))
+	static void SetNameArrayParameter(UPARAM(Ref) FBlueprintCameraNodeEvaluationResult& Result, UCameraRigAsset* CameraRig, const FString& ParameterName, const TArray<FName>& ParameterValue);
+
+	UFUNCTION(BlueprintCallable, Category="Camera", meta=(BlueprintInternalUseOnly="true"))
+	static void SetStringArrayParameter(UPARAM(Ref) FBlueprintCameraNodeEvaluationResult& Result, UCameraRigAsset* CameraRig, const FString& ParameterName, const TArray<FString>& ParameterValue);
+
+	UFUNCTION(BlueprintCallable, Category="Camera", meta=(BlueprintInternalUseOnly="true"))
+	static void SetEnumArrayParameter(UPARAM(Ref) FBlueprintCameraNodeEvaluationResult& Result, UCameraRigAsset* CameraRig, const FString& ParameterName, const UEnum* EnumType, const TArray<uint8>& ParameterValue);
+
+	UFUNCTION(BlueprintCallable, Category="Camera", meta=(BlueprintInternalUseOnly="true"))
+	static void SetStructArrayParameter(UPARAM(Ref) FBlueprintCameraNodeEvaluationResult& Result, UCameraRigAsset* CameraRig, const FString& ParameterName, const TArray<FInstancedStruct>& ParameterValue);
+
+	UFUNCTION(BlueprintCallable, Category="Camera", meta=(BlueprintInternalUseOnly="true"))
+	static void SetObjectArrayParameter(UPARAM(Ref) FBlueprintCameraNodeEvaluationResult& Result, UCameraRigAsset* CameraRig, const FString& ParameterName, const TArray<UObject*>& ParameterValue);
+
+	UFUNCTION(BlueprintCallable, Category="Camera", meta=(BlueprintInternalUseOnly="true"))
+	static void SetClassArrayParameter(UPARAM(Ref) FBlueprintCameraNodeEvaluationResult& Result, UCameraRigAsset* CameraRig, const FString& ParameterName, const TArray<UClass*>& ParameterValue);
 };
 
diff --git a/Engine/Plugins/Cameras/GameplayCameras/Source/GameplayCamerasEditor/Private/Editors/CameraNodeGraphNode.cpp b/Engine/Plugins/Cameras/GameplayCameras/Source/GameplayCamerasEditor/Private/Editors/CameraNodeGraphNode.cpp
--- a/Engine/Plugins/Cameras/GameplayCameras/Source/GameplayCamerasEditor/Private/Editors/CameraNodeGraphNode.cpp
+++ b/Engine/Plugins/Cameras/GameplayCameras/Source/GameplayCamerasEditor/Private/Editors/CameraNodeGraphNode.cpp
@@ -2,6 +2,7 @@
 
 #include "Editors/CameraNodeGraphNode.h"
 
+#include "Core/CameraContextDataTableFwd.h"
 #include "Core/CameraNode.h"
 #include "Core/CameraParameters.h"
 #include "Core/CameraVariableReferences.h"
@@ -56,51 +57,95 @@ void UCameraNodeGraphNode::AllocateDefaultPins()
 	const FName ContextDataMetaData = TEXT("CameraContextData");
 	for (TFieldIterator<FProperty> PropertyIt(CameraNodeClass); PropertyIt; ++PropertyIt)
 	{
-		FStructProperty* StructProperty = CastField<FStructProperty>(*PropertyIt);
-		if (!StructProperty)
-		{
-			continue;
-		}
-
 		const FName PropertyName = PropertyIt->GetFName();
-		
+
 		FEdGraphPinType PinType;
 		const FText PinFriendlyName = FText::FromName(PropertyName);
-		const FString PinToolTip = StructProperty->Struct->GetDisplayNameText().ToString();
+		
+		if (FStructProperty* StructProperty = CastField<FStructProperty>(*PropertyIt))
+		{
+			const FString PinToolTip = StructProperty->Struct->GetDisplayNameText().ToString();
 
 #define UE_CAMERA_VARIABLE_FOR_TYPE(ValueType, ValueName)\
-		if (StructProperty->Struct == F##ValueName##CameraParameter::StaticStruct())\
-		{\
-			PinType.PinCategory = UCameraNodeGraphSchema::PC_CameraParameter;\
-			PinType.PinSubCategory = UEnum::GetValueAsName(ECameraVariableType::ValueName);\
-			PinType.PinSubCategoryObject = F##ValueName##CameraParameter::StaticStruct();\
-			UEdGraphPin* ParameterPin = CreatePin(EGPD_Input, PinType, PropertyName);\
-			ParameterPin->PinFriendlyName = PinFriendlyName;\
-			ParameterPin->PinToolTip = PinToolTip;\
-			continue;\
-		}\
-		if (StructProperty->Struct == F##ValueName##CameraVariableReference::StaticStruct())\
-		{\
-			PinType.PinCategory = UCameraNodeGraphSchema::PC_CameraVariableReference;\
-			PinType.PinSubCategory = UEnum::GetValueAsName(ECameraVariableType::ValueName);\
-			PinType.PinSubCategoryObject = F##ValueName##CameraVariableReference::StaticStruct();\
-			UEdGraphPin* VariableReferencePin = CreatePin(EGPD_Input, PinType, PropertyName);\
-			VariableReferencePin->PinFriendlyName = PinFriendlyName;\
-			VariableReferencePin->PinToolTip = PinToolTip;\
-			continue;\
-		}
-UE_CAMERA_VARIABLE_FOR_ALL_TYPES()
+			if (StructProperty->Struct == F##ValueName##CameraParameter::StaticStruct())\
+			{\
+				PinType.PinCategory = UCameraNodeGraphSchema::PC_CameraParameter;\
+				PinType.PinSubCategory = UEnum::GetValueAsName(ECameraVariableType::ValueName);\
+				PinType.PinSubCategoryObject = F##ValueName##CameraParameter::StaticStruct();\
+				UEdGraphPin* ParameterPin = CreatePin(EGPD_Input, PinType, PropertyName);\
+				ParameterPin->PinFriendlyName = PinFriendlyName;\
+				ParameterPin->PinToolTip = PinToolTip;\
+				continue;\
+			}\
+			if (StructProperty->Struct == F##ValueName##CameraVariableReference::StaticStruct())\
+			{\
+				PinType.PinCategory = UCameraNodeGraphSchema::PC_CameraVariableReference;\
+				PinType.PinSubCategory = UEnum::GetValueAsName(ECameraVariableType::ValueName);\
+				PinType.PinSubCategoryObject = F##ValueName##CameraVariableReference::StaticStruct();\
+				UEdGraphPin* VariableReferencePin = CreatePin(EGPD_Input, PinType, PropertyName);\
+				VariableReferencePin->PinFriendlyName = PinFriendlyName;\
+				VariableReferencePin->PinToolTip = PinToolTip;\
+				continue;\
+			}
+			UE_CAMERA_VARIABLE_FOR_ALL_TYPES()
 #undef UE_CAMERA_VARIABLE_FOR_TYPE
+		}
 
-		if (StructProperty->HasMetaData(ContextDataMetaData))
+		if (PropertyIt->HasMetaData(ContextDataMetaData))
 		{
-			PinType.PinCategory = UCameraNodeGraphSchema::PC_CameraContextData;
-			PinType.PinSubCategory = UEnum::GetValueAsName(ECameraContextDataType::Struct);
-			PinType.PinSubCategoryObject = StructProperty->Struct;
-			UEdGraphPin* ContextDataPin = CreatePin(EGPD_Input, PinType, PropertyName);
-			ContextDataPin->PinFriendlyName = PinFriendlyName;
-			ContextDataPin->PinToolTip = PinToolTip;
-			continue;
+			bool bGotValidDataProperty = true;
+			FProperty* DataProperty = *PropertyIt;
+			FString PinToolTip;
+
+			if (FArrayProperty* ArrayProperty = CastField<FArrayProperty>(DataProperty))
+			{
+				PinType.ContainerType = EPinContainerType::Array;
+				DataProperty = ArrayProperty->Inner;
+			}
+
+			if (FNameProperty* NameProperty = CastField<FNameProperty>(DataProperty))
+			{
+				PinType.PinSubCategory = UEnum::GetValueAsName(ECameraContextDataType::Name);
+			}
+			else if (FStrProperty* StringProperty = CastField<FStrProperty>(DataProperty))
+			{
+				PinType.PinSubCategory = UEnum::GetValueAsName(ECameraContextDataType::String);
+			}
+			else if (FEnumProperty* EnumProperty = CastField<FEnumProperty>(DataProperty))
+			{
+				PinType.PinSubCategory = UEnum::GetValueAsName(ECameraContextDataType::Enum);
+				PinType.PinSubCategoryObject = EnumProperty->GetEnum();
+				PinToolTip = EnumProperty->GetEnum()->GetDisplayNameText().ToString();
+			}
+			else if (FStructProperty* StructProperty = CastField<FStructProperty>(DataProperty))
+			{
+				PinType.PinSubCategory = UEnum::GetValueAsName(ECameraContextDataType::Struct);
+				PinType.PinSubCategoryObject = StructProperty->Struct;
+				PinToolTip = StructProperty->Struct->GetDisplayNameText().ToString();
+			}
+			else if (FObjectProperty* ObjectProperty = CastField<FObjectProperty>(DataProperty))
+			{
+				PinType.PinSubCategory = UEnum::GetValueAsName(ECameraContextDataType::Object);
+				PinType.PinSubCategoryObject = ObjectProperty->PropertyClass;
+			}
+			else if (FClassProperty* ClassProperty = CastField<FClassProperty>(DataProperty))
+			{
+				PinType.PinSubCategory = UEnum::GetValueAsName(ECameraContextDataType::Class);
+				PinType.PinSubCategoryObject = ClassProperty->PropertyClass;
+			}
+			else
+			{
+				bGotValidDataProperty = false;
+			}
+
+			if (bGotValidDataProperty)
+			{
+				PinType.PinCategory = UCameraNodeGraphSchema::PC_CameraContextData;
+				UEdGraphPin* ContextDataPin = CreatePin(EGPD_Input, PinType, PropertyName);
+				ContextDataPin->PinFriendlyName = PinFriendlyName;
+				ContextDataPin->PinToolTip = PinToolTip;
+				continue;
+			}
 		}
 	}
 
@@ -151,6 +196,11 @@ UE_CAMERA_VARIABLE_FOR_ALL_TYPES()
 			PinType.PinSubCategory = DataTypeEnum->GetNameByValue((int64)DataParameter.ParameterType);
 			PinType.PinSubCategoryObject = const_cast<UObject*>(DataParameter.ParameterTypeObject.Get());
 
+			if (DataParameter.ParameterContainerType == ECameraContextDataContainerType::Array)
+			{
+				PinType.ContainerType = EPinContainerType::Array;
+			}
+
 			UEdGraphPin* ContextDataPin = CreatePin(EGPD_Input, PinType, DataParameter.ParameterName);
 			ContextDataPin->PinFriendlyName = FText::FromName(DataParameter.ParameterName);
 
diff --git a/Engine/Plugins/Cameras/GameplayCameras/Source/GameplayCamerasEditor/Private/Editors/CameraNodeGraphSchema.cpp b/Engine/Plugins/Cameras/GameplayCameras/Source/GameplayCamerasEditor/Private/Editors/CameraNodeGraphSchema.cpp
--- a/Engine/Plugins/Cameras/GameplayCameras/Source/GameplayCamerasEditor/Private/Editors/CameraNodeGraphSchema.cpp
+++ b/Engine/Plugins/Cameras/GameplayCameras/Source/GameplayCamerasEditor/Private/Editors/CameraNodeGraphSchema.cpp
@@ -265,16 +265,24 @@ void UCameraNodeGraphSchema::GetGraphContextActions(FGraphContextMenuBuilder& Co
 			else if (DraggedPin->PinType.PinCategory == PC_CameraContextData)
 			{
 				ECameraContextDataType DataType;
+				ECameraContextDataContainerType DataContainerType = ECameraContextDataContainerType::None;
 				const UObject* DataTypeObject = nullptr;
 				
 				FCustomCameraNodeDataParameter DataParameter;
 				if (CustomParameters.FindDataParameter(DraggedPin->PinName, DataParameter))
 				{
 					DataType = DataParameter.ParameterType;
+					DataContainerType = DataParameter.ParameterContainerType;
 					DataTypeObject = DataParameter.ParameterTypeObject;
 				}
 				else if (Property)
 				{
+					if (FArrayProperty* ArrayProperty = CastField<FArrayProperty>(Property))
+					{
+						DataContainerType = ECameraContextDataContainerType::Array;
+						Property = ArrayProperty->Inner;
+					}
+
 					if (FNameProperty* NameProperty = CastField<FNameProperty>(Property))
 					{
 						DataType = ECameraContextDataType::Name;
@@ -320,6 +328,7 @@ void UCameraNodeGraphSchema::GetGraphContextActions(FGraphContextMenuBuilder& Co
 				FCameraRigParameterDefinition NewParameterDefinition;
 				NewParameterDefinition.ParameterType = ECameraRigInterfaceParameterType::Data;
 				NewParameterDefinition.DataType = DataType;
+				NewParameterDefinition.DataContainerType = DataContainerType;
 				NewParameterDefinition.DataTypeObject = DataTypeObject;
 				Action->ParameterDefinition = NewParameterDefinition;
 			}
@@ -591,6 +600,7 @@ UEdGraphNode* FCameraNodeGraphSchemaAction_NewInterfaceParameterNode::PerformAct
 	{
 		UCameraRigDataParameter* NewDataParameter = NewObject<UCameraRigDataParameter>(CameraRig, NAME_None, RF_Transactional);
 		NewDataParameter->DataType = ParameterDefinition.DataType;
+		NewDataParameter->DataContainerType = ParameterDefinition.DataContainerType;
 		NewDataParameter->DataTypeObject = ParameterDefinition.DataTypeObject;
 		CameraRig->Interface.DataParameters.Add(NewDataParameter);
 		NewInterfaceParameter = NewDataParameter;
diff --git a/Engine/Plugins/Cameras/GameplayCameras/Source/GameplayCamerasEditor/Private/Toolkits/CameraRigInterfaceParametersToolkit.cpp b/Engine/Plugins/Cameras/GameplayCameras/Source/GameplayCamerasEditor/Private/Toolkits/CameraRigInterfaceParametersToolkit.cpp
--- a/Engine/Plugins/Cameras/GameplayCameras/Source/GameplayCamerasEditor/Private/Toolkits/CameraRigInterfaceParametersToolkit.cpp
+++ b/Engine/Plugins/Cameras/GameplayCameras/Source/GameplayCamerasEditor/Private/Toolkits/CameraRigInterfaceParametersToolkit.cpp
@@ -3,7 +3,7 @@
 #include "Toolkits/CameraRigInterfaceParametersToolkit.h"
 
 #include "Core/CameraRigAsset.h"
-#include "EdGraphSchema_K2.h"
+#include "EdGraph/EdGraphNode.h"
 #include "Editor.h"
 #include "Editors/CameraNodeGraphDragDropOp.h"
 #include "Framework/MultiBox/MultiBoxBuilder.h"
@@ -154,7 +154,7 @@ protected:
 		{
 			TSharedRef<FGameplayCamerasEditorStyle> GameplayCamerasEditorStyle = FGameplayCamerasEditorStyle::Get();
 
-			const UEdGraphSchema_K2* K2Schema = GetDefault<UEdGraphSchema_K2>();
+			const UEdGraphSchema_K2* K2Schema = GetDefault<UEdGraphSchema_CameraNodeK2>();
 
 			return SNew(SBox)
 				.MinDesiredWidth(125.f)
@@ -188,7 +188,7 @@ protected:
 	{
 		using FPinTypeTreeInfo = UEdGraphSchema_K2::FPinTypeTreeInfo;
 
-		const UEdGraphSchema_K2* K2Schema = GetDefault<UEdGraphSchema_K2>();
+		const UEdGraphSchema_K2* K2Schema = GetDefault<UEdGraphSchema_CameraNodeK2>();
 
 		TypeTree.Reset();
 
@@ -452,7 +452,7 @@ protected:
 		{
 			TSharedRef<FGameplayCamerasEditorStyle> GameplayCamerasEditorStyle = FGameplayCamerasEditorStyle::Get();
 
-			const UEdGraphSchema_K2* K2Schema = GetDefault<UEdGraphSchema_K2>();
+			const UEdGraphSchema_K2* K2Schema = GetDefault<UEdGraphSchema_CameraNodeK2>();
 			TArray<TSharedPtr<IPinTypeSelectorFilter>> PinTypeSelectorFilters{ MakeShared<FDataParameterPinTypeSelectorFilter>() };
 
 			return SNew(SBox)
@@ -465,7 +465,7 @@ protected:
 					.ReadOnly(this, &SCameraRigDataParameterTableRow::IsPinTypeSelectorReadOnly)
 					.Schema(K2Schema)
 					.CustomFilters(PinTypeSelectorFilters)
-					.bAllowArrays(false)
+					.bAllowArrays(true)
 				];
 		}
 
@@ -505,6 +505,13 @@ protected:
 				break;
 		}
 
+		switch (Item->DataContainerType)
+		{
+			case ECameraContextDataContainerType::Array:
+				PinType.ContainerType = EPinContainerType::Array;
+				break;
+		}
+
 		return PinType;
 	}
 
@@ -512,6 +519,7 @@ protected:
 	{
 		bool bIsValidType = true;
 		ECameraContextDataType NewDataType = ECameraContextDataType::Name;
+		ECameraContextDataContainerType NewDataContainerType = ECameraContextDataContainerType::None;
 		const UObject* NewDataTypeObject = PinType.PinSubCategoryObject.Get();
 
 		if (PinType.PinCategory == UEdGraphSchema_K2::PC_Name)
@@ -543,7 +551,20 @@ protected:
 			bIsValidType = false;
 		}
 
-		if (ensure(bIsValidType) && (NewDataType != Item->DataType || NewDataTypeObject != Item->DataTypeObject))
+		switch (PinType.ContainerType)
+		{
+			case EPinContainerType::None:
+				NewDataContainerType = ECameraContextDataContainerType::None;
+				break;
+			case EPinContainerType::Array:
+				NewDataContainerType = ECameraContextDataContainerType::Array;
+				break;
+			default:
+				bIsValidType = false;
+				break;
+		}
+
+		if (ensure(bIsValidType) && (NewDataType != Item->DataType || NewDataContainerType != Item->DataContainerType || NewDataTypeObject != Item->DataTypeObject))
 		{
 			const FScopedTransaction Transaction(LOCTEXT("ChangeDataParameterType", "Change Data Parameter Type"));
 
@@ -556,13 +577,6 @@ protected:
 	{
 		return !IsHovered();
 	}
-
-private:
-
-	mutable ECameraContextDataType CachedDataType;
-	mutable const UObject* CachedDataTypeObject = nullptr;
-	mutable const FSlateBrush* CachedDataTypeImage = nullptr;
-	mutable FSlateColor CachedDataTypeImageColor;
 };
 
 /**
diff --git a/Engine/Plugins/Cameras/GameplayCameras/Source/GameplayCamerasEditor/Public/Toolkits/CameraRigInterfaceParametersToolkit.h b/Engine/Plugins/Cameras/GameplayCameras/Source/GameplayCamerasEditor/Public/Toolkits/CameraRigInterfaceParametersToolkit.h
--- a/Engine/Plugins/Cameras/GameplayCameras/Source/GameplayCamerasEditor/Public/Toolkits/CameraRigInterfaceParametersToolkit.h
+++ b/Engine/Plugins/Cameras/GameplayCameras/Source/GameplayCamerasEditor/Public/Toolkits/CameraRigInterfaceParametersToolkit.h
@@ -2,10 +2,13 @@
 
 #pragma once
 
+#include "EdGraphSchema_K2.h"
 #include "EditorUndoClient.h"
 #include "Templates/SharedPointer.h"
 #include "UObject/ObjectPtr.h"
 
+#include "CameraRigInterfaceParametersToolkit.generated.h"
+
 class SBox;
 class SWidget;
 class UCameraRigAsset;
@@ -59,3 +62,20 @@ private:
 
 }  // namespace UE::Cameras
 
+UCLASS(MinimalAPI, Hidden)
+class UEdGraphSchema_CameraNodeK2 : public UEdGraphSchema_K2
+{
+	GENERATED_BODY()
+
+public:
+
+	virtual bool SupportsPinTypeContainer(TWeakPtr<const FEdGraphSchemaAction> SchemaAction, const FEdGraphPinType& PinType, const EPinContainerType& ContainerType) const
+	{
+		if (ContainerType == EPinContainerType::None || ContainerType == EPinContainerType::Array)
+		{
+			return Super::SupportsPinTypeContainer(SchemaAction, PinType, ContainerType);
+		}
+		return false;
+	}
+};
+
diff --git a/Engine/Plugins/Cameras/GameplayCameras/Source/GameplayCamerasUncookedOnly/Private/BlueprintGraph/K2Node_CameraRigBase.cpp b/Engine/Plugins/Cameras/GameplayCameras/Source/GameplayCamerasUncookedOnly/Private/BlueprintGraph/K2Node_CameraRigBase.cpp
--- a/Engine/Plugins/Cameras/GameplayCameras/Source/GameplayCamerasUncookedOnly/Private/BlueprintGraph/K2Node_CameraRigBase.cpp
+++ b/Engine/Plugins/Cameras/GameplayCameras/Source/GameplayCamerasUncookedOnly/Private/BlueprintGraph/K2Node_CameraRigBase.cpp
@@ -193,10 +193,10 @@ FEdGraphPinType UK2Node_CameraRigBase::MakeBlendableParameterPinType(ECameraVari
 
 FEdGraphPinType UK2Node_CameraRigBase::MakeDataParameterPinType(const UCameraRigDataParameter* DataParameter)
 {
-	return MakeDataParameterPinType(DataParameter->DataType, DataParameter->DataTypeObject);
+	return MakeDataParameterPinType(DataParameter->DataType, DataParameter->DataContainerType, DataParameter->DataTypeObject);
 }
 
-FEdGraphPinType UK2Node_CameraRigBase::MakeDataParameterPinType(ECameraContextDataType CameraContextDataType, const UObject* CameraContextDataTypeObject)
+FEdGraphPinType UK2Node_CameraRigBase::MakeDataParameterPinType(ECameraContextDataType CameraContextDataType, ECameraContextDataContainerType CameraContextDataContainerType, const UObject* CameraContextDataTypeObject)
 {
 	FName PinCategory;
 	UObject* PinSubCategoryObject = const_cast<UObject*>(CameraContextDataTypeObject);
@@ -222,9 +222,19 @@ FEdGraphPinType UK2Node_CameraRigBase::MakeDataParameterPinType(ECameraContextDa
 			break;
 	}
 
+	EPinContainerType PinContainerType = EPinContainerType::None;
+	switch (CameraContextDataContainerType)
+	{
+		case ECameraContextDataContainerType::Array:
+			PinContainerType = EPinContainerType::Array;
+		default:
+			break;
+	}
+
 	FEdGraphPinType PinType;
 	PinType.PinCategory = PinCategory;
 	PinType.PinSubCategoryObject = PinSubCategoryObject;
+	PinType.ContainerType = PinContainerType;
 	return PinType;
 }
 
@@ -282,35 +292,62 @@ FName UK2Node_CameraRigBase::GetDataParameterInteropSettingFunctionName(const UC
 {
 	if (ensure(DataParameter && DataParameter->PrivateDataID.IsValid()))
 	{
-		return GetDataParameterInteropSettingFunctionName(DataParameter->DataType, DataParameter->DataTypeObject);
+		return GetDataParameterInteropSettingFunctionName(DataParameter->DataType, DataParameter->DataContainerType, DataParameter->DataTypeObject);
 	}
 	return NAME_None;
 }
 
-FName UK2Node_CameraRigBase::GetDataParameterInteropSettingFunctionName(ECameraContextDataType CameraContextDataType, const UObject* CameraContextDataTypeObject)
+FName UK2Node_CameraRigBase::GetDataParameterInteropSettingFunctionName(ECameraContextDataType CameraContextDataType, ECameraContextDataContainerType CameraContextDataContainerType, const UObject* CameraContextDataTypeObject)
 {
 	// Figure out the sort of SetXxxParameter function we want to call for this parameter.
 	FName CallSetParameterFuncName;
-	switch (CameraContextDataType)
+	if (CameraContextDataContainerType == ECameraContextDataContainerType::None)
 	{
-		case ECameraContextDataType::Name:
-			CallSetParameterFuncName = GET_FUNCTION_NAME_CHECKED(UCameraRigParameterInterop, SetNameParameter);
-			break;
-		case ECameraContextDataType::String:
-			CallSetParameterFuncName = GET_FUNCTION_NAME_CHECKED(UCameraRigParameterInterop, SetStringParameter);
-			break;
-		case ECameraContextDataType::Enum:
-			CallSetParameterFuncName = GET_FUNCTION_NAME_CHECKED(UCameraRigParameterInterop, SetEnumParameter);
-			break;
-		case ECameraContextDataType::Struct:
-			CallSetParameterFuncName = GET_FUNCTION_NAME_CHECKED(UCameraRigParameterInterop, SetStructParameter);
-			break;
-		case ECameraContextDataType::Object:
-			CallSetParameterFuncName = GET_FUNCTION_NAME_CHECKED(UCameraRigParameterInterop, SetObjectParameter);
-			break;
-		case ECameraContextDataType::Class:
-			CallSetParameterFuncName = GET_FUNCTION_NAME_CHECKED(UCameraRigParameterInterop, SetClassParameter);
-			break;
+		switch (CameraContextDataType)
+		{
+			case ECameraContextDataType::Name:
+				CallSetParameterFuncName = GET_FUNCTION_NAME_CHECKED(UCameraRigParameterInterop, SetNameParameter);
+				break;
+			case ECameraContextDataType::String:
+				CallSetParameterFuncName = GET_FUNCTION_NAME_CHECKED(UCameraRigParameterInterop, SetStringParameter);
+				break;
+			case ECameraContextDataType::Enum:
+				CallSetParameterFuncName = GET_FUNCTION_NAME_CHECKED(UCameraRigParameterInterop, SetEnumParameter);
+				break;
+			case ECameraContextDataType::Struct:
+				CallSetParameterFuncName = GET_FUNCTION_NAME_CHECKED(UCameraRigParameterInterop, SetStructParameter);
+				break;
+			case ECameraContextDataType::Object:
+				CallSetParameterFuncName = GET_FUNCTION_NAME_CHECKED(UCameraRigParameterInterop, SetObjectParameter);
+				break;
+			case ECameraContextDataType::Class:
+				CallSetParameterFuncName = GET_FUNCTION_NAME_CHECKED(UCameraRigParameterInterop, SetClassParameter);
+				break;
+		}
+	}
+	else
+	{
+		switch (CameraContextDataType)
+		{
+			case ECameraContextDataType::Name:
+				CallSetParameterFuncName = GET_FUNCTION_NAME_CHECKED(UCameraRigParameterInterop, SetNameArrayParameter);
+				break;
+			case ECameraContextDataType::String:
+				CallSetParameterFuncName = GET_FUNCTION_NAME_CHECKED(UCameraRigParameterInterop, SetStringArrayParameter);
+				break;
+			case ECameraContextDataType::Enum:
+				CallSetParameterFuncName = GET_FUNCTION_NAME_CHECKED(UCameraRigParameterInterop, SetEnumArrayParameter);
+				break;
+			case ECameraContextDataType::Struct:
+				CallSetParameterFuncName = GET_FUNCTION_NAME_CHECKED(UCameraRigParameterInterop, SetStructArrayParameter);
+				break;
+			case ECameraContextDataType::Object:
+				CallSetParameterFuncName = GET_FUNCTION_NAME_CHECKED(UCameraRigParameterInterop, SetObjectArrayParameter);
+				break;
+			case ECameraContextDataType::Class:
+				CallSetParameterFuncName = GET_FUNCTION_NAME_CHECKED(UCameraRigParameterInterop, SetClassArrayParameter);
+				break;
+		}
 	}
 
 	return CallSetParameterFuncName;
diff --git a/Engine/Plugins/Cameras/GameplayCameras/Source/GameplayCamerasUncookedOnly/Private/BlueprintGraph/K2Node_CameraRigBase.h b/Engine/Plugins/Cameras/GameplayCameras/Source/GameplayCamerasUncookedOnly/Private/BlueprintGraph/K2Node_CameraRigBase.h
--- a/Engine/Plugins/Cameras/GameplayCameras/Source/GameplayCamerasUncookedOnly/Private/BlueprintGraph/K2Node_CameraRigBase.h
+++ b/Engine/Plugins/Cameras/GameplayCameras/Source/GameplayCamerasUncookedOnly/Private/BlueprintGraph/K2Node_CameraRigBase.h
@@ -47,12 +47,12 @@ public:
 	static FEdGraphPinType MakeBlendableParameterPinType(const UCameraRigBlendableParameter* BlendableParameter);
 	static FEdGraphPinType MakeBlendableParameterPinType(ECameraVariableType CameraVariableType, const UScriptStruct* BlendableStructType);
 	static FEdGraphPinType MakeDataParameterPinType(const UCameraRigDataParameter* DataParameter);
-	static FEdGraphPinType MakeDataParameterPinType(ECameraContextDataType CameraContextDataType, const UObject* CameraContextDataTypeObject);
+	static FEdGraphPinType MakeDataParameterPinType(ECameraContextDataType CameraContextDataType, ECameraContextDataContainerType CameraContextDataContainerType, const UObject* CameraContextDataTypeObject);
 
 	static FName GetBlendableParameterInteropSettingFunctionName(const UCameraRigBlendableParameter* BlendableParameter);
 	static FName GetBlendableParameterInteropSettingFunctionName(ECameraVariableType CameraVariableType);
 	static FName GetDataParameterInteropSettingFunctionName(const UCameraRigDataParameter* DataParameter);
-	static FName GetDataParameterInteropSettingFunctionName(ECameraContextDataType CameraContextDataType, const UObject* CameraContextDataTypeObject);
+	static FName GetDataParameterInteropSettingFunctionName(ECameraContextDataType CameraContextDataType, ECameraContextDataContainerType CameraContextDataContainerType, const UObject* CameraContextDataTypeObject);
 
 protected:
 
diff --git a/Engine/Plugins/Cameras/GameplayCameras/Source/GameplayCamerasUncookedOnly/Private/BlueprintGraph/K2Node_SetCameraRigParameter.cpp b/Engine/Plugins/Cameras/GameplayCameras/Source/GameplayCamerasUncookedOnly/Private/BlueprintGraph/K2Node_SetCameraRigParameter.cpp
--- a/Engine/Plugins/Cameras/GameplayCameras/Source/GameplayCamerasUncookedOnly/Private/BlueprintGraph/K2Node_SetCameraRigParameter.cpp
+++ b/Engine/Plugins/Cameras/GameplayCameras/Source/GameplayCamerasUncookedOnly/Private/BlueprintGraph/K2Node_SetCameraRigParameter.cpp
@@ -34,7 +34,7 @@ void UK2Node_SetCameraRigParameter::Initialize(const FAssetData& UnloadedCameraR
 		}
 		else if (const UCameraRigDataParameter* DataParameter = LoadedCameraRig->Interface.FindDataParameterByName(InCameraParameterName))
 		{
-			Initialize(LoadedCameraRig, InCameraParameterName, DataParameter->DataType, DataParameter->DataTypeObject);
+			Initialize(LoadedCameraRig, InCameraParameterName, DataParameter->DataType, DataParameter->DataContainerType, DataParameter->DataTypeObject);
 		}
 		// else, no parameter of that name found...
 	}
@@ -49,12 +49,13 @@ void UK2Node_SetCameraRigParameter::Initialize(UCameraRigAsset* InCameraRig, con
 	BlendableStructType = InBlendableStructType;
 }
 
-void UK2Node_SetCameraRigParameter::Initialize(UCameraRigAsset* InCameraRig, const FString& InCameraParameterName, ECameraContextDataType InCameraContextDataType, const UObject* InCameraContextDataTypeObject)
+void UK2Node_SetCameraRigParameter::Initialize(UCameraRigAsset* InCameraRig, const FString& InCameraParameterName, ECameraContextDataType InCameraContextDataType, ECameraContextDataContainerType InCameraContextDataContainerType, const UObject* InCameraContextDataTypeObject)
 {
 	CameraRig = InCameraRig;
 	CameraParameterName = InCameraParameterName;
 	CameraParameterType = EK2Node_CameraParameterType::Data;
 	DataCameraParameterType = InCameraContextDataType;
+	DataCameraParameterContainerType = InCameraContextDataContainerType;
 	DataCameraParameterTypeObject = const_cast<UObject*>(InCameraContextDataTypeObject);
 }
 
@@ -70,7 +71,7 @@ void UK2Node_SetCameraRigParameter::AllocateDefaultPins()
 			PinType = UK2Node_CameraRigBase::MakeBlendableParameterPinType(BlendableCameraParameterType, BlendableStructType);
 			break;
 		case EK2Node_CameraParameterType::Data:
-			PinType = UK2Node_CameraRigBase::MakeDataParameterPinType(DataCameraParameterType, DataCameraParameterTypeObject);
+			PinType = UK2Node_CameraRigBase::MakeDataParameterPinType(DataCameraParameterType, DataCameraParameterContainerType, DataCameraParameterTypeObject);
 			break;
 		default:
 			PinType.PinCategory = UEdGraphSchema_K2::PC_Wildcard;
@@ -160,7 +161,7 @@ void UK2Node_SetCameraRigParameter::ExpandNode(FKismetCompilerContext& CompilerC
 			break;
 		case EK2Node_CameraParameterType::Data:
 			CallSetParameterFuncName = UK2Node_CameraRigBase::GetDataParameterInteropSettingFunctionName(
-					DataCameraParameterType, DataCameraParameterTypeObject);
+					DataCameraParameterType, DataCameraParameterContainerType, DataCameraParameterTypeObject);
 			break;
 		default:
 			break;
diff --git a/Engine/Plugins/Cameras/GameplayCameras/Source/GameplayCamerasUncookedOnly/Private/BlueprintGraph/K2Node_SetCameraRigParameter.h b/Engine/Plugins/Cameras/GameplayCameras/Source/GameplayCamerasUncookedOnly/Private/BlueprintGraph/K2Node_SetCameraRigParameter.h
--- a/Engine/Plugins/Cameras/GameplayCameras/Source/GameplayCamerasUncookedOnly/Private/BlueprintGraph/K2Node_SetCameraRigParameter.h
+++ b/Engine/Plugins/Cameras/GameplayCameras/Source/GameplayCamerasUncookedOnly/Private/BlueprintGraph/K2Node_SetCameraRigParameter.h
@@ -32,7 +32,7 @@ public:
 
 	void Initialize(const FAssetData& UnloadedCameraRig, const FString& InCameraParameterName);
 	void Initialize(UCameraRigAsset* InCameraRig, const FString& InCameraParameterName, ECameraVariableType InCameraVariableType, const UScriptStruct* InBlendableStructType);
-	void Initialize(UCameraRigAsset* InCameraRig, const FString& InCameraParameterName, ECameraContextDataType InCameraContextDataType, const UObject* InCameraContextDataTypeObject);
+	void Initialize(UCameraRigAsset* InCameraRig, const FString& InCameraParameterName, ECameraContextDataType InCameraContextDataType, ECameraContextDataContainerType InCameraContextDataContainerType, const UObject* InCameraContextDataTypeObject);
 
 public:
 
@@ -65,6 +65,9 @@ private:
 	UPROPERTY()
 	ECameraContextDataType DataCameraParameterType;
 
+	UPROPERTY()
+	ECameraContextDataContainerType DataCameraParameterContainerType;
+
 	UPROPERTY()
 	TObjectPtr<UObject> DataCameraParameterTypeObject;
 };
