Expose some functions in GameplayCameras for external camera system

diff --git a/Engine/Plugins/Cameras/GameplayCameras/Source/GameplayCameras/Public/Core/BlendCameraNode.h b/Engine/Plugins/Cameras/GameplayCameras/Source/GameplayCameras/Public/Core/BlendCameraNode.h
--- a/Engine/Plugins/Cameras/GameplayCameras/Source/GameplayCameras/Public/Core/BlendCameraNode.h
+++ b/Engine/Plugins/Cameras/GameplayCameras/Source/GameplayCameras/Public/Core/BlendCameraNode.h
@@ -109,10 +109,10 @@ class FBlendCameraNodeEvaluator : public FCameraNodeEvaluator
 public:
 
 	/** Blend the parameters produced by a camera node tree over another set of values. */
-	void BlendParameters(const FCameraNodePreBlendParams& Params, FCameraNodePreBlendResult& OutResult);
+	GAMEPLAYCAMERAS_API void BlendParameters(const FCameraNodePreBlendParams& Params, FCameraNodePreBlendResult& OutResult);
 
 	/** Blend the result of a camera node tree over another result. */
-	void BlendResults(const FCameraNodeBlendParams& Params, FCameraNodeBlendResult& OutResult);
+	GAMEPLAYCAMERAS_API void BlendResults(const FCameraNodeBlendParams& Params, FCameraNodeBlendResult& OutResult);
 
 protected:
 
diff --git a/Engine/Plugins/Cameras/GameplayCameras/Source/GameplayCameras/Public/Core/CameraNodeEvaluator.h b/Engine/Plugins/Cameras/GameplayCameras/Source/GameplayCameras/Public/Core/CameraNodeEvaluator.h
--- a/Engine/Plugins/Cameras/GameplayCameras/Source/GameplayCameras/Public/Core/CameraNodeEvaluator.h
+++ b/Engine/Plugins/Cameras/GameplayCameras/Source/GameplayCameras/Public/Core/CameraNodeEvaluator.h
@@ -58,7 +58,7 @@ struct FCameraNodeEvaluatorBuildParams
 	{}
 
 	/** Builds an evaluator for the given camera node. */
-	FCameraNodeEvaluator* BuildEvaluator(const UCameraNode* InNode) const;
+	GAMEPLAYCAMERAS_API FCameraNodeEvaluator* BuildEvaluator(const UCameraNode* InNode) const;
 
 	/** Builds an evaluator for the given camera node, and down-cast it to the given type. */
 	template<typename EvaluatorType>
@@ -282,7 +282,7 @@ public:
 	void UpdateParameters(const FCameraBlendedParameterUpdateParams& Params, FCameraBlendedParameterUpdateResult& OutResult);
 
 	/** Run this evaluator. */
-	void Run(const FCameraNodeEvaluationParams& Params, FCameraNodeEvaluationResult& OutResult);
+	GAMEPLAYCAMERAS_API void Run(const FCameraNodeEvaluationParams& Params, FCameraNodeEvaluationResult& OutResult);
 
 	/** Execute an IK operation. */
 	void ExecuteOperation(const FCameraOperationParams& Params, FCameraOperation& Operation);
diff --git a/Engine/Plugins/Cameras/GameplayCameras/Source/GameplayCameras/Public/Nodes/Blends/OrbitBlendCameraNode.h b/Engine/Plugins/Cameras/GameplayCameras/Source/GameplayCameras/Public/Nodes/Blends/OrbitBlendCameraNode.h
--- a/Engine/Plugins/Cameras/GameplayCameras/Source/GameplayCameras/Public/Nodes/Blends/OrbitBlendCameraNode.h
+++ b/Engine/Plugins/Cameras/GameplayCameras/Source/GameplayCameras/Public/Nodes/Blends/OrbitBlendCameraNode.h
@@ -19,6 +19,8 @@ public:
 
 	UOrbitBlendCameraNode(const FObjectInitializer& ObjInit);
 
+	GAMEPLAYCAMERAS_API void SetDrivingBlend(USimpleBlendCameraNode* DrivingBlendIn) { DrivingBlend = DrivingBlendIn; }
+
 public:
 
 	UPROPERTY()
diff --git a/Engine/Plugins/Cameras/GameplayCameras/Source/GameplayCameras/Public/Nodes/Blends/SimpleBlendCameraNode.h b/Engine/Plugins/Cameras/GameplayCameras/Source/GameplayCameras/Public/Nodes/Blends/SimpleBlendCameraNode.h
--- a/Engine/Plugins/Cameras/GameplayCameras/Source/GameplayCameras/Public/Nodes/Blends/SimpleBlendCameraNode.h
+++ b/Engine/Plugins/Cameras/GameplayCameras/Source/GameplayCameras/Public/Nodes/Blends/SimpleBlendCameraNode.h
@@ -23,6 +23,10 @@ class USimpleFixedTimeBlendCameraNode : public USimpleBlendCameraNode
 {
 	GENERATED_BODY()
 
+public:
+
+	GAMEPLAYCAMERAS_API void SetBlendTime(float BlendTimeIn) { BlendTime = BlendTimeIn; }
+
 public:
 
 	/** Duration of the blend. */
diff --git a/Engine/Plugins/Cameras/GameplayCameras/Source/GameplayCameras/Public/Nodes/Blends/SmoothBlendCameraNode.h b/Engine/Plugins/Cameras/GameplayCameras/Source/GameplayCameras/Public/Nodes/Blends/SmoothBlendCameraNode.h
--- a/Engine/Plugins/Cameras/GameplayCameras/Source/GameplayCameras/Public/Nodes/Blends/SmoothBlendCameraNode.h
+++ b/Engine/Plugins/Cameras/GameplayCameras/Source/GameplayCameras/Public/Nodes/Blends/SmoothBlendCameraNode.h
@@ -24,6 +24,10 @@ class USmoothBlendCameraNode : public USimpleFixedTimeBlendCameraNode
 {
 	GENERATED_BODY()
 
+public:
+
+	GAMEPLAYCAMERAS_API void SetCameraBlendType(ESmoothCameraBlendType BlendTypeIn) { BlendType = BlendTypeIn; }
+
 protected:
 
 	virtual FCameraNodeEvaluatorPtr OnBuildEvaluator(FCameraNodeEvaluatorBuilder& Builder) const override;
