Cameras: attach to actor group node

diff --git a/Engine/Plugins/Cameras/GameplayCameras/Source/GameplayCameras/Private/Nodes/Attach/AttachToActorGroupCameraNode.cpp b/Engine/Plugins/Cameras/GameplayCameras/Source/GameplayCameras/Private/Nodes/Attach/AttachToActorGroupCameraNode.cpp
new file mode 100644
--- /dev/null
+++ b/Engine/Plugins/Cameras/GameplayCameras/Source/GameplayCameras/Private/Nodes/Attach/AttachToActorGroupCameraNode.cpp
@@ -0,0 +1,85 @@
+// Copyright Epic Games, Inc. All Rights Reserved.
+
+#include "Nodes/Attach/AttachToActorGroupCameraNode.h"
+
+#include UE_INLINE_GENERATED_CPP_BY_NAME(AttachToActorGroupCameraNode)
+
+namespace UE::Cameras
+{
+
+class FAttachToActorGroupCameraNodeEvaluator : public FCameraNodeEvaluator
+{
+	UE_DECLARE_CAMERA_NODE_EVALUATOR(GAMEPLAYCAMERAS_API, FAttachToActorGroupCameraNodeEvaluator)
+
+protected:
+
+	virtual void OnInitialize(const FCameraNodeEvaluatorInitializeParams& Params, FCameraNodeEvaluationResult& OutResult) override;
+	virtual void OnRun(const FCameraNodeEvaluationParams& Params, FCameraNodeEvaluationResult& OutResult) override;
+
+private:
+
+#if WITH_EDITOR
+	void UpdateAttachmentReaders();
+#endif  // WITH_EDITOR
+
+private:
+
+	FCameraActorAttachmentInfoArrayReader AttachmentsReader;
+};
+
+UE_DEFINE_CAMERA_NODE_EVALUATOR(FAttachToActorGroupCameraNodeEvaluator)
+
+void FAttachToActorGroupCameraNodeEvaluator::OnInitialize(const FCameraNodeEvaluatorInitializeParams& Params, FCameraNodeEvaluationResult& OutResult)
+{
+	const UAttachToActorGroupCameraNode* AttachNode = GetCameraNodeAs<UAttachToActorGroupCameraNode>();
+	AttachmentsReader.Initialize(AttachNode->Attachments, AttachNode->AttachmentsDataID);
+}
+
+void FAttachToActorGroupCameraNodeEvaluator::OnRun(const FCameraNodeEvaluationParams& Params, FCameraNodeEvaluationResult& OutResult)
+{
+#if WITH_EDITOR
+	UpdateAttachmentReaders();
+#endif  // WITH_EDITOR
+
+	FTransform3d AttachTransform;
+	if (AttachmentsReader.GetAttachmentTransform(OutResult.ContextDataTable, AttachTransform))
+	{
+		const FVector3d AttachLocation = AttachTransform.GetLocation();
+		OutResult.CameraPose.SetLocation(AttachLocation);
+	}
+}
+
+#if WITH_EDITOR
+
+void FAttachToActorGroupCameraNodeEvaluator::UpdateAttachmentReaders()
+{
+	//const UAttachToActorGroupCameraNode* AttachNode = GetCameraNodeAs<UAttachToActorGroupCameraNode>();
+	//if (AttachNode->Attachments.Num() != AttachmentReaders.Num())
+	//{
+	//	AttachmentReaders.Reset();
+	//	for (const FCameraActorAttachmentInfo& Attachment : AttachNode->Attachments)
+	//	{
+	//		AttachmentReaders.Emplace(Attachment, FCameraContextDataID());
+	//	}
+	//}
+	//else
+	//{
+	//	for (int32 Index = 0; Index < AttachmentReaders.Num(); ++Index)
+	//	{
+	//		FCameraActorAttachmentInfoReader& Reader(AttachmentReaders[Index]);
+	//		const FCameraActorAttachmentInfo& Attachment(AttachNode->Attachments[Index]);
+	//		Reader.Initialize(Attachment, FCameraContextDataID());
+	//	}
+	//}
+}
+
+#endif  // WITH_EDITOR
+
+}  // namespace UE::Cameras
+
+FCameraNodeEvaluatorPtr UAttachToActorGroupCameraNode::OnBuildEvaluator(FCameraNodeEvaluatorBuilder& Builder) const
+{
+	using namespace UE::Cameras;
+	return Builder.BuildEvaluator<FAttachToActorGroupCameraNodeEvaluator>();
+}
+
diff --git a/Engine/Plugins/Cameras/GameplayCameras/Source/GameplayCameras/Private/Nodes/Attach/CameraActorAttachmentInfo.cpp b/Engine/Plugins/Cameras/GameplayCameras/Source/GameplayCameras/Private/Nodes/Attach/CameraActorAttachmentInfo.cpp
--- a/Engine/Plugins/Cameras/GameplayCameras/Source/GameplayCameras/Private/Nodes/Attach/CameraActorAttachmentInfo.cpp
+++ b/Engine/Plugins/Cameras/GameplayCameras/Source/GameplayCameras/Private/Nodes/Attach/CameraActorAttachmentInfo.cpp
@@ -30,10 +30,20 @@ void FCameraActorAttachmentInfoReader::CacheAttachmentInfo(const FCameraActorAtt
 		CachedAttachmentInfo = InAttachmentInfo;
 
 		CachedSkeletalMeshComponent = nullptr;
-		if (InAttachmentInfo.Actor && !InAttachmentInfo.SocketName.IsNone())
+		if (InAttachmentInfo.Actor && (!InAttachmentInfo.SocketName.IsNone() || !InAttachmentInfo.BoneName.IsNone()))
 		{
 			CachedSkeletalMeshComponent = InAttachmentInfo.Actor->FindComponentByClass<USkeletalMeshComponent>();
 		}
+
+		CachedBoneName = NAME_None;
+		if (CachedSkeletalMeshComponent)
+		{
+			CachedBoneName = InAttachmentInfo.BoneName;
+			if (!InAttachmentInfo.SocketName.IsNone())
+			{
+				CachedBoneName = CachedSkeletalMeshComponent->GetSocketBoneName(InAttachmentInfo.SocketName);
+			}
+		}
 	}
 }
 
@@ -51,7 +61,7 @@ bool FCameraActorAttachmentInfoReader::GetAttachmentTransform(const FCameraConte
 		}
 	}
 
-	if (!CachedAttachmentInfo.SocketName.IsNone() && CachedSkeletalMeshComponent)
+	if (CachedSkeletalMeshComponent && !CachedBoneName.IsNone())
 	{
 		OutTransform = CachedSkeletalMeshComponent->GetSocketTransform(CachedAttachmentInfo.SocketName);
 		return true;
@@ -64,39 +74,77 @@ bool FCameraActorAttachmentInfoReader::GetAttachmentTransform(const FCameraConte
 	return false;
 }
 
-bool FCameraActorAttachmentInfoReader::GetAttachmentTransform(TArrayView<FCameraActorAttachmentInfoReader> Readers, const FCameraContextDataTable& ContextDataTable,FTransform3d& OutTransform)
+FCameraActorAttachmentInfoArrayReader::FCameraActorAttachmentInfoArrayReader(TConstArrayView<FCameraActorAttachmentInfo> InAttachmentInfos, FCameraContextDataID InDataID)
+{
+	Initialize(InAttachmentInfos, InDataID);
+}
+
+void FCameraActorAttachmentInfoArrayReader::Initialize(TConstArrayView<FCameraActorAttachmentInfo> InAttachmentInfos, FCameraContextDataID InDataID)
 {
-	using FWeightedTransform = TTuple<FTransform3d, float>;
-	
+	DataID = InDataID;
+
+	CacheAttachmentInfos(InAttachmentInfos);
+}
+
+void FCameraActorAttachmentInfoArrayReader::CacheAttachmentInfos(TConstArrayView<FCameraActorAttachmentInfo> InAttachmentInfos)
+{
+	Readers.SetNum(InAttachmentInfos.Num());
+
+	for (int32 Index = 0; Index < InAttachmentInfos.Num(); ++Index)
+	{
+		Readers[Index].CacheAttachmentInfo(InAttachmentInfos[Index]);
+	}
+}
+
+bool FCameraActorAttachmentInfoArrayReader::GetAttachmentTransform(const FCameraContextDataTable& ContextDataTable, FTransform3d& OutTransform)
+{
+	if (DataID.IsValid())
+	{
+		TConstArrayView<FCameraActorAttachmentInfo> NewAttachmentInfos;
+		if (ContextDataTable.TryGetArrayData<FCameraActorAttachmentInfo>(DataID, NewAttachmentInfos))
+		{
+			CacheAttachmentInfos(NewAttachmentInfos);
+		}
+	}
+
 	if (Readers.IsEmpty())
 	{
 		return false;
 	}
 
-	TArray<FWeightedTransform> WeightedTransforms;
+	struct FComputedAttachmentInfo
+	{
+		FTransform3d Transform;
+		float Weight = 1.f;
+	};
+
+	TArray<FComputedAttachmentInfo> ComputedAttachments;
+	ComputedAttachments.SetNum(Readers.Num());
 
-	for (FCameraActorAttachmentInfoReader& Reader : Readers)
+	for (int32 Index = 0; Index < Readers.Num(); ++Index)
 	{
-		FTransform3d CurTransform;
-		if (Reader.GetAttachmentTransform(ContextDataTable, CurTransform))
-		{
-			WeightedTransforms.Emplace(CurTransform, Reader.CachedAttachmentInfo.Weight);
-		}
+		FCameraActorAttachmentInfoReader& Reader(Readers[Index]);
+		FComputedAttachmentInfo& ComputedAttachment(ComputedAttachments[Index]);
+		const bool bGotAttachment = Reader.GetAttachmentTransform(ContextDataTable, ComputedAttachment.Transform);
+		ComputedAttachment.Weight = bGotAttachment ? Reader.CachedAttachmentInfo.Weight : 0.f;
 	}
 
 	float TotalWeight = Algo::Accumulate(
-			WeightedTransforms, 
+			ComputedAttachments, 
 			0.f,
-			[](float Cur, const FWeightedTransform& Item) { return Cur + Item.Get<1>(); });
+			[](float Cur, const FComputedAttachmentInfo& Item) { return Cur + Item.Weight; });
 	if (TotalWeight == 0.f)
 	{
 		return false;
 	}
 
 	OutTransform = FTransform3d::Identity;
-	for (const FWeightedTransform& WeightedTransform : WeightedTransforms)
+	for (const FComputedAttachmentInfo& ComputedAttachment : ComputedAttachments)
 	{
-		OutTransform.BlendWith(WeightedTransform.Get<0>(), WeightedTransform.Get<1>() / TotalWeight);
+		if (ComputedAttachment.Weight > 0.f)
+		{
+			OutTransform.BlendWith(ComputedAttachment.Transform, ComputedAttachment.Weight / TotalWeight);
+		}
 	}
 
 	return true;
diff --git a/Engine/Plugins/Cameras/GameplayCameras/Source/GameplayCameras/Public/Nodes/Attach/AttachToActorGroupCameraNode.h b/Engine/Plugins/Cameras/GameplayCameras/Source/GameplayCameras/Public/Nodes/Attach/AttachToActorGroupCameraNode.h
new file mode 100644
--- /dev/null
+++ b/Engine/Plugins/Cameras/GameplayCameras/Source/GameplayCameras/Public/Nodes/Attach/AttachToActorGroupCameraNode.h
@@ -0,0 +1,33 @@
+// Copyright Epic Games, Inc. All Rights Reserved.
+
+#pragma once
+
+#include "Core/CameraContextDataTableFwd.h"
+#include "Core/CameraNode.h"
+#include "Nodes/Attach/CameraActorAttachmentInfo.h"
+
+#include "AttachToActorGroupCameraNode.generated.h"
+
+/**
+ * A camera node that moves the camera to the weighted center of a group of actors.
+ */
+UCLASS(MinimalAPI, meta=(CameraNodeCategories="Attachment"))
+class UAttachToActorGroupCameraNode : public UCameraNode
+{
+	GENERATED_BODY()
+
+protected:
+
+	// UCameraNode interface.
+	virtual FCameraNodeEvaluatorPtr OnBuildEvaluator(FCameraNodeEvaluatorBuilder& Builder) const override;
+
+public:
+
+	/** The actors to attach to. */
+	UPROPERTY(EditAnywhere, Category=Common, meta=(CameraContextData=true))
+	TArray<FCameraActorAttachmentInfo> Attachments;
+
+	UPROPERTY()
+	FCameraContextDataID AttachmentsDataID;
+};
+
diff --git a/Engine/Plugins/Cameras/GameplayCameras/Source/GameplayCameras/Public/Nodes/Attach/CameraActorAttachmentInfo.h b/Engine/Plugins/Cameras/GameplayCameras/Source/GameplayCameras/Public/Nodes/Attach/CameraActorAttachmentInfo.h
--- a/Engine/Plugins/Cameras/GameplayCameras/Source/GameplayCameras/Public/Nodes/Attach/CameraActorAttachmentInfo.h
+++ b/Engine/Plugins/Cameras/GameplayCameras/Source/GameplayCameras/Public/Nodes/Attach/CameraActorAttachmentInfo.h
@@ -25,6 +25,10 @@ struct GAMEPLAYCAMERAS_API FCameraActorAttachmentInfo
 	UPROPERTY(EditAnywhere, Interp, Category="Attachment")
 	FName SocketName;
 
+	/** An optional bone to attach to on the actor. */
+	UPROPERTY(EditAnywhere, Interp, Category="Target")
+	FName BoneName;
+
 	/** The weight of this attachment. Unused if only one attachment is used. */
 	UPROPERTY(EditAnywhere, Interp, Category="Attachment")
 	float Weight = 1.f;
@@ -45,8 +49,6 @@ struct FCameraActorAttachmentInfoReader
 
 	bool GetAttachmentTransform(const FCameraContextDataTable& ContextDataTable, FTransform3d& OutTransform);
 
-	static bool GetAttachmentTransform(TArrayView<FCameraActorAttachmentInfoReader> Readers, const FCameraContextDataTable& ContextDataTable, FTransform3d& OutTransform);
-
 private:
 
 	void CacheAttachmentInfo(const FCameraActorAttachmentInfo& InAttachmentInfo);
@@ -58,6 +60,29 @@ private:
 
 	FCameraActorAttachmentInfo CachedAttachmentInfo;
 	const USkeletalMeshComponent* CachedSkeletalMeshComponent = nullptr;
+	FName CachedBoneName;
+
+	friend struct FCameraActorAttachmentInfoArrayReader;
+};
+
+/** A special reader class for multiple attachment information. */
+struct FCameraActorAttachmentInfoArrayReader
+{
+	FCameraActorAttachmentInfoArrayReader() {}
+	FCameraActorAttachmentInfoArrayReader(TConstArrayView<FCameraActorAttachmentInfo> InAttachmentInfos, FCameraContextDataID InDataID);
+
+	void Initialize(TConstArrayView<FCameraActorAttachmentInfo> InAttachmentInfos, FCameraContextDataID InDataID);
+
+	bool GetAttachmentTransform(const FCameraContextDataTable& ContextDataTable, FTransform3d& OutTransform);
+
+private:
+
+	void CacheAttachmentInfos(TConstArrayView<FCameraActorAttachmentInfo> InAttachmentInfos);
+
+private:
+
+	TArray<FCameraActorAttachmentInfoReader> Readers;
+	FCameraContextDataID DataID;
 };
 
 }  // namespace UE::Cameras
