Cameras: multi-target support for screen-space framing nodes

diff --git a/Engine/Plugins/Cameras/GameplayCameras/Source/GameplayCameras/Private/Math/CameraFramingZoneMath.cpp b/Engine/Plugins/Cameras/GameplayCameras/Source/GameplayCameras/Private/Math/CameraFramingZoneMath.cpp
--- a/Engine/Plugins/Cameras/GameplayCameras/Source/GameplayCameras/Private/Math/CameraFramingZoneMath.cpp
+++ b/Engine/Plugins/Cameras/GameplayCameras/Source/GameplayCameras/Private/Math/CameraFramingZoneMath.cpp
@@ -108,6 +108,14 @@ void FFramingZone::ClampBounds(const FFramingZone& MustContain)
 	BottomBound = FMath::Clamp(BottomBound, MustContain.BottomBound, 1.0);
 }
 
+void FFramingZone::Add(const FFramingZone& Other)
+{
+	LeftBound = FMath::Min(LeftBound, Other.LeftBound);
+	TopBound = FMath::Min(TopBound, Other.TopBound);
+	RightBound = FMath::Max(RightBound, Other.RightBound);
+	BottomBound = FMath::Max(BottomBound, Other.BottomBound);
+}
+
 bool FFramingZone::Contains(const FVector2d& Point) const
 {
 	return Point.X >= LeftBound && Point.X <= RightBound &&
diff --git a/Engine/Plugins/Cameras/GameplayCameras/Source/GameplayCameras/Private/Nodes/Framing/BaseFramingCameraNode.cpp b/Engine/Plugins/Cameras/GameplayCameras/Source/GameplayCameras/Private/Nodes/Framing/BaseFramingCameraNode.cpp
--- a/Engine/Plugins/Cameras/GameplayCameras/Source/GameplayCameras/Private/Nodes/Framing/BaseFramingCameraNode.cpp
+++ b/Engine/Plugins/Cameras/GameplayCameras/Source/GameplayCameras/Private/Nodes/Framing/BaseFramingCameraNode.cpp
@@ -17,6 +17,7 @@
 #include "Math/ColorList.h"
 #include "Misc/AssertionMacros.h"
 #include "Misc/CoreMiscDefines.h"
+#include "Nodes/Framing/CameraActorTargetInfo.h"
 
 #include UE_INLINE_GENERATED_CPP_BY_NAME(BaseFramingCameraNode)
 
@@ -27,23 +28,30 @@ float GFramingIdealReachedEpsilon = 0.001;
 static FAutoConsoleVariableRef CVarFramingUnlockRadiusEpsilon(
 	TEXT("GameplayCameras.Framing.IdealReachedEpsilon"),
 	GFramingIdealReachedEpsilon,
-	TEXT("(Default: 0.001. The epsilon to determine whether we have reached ideal screen framing."));
+	TEXT("(Default: 0.001) The epsilon to determine whether we have reached ideal screen framing."));
 
-float GFramingExpolationEpsilon = 0.001;
+float GFramingExtrapolationEpsilon = 0.001;
 static FAutoConsoleVariableRef CVarFramingExtrapolationEpsilon(
 	TEXT("GameplayCameras.Framing.ExtrapolationEpsilon"),
-	GFramingIdealReachedEpsilon,
-	TEXT("(Default: 0.001. The epsilon to determine whether target movement extrapolation should be included."));
+	GFramingExtrapolationEpsilon,
+	TEXT("(Default: 0.001) The epsilon to determine whether target movement extrapolation should be included."));
 
 int32 GFramingNumTargetMovementSamples = 10;
 static FAutoConsoleVariableRef CVarFramingNumTargetMovementSamples(
 	TEXT("GameplayCameras.Framing.NumTargetMovementSamples"),
 	GFramingNumTargetMovementSamples,
-	TEXT("(Default: 10. The number of samples to use to extrapolate target movement."));
+	TEXT("(Default: 10) The number of samples to use to extrapolate target movement."));
+
+bool GFramingShowEffectiveDeadZone = false;
+static FAutoConsoleVariableRef CVarFramingShowEffectiveDeadZone(
+	TEXT("GameplayCameras.Framing.ShowEffectiveDeadZone"),
+	GFramingShowEffectiveDeadZone,
+	TEXT("(Default: false) Show the effective dead zone"));
 
 UE_DECLARE_CAMERA_DEBUG_BLOCK_START(GAMEPLAYCAMERAS_API, FBaseFramingCameraDebugBlock)
 	UE_DECLARE_CAMERA_DEBUG_BLOCK_FIELD(FBaseFramingCameraNodeEvaluator::FState, State);
 	UE_DECLARE_CAMERA_DEBUG_BLOCK_FIELD(FBaseFramingCameraNodeEvaluator::FDesired, Desired);
+	UE_DECLARE_CAMERA_DEBUG_BLOCK_FIELD(FBaseFramingCameraNodeEvaluator::FWorldTargetInfos, WorldTargets);
 UE_DECLARE_CAMERA_DEBUG_BLOCK_END()
 
 UE_DEFINE_CAMERA_DEBUG_BLOCK_WITH_FIELDS(FBaseFramingCameraDebugBlock)
@@ -54,7 +62,7 @@ void FBaseFramingCameraNodeEvaluator::OnInitialize(const FCameraNodeEvaluatorIni
 {
 	const UBaseFramingCameraNode* BaseFramingNode = GetCameraNodeAs<UBaseFramingCameraNode>();
 
-	Readers.TargetInfo.Initialize(BaseFramingNode->TargetInfo, BaseFramingNode->TargetInfoDataID);
+	Readers.TargetInfos.Initialize(BaseFramingNode->TargetInfos, BaseFramingNode->TargetInfosDataID);
 
 	Readers.IdealFramingLocation.Initialize(BaseFramingNode->IdealFramingLocation);
 	Readers.SetTargetDistance.Initialize(BaseFramingNode->SetTargetDistance);
@@ -68,10 +76,32 @@ void FBaseFramingCameraNodeEvaluator::OnInitialize(const FCameraNodeEvaluatorIni
 	Readers.DeadZone.Initialize(BaseFramingNode->DeadZone);
 	Readers.SoftZone.Initialize(BaseFramingNode->SoftZone);
 
-	Anticipation.History.Reserve(GFramingNumTargetMovementSamples);
+	ScreenTargetHistory.History.Reserve(GFramingNumTargetMovementSamples);
+}
+
+void FBaseFramingCameraNodeEvaluator::UpdateFramingState(const FCameraNodeEvaluationParams& Params, const FCameraNodeEvaluationResult& OutResult, const FTransform3d& LastFraming)
+{
+	TArray<FCameraActorComputedTargetInfo> AcquiredTargetInfos;
+	const bool bAcquireSuccess = AcquireTargetInfo(Params, OutResult, AcquiredTargetInfos);
+	WorldTargets.TargetInfos = AcquiredTargetInfos;
+	if (bAcquireSuccess && AcquiredTargetInfos.Num() > 0)
+	{
+		ComputeCurrentState(Params, OutResult, LastFraming);
+		ComputeDesiredState(Params.DeltaTime);
+	}
 }
 
-bool FBaseFramingCameraNodeEvaluator::AcquireTargetInfo(const FCameraNodeEvaluationParams& Params, const FCameraNodeEvaluationResult& InResult, FAcquiredTargetInfo& OutInfo)
+void FBaseFramingCameraNodeEvaluator::EndFramingUpdate(const FCameraNodeEvaluationParams& Params, FCameraNodeEvaluationResult& OutResult)
+{
+	if (Readers.SetTargetDistance.Get(OutResult.VariableTable))
+	{
+		FCameraPose& OutCameraPose = OutResult.CameraPose;
+		const double TargetDistance = FVector3d::Distance(State.WorldTarget, OutCameraPose.GetLocation());
+		OutCameraPose.SetTargetDistance(TargetDistance);
+	}
+}
+
+bool FBaseFramingCameraNodeEvaluator::AcquireTargetInfo(const FCameraNodeEvaluationParams& Params, const FCameraNodeEvaluationResult& InResult, TArray<FCameraActorComputedTargetInfo>& OutInfos)
 {
 	const UBaseFramingCameraNode* FramingNode = GetCameraNodeAs<UBaseFramingCameraNode>();
 	if (FramingNode->TargetLocation.IsValid())
@@ -79,26 +109,25 @@ bool FBaseFramingCameraNodeEvaluator::AcquireTargetInfo(const FCameraNodeEvaluat
 		const FVector3d* TargetLocation = FramingNode->TargetLocation.GetValue(InResult.VariableTable);
 		if (TargetLocation)
 		{
-			OutInfo.WorldTransform = FTransform3d(*TargetLocation);
+			FCameraActorComputedTargetInfo OutInfo;
+			OutInfo.Transform = FTransform3d(*TargetLocation);
 			OutInfo.LocalBounds = FBoxSphereBounds3d(EForceInit::ForceInit);
+			OutInfos.Add(OutInfo);
 			return true;
 		}
 	}
-	else if (FramingNode->TargetInfo.IsValid() || FramingNode->TargetInfoDataID.IsValid())
+	else if (FramingNode->TargetInfos.Num() > 0 || FramingNode->TargetInfosDataID.IsValid())
 	{
-		FTransform3d TargetTransform;
-		FBoxSphereBounds3d TargetBounds;
-		if (Readers.TargetInfo.GetTargetInfo(InResult.ContextDataTable, TargetTransform, TargetBounds))
+		if (Readers.TargetInfos.ComputeTargetInfos(InResult.ContextDataTable, OutInfos))
 		{
-			OutInfo.WorldTransform = TargetTransform;
-			OutInfo.LocalBounds = TargetBounds;
 			return true;
 		}
 	}
 	else if (APlayerController* PlayerController = Params.EvaluationContext->GetPlayerController())
 	{
 		APawn* Pawn = PlayerController->GetPawn();
-		OutInfo.WorldTransform = FTransform3d(Pawn->GetActorLocation());
+		FCameraActorComputedTargetInfo OutInfo;
+		OutInfo.Transform = FTransform3d(Pawn->GetActorLocation());
 		OutInfo.LocalBounds = FBoxSphereBounds3d(EForceInit::ForceInit);
 		if (USceneComponent* RootComponent = Pawn->GetRootComponent())
 		{
@@ -110,23 +139,7 @@ bool FBaseFramingCameraNodeEvaluator::AcquireTargetInfo(const FCameraNodeEvaluat
 	return false;
 }
 
-void FBaseFramingCameraNodeEvaluator::UpdateFramingState(const FCameraNodeEvaluationParams& Params, const FCameraNodeEvaluationResult& OutResult, const FAcquiredTargetInfo& TargetInfo, const FTransform3d& LastFraming)
-{
-	ComputeCurrentState(Params, OutResult, TargetInfo, LastFraming);
-	ComputeDesiredState(Params.DeltaTime);
-}
-
-void FBaseFramingCameraNodeEvaluator::EndFramingUpdate(const FCameraNodeEvaluationParams& Params, FCameraNodeEvaluationResult& OutResult)
-{
-	if (Readers.SetTargetDistance.Get(OutResult.VariableTable))
-	{
-		FCameraPose& OutCameraPose = OutResult.CameraPose;
-		const double TargetDistance = FVector3d::Distance(State.WorldTarget, OutCameraPose.GetLocation());
-		OutCameraPose.SetTargetDistance(TargetDistance);
-	}
-}
-
-void FBaseFramingCameraNodeEvaluator::ComputeCurrentState(const FCameraNodeEvaluationParams& Params, const FCameraNodeEvaluationResult& OutResult, const FAcquiredTargetInfo& TargetInfo, const FTransform3d& LastFraming)
+void FBaseFramingCameraNodeEvaluator::ComputeCurrentState(const FCameraNodeEvaluationParams& Params, const FCameraNodeEvaluationResult& OutResult, const FTransform3d& LastFraming)
 {
 	// Get screen-space coordinates of the ideal framing point. These are in 0..1 UI space.
 	State.IdealTarget = Readers.IdealFramingLocation.Get(OutResult.VariableTable);
@@ -151,49 +164,25 @@ void FBaseFramingCameraNodeEvaluator::ComputeCurrentState(const FCameraNodeEvalu
 	State.DeadZone.ClampBounds(State.IdealTarget);
 	State.SoftZone.ClampBounds(State.DeadZone);
 
-	// See if we need to extrapolate where the target will be in "anticipation time" seconds.
-	const FVector3d EffectiveTargetLocation = ComputeAnticipatedScreenTarget(Params.DeltaTime, TargetInfo.WorldTransform.GetLocation());
-
 	// We are going to reframe things iteratively, so we'll use a temporary pose defined by last frame's
 	// shot transform.
 	FCameraPose TempPose(OutResult.CameraPose);
 	TempPose.SetTransform(LastFraming);
 
-	// Get the target in screen-space.
-	APlayerController* PlayerController = Params.EvaluationContext->GetPlayerController();
-	const double AspectRatio = FCameraPoseMath::GetEffectiveAspectRatio(TempPose, PlayerController);
-	const TOptional<FVector2d> ScreenTarget = FCameraPoseMath::ProjectWorldToScreen(TempPose, AspectRatio, EffectiveTargetLocation, true);
-	State.WorldTarget = EffectiveTargetLocation;
-	State.LocalTargetBounds = TargetInfo.LocalBounds;
-	State.ScreenTarget = ScreenTarget.Get(FVector2d(0.5, 0.5));
+	// Process our targets and figure out the weighted average we should be aiming at.
+	FVector2d NewScreenTarget;
+	ComputeFinalTargetInfo(Params, TempPose, State.WorldTarget, NewScreenTarget, State.ScreenTargetBounds);
 
-	// Get the target's bounds in screen-space.
-	FTransform3d EffectiveTargetTransform(TargetInfo.WorldTransform.GetRotation(), EffectiveTargetLocation);
-	State.ScreenTargetBounds = ComputeScreenTargetBounds(TempPose, AspectRatio, EffectiveTargetTransform, TargetInfo.LocalBounds);
+	// See if we need to extrapolate where the target will be in "anticipation time" seconds.
+	State.ScreenTarget = ComputeAnticipatedScreenTarget(Params.DeltaTime, State.ScreenTarget, NewScreenTarget);
 
 	// Compute the effective dead-zone, which is the subset of the dead-zone that encompasses as much
 	// of the target's bound as possible.
-	FFramingZone EffectiveDeadZone = State.DeadZone;
-	const FVector2d DeadZoneCenter = State.DeadZone.Center();
-	const double TargetBoundsHalfWidth = State.ScreenTargetBounds.Width() / 2.0;
-	const double TargetBoundsHalfHeight = State.ScreenTargetBounds.Height() / 2.0;
-	EffectiveDeadZone.LeftBound += TargetBoundsHalfWidth;
-	EffectiveDeadZone.TopBound += TargetBoundsHalfHeight;
-	EffectiveDeadZone.RightBound -= TargetBoundsHalfWidth;
-	EffectiveDeadZone.BottomBound -= TargetBoundsHalfHeight;
-	if (EffectiveDeadZone.LeftBound > EffectiveDeadZone.RightBound)
-	{
-		EffectiveDeadZone.LeftBound = EffectiveDeadZone.RightBound = DeadZoneCenter.X;
-	}
-	if (EffectiveDeadZone.TopBound > EffectiveDeadZone.BottomBound)
-	{
-		EffectiveDeadZone.TopBound = EffectiveDeadZone.BottomBound = DeadZoneCenter.Y;
-	}
-	State.EffectiveDeadZone = EffectiveDeadZone;
+	State.EffectiveDeadZone = ComputeEffectiveDeadZone();
 
 #if UE_GAMEPLAY_CAMERAS_DEBUG
-	State.DeadZoneEdgePoint = State.ScreenTarget;
-	State.HardZoneEdgePoint = State.ScreenTarget;
+	State.DebugDeadZoneEdgePoint = State.ScreenTarget;
+	State.DebugHardZoneEdgePoint = State.ScreenTarget;
 #endif  // UE_GAMEPLAY_CAMERAS_DEBUG
 
 	// Update the reframe damper's damping factor.
@@ -223,8 +212,8 @@ void FBaseFramingCameraNodeEvaluator::ComputeCurrentState(const FCameraNodeEvalu
 			State.ReframeDampingFactorAlpha = Alpha;
 
 #if UE_GAMEPLAY_CAMERAS_DEBUG
-			State.DeadZoneEdgePoint = DeadEdgePoint;
-			State.HardZoneEdgePoint = HardEdgePoint;
+			State.DebugDeadZoneEdgePoint = DeadEdgePoint;
+			State.DebugHardZoneEdgePoint = HardEdgePoint;
 #endif  // UE_GAMEPLAY_CAMERAS_DEBUG
 		}
 		else
@@ -236,7 +225,7 @@ void FBaseFramingCameraNodeEvaluator::ComputeCurrentState(const FCameraNodeEvalu
 	const bool bWasReframing = State.bIsReframingTarget;
 	const bool bIsInSoftZone = State.SoftZone.Contains(State.ScreenTarget);
 	const bool bIsInDeadZone = State.EffectiveDeadZone.Contains(State.ScreenTarget);
-	if (!ScreenTarget.IsSet() || !bIsInSoftZone)
+	if (!bIsInSoftZone)
 	{
 		// Target is out of view or outside the soft zone -- it's therefore in the hard zone and we will
 		// do everything we can to put it back in the soft zone ASAP.
@@ -327,48 +316,150 @@ void FBaseFramingCameraNodeEvaluator::ComputeCurrentState(const FCameraNodeEvalu
 	}
 }
 
-FVector3d FBaseFramingCameraNodeEvaluator::ComputeAnticipatedScreenTarget(float DeltaTime, const FVector3d& InTargetLocation)
+bool FBaseFramingCameraNodeEvaluator::ComputeFinalTargetInfo(const FCameraNodeEvaluationParams& Params, const FCameraPose& CameraPose, FVector3d& OutWorldTarget, FVector2d& OutScreenTarget, FFramingZone& OutScreenBounds)
+{
+	TConstArrayView<FCameraActorComputedTargetInfo> TargetInfos(WorldTargets.TargetInfos);
+
+	if (TargetInfos.IsEmpty())
+	{
+		return false;
+	}
+
+	// Start with projecting all the targets, and their bounds, on screen.
+	APlayerController* PlayerController = Params.EvaluationContext->GetPlayerController();
+	const double AspectRatio = FCameraPoseMath::GetEffectiveAspectRatio(CameraPose, PlayerController);
+
+	struct FComputedTargetScreenInfo
+	{
+		FVector3d WorldTarget;
+		FVector2d ScreenTarget;
+		FFramingZone ScreenBounds;
+		double WorldTargetDistance = 0.f;
+		float NormalizedWeight = 1.f;
+	};
+
+	TArray<FComputedTargetScreenInfo> TargetScreenInfos;
+	TargetScreenInfos.SetNum(TargetInfos.Num());
+#if UE_GAMEPLAY_CAMERAS_DEBUG
+	State.DebugAllScreenTargetBounds.SetNum(TargetInfos.Num());
+#endif
+
+	for (int32 Index = 0; Index < TargetInfos.Num(); ++Index)
+	{
+		const FCameraActorComputedTargetInfo& TargetInfo(TargetInfos[Index]);
+		FComputedTargetScreenInfo& TargetScreenInfo(TargetScreenInfos[Index]);
+
+		const FVector3d& WorldTarget = TargetInfo.Transform.GetLocation();
+		const TOptional<FVector2d> ScreenTarget = FCameraPoseMath::ProjectWorldToScreen(CameraPose, AspectRatio, WorldTarget, true);
+
+		TargetScreenInfo.WorldTarget = WorldTarget;
+		TargetScreenInfo.ScreenTarget = ScreenTarget.Get(FVector2d(0.5, 0.5));
+		TargetScreenInfo.ScreenBounds = ComputeScreenTargetBounds(CameraPose, AspectRatio, TargetInfo.Transform, TargetInfo.LocalBounds);
+		TargetScreenInfo.WorldTargetDistance = FVector3d::Distance(CameraPose.GetLocation(), WorldTarget);
+		TargetScreenInfo.NormalizedWeight = TargetInfo.NormalizedWeight;
+
+#if UE_GAMEPLAY_CAMERAS_DEBUG
+		State.DebugAllScreenTargetBounds[Index] = TargetScreenInfo.ScreenBounds;
+#endif
+	}
+
+	// If we only have one target, just use that.
+	if (TargetInfos.Num() == 1)
+	{
+		const FComputedTargetScreenInfo& TargetScreenInfo(TargetScreenInfos[0]);
+		OutScreenTarget = TargetScreenInfo.ScreenTarget;
+		OutScreenBounds = TargetScreenInfo.ScreenBounds;
+		OutWorldTarget = TargetScreenInfo.WorldTarget;
+		return true;
+	}
+
+	// Compute the weighted average of the screen target, and combine the screen bounds.
+	FVector2d FinalScreenTarget(EForceInit::ForceInit);
+	FFramingZone FinalScreenBounds(TargetScreenInfos[0].ScreenBounds);
+	double FinalWorldTargetDistance = 0.f;
+	for (int32 Index = 0; Index < TargetScreenInfos.Num(); ++Index)
+	{
+		const FComputedTargetScreenInfo& TargetScreenInfo(TargetScreenInfos[Index]);
+		FinalScreenTarget += TargetScreenInfo.ScreenTarget * TargetScreenInfo.NormalizedWeight;
+		FinalScreenBounds.Add(TargetScreenInfo.ScreenBounds);
+		FinalWorldTargetDistance += TargetScreenInfo.WorldTargetDistance * TargetScreenInfo.NormalizedWeight;
+	}
+
+	// Unproject the final screen target, and use the weigted average distance to get, roughly, what
+	// world-space target we might be looking at.
+	const FVector3d FinalWorldTarget = FCameraPoseMath::UnprojectScreenToWorld(CameraPose, AspectRatio, FinalScreenTarget, FinalWorldTargetDistance);
+
+	OutScreenTarget = FinalScreenTarget;
+	OutScreenBounds = FinalScreenBounds;
+	OutWorldTarget = FinalWorldTarget;
+
+	return true;
+}
+
+FVector2d FBaseFramingCameraNodeEvaluator::ComputeAnticipatedScreenTarget(float DeltaTime, const FVector2d& InPreviousAnticipatedScreenTarget, const FVector2d& InScreenTarget)
 {
 	if (State.TargetMovementAnticipationTime <= 0)
 	{
-		return InTargetLocation;
+		return InScreenTarget;
 	}
 	if (DeltaTime <= 0)
 	{
-		return InTargetLocation;
+		return InPreviousAnticipatedScreenTarget;
 	}
 
-	using FHistoryEntry = TTuple<FVector3d, float>;
+	using FHistoryEntry = TTuple<FVector2d, float>;
 
-	const int32 MaxNumSamples = Anticipation.History.Max();
-	if (Anticipation.History.Num() == MaxNumSamples)
+	const int32 MaxNumSamples = ScreenTargetHistory.History.Max();
+	if (ScreenTargetHistory.History.Num() == MaxNumSamples)
 	{
-		Anticipation.History.RemoveAt(0, EAllowShrinking::No);
+		ScreenTargetHistory.History.RemoveAt(0, EAllowShrinking::No);
 	}
-	Anticipation.History.Add({ InTargetLocation, DeltaTime });
+	ScreenTargetHistory.History.Add({ InScreenTarget, DeltaTime });
 
-	FVector3d AverageMoveStep(EForceInit::ForceInit);
-	for (int32 Index = 1; Index < Anticipation.History.Num(); ++Index)
+	FVector2d AverageMoveStep(EForceInit::ForceInit);
+	for (int32 Index = 1; Index < ScreenTargetHistory.History.Num(); ++Index)
 	{
-		const FHistoryEntry& Entry(Anticipation.History[Index]);
-		const FHistoryEntry& PrevEntry(Anticipation.History[Index - 1]);
-		const FVector3d NormalizedMoveStep = ((Entry.Key - PrevEntry.Key) / Entry.Value);
+		const FHistoryEntry& Entry(ScreenTargetHistory.History[Index]);
+		const FHistoryEntry& PrevEntry(ScreenTargetHistory.History[Index - 1]);
+		const FVector2d NormalizedMoveStep = ((Entry.Key - PrevEntry.Key) / Entry.Value);
 		AverageMoveStep += NormalizedMoveStep;
 	}
-	AverageMoveStep /= Anticipation.History.Num();
+	AverageMoveStep /= ScreenTargetHistory.History.Num();
 
-	if (AverageMoveStep.Length() >= GFramingExpolationEpsilon)
+	if (AverageMoveStep.Length() >= GFramingExtrapolationEpsilon)
 	{
-		const FVector3d ExtrapolatedMovement = AverageMoveStep * State.TargetMovementAnticipationTime;
-		const FVector3d ExtrapolatedWorldTarget = InTargetLocation + ExtrapolatedMovement;
-		return ExtrapolatedWorldTarget;
+		const FVector2d ExtrapolatedMovement = AverageMoveStep * State.TargetMovementAnticipationTime;
+		const FVector2d ExtrapolatedScreenTarget = InScreenTarget + ExtrapolatedMovement;
+		return ExtrapolatedScreenTarget;
 	}
 	else
 	{
-		return InTargetLocation;
+		return InScreenTarget;
 	}
 }
 
+FFramingZone FBaseFramingCameraNodeEvaluator::ComputeEffectiveDeadZone()
+{
+	const FVector2d& CurTarget = State.ScreenTarget;
+	const FFramingZone& CurTargetBounds = State.ScreenTargetBounds;
+
+	FFramingZone RelativeTargetBounds;
+	RelativeTargetBounds.LeftBound = FMath::Max(0.0, CurTarget.X - CurTargetBounds.LeftBound);
+	RelativeTargetBounds.TopBound = FMath::Max(0.0, CurTarget.Y - CurTargetBounds.TopBound);
+	RelativeTargetBounds.RightBound = FMath::Max(0.0, CurTargetBounds.RightBound - CurTarget.X);
+	RelativeTargetBounds.BottomBound = FMath::Max(0.0, CurTargetBounds.BottomBound - CurTarget.Y);
+
+	FFramingZone EffectiveDeadZone = State.DeadZone;
+	EffectiveDeadZone.LeftBound += RelativeTargetBounds.LeftBound;
+	EffectiveDeadZone.TopBound += RelativeTargetBounds.TopBound;
+	EffectiveDeadZone.RightBound -= RelativeTargetBounds.RightBound;
+	EffectiveDeadZone.BottomBound -= RelativeTargetBounds.BottomBound;
+
+	EffectiveDeadZone.ClampBounds(State.IdealTarget);
+	
+	return EffectiveDeadZone;
+}
+
 FFramingZone FBaseFramingCameraNodeEvaluator::ComputeScreenTargetBounds(const FCameraPose& CameraPose, float AspectRatio, const FTransform3d& TargetTransform, const FBoxSphereBounds3d& LocalBounds)
 {
 	const FVector3d BoxExtent = LocalBounds.BoxExtent;
@@ -397,7 +488,8 @@ void FBaseFramingCameraNodeEvaluator::ComputeDesiredState(float DeltaTime)
 	// If we  don't have any reframing to do, bail out.
 	FVector2d IdealToTarget(State.ScreenTarget - State.IdealTarget);
 	double DistanceToGo = IdealToTarget.Length();
-	if (!State.bIsReframingTarget || DistanceToGo <= GFramingIdealReachedEpsilon)
+	const double ReframingSpeed = State.ReframeDamper.GetX0Derivative();
+	if (!State.bIsReframingTarget || (DistanceToGo <= GFramingIdealReachedEpsilon && ReframingSpeed <= GFramingIdealReachedEpsilon))
 	{
 		ensure(State.TargetFramingState == ETargetFramingState::InDeadZone);
 		Desired.ScreenTarget = State.ScreenTarget;
@@ -468,27 +560,44 @@ void FBaseFramingCameraNodeEvaluator::DrawFramingState(const FState& State, cons
 		Renderer.Draw2DBox(
 				State.SoftZone.GetCanvasPosition(CanvasSize), 
 				State.SoftZone.GetCanvasSize(CanvasSize),
-				FLinearColor::Red,
+				FColorList::Red,
 				1.f);
 		Renderer.Draw2DBox(
 				State.DeadZone.GetCanvasPosition(CanvasSize), 
 				State.DeadZone.GetCanvasSize(CanvasSize),
-				FLinearColor::Green,
-				1.f);
-		Renderer.Draw2DBox(
-				State.EffectiveDeadZone.GetCanvasPosition(CanvasSize), 
-				State.EffectiveDeadZone.GetCanvasSize(CanvasSize),
-				FColorList::LightGrey,
+				FColorList::Green,
 				1.f);
+
 		Renderer.Draw2DBox(
 				State.ScreenTargetBounds.GetCanvasPosition(CanvasSize),
 				State.ScreenTargetBounds.GetCanvasSize(CanvasSize),
-				FLinearColor::Gray,
+				FColorList::LightGrey,
 				1.f);
 
+		if (State.DebugAllScreenTargetBounds.Num() > 1)
+		{
+			for (const FFramingZone& SingleTargetBounds : State.DebugAllScreenTargetBounds)
+			{
+				Renderer.Draw2DBox(
+						SingleTargetBounds.GetCanvasPosition(CanvasSize),
+						SingleTargetBounds.GetCanvasSize(CanvasSize),
+						FColorList::LightGrey,
+						1.f);
+			}
+		}
+
+		if (GFramingShowEffectiveDeadZone)
+		{
+			Renderer.Draw2DBox(
+					State.EffectiveDeadZone.GetCanvasPosition(CanvasSize), 
+					State.EffectiveDeadZone.GetCanvasSize(CanvasSize),
+					FColorList::LightGrey,
+					1.f);
+		}
+
 		if (State.LowReframeDampingFactor > 0)
 		{
-			Renderer.Draw2DLine(State.DeadZoneEdgePoint, State.HardZoneEdgePoint, FLinearColor(FColorList::LightGrey), 1.f);
+			Renderer.Draw2DLine(State.DebugDeadZoneEdgePoint, State.DebugHardZoneEdgePoint, FLinearColor(FColorList::LightGrey), 1.f);
 		}
 
 		const FVector2D ScreenTarget(State.ScreenTarget.X * CanvasSize.X, State.ScreenTarget.Y * CanvasSize.Y);
@@ -520,6 +629,7 @@ void FBaseFramingCameraNodeEvaluator::OnBuildDebugBlocks(const FCameraDebugBlock
 
 	DebugBlock.State = State;
 	DebugBlock.Desired = Desired;
+	DebugBlock.WorldTargets = WorldTargets;
 }
 
 void FBaseFramingCameraDebugBlock::OnDebugDraw(const FCameraDebugBlockDrawParams& Params, FCameraDebugRenderer& Renderer)
@@ -572,14 +682,16 @@ void FBaseFramingCameraDebugBlock::OnDebugDraw(const FCameraDebugBlockDrawParams
 
 	FBaseFramingCameraNodeEvaluator::DrawFramingState(State, Desired, Renderer);
 
-	Renderer.DrawBox(State.WorldTarget, State.LocalTargetBounds.BoxExtent, FLinearColor::Yellow, 0.5f);
+	for (const FCameraActorComputedTargetInfo& TargetInfo : WorldTargets.TargetInfos)
+	{
+		Renderer.DrawBox(TargetInfo.Transform.GetLocation(), TargetInfo.LocalBounds.BoxExtent, FLinearColor::Yellow, 0.5f);
+	}
 }
 
 #endif  // UE_GAMEPLAY_CAMERAS_DEBUG
 
 void FBaseFramingCameraNodeEvaluator::FState::Serialize(FArchive& Ar)
 {
-	Ar << WorldTarget;
 	Ar << IdealTarget;
 	Ar << ReframeDampingFactor;
 	Ar << LowReframeDampingFactor;
@@ -590,7 +702,12 @@ void FBaseFramingCameraNodeEvaluator::FState::Serialize(FArchive& Ar)
 	Ar << ToggleEngageAlpha;
 	Ar << DeadZone;
 	Ar << SoftZone;
+
+	Ar << WorldTarget;
 	Ar << ScreenTarget;
+	Ar << ScreenTargetBounds;
+	Ar << EffectiveDeadZone;
+
 	Ar << TargetFramingState;
 	Ar << bIsReframingTarget;
 	Ar << ReframeDamper;
@@ -603,6 +720,11 @@ void FBaseFramingCameraNodeEvaluator::FDesired::Serialize(FArchive& Ar)
 	Ar << bHasCorrection;
 }
 
+void FBaseFramingCameraNodeEvaluator::FWorldTargetInfos::Serialize(FArchive& Ar)
+{
+	Ar << TargetInfos;
+}
+
 FArchive& operator <<(FArchive& Ar, FBaseFramingCameraNodeEvaluator::FState& State)
 {
 	State.Serialize(Ar);
@@ -615,6 +737,12 @@ FArchive& operator <<(FArchive& Ar, FBaseFramingCameraNodeEvaluator::FDesired& D
 	return Ar;
 }
 
+FArchive& operator <<(FArchive& Ar, FBaseFramingCameraNodeEvaluator::FWorldTargetInfos& WorldTargets)
+{
+	WorldTargets.Serialize(Ar);
+	return Ar;
+}
+
 }  // namespace UE::Cameras
 
 UBaseFramingCameraNode::UBaseFramingCameraNode(const FObjectInitializer& ObjectInit)
@@ -630,6 +758,17 @@ UBaseFramingCameraNode::UBaseFramingCameraNode(const FObjectInitializer& ObjectI
 	DisengageTime.Value = 1.f;
 }
 
+void UBaseFramingCameraNode::PostLoad()
+{
+	Super::PostLoad();
+
+	if (TargetInfos.IsEmpty() && TargetInfo_DEPRECATED.IsValid())
+	{
+		TargetInfos.Add(TargetInfo_DEPRECATED);
+		TargetInfo_DEPRECATED = FCameraActorTargetInfo();
+	}
+}
+
 void UBaseFramingCameraNode::GetCustomCameraNodeParameters(FCustomCameraNodeParameterInfos& OutParameterInfos)
 {
 	OutParameterInfos.AddBlendableParameter(
diff --git a/Engine/Plugins/Cameras/GameplayCameras/Source/GameplayCameras/Private/Nodes/Framing/CameraActorTargetInfo.cpp b/Engine/Plugins/Cameras/GameplayCameras/Source/GameplayCameras/Private/Nodes/Framing/CameraActorTargetInfo.cpp
--- a/Engine/Plugins/Cameras/GameplayCameras/Source/GameplayCameras/Private/Nodes/Framing/CameraActorTargetInfo.cpp
+++ b/Engine/Plugins/Cameras/GameplayCameras/Source/GameplayCameras/Private/Nodes/Framing/CameraActorTargetInfo.cpp
@@ -2,6 +2,7 @@
 
 #include "Nodes/Framing/CameraActorTargetInfo.h"
 
+#include "Algo/Accumulate.h"
 #include "Components/SkeletalMeshComponent.h"
 #include "Core/CameraContextDataTable.h"
 #include "GameFramework/Actor.h"
@@ -113,5 +114,88 @@ void FCameraActorTargetInfoReader::ComputeTargetBounds(const FVector3d& TargetLo
 	}
 }
 
+FCameraActorTargetInfoArrayReader::FCameraActorTargetInfoArrayReader(TConstArrayView<FCameraActorTargetInfo> InTargetInfos, FCameraContextDataID InDataID)
+{
+	Initialize(InTargetInfos, InDataID);
+}
+
+void FCameraActorTargetInfoArrayReader::Initialize(TConstArrayView<FCameraActorTargetInfo> InTargetInfos, FCameraContextDataID InDataID)
+{
+	DataID = InDataID;
+
+	CacheTargetInfos(InTargetInfos);
+}
+
+void FCameraActorTargetInfoArrayReader::CacheTargetInfos(TConstArrayView<FCameraActorTargetInfo> InTargetInfos)
+{
+	Readers.SetNum(InTargetInfos.Num());
+
+	for (int32 Index = 0; Index < InTargetInfos.Num(); ++Index)
+	{
+		Readers[Index].CacheTargetInfo(InTargetInfos[Index]);
+	}
+}
+
+bool FCameraActorTargetInfoArrayReader::ComputeTargetInfos(const FCameraContextDataTable& ContextDataTable, TArray<FCameraActorComputedTargetInfo>& ComputedTargets)
+{
+	if (DataID.IsValid())
+	{
+		TConstArrayView<FCameraActorTargetInfo> NewTargetInfos;
+		if (ContextDataTable.TryGetArrayData<FCameraActorTargetInfo>(DataID, NewTargetInfos))
+		{
+			CacheTargetInfos(NewTargetInfos);
+		}
+	}
+
+	if (Readers.IsEmpty())
+	{
+		return false;
+	}
+
+	ComputedTargets.SetNum(Readers.Num());
+
+	for (int32 Index = 0; Index < Readers.Num(); ++Index)
+	{
+		FCameraActorTargetInfoReader& Reader(Readers[Index]);
+		FCameraActorComputedTargetInfo& ComputedTarget(ComputedTargets[Index]);
+		Reader.GetTargetInfo(ContextDataTable, ComputedTarget.Transform, ComputedTarget.LocalBounds);
+		ComputedTarget.NormalizedWeight = Reader.CachedTargetInfo.Weight;
+	}
+
+	float TotalWeight = Algo::Accumulate(
+			ComputedTargets, 
+			0.f,
+			[](float Cur, const FCameraActorComputedTargetInfo& Item) { return Cur + Item.NormalizedWeight; });
+	if (TotalWeight == 0.f)
+	{
+		return false;
+	}
+
+	for (int32 Index = 0; Index < Readers.Num(); ++Index)
+	{
+		FCameraActorComputedTargetInfo& ComputedTarget(ComputedTargets[Index]);
+		ComputedTarget.NormalizedWeight /= TotalWeight;
+	}
+
+	return true;
+}
+
+#if WITH_EDITOR
+
+void FCameraActorTargetInfoArrayReader::Refresh(TConstArrayView<FCameraActorTargetInfo> InTargetInfos)
+{
+	CacheTargetInfos(InTargetInfos);
+}
+
+#endif  // WITH_EDITOR
+
+FArchive& operator <<(FArchive& Ar, FCameraActorComputedTargetInfo& TargetInfo)
+{
+	Ar << TargetInfo.Transform;
+	Ar << TargetInfo.LocalBounds;
+	Ar << TargetInfo.NormalizedWeight;
+	return Ar;
+}
+
 }  // namespace UE::Cameras
 
diff --git a/Engine/Plugins/Cameras/GameplayCameras/Source/GameplayCameras/Private/Nodes/Framing/DollyFramingCameraNode.cpp b/Engine/Plugins/Cameras/GameplayCameras/Source/GameplayCameras/Private/Nodes/Framing/DollyFramingCameraNode.cpp
--- a/Engine/Plugins/Cameras/GameplayCameras/Source/GameplayCameras/Private/Nodes/Framing/DollyFramingCameraNode.cpp
+++ b/Engine/Plugins/Cameras/GameplayCameras/Source/GameplayCameras/Private/Nodes/Framing/DollyFramingCameraNode.cpp
@@ -75,15 +75,9 @@ void FDollyFramingCameraNodeEvaluator::OnInitialize(const FCameraNodeEvaluatorIn
 
 void FDollyFramingCameraNodeEvaluator::OnRun(const FCameraNodeEvaluationParams& Params, FCameraNodeEvaluationResult& OutResult)
 {
-	FAcquiredTargetInfo TargetInfo;
-	if (!AcquireTargetInfo(Params, OutResult, TargetInfo))
-	{
-		return;
-	}
-
 	// Let the base class figure out all the screen-space framing stuff.
 	const FTransform3d LastShotTransform = BuildDollyShotTransform(OutResult.CameraPose);
-	UpdateFramingState(Params, OutResult, TargetInfo, LastShotTransform);
+	UpdateFramingState(Params, OutResult, LastShotTransform);
 
 	// If we need to reframe the target this tick, figure out how much we need to move the dolly
 	// to accomplish that.
diff --git a/Engine/Plugins/Cameras/GameplayCameras/Source/GameplayCameras/Private/Nodes/Framing/PanningFramingCameraNode.cpp b/Engine/Plugins/Cameras/GameplayCameras/Source/GameplayCameras/Private/Nodes/Framing/PanningFramingCameraNode.cpp
--- a/Engine/Plugins/Cameras/GameplayCameras/Source/GameplayCameras/Private/Nodes/Framing/PanningFramingCameraNode.cpp
+++ b/Engine/Plugins/Cameras/GameplayCameras/Source/GameplayCameras/Private/Nodes/Framing/PanningFramingCameraNode.cpp
@@ -69,15 +69,9 @@ void FPanningFramingCameraNodeEvaluator::OnInitialize(const FCameraNodeEvaluator
 
 void FPanningFramingCameraNodeEvaluator::OnRun(const FCameraNodeEvaluationParams& Params, FCameraNodeEvaluationResult& OutResult)
 {
-	FAcquiredTargetInfo TargetInfo;
-	if (!AcquireTargetInfo(Params, OutResult, TargetInfo))
-	{
-		return;
-	}
-
 	// Let the base class figure out all the screen-space framing stuff.
 	const FTransform3d LastShotTransform = BuildPanningShotTransform(OutResult.CameraPose);
-	UpdateFramingState(Params, OutResult, TargetInfo, LastShotTransform);
+	UpdateFramingState(Params, OutResult, LastShotTransform);
 
 	// If we need to reframe the target this tick, figure out how much we need to rotate the camera
 	// to accomplish that.
diff --git a/Engine/Plugins/Cameras/GameplayCameras/Source/GameplayCameras/Public/Math/CameraFramingZoneMath.h b/Engine/Plugins/Cameras/GameplayCameras/Source/GameplayCameras/Public/Math/CameraFramingZoneMath.h
--- a/Engine/Plugins/Cameras/GameplayCameras/Source/GameplayCameras/Public/Math/CameraFramingZoneMath.h
+++ b/Engine/Plugins/Cameras/GameplayCameras/Source/GameplayCameras/Public/Math/CameraFramingZoneMath.h
@@ -59,6 +59,11 @@ struct FFramingZone
 	 */
 	void ClampBounds(const FFramingZone& MustContain);
 
+	/**
+	 * Expands this framing zone to include the other framing zone.
+	 */
+	void Add(const FFramingZone& Other);
+
 	/** Checks whether the given point (in 0..1 UI space) is inside this zone. */
 	bool Contains(const FVector2d& Point) const;
 
diff --git a/Engine/Plugins/Cameras/GameplayCameras/Source/GameplayCameras/Public/Nodes/Framing/BaseFramingCameraNode.h b/Engine/Plugins/Cameras/GameplayCameras/Source/GameplayCameras/Public/Nodes/Framing/BaseFramingCameraNode.h
--- a/Engine/Plugins/Cameras/GameplayCameras/Source/GameplayCameras/Public/Nodes/Framing/BaseFramingCameraNode.h
+++ b/Engine/Plugins/Cameras/GameplayCameras/Source/GameplayCameras/Public/Nodes/Framing/BaseFramingCameraNode.h
@@ -30,16 +30,17 @@ public:
 	/** 
 	 * A variable whose value is the desired target's location in world space.
 	 * If set, and if the variable has been set, the obtained value takes priority
-	 * over the TargetInfo property.
+	 * over the TargetInfos property.
 	 */
 	UPROPERTY(EditAnywhere, Category="Target")
 	FVector3dCameraVariableReference TargetLocation;
 
+	/** Specifies one or more target actors to frame. */
 	UPROPERTY(EditAnywhere, Category="Target", meta=(CameraContextData=true))
-	FCameraActorTargetInfo TargetInfo;
+	TArray<FCameraActorTargetInfo> TargetInfos;
 
 	UPROPERTY()
-	FCameraContextDataID TargetInfoDataID;
+	FCameraContextDataID TargetInfosDataID;
 
 	/**
 	 * Whether the camera pose's target distance should be set to the distance between
@@ -102,12 +103,20 @@ public:
 	UPROPERTY(EditAnywhere, Category="Framing Zones")
 	FCameraFramingZoneParameter SoftZone;
 
+private:
+
+	UPROPERTY()
+	FCameraActorTargetInfo TargetInfo_DEPRECATED;
+
 public:
 
 	UBaseFramingCameraNode(const FObjectInitializer& ObjectInit);
 
 public:
 
+	// UObject interface.
+	virtual void PostLoad() override;
+
 	// ICustomCameraNodeParameterProvider interface.
 	virtual void GetCustomCameraNodeParameters(FCustomCameraNodeParameterInfos& OutParameterInfos) override;
 };
@@ -158,12 +167,9 @@ protected:
 
 	struct FState;
 	struct FDesired;
-	struct FAcquiredTargetInfo;
 
-	/** Gets the target location. */
-	bool AcquireTargetInfo(const FCameraNodeEvaluationParams& Params, const FCameraNodeEvaluationResult& InResult, FAcquiredTargetInfo& OutInfo);
 	/** Updates the framing state for the current tick, see State member field. */
-	void UpdateFramingState(const FCameraNodeEvaluationParams& Params, const FCameraNodeEvaluationResult& OutResult, const FAcquiredTargetInfo& TargetInfo, const FTransform3d& LastFraming);
+	void UpdateFramingState(const FCameraNodeEvaluationParams& Params, const FCameraNodeEvaluationResult& OutResult, const FTransform3d& LastFraming);
 	/** Wraps-up the update with optional operations. */
 	void EndFramingUpdate(const FCameraNodeEvaluationParams& Params, FCameraNodeEvaluationResult& OutResult);
 
@@ -175,11 +181,16 @@ private:
 
 	FVector2d GetHardReframeCoords() const;
 
-	void ComputeCurrentState(const FCameraNodeEvaluationParams& Params, const FCameraNodeEvaluationResult& OutResult, const FAcquiredTargetInfo& TargetInfo, const FTransform3d& LastFraming);
-	FVector3d ComputeAnticipatedScreenTarget(float DeltaTime, const FVector3d& InTargetLocation);
-	FFramingZone ComputeScreenTargetBounds(const FCameraPose& CameraPose, float AspectRatio, const FTransform3d& TargetTransform, const FBoxSphereBounds3d& LocalBounds);
+	void ComputeCurrentState(const FCameraNodeEvaluationParams& Params, const FCameraNodeEvaluationResult& OutResult, const FTransform3d& LastFraming);
 	void ComputeDesiredState(float DeltaTime);
 
+	bool AcquireTargetInfo(const FCameraNodeEvaluationParams& Params, const FCameraNodeEvaluationResult& InResult, TArray<FCameraActorComputedTargetInfo>& OutInfos);
+	bool ComputeFinalTargetInfo(const FCameraNodeEvaluationParams& Params, const FCameraPose& CameraPose, FVector3d& OutWorldTarget, FVector2d& OutScreenTarget, FFramingZone& OutScreenBounds);
+	FVector2d ComputeAnticipatedScreenTarget(float DeltaTime, const FVector2d& InPreviousAnticipatedScreenTarget, const FVector2d& InScreenTarget);
+	FFramingZone ComputeEffectiveDeadZone();
+
+	static FFramingZone ComputeScreenTargetBounds(const FCameraPose& CameraPose, float AspectRatio, const FTransform3d& TargetTransform, const FBoxSphereBounds3d& LocalBounds);
+
 protected:
 
 	/** The current location of the target. */
@@ -202,16 +213,10 @@ protected:
 		InHardZone
 	};
 
-	struct FAcquiredTargetInfo
-	{
-		FTransform3d WorldTransform;
-		FBoxSphereBounds3d LocalBounds;
-	};
-
 	/** Utility structure for all the parameter readers we need every frame. */
 	struct FReaders
 	{
-		FCameraActorTargetInfoReader TargetInfo;
+		FCameraActorTargetInfoArrayReader TargetInfos;
 
 		TCameraParameterReader<FVector2d> IdealFramingLocation;
 		TCameraParameterReader<bool> SetTargetDistance;
@@ -230,11 +235,6 @@ protected:
 	/** Utility struct for storing the current known state. */
 	struct FState
 	{
-		/** World position of the tracked target. */
-		FVector3d WorldTarget;
-		/** Bounds of the tracked target. */
-		FBoxSphereBounds3d LocalTargetBounds;
-
 		/** Screen-space position of the ideal framing position. */
 		FVector2d IdealTarget;
 		/** Current reframing damping factor. */
@@ -253,17 +253,20 @@ protected:
 		float ToggleEngageAlpha;
 		/** Current look-ahead time for anticipating target movement */
 		float TargetMovementAnticipationTime;
-		/** Current target bounds zone. */
-		FFramingZone ScreenTargetBounds;
-		/** Dead zone minus the screen target bounds. */
-		FFramingZone EffectiveDeadZone;
 		/** Current coordinates of the dead zone. */
 		FFramingZone DeadZone;
 		/** Current coordinates of the soft zone. */
 		FFramingZone SoftZone;
 
+		/** Current world-space position of the tracked target. */
+		FVector3d WorldTarget;
 		/** Current screen-space position of the tracked target. */
 		FVector2d ScreenTarget;
+		/** Current target bounds zone. */
+		FFramingZone ScreenTargetBounds;
+		/** Dead zone minus the screen target bounds. */
+		FFramingZone EffectiveDeadZone;
+
 		/** Current state of the tracked target. */
 		ETargetFramingState TargetFramingState;
 		/** Whether we are actively trying to bring the target back to the ideal position. */
@@ -274,9 +277,11 @@ protected:
 
 #if UE_GAMEPLAY_CAMERAS_DEBUG
 		/** Intersection of the reframing vector with the dead zone box. */
-		FVector2d DeadZoneEdgePoint;
+		FVector2d DebugDeadZoneEdgePoint;
 		/** Intersection of the reframing vector with the hard zone box. */
-		FVector2d HardZoneEdgePoint;
+		FVector2d DebugHardZoneEdgePoint;
+		/** Screen bounds for all the targets. */
+		TArray<FFramingZone, TInlineAllocator<4>> DebugAllScreenTargetBounds;
 #endif  // UE_GAMEPLAY_CAMERAS_DEBUG
 
 		void Serialize(FArchive& Ar);
@@ -306,15 +311,25 @@ protected:
 	};
 	FDesired Desired;
 
-	struct FWorldTargetAnticipation
+	struct FWorldTargetInfos
+	{
+		TArray<FCameraActorComputedTargetInfo, TInlineAllocator<4>> TargetInfos;
+
+		void Serialize(FArchive& Ar);
+	};
+	FWorldTargetInfos WorldTargets;
+
+	struct FScreenTargetHistory
 	{
-		TArray<TTuple<FVector3d, float>, TInlineAllocator<10>> History;
+		FVector2d UnanticipatedScreenTarget;
+		TArray<TTuple<FVector2d, float>, TInlineAllocator<10>> History;
 	};
-	FWorldTargetAnticipation Anticipation;
+	FScreenTargetHistory ScreenTargetHistory;
 
 	friend class FBaseFramingCameraDebugBlock;
 	friend FArchive& operator <<(FArchive& Ar, FState& State);
 	friend FArchive& operator <<(FArchive& Ar, FDesired& Desired);
+	friend FArchive& operator <<(FArchive& Ar, FWorldTargetInfos& WorldTargets);
 };
 
 FArchive& operator <<(FArchive& Ar, FBaseFramingCameraNodeEvaluator::FState& State);
diff --git a/Engine/Plugins/Cameras/GameplayCameras/Source/GameplayCameras/Public/Nodes/Framing/CameraActorTargetInfo.h b/Engine/Plugins/Cameras/GameplayCameras/Source/GameplayCameras/Public/Nodes/Framing/CameraActorTargetInfo.h
--- a/Engine/Plugins/Cameras/GameplayCameras/Source/GameplayCameras/Public/Nodes/Framing/CameraActorTargetInfo.h
+++ b/Engine/Plugins/Cameras/GameplayCameras/Source/GameplayCameras/Public/Nodes/Framing/CameraActorTargetInfo.h
@@ -8,6 +8,7 @@
 #include "CameraActorTargetInfo.generated.h"
 
 class AActor;
+struct FCameraPose;
 
 namespace UE::Cameras
 {
@@ -89,6 +90,41 @@ private:
 	const USkeletalMeshComponent* CachedSkeletalMeshComponent = nullptr;
 	FName CachedBoneName;
 	FName CachedParentBoneName;
+
+	friend struct FCameraActorTargetInfoArrayReader;
+};
+
+struct FCameraActorComputedTargetInfo
+{
+	FTransform3d Transform;
+	FBoxSphereBounds3d LocalBounds;
+	float NormalizedWeight = 1.f;
+};
+
+FArchive& operator <<(FArchive& Ar, FCameraActorComputedTargetInfo& TargetInfo);
+
+/** A special reader class for multiple targeting information. */
+struct FCameraActorTargetInfoArrayReader
+{
+	FCameraActorTargetInfoArrayReader() {}
+	FCameraActorTargetInfoArrayReader(TConstArrayView<FCameraActorTargetInfo> InTargetInfos, FCameraContextDataID InDataID);
+
+	void Initialize(TConstArrayView<FCameraActorTargetInfo> InTargetInfos, FCameraContextDataID InDataID);
+
+	bool ComputeTargetInfos(const FCameraContextDataTable& ContextDataTable, TArray<FCameraActorComputedTargetInfo>& ComputedTargets);
+
+#if WITH_EDITOR
+	void Refresh(TConstArrayView<FCameraActorTargetInfo> InTargetInfos);
+#endif
+
+private:
+
+	void CacheTargetInfos(TConstArrayView<FCameraActorTargetInfo> InTargetInfos);
+
+private:
+
+	TArray<FCameraActorTargetInfoReader> Readers;
+	FCameraContextDataID DataID;
 };
 
 }  // namespace UE::Cameras
diff --git a/Engine/Plugins/Cameras/GameplayCameras/Source/GameplayCameras/Public/Nodes/Framing/CameraFramingZone.h b/Engine/Plugins/Cameras/GameplayCameras/Source/GameplayCameras/Public/Nodes/Framing/CameraFramingZone.h
--- a/Engine/Plugins/Cameras/GameplayCameras/Source/GameplayCameras/Public/Nodes/Framing/CameraFramingZone.h
+++ b/Engine/Plugins/Cameras/GameplayCameras/Source/GameplayCameras/Public/Nodes/Framing/CameraFramingZone.h
@@ -10,8 +10,9 @@
 /**
  * A structure that defines a zone for use in framing subjects in screen-space.
  *
- * All margins are defined in percentages of the screen's horizontal size. They are also 
- * all defined relative to their respective edges.
+ * Margins are generally expressed in screen size percentages (between 0 and 1), but don't have a standard meaning.
+ * In some cases, they may be margins from the screen's edge, while in other cases they may be margins relative to
+ * a given screen point.
  */
 USTRUCT(BlueprintType)
 struct FCameraFramingZone
