Cameras: fix possible default parameterrs memory corruption and handle blendable structs

diff --git a/Engine/Plugins/Cameras/GameplayCameras/Source/GameplayCameras/Private/Core/CameraRigParameterDefinition.cpp b/Engine/Plugins/Cameras/GameplayCameras/Source/GameplayCameras/Private/Core/CameraRigParameterDefinition.cpp
--- a/Engine/Plugins/Cameras/GameplayCameras/Source/GameplayCameras/Private/Core/CameraRigParameterDefinition.cpp
+++ b/Engine/Plugins/Cameras/GameplayCameras/Source/GameplayCameras/Private/Core/CameraRigParameterDefinition.cpp
@@ -131,8 +131,9 @@ void FCameraRigParameterBuilder::SetDefaultParameterValues(const UCameraRigAsset
 			continue;
 		}
 
-		const uint8* RawSourceValuePtr = nullptr;
+		const void* RawSourceValuePtr = nullptr;
 
+		// First check if the value is found on a custom parameter.
 		if (ICustomCameraNodeParameterProvider* CustomParameterProvider = Cast<ICustomCameraNodeParameterProvider>(CameraNode))
 		{
 			FCustomCameraNodeParameterInfos CustomParameters;
@@ -150,6 +151,7 @@ void FCameraRigParameterBuilder::SetDefaultParameterValues(const UCameraRigAsset
 			}
 		}
 
+		// If not found, check on a reflected UObject property.
 		if (!RawSourceValuePtr)
 		{
 			const UClass* TargetClass = CameraNode->GetClass();
@@ -163,15 +165,20 @@ void FCameraRigParameterBuilder::SetDefaultParameterValues(const UCameraRigAsset
 					case ECameraVariableType::ValueName:\
 						{\
 							using CameraParameterType = F##ValueName##CameraParameter;\
-							if (StructProperty->Struct == CameraParameterType::StaticStruct())\
+							if (ensure(StructProperty->Struct == CameraParameterType::StaticStruct()))\
 							{\
-								auto* CameraParameterPtr = StructProperty->ContainerPtrToValuePtr<CameraParameterType>(CameraNode);\
-								RawSourceValuePtr = reinterpret_cast<uint8*>(&CameraParameterPtr->Value);\
+								CameraParameterType* CameraParameterPtr = StructProperty->ContainerPtrToValuePtr<CameraParameterType>(CameraNode);\
+								RawSourceValuePtr = static_cast<void*>(&CameraParameterPtr->Value);\
 							}\
 						}\
 						break;
 					UE_CAMERA_VARIABLE_FOR_ALL_TYPES()
 #undef UE_CAMERA_VARIABLE_FOR_TYPE
+					case ECameraVariableType::BlendableStruct:
+						{
+							RawSourceValuePtr = StructProperty->ContainerPtrToValuePtr<void>(CameraNode);
+						}
+						break;
 				}
 			}
 		}
@@ -181,15 +188,17 @@ void FCameraRigParameterBuilder::SetDefaultParameterValues(const UCameraRigAsset
 			continue;
 		}
 
+		// Find the corresponding property on the default parameters' property bag.
 		const FPropertyBagPropertyDesc* PropertyDesc = PropertyBagStruct->FindPropertyDescByID(BlendableParameter->GetGuid());
 		if (!ensure(PropertyDesc && PropertyDesc->CachedProperty))
 		{
 			continue;
 		}
 
-		void* RawDestinationValuePtr = PropertyDesc->CachedProperty->ContainerPtrToValuePtr<void>(PropertyBagValue);
-
-		if (!ensure(RawDestinationValuePtr))
+		// This property should be a structure: either a camera parameter for all the standard blendable types,
+		// or a blendable structure.
+		const FStructProperty* DefaultParameterProperty = CastField<FStructProperty>(PropertyDesc->CachedProperty);
+		if (!ensure(DefaultParameterProperty))
 		{
 			continue;
 		}
@@ -200,16 +209,20 @@ void FCameraRigParameterBuilder::SetDefaultParameterValues(const UCameraRigAsset
 			case ECameraVariableType::ValueName:\
 				{\
 					using CameraParameterType = F##ValueName##CameraParameter;\
-					const ValueType* SourceValuePtr = reinterpret_cast<const ValueType*>(RawSourceValuePtr);\
-					CameraParameterType* DestinationParameter = reinterpret_cast<CameraParameterType*>(RawDestinationValuePtr);\
-					DestinationParameter->Value = *SourceValuePtr;\
+					if (ensure(DefaultParameterProperty->Struct == CameraParameterType::StaticStruct()))\
+					{\
+						const ValueType* SourceValuePtr = reinterpret_cast<const ValueType*>(RawSourceValuePtr);\
+						CameraParameterType* DestinationParameter = DefaultParameterProperty->ContainerPtrToValuePtr<CameraParameterType>(PropertyBagValue);\
+						DestinationParameter->Value = *SourceValuePtr;\
+					}\
 				}\
 				break;
 			UE_CAMERA_VARIABLE_FOR_ALL_TYPES()
 #undef UE_CAMERA_VARIABLE_FOR_TYPE
 			case ECameraVariableType::BlendableStruct:
-				if (ensure(BlendableParameter->BlendableStructType))
+				if (ensure(DefaultParameterProperty->Struct == BlendableParameter->BlendableStructType))
 				{
+					void* RawDestinationValuePtr = DefaultParameterProperty->ContainerPtrToValuePtr<void>(PropertyBagValue);
 					BlendableParameter->BlendableStructType->CopyScriptStruct(RawDestinationValuePtr, RawSourceValuePtr);
 				}
 				break;
