Cameras: fix incorrect test for auto-saving

diff --git a/Engine/Plugins/Cameras/GameplayCameras/Source/GameplayCameras/Private/Core/CameraAsset.cpp b/Engine/Plugins/Cameras/GameplayCameras/Source/GameplayCameras/Private/Core/CameraAsset.cpp
--- a/Engine/Plugins/Cameras/GameplayCameras/Source/GameplayCameras/Private/Core/CameraAsset.cpp
+++ b/Engine/Plugins/Cameras/GameplayCameras/Source/GameplayCameras/Private/Core/CameraAsset.cpp
@@ -163,7 +163,7 @@ void UCameraAsset::PreSave(FObjectPreSaveContext ObjectSaveContext)
 #if WITH_EDITOR
 
 	const bool bIsUserObject = !HasAnyFlags(RF_ClassDefaultObject | RF_ArchetypeObject);
-	const bool bIsEditorAutoSave = !(ObjectSaveContext.GetSaveFlags() & SAVE_FromAutosave);
+	const bool bIsEditorAutoSave = ((ObjectSaveContext.GetSaveFlags() & SAVE_FromAutosave) != 0);
 	if (bIsUserObject && !bIsEditorAutoSave)
 	{
 		// Build when saving/cooking.
diff --git a/Engine/Plugins/Cameras/GameplayCameras/Source/GameplayCameras/Private/Core/CameraRigAsset.cpp b/Engine/Plugins/Cameras/GameplayCameras/Source/GameplayCameras/Private/Core/CameraRigAsset.cpp
--- a/Engine/Plugins/Cameras/GameplayCameras/Source/GameplayCameras/Private/Core/CameraRigAsset.cpp
+++ b/Engine/Plugins/Cameras/GameplayCameras/Source/GameplayCameras/Private/Core/CameraRigAsset.cpp
@@ -272,7 +272,7 @@ void UCameraRigAsset::PreSave(FObjectPreSaveContext ObjectSaveContext)
 #if WITH_EDITOR
 
 	const bool bIsUserObject = !HasAnyFlags(RF_ClassDefaultObject | RF_ArchetypeObject);
-	const bool bIsEditorAutoSave = !(ObjectSaveContext.GetSaveFlags() & SAVE_FromAutosave);
+	const bool bIsEditorAutoSave = ((ObjectSaveContext.GetSaveFlags() & SAVE_FromAutosave) != 0);
 	if (bIsUserObject && !bIsEditorAutoSave)
 	{
 		// Build when saving/cooking.
