Cameras: fix sequencer animation for array data parameters

diff --git a/Engine/Plugins/Cameras/GameplayCameras/Source/GameplayCameras/Private/MovieScene/MovieSceneCameraParameterInstantiator.cpp b/Engine/Plugins/Cameras/GameplayCameras/Source/GameplayCameras/Private/MovieScene/MovieSceneCameraParameterInstantiator.cpp
--- a/Engine/Plugins/Cameras/GameplayCameras/Source/GameplayCameras/Private/MovieScene/MovieSceneCameraParameterInstantiator.cpp
+++ b/Engine/Plugins/Cameras/GameplayCameras/Source/GameplayCameras/Private/MovieScene/MovieSceneCameraParameterInstantiator.cpp
@@ -118,7 +118,19 @@ struct FSetupCameraParameterOverrideTask
 			return FGuid();
 		}
 
-		const FString& ParameterName = PropertyNames[3];
+		FString& ParameterName = PropertyNames[3];
+		
+		// Strip out the array element index from the name, if we are animating something inside an
+		// array data parameter.
+		if (ParameterName.Len() > 0 && ParameterName.GetCharArray()[ParameterName.Len() - 1] == ']')
+		{
+			int32 OpenIndex = 0;
+			if (ParameterName.FindLastChar('[', OpenIndex))
+			{
+				ParameterName = FString::ConstructFromPtrSize(*ParameterName, OpenIndex);
+			}
+		}
+
 		for (const FCameraRigParameterDefinition& ParameterDefintion : CameraAsset->GetParameterDefinitions())
 		{
 			if (ParameterDefintion.ParameterName == ParameterName)
