Cameras: add viewport override option on evaluation contexts

This allows us to use the editor viewport's aspect ratio when running preview rigs in the editor.

diff --git a/Engine/Plugins/Cameras/GameplayCameras/Source/GameplayCameras/Private/Core/CameraEvaluationContext.cpp b/Engine/Plugins/Cameras/GameplayCameras/Source/GameplayCameras/Private/Core/CameraEvaluationContext.cpp
--- a/Engine/Plugins/Cameras/GameplayCameras/Source/GameplayCameras/Private/Core/CameraEvaluationContext.cpp
+++ b/Engine/Plugins/Cameras/GameplayCameras/Source/GameplayCameras/Private/Core/CameraEvaluationContext.cpp
@@ -108,6 +108,23 @@ UWorld* FCameraEvaluationContext::GetWorld() const
 	return nullptr;
 }
 
+FIntPoint FCameraEvaluationContext::GetViewportSize() const
+{
+	if (OverrideViewportSize.IsSet())
+	{
+		return OverrideViewportSize.GetValue();
+	}
+
+	if (APlayerController* PlayerController = GetPlayerController())
+	{
+		int32 ViewportSizeX = 0, ViewportSizeY = 0;
+		PlayerController->GetViewportSize(ViewportSizeX, ViewportSizeY);
+		return FIntPoint(ViewportSizeX, ViewportSizeY);
+	}
+
+	return FIntPoint(EForceInit::ForceInit);
+}
+
 void FCameraEvaluationContext::AddReferencedObjects(FReferenceCollector& Collector)
 {
 	Collector.AddReferencedObject(CameraAsset);
@@ -261,6 +278,7 @@ void FCameraEvaluationContext::Activate(const FCameraEvaluationContextActivatePa
 		DirectorEvaluator->Activate(DirectorParams);
 	}
 
+	CameraSystemEvaluator = Params.Evaluator;
 	bActivated = true;
 }
 
@@ -283,6 +301,7 @@ void FCameraEvaluationContext::Deactivate(const FCameraEvaluationContextDeactiva
 
 	OnDeactivate(Params);
 
+	CameraSystemEvaluator = nullptr;
 	bActivated = false;
 }
 
diff --git a/Engine/Plugins/Cameras/GameplayCameras/Source/GameplayCameras/Private/Math/CameraPoseMath.cpp b/Engine/Plugins/Cameras/GameplayCameras/Source/GameplayCameras/Private/Math/CameraPoseMath.cpp
--- a/Engine/Plugins/Cameras/GameplayCameras/Source/GameplayCameras/Private/Math/CameraPoseMath.cpp
+++ b/Engine/Plugins/Cameras/GameplayCameras/Source/GameplayCameras/Private/Math/CameraPoseMath.cpp
@@ -3,12 +3,11 @@
 #include "Math/CameraPoseMath.h"
 
 #include "Core/CameraPose.h"
+#include "Core/CameraEvaluationContext.h"
 #include "CoreGlobals.h"
-#include "GameFramework/PlayerController.h"
 #include "Math/InverseRotationMatrix.h"
 #include "Math/PerspectiveMatrix.h"
 #include "Math/TranslationMatrix.h"
-#include "RHI.h"
 
 namespace UE::Cameras
 {
@@ -18,22 +17,14 @@ FCameraFieldsOfView FCameraPoseMath::GetEffectiveFieldsOfView(const FCameraPose&
 	return GetEffectiveFieldsOfView(CameraPose, CameraPose.GetSensorAspectRatio());
 }
 
-FCameraFieldsOfView FCameraPoseMath::GetEffectiveFieldsOfView(const FCameraPose& CameraPose, const APlayerController* PlayerController)
+FCameraFieldsOfView FCameraPoseMath::GetEffectiveFieldsOfView(const FCameraPose& CameraPose, TSharedPtr<const FCameraEvaluationContext> EvaluationContext)
 {
-	const double AspectRatio = GetEffectiveAspectRatio(CameraPose, PlayerController);
+	const double AspectRatio = GetEffectiveAspectRatio(CameraPose, EvaluationContext);
 	return GetEffectiveFieldsOfView(CameraPose, AspectRatio);
 }
 
 FCameraFieldsOfView FCameraPoseMath::GetEffectiveFieldsOfView(const FCameraPose& CameraPose, double AspectRatio)
 {
-	if (CameraPose.GetConstrainAspectRatio() || 
-			!ensureMsgf(AspectRatio > 0.0, TEXT("Invalid aspect ratio! Using sensor aspect ratio instead.")))
-	{
-		// Aspect ratio is constrained, there will be black bars to enforce it to be equal to
-		// our sensor aspect ratio.
-		return GetEffectiveFieldsOfView(CameraPose, CameraPose.GetSensorAspectRatio());
-	}
-
 	const double HorizontalFOV = CameraPose.GetEffectiveFieldOfView();
 
 	// Check the sort of aspect ratio axis constraint we have.
@@ -69,20 +60,19 @@ FCameraFieldsOfView FCameraPoseMath::GetEffectiveFieldsOfView(const FCameraPose&
 	}
 }
 
-double FCameraPoseMath::GetEffectiveAspectRatio(const FCameraPose& CameraPose, const APlayerController* PlayerController)
+double FCameraPoseMath::GetEffectiveAspectRatio(const FCameraPose& CameraPose, TSharedPtr<const FCameraEvaluationContext> EvaluationContext)
 {
-	if (CameraPose.GetConstrainAspectRatio() || PlayerController == nullptr)
+	if (CameraPose.GetConstrainAspectRatio() || !EvaluationContext)
 	{
 		return CameraPose.GetSensorAspectRatio();
 	}
 	else
 	{
-		int32 ViewportSizeX = 0, ViewportSizeY = 0;
-		PlayerController->GetViewportSize(ViewportSizeX, ViewportSizeY);
-		if (ensureMsgf(ViewportSizeX > 0 && ViewportSizeY > 0,
+		FIntPoint ViewportSize = EvaluationContext->GetViewportSize();
+		if (ensureMsgf(ViewportSize.X > 0 && ViewportSize.Y > 0,
 				TEXT("Can't get viewport aspect ratio! Using sensor aspect ratio instead.")))
 		{
-			return (double)ViewportSizeX / (double)ViewportSizeY;
+			return (double)ViewportSize.X / (double)ViewportSize.Y;
 		}
 		else
 		{
diff --git a/Engine/Plugins/Cameras/GameplayCameras/Source/GameplayCameras/Private/Nodes/Framing/BaseFramingCameraNode.cpp b/Engine/Plugins/Cameras/GameplayCameras/Source/GameplayCameras/Private/Nodes/Framing/BaseFramingCameraNode.cpp
--- a/Engine/Plugins/Cameras/GameplayCameras/Source/GameplayCameras/Private/Nodes/Framing/BaseFramingCameraNode.cpp
+++ b/Engine/Plugins/Cameras/GameplayCameras/Source/GameplayCameras/Private/Nodes/Framing/BaseFramingCameraNode.cpp
@@ -326,8 +326,7 @@ bool FBaseFramingCameraNodeEvaluator::ComputeFinalTargetInfo(const FCameraNodeEv
 	}
 
 	// Start with projecting all the targets, and their bounds, on screen.
-	APlayerController* PlayerController = Params.EvaluationContext->GetPlayerController();
-	const double AspectRatio = FCameraPoseMath::GetEffectiveAspectRatio(CameraPose, PlayerController);
+	const double AspectRatio = FCameraPoseMath::GetEffectiveAspectRatio(CameraPose, Params.EvaluationContext);
 
 	struct FComputedTargetScreenInfo
 	{
diff --git a/Engine/Plugins/Cameras/GameplayCameras/Source/GameplayCameras/Private/Nodes/Framing/DollyFramingCameraNode.cpp b/Engine/Plugins/Cameras/GameplayCameras/Source/GameplayCameras/Private/Nodes/Framing/DollyFramingCameraNode.cpp
--- a/Engine/Plugins/Cameras/GameplayCameras/Source/GameplayCameras/Private/Nodes/Framing/DollyFramingCameraNode.cpp
+++ b/Engine/Plugins/Cameras/GameplayCameras/Source/GameplayCameras/Private/Nodes/Framing/DollyFramingCameraNode.cpp
@@ -36,7 +36,7 @@ protected:
 private:
 
 	FTransform3d BuildDollyShotTransform(const FCameraPose& CameraPose) const;
-	FVector3d ComputeFramingTranslation(const FCameraPose& CameraPose, APlayerController* PlayerController);
+	FVector3d ComputeFramingTranslation(const FCameraPose& CameraPose, TSharedPtr<const FCameraEvaluationContext> EvaluationContext);
 
 private:
 
@@ -86,8 +86,7 @@ void FDollyFramingCameraNodeEvaluator::OnRun(const FCameraNodeEvaluationParams&
 		FCameraPose LastShotPose(OutResult.CameraPose);
 		LastShotPose.SetTransform(LastShotTransform);
 
-		APlayerController* PlayerController = Params.EvaluationContext->GetPlayerController();
-		FVector3d DesiredLocalOffset = ComputeFramingTranslation(LastShotPose, PlayerController);
+		FVector3d DesiredLocalOffset = ComputeFramingTranslation(LastShotPose, Params.EvaluationContext);
 
 		// We never bring the dolly forward or backward (we only move it vertically or horizontally).
 		DesiredLocalOffset.X = 0.f;
@@ -128,9 +127,9 @@ FTransform3d FDollyFramingCameraNodeEvaluator::BuildDollyShotTransform(const FCa
 	return Transform;
 }
 
-FVector3d FDollyFramingCameraNodeEvaluator::ComputeFramingTranslation(const FCameraPose& CameraPose, APlayerController* PlayerController)
+FVector3d FDollyFramingCameraNodeEvaluator::ComputeFramingTranslation(const FCameraPose& CameraPose, TSharedPtr<const FCameraEvaluationContext> EvaluationContext)
 {
-	const float AspectRatio = FCameraPoseMath::GetEffectiveAspectRatio(CameraPose, PlayerController);
+	const float AspectRatio = FCameraPoseMath::GetEffectiveAspectRatio(CameraPose, EvaluationContext);
 	const FCameraFieldsOfView FOVs(FCameraPoseMath::GetEffectiveFieldsOfView(CameraPose, AspectRatio));
 
 	// Get the position of the current target in camera space.
diff --git a/Engine/Plugins/Cameras/GameplayCameras/Source/GameplayCameras/Private/Nodes/Framing/PanningFramingCameraNode.cpp b/Engine/Plugins/Cameras/GameplayCameras/Source/GameplayCameras/Private/Nodes/Framing/PanningFramingCameraNode.cpp
--- a/Engine/Plugins/Cameras/GameplayCameras/Source/GameplayCameras/Private/Nodes/Framing/PanningFramingCameraNode.cpp
+++ b/Engine/Plugins/Cameras/GameplayCameras/Source/GameplayCameras/Private/Nodes/Framing/PanningFramingCameraNode.cpp
@@ -80,8 +80,7 @@ void FPanningFramingCameraNodeEvaluator::OnRun(const FCameraNodeEvaluationParams
 		FCameraPose LastShotPose(OutResult.CameraPose);
 		LastShotPose.SetTransform(LastShotTransform);
 
-		APlayerController* PlayerController = Params.EvaluationContext->GetPlayerController();
-		const float AspectRatio = FCameraPoseMath::GetEffectiveAspectRatio(LastShotPose, PlayerController);
+		const float AspectRatio = FCameraPoseMath::GetEffectiveAspectRatio(LastShotPose, Params.EvaluationContext);
 		const FCameraFieldsOfView FOVs(FCameraPoseMath::GetEffectiveFieldsOfView(LastShotPose, AspectRatio));
 
 		// Get yaw/pitch angles for where the target is currently at, and where we want it to be at.
diff --git a/Engine/Plugins/Cameras/GameplayCameras/Source/GameplayCameras/Public/Core/CameraEvaluationContext.h b/Engine/Plugins/Cameras/GameplayCameras/Source/GameplayCameras/Public/Core/CameraEvaluationContext.h
--- a/Engine/Plugins/Cameras/GameplayCameras/Source/GameplayCameras/Public/Core/CameraEvaluationContext.h
+++ b/Engine/Plugins/Cameras/GameplayCameras/Source/GameplayCameras/Public/Core/CameraEvaluationContext.h
@@ -106,6 +106,9 @@ public:
 	/** Gets the children evaluation contexts running inside this context. */
 	TArrayView<const TSharedPtr<FCameraEvaluationContext>> GetChildrenContexts() const { return ChildrenContexts; }
 
+	/** Gets the size of the viewport that cameras inside the context are being viewed through. */
+	FIntPoint GetViewportSize() const;
+
 public:
 
 	/** Gets an evaluation result that will be conditionally overlaid on some camera rigs in this context. */
@@ -149,6 +152,8 @@ protected:
 	virtual void OnActivate(const FCameraEvaluationContextActivateParams& Params) {}
 	virtual void OnDeactivate(const FCameraEvaluationContextDeactivateParams& Params) {}
 
+	FCameraSystemEvaluator* GetCameraSystemEvaluator() const { return CameraSystemEvaluator; }
+
 	void AutoCreateDirectorEvaluator();
 
 #if WITH_EDITOR
@@ -172,18 +177,26 @@ protected:
 	/** The initial result for all camera rigs in this context. */
 	FCameraNodeEvaluationResult InitialResult;
 
+	/** An optional viewport size, if there is no PlayerController, or its viewport should not be used. */
+	TOptional<FIntPoint> OverrideViewportSize;
+
 private:
 
+	/** The camera director evaluator, usually only valid if this context has been activated. */
 	FCameraDirectorEvaluatorStorage DirectorEvaluatorStorage;
 	FCameraDirectorEvaluator* DirectorEvaluator = nullptr;
 
+	/** The camera system that this context has been activated for, or null if the context isn't currently active. */
 	FCameraSystemEvaluator* CameraSystemEvaluator = nullptr;
 
+	/** The parent context, if any. */
 	TWeakPtr<FCameraEvaluationContext> WeakParent;
 
+	/** The children contexts, if any. */
 	using FChildrenContexts = TArray<TSharedPtr<FCameraEvaluationContext>>;
 	FChildrenContexts ChildrenContexts;
 
+	/** Conditional camera data for camera rigs running inside this context. */
 	using FConditionalResults = TMap<ECameraEvaluationDataCondition, FCameraNodeEvaluationResult>;
 	FConditionalResults ConditionalResults;
 
diff --git a/Engine/Plugins/Cameras/GameplayCameras/Source/GameplayCameras/Public/Math/CameraPoseMath.h b/Engine/Plugins/Cameras/GameplayCameras/Source/GameplayCameras/Public/Math/CameraPoseMath.h
--- a/Engine/Plugins/Cameras/GameplayCameras/Source/GameplayCameras/Public/Math/CameraPoseMath.h
+++ b/Engine/Plugins/Cameras/GameplayCameras/Source/GameplayCameras/Public/Math/CameraPoseMath.h
@@ -6,13 +6,14 @@
 #include "Math/MathFwd.h"
 #include "Math/Matrix.h"
 
-class APlayerController;
 struct FCameraFieldsOfView;
 struct FCameraPose;
 
 namespace UE::Cameras
 {
 
+class FCameraEvaluationContext;
+
 /**
  * Simple struct for holding horizontal and vertical fields of view.
  */
@@ -33,13 +34,13 @@ public:
 	static FCameraFieldsOfView GetEffectiveFieldsOfView(const FCameraPose& CameraPose);
 
 	/** Gets both horizontal and vertical effective fields of view, using the given player controller's viewport aspect ratio. */
-	static FCameraFieldsOfView GetEffectiveFieldsOfView(const FCameraPose& CameraPose, const APlayerController* PlayerController);
+	static FCameraFieldsOfView GetEffectiveFieldsOfView(const FCameraPose& CameraPose, TSharedPtr<const FCameraEvaluationContext> EvaluationContext);
 
 	/** Gets both horizontal and vertical effective fields of view, using the given aspect ratio. */
 	static FCameraFieldsOfView GetEffectiveFieldsOfView(const FCameraPose& CameraPose, double AspectRatio);
 
 	/** Gets the aspect ratio of the viewport associated with the given player controller. */
-	static double GetEffectiveAspectRatio(const FCameraPose& CameraPose, const APlayerController* PlayerController);
+	static double GetEffectiveAspectRatio(const FCameraPose& CameraPose, TSharedPtr<const FCameraEvaluationContext> EvaluationContext);
 
 	/** 
 	 * Builds the projection matrix of the given camera pose. 
