Cameras: display proper 3d bounds for framing node targets

diff --git a/Engine/Plugins/Cameras/GameplayCameras/Source/GameplayCameras/Private/Debug/CameraDebugRenderer.cpp b/Engine/Plugins/Cameras/GameplayCameras/Source/GameplayCameras/Private/Debug/CameraDebugRenderer.cpp
--- a/Engine/Plugins/Cameras/GameplayCameras/Source/GameplayCameras/Private/Debug/CameraDebugRenderer.cpp
+++ b/Engine/Plugins/Cameras/GameplayCameras/Source/GameplayCameras/Private/Debug/CameraDebugRenderer.cpp
@@ -462,6 +462,50 @@ void FCameraDebugRenderer::DrawBox(const FVector3d& Center, const FVector3d& Siz
 	}
 }
 
+void FCameraDebugRenderer::DrawBox(const FTransform3d& Transform, const FVector3d& Size, const FLinearColor& LineColor, float LineThickness)
+{
+	if (ULineBatchComponent* LineBatcher = GetDebugLineBatcher())
+	{
+		// Create all box corners in world space.
+		FVector3d TopCorners[4] = {
+			FVector3d(Size.X, Size.Y, Size.Z),
+			FVector3d(-Size.X, Size.Y, Size.Z),
+			FVector3d(-Size.X, -Size.Y, Size.Z),
+			FVector3d(Size.X, -Size.Y, Size.Z)
+		};
+		FVector3d BottomCorners[4] = {
+			FVector3d(Size.X, Size.Y, -Size.Z),
+			FVector3d(-Size.X, Size.Y, -Size.Z),
+			FVector3d(-Size.X, -Size.Y, -Size.Z),
+			FVector3d(Size.X, -Size.Y, -Size.Z)
+		};
+		for (int32 Index = 0; Index < 4; ++Index)
+		{
+			TopCorners[Index] = Transform.TransformVectorNoScale(TopCorners[Index]);
+			BottomCorners[Index] = Transform.TransformVectorNoScale(BottomCorners[Index]);
+		}
+
+		TArray<FBatchedLine> Lines;
+		const FVector3d Center = Transform.GetLocation();
+
+		// Draw the top and bottom squares, and the lines in between.
+		for (int32 Index = 0; Index < 4; ++Index)
+		{
+			Lines.Emplace(
+					Center + TopCorners[Index], Center + TopCorners[(Index + 1) % 4], 
+					LineColor, 0.f, LineThickness, SDPG_Foreground);
+			Lines.Emplace(
+					Center + BottomCorners[Index], Center + BottomCorners[(Index + 1) % 4], 
+					LineColor, 0.f, LineThickness, SDPG_Foreground);
+			Lines.Emplace(
+					Center + TopCorners[Index], Center + BottomCorners[Index], 
+					LineColor, 0.f, LineThickness, SDPG_Foreground);
+		}
+
+		LineBatcher->DrawLines(Lines);
+	}
+}
+
 void FCameraDebugRenderer::DrawSphere(const FVector3d& Center, float Radius, int32 Segments, const FLinearColor& LineColor, float LineThickness)
 {
 	if (ULineBatchComponent* LineBatcher = GetDebugLineBatcher())
diff --git a/Engine/Plugins/Cameras/GameplayCameras/Source/GameplayCameras/Private/Nodes/Framing/BaseFramingCameraNode.cpp b/Engine/Plugins/Cameras/GameplayCameras/Source/GameplayCameras/Private/Nodes/Framing/BaseFramingCameraNode.cpp
--- a/Engine/Plugins/Cameras/GameplayCameras/Source/GameplayCameras/Private/Nodes/Framing/BaseFramingCameraNode.cpp
+++ b/Engine/Plugins/Cameras/GameplayCameras/Source/GameplayCameras/Private/Nodes/Framing/BaseFramingCameraNode.cpp
@@ -683,7 +683,7 @@ void FBaseFramingCameraDebugBlock::OnDebugDraw(const FCameraDebugBlockDrawParams
 
 	for (const FCameraActorComputedTargetInfo& TargetInfo : WorldTargets.TargetInfos)
 	{
-		Renderer.DrawBox(TargetInfo.Transform.GetLocation(), TargetInfo.LocalBounds.BoxExtent, FLinearColor::Yellow, 0.5f);
+		Renderer.DrawBox(TargetInfo.Transform, TargetInfo.LocalBounds.BoxExtent, FColorList::LightSteelBlue, 0.1f);
 	}
 }
 
diff --git a/Engine/Plugins/Cameras/GameplayCameras/Source/GameplayCameras/Public/Debug/CameraDebugRenderer.h b/Engine/Plugins/Cameras/GameplayCameras/Source/GameplayCameras/Public/Debug/CameraDebugRenderer.h
--- a/Engine/Plugins/Cameras/GameplayCameras/Source/GameplayCameras/Public/Debug/CameraDebugRenderer.h
+++ b/Engine/Plugins/Cameras/GameplayCameras/Source/GameplayCameras/Public/Debug/CameraDebugRenderer.h
@@ -115,6 +115,8 @@ public:
 	GAMEPLAYCAMERAS_API void DrawLine(const FVector3d& Start, const FVector3d& End, const FLinearColor& LineColor, float LineThickness = 1.f);
 	/** Draws a 3D box. */
 	GAMEPLAYCAMERAS_API void DrawBox(const FVector3d& Center, const FVector3d& Size, const FLinearColor& LineColor, float LineThickness = 1.f);
+	/** Draws a 3D box. */
+	GAMEPLAYCAMERAS_API void DrawBox(const FTransform3d& Transform, const FVector3d& Size, const FLinearColor& LineColor, float LineThickness = 1.f);
 	/** Draws a 3D sphere. */
 	GAMEPLAYCAMERAS_API void DrawSphere(const FVector3d& Center, float Radius, int32 Segments, const FLinearColor& LineColor, float LineThickness = 1.f);
 	/** Draws a 3D arrow. */
