Cameras: make the effective dead zone at least "epsilon" big, remove false positive assert

diff --git a/Engine/Plugins/Cameras/GameplayCameras/Source/GameplayCameras/Private/Math/CameraFramingZoneMath.cpp b/Engine/Plugins/Cameras/GameplayCameras/Source/GameplayCameras/Private/Math/CameraFramingZoneMath.cpp
--- a/Engine/Plugins/Cameras/GameplayCameras/Source/GameplayCameras/Private/Math/CameraFramingZoneMath.cpp
+++ b/Engine/Plugins/Cameras/GameplayCameras/Source/GameplayCameras/Private/Math/CameraFramingZoneMath.cpp
@@ -100,6 +100,14 @@ void FFramingZone::ClampBounds(const FVector2d& MustContain)
 	BottomBound = FMath::Clamp(BottomBound, MustContain.Y, 1.0);
 }
 
+void FFramingZone::ClampBounds(const FVector2d& MustContain, double Margin)
+{
+	LeftBound = FMath::Clamp(LeftBound, 0.0, MustContain.X - Margin);
+	TopBound = FMath::Clamp(TopBound, 0.0, MustContain.Y - Margin);
+	RightBound = FMath::Clamp(RightBound, MustContain.X + Margin, 1.0);
+	BottomBound = FMath::Clamp(BottomBound, MustContain.Y + Margin, 1.0);
+}
+
 void FFramingZone::ClampBounds(const FFramingZone& MustContain)
 {
 	LeftBound = FMath::Clamp(LeftBound, 0.0, MustContain.LeftBound);
diff --git a/Engine/Plugins/Cameras/GameplayCameras/Source/GameplayCameras/Private/Nodes/Framing/BaseFramingCameraNode.cpp b/Engine/Plugins/Cameras/GameplayCameras/Source/GameplayCameras/Private/Nodes/Framing/BaseFramingCameraNode.cpp
--- a/Engine/Plugins/Cameras/GameplayCameras/Source/GameplayCameras/Private/Nodes/Framing/BaseFramingCameraNode.cpp
+++ b/Engine/Plugins/Cameras/GameplayCameras/Source/GameplayCameras/Private/Nodes/Framing/BaseFramingCameraNode.cpp
@@ -454,7 +454,7 @@ FFramingZone FBaseFramingCameraNodeEvaluator::ComputeEffectiveDeadZone()
 	EffectiveDeadZone.RightBound -= RelativeTargetBounds.RightBound;
 	EffectiveDeadZone.BottomBound -= RelativeTargetBounds.BottomBound;
 
-	EffectiveDeadZone.ClampBounds(State.IdealTarget);
+	EffectiveDeadZone.ClampBounds(State.IdealTarget, UE_DOUBLE_KINDA_SMALL_NUMBER);
 	
 	return EffectiveDeadZone;
 }
@@ -490,7 +490,6 @@ void FBaseFramingCameraNodeEvaluator::ComputeDesiredState(float DeltaTime)
 	const double ReframingSpeed = State.ReframeDamper.GetX0Derivative();
 	if (!State.bIsReframingTarget || (DistanceToGo <= GFramingIdealReachedEpsilon && ReframingSpeed <= GFramingIdealReachedEpsilon))
 	{
-		ensure(State.TargetFramingState == ETargetFramingState::InDeadZone);
 		Desired.ScreenTarget = State.ScreenTarget;
 		Desired.FramingCorrection = FVector2d::ZeroVector;
 		Desired.bHasCorrection = false;
diff --git a/Engine/Plugins/Cameras/GameplayCameras/Source/GameplayCameras/Public/Math/CameraFramingZoneMath.h b/Engine/Plugins/Cameras/GameplayCameras/Source/GameplayCameras/Public/Math/CameraFramingZoneMath.h
--- a/Engine/Plugins/Cameras/GameplayCameras/Source/GameplayCameras/Public/Math/CameraFramingZoneMath.h
+++ b/Engine/Plugins/Cameras/GameplayCameras/Source/GameplayCameras/Public/Math/CameraFramingZoneMath.h
@@ -53,6 +53,12 @@ struct FFramingZone
 	 */
 	void ClampBounds(const FVector2d& MustContain);
 
+	/**
+	 * Makes sure all the bounds have valid values between 0 and 1, and that the
+	 * enclosed rectangle contains the given target point, with a uniform margin.
+	 */
+	void ClampBounds(const FVector2d& MustContain, double Margin);
+
 	/**
 	 * Makes sure all the bounds have valid values between 0 and 1, and that the
 	 * enclosed rectangle contains the given inner rectangle.
