Cameras: use the active editor viewport's size for aspect ratio when previewing camera rigs in editor

diff --git a/Engine/Plugins/Cameras/GameplayCameras/Source/GameplayCameras/Private/GameFramework/GameplayCameraComponent.cpp b/Engine/Plugins/Cameras/GameplayCameras/Source/GameplayCameras/Private/GameFramework/GameplayCameraComponent.cpp
--- a/Engine/Plugins/Cameras/GameplayCameras/Source/GameplayCameras/Private/GameFramework/GameplayCameraComponent.cpp
+++ b/Engine/Plugins/Cameras/GameplayCameras/Source/GameplayCameras/Private/GameFramework/GameplayCameraComponent.cpp
@@ -21,6 +21,11 @@
 #include "SceneView.h"
 #include "ShowFlags.h"
 
+#if WITH_EDITOR
+#include "Editor.h"
+#include "LevelEditorViewport.h"
+#endif
+
 #include UE_INLINE_GENERATED_CPP_BY_NAME(GameplayCameraComponent)
 
 #define LOCTEXT_NAMESPACE "GameplayCameraComponent"
@@ -381,6 +386,10 @@ void UGameplayCameraComponent::UpdateCameraEvaluationContext(bool bForceApplyPar
 
 	const bool bApplyDrivenParametersOnly = !bForceApplyParameterOverrides;
 	CameraReference.ApplyParameterOverrides(InitialResult, bApplyDrivenParametersOnly);
+
+#if WITH_EDITOR
+	EvaluationContext->UpdateForEditorPreview();
+#endif  // WITH_EDITOR
 }
 
 void UGameplayCameraComponent::UpdateOutputCameraComponent()
@@ -633,6 +642,27 @@ namespace UE::Cameras
 
 UE_DEFINE_CAMERA_EVALUATION_CONTEXT(FGameplayCameraComponentEvaluationContext)
 
+#if WITH_EDITOR
+
+void FGameplayCameraComponentEvaluationContext::UpdateForEditorPreview()
+{
+	FCameraSystemEvaluator* ActiveEvaluator = GetCameraSystemEvaluator();
+	if (ActiveEvaluator && ActiveEvaluator->GetRole() == ECameraSystemEvaluatorRole::EditorPreview)
+	{
+		if (GCurrentLevelEditingViewportClient && GCurrentLevelEditingViewportClient->Viewport)
+		{
+			FIntPoint ViewportSize = GCurrentLevelEditingViewportClient->Viewport->GetSizeXY();
+			OverrideViewportSize = ViewportSize;
+		}
+		else
+		{
+			OverrideViewportSize.Reset();
+		}
+	}
+}
+
+#endif  // WITH_EDITOR
+
 }  // namespace UE::Cameras
 
 #undef LOCTEXT_NAMESPACE
diff --git a/Engine/Plugins/Cameras/GameplayCameras/Source/GameplayCameras/Public/GameFramework/GameplayCameraComponent.h b/Engine/Plugins/Cameras/GameplayCameras/Source/GameplayCameras/Public/GameFramework/GameplayCameraComponent.h
--- a/Engine/Plugins/Cameras/GameplayCameras/Source/GameplayCameras/Public/GameFramework/GameplayCameraComponent.h
+++ b/Engine/Plugins/Cameras/GameplayCameras/Source/GameplayCameras/Public/GameFramework/GameplayCameraComponent.h
@@ -202,6 +202,14 @@ namespace UE::Cameras
 class FGameplayCameraComponentEvaluationContext : public FCameraEvaluationContext
 {
 	UE_DECLARE_CAMERA_EVALUATION_CONTEXT(GAMEPLAYCAMERAS_API, FGameplayCameraComponentEvaluationContext)
+
+#if WITH_EDITOR
+
+public:
+
+	void UpdateForEditorPreview();
+
+#endif  // WITH_EDITOR
 };
 
 }  // namespace UE::Cameras
