Cameras: move various blend stack implementations into separate files

diff --git a/Engine/Plugins/Cameras/GameplayCameras/Source/GameplayCameras/Private/Core/BlendStackCameraNode.cpp b/Engine/Plugins/Cameras/GameplayCameras/Source/GameplayCameras/Private/Core/BlendStackCameraNode.cpp
--- a/Engine/Plugins/Cameras/GameplayCameras/Source/GameplayCameras/Private/Core/BlendStackCameraNode.cpp
+++ b/Engine/Plugins/Cameras/GameplayCameras/Source/GameplayCameras/Private/Core/BlendStackCameraNode.cpp
@@ -2,16 +2,15 @@
 
 #include "Core/BlendStackCameraNode.h"
 
-#include "Algo/AnyOf.h"
-#include "Core/BlendCameraNode.h"
 #include "Core/BlendStackCameraRigEvent.h"
 #include "Core/BlendStackRootCameraNode.h"
 #include "Core/CameraAsset.h"
 #include "Core/CameraEvaluationContext.h"
 #include "Core/CameraRigAsset.h"
 #include "Core/CameraRigAssetReference.h"
-#include "Core/CameraRigCombinationRegistry.h"
 #include "Core/CameraSystemEvaluator.h"
+#include "Core/PersistentBlendStackCameraNode.h"
+#include "Core/TransientBlendStackCameraNode.h"
 #include "Debug/CameraDebugBlockBuilder.h"
 #include "Debug/CameraDebugRenderer.h"
 #include "Debug/CameraNodeEvaluationResultDebugBlock.h"
@@ -59,8 +58,6 @@ static FAutoConsoleVariableRef CVarGameplayCamerasDebugBlendStackShowVariableIDs
 	TEXT(""));
 
 UE_DEFINE_CAMERA_NODE_EVALUATOR(FBlendStackCameraNodeEvaluator)
-UE_DEFINE_CAMERA_NODE_EVALUATOR(FTransientBlendStackCameraNodeEvaluator)
-UE_DEFINE_CAMERA_NODE_EVALUATOR(FPersistentBlendStackCameraNodeEvaluator)
 
 FBlendStackCameraNodeEvaluator::FBlendStackCameraNodeEvaluator()
 {
@@ -500,814 +497,6 @@ void FBlendStackCameraNodeEvaluator::OnPostBuildAsset(const FGameplayCameraAsset
 
 #endif  // WITH_EDITOR
 
-void FTransientBlendStackCameraNodeEvaluator::Push(const FBlendStackCameraPushParams& Params)
-{
-	bool bSearchedForTransition = false;
-	const UCameraRigTransition* Transition = nullptr;
-
-	if (!Entries.IsEmpty())
-	{
-		FCameraRigEntry& TopEntry(Entries.Top());
-		if (!TopEntry.bIsFrozen 
-				&& TopEntry.EvaluationContext == Params.EvaluationContext)
-		{
-			// Don't push anything if what is being requested is already the active 
-			// camera rig.
-			if (!Params.bForcePush && TopEntry.CameraRig == Params.CameraRig)
-			{
-				return;
-			}
-
-			// See if we can merge the new camera rig onto the active camera rig.
-			const ECameraRigMergingEligibility Eligibility = TopEntry.RootEvaluator->CompareCameraRigForMerging(Params.CameraRig);
-
-			if (!Params.bForcePush && Eligibility == ECameraRigMergingEligibility::Active)
-			{
-				// This camera rig is already the active one on the merged stack.
-				return;
-			}
-
-			if (Eligibility == ECameraRigMergingEligibility::EligibleForMerge)
-			{
-				// This camera rig can be merged with the one current running. However, we
-				// only do it if the transition explicitly allows it.
-				bSearchedForTransition = true;
-				Transition = FindTransition(Params);
-
-				if (Transition && Transition->bAllowCameraRigMerging)
-				{
-					PushMergedEntry(Params, Transition);
-					return;
-				}
-			}
-		}
-	}
-
-	// It's a legitimate new entry in the blend stack.
-	if (!bSearchedForTransition)
-	{
-		bSearchedForTransition = true;
-		Transition = FindTransition(Params);
-	}
-
-	PushNewEntry(Params, Transition);
-}
-
-void FTransientBlendStackCameraNodeEvaluator::PushNewEntry(const FBlendStackCameraPushParams& Params, const UCameraRigTransition* Transition)
-{
-	// Create the new root node to wrap the new camera rig's root node, and the specific
-	// blend node for this transition.
-	// We need to const-cast here to be able to use our own blend stack node as the outer
-	// of the new node.
-	const UCameraRigTransition* UsedTransition = nullptr;
-	UObject* Outer = const_cast<UObject*>((UObject*)GetCameraNode());
-	UBlendStackRootCameraNode* EntryRootNode = NewObject<UBlendStackRootCameraNode>(Outer, NAME_None);
-	{
-		EntryRootNode->RootNode = Params.CameraRig->RootNode;
-
-		// Find a transition and use its blend. If no transition is found,
-		// make a camera cut transition.
-		UBlendCameraNode* ModeBlend = nullptr;
-		if (Transition)
-		{
-			ModeBlend = Transition->Blend;
-			UsedTransition = Transition;
-		}
-		if (!ModeBlend)
-		{
-			ModeBlend = NewObject<UPopBlendCameraNode>(EntryRootNode, NAME_None);
-		}
-		EntryRootNode->Blend = ModeBlend;
-	}
-
-	// Make the new stack entry, and use its storage buffer to build the tree of evaluators.
-	FCameraRigEntry NewEntry;
-	const bool bInitialized = InitializeEntry(
-			NewEntry, 
-			Params.CameraRig,
-			Params.EvaluationContext,
-			EntryRootNode,
-			true);
-	if (!bInitialized)
-	{
-		return;
-	}
-	
-	// Take a snapshot of the initial result for having base values during pre-blending.
-	FCameraRigEntryExtraInfo NewExtraInfo;
-	{
-		FCameraNodeEvaluationResult& UnblendedResult = NewExtraInfo.UnblendedResult;
-
-		if (Params.CameraRig)
-		{
-			const FCameraRigAllocationInfo& AllocationInfo = Params.CameraRig->AllocationInfo;
-			UnblendedResult.VariableTable.Initialize(AllocationInfo.VariableTableInfo);
-			UnblendedResult.ContextDataTable.Initialize(AllocationInfo.ContextDataTableInfo);
-		}
-
-		UnblendedResult.OverrideAll(NewEntry.Result, true);
-	}
-
-#if WITH_EDITOR
-	// Listen to changes to the packages inside which this camera rig is defined. We will hot-reload the
-	// camera node evaluators for this camera rig when we detect changes.
-	AddPackageListeners(NewEntry);
-#endif  // WITH_EDITOR
-
-	// Important: we need to move the new entry here because copying evaluator storage
-	// is disabled.
-	Entries.Add(MoveTemp(NewEntry));
-	EntryExtraInfos.Add(MoveTemp(NewExtraInfo));
-
-	if (OnCameraRigEventDelegate.IsBound())
-	{
-		BroadcastCameraRigEvent(EBlendStackCameraRigEventType::Pushed, Entries.Last(), UsedTransition);
-	}
-}
-
-void FTransientBlendStackCameraNodeEvaluator::PushMergedEntry(const FBlendStackCameraPushParams& PushParams, const UCameraRigTransition* Transition)
-{
-	const UBlendCameraNode* Blend = Transition ? Transition->Blend : nullptr;
-
-	FCameraRigEntry& TopEntry = Entries.Top();
-	FCameraNodeEvaluatorBuilder Builder(TopEntry.EvaluatorStorage);
-	FCameraNodeEvaluatorBuildParams BuildParams(Builder);
-	TopEntry.RootEvaluator->MergeCameraRig(BuildParams, PushParams.CameraRig, Blend);
-
-	// Merging a camera rig the first time changes the evaluator tree by removing any prefab nodes
-	// near the root... so we need to rebuild our hierarchy cache.
-	TopEntry.EvaluatorHierarchy.Build(TopEntry.RootEvaluator);
-
-	// Swap out the camera rig registered as "active" for this entry.
-#if WITH_EDITOR
-	RemoveListenedPackages(TopEntry);
-#endif
-	{
-		TopEntry.CameraRig = PushParams.CameraRig;
-	}
-#if WITH_EDITOR
-	AddPackageListeners(TopEntry);
-#endif
-}
-
-void FTransientBlendStackCameraNodeEvaluator::Freeze(const FBlendStackCameraFreezeParams& Params)
-{
-	for (FCameraRigEntry& Entry : Entries)
-	{
-		if (!Entry.bIsFrozen && 
-				Entry.CameraRig == Params.CameraRig &&
-				Entry.EvaluationContext == Params.EvaluationContext)
-		{
-			FreezeEntry(Entry);
-		}
-	}
-}
-
-void FTransientBlendStackCameraNodeEvaluator::FreezeAll(TSharedPtr<FCameraEvaluationContext> EvaluationContext)
-{
-	for (FCameraRigEntry& Entry : Entries)
-	{
-		if (!Entry.bIsFrozen && Entry.EvaluationContext == EvaluationContext)
-		{
-			FreezeEntry(Entry);
-		}
-	}
-}
-
-void FTransientBlendStackCameraNodeEvaluator::OnRun(const FCameraNodeEvaluationParams& Params, FCameraNodeEvaluationResult& OutResult)
-{
-	ensure(Entries.Num() == EntryExtraInfos.Num());
-
-	// Validate our entries and resolve evaluation context weak pointers.
-	TArray<FResolvedEntry> ResolvedEntries;
-	ResolveEntries(ResolvedEntries);
-
-	// Gather parameters to pre-blend, and evaluate blend nodes.
-	InternalPreBlendPrepare(ResolvedEntries, Params, OutResult);
-
-	// Blend input variables.
-	InternalPreBlendExecute(ResolvedEntries, Params, OutResult);
-
-	// Run the root nodes. They will use the pre-blended inputs from the last step.
-	// Frozen entries are skipped, since they only ever use the last result they produced.
-	InternalUpdate(ResolvedEntries, Params, OutResult);
-
-	// Now blend all the results, keeping track of blends that have reached 100% so
-	// that we can remove any camera rigs below (since they would have been completely
-	// blended out by that).
-	InternalPostBlendExecute(ResolvedEntries, Params, OutResult);
-
-	// Tidy up.
-	OnRunFinished(OutResult);
-	InternalRunFinished(OutResult);
-}
-
-void FTransientBlendStackCameraNodeEvaluator::InternalPreBlendPrepare(TArrayView<FResolvedEntry> ResolvedEntries, const FCameraNodeEvaluationParams& Params, FCameraNodeEvaluationResult& OutResult)
-{
-	constexpr ECameraVariableTableFilter VariableTableFilter = ECameraVariableTableFilter::ChangedOnly;
-	constexpr ECameraContextDataTableFilter ContextDataTableFilter = ECameraContextDataTableFilter::ChangedOnly;
-
-	for (FResolvedEntry& ResolvedEntry : ResolvedEntries)
-	{
-		FCameraRigEntry& Entry(ResolvedEntry.Entry);
-		FCameraRigEntryExtraInfo& EntryExtraInfo(EntryExtraInfos[ResolvedEntry.EntryIndex]);
-
-		if (UNLIKELY(Entry.bIsFrozen))
-		{
-			continue;
-		}
-
-		FCameraNodeEvaluationParams CurParams(Params);
-		CurParams.EvaluationContext = ResolvedEntry.Context;
-		CurParams.bIsFirstFrame = Entry.bIsFirstFrame;
-
-		FCameraNodeEvaluationResult& UnblendedResult(EntryExtraInfo.UnblendedResult);
-
-		// Reset frame flags to match what happend to the Entry.Result in ResolveEntries.
-		UnblendedResult.ResetFrameFlags();
-
-		// Start with the input given to us.
-		UnblendedResult.VariableTable.OverrideAll(OutResult.VariableTable);
-
-		// Override it with whatever the evaluation context has set on its result this frame.
-		// Evaluation contexts may have private variables we need to pass along, such as when rig parameter
-		// overrides have been set on them, so include private variables in the filter.
-		const FCameraNodeEvaluationResult& ContextResult(ResolvedEntry.Context->GetInitialResult());
-		UnblendedResult.VariableTable.Override(ContextResult.VariableTable, VariableTableFilter);
-		UnblendedResult.ContextDataTable.Override(ContextResult.ContextDataTable, ContextDataTableFilter);
-		if (ResolvedEntry.bIsActiveEntry)
-		{
-			if (const FCameraNodeEvaluationResult* ActiveOnlyResult = ResolvedEntry.Context->GetConditionalResult(ECameraEvaluationDataCondition::ActiveCameraRig))
-			{
-				UnblendedResult.VariableTable.Override(ActiveOnlyResult->VariableTable, VariableTableFilter);
-				UnblendedResult.ContextDataTable.Override(ActiveOnlyResult->ContextDataTable, ContextDataTableFilter);
-			}
-		}
-
-		// Gather input parameters if needed (and remember if it was indeed needed).
-		if (!EntryExtraInfo.bInputRunThisFrame)
-		{
-			bool bHasPreBlendedParameters = false;
-			FCameraBlendedParameterUpdateParams InputParams(CurParams, UnblendedResult.CameraPose);
-			FCameraBlendedParameterUpdateResult InputResult(UnblendedResult.VariableTable);
-
-			Entry.EvaluatorHierarchy.ForEachEvaluator(ECameraNodeEvaluatorFlags::NeedsParameterUpdate,
-					[&bHasPreBlendedParameters, &InputParams, &InputResult](FCameraNodeEvaluator* ParameterEvaluator)
-					{
-						ParameterEvaluator->UpdateParameters(InputParams, InputResult);
-						bHasPreBlendedParameters = true;
-					});
-
-			EntryExtraInfo.bHasPreBlendedParameters = bHasPreBlendedParameters;
-			EntryExtraInfo.bInputRunThisFrame = true;
-		}
-
-		// Run blends.
-		// Note that we pass last frame's camera pose to the Run() method. This may change.
-		// Blends aren't expected to use the camera pose to do any logic until BlendResults().
-		if (!EntryExtraInfo.bBlendRunThisFrame)
-		{
-			FBlendCameraNodeEvaluator* EntryBlendEvaluator = Entry.RootEvaluator->GetBlendEvaluator();
-			if (EntryBlendEvaluator)
-			{
-				EntryBlendEvaluator->Run(CurParams, UnblendedResult);
-			}
-
-			EntryExtraInfo.bBlendRunThisFrame = true;
-		}
-
-		// Assign the new unblended result back to the entry result. The goal is that:
-		//
-		// 1) We only write values from the context result when they have been written to.
-		//    That is, we don't always set all values every frame. When "left alone", a camera rig
-		//    should retain whatever values it was last set with. For instance, if a camera rig 
-		//    received special values when it was the active camera rig, we want to keep those 
-		//    values when it starts blending out, instead of immediately overwriting them with
-		//    the "common" values (i.e. the values from the initial result).
-		//
-		// 2) However, we do want to write all the values into the PreBlendVariableTable every
-		//    frame, so that we can have them blend together. The blended values are then reassigned
-		//    to Entry.Result in InternalPreBlendExecute, which means we lose each camera rig's
-		//    "intended" values. Getting them from the evaluation context each frame would violate
-		//    previous item (1) so that's why we use the UnblendedResult as a snapshot.
-		//
-		FCameraNodeEvaluationResult& CurResult(Entry.Result);
-		CurResult.OverrideAll(UnblendedResult, true);
-	}
-}
-
-void FTransientBlendStackCameraNodeEvaluator::InternalPreBlendExecute(TArrayView<FResolvedEntry> ResolvedEntries, const FCameraNodeEvaluationParams& Params, FCameraNodeEvaluationResult& OutResult)
-{
-	// Blend all the camera rigs' input variables (include private variables such as camera rig parameters).
-	// The result of this pre-blending goes into PreBlendVariableTable, which will then have a big mix
-	// of all the variables.
-	PreBlendVariableTable.ClearAllWrittenThisFrameFlags();
-
-	if (!Algo::AnyOf(EntryExtraInfos, [](const FCameraRigEntryExtraInfo& Item) { return Item.bHasPreBlendedParameters; }))
-	{
-		return;
-	}
-
-	constexpr ECameraVariableTableFilter VariableTableFilter = ECameraVariableTableFilter::None;
-
-	for (FResolvedEntry& ResolvedEntry : ResolvedEntries)
-	{
-		FCameraRigEntry& Entry(ResolvedEntry.Entry);
-		FCameraRigEntryExtraInfo& EntryExtraInfo(EntryExtraInfos[ResolvedEntry.EntryIndex]);
-
-		FCameraNodeEvaluationResult& UnblendedResult(EntryExtraInfo.UnblendedResult);
-
-		if (!Entry.bIsFrozen)
-		{
-			FCameraNodeEvaluationParams CurParams(Params);
-			CurParams.EvaluationContext = ResolvedEntry.Context;
-			CurParams.bIsFirstFrame = Entry.bIsFirstFrame;
-
-			FCameraNodePreBlendParams PreBlendParams(CurParams, UnblendedResult.CameraPose, UnblendedResult.VariableTable);
-			PreBlendParams.VariableTableFilter = VariableTableFilter;
-
-			FCameraNodePreBlendResult PreBlendResult(PreBlendVariableTable);
-
-			FBlendCameraNodeEvaluator* EntryBlendEvaluator = Entry.RootEvaluator->GetBlendEvaluator();
-			if (EntryBlendEvaluator)
-			{
-				EntryBlendEvaluator->BlendParameters(PreBlendParams, PreBlendResult);
-				EntryExtraInfo.bIsPreBlendFull = (PreBlendResult.bIsBlendFinished && PreBlendResult.bIsBlendFull);
-			}
-			else
-			{
-				PreBlendVariableTable.Override(UnblendedResult.VariableTable, PreBlendParams.VariableTableFilter);
-				EntryExtraInfo.bIsPreBlendFull = true;
-			}
-		}
-		else
-		{
-			// Frozen entries still contribute to the blend using their last evaluated values.
-			PreBlendVariableTable.Override(UnblendedResult.VariableTable, VariableTableFilter);
-		}
-	}
-
-	// Write the values back to each entry table, so that each of these camera rigs will run with
-	// the pre-blended values. We limit this writing to the variables each of them knows, since
-	// there's no need to add entries they don't use to their variable table.
-	for (FResolvedEntry& ResolvedEntry : ResolvedEntries)
-	{
-		FCameraRigEntry& Entry(ResolvedEntry.Entry);
-		FCameraRigEntryExtraInfo& EntryExtraInfo(EntryExtraInfos[ResolvedEntry.EntryIndex]);
-
-		if (!Entry.bIsFrozen)
-		{
-			FCameraNodeEvaluationResult& UnblendedResult(EntryExtraInfo.UnblendedResult);
-			FCameraNodeEvaluationResult& CurResult(Entry.Result);
-
-			Entry.Result.VariableTable.Override(PreBlendVariableTable, ECameraVariableTableFilter::KnownOnly);
-		}
-	}
-}
-
-void FTransientBlendStackCameraNodeEvaluator::InternalUpdate(TArrayView<FResolvedEntry> ResolvedEntries, const FCameraNodeEvaluationParams& Params, FCameraNodeEvaluationResult& OutResult)
-{
-	for (FResolvedEntry& ResolvedEntry : ResolvedEntries)
-	{
-		FCameraRigEntry& Entry(ResolvedEntry.Entry);
-		FCameraRigEntryExtraInfo& EntryExtraInfo(EntryExtraInfos[ResolvedEntry.EntryIndex]);
-
-		if (UNLIKELY(Entry.bIsFrozen))
-		{
-			continue;
-		}
-
-		FCameraNodeEvaluationResult& CurResult(Entry.Result);
-
-		FCameraNodeEvaluationParams CurParams(Params);
-		CurParams.EvaluationContext = ResolvedEntry.Context;
-		CurParams.bIsFirstFrame = Entry.bIsFirstFrame;
-
-		// Start with the input given to us.
-		CurResult.Reset();
-		CurResult.CameraPose = OutResult.CameraPose;
-		CurResult.CameraRigJoints.OverrideAll(OutResult.CameraRigJoints);
-		CurResult.PostProcessSettings.OverrideAll(OutResult.PostProcessSettings);
-
-		// Override it with whatever the evaluation context has set on its result.
-		const FCameraNodeEvaluationResult& ContextResult(ResolvedEntry.Context->GetInitialResult());
-		CurResult.CameraPose.OverrideChanged(ContextResult.CameraPose);
-		CurResult.bIsCameraCut = OutResult.bIsCameraCut || ContextResult.bIsCameraCut || Entry.bForceCameraCut;
-		if (ResolvedEntry.bIsActiveEntry)
-		{
-			if (const FCameraNodeEvaluationResult* ActiveOnlyResult = ResolvedEntry.Context->GetConditionalResult(ECameraEvaluationDataCondition::ActiveCameraRig))
-			{
-				CurResult.CameraPose.OverrideChanged(ActiveOnlyResult->CameraPose);
-				CurResult.bIsCameraCut |= ActiveOnlyResult->bIsCameraCut;
-			}
-		}
-		
-		CurResult.bIsValid = true;
-
-#if WITH_EDITOR || UE_GAMEPLAY_CAMERAS_DEBUG
-		CurResult.AddCameraPoseTrailPointIfNeeded(ContextResult.CameraPose.GetLocation());
-#endif  // WITH_EDITOR || UE_GAMEPLAY_CAMERAS_DEBUG
-
-		// Run the camera rig's root node.
-		FCameraNodeEvaluator* RootEvaluator = Entry.RootEvaluator->GetRootEvaluator();
-		if (RootEvaluator)
-		{
-			RootEvaluator->Run(CurParams, CurResult);
-		}
-	}
-}
-
-void FTransientBlendStackCameraNodeEvaluator::InternalPostBlendExecute(TArrayView<FResolvedEntry> ResolvedEntries, const FCameraNodeEvaluationParams& Params, FCameraNodeEvaluationResult& OutResult)
-{
-	int32 PopEntriesBelow = INDEX_NONE;
-	for (FResolvedEntry& ResolvedEntry : ResolvedEntries)
-	{
-		FCameraRigEntry& Entry(ResolvedEntry.Entry);
-		FCameraRigEntryExtraInfo& EntryExtraInfo(EntryExtraInfos[ResolvedEntry.EntryIndex]);
-
-		FCameraNodeEvaluationResult& CurResult(Entry.Result);
-
-		if (!Entry.bIsFrozen)
-		{
-			FCameraNodeEvaluationParams CurParams(Params);
-			CurParams.EvaluationContext = ResolvedEntry.Context;
-			CurParams.bIsFirstFrame = Entry.bIsFirstFrame;
-			FCameraNodeBlendParams BlendParams(CurParams, CurResult);
-
-			FCameraNodeBlendResult BlendResult(OutResult);
-
-			FBlendCameraNodeEvaluator* EntryBlendEvaluator = Entry.RootEvaluator->GetBlendEvaluator();
-			if (EntryBlendEvaluator)
-			{
-				EntryBlendEvaluator->BlendResults(BlendParams, BlendResult);
-
-				if (BlendResult.bIsBlendFull && BlendResult.bIsBlendFinished)
-				{
-					PopEntriesBelow = ResolvedEntry.EntryIndex;
-				}
-			}
-			else
-			{
-				OutResult.OverrideAll(CurResult);
-
-				PopEntriesBelow = ResolvedEntry.EntryIndex;
-			}
-		}
-		else
-		{
-			OutResult.OverrideAll(CurResult);
-
-			PopEntriesBelow = ResolvedEntry.EntryIndex;
-		}
-		// else, merged entries only contribute to the pre-blend.
-	}
-
-	// Pop out camera rigs that have been blended out.
-	const UBlendStackCameraNode* BlendStackNode = GetCameraNodeAs<UBlendStackCameraNode>();
-	if (BlendStackNode->BlendStackType == ECameraBlendStackType::IsolatedTransient && PopEntriesBelow != INDEX_NONE)
-	{
-		PopEntries(PopEntriesBelow);
-
-		EntryExtraInfos.RemoveAt(0, PopEntriesBelow);
-	}
-}
-
-void FTransientBlendStackCameraNodeEvaluator::InternalRunFinished(FCameraNodeEvaluationResult& OutResult)
-{
-	for (FCameraRigEntryExtraInfo& ExtraInfo : EntryExtraInfos)
-	{
-		ExtraInfo.bInputRunThisFrame = false;
-		ExtraInfo.bBlendRunThisFrame = false;
-		ExtraInfo.bHasPreBlendedParameters = false;
-		ExtraInfo.bIsPreBlendFull = false;
-	}
-}
-
-void FTransientBlendStackCameraNodeEvaluator::OnSerialize(const FCameraNodeEvaluatorSerializeParams& Params, FArchive& Ar)
-{
-	Super::OnSerialize(Params, Ar);
-
-	for (FCameraRigEntryExtraInfo& ExtraInfo : EntryExtraInfos)
-	{
-		Ar << ExtraInfo.bInputRunThisFrame;
-		Ar << ExtraInfo.bBlendRunThisFrame;
-		Ar << ExtraInfo.bHasPreBlendedParameters;
-		Ar << ExtraInfo.bIsPreBlendFull;
-	}
-}
-
-const UCameraRigTransition* FTransientBlendStackCameraNodeEvaluator::FindTransition(const FBlendStackCameraPushParams& Params) const
-{
-	// If we are forced to use a specific transition, our search is over.
-	if (Params.TransitionOverride)
-	{
-		return Params.TransitionOverride;
-	}
-
-	// Find a transition that works for blending towards ToCameraRig.
-	// If the stack isn't empty, we need to find a transition that works between the previous and 
-	// next camera rigs. If the stack is empty, we blend the new camera rig in from nothing if
-	// appropriate.
-	if (!Entries.IsEmpty())
-	{
-		// Grab information about the new entry to push.
-		TSharedPtr<const FCameraEvaluationContext> ToContext = Params.EvaluationContext;
-		const UCameraAsset* ToCameraAsset = ToContext ? ToContext->GetCameraAsset() : nullptr;
-
-		// Grab information about the top entry (i.e. the currently active camera rig).
-		const FCameraRigEntry& TopEntry = Entries.Top();
-		TSharedPtr<const FCameraEvaluationContext> FromContext = TopEntry.EvaluationContext.Pin();
-		const UCameraAsset* FromCameraAsset = FromContext ? FromContext->GetCameraAsset() : nullptr;
-
-		// If the new or current top entries are a combination, look for transitions on all 
-		// their combined camera rigs.
-		TArray<const UCameraRigAsset*> ToCombinedCameraRigs;
-		UCombinedCameraRigsCameraNode::GetAllCombinationCameraRigs(Params.CameraRig, ToCombinedCameraRigs);
-
-		TArray<const UCameraRigAsset*> FromCombinedCameraRigs;
-		UCombinedCameraRigsCameraNode::GetAllCombinationCameraRigs(TopEntry.CameraRig, FromCombinedCameraRigs);
-
-		const bool bFromFrozen = TopEntry.bIsFrozen;
-		const UCameraRigTransition* TransitionToUse = nullptr;
-
-		// Start by looking at exit transitions on the last active (top) camera rig.
-		for (const UCameraRigAsset* FromCameraRig : FromCombinedCameraRigs)
-		{
-			if (!FromCameraRig->ExitTransitions.IsEmpty())
-			{
-				// Look for exit transitions on the last active camera rig itself.
-				for (const UCameraRigAsset* ToCameraRig : ToCombinedCameraRigs)
-				{
-					TransitionToUse = FindTransition(
-							FromCameraRig->ExitTransitions,
-							FromCameraRig, FromCameraAsset, bFromFrozen,
-							ToCameraRig, ToCameraAsset);
-					if (TransitionToUse)
-					{
-						return TransitionToUse;
-					}
-				}
-			}
-		}
-		for (const UCameraRigAsset* FromCameraRig : FromCombinedCameraRigs)
-		{
-			if (FromCameraAsset && !FromCameraAsset->GetExitTransitions().IsEmpty())
-			{
-				// Look for exit transitions on its parent camera asset.
-				for (const UCameraRigAsset* ToCameraRig : ToCombinedCameraRigs)
-				{
-					TransitionToUse = FindTransition(
-							FromCameraAsset->GetExitTransitions(),
-							FromCameraRig, FromCameraAsset, bFromFrozen,
-							ToCameraRig, ToCameraAsset);
-					if (TransitionToUse)
-					{
-						return TransitionToUse;
-					}
-				}
-			}
-		}
-
-		// Now look at enter transitions on the new camera rig.
-		for (const UCameraRigAsset* ToCameraRig : ToCombinedCameraRigs)
-		{
-			if (!ToCameraRig->EnterTransitions.IsEmpty())
-			{
-				// Look for enter transitions on the new camera rig itself.
-				for (const UCameraRigAsset* FromCameraRig : FromCombinedCameraRigs)
-				{
-					TransitionToUse = FindTransition(
-							ToCameraRig->EnterTransitions,
-							FromCameraRig, FromCameraAsset, bFromFrozen,
-							ToCameraRig, ToCameraAsset);
-					if (TransitionToUse)
-					{
-						return TransitionToUse;
-					}
-				}
-			}
-		}
-		for (const UCameraRigAsset* ToCameraRig : ToCombinedCameraRigs)
-		{
-			if (ToCameraAsset && !ToCameraAsset->GetEnterTransitions().IsEmpty())
-			{
-				// Look at enter transitions on its parent camera asset.
-				for (const UCameraRigAsset* FromCameraRig : FromCombinedCameraRigs)
-				{
-					TransitionToUse = FindTransition(
-							ToCameraAsset->GetEnterTransitions(),
-							FromCameraRig, FromCameraAsset, bFromFrozen,
-							ToCameraRig, ToCameraAsset);
-					if (TransitionToUse)
-					{
-						return TransitionToUse;
-					}
-				}
-			}
-		}
-	}
-	// else: make the first camera rig in the stack start at 100% blend immediately.
-
-	return nullptr;
-}
-
-const UCameraRigTransition* FTransientBlendStackCameraNodeEvaluator::FindTransition(
-			TArrayView<const TObjectPtr<UCameraRigTransition>> Transitions, 
-			const UCameraRigAsset* FromCameraRig, const UCameraAsset* FromCameraAsset, bool bFromFrozen,
-			const UCameraRigAsset* ToCameraRig, const UCameraAsset* ToCameraAsset) const
-{
-	FCameraRigTransitionConditionMatchParams MatchParams;
-	MatchParams.FromCameraRig = FromCameraRig;
-	MatchParams.FromCameraAsset = FromCameraAsset;
-	MatchParams.ToCameraRig = ToCameraRig;
-	MatchParams.ToCameraAsset = ToCameraAsset;
-
-	// The transition should be used if all its conditions pass.
-	for (TObjectPtr<const UCameraRigTransition> Transition : Transitions)
-	{
-		const bool bConditionsPass = Transition->AllConditionsMatch(MatchParams);
-		if (bConditionsPass)
-		{
-			return Transition;
-		}
-	}
-
-	return nullptr;
-}
-
-void FPersistentBlendStackCameraNodeEvaluator::Insert(const FBlendStackCameraInsertParams& Params)
-{
-	// See if we already have this camera rig and evaluation context in the stack.
-	if (!Params.bForceInsert)
-	{
-		for (const FCameraRigEntry& Entry : Entries)
-		{
-			if (!Entry.bIsFrozen &&
-					Entry.CameraRig == Params.CameraRig &&
-					Entry.EvaluationContext == Params.EvaluationContext)
-			{
-				return;
-			}
-		}
-	}
-
-	// TODO: add support for slot indices or something, to allow callers to specify a place in the stack.
-	UObject* Outer = const_cast<UObject*>((UObject*)GetCameraNode());
-	UBlendStackRootCameraNode* EntryRootNode = NewObject<UBlendStackRootCameraNode>(Outer, NAME_None);
-	{
-		EntryRootNode->RootNode = Params.CameraRig->RootNode;
-		// TODO: add support for blending in and out.
-	}
-
-	FCameraRigEntry NewEntry;
-	const bool bInitialized = InitializeEntry(
-			NewEntry, 
-			Params.CameraRig,
-			Params.EvaluationContext,
-			EntryRootNode,
-			false);
-	if (!bInitialized)
-	{
-		return;
-	}
-
-#if WITH_EDITOR
-	AddPackageListeners(NewEntry);
-#endif  // WITH_EDITOR
-
-	Entries.Add(MoveTemp(NewEntry));
-
-	if (OnCameraRigEventDelegate.IsBound())
-	{
-		BroadcastCameraRigEvent(EBlendStackCameraRigEventType::Pushed, Entries.Last(), nullptr);
-	}
-}
-
-void FPersistentBlendStackCameraNodeEvaluator::Remove(const FBlendStackCameraRemoveParams& Params)
-{
-	for (int32 Index = Entries.Num() - 1; Index >= 0; --Index)
-	{
-		FCameraRigEntry& Entry(Entries[Index]);
-		if (Entry.CameraRig == Params.CameraRig &&
-				Entry.EvaluationContext == Params.EvaluationContext)
-		{
-			PopEntry(Index);
-		}
-	}
-}
-
-void FPersistentBlendStackCameraNodeEvaluator::OnRun(const FCameraNodeEvaluationParams& Params, FCameraNodeEvaluationResult& OutResult)
-{
-	// Validate our entries and resolve evaluation context weak pointers.
-	TArray<FResolvedEntry> ResolvedEntries;
-	ResolveEntries(ResolvedEntries);
-
-	// Run the stack!
-	InternalUpdate(ResolvedEntries, Params, OutResult);
-
-	// Tidy things up.
-	OnRunFinished(OutResult);
-}
-
-void FPersistentBlendStackCameraNodeEvaluator::InternalUpdate(TArrayView<FResolvedEntry> ResolvedEntries, const FCameraNodeEvaluationParams& Params, FCameraNodeEvaluationResult& OutResult)
-{
-	constexpr ECameraVariableTableFilter VariableTableFilter = ECameraVariableTableFilter::ChangedOnly;
-	constexpr ECameraContextDataTableFilter ContextDataTableFilter = ECameraContextDataTableFilter::ChangedOnly;
-
-	for (FResolvedEntry& ResolvedEntry : ResolvedEntries)
-	{
-		FCameraRigEntry& Entry(ResolvedEntry.Entry);
-
-		if (!Entry.bIsFrozen)
-		{
-			FCameraNodeEvaluationParams CurParams(Params);
-			CurParams.EvaluationContext = ResolvedEntry.Context;
-			CurParams.bIsFirstFrame = Entry.bIsFirstFrame;
-
-			FCameraNodeEvaluationResult& CurResult(Entry.Result);
-
-			// Start with the input given to us.
-			{
-				CurResult.CameraPose = OutResult.CameraPose;
-				CurResult.VariableTable.OverrideAll(OutResult.VariableTable);
-				CurResult.CameraRigJoints.OverrideAll(OutResult.CameraRigJoints);
-				CurResult.PostProcessSettings.OverrideAll(OutResult.PostProcessSettings);
-
-				// Override it with whatever the evaluation context has set on its result.
-				// Evaluation contexts may have private variables we need to pass along, such as when rig parameter
-				// overrides have been set on them, so include private variables in the filter.
-				const FCameraNodeEvaluationResult& ContextResult(ResolvedEntry.Context->GetInitialResult());
-				CurResult.CameraPose.OverrideChanged(ContextResult.CameraPose);
-				CurResult.VariableTable.Override(ContextResult.VariableTable, VariableTableFilter);
-				CurResult.ContextDataTable.Override(ContextResult.ContextDataTable, ContextDataTableFilter);
-
-				// Setup flags.
-				CurResult.bIsCameraCut = OutResult.bIsCameraCut || ContextResult.bIsCameraCut || Entry.bForceCameraCut;
-				CurResult.bIsValid = true;
-			}
-
-			// Update pre-blended parameters.
-			{
-				FCameraBlendedParameterUpdateParams InputParams(CurParams, CurResult.CameraPose);
-				FCameraBlendedParameterUpdateResult InputResult(CurResult.VariableTable);
-
-				Entry.EvaluatorHierarchy.ForEachEvaluator(ECameraNodeEvaluatorFlags::NeedsParameterUpdate,
-						[&InputParams, &InputResult](FCameraNodeEvaluator* ParameterEvaluator)
-						{
-							ParameterEvaluator->UpdateParameters(InputParams, InputResult);
-						});
-			}
-
-			// Run the blend node.
-			FBlendCameraNodeEvaluator* EntryBlendEvaluator = Entry.RootEvaluator->GetBlendEvaluator();
-			if (EntryBlendEvaluator)
-			{
-				EntryBlendEvaluator->Run(CurParams, CurResult);
-			}
-
-			// Blend pre-blended parameters.
-			if (EntryBlendEvaluator)
-			{
-				FCameraNodePreBlendParams PreBlendParams(CurParams, CurResult.CameraPose, CurResult.VariableTable);
-				FCameraNodePreBlendResult PreBlendResult(OutResult.VariableTable);
-
-				EntryBlendEvaluator->BlendParameters(PreBlendParams, PreBlendResult);
-			}
-			else
-			{
-				OutResult.VariableTable.Override(CurResult.VariableTable, ECameraVariableTableFilter::None);
-			}
-
-			// Run the camera rig's root node.
-			FCameraNodeEvaluator* RootEvaluator = Entry.RootEvaluator->GetRootEvaluator();
-			if (RootEvaluator)
-			{
-				RootEvaluator->Run(CurParams, CurResult);
-			}
-
-			// Blend the results.
-			if (EntryBlendEvaluator)
-			{
-				FCameraNodeBlendParams BlendParams(CurParams, CurResult);
-				FCameraNodeBlendResult BlendResult(OutResult);
-
-				EntryBlendEvaluator->BlendResults(BlendParams, BlendResult);
-			}
-			else
-			{
-				OutResult.OverrideAll(CurResult);
-			}
-		}
-		else
-		{
-			FCameraNodeEvaluationResult& CurResult(Entry.Result);
-
-			OutResult.VariableTable.Override(CurResult.VariableTable, ECameraVariableTableFilter::None);
-			OutResult.OverrideAll(CurResult);
-		}
-	}
-}
-
 #if UE_GAMEPLAY_CAMERAS_DEBUG
 
 void FBlendStackCameraNodeEvaluator::OnBuildDebugBlocks(const FCameraDebugBlockBuildParams& Params, FCameraDebugBlockBuilder& Builder)
diff --git a/Engine/Plugins/Cameras/GameplayCameras/Source/GameplayCameras/Private/Core/DefaultRootCameraNode.cpp b/Engine/Plugins/Cameras/GameplayCameras/Source/GameplayCameras/Private/Core/DefaultRootCameraNode.cpp
--- a/Engine/Plugins/Cameras/GameplayCameras/Source/GameplayCameras/Private/Core/DefaultRootCameraNode.cpp
+++ b/Engine/Plugins/Cameras/GameplayCameras/Source/GameplayCameras/Private/Core/DefaultRootCameraNode.cpp
@@ -5,7 +5,9 @@
 #include "Core/BlendStackCameraNode.h"
 #include "Core/CameraEvaluationContext.h"
 #include "Core/CameraNodeEvaluatorHierarchy.h"
+#include "Core/PersistentBlendStackCameraNode.h"
 #include "Core/RootCameraNodeCameraRigEvent.h"
+#include "Core/TransientBlendStackCameraNode.h"
 #include "Debug/BlendStacksCameraDebugBlock.h"
 #include "Debug/CameraDebugBlockBuilder.h"
 #include "Debug/RootCameraDebugBlock.h"
diff --git a/Engine/Plugins/Cameras/GameplayCameras/Source/GameplayCameras/Private/Core/PersistentBlendStackCameraNode.cpp b/Engine/Plugins/Cameras/GameplayCameras/Source/GameplayCameras/Private/Core/PersistentBlendStackCameraNode.cpp
new file mode 100644
--- /dev/null
+++ b/Engine/Plugins/Cameras/GameplayCameras/Source/GameplayCameras/Private/Core/PersistentBlendStackCameraNode.cpp
@@ -0,0 +1,189 @@
+// Copyright Epic Games, Inc. All Rights Reserved.
+
+#include "Core/PersistentBlendStackCameraNode.h"
+
+#include "Core/BlendCameraNode.h"
+#include "Core/BlendStackCameraRigEvent.h"
+#include "Core/BlendStackRootCameraNode.h"
+#include "Core/CameraEvaluationContext.h"
+
+namespace UE::Cameras
+{
+
+UE_DEFINE_CAMERA_NODE_EVALUATOR(FPersistentBlendStackCameraNodeEvaluator)
+
+void FPersistentBlendStackCameraNodeEvaluator::Insert(const FBlendStackCameraInsertParams& Params)
+{
+	// See if we already have this camera rig and evaluation context in the stack.
+	if (!Params.bForceInsert)
+	{
+		for (const FCameraRigEntry& Entry : Entries)
+		{
+			if (!Entry.bIsFrozen &&
+					Entry.CameraRig == Params.CameraRig &&
+					Entry.EvaluationContext == Params.EvaluationContext)
+			{
+				return;
+			}
+		}
+	}
+
+	// TODO: add support for slot indices or something, to allow callers to specify a place in the stack.
+	UObject* Outer = const_cast<UObject*>((UObject*)GetCameraNode());
+	UBlendStackRootCameraNode* EntryRootNode = NewObject<UBlendStackRootCameraNode>(Outer, NAME_None);
+	{
+		EntryRootNode->RootNode = Params.CameraRig->RootNode;
+		// TODO: add support for blending in and out.
+	}
+
+	FCameraRigEntry NewEntry;
+	const bool bInitialized = InitializeEntry(
+			NewEntry, 
+			Params.CameraRig,
+			Params.EvaluationContext,
+			EntryRootNode,
+			false);
+	if (!bInitialized)
+	{
+		return;
+	}
+
+#if WITH_EDITOR
+	AddPackageListeners(NewEntry);
+#endif  // WITH_EDITOR
+
+	Entries.Add(MoveTemp(NewEntry));
+
+	if (OnCameraRigEventDelegate.IsBound())
+	{
+		BroadcastCameraRigEvent(EBlendStackCameraRigEventType::Pushed, Entries.Last(), nullptr);
+	}
+}
+
+void FPersistentBlendStackCameraNodeEvaluator::Remove(const FBlendStackCameraRemoveParams& Params)
+{
+	for (int32 Index = Entries.Num() - 1; Index >= 0; --Index)
+	{
+		FCameraRigEntry& Entry(Entries[Index]);
+		if (Entry.CameraRig == Params.CameraRig &&
+				Entry.EvaluationContext == Params.EvaluationContext)
+		{
+			PopEntry(Index);
+		}
+	}
+}
+
+void FPersistentBlendStackCameraNodeEvaluator::OnRun(const FCameraNodeEvaluationParams& Params, FCameraNodeEvaluationResult& OutResult)
+{
+	// Validate our entries and resolve evaluation context weak pointers.
+	TArray<FResolvedEntry> ResolvedEntries;
+	ResolveEntries(ResolvedEntries);
+
+	// Run the stack!
+	InternalUpdate(ResolvedEntries, Params, OutResult);
+
+	// Tidy things up.
+	OnRunFinished(OutResult);
+}
+
+void FPersistentBlendStackCameraNodeEvaluator::InternalUpdate(TArrayView<FResolvedEntry> ResolvedEntries, const FCameraNodeEvaluationParams& Params, FCameraNodeEvaluationResult& OutResult)
+{
+	constexpr ECameraVariableTableFilter VariableTableFilter = ECameraVariableTableFilter::ChangedOnly;
+	constexpr ECameraContextDataTableFilter ContextDataTableFilter = ECameraContextDataTableFilter::ChangedOnly;
+
+	for (FResolvedEntry& ResolvedEntry : ResolvedEntries)
+	{
+		FCameraRigEntry& Entry(ResolvedEntry.Entry);
+
+		if (!Entry.bIsFrozen)
+		{
+			FCameraNodeEvaluationParams CurParams(Params);
+			CurParams.EvaluationContext = ResolvedEntry.Context;
+			CurParams.bIsFirstFrame = Entry.bIsFirstFrame;
+
+			FCameraNodeEvaluationResult& CurResult(Entry.Result);
+
+			// Start with the input given to us.
+			{
+				CurResult.CameraPose = OutResult.CameraPose;
+				CurResult.VariableTable.OverrideAll(OutResult.VariableTable);
+				CurResult.CameraRigJoints.OverrideAll(OutResult.CameraRigJoints);
+				CurResult.PostProcessSettings.OverrideAll(OutResult.PostProcessSettings);
+
+				// Override it with whatever the evaluation context has set on its result.
+				// Evaluation contexts may have private variables we need to pass along, such as when rig parameter
+				// overrides have been set on them, so include private variables in the filter.
+				const FCameraNodeEvaluationResult& ContextResult(ResolvedEntry.Context->GetInitialResult());
+				CurResult.CameraPose.OverrideChanged(ContextResult.CameraPose);
+				CurResult.VariableTable.Override(ContextResult.VariableTable, VariableTableFilter);
+				CurResult.ContextDataTable.Override(ContextResult.ContextDataTable, ContextDataTableFilter);
+
+				// Setup flags.
+				CurResult.bIsCameraCut = OutResult.bIsCameraCut || ContextResult.bIsCameraCut || Entry.bForceCameraCut;
+				CurResult.bIsValid = true;
+			}
+
+			// Update pre-blended parameters.
+			{
+				FCameraBlendedParameterUpdateParams InputParams(CurParams, CurResult.CameraPose);
+				FCameraBlendedParameterUpdateResult InputResult(CurResult.VariableTable);
+
+				Entry.EvaluatorHierarchy.ForEachEvaluator(ECameraNodeEvaluatorFlags::NeedsParameterUpdate,
+						[&InputParams, &InputResult](FCameraNodeEvaluator* ParameterEvaluator)
+						{
+							ParameterEvaluator->UpdateParameters(InputParams, InputResult);
+						});
+			}
+
+			// Run the blend node.
+			FBlendCameraNodeEvaluator* EntryBlendEvaluator = Entry.RootEvaluator->GetBlendEvaluator();
+			if (EntryBlendEvaluator)
+			{
+				EntryBlendEvaluator->Run(CurParams, CurResult);
+			}
+
+			// Blend pre-blended parameters.
+			if (EntryBlendEvaluator)
+			{
+				FCameraNodePreBlendParams PreBlendParams(CurParams, CurResult.CameraPose, CurResult.VariableTable);
+				FCameraNodePreBlendResult PreBlendResult(OutResult.VariableTable);
+
+				EntryBlendEvaluator->BlendParameters(PreBlendParams, PreBlendResult);
+			}
+			else
+			{
+				OutResult.VariableTable.Override(CurResult.VariableTable, ECameraVariableTableFilter::None);
+			}
+
+			// Run the camera rig's root node.
+			FCameraNodeEvaluator* RootEvaluator = Entry.RootEvaluator->GetRootEvaluator();
+			if (RootEvaluator)
+			{
+				RootEvaluator->Run(CurParams, CurResult);
+			}
+
+			// Blend the results.
+			if (EntryBlendEvaluator)
+			{
+				FCameraNodeBlendParams BlendParams(CurParams, CurResult);
+				FCameraNodeBlendResult BlendResult(OutResult);
+
+				EntryBlendEvaluator->BlendResults(BlendParams, BlendResult);
+			}
+			else
+			{
+				OutResult.OverrideAll(CurResult);
+			}
+		}
+		else
+		{
+			FCameraNodeEvaluationResult& CurResult(Entry.Result);
+
+			OutResult.VariableTable.Override(CurResult.VariableTable, ECameraVariableTableFilter::None);
+			OutResult.OverrideAll(CurResult);
+		}
+	}
+}
+
+}  // namespace UE::Cameras
+
diff --git a/Engine/Plugins/Cameras/GameplayCameras/Source/GameplayCameras/Private/Core/TransientBlendStackCameraNode.cpp b/Engine/Plugins/Cameras/GameplayCameras/Source/GameplayCameras/Private/Core/TransientBlendStackCameraNode.cpp
new file mode 100644
--- /dev/null
+++ b/Engine/Plugins/Cameras/GameplayCameras/Source/GameplayCameras/Private/Core/TransientBlendStackCameraNode.cpp
@@ -0,0 +1,657 @@
+// Copyright Epic Games, Inc. All Rights Reserved.
+
+#include "Core/TransientBlendStackCameraNode.h"
+
+#include "Algo/AnyOf.h"
+#include "Core/BlendCameraNode.h"
+#include "Core/BlendStackCameraRigEvent.h"
+#include "Core/BlendStackRootCameraNode.h"
+#include "Core/CameraAsset.h"
+#include "Core/CameraEvaluationContext.h"
+#include "Core/CameraRigAsset.h"
+#include "Core/CameraRigCombinationRegistry.h"
+#include "Core/CameraRigTransition.h"
+#include "Nodes/Blends/PopBlendCameraNode.h"
+
+namespace UE::Cameras
+{
+
+UE_DEFINE_CAMERA_NODE_EVALUATOR(FTransientBlendStackCameraNodeEvaluator)
+
+void FTransientBlendStackCameraNodeEvaluator::Push(const FBlendStackCameraPushParams& Params)
+{
+	bool bSearchedForTransition = false;
+	const UCameraRigTransition* Transition = nullptr;
+
+	if (!Entries.IsEmpty())
+	{
+		FCameraRigEntry& TopEntry(Entries.Top());
+		if (!TopEntry.bIsFrozen 
+				&& TopEntry.EvaluationContext == Params.EvaluationContext)
+		{
+			// Don't push anything if what is being requested is already the active 
+			// camera rig.
+			if (!Params.bForcePush && TopEntry.CameraRig == Params.CameraRig)
+			{
+				return;
+			}
+
+			// See if we can merge the new camera rig onto the active camera rig.
+			const ECameraRigMergingEligibility Eligibility = TopEntry.RootEvaluator->CompareCameraRigForMerging(Params.CameraRig);
+
+			if (!Params.bForcePush && Eligibility == ECameraRigMergingEligibility::Active)
+			{
+				// This camera rig is already the active one on the merged stack.
+				return;
+			}
+
+			if (Eligibility == ECameraRigMergingEligibility::EligibleForMerge)
+			{
+				// This camera rig can be merged with the one current running. However, we
+				// only do it if the transition explicitly allows it.
+				bSearchedForTransition = true;
+				Transition = FindTransition(Params);
+
+				if (Transition && Transition->bAllowCameraRigMerging)
+				{
+					PushMergedEntry(Params, Transition);
+					return;
+				}
+			}
+		}
+	}
+
+	// It's a legitimate new entry in the blend stack.
+	if (!bSearchedForTransition)
+	{
+		bSearchedForTransition = true;
+		Transition = FindTransition(Params);
+	}
+
+	PushNewEntry(Params, Transition);
+}
+
+void FTransientBlendStackCameraNodeEvaluator::PushNewEntry(const FBlendStackCameraPushParams& Params, const UCameraRigTransition* Transition)
+{
+	// Create the new root node to wrap the new camera rig's root node, and the specific
+	// blend node for this transition.
+	// We need to const-cast here to be able to use our own blend stack node as the outer
+	// of the new node.
+	const UCameraRigTransition* UsedTransition = nullptr;
+	UObject* Outer = const_cast<UObject*>((UObject*)GetCameraNode());
+	UBlendStackRootCameraNode* EntryRootNode = NewObject<UBlendStackRootCameraNode>(Outer, NAME_None);
+	{
+		EntryRootNode->RootNode = Params.CameraRig->RootNode;
+
+		// Find a transition and use its blend. If no transition is found,
+		// make a camera cut transition.
+		UBlendCameraNode* ModeBlend = nullptr;
+		if (Transition)
+		{
+			ModeBlend = Transition->Blend;
+			UsedTransition = Transition;
+		}
+		if (!ModeBlend)
+		{
+			ModeBlend = NewObject<UPopBlendCameraNode>(EntryRootNode, NAME_None);
+		}
+		EntryRootNode->Blend = ModeBlend;
+	}
+
+	// Make the new stack entry, and use its storage buffer to build the tree of evaluators.
+	FCameraRigEntry NewEntry;
+	const bool bInitialized = InitializeEntry(
+			NewEntry, 
+			Params.CameraRig,
+			Params.EvaluationContext,
+			EntryRootNode,
+			true);
+	if (!bInitialized)
+	{
+		return;
+	}
+	
+	// Take a snapshot of the initial result for having base values during pre-blending.
+	FCameraRigEntryExtraInfo NewExtraInfo;
+	{
+		FCameraNodeEvaluationResult& UnblendedResult = NewExtraInfo.UnblendedResult;
+
+		if (Params.CameraRig)
+		{
+			const FCameraRigAllocationInfo& AllocationInfo = Params.CameraRig->AllocationInfo;
+			UnblendedResult.VariableTable.Initialize(AllocationInfo.VariableTableInfo);
+			UnblendedResult.ContextDataTable.Initialize(AllocationInfo.ContextDataTableInfo);
+		}
+
+		UnblendedResult.OverrideAll(NewEntry.Result, true);
+	}
+
+#if WITH_EDITOR
+	// Listen to changes to the packages inside which this camera rig is defined. We will hot-reload the
+	// camera node evaluators for this camera rig when we detect changes.
+	AddPackageListeners(NewEntry);
+#endif  // WITH_EDITOR
+
+	// Important: we need to move the new entry here because copying evaluator storage
+	// is disabled.
+	Entries.Add(MoveTemp(NewEntry));
+	EntryExtraInfos.Add(MoveTemp(NewExtraInfo));
+
+	if (OnCameraRigEventDelegate.IsBound())
+	{
+		BroadcastCameraRigEvent(EBlendStackCameraRigEventType::Pushed, Entries.Last(), UsedTransition);
+	}
+}
+
+void FTransientBlendStackCameraNodeEvaluator::PushMergedEntry(const FBlendStackCameraPushParams& PushParams, const UCameraRigTransition* Transition)
+{
+	const UBlendCameraNode* Blend = Transition ? Transition->Blend : nullptr;
+
+	FCameraRigEntry& TopEntry = Entries.Top();
+	FCameraNodeEvaluatorBuilder Builder(TopEntry.EvaluatorStorage);
+	FCameraNodeEvaluatorBuildParams BuildParams(Builder);
+	TopEntry.RootEvaluator->MergeCameraRig(BuildParams, PushParams.CameraRig, Blend);
+
+	// Merging a camera rig the first time changes the evaluator tree by removing any prefab nodes
+	// near the root... so we need to rebuild our hierarchy cache.
+	TopEntry.EvaluatorHierarchy.Build(TopEntry.RootEvaluator);
+
+	// Swap out the camera rig registered as "active" for this entry.
+#if WITH_EDITOR
+	RemoveListenedPackages(TopEntry);
+#endif
+	{
+		TopEntry.CameraRig = PushParams.CameraRig;
+	}
+#if WITH_EDITOR
+	AddPackageListeners(TopEntry);
+#endif
+}
+
+void FTransientBlendStackCameraNodeEvaluator::Freeze(const FBlendStackCameraFreezeParams& Params)
+{
+	for (FCameraRigEntry& Entry : Entries)
+	{
+		if (!Entry.bIsFrozen && 
+				Entry.CameraRig == Params.CameraRig &&
+				Entry.EvaluationContext == Params.EvaluationContext)
+		{
+			FreezeEntry(Entry);
+		}
+	}
+}
+
+void FTransientBlendStackCameraNodeEvaluator::FreezeAll(TSharedPtr<FCameraEvaluationContext> EvaluationContext)
+{
+	for (FCameraRigEntry& Entry : Entries)
+	{
+		if (!Entry.bIsFrozen && Entry.EvaluationContext == EvaluationContext)
+		{
+			FreezeEntry(Entry);
+		}
+	}
+}
+
+void FTransientBlendStackCameraNodeEvaluator::OnRun(const FCameraNodeEvaluationParams& Params, FCameraNodeEvaluationResult& OutResult)
+{
+	ensure(Entries.Num() == EntryExtraInfos.Num());
+
+	// Validate our entries and resolve evaluation context weak pointers.
+	TArray<FResolvedEntry> ResolvedEntries;
+	ResolveEntries(ResolvedEntries);
+
+	// Gather parameters to pre-blend, and evaluate blend nodes.
+	InternalPreBlendPrepare(ResolvedEntries, Params, OutResult);
+
+	// Blend input variables.
+	InternalPreBlendExecute(ResolvedEntries, Params, OutResult);
+
+	// Run the root nodes. They will use the pre-blended inputs from the last step.
+	// Frozen entries are skipped, since they only ever use the last result they produced.
+	InternalUpdate(ResolvedEntries, Params, OutResult);
+
+	// Now blend all the results, keeping track of blends that have reached 100% so
+	// that we can remove any camera rigs below (since they would have been completely
+	// blended out by that).
+	InternalPostBlendExecute(ResolvedEntries, Params, OutResult);
+
+	// Tidy up.
+	OnRunFinished(OutResult);
+	InternalRunFinished(OutResult);
+}
+
+void FTransientBlendStackCameraNodeEvaluator::InternalPreBlendPrepare(TArrayView<FResolvedEntry> ResolvedEntries, const FCameraNodeEvaluationParams& Params, FCameraNodeEvaluationResult& OutResult)
+{
+	constexpr ECameraVariableTableFilter VariableTableFilter = ECameraVariableTableFilter::ChangedOnly;
+	constexpr ECameraContextDataTableFilter ContextDataTableFilter = ECameraContextDataTableFilter::ChangedOnly;
+
+	for (FResolvedEntry& ResolvedEntry : ResolvedEntries)
+	{
+		FCameraRigEntry& Entry(ResolvedEntry.Entry);
+		FCameraRigEntryExtraInfo& EntryExtraInfo(EntryExtraInfos[ResolvedEntry.EntryIndex]);
+
+		if (UNLIKELY(Entry.bIsFrozen))
+		{
+			continue;
+		}
+
+		FCameraNodeEvaluationParams CurParams(Params);
+		CurParams.EvaluationContext = ResolvedEntry.Context;
+		CurParams.bIsFirstFrame = Entry.bIsFirstFrame;
+
+		FCameraNodeEvaluationResult& UnblendedResult(EntryExtraInfo.UnblendedResult);
+
+		// Reset frame flags to match what happend to the Entry.Result in ResolveEntries.
+		UnblendedResult.ResetFrameFlags();
+
+		// Start with the input given to us.
+		UnblendedResult.VariableTable.OverrideAll(OutResult.VariableTable);
+
+		// Override it with whatever the evaluation context has set on its result this frame.
+		// Evaluation contexts may have private variables we need to pass along, such as when rig parameter
+		// overrides have been set on them, so include private variables in the filter.
+		const FCameraNodeEvaluationResult& ContextResult(ResolvedEntry.Context->GetInitialResult());
+		UnblendedResult.VariableTable.Override(ContextResult.VariableTable, VariableTableFilter);
+		UnblendedResult.ContextDataTable.Override(ContextResult.ContextDataTable, ContextDataTableFilter);
+		if (ResolvedEntry.bIsActiveEntry)
+		{
+			if (const FCameraNodeEvaluationResult* ActiveOnlyResult = ResolvedEntry.Context->GetConditionalResult(ECameraEvaluationDataCondition::ActiveCameraRig))
+			{
+				UnblendedResult.VariableTable.Override(ActiveOnlyResult->VariableTable, VariableTableFilter);
+				UnblendedResult.ContextDataTable.Override(ActiveOnlyResult->ContextDataTable, ContextDataTableFilter);
+			}
+		}
+
+		// Gather input parameters if needed (and remember if it was indeed needed).
+		if (!EntryExtraInfo.bInputRunThisFrame)
+		{
+			bool bHasPreBlendedParameters = false;
+			FCameraBlendedParameterUpdateParams InputParams(CurParams, UnblendedResult.CameraPose);
+			FCameraBlendedParameterUpdateResult InputResult(UnblendedResult.VariableTable);
+
+			Entry.EvaluatorHierarchy.ForEachEvaluator(ECameraNodeEvaluatorFlags::NeedsParameterUpdate,
+					[&bHasPreBlendedParameters, &InputParams, &InputResult](FCameraNodeEvaluator* ParameterEvaluator)
+					{
+						ParameterEvaluator->UpdateParameters(InputParams, InputResult);
+						bHasPreBlendedParameters = true;
+					});
+
+			EntryExtraInfo.bHasPreBlendedParameters = bHasPreBlendedParameters;
+			EntryExtraInfo.bInputRunThisFrame = true;
+		}
+
+		// Run blends.
+		// Note that we pass last frame's camera pose to the Run() method. This may change.
+		// Blends aren't expected to use the camera pose to do any logic until BlendResults().
+		if (!EntryExtraInfo.bBlendRunThisFrame)
+		{
+			FBlendCameraNodeEvaluator* EntryBlendEvaluator = Entry.RootEvaluator->GetBlendEvaluator();
+			if (EntryBlendEvaluator)
+			{
+				EntryBlendEvaluator->Run(CurParams, UnblendedResult);
+			}
+
+			EntryExtraInfo.bBlendRunThisFrame = true;
+		}
+
+		// Assign the new unblended result back to the entry result. The goal is that:
+		//
+		// 1) We only write values from the context result when they have been written to.
+		//    That is, we don't always set all values every frame. When "left alone", a camera rig
+		//    should retain whatever values it was last set with. For instance, if a camera rig 
+		//    received special values when it was the active camera rig, we want to keep those 
+		//    values when it starts blending out, instead of immediately overwriting them with
+		//    the "common" values (i.e. the values from the initial result).
+		//
+		// 2) However, we do want to write all the values into the PreBlendVariableTable every
+		//    frame, so that we can have them blend together. The blended values are then reassigned
+		//    to Entry.Result in InternalPreBlendExecute, which means we lose each camera rig's
+		//    "intended" values. Getting them from the evaluation context each frame would violate
+		//    previous item (1) so that's why we use the UnblendedResult as a snapshot.
+		//
+		FCameraNodeEvaluationResult& CurResult(Entry.Result);
+		CurResult.OverrideAll(UnblendedResult, true);
+	}
+}
+
+void FTransientBlendStackCameraNodeEvaluator::InternalPreBlendExecute(TArrayView<FResolvedEntry> ResolvedEntries, const FCameraNodeEvaluationParams& Params, FCameraNodeEvaluationResult& OutResult)
+{
+	// Blend all the camera rigs' input variables (include private variables such as camera rig parameters).
+	// The result of this pre-blending goes into PreBlendVariableTable, which will then have a big mix
+	// of all the variables.
+	PreBlendVariableTable.ClearAllWrittenThisFrameFlags();
+
+	if (!Algo::AnyOf(EntryExtraInfos, [](const FCameraRigEntryExtraInfo& Item) { return Item.bHasPreBlendedParameters; }))
+	{
+		return;
+	}
+
+	constexpr ECameraVariableTableFilter VariableTableFilter = ECameraVariableTableFilter::None;
+
+	for (FResolvedEntry& ResolvedEntry : ResolvedEntries)
+	{
+		FCameraRigEntry& Entry(ResolvedEntry.Entry);
+		FCameraRigEntryExtraInfo& EntryExtraInfo(EntryExtraInfos[ResolvedEntry.EntryIndex]);
+
+		FCameraNodeEvaluationResult& UnblendedResult(EntryExtraInfo.UnblendedResult);
+
+		if (!Entry.bIsFrozen)
+		{
+			FCameraNodeEvaluationParams CurParams(Params);
+			CurParams.EvaluationContext = ResolvedEntry.Context;
+			CurParams.bIsFirstFrame = Entry.bIsFirstFrame;
+
+			FCameraNodePreBlendParams PreBlendParams(CurParams, UnblendedResult.CameraPose, UnblendedResult.VariableTable);
+			PreBlendParams.VariableTableFilter = VariableTableFilter;
+
+			FCameraNodePreBlendResult PreBlendResult(PreBlendVariableTable);
+
+			FBlendCameraNodeEvaluator* EntryBlendEvaluator = Entry.RootEvaluator->GetBlendEvaluator();
+			if (EntryBlendEvaluator)
+			{
+				EntryBlendEvaluator->BlendParameters(PreBlendParams, PreBlendResult);
+				EntryExtraInfo.bIsPreBlendFull = (PreBlendResult.bIsBlendFinished && PreBlendResult.bIsBlendFull);
+			}
+			else
+			{
+				PreBlendVariableTable.Override(UnblendedResult.VariableTable, PreBlendParams.VariableTableFilter);
+				EntryExtraInfo.bIsPreBlendFull = true;
+			}
+		}
+		else
+		{
+			// Frozen entries still contribute to the blend using their last evaluated values.
+			PreBlendVariableTable.Override(UnblendedResult.VariableTable, VariableTableFilter);
+		}
+	}
+
+	// Write the values back to each entry table, so that each of these camera rigs will run with
+	// the pre-blended values. We limit this writing to the variables each of them knows, since
+	// there's no need to add entries they don't use to their variable table.
+	for (FResolvedEntry& ResolvedEntry : ResolvedEntries)
+	{
+		FCameraRigEntry& Entry(ResolvedEntry.Entry);
+		FCameraRigEntryExtraInfo& EntryExtraInfo(EntryExtraInfos[ResolvedEntry.EntryIndex]);
+
+		if (!Entry.bIsFrozen)
+		{
+			FCameraNodeEvaluationResult& UnblendedResult(EntryExtraInfo.UnblendedResult);
+			FCameraNodeEvaluationResult& CurResult(Entry.Result);
+
+			Entry.Result.VariableTable.Override(PreBlendVariableTable, ECameraVariableTableFilter::KnownOnly);
+		}
+	}
+}
+
+void FTransientBlendStackCameraNodeEvaluator::InternalUpdate(TArrayView<FResolvedEntry> ResolvedEntries, const FCameraNodeEvaluationParams& Params, FCameraNodeEvaluationResult& OutResult)
+{
+	for (FResolvedEntry& ResolvedEntry : ResolvedEntries)
+	{
+		FCameraRigEntry& Entry(ResolvedEntry.Entry);
+		FCameraRigEntryExtraInfo& EntryExtraInfo(EntryExtraInfos[ResolvedEntry.EntryIndex]);
+
+		if (UNLIKELY(Entry.bIsFrozen))
+		{
+			continue;
+		}
+
+		FCameraNodeEvaluationResult& CurResult(Entry.Result);
+
+		FCameraNodeEvaluationParams CurParams(Params);
+		CurParams.EvaluationContext = ResolvedEntry.Context;
+		CurParams.bIsFirstFrame = Entry.bIsFirstFrame;
+
+		// Start with the input given to us.
+		CurResult.Reset();
+		CurResult.CameraPose = OutResult.CameraPose;
+		CurResult.CameraRigJoints.OverrideAll(OutResult.CameraRigJoints);
+		CurResult.PostProcessSettings.OverrideAll(OutResult.PostProcessSettings);
+
+		// Override it with whatever the evaluation context has set on its result.
+		const FCameraNodeEvaluationResult& ContextResult(ResolvedEntry.Context->GetInitialResult());
+		CurResult.CameraPose.OverrideChanged(ContextResult.CameraPose);
+		CurResult.bIsCameraCut = OutResult.bIsCameraCut || ContextResult.bIsCameraCut || Entry.bForceCameraCut;
+		if (ResolvedEntry.bIsActiveEntry)
+		{
+			if (const FCameraNodeEvaluationResult* ActiveOnlyResult = ResolvedEntry.Context->GetConditionalResult(ECameraEvaluationDataCondition::ActiveCameraRig))
+			{
+				CurResult.CameraPose.OverrideChanged(ActiveOnlyResult->CameraPose);
+				CurResult.bIsCameraCut |= ActiveOnlyResult->bIsCameraCut;
+			}
+		}
+		
+		CurResult.bIsValid = true;
+
+#if WITH_EDITOR || UE_GAMEPLAY_CAMERAS_DEBUG
+		CurResult.AddCameraPoseTrailPointIfNeeded(ContextResult.CameraPose.GetLocation());
+#endif  // WITH_EDITOR || UE_GAMEPLAY_CAMERAS_DEBUG
+
+		// Run the camera rig's root node.
+		FCameraNodeEvaluator* RootEvaluator = Entry.RootEvaluator->GetRootEvaluator();
+		if (RootEvaluator)
+		{
+			RootEvaluator->Run(CurParams, CurResult);
+		}
+	}
+}
+
+void FTransientBlendStackCameraNodeEvaluator::InternalPostBlendExecute(TArrayView<FResolvedEntry> ResolvedEntries, const FCameraNodeEvaluationParams& Params, FCameraNodeEvaluationResult& OutResult)
+{
+	int32 PopEntriesBelow = INDEX_NONE;
+	for (FResolvedEntry& ResolvedEntry : ResolvedEntries)
+	{
+		FCameraRigEntry& Entry(ResolvedEntry.Entry);
+		FCameraRigEntryExtraInfo& EntryExtraInfo(EntryExtraInfos[ResolvedEntry.EntryIndex]);
+
+		FCameraNodeEvaluationResult& CurResult(Entry.Result);
+
+		if (!Entry.bIsFrozen)
+		{
+			FCameraNodeEvaluationParams CurParams(Params);
+			CurParams.EvaluationContext = ResolvedEntry.Context;
+			CurParams.bIsFirstFrame = Entry.bIsFirstFrame;
+			FCameraNodeBlendParams BlendParams(CurParams, CurResult);
+
+			FCameraNodeBlendResult BlendResult(OutResult);
+
+			FBlendCameraNodeEvaluator* EntryBlendEvaluator = Entry.RootEvaluator->GetBlendEvaluator();
+			if (EntryBlendEvaluator)
+			{
+				EntryBlendEvaluator->BlendResults(BlendParams, BlendResult);
+
+				if (BlendResult.bIsBlendFull && BlendResult.bIsBlendFinished)
+				{
+					PopEntriesBelow = ResolvedEntry.EntryIndex;
+				}
+			}
+			else
+			{
+				OutResult.OverrideAll(CurResult);
+
+				PopEntriesBelow = ResolvedEntry.EntryIndex;
+			}
+		}
+		else
+		{
+			OutResult.OverrideAll(CurResult);
+
+			PopEntriesBelow = ResolvedEntry.EntryIndex;
+		}
+		// else, merged entries only contribute to the pre-blend.
+	}
+
+	// Pop out camera rigs that have been blended out.
+	const UBlendStackCameraNode* BlendStackNode = GetCameraNodeAs<UBlendStackCameraNode>();
+	if (BlendStackNode->BlendStackType == ECameraBlendStackType::IsolatedTransient && PopEntriesBelow != INDEX_NONE)
+	{
+		PopEntries(PopEntriesBelow);
+
+		EntryExtraInfos.RemoveAt(0, PopEntriesBelow);
+	}
+}
+
+void FTransientBlendStackCameraNodeEvaluator::InternalRunFinished(FCameraNodeEvaluationResult& OutResult)
+{
+	for (FCameraRigEntryExtraInfo& ExtraInfo : EntryExtraInfos)
+	{
+		ExtraInfo.bInputRunThisFrame = false;
+		ExtraInfo.bBlendRunThisFrame = false;
+		ExtraInfo.bHasPreBlendedParameters = false;
+		ExtraInfo.bIsPreBlendFull = false;
+	}
+}
+
+void FTransientBlendStackCameraNodeEvaluator::OnSerialize(const FCameraNodeEvaluatorSerializeParams& Params, FArchive& Ar)
+{
+	Super::OnSerialize(Params, Ar);
+
+	for (FCameraRigEntryExtraInfo& ExtraInfo : EntryExtraInfos)
+	{
+		Ar << ExtraInfo.bInputRunThisFrame;
+		Ar << ExtraInfo.bBlendRunThisFrame;
+		Ar << ExtraInfo.bHasPreBlendedParameters;
+		Ar << ExtraInfo.bIsPreBlendFull;
+	}
+}
+
+const UCameraRigTransition* FTransientBlendStackCameraNodeEvaluator::FindTransition(const FBlendStackCameraPushParams& Params) const
+{
+	// If we are forced to use a specific transition, our search is over.
+	if (Params.TransitionOverride)
+	{
+		return Params.TransitionOverride;
+	}
+
+	// Find a transition that works for blending towards ToCameraRig.
+	// If the stack isn't empty, we need to find a transition that works between the previous and 
+	// next camera rigs. If the stack is empty, we blend the new camera rig in from nothing if
+	// appropriate.
+	if (!Entries.IsEmpty())
+	{
+		// Grab information about the new entry to push.
+		TSharedPtr<const FCameraEvaluationContext> ToContext = Params.EvaluationContext;
+		const UCameraAsset* ToCameraAsset = ToContext ? ToContext->GetCameraAsset() : nullptr;
+
+		// Grab information about the top entry (i.e. the currently active camera rig).
+		const FCameraRigEntry& TopEntry = Entries.Top();
+		TSharedPtr<const FCameraEvaluationContext> FromContext = TopEntry.EvaluationContext.Pin();
+		const UCameraAsset* FromCameraAsset = FromContext ? FromContext->GetCameraAsset() : nullptr;
+
+		// If the new or current top entries are a combination, look for transitions on all 
+		// their combined camera rigs.
+		TArray<const UCameraRigAsset*> ToCombinedCameraRigs;
+		UCombinedCameraRigsCameraNode::GetAllCombinationCameraRigs(Params.CameraRig, ToCombinedCameraRigs);
+
+		TArray<const UCameraRigAsset*> FromCombinedCameraRigs;
+		UCombinedCameraRigsCameraNode::GetAllCombinationCameraRigs(TopEntry.CameraRig, FromCombinedCameraRigs);
+
+		const bool bFromFrozen = TopEntry.bIsFrozen;
+		const UCameraRigTransition* TransitionToUse = nullptr;
+
+		// Start by looking at exit transitions on the last active (top) camera rig.
+		for (const UCameraRigAsset* FromCameraRig : FromCombinedCameraRigs)
+		{
+			if (!FromCameraRig->ExitTransitions.IsEmpty())
+			{
+				// Look for exit transitions on the last active camera rig itself.
+				for (const UCameraRigAsset* ToCameraRig : ToCombinedCameraRigs)
+				{
+					TransitionToUse = FindTransition(
+							FromCameraRig->ExitTransitions,
+							FromCameraRig, FromCameraAsset, bFromFrozen,
+							ToCameraRig, ToCameraAsset);
+					if (TransitionToUse)
+					{
+						return TransitionToUse;
+					}
+				}
+			}
+		}
+		for (const UCameraRigAsset* FromCameraRig : FromCombinedCameraRigs)
+		{
+			if (FromCameraAsset && !FromCameraAsset->GetExitTransitions().IsEmpty())
+			{
+				// Look for exit transitions on its parent camera asset.
+				for (const UCameraRigAsset* ToCameraRig : ToCombinedCameraRigs)
+				{
+					TransitionToUse = FindTransition(
+							FromCameraAsset->GetExitTransitions(),
+							FromCameraRig, FromCameraAsset, bFromFrozen,
+							ToCameraRig, ToCameraAsset);
+					if (TransitionToUse)
+					{
+						return TransitionToUse;
+					}
+				}
+			}
+		}
+
+		// Now look at enter transitions on the new camera rig.
+		for (const UCameraRigAsset* ToCameraRig : ToCombinedCameraRigs)
+		{
+			if (!ToCameraRig->EnterTransitions.IsEmpty())
+			{
+				// Look for enter transitions on the new camera rig itself.
+				for (const UCameraRigAsset* FromCameraRig : FromCombinedCameraRigs)
+				{
+					TransitionToUse = FindTransition(
+							ToCameraRig->EnterTransitions,
+							FromCameraRig, FromCameraAsset, bFromFrozen,
+							ToCameraRig, ToCameraAsset);
+					if (TransitionToUse)
+					{
+						return TransitionToUse;
+					}
+				}
+			}
+		}
+		for (const UCameraRigAsset* ToCameraRig : ToCombinedCameraRigs)
+		{
+			if (ToCameraAsset && !ToCameraAsset->GetEnterTransitions().IsEmpty())
+			{
+				// Look at enter transitions on its parent camera asset.
+				for (const UCameraRigAsset* FromCameraRig : FromCombinedCameraRigs)
+				{
+					TransitionToUse = FindTransition(
+							ToCameraAsset->GetEnterTransitions(),
+							FromCameraRig, FromCameraAsset, bFromFrozen,
+							ToCameraRig, ToCameraAsset);
+					if (TransitionToUse)
+					{
+						return TransitionToUse;
+					}
+				}
+			}
+		}
+	}
+	// else: make the first camera rig in the stack start at 100% blend immediately.
+
+	return nullptr;
+}
+
+const UCameraRigTransition* FTransientBlendStackCameraNodeEvaluator::FindTransition(
+			TArrayView<const TObjectPtr<UCameraRigTransition>> Transitions, 
+			const UCameraRigAsset* FromCameraRig, const UCameraAsset* FromCameraAsset, bool bFromFrozen,
+			const UCameraRigAsset* ToCameraRig, const UCameraAsset* ToCameraAsset) const
+{
+	FCameraRigTransitionConditionMatchParams MatchParams;
+	MatchParams.FromCameraRig = FromCameraRig;
+	MatchParams.FromCameraAsset = FromCameraAsset;
+	MatchParams.ToCameraRig = ToCameraRig;
+	MatchParams.ToCameraAsset = ToCameraAsset;
+
+	// The transition should be used if all its conditions pass.
+	for (TObjectPtr<const UCameraRigTransition> Transition : Transitions)
+	{
+		const bool bConditionsPass = Transition->AllConditionsMatch(MatchParams);
+		if (bConditionsPass)
+		{
+			return Transition;
+		}
+	}
+
+	return nullptr;
+}
+
+}  // namespace UE::Cameras
+
diff --git a/Engine/Plugins/Cameras/GameplayCameras/Source/GameplayCameras/Public/Core/BlendStackCameraNode.h b/Engine/Plugins/Cameras/GameplayCameras/Source/GameplayCameras/Public/Core/BlendStackCameraNode.h
--- a/Engine/Plugins/Cameras/GameplayCameras/Source/GameplayCameras/Public/Core/BlendStackCameraNode.h
+++ b/Engine/Plugins/Cameras/GameplayCameras/Source/GameplayCameras/Public/Core/BlendStackCameraNode.h
@@ -226,165 +226,6 @@ protected:
 #endif  // UE_GAMEPLAY_CAMERAS_DEBUG
 };
 
-/**
- * Parameter structure for pushing a camera rig onto a transient blend stack.
- */
-struct FBlendStackCameraPushParams
-{
-	/** The evaluation context within which a camera rig's node tree should run. */
-	TSharedPtr<const FCameraEvaluationContext> EvaluationContext;
-
-	/** The source camera rig asset to instantiate and push on the blend stack. */
-	TObjectPtr<const UCameraRigAsset> CameraRig;
-
-	/** A transition to use, instead of looking one up. */
-	TObjectPtr<const UCameraRigTransition> TransitionOverride;
-
-	/** Whether to force pushing a new instance of the camera rig, even if it is currently active. */
-	bool bForcePush = false;
-};
-
-/**
- * Parameter structure for freezing a camera rig inside a transient blend stack.
- */
-struct FBlendStackCameraFreezeParams
-{
-	/** The evaluation context within which a camera rig's node tree is running. */
-	TSharedPtr<const FCameraEvaluationContext> EvaluationContext;
-
-	/** The source camera rig asset that is running. */
-	TObjectPtr<const UCameraRigAsset> CameraRig;
-};
-
-/**
- * Evaluator for a transient blend stack.
- */
-class FTransientBlendStackCameraNodeEvaluator 
-	: public FBlendStackCameraNodeEvaluator
-{
-
-	UE_DECLARE_CAMERA_NODE_EVALUATOR_EX(GAMEPLAYCAMERAS_API, FTransientBlendStackCameraNodeEvaluator, FBlendStackCameraNodeEvaluator)
-
-public:
-
-	/** Push a new camera rig onto the blend stack. */
-	void Push(const FBlendStackCameraPushParams& Params);
-
-	/** Freeze a camera rig. */
-	void Freeze(const FBlendStackCameraFreezeParams& Params);
-
-	/** Freeze all camera rigs that belong to a given evaluation context. */
-	void FreezeAll(TSharedPtr<FCameraEvaluationContext> EvaluationContext);
-
-protected:
-
-	// FCameraNodeEvaluator interface.
-	virtual void OnRun(const FCameraNodeEvaluationParams& Params, FCameraNodeEvaluationResult& OutResult) override;
-	virtual void OnSerialize(const FCameraNodeEvaluatorSerializeParams& Params, FArchive& Ar) override;
-
-private:
-
-	// Update methods.
-	void InternalPreBlendPrepare(TArrayView<FResolvedEntry> ResolvedEntries, const FCameraNodeEvaluationParams& Params, FCameraNodeEvaluationResult& OutResult);
-	void InternalPreBlendExecute(TArrayView<FResolvedEntry> ResolvedEntries, const FCameraNodeEvaluationParams& Params, FCameraNodeEvaluationResult& OutResult);
-	void InternalUpdate(TArrayView<FResolvedEntry> ResolvedEntries, const FCameraNodeEvaluationParams& Params, FCameraNodeEvaluationResult& OutResult);
-	void InternalPostBlendExecute(TArrayView<FResolvedEntry> ResolvedEntries, const FCameraNodeEvaluationParams& Params, FCameraNodeEvaluationResult& OutResult);
-	void InternalRunFinished(FCameraNodeEvaluationResult& OutResult);
-
-	// Utility functions for finding an appropriate transition.
-	const UCameraRigTransition* FindTransition(const FBlendStackCameraPushParams& Params) const;
-	const UCameraRigTransition* FindTransition(
-			TArrayView<const TObjectPtr<UCameraRigTransition>> Transitions, 
-			const UCameraRigAsset* FromCameraRig, const UCameraAsset* FromCameraAsset, bool bFromFrozen,
-			const UCameraRigAsset* ToCameraRig, const UCameraAsset* ToCameraAsset) const;
-
-	void PushNewEntry(const FBlendStackCameraPushParams& Params, const UCameraRigTransition* Transition);
-	void PushMergedEntry(const FBlendStackCameraPushParams& Params, const UCameraRigTransition* Transition);
-
-private:
-
-	/** Extra blending-related info for each camera rig entry. */
-	struct FCameraRigEntryExtraInfo
-	{
-		/** Result without pre-blending. */
-		FCameraNodeEvaluationResult UnblendedResult;
-
-		/** Whether input slots were run. */
-		bool bInputRunThisFrame = false;
-		/** Whether the blend node was run. */
-		bool bBlendRunThisFrame = false;
-		/** Whether this camera rig has any parameters to pre-blend. */
-		bool bHasPreBlendedParameters = false;
-		/** Whether this camera rig's pre-blend was full and finished this frame. */
-		bool bIsPreBlendFull = false;
-	};
-
-	TArray<FCameraRigEntryExtraInfo> EntryExtraInfos;
-
-	/** Variable table for pre-blending. */
-	FCameraVariableTable PreBlendVariableTable;
-};
-
-/**
- * Parameter structure for inserting a camera rig into a persistent blend stack.
- */
-struct FBlendStackCameraInsertParams
-{
-	/** The evaluation context within which a camera rig's node tree should run. */
-	TSharedPtr<const FCameraEvaluationContext> EvaluationContext;
-
-	/** The source camera rig asset to instantiate and push on the blend stack. */
-	TObjectPtr<const UCameraRigAsset> CameraRig;
-
-	/** A transition to use, instead of looking one up. */
-	TObjectPtr<const UCameraRigTransition> TransitionOverride;
-
-	/** Whether to force insert a new instance of the camera rig, even if there is already one in the stack. */
-	bool bForceInsert = false;
-};
-
-/**
- * Parameter structure for removing a camera rig from a persistent blend stack.
- */
-struct FBlendStackCameraRemoveParams
-{
-	/** The evaluation context within which a camera rig's node tree should run. */
-	TSharedPtr<const FCameraEvaluationContext> EvaluationContext;
-
-	/** The source camera rig asset to instantiate and push on the blend stack. */
-	TObjectPtr<const UCameraRigAsset> CameraRig;
-};
-
-/**
- * Evaluator for a persistent blend stack.
- */
-class FPersistentBlendStackCameraNodeEvaluator 
-	: public FBlendStackCameraNodeEvaluator
-{
-	UE_DECLARE_CAMERA_NODE_EVALUATOR_EX(GAMEPLAYCAMERAS_API, FPersistentBlendStackCameraNodeEvaluator, FBlendStackCameraNodeEvaluator)
-
-public:
-
-	/** Insert a new camera rig onto the blend stack. */
-	void Insert(const FBlendStackCameraInsertParams& Params);
-
-	/** Remove an existing camera rig from the blend stack. */
-	void Remove(const FBlendStackCameraRemoveParams& Params);
-
-protected:
-
-	// FCameraNodeEvaluator interface.
-	virtual void OnRun(const FCameraNodeEvaluationParams& Params, FCameraNodeEvaluationResult& OutResult) override;
-
-private:
-
-	// Update methods.
-	void InternalUpdate(TArrayView<FResolvedEntry> ResolvedEntries, const FCameraNodeEvaluationParams& Params, FCameraNodeEvaluationResult& OutResult);
-
-	// Utility functions for finding an appropriate transition.
-	const UCameraRigTransition* FindTransition(const FBlendStackCameraPushParams& Params) const;
-};
-
 #if UE_GAMEPLAY_CAMERAS_DEBUG
 
 class FBlendStackSummaryCameraDebugBlock : public FCameraDebugBlock
diff --git a/Engine/Plugins/Cameras/GameplayCameras/Source/GameplayCameras/Public/Core/PersistentBlendStackCameraNode.h b/Engine/Plugins/Cameras/GameplayCameras/Source/GameplayCameras/Public/Core/PersistentBlendStackCameraNode.h
new file mode 100644
--- /dev/null
+++ b/Engine/Plugins/Cameras/GameplayCameras/Source/GameplayCameras/Public/Core/PersistentBlendStackCameraNode.h
@@ -0,0 +1,73 @@
+// Copyright Epic Games, Inc. All Rights Reserved.
+
+#pragma once
+
+#include "Core/BlendStackCameraNode.h"
+
+namespace UE::Cameras
+{
+
+/**
+ * Parameter structure for inserting a camera rig into a persistent blend stack.
+ */
+struct FBlendStackCameraInsertParams
+{
+	/** The evaluation context within which a camera rig's node tree should run. */
+	TSharedPtr<const FCameraEvaluationContext> EvaluationContext;
+
+	/** The source camera rig asset to instantiate and push on the blend stack. */
+	TObjectPtr<const UCameraRigAsset> CameraRig;
+
+	/** A transition to use, instead of looking one up. */
+	TObjectPtr<const UCameraRigTransition> TransitionOverride;
+
+	/** Whether to force insert a new instance of the camera rig, even if there is already one in the stack. */
+	bool bForceInsert = false;
+};
+
+/**
+ * Parameter structure for removing a camera rig from a persistent blend stack.
+ */
+struct FBlendStackCameraRemoveParams
+{
+	/** The evaluation context within which a camera rig's node tree should run. */
+	TSharedPtr<const FCameraEvaluationContext> EvaluationContext;
+
+	/** The source camera rig asset to instantiate and push on the blend stack. */
+	TObjectPtr<const UCameraRigAsset> CameraRig;
+};
+
+/**
+ * Evaluator for a persistent blend stack, i.e. a blend stack in which camera rigs blend additively on top of 
+ * each other, but without automatically popping out any fully blended-out entries. 
+ * This is a stack suitable for a "camera modifier stack" of sorts.
+ */
+class FPersistentBlendStackCameraNodeEvaluator 
+	: public FBlendStackCameraNodeEvaluator
+{
+	UE_DECLARE_CAMERA_NODE_EVALUATOR_EX(GAMEPLAYCAMERAS_API, FPersistentBlendStackCameraNodeEvaluator, FBlendStackCameraNodeEvaluator)
+
+public:
+
+	/** Insert a new camera rig onto the blend stack. */
+	void Insert(const FBlendStackCameraInsertParams& Params);
+
+	/** Remove an existing camera rig from the blend stack. */
+	void Remove(const FBlendStackCameraRemoveParams& Params);
+
+protected:
+
+	// FCameraNodeEvaluator interface.
+	virtual void OnRun(const FCameraNodeEvaluationParams& Params, FCameraNodeEvaluationResult& OutResult) override;
+
+private:
+
+	// Update methods.
+	void InternalUpdate(TArrayView<FResolvedEntry> ResolvedEntries, const FCameraNodeEvaluationParams& Params, FCameraNodeEvaluationResult& OutResult);
+
+	// Utility functions for finding an appropriate transition.
+	const UCameraRigTransition* FindTransition(const FBlendStackCameraInsertParams& Params) const;
+};
+
+}  // namespace UE::Cameras
+
diff --git a/Engine/Plugins/Cameras/GameplayCameras/Source/GameplayCameras/Public/Core/TransientBlendStackCameraNode.h b/Engine/Plugins/Cameras/GameplayCameras/Source/GameplayCameras/Public/Core/TransientBlendStackCameraNode.h
new file mode 100644
--- /dev/null
+++ b/Engine/Plugins/Cameras/GameplayCameras/Source/GameplayCameras/Public/Core/TransientBlendStackCameraNode.h
@@ -0,0 +1,113 @@
+// Copyright Epic Games, Inc. All Rights Reserved.
+
+#pragma once
+
+#include "Core/BlendStackCameraNode.h"
+
+namespace UE::Cameras
+{
+
+/**
+ * Parameter structure for pushing a camera rig onto a transient blend stack.
+ */
+struct FBlendStackCameraPushParams
+{
+	/** The evaluation context within which a camera rig's node tree should run. */
+	TSharedPtr<const FCameraEvaluationContext> EvaluationContext;
+
+	/** The source camera rig asset to instantiate and push on the blend stack. */
+	TObjectPtr<const UCameraRigAsset> CameraRig;
+
+	/** A transition to use, instead of looking one up. */
+	TObjectPtr<const UCameraRigTransition> TransitionOverride;
+
+	/** Whether to force pushing a new instance of the camera rig, even if it is currently active. */
+	bool bForcePush = false;
+};
+
+/**
+ * Parameter structure for freezing a camera rig inside a transient blend stack.
+ */
+struct FBlendStackCameraFreezeParams
+{
+	/** The evaluation context within which a camera rig's node tree is running. */
+	TSharedPtr<const FCameraEvaluationContext> EvaluationContext;
+
+	/** The source camera rig asset that is running. */
+	TObjectPtr<const UCameraRigAsset> CameraRig;
+};
+
+/**
+ * Evaluator for a transient blend stack, i.e. a blend stack where camera rigs evaluate together,
+ * blending together and blend out those below them, which are then automatically popped out of
+ * the stack.
+ * This is a stack suitable for "camera modes" of sorts.
+ */
+class FTransientBlendStackCameraNodeEvaluator 
+	: public FBlendStackCameraNodeEvaluator
+{
+
+	UE_DECLARE_CAMERA_NODE_EVALUATOR_EX(GAMEPLAYCAMERAS_API, FTransientBlendStackCameraNodeEvaluator, FBlendStackCameraNodeEvaluator)
+
+public:
+
+	/** Push a new camera rig onto the blend stack. */
+	void Push(const FBlendStackCameraPushParams& Params);
+
+	/** Freeze a camera rig. */
+	void Freeze(const FBlendStackCameraFreezeParams& Params);
+
+	/** Freeze all camera rigs that belong to a given evaluation context. */
+	void FreezeAll(TSharedPtr<FCameraEvaluationContext> EvaluationContext);
+
+protected:
+
+	// FCameraNodeEvaluator interface.
+	virtual void OnRun(const FCameraNodeEvaluationParams& Params, FCameraNodeEvaluationResult& OutResult) override;
+	virtual void OnSerialize(const FCameraNodeEvaluatorSerializeParams& Params, FArchive& Ar) override;
+
+private:
+
+	// Update methods.
+	void InternalPreBlendPrepare(TArrayView<FResolvedEntry> ResolvedEntries, const FCameraNodeEvaluationParams& Params, FCameraNodeEvaluationResult& OutResult);
+	void InternalPreBlendExecute(TArrayView<FResolvedEntry> ResolvedEntries, const FCameraNodeEvaluationParams& Params, FCameraNodeEvaluationResult& OutResult);
+	void InternalUpdate(TArrayView<FResolvedEntry> ResolvedEntries, const FCameraNodeEvaluationParams& Params, FCameraNodeEvaluationResult& OutResult);
+	void InternalPostBlendExecute(TArrayView<FResolvedEntry> ResolvedEntries, const FCameraNodeEvaluationParams& Params, FCameraNodeEvaluationResult& OutResult);
+	void InternalRunFinished(FCameraNodeEvaluationResult& OutResult);
+
+	// Utility functions for finding an appropriate transition.
+	const UCameraRigTransition* FindTransition(const FBlendStackCameraPushParams& Params) const;
+	const UCameraRigTransition* FindTransition(
+			TArrayView<const TObjectPtr<UCameraRigTransition>> Transitions, 
+			const UCameraRigAsset* FromCameraRig, const UCameraAsset* FromCameraAsset, bool bFromFrozen,
+			const UCameraRigAsset* ToCameraRig, const UCameraAsset* ToCameraAsset) const;
+
+	void PushNewEntry(const FBlendStackCameraPushParams& Params, const UCameraRigTransition* Transition);
+	void PushMergedEntry(const FBlendStackCameraPushParams& Params, const UCameraRigTransition* Transition);
+
+private:
+
+	/** Extra blending-related info for each camera rig entry. */
+	struct FCameraRigEntryExtraInfo
+	{
+		/** Result without pre-blending. */
+		FCameraNodeEvaluationResult UnblendedResult;
+
+		/** Whether input slots were run. */
+		bool bInputRunThisFrame = false;
+		/** Whether the blend node was run. */
+		bool bBlendRunThisFrame = false;
+		/** Whether this camera rig has any parameters to pre-blend. */
+		bool bHasPreBlendedParameters = false;
+		/** Whether this camera rig's pre-blend was full and finished this frame. */
+		bool bIsPreBlendFull = false;
+	};
+
+	TArray<FCameraRigEntryExtraInfo> EntryExtraInfos;
+
+	/** Variable table for pre-blending. */
+	FCameraVariableTable PreBlendVariableTable;
+};
+
+}  // namespace UE::Cameras
+
