Cameras: fix const-casting for objects using the camera system RTTI.

diff --git a/Engine/Plugins/Cameras/GameplayCameras/Source/GameplayCameras/Public/Core/CameraObjectRtti.h b/Engine/Plugins/Cameras/GameplayCameras/Source/GameplayCameras/Public/Core/CameraObjectRtti.h
--- a/Engine/Plugins/Cameras/GameplayCameras/Source/GameplayCameras/Public/Core/CameraObjectRtti.h
+++ b/Engine/Plugins/Cameras/GameplayCameras/Source/GameplayCameras/Public/Core/CameraObjectRtti.h
@@ -148,9 +148,9 @@ private:
 		virtual bool IsKindOf(const FCameraObjectTypeID& InTypeID) const { return InTypeID == ClassName::PrivateTypeID; }\
 		template<typename Type> bool IsKindOf() const { return IsKindOf(Type::StaticTypeID()); }\
 		template<typename Type> Type* CastThis() { return IsKindOf<Type>() ? static_cast<Type*>(this) : nullptr; }\
-		template<typename Type> const Type* CastThis() const { return IsKindOf<Type>() ? static_cast<Type*>(this) : nullptr; }\
+		template<typename Type> const Type* CastThis() const { return IsKindOf<Type>() ? static_cast<const Type*>(this) : nullptr; }\
 		template<typename Type> Type* CastThisChecked() { check(IsKindOf<Type>()); return static_cast<Type*>(this); }\
-		template<typename Type> const Type* CastThisChecked() const { check(IsKindOf<Type>()); return static_cast<Type*>(this); }\
+		template<typename Type> const Type* CastThisChecked() const { check(IsKindOf<Type>()); return static_cast<const Type*>(this); }\
 	private:\
 		ApiDeclSpec static const ::UE::Cameras::TCameraObjectTypeID<ClassName> PrivateTypeID;
 
