Cameras: add camera rig instance IDs

These IDs lets us uniquely identify a running, or previously running, camera rig instance.

diff --git a/Engine/Plugins/Cameras/GameplayCameras/Source/GameplayCameras/Private/Core/BlendStackCameraNode.cpp b/Engine/Plugins/Cameras/GameplayCameras/Source/GameplayCameras/Private/Core/BlendStackCameraNode.cpp
--- a/Engine/Plugins/Cameras/GameplayCameras/Source/GameplayCameras/Private/Core/BlendStackCameraNode.cpp
+++ b/Engine/Plugins/Cameras/GameplayCameras/Source/GameplayCameras/Private/Core/BlendStackCameraNode.cpp
@@ -123,6 +123,7 @@ bool FBlendStackCameraNodeEvaluator::InitializeEntry(
 	RootEvaluator->Initialize(InitParams, NewEntry.Result);
 
 	// Wrap up!
+	NewEntry.EntryID = FBlendStackEntryID(NextEntryID++);
 	NewEntry.EvaluationContext = EvaluationContext;
 	NewEntry.CameraRig = CameraRig;
 	NewEntry.RootNode = EntryRootNode;
@@ -133,6 +134,14 @@ bool FBlendStackCameraNodeEvaluator::InitializeEntry(
 	return true;
 }
 
+int32 FBlendStackCameraNodeEvaluator::IndexOfEntry(const FBlendStackEntryID EntryID) const
+{
+	return Entries.IndexOfByPredicate([EntryID](const FCameraRigEntry& Item)
+			{
+				return Item.EntryID == EntryID;
+			});
+}
+
 void FBlendStackCameraNodeEvaluator::FreezeEntry(FCameraRigEntry& Entry)
 {
 	// Deallocate our node evaluators and clear any pointers we kept to them.
@@ -157,6 +166,7 @@ FCameraRigEvaluationInfo FBlendStackCameraNodeEvaluator::GetActiveCameraRigEvalu
 	{
 		const FCameraRigEntry& ActiveEntry = Entries[0];
 		FCameraRigEvaluationInfo Info(
+				FCameraRigInstanceID::FromBlendStackEntryID(ActiveEntry.EntryID, Layer),
 				ActiveEntry.EvaluationContext.Pin(),
 				ActiveEntry.CameraRig, 
 				&ActiveEntry.Result,
@@ -166,6 +176,23 @@ FCameraRigEvaluationInfo FBlendStackCameraNodeEvaluator::GetActiveCameraRigEvalu
 	return FCameraRigEvaluationInfo();
 }
 
+FCameraRigEvaluationInfo FBlendStackCameraNodeEvaluator::GetCameraRigEvaluationInfo(FBlendStackEntryID EntryID) const
+{
+	const int32 EntryIndex = IndexOfEntry(EntryID);
+	if (EntryIndex != INDEX_NONE)
+	{
+		const FCameraRigEntry& Entry = Entries[EntryIndex];
+		FCameraRigEvaluationInfo Info(
+				FCameraRigInstanceID::FromBlendStackEntryID(Entry.EntryID, Layer),
+				Entry.EvaluationContext.Pin(),
+				Entry.CameraRig, 
+				&Entry.Result,
+				Entry.RootEvaluator ? Entry.RootEvaluator->GetRootEvaluator() : nullptr);
+		return Info;
+	}
+	return FCameraRigEvaluationInfo();
+}
+
 FCameraNodeEvaluatorChildrenView FBlendStackCameraNodeEvaluator::OnGetChildren()
 {
 	FCameraNodeEvaluatorChildrenView View;
@@ -182,6 +209,9 @@ FCameraNodeEvaluatorChildrenView FBlendStackCameraNodeEvaluator::OnGetChildren()
 void FBlendStackCameraNodeEvaluator::OnInitialize(const FCameraNodeEvaluatorInitializeParams& Params, FCameraNodeEvaluationResult& OutResult)
 {
 	OwningEvaluator = Params.Evaluator;
+
+	const UBlendStackCameraNode* BlendStack = GetCameraNodeAs<UBlendStackCameraNode>();
+	Layer = BlendStack->Layer;
 }
 
 void FBlendStackCameraNodeEvaluator::ResolveEntries(TArray<FResolvedEntry>& OutResolvedEntries)
@@ -413,6 +443,7 @@ void FBlendStackCameraNodeEvaluator::BroadcastCameraRigEvent(EBlendStackCameraRi
 	Event.EventType = EventType;
 	Event.BlendStackEvaluator = this;
 	Event.CameraRigInfo = FCameraRigEvaluationInfo(
+			FCameraRigInstanceID::FromBlendStackEntryID(Entry.EntryID, Layer),
 			Entry.EvaluationContext.Pin(),
 			Entry.CameraRig,
 			&Entry.Result,
diff --git a/Engine/Plugins/Cameras/GameplayCameras/Source/GameplayCameras/Private/Core/DefaultRootCameraNode.cpp b/Engine/Plugins/Cameras/GameplayCameras/Source/GameplayCameras/Private/Core/DefaultRootCameraNode.cpp
--- a/Engine/Plugins/Cameras/GameplayCameras/Source/GameplayCameras/Private/Core/DefaultRootCameraNode.cpp
+++ b/Engine/Plugins/Cameras/GameplayCameras/Source/GameplayCameras/Private/Core/DefaultRootCameraNode.cpp
@@ -19,11 +19,12 @@ namespace UE::Cameras::Private
 
 TObjectPtr<UBlendStackCameraNode> CreateBlendStack(
 		UObject* This, const FObjectInitializer& ObjectInit,
-		const FName& Name, ECameraBlendStackType BlendStackType)
+		const FName& Name, ECameraBlendStackType BlendStackType, ECameraRigLayer Layer)
 {
 	TObjectPtr<UBlendStackCameraNode> NewBlendStack = ObjectInit.CreateDefaultSubobject<UBlendStackCameraNode>(
 			This, Name);
 	NewBlendStack->BlendStackType = BlendStackType;
+	NewBlendStack->Layer = Layer;
 	return NewBlendStack;
 }
 
@@ -34,10 +35,10 @@ UDefaultRootCameraNode::UDefaultRootCameraNode(const FObjectInitializer& ObjectI
 {
 	using namespace UE::Cameras::Private;
 
-	BaseLayer = CreateBlendStack(this, ObjectInit, TEXT("BaseLayer"), ECameraBlendStackType::AdditivePersistent);
-	MainLayer = CreateBlendStack(this, ObjectInit, TEXT("MainLayer"), ECameraBlendStackType::IsolatedTransient);
-	GlobalLayer = CreateBlendStack(this, ObjectInit, TEXT("GlobalLayer"), ECameraBlendStackType::AdditivePersistent);
-	VisualLayer = CreateBlendStack(this, ObjectInit, TEXT("VisualLayer"), ECameraBlendStackType::AdditivePersistent);
+	BaseLayer = CreateBlendStack(this, ObjectInit, TEXT("BaseLayer"), ECameraBlendStackType::AdditivePersistent, ECameraRigLayer::Base);
+	MainLayer = CreateBlendStack(this, ObjectInit, TEXT("MainLayer"), ECameraBlendStackType::IsolatedTransient, ECameraRigLayer::Main);
+	GlobalLayer = CreateBlendStack(this, ObjectInit, TEXT("GlobalLayer"), ECameraBlendStackType::AdditivePersistent, ECameraRigLayer::Global);
+	VisualLayer = CreateBlendStack(this, ObjectInit, TEXT("VisualLayer"), ECameraBlendStackType::AdditivePersistent, ECameraRigLayer::Visual);
 }
 
 FCameraNodeEvaluatorPtr UDefaultRootCameraNode::OnBuildEvaluator(FCameraNodeEvaluatorBuilder& Builder) const
@@ -81,7 +82,7 @@ void FDefaultRootCameraNodeEvaluator::OnRun(const FCameraNodeEvaluationParams& P
 	VisualLayer->Run(Params, OutResult);
 }
 
-void FDefaultRootCameraNodeEvaluator::OnActivateCameraRig(const FActivateCameraRigParams& Params)
+FCameraRigInstanceID FDefaultRootCameraNodeEvaluator::OnActivateCameraRig(const FActivateCameraRigParams& Params)
 {
 	if (Params.Layer == ECameraRigLayer::Main)
 	{
@@ -90,7 +91,9 @@ void FDefaultRootCameraNodeEvaluator::OnActivateCameraRig(const FActivateCameraR
 		PushParams.CameraRig = Params.CameraRig;
 		PushParams.TransitionOverride = Params.TransitionOverride;
 		PushParams.bForcePush = Params.bForceActivate;
-		MainLayer->Push(PushParams);
+		
+		const FBlendStackEntryID EntryID = MainLayer->Push(PushParams);
+		return FCameraRigInstanceID::FromBlendStackEntryID(EntryID, ECameraRigLayer::Main);
 	}
 	else
 	{
@@ -114,16 +117,22 @@ void FDefaultRootCameraNodeEvaluator::OnActivateCameraRig(const FActivateCameraR
 			InsertParams.CameraRig = Params.CameraRig;
 			InsertParams.TransitionOverride = Params.TransitionOverride;
 			InsertParams.bForceInsert = Params.bForceActivate;
-			TargetLayer->Insert(InsertParams);
+
+			const FBlendStackEntryID EntryID = TargetLayer->Insert(InsertParams);
+			return FCameraRigInstanceID::FromBlendStackEntryID(EntryID, Params.Layer);
 		}
 	}
+
+	return FCameraRigInstanceID();
 }
 
 void FDefaultRootCameraNodeEvaluator::OnDeactivateCameraRig(const FDeactivateCameraRigParams& Params)
 {
-	if (Params.Layer == ECameraRigLayer::Main)
+	ECameraRigLayer Layer = (Params.InstanceID.IsValid() ? Params.InstanceID.GetLayer() : Params.Layer);
+	if (Layer == ECameraRigLayer::Main)
 	{
 		FBlendStackCameraFreezeParams FreezeParams;
+		FreezeParams.EntryID = Params.InstanceID.ToBlendStackEntryID();
 		FreezeParams.CameraRig = Params.CameraRig;
 		FreezeParams.EvaluationContext = Params.EvaluationContext;
 		MainLayer->Freeze(FreezeParams);
@@ -131,7 +140,7 @@ void FDefaultRootCameraNodeEvaluator::OnDeactivateCameraRig(const FDeactivateCam
 	else
 	{
 		FPersistentBlendStackCameraNodeEvaluator* TargetLayer = nullptr;
-		switch (Params.Layer)
+		switch (Layer)
 		{
 			case ECameraRigLayer::Base:
 				TargetLayer = BaseLayer;
@@ -146,6 +155,7 @@ void FDefaultRootCameraNodeEvaluator::OnDeactivateCameraRig(const FDeactivateCam
 		if (ensure(TargetLayer))
 		{
 			FBlendStackCameraRemoveParams RemoveParams;
+			RemoveParams.EntryID = Params.InstanceID.ToBlendStackEntryID();
 			RemoveParams.EvaluationContext = Params.EvaluationContext;
 			RemoveParams.CameraRig = Params.CameraRig;
 			TargetLayer->Remove(RemoveParams);
@@ -158,6 +168,27 @@ void FDefaultRootCameraNodeEvaluator::OnGetActiveCameraRigInfo(FCameraRigEvaluat
 	OutCameraRigInfo = MainLayer->GetActiveCameraRigEvaluationInfo();
 }
 
+void FDefaultRootCameraNodeEvaluator::OnGetCameraRigInfo(const FCameraRigInstanceID InstanceID, FCameraRigEvaluationInfo& OutCameraRigInfo) const
+{
+	FPersistentBlendStackCameraNodeEvaluator* TargetLayer = nullptr;
+	switch (InstanceID.GetLayer())
+	{
+		case ECameraRigLayer::Base:
+			TargetLayer = BaseLayer;
+			break;
+		case ECameraRigLayer::Global:
+			TargetLayer = GlobalLayer;
+			break;
+		case ECameraRigLayer::Visual:
+			TargetLayer = VisualLayer;
+			break;
+	}
+	if (ensure(TargetLayer))
+	{
+		OutCameraRigInfo = TargetLayer->GetCameraRigEvaluationInfo(InstanceID.ToBlendStackEntryID());
+	}
+}
+
 void FDefaultRootCameraNodeEvaluator::OnBuildSingleCameraRigHierarchy(const FSingleCameraRigHierarchyBuildParams& Params, FCameraNodeEvaluatorHierarchy& OutHierarchy)
 {
 	OutHierarchy.Build(BaseLayer);
diff --git a/Engine/Plugins/Cameras/GameplayCameras/Source/GameplayCameras/Private/Core/PersistentBlendStackCameraNode.cpp b/Engine/Plugins/Cameras/GameplayCameras/Source/GameplayCameras/Private/Core/PersistentBlendStackCameraNode.cpp
--- a/Engine/Plugins/Cameras/GameplayCameras/Source/GameplayCameras/Private/Core/PersistentBlendStackCameraNode.cpp
+++ b/Engine/Plugins/Cameras/GameplayCameras/Source/GameplayCameras/Private/Core/PersistentBlendStackCameraNode.cpp
@@ -12,7 +12,7 @@ namespace UE::Cameras
 
 UE_DEFINE_CAMERA_NODE_EVALUATOR(FPersistentBlendStackCameraNodeEvaluator)
 
-void FPersistentBlendStackCameraNodeEvaluator::Insert(const FBlendStackCameraInsertParams& Params)
+FBlendStackEntryID FPersistentBlendStackCameraNodeEvaluator::Insert(const FBlendStackCameraInsertParams& Params)
 {
 	// See if we already have this camera rig and evaluation context in the stack.
 	if (!Params.bForceInsert)
@@ -23,7 +23,7 @@ void FPersistentBlendStackCameraNodeEvaluator::Insert(const FBlendStackCameraIns
 					Entry.CameraRig == Params.CameraRig &&
 					Entry.EvaluationContext == Params.EvaluationContext)
 			{
-				return;
+				return FBlendStackEntryID();
 			}
 		}
 	}
@@ -45,7 +45,7 @@ void FPersistentBlendStackCameraNodeEvaluator::Insert(const FBlendStackCameraIns
 			false);
 	if (!bInitialized)
 	{
-		return;
+		return FBlendStackEntryID();
 	}
 
 #if WITH_EDITOR
@@ -53,24 +53,47 @@ void FPersistentBlendStackCameraNodeEvaluator::Insert(const FBlendStackCameraIns
 #endif  // WITH_EDITOR
 
 	Entries.Add(MoveTemp(NewEntry));
+	const FBlendStackEntryID AddedEntryID(NewEntry.EntryID);
 
 	if (OnCameraRigEventDelegate.IsBound())
 	{
 		BroadcastCameraRigEvent(EBlendStackCameraRigEventType::Pushed, Entries.Last(), nullptr);
 	}
+
+	return AddedEntryID;
 }
 
 void FPersistentBlendStackCameraNodeEvaluator::Remove(const FBlendStackCameraRemoveParams& Params)
 {
-	for (int32 Index = Entries.Num() - 1; Index >= 0; --Index)
+	TArray<int32, TInlineAllocator<4>> EntriesToRemove;
+
+	if (Params.EntryID.IsValid())
+	{
+		// Remove the entry by ID.
+		const int32 EntryIndex = IndexOfEntry(Params.EntryID);
+		if (EntryIndex != INDEX_NONE)
+		{
+			EntriesToRemove.Add(EntryIndex);
+		}
+	}
+	else
 	{
-		FCameraRigEntry& Entry(Entries[Index]);
-		if (Entry.CameraRig == Params.CameraRig &&
-				Entry.EvaluationContext == Params.EvaluationContext)
+		// Remove any entries matching the given context and rig asset.
+		for (int32 Index = Entries.Num() - 1; Index >= 0; --Index)
 		{
-			PopEntry(Index);
+			FCameraRigEntry& Entry(Entries[Index]);
+			if (Entry.CameraRig == Params.CameraRig &&
+					Entry.EvaluationContext == Params.EvaluationContext)
+			{
+				EntriesToRemove.Add(Index);
+			}
 		}
 	}
+
+	for (int32 Index : EntriesToRemove)
+	{
+		PopEntry(Index);
+	}
 }
 
 void FPersistentBlendStackCameraNodeEvaluator::OnRun(const FCameraNodeEvaluationParams& Params, FCameraNodeEvaluationResult& OutResult)
diff --git a/Engine/Plugins/Cameras/GameplayCameras/Source/GameplayCameras/Private/Core/RootCameraNode.cpp b/Engine/Plugins/Cameras/GameplayCameras/Source/GameplayCameras/Private/Core/RootCameraNode.cpp
--- a/Engine/Plugins/Cameras/GameplayCameras/Source/GameplayCameras/Private/Core/RootCameraNode.cpp
+++ b/Engine/Plugins/Cameras/GameplayCameras/Source/GameplayCameras/Private/Core/RootCameraNode.cpp
@@ -16,9 +16,9 @@ void FRootCameraNodeEvaluator::OnInitialize(const FCameraNodeEvaluatorInitialize
 	OwningEvaluator = Params.Evaluator;
 }
 
-void FRootCameraNodeEvaluator::ActivateCameraRig(const FActivateCameraRigParams& Params)
+FCameraRigInstanceID FRootCameraNodeEvaluator::ActivateCameraRig(const FActivateCameraRigParams& Params)
 {
-	OnActivateCameraRig(Params);
+	return OnActivateCameraRig(Params);
 }
 
 void FRootCameraNodeEvaluator::DeactivateCameraRig(const FDeactivateCameraRigParams& Params)
@@ -31,6 +31,11 @@ void FRootCameraNodeEvaluator::GetActiveCameraRigInfo(FCameraRigEvaluationInfo&
 	OnGetActiveCameraRigInfo(OutCameraRigInfo);
 }
 
+void FRootCameraNodeEvaluator::GetCameraRigInfo(const FCameraRigInstanceID InstanceID, FCameraRigEvaluationInfo& OutCameraRigInfo) const
+{
+	OnGetCameraRigInfo(InstanceID, OutCameraRigInfo);
+}
+
 bool FRootCameraNodeEvaluator::HasAnyActiveCameraRig() const
 {
 	FCameraRigEvaluationInfo RigInfo;
diff --git a/Engine/Plugins/Cameras/GameplayCameras/Source/GameplayCameras/Private/Core/TransientBlendStackCameraNode.cpp b/Engine/Plugins/Cameras/GameplayCameras/Source/GameplayCameras/Private/Core/TransientBlendStackCameraNode.cpp
--- a/Engine/Plugins/Cameras/GameplayCameras/Source/GameplayCameras/Private/Core/TransientBlendStackCameraNode.cpp
+++ b/Engine/Plugins/Cameras/GameplayCameras/Source/GameplayCameras/Private/Core/TransientBlendStackCameraNode.cpp
@@ -18,7 +18,7 @@ namespace UE::Cameras
 
 UE_DEFINE_CAMERA_NODE_EVALUATOR(FTransientBlendStackCameraNodeEvaluator)
 
-void FTransientBlendStackCameraNodeEvaluator::Push(const FBlendStackCameraPushParams& Params)
+FBlendStackEntryID FTransientBlendStackCameraNodeEvaluator::Push(const FBlendStackCameraPushParams& Params)
 {
 	bool bSearchedForTransition = false;
 	const UCameraRigTransition* Transition = nullptr;
@@ -33,7 +33,7 @@ void FTransientBlendStackCameraNodeEvaluator::Push(const FBlendStackCameraPushPa
 			// camera rig.
 			if (!Params.bForcePush && TopEntry.CameraRig == Params.CameraRig)
 			{
-				return;
+				return FBlendStackEntryID();
 			}
 
 			// See if we can merge the new camera rig onto the active camera rig.
@@ -42,7 +42,7 @@ void FTransientBlendStackCameraNodeEvaluator::Push(const FBlendStackCameraPushPa
 			if (!Params.bForcePush && Eligibility == ECameraRigMergingEligibility::Active)
 			{
 				// This camera rig is already the active one on the merged stack.
-				return;
+				return FBlendStackEntryID();
 			}
 
 			if (Eligibility == ECameraRigMergingEligibility::EligibleForMerge)
@@ -54,8 +54,7 @@ void FTransientBlendStackCameraNodeEvaluator::Push(const FBlendStackCameraPushPa
 
 				if (Transition && Transition->bAllowCameraRigMerging)
 				{
-					PushMergedEntry(Params, Transition);
-					return;
+					return PushMergedEntry(Params, Transition);
 				}
 			}
 		}
@@ -68,10 +67,10 @@ void FTransientBlendStackCameraNodeEvaluator::Push(const FBlendStackCameraPushPa
 		Transition = FindTransition(Params);
 	}
 
-	PushNewEntry(Params, Transition);
+	return PushNewEntry(Params, Transition);
 }
 
-void FTransientBlendStackCameraNodeEvaluator::PushNewEntry(const FBlendStackCameraPushParams& Params, const UCameraRigTransition* Transition)
+FBlendStackEntryID FTransientBlendStackCameraNodeEvaluator::PushNewEntry(const FBlendStackCameraPushParams& Params, const UCameraRigTransition* Transition)
 {
 	// Create the new root node to wrap the new camera rig's root node, and the specific
 	// blend node for this transition.
@@ -108,7 +107,7 @@ void FTransientBlendStackCameraNodeEvaluator::PushNewEntry(const FBlendStackCame
 			true);
 	if (!bInitialized)
 	{
-		return;
+		return FBlendStackEntryID();
 	}
 	
 	// Take a snapshot of the initial result for having base values during pre-blending.
@@ -132,6 +131,8 @@ void FTransientBlendStackCameraNodeEvaluator::PushNewEntry(const FBlendStackCame
 	AddPackageListeners(NewEntry);
 #endif  // WITH_EDITOR
 
+	const FBlendStackEntryID AddedEntryID(NewEntry.EntryID);
+
 	// Important: we need to move the new entry here because copying evaluator storage
 	// is disabled.
 	Entries.Add(MoveTemp(NewEntry));
@@ -141,9 +142,11 @@ void FTransientBlendStackCameraNodeEvaluator::PushNewEntry(const FBlendStackCame
 	{
 		BroadcastCameraRigEvent(EBlendStackCameraRigEventType::Pushed, Entries.Last(), UsedTransition);
 	}
+
+	return AddedEntryID;
 }
 
-void FTransientBlendStackCameraNodeEvaluator::PushMergedEntry(const FBlendStackCameraPushParams& PushParams, const UCameraRigTransition* Transition)
+FBlendStackEntryID FTransientBlendStackCameraNodeEvaluator::PushMergedEntry(const FBlendStackCameraPushParams& PushParams, const UCameraRigTransition* Transition)
 {
 	const UBlendCameraNode* Blend = Transition ? Transition->Blend : nullptr;
 
@@ -166,17 +169,33 @@ void FTransientBlendStackCameraNodeEvaluator::PushMergedEntry(const FBlendStackC
 #if WITH_EDITOR
 	AddPackageListeners(TopEntry);
 #endif
+
+	return TopEntry.EntryID;
 }
 
 void FTransientBlendStackCameraNodeEvaluator::Freeze(const FBlendStackCameraFreezeParams& Params)
 {
-	for (FCameraRigEntry& Entry : Entries)
+	if (Params.EntryID.IsValid())
 	{
-		if (!Entry.bIsFrozen && 
-				Entry.CameraRig == Params.CameraRig &&
-				Entry.EvaluationContext == Params.EvaluationContext)
+		// Freeze the entry by ID.
+		const int32 EntryIndex = IndexOfEntry(Params.EntryID);
+		if (EntryIndex != INDEX_NONE)
 		{
-			FreezeEntry(Entry);
+			FreezeEntry(Entries[EntryIndex]);
+		}
+	}
+	else
+	{
+
+		// Freeze any entries matching the given context and rig asset.
+		for (FCameraRigEntry& Entry : Entries)
+		{
+			if (!Entry.bIsFrozen && 
+					Entry.CameraRig == Params.CameraRig &&
+					Entry.EvaluationContext == Params.EvaluationContext)
+			{
+				FreezeEntry(Entry);
+			}
 		}
 	}
 }
diff --git a/Engine/Plugins/Cameras/GameplayCameras/Source/GameplayCameras/Public/Core/BlendStackCameraNode.h b/Engine/Plugins/Cameras/GameplayCameras/Source/GameplayCameras/Public/Core/BlendStackCameraNode.h
--- a/Engine/Plugins/Cameras/GameplayCameras/Source/GameplayCameras/Public/Core/BlendStackCameraNode.h
+++ b/Engine/Plugins/Cameras/GameplayCameras/Source/GameplayCameras/Public/Core/BlendStackCameraNode.h
@@ -2,12 +2,14 @@
 
 #pragma once
 
+#include "Core/BlendStackEntryID.h"
 #include "Core/CameraNode.h"
 #include "Core/CameraNodeEvaluator.h"
 #include "Core/CameraNodeEvaluatorHierarchy.h"
 #include "Core/CameraNodeEvaluatorStorage.h"
 #include "Core/CameraRigAsset.h"
 #include "Core/CameraRigEvaluationInfo.h"
+#include "Core/CameraRigInstanceID.h"
 #include "Debug/CameraDebugBlock.h"
 #include "IGameplayCamerasLiveEditListener.h"
 
@@ -72,11 +74,13 @@ protected:
 
 public:
 
-	/** 
-	 * The type of blend stack this should run as.
-	 */
+	/** The type of blend stack this should run as. */
 	UPROPERTY()
 	ECameraBlendStackType BlendStackType = ECameraBlendStackType::IsolatedTransient;
+
+	/** The layer that this blend stack represents, if any. */
+	UPROPERTY()
+	ECameraRigLayer Layer = ECameraRigLayer::None;
 };
 
 namespace UE::Cameras
@@ -103,6 +107,9 @@ public:
 	/** Returns information about the top (active) camera rig, if any. */
 	FCameraRigEvaluationInfo GetActiveCameraRigEvaluationInfo() const;
 
+	/** Returns information about a given camera rig, if any. */
+	FCameraRigEvaluationInfo GetCameraRigEvaluationInfo(FBlendStackEntryID EntryID) const;
+
 #if UE_GAMEPLAY_CAMERAS_DEBUG
 	FBlendStackCameraDebugBlock* BuildDetailedDebugBlock(const FCameraDebugBlockBuildParams& Params, FCameraDebugBlockBuilder& Builder);
 #endif  // UE_GAMEPLAY_CAMERAS_DEBUG
@@ -140,6 +147,8 @@ protected:
 		UBlendStackRootCameraNode* EntryRootNode,
 		bool bSetActiveResult);
 
+	int32 IndexOfEntry(const FBlendStackEntryID EntryID) const;
+
 	void FreezeEntry(FCameraRigEntry& Entry);
 
 	void PopEntry(int32 EntryIndex);
@@ -174,6 +183,8 @@ protected:
 
 	struct FCameraRigEntry
 	{
+		/** The ID for this entry. */
+		FBlendStackEntryID EntryID;
 		/** Evaluation context in which this entry runs. */
 		TWeakPtr<const FCameraEvaluationContext> EvaluationContext;
 		/** The camera rig asset that this entry runs. */
@@ -213,6 +224,12 @@ protected:
 	/** Entries in the blend stack. */
 	TArray<FCameraRigEntry> Entries;
 
+	/** Next entry ID to use. */
+	uint32 NextEntryID = 0;
+
+	/** The layer that this blend stack represents, if any. */
+	ECameraRigLayer Layer = ECameraRigLayer::None;
+
 	/** The delegate to invoke when an event occurs in this blend stack. */
 	FOnBlendStackCameraRigEvent OnCameraRigEventDelegate;
 
diff --git a/Engine/Plugins/Cameras/GameplayCameras/Source/GameplayCameras/Public/Core/BlendStackEntryID.h b/Engine/Plugins/Cameras/GameplayCameras/Source/GameplayCameras/Public/Core/BlendStackEntryID.h
new file mode 100644
--- /dev/null
+++ b/Engine/Plugins/Cameras/GameplayCameras/Source/GameplayCameras/Public/Core/BlendStackEntryID.h
@@ -0,0 +1,71 @@
+// Copyright Epic Games, Inc. All Rights Reserved.
+
+#pragma once
+
+#include "CoreTypes.h"
+#include "Serialization/Archive.h"
+
+struct FCameraRigInstanceID;
+
+namespace UE::Cameras
+{
+
+class FBlendStackCameraNodeEvaluator;
+
+/**
+ * An ID for an entry in a camera blend stack.
+ *
+ * Note about overflowing: the max value for a uint32 is 4294967295. We use it for an INVALID
+ * blend stack entry, so we have at most 4294967294 instances to go through before overflowing.
+ * If somehow we wanted to push, on average, one new camera rig on a blend stack every second,
+ * we would only overflow after ~136 years.
+ */
+struct FBlendStackEntryID
+{
+public:
+
+	FBlendStackEntryID() : Value(INVALID) {}
+
+	bool IsValid() const { return Value != INVALID; }
+
+	explicit operator bool() const { return IsValid(); }
+
+public:
+
+	friend bool operator==(FBlendStackEntryID A, FBlendStackEntryID B)
+	{
+		return A.Value == B.Value;
+	}
+
+	friend bool operator!=(FBlendStackEntryID A, FBlendStackEntryID B)
+	{
+		return !(A == B);
+	}
+
+	friend uint32 GetTypeHash(FBlendStackEntryID In)
+	{
+		return In.Value;
+	}
+
+	friend FArchive& operator<< (FArchive& Ar, FBlendStackEntryID& In)
+	{
+		Ar << In.Value;
+		return Ar;
+	}
+
+private:
+
+	FBlendStackEntryID(uint32 InValue) 
+		: Value(InValue)
+	{}
+
+	static const uint32 INVALID = uint32(-1);
+
+	uint32 Value;
+
+	friend struct ::FCameraRigInstanceID;
+	friend class FBlendStackCameraNodeEvaluator;
+};
+
+}  // namespace UE::Cameras
+
diff --git a/Engine/Plugins/Cameras/GameplayCameras/Source/GameplayCameras/Public/Core/CameraRigEvaluationInfo.h b/Engine/Plugins/Cameras/GameplayCameras/Source/GameplayCameras/Public/Core/CameraRigEvaluationInfo.h
--- a/Engine/Plugins/Cameras/GameplayCameras/Source/GameplayCameras/Public/Core/CameraRigEvaluationInfo.h
+++ b/Engine/Plugins/Cameras/GameplayCameras/Source/GameplayCameras/Public/Core/CameraRigEvaluationInfo.h
@@ -2,6 +2,7 @@
 
 #pragma once
 
+#include "Core/CameraRigInstanceID.h"
 #include "Templates/SharedPointer.h"
 #include "UObject/ObjectPtr.h"
 
@@ -22,6 +23,8 @@ struct FCameraNodeEvaluationResult;
  */
 struct GAMEPLAYCAMERAS_API FCameraRigEvaluationInfo
 {
+	/** The instance ID of this camera rig. */
+	FCameraRigInstanceID InstanceID;
 	/** The context inside which the evaluation occurs. */
 	TSharedPtr<const FCameraEvaluationContext> EvaluationContext;
 	/** The camera rig being evaluated. */
@@ -35,11 +38,13 @@ struct GAMEPLAYCAMERAS_API FCameraRigEvaluationInfo
 	{}
 
 	FCameraRigEvaluationInfo(
+			FCameraRigInstanceID InInstanceID,
 			TSharedPtr<const FCameraEvaluationContext> InEvaluationContext,
 			TObjectPtr<const UCameraRigAsset> InCameraRig,
 			const FCameraNodeEvaluationResult* InLastResult,
 			FCameraNodeEvaluator* InRootEvaluator)
-		: EvaluationContext(InEvaluationContext)
+		: InstanceID(InInstanceID)
+		, EvaluationContext(InEvaluationContext)
 		, CameraRig(InCameraRig)
 		, LastResult(InLastResult)
 		, RootEvaluator(InRootEvaluator)
diff --git a/Engine/Plugins/Cameras/GameplayCameras/Source/GameplayCameras/Public/Core/CameraRigInstanceID.h b/Engine/Plugins/Cameras/GameplayCameras/Source/GameplayCameras/Public/Core/CameraRigInstanceID.h
new file mode 100644
--- /dev/null
+++ b/Engine/Plugins/Cameras/GameplayCameras/Source/GameplayCameras/Public/Core/CameraRigInstanceID.h
@@ -0,0 +1,119 @@
+// Copyright Epic Games, Inc. All Rights Reserved.
+
+#pragma once
+
+#include "CoreTypes.h"
+#include "Core/BlendStackEntryID.h"
+#include "Kismet/BlueprintFunctionLibrary.h"
+
+#include "CameraRigInstanceID.generated.h"
+
+/**
+ * Defines evaluation layers for camera rigs.
+ */
+UENUM(BlueprintType)
+enum class ECameraRigLayer : uint8
+{
+	None UMETA(Hidden),
+	Base UMETA(DisplayName="Base Layer"),
+	Main UMETA(DisplayName="Main Layer"),
+	Global UMETA(DisplayName="Global Layer"),
+	Visual UMETA(DisplayName="Visual Layer")
+};
+ENUM_CLASS_FLAGS(ECameraRigLayer)
+
+/**
+ * A unique instance ID for a running, or previously running, camera rig.
+ */
+USTRUCT(BlueprintType)
+struct FCameraRigInstanceID
+{
+	GENERATED_BODY()
+
+public:
+
+	FCameraRigInstanceID() : Value(INVALID), Layer(ECameraRigLayer::None) {}
+
+	bool IsValid() const { return Value != INVALID; }
+
+	explicit operator bool() const { return IsValid(); }
+
+public:
+
+	friend bool operator==(FCameraRigInstanceID A, FCameraRigInstanceID B)
+	{
+		return A.Value == B.Value && A.Layer == B.Layer;
+	}
+
+	friend bool operator!=(FCameraRigInstanceID A, FCameraRigInstanceID B)
+	{
+		return !(A == B);
+	}
+
+	friend uint32 GetTypeHash(FCameraRigInstanceID In)
+	{
+		return HashCombineFast(In.Value, (uint32)In.Layer);
+	}
+
+	friend FArchive& operator<< (FArchive& Ar, FCameraRigInstanceID& In)
+	{
+		Ar << In.Value;
+		Ar << In.Layer;
+		return Ar;
+	}
+
+public:
+
+	/** Gets the layer this camera rig is, or was, running on. */
+	ECameraRigLayer GetLayer() const { return Layer; }
+
+	/** Gets a blend stack ID from this instance ID. */
+	UE::Cameras::FBlendStackEntryID ToBlendStackEntryID() const { return UE::Cameras::FBlendStackEntryID(Value); }
+
+	/** Creates an instance ID from a blend stack ID and a layer. */
+	static FCameraRigInstanceID FromBlendStackEntryID(UE::Cameras::FBlendStackEntryID EntryID, ECameraRigLayer Layer)
+	{
+		return FCameraRigInstanceID(EntryID, Layer);
+	}
+
+private:
+
+	FCameraRigInstanceID(uint32 InValue, ECameraRigLayer InLayer) 
+		: Value(InValue)
+		, Layer(InLayer)
+	{}
+
+	FCameraRigInstanceID(const UE::Cameras::FBlendStackEntryID InBlendStackID, ECameraRigLayer InLayer) 
+		: Value(InBlendStackID.Value)
+		, Layer(InLayer)
+	{}
+
+	static const uint32 INVALID = uint32(-1);
+
+	UPROPERTY()
+	uint32 Value;
+
+	UPROPERTY()
+	ECameraRigLayer Layer;
+
+	friend class UE::Cameras::FBlendStackCameraNodeEvaluator;
+};
+
+/**
+ * Blueprint functions for camera rig instance IDs.
+ */
+UCLASS(MinimalAPI)
+class UCameraRigInstanceFunctions : public UBlueprintFunctionLibrary
+{
+	GENERATED_BODY()
+
+public:
+
+	/** 
+	 * Whether the given camera rig instance ID is valid.
+	 * A valid ID doesn't necessarily correspond to a camera rig instance that is still running.
+	 */
+	UFUNCTION(BlueprintPure, Category="Camera")
+	static bool IsValid(FCameraRigInstanceID InstanceID) { return InstanceID.IsValid(); }
+};
+
diff --git a/Engine/Plugins/Cameras/GameplayCameras/Source/GameplayCameras/Public/Core/DefaultRootCameraNode.h b/Engine/Plugins/Cameras/GameplayCameras/Source/GameplayCameras/Public/Core/DefaultRootCameraNode.h
--- a/Engine/Plugins/Cameras/GameplayCameras/Source/GameplayCameras/Public/Core/DefaultRootCameraNode.h
+++ b/Engine/Plugins/Cameras/GameplayCameras/Source/GameplayCameras/Public/Core/DefaultRootCameraNode.h
@@ -65,9 +65,10 @@ protected:
 #endif  // UE_GAMEPLAY_CAMERAS_DEBUG
 
 	// FRootCameraNodeEvaluator interface.
-	virtual void OnActivateCameraRig(const FActivateCameraRigParams& Params) override;
+	virtual FCameraRigInstanceID OnActivateCameraRig(const FActivateCameraRigParams& Params) override;
 	virtual void OnDeactivateCameraRig(const FDeactivateCameraRigParams& Params) override;
 	virtual void OnGetActiveCameraRigInfo(FCameraRigEvaluationInfo& OutCameraRigInfo) const override;
+	virtual void OnGetCameraRigInfo(const FCameraRigInstanceID InstanceID, FCameraRigEvaluationInfo& OutCameraRigInfo) const override;
 	virtual void OnBuildSingleCameraRigHierarchy(const FSingleCameraRigHierarchyBuildParams& Params, FCameraNodeEvaluatorHierarchy& OutHierarchy) override;
 	virtual void OnRunSingleCameraRig(const FSingleCameraRigEvaluationParams& Params, FCameraNodeEvaluationResult& OutResult) override;
 
diff --git a/Engine/Plugins/Cameras/GameplayCameras/Source/GameplayCameras/Public/Core/PersistentBlendStackCameraNode.h b/Engine/Plugins/Cameras/GameplayCameras/Source/GameplayCameras/Public/Core/PersistentBlendStackCameraNode.h
--- a/Engine/Plugins/Cameras/GameplayCameras/Source/GameplayCameras/Public/Core/PersistentBlendStackCameraNode.h
+++ b/Engine/Plugins/Cameras/GameplayCameras/Source/GameplayCameras/Public/Core/PersistentBlendStackCameraNode.h
@@ -27,13 +27,18 @@ struct FBlendStackCameraInsertParams
 
 /**
  * Parameter structure for removing a camera rig from a persistent blend stack.
+ * The camera rig to remove is identified first by a given ID or, if no ID is given,
+ * by finding a camera rig matching the provided EvaluationContext and CameraRig.
  */
 struct FBlendStackCameraRemoveParams
 {
-	/** The evaluation context within which a camera rig's node tree should run. */
+	/** The ID of the blend stack entry to remove. */
+	FBlendStackEntryID EntryID;
+
+	/** The evaluation context within which the camera rig to remove is being run. */
 	TSharedPtr<const FCameraEvaluationContext> EvaluationContext;
 
-	/** The source camera rig asset to instantiate and push on the blend stack. */
+	/** The source camera rig asset used by the instanced to remove. */
 	TObjectPtr<const UCameraRigAsset> CameraRig;
 };
 
@@ -50,7 +55,7 @@ class FPersistentBlendStackCameraNodeEvaluator
 public:
 
 	/** Insert a new camera rig onto the blend stack. */
-	void Insert(const FBlendStackCameraInsertParams& Params);
+	FBlendStackEntryID Insert(const FBlendStackCameraInsertParams& Params);
 
 	/** Remove an existing camera rig from the blend stack. */
 	void Remove(const FBlendStackCameraRemoveParams& Params);
diff --git a/Engine/Plugins/Cameras/GameplayCameras/Source/GameplayCameras/Public/Core/RootCameraNode.h b/Engine/Plugins/Cameras/GameplayCameras/Source/GameplayCameras/Public/Core/RootCameraNode.h
--- a/Engine/Plugins/Cameras/GameplayCameras/Source/GameplayCameras/Public/Core/RootCameraNode.h
+++ b/Engine/Plugins/Cameras/GameplayCameras/Source/GameplayCameras/Public/Core/RootCameraNode.h
@@ -5,25 +5,13 @@
 #include "Core/CameraNode.h"
 #include "Core/CameraNodeEvaluator.h"
 #include "Core/CameraRigEvaluationInfo.h"
+#include "Core/CameraRigInstanceID.h"
 
 #include "RootCameraNode.generated.h"
 
 class UCameraRigAsset;
 class UCameraRigTransition;
 
-/**
- * Defines evaluation layers for camera rigs.
- */
-UENUM(BlueprintType)
-enum class ECameraRigLayer : uint8
-{
-	Base UMETA(DisplayName="Base Layer"),
-	Main UMETA(DisplayName="Main Layer"),
-	Global UMETA(DisplayName="Global Layer"),
-	Visual UMETA(DisplayName="Visual Layer")
-};
-ENUM_CLASS_FLAGS(ECameraRigLayer)
-
 /**
  * The base class for a camera node that can act as the root of the
  * camera system evaluation.
@@ -63,10 +51,13 @@ struct FActivateCameraRigParams
 };
 
 /**
- * Parameter structure for deaactivating a running camera rig.
+ * Parameter structure for deactivating a running camera rig.
  */
 struct FDeactivateCameraRigParams
 {
+	/** The instance to deactivate. */
+	FCameraRigInstanceID InstanceID;
+
 	/** The evaluation context in which the camera rig runs. */
 	TSharedPtr<const FCameraEvaluationContext> EvaluationContext;
 
@@ -110,15 +101,24 @@ class FRootCameraNodeEvaluator : public FCameraNodeEvaluator
 {
 public:
 
-	/** Activates a camera rig. */
-	void ActivateCameraRig(const FActivateCameraRigParams& Params);
+	/** 
+	 * Activates a camera rig.
+	 * What it means to activate a camera rig may differ depending on the layer it runs on.
+	 */
+	FCameraRigInstanceID ActivateCameraRig(const FActivateCameraRigParams& Params);
 
-	/** Deactivates a camera rig. */
+	/** 
+	 * Deactivates a camera rig. 
+	 * What it means to deactivate a camera rig may differ depending on the layer it runs on.
+	 */
 	void DeactivateCameraRig(const FDeactivateCameraRigParams& Params);
 
 	/** Gets information about the active camera rig in the main layer. */
 	void GetActiveCameraRigInfo(FCameraRigEvaluationInfo& OutCameraRigInfo) const;
 
+	/** Gets information about a specified camera rig. */
+	void GetCameraRigInfo(const FCameraRigInstanceID InstanceID, FCameraRigEvaluationInfo& OutCameraRigInfo) const;
+
 	/** Gets whether any camera rig is running. */
 	bool HasAnyActiveCameraRig() const;
 
@@ -148,7 +148,7 @@ protected:
 protected:
 
 	/** Activates a camera rig. */
-	virtual void OnActivateCameraRig(const FActivateCameraRigParams& Params) {}
+	virtual FCameraRigInstanceID OnActivateCameraRig(const FActivateCameraRigParams& Params) { return FCameraRigInstanceID(); }
 	
 	/** Deactivates a camera rig. */
 	virtual void OnDeactivateCameraRig(const FDeactivateCameraRigParams& Params) {}
@@ -156,6 +156,9 @@ protected:
 	/** Gets information about the active camera rig in the main layer. */
 	virtual void OnGetActiveCameraRigInfo(FCameraRigEvaluationInfo& OutCameraRigInfo) const {}
 
+	/** Gets information about a specified camera rig. */
+	virtual void OnGetCameraRigInfo(const FCameraRigInstanceID InstanceID, FCameraRigEvaluationInfo& OutCameraRigInfo) const {}
+
 	/* Builds the hierarchy of the system for a given single camera rig. */
 	virtual void OnBuildSingleCameraRigHierarchy(const FSingleCameraRigHierarchyBuildParams& Params, FCameraNodeEvaluatorHierarchy& OutHierarchy) {}
 
diff --git a/Engine/Plugins/Cameras/GameplayCameras/Source/GameplayCameras/Public/Core/TransientBlendStackCameraNode.h b/Engine/Plugins/Cameras/GameplayCameras/Source/GameplayCameras/Public/Core/TransientBlendStackCameraNode.h
--- a/Engine/Plugins/Cameras/GameplayCameras/Source/GameplayCameras/Public/Core/TransientBlendStackCameraNode.h
+++ b/Engine/Plugins/Cameras/GameplayCameras/Source/GameplayCameras/Public/Core/TransientBlendStackCameraNode.h
@@ -3,6 +3,7 @@
 #pragma once
 
 #include "Core/BlendStackCameraNode.h"
+#include "Core/BlendStackEntryID.h"
 
 namespace UE::Cameras
 {
@@ -30,6 +31,9 @@ struct FBlendStackCameraPushParams
  */
 struct FBlendStackCameraFreezeParams
 {
+	/** The ID of the blend stack entry to freeze. */
+	FBlendStackEntryID EntryID;
+
 	/** The evaluation context within which a camera rig's node tree is running. */
 	TSharedPtr<const FCameraEvaluationContext> EvaluationContext;
 
@@ -52,7 +56,7 @@ class FTransientBlendStackCameraNodeEvaluator
 public:
 
 	/** Push a new camera rig onto the blend stack. */
-	void Push(const FBlendStackCameraPushParams& Params);
+	FBlendStackEntryID Push(const FBlendStackCameraPushParams& Params);
 
 	/** Freeze a camera rig. */
 	void Freeze(const FBlendStackCameraFreezeParams& Params);
@@ -82,8 +86,8 @@ private:
 			const UCameraRigAsset* FromCameraRig, const UCameraAsset* FromCameraAsset, bool bFromFrozen,
 			const UCameraRigAsset* ToCameraRig, const UCameraAsset* ToCameraAsset) const;
 
-	void PushNewEntry(const FBlendStackCameraPushParams& Params, const UCameraRigTransition* Transition);
-	void PushMergedEntry(const FBlendStackCameraPushParams& Params, const UCameraRigTransition* Transition);
+	FBlendStackEntryID PushNewEntry(const FBlendStackCameraPushParams& Params, const UCameraRigTransition* Transition);
+	FBlendStackEntryID PushMergedEntry(const FBlendStackCameraPushParams& Params, const UCameraRigTransition* Transition);
 
 private:
 
