Cameras: base/global/visual stacks now support priorities and blend in/out

A "sort key" can be provided when inserting a camera rig inside a persistent/additive blend stack. It will be inserted at an appropriate place in the stack.

For blending in and out, enter/exit transitions are scanned as usual. Blending out must be handled as a special case however. A new API on blend nodes allows running some blends in reverse. If that's not supported, a "reverse blend" wrapper is created. Similarly, a new API lets a blend node initialize itself from an interrupted blend node. If that's not supported, an "interrupted blend" wrapper is created to "freeze" the interrupted blend.

diff --git a/Engine/Plugins/Cameras/GameplayCameras/Source/GameplayCameras/Private/Core/BlendCameraNode.cpp b/Engine/Plugins/Cameras/GameplayCameras/Source/GameplayCameras/Private/Core/BlendCameraNode.cpp
--- a/Engine/Plugins/Cameras/GameplayCameras/Source/GameplayCameras/Private/Core/BlendCameraNode.cpp
+++ b/Engine/Plugins/Cameras/GameplayCameras/Source/GameplayCameras/Private/Core/BlendCameraNode.cpp
@@ -19,5 +19,15 @@ void FBlendCameraNodeEvaluator::BlendResults(const FCameraNodeBlendParams& Param
 	OnBlendResults(Params, OutResult);
 }
 
+bool FBlendCameraNodeEvaluator::InitializeFromInterruption(const FCameraNodeBlendInterruptionParams& Params)
+{
+	return OnInitializeFromInterruption(Params);
+}
+
+bool FBlendCameraNodeEvaluator::SetReversed(bool bInReverse)
+{
+	return OnSetReversed(bInReverse);
+}
+
 }  // namespace UE::Cameras
 
diff --git a/Engine/Plugins/Cameras/GameplayCameras/Source/GameplayCameras/Private/Core/BlendStackRootCameraNode.cpp b/Engine/Plugins/Cameras/GameplayCameras/Source/GameplayCameras/Private/Core/BlendStackRootCameraNode.cpp
--- a/Engine/Plugins/Cameras/GameplayCameras/Source/GameplayCameras/Private/Core/BlendStackRootCameraNode.cpp
+++ b/Engine/Plugins/Cameras/GameplayCameras/Source/GameplayCameras/Private/Core/BlendStackRootCameraNode.cpp
@@ -99,6 +99,11 @@ void FBlendStackRootCameraNodeEvaluator::OnInitialize(const FCameraNodeEvaluator
 	}
 }
 
+void FBlendStackRootCameraNodeEvaluator::SetBlendEvaluator(FBlendCameraNodeEvaluator* InBlendEvaluator)
+{
+	BlendEvaluator = InBlendEvaluator;
+}
+
 ECameraRigMergingEligibility FBlendStackRootCameraNodeEvaluator::CompareCameraRigForMerging(const UCameraRigAsset* CameraRig) const
 {
 	const UCameraRigAsset* NewCameraRigPrefab = FindInnermostCameraRigPrefab(CameraRig);
diff --git a/Engine/Plugins/Cameras/GameplayCameras/Source/GameplayCameras/Private/Core/BlendStackRootCameraNode.h b/Engine/Plugins/Cameras/GameplayCameras/Source/GameplayCameras/Private/Core/BlendStackRootCameraNode.h
--- a/Engine/Plugins/Cameras/GameplayCameras/Source/GameplayCameras/Private/Core/BlendStackRootCameraNode.h
+++ b/Engine/Plugins/Cameras/GameplayCameras/Source/GameplayCameras/Private/Core/BlendStackRootCameraNode.h
@@ -68,6 +68,8 @@ public:
 	FBlendCameraNodeEvaluator* GetBlendEvaluator() const { return BlendEvaluator; }
 	FCameraNodeEvaluator* GetRootEvaluator() const { return RootEvaluator; }
 
+	void SetBlendEvaluator(FBlendCameraNodeEvaluator* InBlendEvaluator);
+
 	ECameraRigMergingEligibility CompareCameraRigForMerging(const UCameraRigAsset* CameraRig) const;
 
 	void MergeCameraRig(const FCameraNodeEvaluatorBuildParams& Params, const UCameraRigAsset* CameraRig, const UBlendCameraNode* Blend);
diff --git a/Engine/Plugins/Cameras/GameplayCameras/Source/GameplayCameras/Private/Core/DefaultRootCameraNode.cpp b/Engine/Plugins/Cameras/GameplayCameras/Source/GameplayCameras/Private/Core/DefaultRootCameraNode.cpp
--- a/Engine/Plugins/Cameras/GameplayCameras/Source/GameplayCameras/Private/Core/DefaultRootCameraNode.cpp
+++ b/Engine/Plugins/Cameras/GameplayCameras/Source/GameplayCameras/Private/Core/DefaultRootCameraNode.cpp
@@ -86,6 +86,7 @@ FCameraRigInstanceID FDefaultRootCameraNodeEvaluator::OnActivateCameraRig(const
 {
 	if (Params.Layer == ECameraRigLayer::Main)
 	{
+		ensure(Params.OrderKey == 0);
 		FBlendStackCameraPushParams PushParams;
 		PushParams.EvaluationContext = Params.EvaluationContext;
 		PushParams.CameraRig = Params.CameraRig;
@@ -116,6 +117,7 @@ FCameraRigInstanceID FDefaultRootCameraNodeEvaluator::OnActivateCameraRig(const
 			InsertParams.EvaluationContext = Params.EvaluationContext;
 			InsertParams.CameraRig = Params.CameraRig;
 			InsertParams.TransitionOverride = Params.TransitionOverride;
+			InsertParams.StackOrder = Params.OrderKey;
 			InsertParams.bForceInsert = Params.bForceActivate;
 
 			const FBlendStackEntryID EntryID = TargetLayer->Insert(InsertParams);
@@ -158,6 +160,8 @@ void FDefaultRootCameraNodeEvaluator::OnDeactivateCameraRig(const FDeactivateCam
 			RemoveParams.EntryID = Params.InstanceID.ToBlendStackEntryID();
 			RemoveParams.EvaluationContext = Params.EvaluationContext;
 			RemoveParams.CameraRig = Params.CameraRig;
+			RemoveParams.TransitionOverride = Params.TransitionOverride;
+			RemoveParams.bRemoveImmediately = Params.bDeactiveImmediately;
 			TargetLayer->Remove(RemoveParams);
 		}
 	}
diff --git a/Engine/Plugins/Cameras/GameplayCameras/Source/GameplayCameras/Private/Core/PersistentBlendStackCameraNode.cpp b/Engine/Plugins/Cameras/GameplayCameras/Source/GameplayCameras/Private/Core/PersistentBlendStackCameraNode.cpp
--- a/Engine/Plugins/Cameras/GameplayCameras/Source/GameplayCameras/Private/Core/PersistentBlendStackCameraNode.cpp
+++ b/Engine/Plugins/Cameras/GameplayCameras/Source/GameplayCameras/Private/Core/PersistentBlendStackCameraNode.cpp
@@ -2,10 +2,15 @@
 
 #include "Core/PersistentBlendStackCameraNode.h"
 
+#include "Algo/BinarySearch.h"
 #include "Core/BlendCameraNode.h"
 #include "Core/BlendStackCameraRigEvent.h"
 #include "Core/BlendStackRootCameraNode.h"
 #include "Core/CameraEvaluationContext.h"
+#include "Helpers/CameraRigTransitionFinder.h"
+#include "Nodes/Blends/InterruptedBlendCameraNode.h"
+#include "Nodes/Blends/PopBlendCameraNode.h"
+#include "Nodes/Blends/ReverseBlendCameraNode.h"
 
 namespace UE::Cameras
 {
@@ -17,23 +22,37 @@ FBlendStackEntryID FPersistentBlendStackCameraNodeEvaluator::Insert(const FBlend
 	// See if we already have this camera rig and evaluation context in the stack.
 	if (!Params.bForceInsert)
 	{
-		for (const FCameraRigEntry& Entry : Entries)
+		for (int32 Index = 0; Index < Entries.Num(); ++Index)
 		{
+			const FCameraRigEntry& Entry(Entries[Index]);
+			const FCameraRigEntryExtraInfo& EntryExtraInfo(EntryExtraInfos[Index]);
+
 			if (!Entry.bIsFrozen &&
 					Entry.CameraRig == Params.CameraRig &&
-					Entry.EvaluationContext == Params.EvaluationContext)
+					Entry.EvaluationContext == Params.EvaluationContext &&
+					EntryExtraInfo.StackOrder == Params.StackOrder)
 			{
 				return FBlendStackEntryID();
 			}
 		}
 	}
 
-	// TODO: add support for slot indices or something, to allow callers to specify a place in the stack.
 	UObject* Outer = const_cast<UObject*>((UObject*)GetCameraNode());
 	UBlendStackRootCameraNode* EntryRootNode = NewObject<UBlendStackRootCameraNode>(Outer, NAME_None);
 	{
 		EntryRootNode->RootNode = Params.CameraRig->RootNode;
-		// TODO: add support for blending in and out.
+
+		// Find a transition to blend in. If no transition is found, use a pop blend.
+		UBlendCameraNode* ModeBlend = nullptr;
+		if (const UCameraRigTransition* Transition = FindEnterTransition(Params))
+		{
+			ModeBlend = Transition->Blend;
+		}
+		if (!ModeBlend)
+		{
+			ModeBlend = NewObject<UPopBlendCameraNode>(EntryRootNode, NAME_None);
+		}
+		EntryRootNode->Blend = ModeBlend;
 	}
 
 	FCameraRigEntry NewEntry;
@@ -48,16 +67,25 @@ FBlendStackEntryID FPersistentBlendStackCameraNodeEvaluator::Insert(const FBlend
 		return FBlendStackEntryID();
 	}
 
+	FCameraRigEntryExtraInfo NewExtraInfo;
+	NewExtraInfo.StackOrder = Params.StackOrder;
+	NewExtraInfo.BlendStatus = EBlendStatus::BlendIn;
+
 #if WITH_EDITOR
 	AddPackageListeners(NewEntry);
 #endif  // WITH_EDITOR
 
-	Entries.Add(MoveTemp(NewEntry));
 	const FBlendStackEntryID AddedEntryID(NewEntry.EntryID);
 
+	ensure(Entries.Num() == EntryExtraInfos.Num());
+
+	const int32 AddedIndex = Algo::UpperBoundBy(EntryExtraInfos, Params.StackOrder, &FCameraRigEntryExtraInfo::StackOrder);
+	Entries.Insert(MoveTemp(NewEntry), AddedIndex);
+	EntryExtraInfos.Insert(MoveTemp(NewExtraInfo), AddedIndex);
+
 	if (OnCameraRigEventDelegate.IsBound())
 	{
-		BroadcastCameraRigEvent(EBlendStackCameraRigEventType::Pushed, Entries.Last(), nullptr);
+		BroadcastCameraRigEvent(EBlendStackCameraRigEventType::Pushed, Entries[AddedIndex], nullptr);
 	}
 
 	return AddedEntryID;
@@ -90,14 +118,79 @@ void FPersistentBlendStackCameraNodeEvaluator::Remove(const FBlendStackCameraRem
 		}
 	}
 
-	for (int32 Index : EntriesToRemove)
+	// If we need to remove the camera rigs immediately, simply pop out their entries.
+	if (Params.bRemoveImmediately)
+	{
+		for (int32 Index : EntriesToRemove)
+		{
+			PopEntry(Index);
+			EntryExtraInfos.RemoveAt(Index);
+		}
+	}
+	// Else, we need to start blending out these entries.
+	else
 	{
-		PopEntry(Index);
+		for (int32 Index : EntriesToRemove)
+		{
+			FCameraRigEntry& Entry(Entries[Index]);
+			FCameraRigEntryExtraInfo& EntryExtraInfo(EntryExtraInfos[Index]);
+			const UCameraRigTransition* Transition = FindExitTransition(Params, Entry);
+			if (Transition && Transition->Blend)
+			{
+				// Swap the blend-in evaluator on this entry with a blend-out one.
+				if (EntryExtraInfo.BlendStatus != EBlendStatus::BlendOut)
+				{
+					FCameraNodeEvaluatorBuilder BlendOutBuilder(Entry.EvaluatorStorage);
+					FCameraNodeEvaluatorBuildParams BlendOutBuildParams(BlendOutBuilder);
+					FBlendCameraNodeEvaluator* BlendOutEvaluator = BlendOutBuildParams.BuildEvaluatorAs<FBlendCameraNodeEvaluator>(Transition->Blend);
+
+					FCameraNodeEvaluatorInitializeParams BlendOutInitParams;
+					BlendOutInitParams.Evaluator = OwningEvaluator;
+					BlendOutInitParams.EvaluationContext = Entry.EvaluationContext.Pin();
+					BlendOutEvaluator->Initialize(BlendOutInitParams, Entry.Result);
+
+					// Reverse this blend so it plays as a blend-out. Also, see if we are going to 
+					// interrupt an ongoing blend-in... if so, give a chance for the blend-out to
+					// start at an "equivalent spot".
+					if (!BlendOutEvaluator->SetReversed(true))
+					{
+						BlendOutEvaluator = Entry.EvaluatorStorage.BuildEvaluator<FReverseBlendCameraNodeEvaluator>(BlendOutEvaluator);
+					}
+					if (EntryExtraInfo.BlendStatus == EBlendStatus::BlendIn)
+					{
+						FBlendCameraNodeEvaluator* OngoingBlend = Entry.RootEvaluator->GetBlendEvaluator();
+
+						FCameraNodeBlendInterruptionParams InterruptionParams;
+						InterruptionParams.InterruptedBlend = OngoingBlend;
+						if (!BlendOutEvaluator->InitializeFromInterruption(InterruptionParams))
+						{
+							BlendOutEvaluator = Entry.EvaluatorStorage.BuildEvaluator<FInterruptedBlendCameraNodeEvaluator>(BlendOutEvaluator, OngoingBlend);
+						}
+					}
+					// Note: neither the reverse or interrupted blends need initialization, but
+					// technically we're missing calling it on them.
+					Entry.RootEvaluator->SetBlendEvaluator(BlendOutEvaluator);
+
+					EntryExtraInfo.BlendStatus = EBlendStatus::BlendOut;
+					EntryExtraInfo.bIsBlendFinished = false;
+					EntryExtraInfo.bIsBlendFull = false;
+				}
+				// else: we were already blending out, so let this continue.
+			}
+			else
+			{
+				// No transition found... just cut.
+				PopEntry(Index);
+				EntryExtraInfos.RemoveAt(Index);
+			}
+		}
 	}
 }
 
 void FPersistentBlendStackCameraNodeEvaluator::OnRun(const FCameraNodeEvaluationParams& Params, FCameraNodeEvaluationResult& OutResult)
 {
+	ensure(Entries.Num() == EntryExtraInfos.Num());
+
 	// Validate our entries and resolve evaluation context weak pointers.
 	TArray<FResolvedEntry> ResolvedEntries;
 	ResolveEntries(ResolvedEntries);
@@ -114,9 +207,12 @@ void FPersistentBlendStackCameraNodeEvaluator::InternalUpdate(TArrayView<FResolv
 	constexpr ECameraVariableTableFilter VariableTableFilter = ECameraVariableTableFilter::ChangedOnly;
 	constexpr ECameraContextDataTableFilter ContextDataTableFilter = ECameraContextDataTableFilter::ChangedOnly;
 
+	TArray<int32, TInlineAllocator<4>> EntriesToRemove;
+
 	for (FResolvedEntry& ResolvedEntry : ResolvedEntries)
 	{
 		FCameraRigEntry& Entry(ResolvedEntry.Entry);
+		FCameraRigEntryExtraInfo& EntryExtraInfo(EntryExtraInfos[ResolvedEntry.EntryIndex]);
 
 		if (!Entry.bIsFrozen)
 		{
@@ -192,11 +288,30 @@ void FPersistentBlendStackCameraNodeEvaluator::InternalUpdate(TArrayView<FResolv
 				FCameraNodeBlendResult BlendResult(OutResult);
 
 				EntryBlendEvaluator->BlendResults(BlendParams, BlendResult);
+
+				EntryExtraInfo.bIsBlendFinished = BlendResult.bIsBlendFinished;
+				EntryExtraInfo.bIsBlendFull = BlendResult.bIsBlendFull;
 			}
 			else
 			{
 				OutResult.OverrideAll(CurResult);
 			}
+
+			// Update blend state.
+			if (EntryExtraInfo.BlendStatus == EBlendStatus::BlendIn)
+			{
+				if (EntryExtraInfo.bIsBlendFull && EntryExtraInfo.bIsBlendFinished)
+				{
+					EntryExtraInfo.BlendStatus = EBlendStatus::None;
+				}
+			}
+			else if (EntryExtraInfo.BlendStatus == EBlendStatus::BlendOut)
+			{
+				if (EntryExtraInfo.bIsBlendFull && EntryExtraInfo.bIsBlendFinished)
+				{
+					EntriesToRemove.Add(ResolvedEntry.EntryIndex);
+				}
+			}
 		}
 		else
 		{
@@ -206,6 +321,42 @@ void FPersistentBlendStackCameraNodeEvaluator::InternalUpdate(TArrayView<FResolv
 			OutResult.OverrideAll(CurResult);
 		}
 	}
+
+	for (int32 Index = EntriesToRemove.Num() - 1; Index >= 0; --Index)
+	{
+		PopEntry(EntriesToRemove[Index]);
+		EntryExtraInfos.RemoveAt(EntriesToRemove[Index]);
+	}
+}
+
+const UCameraRigTransition* FPersistentBlendStackCameraNodeEvaluator::FindEnterTransition(const FBlendStackCameraInsertParams& Params) const
+{
+	// If we are forced to use a specific transition, our search is over.
+	if (Params.TransitionOverride)
+	{
+		return Params.TransitionOverride;
+	}
+
+	// Find a transition that works for blending the given camera rig in.
+	return FCameraRigTransitionFinder::FindTransition(
+			Params.CameraRig->EnterTransitions,
+			nullptr, nullptr, false,
+			Params.CameraRig, nullptr);
+}
+
+const UCameraRigTransition* FPersistentBlendStackCameraNodeEvaluator::FindExitTransition(const FBlendStackCameraRemoveParams& Params, const FCameraRigEntry& Entry) const
+{
+	// If we are forced to use a specific transition, our search is over.
+	if (Params.TransitionOverride)
+	{
+		return Params.TransitionOverride;
+	}
+
+	// Find a transition that works for blending the given camera rig out.
+	return FCameraRigTransitionFinder::FindTransition(
+			Entry.CameraRig->ExitTransitions,
+			Entry.CameraRig, nullptr, Entry.bIsFrozen,
+			nullptr, nullptr);
 }
 
 }  // namespace UE::Cameras
diff --git a/Engine/Plugins/Cameras/GameplayCameras/Source/GameplayCameras/Private/Core/TransientBlendStackCameraNode.cpp b/Engine/Plugins/Cameras/GameplayCameras/Source/GameplayCameras/Private/Core/TransientBlendStackCameraNode.cpp
--- a/Engine/Plugins/Cameras/GameplayCameras/Source/GameplayCameras/Private/Core/TransientBlendStackCameraNode.cpp
+++ b/Engine/Plugins/Cameras/GameplayCameras/Source/GameplayCameras/Private/Core/TransientBlendStackCameraNode.cpp
@@ -11,6 +11,7 @@
 #include "Core/CameraRigAsset.h"
 #include "Core/CameraRigCombinationRegistry.h"
 #include "Core/CameraRigTransition.h"
+#include "Helpers/CameraRigTransitionFinder.h"
 #include "Nodes/Blends/PopBlendCameraNode.h"
 
 namespace UE::Cameras
@@ -575,7 +576,7 @@ const UCameraRigTransition* FTransientBlendStackCameraNodeEvaluator::FindTransit
 				// Look for exit transitions on the last active camera rig itself.
 				for (const UCameraRigAsset* ToCameraRig : ToCombinedCameraRigs)
 				{
-					TransitionToUse = FindTransition(
+					TransitionToUse = FCameraRigTransitionFinder::FindTransition(
 							FromCameraRig->ExitTransitions,
 							FromCameraRig, FromCameraAsset, bFromFrozen,
 							ToCameraRig, ToCameraAsset);
@@ -593,7 +594,7 @@ const UCameraRigTransition* FTransientBlendStackCameraNodeEvaluator::FindTransit
 				// Look for exit transitions on its parent camera asset.
 				for (const UCameraRigAsset* ToCameraRig : ToCombinedCameraRigs)
 				{
-					TransitionToUse = FindTransition(
+					TransitionToUse = FCameraRigTransitionFinder::FindTransition(
 							FromCameraAsset->GetExitTransitions(),
 							FromCameraRig, FromCameraAsset, bFromFrozen,
 							ToCameraRig, ToCameraAsset);
@@ -613,7 +614,7 @@ const UCameraRigTransition* FTransientBlendStackCameraNodeEvaluator::FindTransit
 				// Look for enter transitions on the new camera rig itself.
 				for (const UCameraRigAsset* FromCameraRig : FromCombinedCameraRigs)
 				{
-					TransitionToUse = FindTransition(
+					TransitionToUse = FCameraRigTransitionFinder::FindTransition(
 							ToCameraRig->EnterTransitions,
 							FromCameraRig, FromCameraAsset, bFromFrozen,
 							ToCameraRig, ToCameraAsset);
@@ -631,7 +632,7 @@ const UCameraRigTransition* FTransientBlendStackCameraNodeEvaluator::FindTransit
 				// Look at enter transitions on its parent camera asset.
 				for (const UCameraRigAsset* FromCameraRig : FromCombinedCameraRigs)
 				{
-					TransitionToUse = FindTransition(
+					TransitionToUse = FCameraRigTransitionFinder::FindTransition(
 							ToCameraAsset->GetEnterTransitions(),
 							FromCameraRig, FromCameraAsset, bFromFrozen,
 							ToCameraRig, ToCameraAsset);
@@ -648,29 +649,5 @@ const UCameraRigTransition* FTransientBlendStackCameraNodeEvaluator::FindTransit
 	return nullptr;
 }
 
-const UCameraRigTransition* FTransientBlendStackCameraNodeEvaluator::FindTransition(
-			TArrayView<const TObjectPtr<UCameraRigTransition>> Transitions, 
-			const UCameraRigAsset* FromCameraRig, const UCameraAsset* FromCameraAsset, bool bFromFrozen,
-			const UCameraRigAsset* ToCameraRig, const UCameraAsset* ToCameraAsset) const
-{
-	FCameraRigTransitionConditionMatchParams MatchParams;
-	MatchParams.FromCameraRig = FromCameraRig;
-	MatchParams.FromCameraAsset = FromCameraAsset;
-	MatchParams.ToCameraRig = ToCameraRig;
-	MatchParams.ToCameraAsset = ToCameraAsset;
-
-	// The transition should be used if all its conditions pass.
-	for (TObjectPtr<const UCameraRigTransition> Transition : Transitions)
-	{
-		const bool bConditionsPass = Transition->AllConditionsMatch(MatchParams);
-		if (bConditionsPass)
-		{
-			return Transition;
-		}
-	}
-
-	return nullptr;
-}
-
 }  // namespace UE::Cameras
 
diff --git a/Engine/Plugins/Cameras/GameplayCameras/Source/GameplayCameras/Private/Helpers/CameraRigTransitionFinder.cpp b/Engine/Plugins/Cameras/GameplayCameras/Source/GameplayCameras/Private/Helpers/CameraRigTransitionFinder.cpp
new file mode 100644
--- /dev/null
+++ b/Engine/Plugins/Cameras/GameplayCameras/Source/GameplayCameras/Private/Helpers/CameraRigTransitionFinder.cpp
@@ -0,0 +1,38 @@
+// Copyright Epic Games, Inc. All Rights Reserved.
+
+#include "Helpers/CameraRigTransitionFinder.h"
+
+#include "Core/CameraAsset.h"
+#include "Core/CameraRigAsset.h"
+#include "Core/CameraRigTransition.h"
+
+namespace UE::Cameras
+{
+
+const UCameraRigTransition* FCameraRigTransitionFinder::FindTransition(
+		TArrayView<const TObjectPtr<UCameraRigTransition>> Transitions, 
+		const UCameraRigAsset* FromCameraRig, const UCameraAsset* FromCameraAsset, bool bFromFrozen,
+		const UCameraRigAsset* ToCameraRig, const UCameraAsset* ToCameraAsset)
+{
+	FCameraRigTransitionConditionMatchParams MatchParams;
+	MatchParams.FromCameraRig = FromCameraRig;
+	MatchParams.FromCameraAsset = FromCameraAsset;
+	MatchParams.ToCameraRig = ToCameraRig;
+	MatchParams.ToCameraAsset = ToCameraAsset;
+
+	// The transition should be used if all its conditions pass.
+	for (TObjectPtr<const UCameraRigTransition> Transition : Transitions)
+	{
+		const bool bConditionsPass = Transition->AllConditionsMatch(MatchParams);
+		if (bConditionsPass)
+		{
+			return Transition;
+		}
+	}
+
+	return nullptr;
+}
+
+}  // namespace UE::Cameras
+
+
diff --git a/Engine/Plugins/Cameras/GameplayCameras/Source/GameplayCameras/Private/Helpers/CameraRigTransitionFinder.h b/Engine/Plugins/Cameras/GameplayCameras/Source/GameplayCameras/Private/Helpers/CameraRigTransitionFinder.h
new file mode 100644
--- /dev/null
+++ b/Engine/Plugins/Cameras/GameplayCameras/Source/GameplayCameras/Private/Helpers/CameraRigTransitionFinder.h
@@ -0,0 +1,28 @@
+// Copyright Epic Games, Inc. All Rights Reserved.
+
+#pragma once
+
+#include "UObject/ObjectPtr.h"
+
+class UCameraAsset;
+class UCameraRigAsset;
+class UCameraRigTransition;
+
+namespace UE::Cameras
+{
+
+/**
+ * Helper class for finding a camera rig transition that matches a given situation.
+ */
+class FCameraRigTransitionFinder
+{
+public:
+
+	static const UCameraRigTransition* FindTransition(
+			TArrayView<const TObjectPtr<UCameraRigTransition>> Transitions, 
+			const UCameraRigAsset* FromCameraRig, const UCameraAsset* FromCameraAsset, bool bFromFrozen,
+			const UCameraRigAsset* ToCameraRig, const UCameraAsset* ToCameraAsset);
+};
+
+}  // namespace UE::Cameras
+
diff --git a/Engine/Plugins/Cameras/GameplayCameras/Source/GameplayCameras/Private/Nodes/Blends/InterruptedBlendCameraNode.cpp b/Engine/Plugins/Cameras/GameplayCameras/Source/GameplayCameras/Private/Nodes/Blends/InterruptedBlendCameraNode.cpp
new file mode 100644
--- /dev/null
+++ b/Engine/Plugins/Cameras/GameplayCameras/Source/GameplayCameras/Private/Nodes/Blends/InterruptedBlendCameraNode.cpp
@@ -0,0 +1,57 @@
+// Copyright Epic Games, Inc. All Rights Reserved.
+
+#include "Nodes/Blends/InterruptedBlendCameraNode.h"
+
+namespace UE::Cameras
+{
+
+UE_DEFINE_BLEND_CAMERA_NODE_EVALUATOR(FInterruptedBlendCameraNodeEvaluator)
+
+FCameraNodeEvaluatorChildrenView FInterruptedBlendCameraNodeEvaluator::OnGetChildren()
+{
+	return FCameraNodeEvaluatorChildrenView{ ChildBlend, FrozenBlend };
+}
+
+void FInterruptedBlendCameraNodeEvaluator::OnRun(const FCameraNodeEvaluationParams& Params, FCameraNodeEvaluationResult& OutResult)
+{
+	if (ChildBlend)
+	{
+		ChildBlend->Run(Params, OutResult);
+	}
+	// Don't run the frozen blend... it's... frozen.
+}
+
+void FInterruptedBlendCameraNodeEvaluator::OnBlendParameters(const FCameraNodePreBlendParams& Params, FCameraNodePreBlendResult& OutResult)
+{
+	FrozenResult.VariableTable.OverrideAll(OutResult.VariableTable, true);
+
+	if (FrozenBlend)
+	{
+		FCameraNodePreBlendResult FrozenPreBlendResult(FrozenResult.VariableTable);
+		FrozenBlend->BlendParameters(Params, FrozenPreBlendResult);
+	}
+	if (ChildBlend)
+	{
+		FCameraNodePreBlendParams FrozenPreBlendParams(Params.EvaluationParams, Params.LastCameraPose, FrozenResult.VariableTable);
+		ChildBlend->BlendParameters(FrozenPreBlendParams, OutResult);
+	}
+}
+
+void FInterruptedBlendCameraNodeEvaluator::OnBlendResults(const FCameraNodeBlendParams& Params, FCameraNodeBlendResult& OutResult)
+{
+	FrozenResult.OverrideAll(OutResult.BlendedResult, true);
+
+	if (FrozenBlend)
+	{
+		FCameraNodeBlendResult FrozenBlendResult(FrozenResult);
+		FrozenBlend->BlendResults(Params, FrozenBlendResult);
+	}
+	if (ChildBlend)
+	{
+		FCameraNodeBlendParams FrozenBlendParams(Params.ChildParams, FrozenResult);
+		ChildBlend->BlendResults(FrozenBlendParams, OutResult);
+	}
+}
+
+}  // namespace UE::Cameras
+
diff --git a/Engine/Plugins/Cameras/GameplayCameras/Source/GameplayCameras/Private/Nodes/Blends/ReverseBlendCameraNode.cpp b/Engine/Plugins/Cameras/GameplayCameras/Source/GameplayCameras/Private/Nodes/Blends/ReverseBlendCameraNode.cpp
new file mode 100644
--- /dev/null
+++ b/Engine/Plugins/Cameras/GameplayCameras/Source/GameplayCameras/Private/Nodes/Blends/ReverseBlendCameraNode.cpp
@@ -0,0 +1,53 @@
+// Copyright Epic Games, Inc. All Rights Reserved.
+
+#include "Nodes/Blends/ReverseBlendCameraNode.h"
+
+namespace UE::Cameras
+{
+
+UE_DEFINE_BLEND_CAMERA_NODE_EVALUATOR(FReverseBlendCameraNodeEvaluator)
+
+FCameraNodeEvaluatorChildrenView FReverseBlendCameraNodeEvaluator::OnGetChildren()
+{
+	return FCameraNodeEvaluatorChildrenView{ ChildBlend };
+}
+
+void FReverseBlendCameraNodeEvaluator::OnRun(const FCameraNodeEvaluationParams& Params, FCameraNodeEvaluationResult& OutResult)
+{
+	if (ChildBlend)
+	{
+		ChildBlend->Run(Params, OutResult);
+	}
+}
+
+void FReverseBlendCameraNodeEvaluator::OnBlendParameters(const FCameraNodePreBlendParams& Params, FCameraNodePreBlendResult& OutResult)
+{
+	if (ChildBlend)
+	{
+		FCameraNodePreBlendParams ReversedParams(Params.EvaluationParams, Params.LastCameraPose, OutResult.VariableTable);
+		FCameraNodePreBlendResult ReversedResult(const_cast<FCameraVariableTable&>(Params.ChildVariableTable));
+		ChildBlend->BlendParameters(ReversedParams, ReversedResult);
+	}
+}
+
+void FReverseBlendCameraNodeEvaluator::OnBlendResults(const FCameraNodeBlendParams& Params, FCameraNodeBlendResult& OutResult)
+{
+	if (ChildBlend)
+	{
+		FCameraNodeBlendParams ReversedParams(Params.ChildParams, OutResult.BlendedResult);
+		FCameraNodeBlendResult ReversedResult(const_cast<FCameraNodeEvaluationResult&>(Params.ChildResult));
+		ChildBlend->BlendResults(ReversedParams, ReversedResult);
+	}
+}
+
+bool FReverseBlendCameraNodeEvaluator::OnInitializeFromInterruption(const FCameraNodeBlendInterruptionParams& Params)
+{
+	if (ChildBlend)
+	{
+		return ChildBlend->InitializeFromInterruption(Params);
+	}
+	return true;
+}
+
+}  // namespace UE::Cameras
+
diff --git a/Engine/Plugins/Cameras/GameplayCameras/Source/GameplayCameras/Private/Nodes/Blends/SimpleBlendCameraNode.cpp b/Engine/Plugins/Cameras/GameplayCameras/Source/GameplayCameras/Private/Nodes/Blends/SimpleBlendCameraNode.cpp
--- a/Engine/Plugins/Cameras/GameplayCameras/Source/GameplayCameras/Private/Nodes/Blends/SimpleBlendCameraNode.cpp
+++ b/Engine/Plugins/Cameras/GameplayCameras/Source/GameplayCameras/Private/Nodes/Blends/SimpleBlendCameraNode.cpp
@@ -23,7 +23,11 @@ void FSimpleBlendCameraNodeEvaluator::OnRun(const FCameraNodeEvaluationParams& P
 {
 	FSimpleBlendCameraNodeEvaluationResult FactorResult;
 	OnComputeBlendFactor(Params, FactorResult);
-	BlendFactor = FactorResult.BlendFactor;
+	BlendFactor = FMath::Clamp(FactorResult.BlendFactor, 0.f, 1.f);
+	if (bReverse)
+	{
+		BlendFactor = 1.f - BlendFactor;
+	}
 }
 
 void FSimpleBlendCameraNodeEvaluator::OnBlendParameters(const FCameraNodePreBlendParams& Params, FCameraNodePreBlendResult& OutResult)
@@ -31,7 +35,7 @@ void FSimpleBlendCameraNodeEvaluator::OnBlendParameters(const FCameraNodePreBlen
 	const FCameraVariableTable& ChildVariableTable(Params.ChildVariableTable);
 	OutResult.VariableTable.Lerp(ChildVariableTable, Params.VariableTableFilter, BlendFactor);
 
-	OutResult.bIsBlendFull = BlendFactor >= 1.f;
+	OutResult.bIsBlendFull = (bReverse ? BlendFactor <= 0.f : BlendFactor >= 1.f);
 	OutResult.bIsBlendFinished = bIsBlendFinished;
 }
 
@@ -42,7 +46,7 @@ void FSimpleBlendCameraNodeEvaluator::OnBlendResults(const FCameraNodeBlendParam
 
 	BlendedResult.LerpAll(ChildResult, BlendFactor);
 
-	OutResult.bIsBlendFull = BlendFactor >= 1.f;
+	OutResult.bIsBlendFull = (bReverse ? BlendFactor <= 0.f : BlendFactor >= 1.f);
 	OutResult.bIsBlendFinished = bIsBlendFinished;
 }
 
@@ -50,6 +54,13 @@ void FSimpleBlendCameraNodeEvaluator::OnSerialize(const FCameraNodeEvaluatorSeri
 {
 	Ar << BlendFactor;
 	Ar << bIsBlendFinished;
+	Ar << bReverse;
+}
+
+bool FSimpleBlendCameraNodeEvaluator::OnSetReversed(bool bInReverse)
+{
+	bReverse = bInReverse;
+	return true;
 }
 
 #if UE_GAMEPLAY_CAMERAS_DEBUG
@@ -69,25 +80,50 @@ void FSimpleBlendCameraDebugBlock::OnDebugDraw(const FCameraDebugBlockDrawParams
 
 UE_DEFINE_BLEND_CAMERA_NODE_EVALUATOR(FSimpleFixedTimeBlendCameraNodeEvaluator)
 
-void FSimpleFixedTimeBlendCameraNodeEvaluator::OnRun(const FCameraNodeEvaluationParams& Params, FCameraNodeEvaluationResult& OutResult)
+void FSimpleFixedTimeBlendCameraNodeEvaluator::OnInitialize(const FCameraNodeEvaluatorInitializeParams& Params, FCameraNodeEvaluationResult& OutResult)
 {
+	Super::OnInitialize(Params, OutResult);
+
 	const USimpleFixedTimeBlendCameraNode* BlendNode = GetCameraNodeAs<USimpleFixedTimeBlendCameraNode>();
+	TotalTime = BlendNode->BlendTime;
+}
+
+void FSimpleFixedTimeBlendCameraNodeEvaluator::OnRun(const FCameraNodeEvaluationParams& Params, FCameraNodeEvaluationResult& OutResult)
+{
 	CurrentTime += Params.DeltaTime;
-	if (CurrentTime >= BlendNode->BlendTime)
+	if (CurrentTime >= TotalTime)
 	{
-		CurrentTime = BlendNode->BlendTime;
+		CurrentTime = TotalTime;
 		SetBlendFinished();
 	}
 
 	FSimpleBlendCameraNodeEvaluator::OnRun(Params, OutResult);
 }
 
+bool FSimpleFixedTimeBlendCameraNodeEvaluator::OnInitializeFromInterruption(const FCameraNodeBlendInterruptionParams& Params)
+{
+	// If we are interrupting a fixed-time blend, adjust our own time to the complementary time ratio.
+	// That is: if we interrupted a blend that was 70% complete, let's set reduce our blend time to 70% 
+	// of its original value.
+	if (Params.InterruptedBlend)
+	{
+		const FSimpleFixedTimeBlendCameraNodeEvaluator* InterruptedBlend = Params.InterruptedBlend->CastThis<FSimpleFixedTimeBlendCameraNodeEvaluator>();
+		if (InterruptedBlend)
+		{
+			const float InterruptedTimeFactor = InterruptedBlend->GetTimeFactor();
+			TotalTime = TotalTime * InterruptedTimeFactor;
+		}
+	}
+
+	// We still want to be wrapped in an interrupted blend.
+	return false;
+}
+
 float FSimpleFixedTimeBlendCameraNodeEvaluator::GetTimeFactor() const
 {
-	const USimpleFixedTimeBlendCameraNode* BlendNode = GetCameraNodeAs<USimpleFixedTimeBlendCameraNode>();
-	if (BlendNode->BlendTime > 0.f)
+	if (TotalTime > 0.f)
 	{
-		return CurrentTime / BlendNode->BlendTime;
+		return CurrentTime / TotalTime;
 	}
 	return 1.f;
 }
diff --git a/Engine/Plugins/Cameras/GameplayCameras/Source/GameplayCameras/Public/Core/BlendCameraNode.h b/Engine/Plugins/Cameras/GameplayCameras/Source/GameplayCameras/Public/Core/BlendCameraNode.h
--- a/Engine/Plugins/Cameras/GameplayCameras/Source/GameplayCameras/Public/Core/BlendCameraNode.h
+++ b/Engine/Plugins/Cameras/GameplayCameras/Source/GameplayCameras/Public/Core/BlendCameraNode.h
@@ -19,6 +19,8 @@ class UBlendCameraNode : public UCameraNode
 namespace UE::Cameras
 {
 
+class FBlendCameraNodeEvaluator;
+
 /**
  * Parameter struct for blending camera node parameters.
  */
@@ -99,6 +101,15 @@ struct FCameraNodeBlendResult
 	bool bIsBlendFinished = false;
 };
 
+/**
+ * Parameter struct for initializing a blend from an interrupted blend.
+ */
+struct FCameraNodeBlendInterruptionParams
+{
+	/** The existing blend that was interrupted. */
+	const FBlendCameraNodeEvaluator* InterruptedBlend = nullptr;
+};
+
 /**
  * Base evaluator class for blend camera nodes.
  */
@@ -114,6 +125,26 @@ public:
 	/** Blend the result of a camera node tree over another result. */
 	GAMEPLAYCAMERAS_API void BlendResults(const FCameraNodeBlendParams& Params, FCameraNodeBlendResult& OutResult);
 
+public:
+
+	/**
+	 * Initialize this blend from an interrupted blend.
+	 *
+	 * @return If true, the blend sub-class fully supports seamless transition from the interrupted blend.
+	 *		   If false, the blend will be wrapped in an Interrupted Blend node that will freeze the
+	 *		   interrupted blend to ensure a seamless transition.
+	 */
+	bool InitializeFromInterruption(const FCameraNodeBlendInterruptionParams& Params);
+
+	/**
+	 * Reverse the direction of this blend.
+	 *
+	 * @return If true, the blend sub-class fully supports reverse blending.
+	 *         If false, the blend will be wrapped in a Reverse Blend node that will cache and swap the
+	 *         "from" and "to" results given to the BlendParameters and BlendResults methods.
+	 */
+	bool SetReversed(bool bInReverse);
+
 protected:
 
 	/** Blend the parameters produced by a camera node tree over another set of values. */
@@ -121,6 +152,12 @@ protected:
 
 	/** Blend the result of a camera node tree over another result. */
 	virtual void OnBlendResults(const FCameraNodeBlendParams& Params, FCameraNodeBlendResult& OutResult) {}
+
+	/** Initialize this blend from an interrupted blend. See comments from InitializeFromInterruption. */
+	virtual bool OnInitializeFromInterruption(const FCameraNodeBlendInterruptionParams& Params) { return false; }
+
+	/** Reverse the direction of this blend. See comments from SetReversed. */
+	virtual bool OnSetReversed(bool bInReverse) { return false; }
 };
 
 }  // namespace UE::Cameras
diff --git a/Engine/Plugins/Cameras/GameplayCameras/Source/GameplayCameras/Public/Core/PersistentBlendStackCameraNode.h b/Engine/Plugins/Cameras/GameplayCameras/Source/GameplayCameras/Public/Core/PersistentBlendStackCameraNode.h
--- a/Engine/Plugins/Cameras/GameplayCameras/Source/GameplayCameras/Public/Core/PersistentBlendStackCameraNode.h
+++ b/Engine/Plugins/Cameras/GameplayCameras/Source/GameplayCameras/Public/Core/PersistentBlendStackCameraNode.h
@@ -21,7 +21,19 @@ struct FBlendStackCameraInsertParams
 	/** A transition to use, instead of looking one up. */
 	TObjectPtr<const UCameraRigTransition> TransitionOverride;
 
-	/** Whether to force insert a new instance of the camera rig, even if there is already one in the stack. */
+	/**
+	 * An ordering value for where to insert the camera rig in the stack.
+	 * Higher values place the camera rig higher in the stack (i.e. evaluating later)
+	 * Lower values place the camera rig lower in the stack (i.e. evluating earlier)
+	 * Insertion of an equal pre-existing value will be done on top of (after) existing 
+	 * entries in the stack.
+	 */
+	int32 StackOrder = 0;
+
+	/** 
+	 * Whether to force insert a new instance of the camera rig, even if there is already 
+	 * one in the stack with the same StackOrder.
+	 */
 	bool bForceInsert = false;
 };
 
@@ -40,6 +52,15 @@ struct FBlendStackCameraRemoveParams
 
 	/** The source camera rig asset used by the instanced to remove. */
 	TObjectPtr<const UCameraRigAsset> CameraRig;
+
+	/** A transition to use, instead of looking one up. */
+	TObjectPtr<const UCameraRigTransition> TransitionOverride;
+
+	/** 
+	 * Whether to immediately remove the given camera rig instead of blending it out.
+	 * Equivalent to passing a pop blend as the TransitionOverride.
+	 */
+	bool bRemoveImmediately = false;
 };
 
 /**
@@ -67,11 +88,29 @@ protected:
 
 private:
 
-	// Update methods.
 	void InternalUpdate(TArrayView<FResolvedEntry> ResolvedEntries, const FCameraNodeEvaluationParams& Params, FCameraNodeEvaluationResult& OutResult);
 
-	// Utility functions for finding an appropriate transition.
-	const UCameraRigTransition* FindTransition(const FBlendStackCameraInsertParams& Params) const;
+	const UCameraRigTransition* FindEnterTransition(const FBlendStackCameraInsertParams& Params) const;
+	const UCameraRigTransition* FindExitTransition(const FBlendStackCameraRemoveParams& Params, const FCameraRigEntry& Entry) const;
+
+private:
+
+	enum class EBlendStatus
+	{
+		None,
+		BlendIn,
+		BlendOut
+	};
+
+	struct FCameraRigEntryExtraInfo
+	{
+		int32 StackOrder = 0;
+		EBlendStatus BlendStatus = EBlendStatus::None;
+		bool bIsBlendFull = false;
+		bool bIsBlendFinished = false;
+	};
+
+	TArray<FCameraRigEntryExtraInfo> EntryExtraInfos;
 };
 
 }  // namespace UE::Cameras
diff --git a/Engine/Plugins/Cameras/GameplayCameras/Source/GameplayCameras/Public/Core/RootCameraNode.h b/Engine/Plugins/Cameras/GameplayCameras/Source/GameplayCameras/Public/Core/RootCameraNode.h
--- a/Engine/Plugins/Cameras/GameplayCameras/Source/GameplayCameras/Public/Core/RootCameraNode.h
+++ b/Engine/Plugins/Cameras/GameplayCameras/Source/GameplayCameras/Public/Core/RootCameraNode.h
@@ -46,6 +46,9 @@ struct FActivateCameraRigParams
 	/** The evaluation layer on which to instantiate the camera rig. */
 	ECameraRigLayer Layer = ECameraRigLayer::Main;
 
+	/** If the given layer supports ordering, this specifies a priority key. */
+	int32 OrderKey = 0;
+
 	/** Whether a new instance of the camera rig should be created, even if it is already active in the given layer. */
 	bool bForceActivate = false;
 };
@@ -66,6 +69,12 @@ struct FDeactivateCameraRigParams
 
 	/** The evaluation layer on which the camera rig is running. */
 	ECameraRigLayer Layer = ECameraRigLayer::Main;
+
+	/** If deactivation uses a transition (e.g. blending out), then force using the given transition. */
+	TObjectPtr<const UCameraRigTransition> TransitionOverride;
+
+	/** If deactivation can take time (e.g. blending out), then force an immediate deactivation. */
+	bool bDeactiveImmediately = false;
 };
 
 /**
diff --git a/Engine/Plugins/Cameras/GameplayCameras/Source/GameplayCameras/Public/Core/TransientBlendStackCameraNode.h b/Engine/Plugins/Cameras/GameplayCameras/Source/GameplayCameras/Public/Core/TransientBlendStackCameraNode.h
--- a/Engine/Plugins/Cameras/GameplayCameras/Source/GameplayCameras/Public/Core/TransientBlendStackCameraNode.h
+++ b/Engine/Plugins/Cameras/GameplayCameras/Source/GameplayCameras/Public/Core/TransientBlendStackCameraNode.h
@@ -79,12 +79,7 @@ private:
 	void InternalPostBlendExecute(TArrayView<FResolvedEntry> ResolvedEntries, const FCameraNodeEvaluationParams& Params, FCameraNodeEvaluationResult& OutResult);
 	void InternalRunFinished(FCameraNodeEvaluationResult& OutResult);
 
-	// Utility functions for finding an appropriate transition.
 	const UCameraRigTransition* FindTransition(const FBlendStackCameraPushParams& Params) const;
-	const UCameraRigTransition* FindTransition(
-			TArrayView<const TObjectPtr<UCameraRigTransition>> Transitions, 
-			const UCameraRigAsset* FromCameraRig, const UCameraAsset* FromCameraAsset, bool bFromFrozen,
-			const UCameraRigAsset* ToCameraRig, const UCameraAsset* ToCameraAsset) const;
 
 	FBlendStackEntryID PushNewEntry(const FBlendStackCameraPushParams& Params, const UCameraRigTransition* Transition);
 	FBlendStackEntryID PushMergedEntry(const FBlendStackCameraPushParams& Params, const UCameraRigTransition* Transition);
diff --git a/Engine/Plugins/Cameras/GameplayCameras/Source/GameplayCameras/Public/Nodes/Blends/InterruptedBlendCameraNode.h b/Engine/Plugins/Cameras/GameplayCameras/Source/GameplayCameras/Public/Nodes/Blends/InterruptedBlendCameraNode.h
new file mode 100644
--- /dev/null
+++ b/Engine/Plugins/Cameras/GameplayCameras/Source/GameplayCameras/Public/Nodes/Blends/InterruptedBlendCameraNode.h
@@ -0,0 +1,48 @@
+// Copyright Epic Games, Inc. All Rights Reserved.
+
+#pragma once
+
+#include "Core/BlendCameraNode.h"
+
+namespace UE::Cameras
+{
+
+/**
+ * A special blend meant to "freeze" an ongoing blend, and then start running another blend.
+ * This is used for instance by the Camera Modifier Service to blend out a camera rig instance that had not
+ * fully blended in yet.
+ *
+ * NOTE: this blend doesn't have any data, and can't be added by the user in the editor. It is meant to be
+ * created programmatically only.
+ */
+class FInterruptedBlendCameraNodeEvaluator : public FBlendCameraNodeEvaluator
+{
+	UE_DECLARE_BLEND_CAMERA_NODE_EVALUATOR(GAMEPLAYCAMERAS_API, FInterruptedBlendCameraNodeEvaluator)
+
+public:
+
+	FInterruptedBlendCameraNodeEvaluator()
+	{}
+	FInterruptedBlendCameraNodeEvaluator(FBlendCameraNodeEvaluator* InChildBlend, FBlendCameraNodeEvaluator* InFrozenBlend) 
+		: ChildBlend(InChildBlend)
+		, FrozenBlend(InFrozenBlend)
+	{}
+
+protected:
+
+	// FBlendCameraNodeEvaluator interface.
+	virtual FCameraNodeEvaluatorChildrenView OnGetChildren() override;
+	virtual void OnRun(const FCameraNodeEvaluationParams& Params, FCameraNodeEvaluationResult& OutResult) override;
+	virtual void OnBlendParameters(const FCameraNodePreBlendParams& Params, FCameraNodePreBlendResult& OutResult) override;
+	virtual void OnBlendResults(const FCameraNodeBlendParams& Params, FCameraNodeBlendResult& OutResult) override;
+
+public:
+
+	FBlendCameraNodeEvaluator* ChildBlend = nullptr;
+	FBlendCameraNodeEvaluator* FrozenBlend = nullptr;
+
+	FCameraNodeEvaluationResult FrozenResult;
+};
+
+}  // namespace UE::Cameras
+
diff --git a/Engine/Plugins/Cameras/GameplayCameras/Source/GameplayCameras/Public/Nodes/Blends/ReverseBlendCameraNode.h b/Engine/Plugins/Cameras/GameplayCameras/Source/GameplayCameras/Public/Nodes/Blends/ReverseBlendCameraNode.h
new file mode 100644
--- /dev/null
+++ b/Engine/Plugins/Cameras/GameplayCameras/Source/GameplayCameras/Public/Nodes/Blends/ReverseBlendCameraNode.h
@@ -0,0 +1,44 @@
+// Copyright Epic Games, Inc. All Rights Reserved.
+
+#pragma once
+
+#include "Core/BlendCameraNode.h"
+
+namespace UE::Cameras
+{
+
+/**
+ * A special blend that can run another blend in "reverse", when that blend doesn't support reverse mode.
+ * This is used for instance by the Camera Modifier Service to blend things out.
+ *
+ * NOTE: this blend doesn't have any data, and can't be added by the user in the editor. It is meant to be
+ * created programmatically only.
+ */
+class FReverseBlendCameraNodeEvaluator : public FBlendCameraNodeEvaluator
+{
+	UE_DECLARE_BLEND_CAMERA_NODE_EVALUATOR(GAMEPLAYCAMERAS_API, FReverseBlendCameraNodeEvaluator)
+
+public:
+
+	FReverseBlendCameraNodeEvaluator()
+	{}
+	FReverseBlendCameraNodeEvaluator(FBlendCameraNodeEvaluator* InChildBlend) 
+		: ChildBlend(InChildBlend)
+	{}
+
+protected:
+
+	// FBlendCameraNodeEvaluator interface.
+	virtual FCameraNodeEvaluatorChildrenView OnGetChildren() override;
+	virtual void OnRun(const FCameraNodeEvaluationParams& Params, FCameraNodeEvaluationResult& OutResult) override;
+	virtual void OnBlendParameters(const FCameraNodePreBlendParams& Params, FCameraNodePreBlendResult& OutResult) override;
+	virtual void OnBlendResults(const FCameraNodeBlendParams& Params, FCameraNodeBlendResult& OutResult) override;
+	virtual bool OnInitializeFromInterruption(const FCameraNodeBlendInterruptionParams& Params) override;
+
+public:
+
+	FBlendCameraNodeEvaluator* ChildBlend = nullptr;
+};
+
+}  // namespace UE::Cameras
+
diff --git a/Engine/Plugins/Cameras/GameplayCameras/Source/GameplayCameras/Public/Nodes/Blends/SimpleBlendCameraNode.h b/Engine/Plugins/Cameras/GameplayCameras/Source/GameplayCameras/Public/Nodes/Blends/SimpleBlendCameraNode.h
--- a/Engine/Plugins/Cameras/GameplayCameras/Source/GameplayCameras/Public/Nodes/Blends/SimpleBlendCameraNode.h
+++ b/Engine/Plugins/Cameras/GameplayCameras/Source/GameplayCameras/Public/Nodes/Blends/SimpleBlendCameraNode.h
@@ -67,10 +67,12 @@ public:
 
 protected:
 
+	// FBlendCameraNodeEvaluator interface.
 	GAMEPLAYCAMERAS_API virtual void OnRun(const FCameraNodeEvaluationParams& Params, FCameraNodeEvaluationResult& OutResult) override;
 	GAMEPLAYCAMERAS_API virtual void OnBlendParameters(const FCameraNodePreBlendParams& Params, FCameraNodePreBlendResult& OutResult) override;
 	GAMEPLAYCAMERAS_API virtual void OnBlendResults(const FCameraNodeBlendParams& Params, FCameraNodeBlendResult& OutResult) override;
 	GAMEPLAYCAMERAS_API virtual void OnSerialize(const FCameraNodeEvaluatorSerializeParams& Params, FArchive& Ar) override;
+	GAMEPLAYCAMERAS_API virtual bool OnSetReversed(bool bInReverse) override;
 
 	virtual void OnComputeBlendFactor(const FCameraNodeEvaluationParams& Params, FSimpleBlendCameraNodeEvaluationResult& OutResult) {}
 
@@ -84,6 +86,7 @@ private:
 
 	float BlendFactor = 0.f;
 	bool bIsBlendFinished = false;
+	bool bReverse = false;
 };
 
 class FSimpleFixedTimeBlendCameraNodeEvaluator : public FSimpleBlendCameraNodeEvaluator
@@ -92,12 +95,16 @@ class FSimpleFixedTimeBlendCameraNodeEvaluator : public FSimpleBlendCameraNodeEv
 
 protected:
 
+	// FBlendCameraNodeEvaluator interface.
+	GAMEPLAYCAMERAS_API virtual void OnInitialize(const FCameraNodeEvaluatorInitializeParams& Params, FCameraNodeEvaluationResult& OutResult) override;
 	GAMEPLAYCAMERAS_API virtual void OnRun(const FCameraNodeEvaluationParams& Params, FCameraNodeEvaluationResult& OutResult) override;
+	GAMEPLAYCAMERAS_API virtual bool OnInitializeFromInterruption(const FCameraNodeBlendInterruptionParams& Params) override;
 
 	float GetTimeFactor() const;
 
 private:
 
+	float TotalTime = 0.f;
 	float CurrentTime = 0.f;
 };
 
