Cameras: add a "camera modifier" service for easily starting and stopping camera modifier rigs

diff --git a/Engine/Plugins/Cameras/GameplayCameras/Source/GameplayCameras/Private/Core/CameraSystemEvaluator.cpp b/Engine/Plugins/Cameras/GameplayCameras/Source/GameplayCameras/Private/Core/CameraSystemEvaluator.cpp
--- a/Engine/Plugins/Cameras/GameplayCameras/Source/GameplayCameras/Private/Core/CameraSystemEvaluator.cpp
+++ b/Engine/Plugins/Cameras/GameplayCameras/Source/GameplayCameras/Private/Core/CameraSystemEvaluator.cpp
@@ -15,6 +15,7 @@
 #include "Debug/RootCameraDebugBlock.h"
 #include "GameplayCamerasSettings.h"
 #include "Math/ColorList.h"
+#include "Services/CameraModifierService.h"
 #include "Services/OrientationInitializationService.h"
 #include "UObject/Package.h"
 #include "UObject/UObjectGlobals.h"
@@ -91,6 +92,7 @@ void FCameraSystemEvaluator::Initialize(const FCameraSystemEvaluatorCreateParams
 		RootEvaluator->Initialize(InitParams, RootNodeResult);
 	}
 
+	RegisterEvaluationService(MakeShared<FCameraModifierService>());
 	RegisterEvaluationService(MakeShared<FOrientationInitializationService>());
 
 	CameraRigCombinationRegistry = MakeShared<FCameraRigCombinationRegistry>();
diff --git a/Engine/Plugins/Cameras/GameplayCameras/Source/GameplayCameras/Private/GameFramework/GameplayCameraSystemComponent.cpp b/Engine/Plugins/Cameras/GameplayCameras/Source/GameplayCameras/Private/GameFramework/GameplayCameraSystemComponent.cpp
--- a/Engine/Plugins/Cameras/GameplayCameras/Source/GameplayCameras/Private/GameFramework/GameplayCameraSystemComponent.cpp
+++ b/Engine/Plugins/Cameras/GameplayCameras/Source/GameplayCameras/Private/GameFramework/GameplayCameraSystemComponent.cpp
@@ -3,13 +3,17 @@
 #include "GameFramework/GameplayCameraSystemComponent.h"
 
 #include "Components/BillboardComponent.h"
+#include "Core/CameraEvaluationContext.h"
 #include "Core/CameraSystemEvaluator.h"
+#include "Core/PersistentBlendStackCameraNode.h"
+#include "Core/RootCameraNode.h"
 #include "Engine/Canvas.h"
 #include "Engine/Texture2D.h"
 #include "Engine/World.h"
 #include "GameFramework/Controller.h"
 #include "GameFramework/GameplayCameraSystemHost.h"
 #include "Kismet/GameplayStatics.h"
+#include "Services/CameraModifierService.h"
 #include "UObject/ConstructorHelpers.h"
 #include "UObject/ICookInfo.h"
 
@@ -211,5 +215,42 @@ void UGameplayCameraSystemComponent::EndPlay(const EEndPlayReason::Type EndPlayR
 	Super::EndPlay(EndPlayReason);
 }
 
+FCameraRigInstanceID UGameplayCameraSystemComponent::StartGlobalCameraModifierRig(const UCameraRigAsset* CameraRig, int32 OrderKey)
+{
+	using namespace UE::Cameras;
+
+	if (TSharedPtr<FCameraSystemEvaluator> Evaluator = GetCameraSystemEvaluator())
+	{
+		TSharedPtr<FCameraModifierService> CameraModifierService = Evaluator->FindEvaluationService<FCameraModifierService>();
+		return CameraModifierService->StartCameraModifierRig(CameraRig, ECameraRigLayer::Global, OrderKey);
+	}
+
+	return FCameraRigInstanceID();
+}
+
+FCameraRigInstanceID UGameplayCameraSystemComponent::StartVisualCameraModifierRig(const UCameraRigAsset* CameraRig, int32 OrderKey)
+{
+	using namespace UE::Cameras;
+
+	if (TSharedPtr<FCameraSystemEvaluator> Evaluator = GetCameraSystemEvaluator())
+	{
+		TSharedPtr<FCameraModifierService> CameraModifierService = Evaluator->FindEvaluationService<FCameraModifierService>();
+		return CameraModifierService->StartCameraModifierRig(CameraRig, ECameraRigLayer::Visual);
+	}
+
+	return FCameraRigInstanceID();
+}
+
+void UGameplayCameraSystemComponent::StopCameraModifierRig(FCameraRigInstanceID InstanceID, bool bImmediately)
+{
+	using namespace UE::Cameras;
+
+	if (TSharedPtr<FCameraSystemEvaluator> Evaluator = GetCameraSystemEvaluator())
+	{
+		TSharedPtr<FCameraModifierService> CameraModifierService = Evaluator->FindEvaluationService<FCameraModifierService>();
+		CameraModifierService->StopCameraModifierRig(InstanceID, bImmediately);
+	}
+}
+
 #undef LOCTEXT_NAMESPACE
 
diff --git a/Engine/Plugins/Cameras/GameplayCameras/Source/GameplayCameras/Private/GameFramework/GameplayCamerasPlayerCameraManager.cpp b/Engine/Plugins/Cameras/GameplayCameras/Source/GameplayCameras/Private/GameFramework/GameplayCamerasPlayerCameraManager.cpp
--- a/Engine/Plugins/Cameras/GameplayCameras/Source/GameplayCameras/Private/GameFramework/GameplayCamerasPlayerCameraManager.cpp
+++ b/Engine/Plugins/Cameras/GameplayCameras/Source/GameplayCameras/Private/GameFramework/GameplayCamerasPlayerCameraManager.cpp
@@ -18,6 +18,7 @@
 #include "GameFramework/GameplayCameraSystemComponent.h"
 #include "GameFramework/GameplayCameraSystemHost.h"
 #include "GameFramework/PlayerController.h"
+#include "Services/CameraModifierService.h"
 #include "UObject/Package.h"
 #include "UObject/UObjectGlobals.h"
 
@@ -185,6 +186,43 @@ void AGameplayCamerasPlayerCameraManager::ReleasePlayerController()
 	PCOwner = nullptr;
 }
 
+FCameraRigInstanceID AGameplayCamerasPlayerCameraManager::StartGlobalCameraModifierRig(const UCameraRigAsset* CameraRig, int32 OrderKey)
+{
+	using namespace UE::Cameras;
+
+	if (CameraSystemEvaluator)
+	{
+		TSharedPtr<FCameraModifierService> CameraModifierService = CameraSystemEvaluator->FindEvaluationService<FCameraModifierService>();
+		return CameraModifierService->StartCameraModifierRig(CameraRig, ECameraRigLayer::Global, OrderKey);
+	}
+
+	return FCameraRigInstanceID();
+}
+
+FCameraRigInstanceID AGameplayCamerasPlayerCameraManager::StartVisualCameraModifierRig(const UCameraRigAsset* CameraRig, int32 OrderKey)
+{
+	using namespace UE::Cameras;
+
+	if (CameraSystemEvaluator)
+	{
+		TSharedPtr<FCameraModifierService> CameraModifierService = CameraSystemEvaluator->FindEvaluationService<FCameraModifierService>();
+		return CameraModifierService->StartCameraModifierRig(CameraRig, ECameraRigLayer::Visual, OrderKey);
+	}
+
+	return FCameraRigInstanceID();
+}
+
+void AGameplayCamerasPlayerCameraManager::StopCameraModifierRig(FCameraRigInstanceID InstanceID, bool bImmediately)
+{
+	using namespace UE::Cameras;
+
+	if (CameraSystemEvaluator)
+	{
+		TSharedPtr<FCameraModifierService> CameraModifierService = CameraSystemEvaluator->FindEvaluationService<FCameraModifierService>();
+		return CameraModifierService->StopCameraModifierRig(InstanceID, bImmediately);
+	}
+}
+
 void AGameplayCamerasPlayerCameraManager::InitializeFor(APlayerController* PlayerController)
 {
 	AcquireCameraSystemHost(PlayerController);
diff --git a/Engine/Plugins/Cameras/GameplayCameras/Source/GameplayCameras/Private/Services/CameraModifierService.cpp b/Engine/Plugins/Cameras/GameplayCameras/Source/GameplayCameras/Private/Services/CameraModifierService.cpp
new file mode 100644
--- /dev/null
+++ b/Engine/Plugins/Cameras/GameplayCameras/Source/GameplayCameras/Private/Services/CameraModifierService.cpp
@@ -0,0 +1,63 @@
+// Copyright Epic Games, Inc. All Rights Reserved.
+
+#include "Services/CameraModifierService.h"
+
+#include "Core/CameraEvaluationContext.h"
+#include "Core/CameraSystemEvaluator.h"
+#include "Core/RootCameraNode.h"
+#include "UObject/ObjectMacros.h"
+
+namespace UE::Cameras
+{
+
+UE_DEFINE_CAMERA_EVALUATION_SERVICE(FCameraModifierService)
+
+void FCameraModifierService::OnInitialize(const FCameraEvaluationServiceInitializeParams& Params)
+{
+	ensure(Evaluator == nullptr);
+	Evaluator = Params.Evaluator;
+}
+
+void FCameraModifierService::OnTeardown(const FCameraEvaluationServiceTeardownParams& Params)
+{
+	ensure(Evaluator != nullptr);
+	Evaluator = nullptr;
+}
+
+FCameraRigInstanceID FCameraModifierService::StartCameraModifierRig(const UCameraRigAsset* CameraRig, ECameraRigLayer Layer, int32 OrderKey)
+{
+	if (ensure(Evaluator))
+	{
+		EnsureModifierContextCreated();
+
+		FActivateCameraRigParams ActivateParams;
+		ActivateParams.EvaluationContext = ModifierContext;
+		ActivateParams.CameraRig = CameraRig;
+		ActivateParams.Layer = Layer;
+		ActivateParams.OrderKey = FirstBlendStackOrderKey + OrderKey;
+		return Evaluator->GetRootNodeEvaluator()->ActivateCameraRig(ActivateParams);
+	}
+	return FCameraRigInstanceID();
+}
+
+void FCameraModifierService::StopCameraModifierRig(FCameraRigInstanceID CameraRigID, bool bImmediately)
+{
+	if (ensure(Evaluator))
+	{
+		FDeactivateCameraRigParams DeactivateParams;
+		DeactivateParams.InstanceID = CameraRigID;
+		DeactivateParams.bDeactiveImmediately = bImmediately;
+		Evaluator->GetRootNodeEvaluator()->DeactivateCameraRig(DeactivateParams);
+	}
+}
+
+void FCameraModifierService::EnsureModifierContextCreated()
+{
+	if (!ModifierContext)
+	{
+		ModifierContext = MakeShared<FCameraEvaluationContext>();
+	}
+}
+
+}  // namespace UE::Cameras
+
diff --git a/Engine/Plugins/Cameras/GameplayCameras/Source/GameplayCameras/Public/GameFramework/GameplayCameraSystemComponent.h b/Engine/Plugins/Cameras/GameplayCameras/Source/GameplayCameras/Public/GameFramework/GameplayCameraSystemComponent.h
--- a/Engine/Plugins/Cameras/GameplayCameras/Source/GameplayCameras/Public/GameFramework/GameplayCameraSystemComponent.h
+++ b/Engine/Plugins/Cameras/GameplayCameras/Source/GameplayCameras/Public/GameFramework/GameplayCameraSystemComponent.h
@@ -11,6 +11,7 @@ class APlayerController;
 class UCameraRigAsset;
 class UCanvas;
 class UGameplayCameraSystemHost;
+struct FCameraRigInstanceID;
 struct FMinimalViewInfo;
 
 namespace UE::Cameras
@@ -54,6 +55,17 @@ public:
 	UFUNCTION(BlueprintCallable, Category=Camera)
 	GAMEPLAYCAMERAS_API void DeactivateCameraSystem(AActor* NextViewTarget = nullptr);
 
+public:
+
+	UFUNCTION(BlueprintCallable, Category=Camera)
+	GAMEPLAYCAMERAS_API FCameraRigInstanceID StartGlobalCameraModifierRig(const UCameraRigAsset* CameraRig, int32 OrderKey = 0);
+
+	UFUNCTION(BlueprintCallable, Category=Camera)
+	GAMEPLAYCAMERAS_API FCameraRigInstanceID StartVisualCameraModifierRig(const UCameraRigAsset* CameraRig, int32 OrderKey = 0);
+
+	UFUNCTION(BlueprintCallable, Category=Camera)
+	GAMEPLAYCAMERAS_API void StopCameraModifierRig(FCameraRigInstanceID InstanceID, bool bImmediately = false);
+
 public:
 
 	// UActorComponent interface
diff --git a/Engine/Plugins/Cameras/GameplayCameras/Source/GameplayCameras/Public/GameFramework/GameplayCamerasPlayerCameraManager.h b/Engine/Plugins/Cameras/GameplayCameras/Source/GameplayCameras/Public/GameFramework/GameplayCamerasPlayerCameraManager.h
--- a/Engine/Plugins/Cameras/GameplayCameras/Source/GameplayCameras/Public/GameFramework/GameplayCamerasPlayerCameraManager.h
+++ b/Engine/Plugins/Cameras/GameplayCameras/Source/GameplayCameras/Public/GameFramework/GameplayCamerasPlayerCameraManager.h
@@ -10,6 +10,7 @@
 
 class AGameplayCameraSystemActor;
 class UGameplayCameraSystemHost;
+struct FCameraRigInstanceID;
 
 namespace UE::Cameras
 {
@@ -58,6 +59,17 @@ public:
 	UFUNCTION(BlueprintCallable, Category="Camera")
 	void ReleasePlayerController();
 
+public:
+
+	UFUNCTION(BlueprintCallable, Category="Camera")
+	FCameraRigInstanceID StartGlobalCameraModifierRig(const UCameraRigAsset* CameraRig, int32 OrderKey = 0);
+
+	UFUNCTION(BlueprintCallable, Category="Camera")
+	FCameraRigInstanceID StartVisualCameraModifierRig(const UCameraRigAsset* CameraRig, int32 OrderKey = 0);
+
+	UFUNCTION(BlueprintCallable, Category="Camera")
+	void StopCameraModifierRig(FCameraRigInstanceID InstanceID, bool bImmediately = false);
+
 public:
 
 	// APlayerCameraManager interface.
diff --git a/Engine/Plugins/Cameras/GameplayCameras/Source/GameplayCameras/Public/Services/CameraModifierService.h b/Engine/Plugins/Cameras/GameplayCameras/Source/GameplayCameras/Public/Services/CameraModifierService.h
new file mode 100644
--- /dev/null
+++ b/Engine/Plugins/Cameras/GameplayCameras/Source/GameplayCameras/Public/Services/CameraModifierService.h
@@ -0,0 +1,54 @@
+// Copyright Epic Games, Inc. All Rights Reserved.
+
+#pragma once
+
+#include "Core/CameraEvaluationService.h"
+#include "Core/CameraRigInstanceID.h"
+
+class UCameraRigAsset;
+
+namespace UE::Cameras
+{
+
+class FCameraEvaluationContext;
+
+/**
+ * An evaluation service for running "camera modifiers", which are context-less camera rig instances
+ * meant to run additively in the camera system.
+ */
+class FCameraModifierService : public FCameraEvaluationService
+{
+	UE_DECLARE_CAMERA_EVALUATION_SERVICE(GAMEPLAYCAMERAS_API, FCameraModifierService)
+
+public:
+
+	/** Blend stack order key for modifiers. */
+	static int32 GetFirstBlendStackOrderKey() { return FirstBlendStackOrderKey; }
+
+	/** Starts a new instance of the given camera rig, using a "null" context. */
+	GAMEPLAYCAMERAS_API FCameraRigInstanceID StartCameraModifierRig(const UCameraRigAsset* CameraRig, ECameraRigLayer Layer, int32 OrderKey = 0);
+
+	/** Stops a running instance of a camera rig. */
+	GAMEPLAYCAMERAS_API void StopCameraModifierRig(FCameraRigInstanceID CameraRigID, bool bImmediately = false);
+
+protected:
+
+	// FCameraEvaluationService interface.
+	virtual void OnInitialize(const FCameraEvaluationServiceInitializeParams& Params) override;
+	virtual void OnTeardown(const FCameraEvaluationServiceTeardownParams& Params) override;
+
+private:
+
+	void EnsureModifierContextCreated();
+
+private:
+
+	static const int32 FirstBlendStackOrderKey = 100;
+
+	FCameraSystemEvaluator* Evaluator = nullptr;
+
+	TSharedPtr<FCameraEvaluationContext> ModifierContext;
+};
+
+}  // namespace UE::Cameras
+
