Cameras: fix slightly incorrect clamping code.

diff --git a/Engine/Plugins/Cameras/GameplayCameras/Source/GameplayCameras/Private/Nodes/Framing/BaseFramingCameraNode.cpp b/Engine/Plugins/Cameras/GameplayCameras/Source/GameplayCameras/Private/Nodes/Framing/BaseFramingCameraNode.cpp
--- a/Engine/Plugins/Cameras/GameplayCameras/Source/GameplayCameras/Private/Nodes/Framing/BaseFramingCameraNode.cpp
+++ b/Engine/Plugins/Cameras/GameplayCameras/Source/GameplayCameras/Private/Nodes/Framing/BaseFramingCameraNode.cpp
@@ -208,7 +208,7 @@ void FBaseFramingCameraNodeEvaluator::ComputeCurrentState(const FCameraNodeEvalu
 			const double DeadToCurrent = FVector2d::Distance(DeadEdgePoint, State.ScreenTarget);
 			const double DeadToHardEdge = FVector2d::Distance(DeadEdgePoint, HardEdgePoint);
 
-			const double Alpha = FMath::Clamp(DeadToCurrent / DeadToHardEdge, 0.0, 1.1);
+			const double Alpha = FMath::Clamp(DeadToCurrent / DeadToHardEdge, 0.0, 1.0);
 			State.ReframeDampingFactorAlpha = Alpha;
 
 #if UE_GAMEPLAY_CAMERAS_DEBUG
