Cameras: tweak IObjectTreeGraphObject API to support custom graph node titles

diff --git a/Engine/Plugins/Cameras/GameplayCameras/Source/GameplayCameras/Public/Core/ObjectTreeGraphObject.h b/Engine/Plugins/Cameras/GameplayCameras/Source/GameplayCameras/Public/Core/ObjectTreeGraphObject.h
--- a/Engine/Plugins/Cameras/GameplayCameras/Source/GameplayCameras/Public/Core/ObjectTreeGraphObject.h
+++ b/Engine/Plugins/Cameras/GameplayCameras/Source/GameplayCameras/Public/Core/ObjectTreeGraphObject.h
@@ -23,8 +23,10 @@ enum class EObjectTreeGraphObjectSupportFlags
 	None,
 	/** Supports storing a comment text. */
 	CommentText = 1 << 0,
+	/** Has a custom graph node title. */
+	CustomTitle = 1 << 1,
 	/** Supports custom renaming the graph node. */
-	CustomRename = 1 << 1
+	CustomRename = 1 << 2
 };
 ENUM_CLASS_FLAGS(EObjectTreeGraphObjectSupportFlags);
 
@@ -54,7 +56,7 @@ public:
 	/** Gets optional APIs support flags. */
 	virtual EObjectTreeGraphObjectSupportFlags GetSupportFlags(FName InGraphName) const { return EObjectTreeGraphObjectSupportFlags::None; }
 	/** Gets whether this object supports the given optional API. */
-	bool HasSupportFlags(FName InGraphName, EObjectTreeGraphObjectSupportFlags InFlags) const { return EnumHasAllFlags(GetSupportFlags(InGraphName), InFlags); }
+	bool HasAnySupportFlags(FName InGraphName, EObjectTreeGraphObjectSupportFlags InFlags) const { return EnumHasAnyFlags(GetSupportFlags(InGraphName), InFlags); }
 
 	/** Gets the graph node's comment text. */
 	virtual const FString& GetGraphNodeCommentText(FName InGraphName) const { static const FString EmptyString; return EmptyString; }
diff --git a/Engine/Plugins/Cameras/GameplayCameras/Source/GameplayCamerasEditor/Private/Editors/ObjectTreeGraphConfig.cpp b/Engine/Plugins/Cameras/GameplayCameras/Source/GameplayCamerasEditor/Private/Editors/ObjectTreeGraphConfig.cpp
--- a/Engine/Plugins/Cameras/GameplayCameras/Source/GameplayCamerasEditor/Private/Editors/ObjectTreeGraphConfig.cpp
+++ b/Engine/Plugins/Cameras/GameplayCameras/Source/GameplayCamerasEditor/Private/Editors/ObjectTreeGraphConfig.cpp
@@ -194,7 +194,8 @@ FText FObjectTreeGraphConfig::GetDisplayNameText(const UObject* InObject) const
 		const FObjectTreeGraphClassConfigs ClassConfig = GetObjectClassConfigs(InObject->GetClass());
 
 		const IObjectTreeGraphObject* GraphObject = Cast<IObjectTreeGraphObject>(InObject);
-		if (GraphObject && GraphObject->HasSupportFlags(GraphName, EObjectTreeGraphObjectSupportFlags::CustomRename))
+		if (GraphObject && GraphObject->HasAnySupportFlags(
+					GraphName, EObjectTreeGraphObjectSupportFlags::CustomTitle | EObjectTreeGraphObjectSupportFlags::CustomRename))
 		{
 			GraphObject->GetGraphNodeName(GraphName, DisplayNameText);
 		}
diff --git a/Engine/Plugins/Cameras/GameplayCameras/Source/GameplayCamerasEditor/Private/Editors/ObjectTreeGraphNode.cpp b/Engine/Plugins/Cameras/GameplayCameras/Source/GameplayCamerasEditor/Private/Editors/ObjectTreeGraphNode.cpp
--- a/Engine/Plugins/Cameras/GameplayCameras/Source/GameplayCamerasEditor/Private/Editors/ObjectTreeGraphNode.cpp
+++ b/Engine/Plugins/Cameras/GameplayCameras/Source/GameplayCamerasEditor/Private/Editors/ObjectTreeGraphNode.cpp
@@ -28,7 +28,7 @@ void UObjectTreeGraphNode::Initialize(UObject* InObject)
 
 	const FNodeContext NodeContext = GetNodeContext();
 	IObjectTreeGraphObject* GraphObject = Cast<IObjectTreeGraphObject>(InObject);
-	if (GraphObject && GraphObject->HasSupportFlags(NodeContext.GraphConfig.GraphName, EObjectTreeGraphObjectSupportFlags::CommentText))
+	if (GraphObject && GraphObject->HasAnySupportFlags(NodeContext.GraphConfig.GraphName, EObjectTreeGraphObjectSupportFlags::CommentText))
 	{
 		NodeComment = GraphObject->GetGraphNodeCommentText(NodeContext.GraphConfig.GraphName);
 	}
@@ -703,7 +703,7 @@ bool UObjectTreeGraphNode::GetCanRenameNode() const
 	UObject* Object = WeakObject.Get();
 	const FNodeContext NodeContext = GetNodeContext();
 	IObjectTreeGraphObject* GraphObject = Cast<IObjectTreeGraphObject>(Object);
-	return GraphObject && GraphObject->HasSupportFlags(NodeContext.GraphConfig.GraphName, EObjectTreeGraphObjectSupportFlags::CustomRename);
+	return GraphObject && GraphObject->HasAnySupportFlags(NodeContext.GraphConfig.GraphName, EObjectTreeGraphObjectSupportFlags::CustomRename);
 }
 
 void UObjectTreeGraphNode::OnRenameNode(const FString& NewName)
@@ -747,7 +747,7 @@ bool UObjectTreeGraphNode::SupportsCommentBubble() const
 	UObject* Object = WeakObject.Get();
 	const FNodeContext NodeContext = GetNodeContext();
 	IObjectTreeGraphObject* GraphObject = Cast<IObjectTreeGraphObject>(Object);
-	return GraphObject && GraphObject->HasSupportFlags(NodeContext.GraphConfig.GraphName, EObjectTreeGraphObjectSupportFlags::CommentText);
+	return GraphObject && GraphObject->HasAnySupportFlags(NodeContext.GraphConfig.GraphName, EObjectTreeGraphObjectSupportFlags::CommentText);
 }
 
 void UObjectTreeGraphNode::OnUpdateCommentText(const FString& NewComment)
