Cameras: Blueprint camera node now shows the BP evaluator name in the UI and debug HUD

diff --git a/Engine/Plugins/Cameras/GameplayCameras/Source/GameplayCameras/Private/Nodes/Utility/BlueprintCameraNode.cpp b/Engine/Plugins/Cameras/GameplayCameras/Source/GameplayCameras/Private/Nodes/Utility/BlueprintCameraNode.cpp
--- a/Engine/Plugins/Cameras/GameplayCameras/Source/GameplayCameras/Private/Nodes/Utility/BlueprintCameraNode.cpp
+++ b/Engine/Plugins/Cameras/GameplayCameras/Source/GameplayCameras/Private/Nodes/Utility/BlueprintCameraNode.cpp
@@ -11,6 +11,9 @@
 #include "Core/CameraSystemEvaluator.h"
 #include "Core/CameraVariableAssets.h"
 #include "Core/CameraVariableTable.h"
+#include "Debug/CameraDebugBlock.h"
+#include "Debug/CameraDebugBlockBuilder.h"
+#include "Debug/CameraDebugRenderer.h"
 #include "Engine/World.h"
 #include "GameFramework/Actor.h"
 #include "GameplayCameras.h"
@@ -35,6 +38,9 @@ protected:
 	virtual void OnInitialize(const FCameraNodeEvaluatorInitializeParams& Params, FCameraNodeEvaluationResult& OutResult) override;
 	virtual void OnRun(const FCameraNodeEvaluationParams& Params, FCameraNodeEvaluationResult& OutResult) override;
 	virtual void OnAddReferencedObjects(FReferenceCollector& Collector) override;
+#if UE_GAMEPLAY_CAMERAS_DEBUG
+	 virtual void OnBuildDebugBlocks(const FCameraDebugBlockBuildParams& Params, FCameraDebugBlockBuilder& Builder) override;
+#endif
 
 private:
 
@@ -51,6 +57,13 @@ private:
 
 UE_DEFINE_CAMERA_NODE_EVALUATOR(FBlueprintCameraNodeEvaluator)
 
+UE_DECLARE_CAMERA_DEBUG_BLOCK_START(GAMEPLAYCAMERAS_API, FBlueprintCameraDebugBlock)
+	UE_DECLARE_CAMERA_DEBUG_BLOCK_FIELD(FString, BlueprintEvaluatorName);
+UE_DECLARE_CAMERA_DEBUG_BLOCK_END()
+
+UE_DEFINE_CAMERA_DEBUG_BLOCK_WITH_FIELDS(FBlueprintCameraDebugBlock)
+
+
 void FBlueprintCameraNodeEvaluator::OnInitialize(const FCameraNodeEvaluatorInitializeParams& Params, FCameraNodeEvaluationResult& OutResult)
 {
 	const UBlueprintCameraNode* BlueprintNode = GetCameraNodeAs<UBlueprintCameraNode>();
@@ -138,6 +151,28 @@ void FBlueprintCameraNodeEvaluator::ApplyParameterOverrides(const FCameraVariabl
 	}
 }
 
+#if UE_GAMEPLAY_CAMERAS_DEBUG
+
+void FBlueprintCameraNodeEvaluator::OnBuildDebugBlocks(const FCameraDebugBlockBuildParams& Params, FCameraDebugBlockBuilder& Builder)
+{
+	FBlueprintCameraDebugBlock& DebugBlock = Builder.AttachDebugBlock<FBlueprintCameraDebugBlock>();
+
+	const UBlueprintCameraNode* BlueprintNode = GetCameraNodeAs<UBlueprintCameraNode>();
+	if (BlueprintNode->CameraNodeEvaluatorTemplate)
+	{
+		const UClass* BlueprintEvaluatorClass = BlueprintNode->CameraNodeEvaluatorTemplate->GetClass();
+		DebugBlock.BlueprintEvaluatorName = GetNameSafe(BlueprintEvaluatorClass);
+	}
+}
+
+void FBlueprintCameraDebugBlock::OnDebugDraw(const FCameraDebugBlockDrawParams& Params, FCameraDebugRenderer& Renderer)
+{
+	Renderer.AddText(BlueprintEvaluatorName);
+}
+
+#endif  // UE_GAMEPLAY_CAMERAS_DEBUG
+
+
 }  // namespace UE::Cameras
 
 void UBlueprintCameraNodeEvaluator::NativeInitializeCameraNode(const FCameraNodeEvaluatorInitializeParams& Params, FCameraNodeEvaluationResult& OutResult)
@@ -595,7 +630,20 @@ void UBlueprintCameraNode::PostEditChangeProperty( struct FPropertyChangedEvent&
 	}
 }
 
-#endif
+EObjectTreeGraphObjectSupportFlags UBlueprintCameraNode::GetSupportFlags(FName InGraphName) const
+{
+	return (Super::GetSupportFlags(InGraphName) | EObjectTreeGraphObjectSupportFlags::CustomTitle);
+}
+
+void UBlueprintCameraNode::GetGraphNodeName(FName InGraphName, FText& OutName) const
+{
+	const UClass* EvaluatorBlueprintClass = CameraNodeEvaluatorTemplate ? CameraNodeEvaluatorTemplate->GetClass() : nullptr;
+	const FText EvaluatorBlueprintName = EvaluatorBlueprintClass ? EvaluatorBlueprintClass->GetDisplayNameText() : LOCTEXT("None", "None");
+
+	OutName = FText::Format(LOCTEXT("GraphNodeNameFormat", "Blueprint ({0})"), EvaluatorBlueprintName);
+}
+
+#endif  // WITH_EDITOR
 
 #undef LOCTEXT_NAMESPACE
 
diff --git a/Engine/Plugins/Cameras/GameplayCameras/Source/GameplayCameras/Public/Nodes/Utility/BlueprintCameraNode.h b/Engine/Plugins/Cameras/GameplayCameras/Source/GameplayCameras/Public/Nodes/Utility/BlueprintCameraNode.h
--- a/Engine/Plugins/Cameras/GameplayCameras/Source/GameplayCameras/Public/Nodes/Utility/BlueprintCameraNode.h
+++ b/Engine/Plugins/Cameras/GameplayCameras/Source/GameplayCameras/Public/Nodes/Utility/BlueprintCameraNode.h
@@ -4,8 +4,8 @@
 
 #include "Core/CameraEvaluationContext.h"
 #include "Core/CameraNode.h"
-#include "Core/CameraVariableTableFwd.h"
 #include "Core/ICustomCameraNodeParameterProvider.h"
+#include "Core/ObjectTreeGraphObject.h"
 #include "GameFramework/BlueprintCameraPose.h"
 #include "GameFramework/BlueprintCameraVariableTable.h"
 
@@ -121,6 +121,12 @@ protected:
 	virtual void OnBuild(FCameraRigBuildContext& BuildContext) override;
 	virtual FCameraNodeEvaluatorPtr OnBuildEvaluator(FCameraNodeEvaluatorBuilder& Builder) const override;
 
+	// IObjectTreeGraphObject interface.
+#if WITH_EDITOR
+	virtual EObjectTreeGraphObjectSupportFlags GetSupportFlags(FName InGraphName) const override;
+	virtual void GetGraphNodeName(FName InGraphName, FText& OutName) const override;
+#endif  // WITH_EDITOR
+
 	// ICustomCameraNodeParameterProvider interface.
 	virtual void GetCustomCameraNodeParameters(FCustomCameraNodeParameterInfos& OutParameterInfos) override;
 
