Cameras: extract most of the camera rig logic into a base "camera object" class so it can be shared with the upcoming camera shakes

- Camera objects have a root camera node, some parameters, and an allocation info.
- The vast majority of the change is simply about renaming classes and moving files. Nothing should be different.

diff --git a/Engine/Plugins/Cameras/GameplayCameras/Config/BaseGameplayCameras.ini b/Engine/Plugins/Cameras/GameplayCameras/Config/BaseGameplayCameras.ini
--- a/Engine/Plugins/Cameras/GameplayCameras/Config/BaseGameplayCameras.ini
+++ b/Engine/Plugins/Cameras/GameplayCameras/Config/BaseGameplayCameras.ini
@@ -10,3 +10,12 @@
 +StructRedirects=(OldName="/Script/GameplayCameras.CameraTargetInfo",NewName="CameraActorTargetInfo")
 +FunctionRedirects=(OldName="BlueprintCameraDirectorEvaluator.ActivateCameraRigPrefab", NewName="ActivateCameraRig")
 +PropertyRedirects=(OldName="/Script/GameplayCameras.CameraRigInterface.InterfaceParameters", NewName="BlendableParameters")
+
++EnumRedirects=(OldName="/Script/GameplayCameras.ECameraRigInterfaceParameterType",NewName="ECameraObjectInterfaceCameraParameterType")
++StructRedirects=(OldName="/Script/GameplayCameras.CameraRigInterface",NewName="CameraObjectInterface")
++ClassRedirects=(OldName="/Script/GameplayCameras.CameraRigInterfaceParameterBase",NewName="CameraObjectInterfaceParameterBase")
++ClassRedirects=(OldName="/Script/GameplayCameras.CameraRigBlendableParameter",NewName="CameraObjectInterfaceBlendableParameter")
++ClassRedirects=(OldName="/Script/GameplayCameras.CameraRigDataParameter",NewName="CameraObjectInterfaceDataParameter")
++StructRedirects=(OldName="/Script/GameplayCameras.CameraRigParameterDefinition",NewName="CameraObjectInterfaceParameterDefinition")
++StructRedirects=(OldName="/Script/GameplayCameras.CameraRigAssetReferenceParameterMetaData",NewName="CameraObjectInterfaceParameterMetaData")
+
diff --git a/Engine/Plugins/Cameras/GameplayCameras/Source/GameplayCameras/Private/Build/CameraObjectInterfaceBuilder.cpp b/Engine/Plugins/Cameras/GameplayCameras/Source/GameplayCameras/Private/Build/CameraObjectInterfaceBuilder.cpp
new file mode 100644
--- /dev/null
+++ b/Engine/Plugins/Cameras/GameplayCameras/Source/GameplayCameras/Private/Build/CameraObjectInterfaceBuilder.cpp
@@ -0,0 +1,898 @@
+// Copyright Epic Games, Inc. All Rights Reserved.
+
+#include "Build/CameraObjectInterfaceBuilder.h"
+
+#include "Core/BaseCameraObject.h"
+#include "Core/CameraNode.h"
+#include "Core/CameraNodeHierarchy.h"
+#include "Core/CameraObjectInterface.h"
+#include "Core/CameraParameters.h"
+#include "Core/CameraVariableReferences.h"
+#include "Core/ICustomCameraNodeParameterProvider.h"
+#include "GameplayCameras.h"
+
+#define LOCTEXT_NAMESPACE "CameraObjectInterfaceBuilder"
+
+namespace UE::Cameras
+{
+
+namespace Internal
+{
+
+struct FInterfaceParameterBindingBuilder
+{
+	UBaseCameraObject* CameraObject;
+
+	FInterfaceParameterBindingBuilder(FCameraObjectInterfaceBuilder& InOwner)
+		: Owner(InOwner)
+	{
+		CameraObject = Owner.CameraObject;
+	}
+
+	void ReportError(FText&& ErrorMessage)
+	{
+		ReportError(nullptr, MoveTemp(ErrorMessage));
+	}
+
+	void ReportError(UObject* Object, FText&& ErrorMessage)
+	{
+		Owner.BuildLog.AddMessage(EMessageSeverity::Error, MoveTemp(ErrorMessage));
+	}
+
+	template<typename CameraParameterOrVariableReferenceType>
+	void SetCameraParameterOrVariableReferenceOverride(
+			const UCameraObjectInterfaceBlendableParameter* BlendableParameter,
+			FStructProperty* TargetProperty,
+			CameraParameterOrVariableReferenceType* CameraParameterOrVariableReferencePtr)
+	{
+		using VariableAssetType = typename CameraParameterOrVariableReferenceType::VariableAssetType;
+		using ValueType = typename VariableAssetType::ValueType;
+
+		ensure(BlendableParameter->PrivateVariableID.IsValid());
+		ensure(BlendableParameter->TargetPropertyName == TargetProperty->GetFName());
+
+		const bool bIsValid = CheckIfParameterCanBeOverridden(BlendableParameter, CameraParameterOrVariableReferencePtr);
+		if (!bIsValid)
+		{
+			return;
+		}
+
+		UObject* TargetNode = BlendableParameter->Target;
+		FCameraVariableID PreviousVariableID = FindOldDrivingVariableID(TargetProperty->GetFName(), TargetNode);
+		if (PreviousVariableID != BlendableParameter->PrivateVariableID)
+		{
+			TargetNode->Modify();
+		}
+
+		CameraParameterOrVariableReferencePtr->VariableID = BlendableParameter->PrivateVariableID;
+	}
+
+	template<typename VariableAssetType>
+	void SetCustomBlendableParameterOverride(
+			const UCameraObjectInterfaceBlendableParameter* BlendableParameter,
+			const FCustomCameraNodeParameterInfos::FBlendableParameterInfo& CustomParameter)
+	{
+		using ValueType = typename VariableAssetType::ValueType;
+
+		ensure(BlendableParameter->PrivateVariableID.IsValid());
+		ensure(BlendableParameter->TargetPropertyName == CustomParameter.ParameterName);
+
+		const bool bIsValid = CheckIfParameterCanBeOverridden(BlendableParameter, CustomParameter);
+		if (!bIsValid)
+		{
+			return;
+		}
+
+		UObject* TargetNode = BlendableParameter->Target;
+		FCameraVariableID PreviousVariableID = FindOldDrivingVariableID(CustomParameter.ParameterName, TargetNode);
+		if (PreviousVariableID != BlendableParameter->PrivateVariableID)
+		{
+			TargetNode->Modify();
+		}
+
+		(*CustomParameter.OverrideVariableID) = BlendableParameter->PrivateVariableID;
+	}
+
+	void SetCustomBlendableStructParameterOverride(
+			const UCameraObjectInterfaceBlendableParameter* BlendableParameter,
+			const FCustomCameraNodeParameterInfos::FBlendableParameterInfo& CustomParameter)
+	{
+		ensure(BlendableParameter->PrivateVariableID.IsValid());
+		ensure(BlendableParameter->TargetPropertyName == CustomParameter.ParameterName);
+
+		const bool bIsValid = CheckIfParameterCanBeOverridden(BlendableParameter, CustomParameter.OverrideVariableID);
+		if (!bIsValid)
+		{
+			return;
+		}
+
+		// Also ensure the struct type is compatible.
+		if (CustomParameter.BlendableStructType != BlendableParameter->BlendableStructType)
+		{
+			ReportError(BlendableParameter->Target,
+					FText::Format(
+						LOCTEXT(
+							"IncompatibleBlendableStructType", 
+							"Invalid interface parameter '{0}', driving property '{1}' on '{2}': expected type {3} but was {4}"),
+						FText::FromString(BlendableParameter->InterfaceParameterName),
+						FText::FromName(BlendableParameter->TargetPropertyName),
+						FText::FromName(BlendableParameter->Target->GetFName()),
+#if WITH_EDITORONLY_DATA
+						CustomParameter.BlendableStructType->GetDisplayNameText(),
+						BlendableParameter->BlendableStructType->GetDisplayNameText()
+#else
+						FText::FromName(CustomParameter.BlendableStructType->GetFName()),
+						FText::FromName(BlendableParameter->BlendableStructType->GetFName())
+#endif
+						));
+			return;
+		}
+
+		UObject* TargetNode = BlendableParameter->Target;
+		FCameraVariableID PreviousVariableID = FindOldDrivingVariableID(CustomParameter.ParameterName, TargetNode);
+		if (PreviousVariableID != BlendableParameter->PrivateVariableID)
+		{
+			TargetNode->Modify();
+		}
+
+		(*CustomParameter.OverrideVariableID) = BlendableParameter->PrivateVariableID;
+	}
+
+	void SetDataContextPropertyOverride(
+			const UCameraObjectInterfaceDataParameter* DataParameter,
+			FProperty* TargetProperty,
+			FCameraContextDataID* OverrideDataID)
+	{
+		ensure(DataParameter->PrivateDataID);
+
+		const bool bIsValid = CheckIfParameterCanBeOverridden(DataParameter, OverrideDataID);
+		if (!bIsValid)
+		{
+			return;
+		}
+
+		UObject* TargetNode = DataParameter->Target;
+		FCameraContextDataID PreviousDataID = FindOldDrivingDataID(TargetProperty->GetFName(), TargetNode);
+		if (PreviousDataID != DataParameter->PrivateDataID)
+		{
+			TargetNode->Modify();
+		}
+
+		(*OverrideDataID) = DataParameter->PrivateDataID;
+	}
+
+	void SetCustomDataParameterOverride(
+			const UCameraObjectInterfaceDataParameter* DataParameter,
+			const FCustomCameraNodeParameterInfos::FDataParameterInfo& CustomParameter)
+	{
+		ensure(DataParameter->PrivateDataID);
+		ensure(DataParameter->TargetPropertyName == CustomParameter.ParameterName);
+
+		const bool bIsValid = CheckIfParameterCanBeOverridden(DataParameter, CustomParameter.OverrideDataID);
+		if (!bIsValid)
+		{
+			return;
+		}
+
+		UObject* TargetNode = DataParameter->Target;
+		FCameraContextDataID PreviousDataID = FindOldDrivingDataID(CustomParameter.ParameterName, TargetNode);
+		if (PreviousDataID != DataParameter->PrivateDataID)
+		{
+			TargetNode->Modify();
+		}
+
+		(*CustomParameter.OverrideDataID) = DataParameter->PrivateDataID;
+	}
+
+private:
+
+	FCameraVariableID FindOldDrivingVariableID(FName ForParameterName, UObject* ForObject)
+	{
+		using FDrivenParameterKey = FCameraObjectInterfaceBuilder::FDrivenParameterKey;
+
+		FDrivenParameterKey ParameterKey{ ForParameterName, ForObject };
+		FCameraVariableID ReusedVariableID;
+		Owner.OldDrivenBlendableParameters.RemoveAndCopyValue(ParameterKey, ReusedVariableID);
+		return ReusedVariableID;
+	}
+
+	FCameraContextDataID FindOldDrivingDataID(FName ForParameterName, UObject* ForObject)
+	{
+		using FDrivenParameterKey = FCameraObjectInterfaceBuilder::FDrivenParameterKey;
+
+		FDrivenParameterKey ParameterKey{ ForParameterName, ForObject };
+		FCameraContextDataID ReusedDataID;
+		Owner.OldDrivenDataParameters.RemoveAndCopyValue(ParameterKey, ReusedDataID);
+		return ReusedDataID;
+	}
+
+	template<typename CameraParameterOrVariableReferenceType>
+	bool CheckIfParameterCanBeOverridden(
+			const UCameraObjectInterfaceBlendableParameter* BlendableParameter, 
+			CameraParameterOrVariableReferenceType* CameraParameterOrVariableReference)
+	{
+		if (CameraParameterOrVariableReference->Variable != nullptr)
+		{
+			ReportError(BlendableParameter->Target,
+					FText::Format(
+						LOCTEXT(
+							"BlendableParameterDrivenTwice", 
+							"Camera node parameter '{0}.{1}' is both exposed and driven by a variable!"),
+						FText::FromName(BlendableParameter->Target->GetFName()), 
+						FText::FromName(BlendableParameter->TargetPropertyName)));
+			return false;
+		}
+
+		return CheckIfParameterCanBeOverridden(BlendableParameter, &CameraParameterOrVariableReference->VariableID);
+	}
+
+	bool CheckIfParameterCanBeOverridden(
+			const UCameraObjectInterfaceBlendableParameter* BlendableParameter, 
+			const FCustomCameraNodeParameterInfos::FBlendableParameterInfo& CustomParameter)
+	{
+		if (CustomParameter.OverrideVariable != nullptr)
+		{
+			ReportError(BlendableParameter->Target,
+					FText::Format(
+						LOCTEXT(
+							"BlendableParameterDrivenTwice", 
+							"Camera node parameter '{0}.{1}' is both exposed and driven by a variable!"),
+						FText::FromName(BlendableParameter->Target->GetFName()), 
+						FText::FromName(BlendableParameter->TargetPropertyName)));
+			return false;
+		}
+
+		return CheckIfParameterCanBeOverridden(BlendableParameter, CustomParameter.OverrideVariableID);
+	}
+
+	bool CheckIfParameterCanBeOverridden(
+			const UCameraObjectInterfaceBlendableParameter* BlendableParameter, 
+			FCameraVariableID* VariableID)
+	{
+		if (!VariableID)
+		{
+			ReportError(BlendableParameter->Target,
+					FText::Format(
+						LOCTEXT(
+							"BlendableParameterMissingOverrideID", 
+							"Camera node parameter '{0}.{1}' cannot be overriden by a parameter"),
+						FText::FromName(BlendableParameter->Target->GetFName()), 
+						FText::FromName(BlendableParameter->TargetPropertyName)));
+			return false;
+		}
+		if (VariableID->IsValid())
+		{
+			ReportError(BlendableParameter->Target,
+					FText::Format(
+						LOCTEXT(
+							"BlendableParameterDrivenTwice", 
+							"Camera node parameter '{0}.{1}' is both exposed and driven by a variable!"),
+						FText::FromName(BlendableParameter->Target->GetFName()), 
+						FText::FromName(BlendableParameter->TargetPropertyName)));
+			return false;
+		}
+		return true;
+	}
+
+	bool CheckIfParameterCanBeOverridden(
+			const UCameraObjectInterfaceDataParameter* DataParameter,
+			FCameraContextDataID* DataID)
+	{
+		if (!DataID)
+		{
+			ReportError(DataParameter->Target,
+					FText::Format(
+						LOCTEXT(
+							"DataParameterMissingOverrideID", 
+							"Camera node parameter '{0}.{1}' cannot be overriden by a parameter"),
+						FText::FromName(DataParameter->Target->GetFName()), 
+						FText::FromName(DataParameter->TargetPropertyName)));
+			return false;
+		}
+		if (DataID->IsValid())
+		{
+			ReportError(DataParameter->Target,
+					FText::Format(
+						LOCTEXT(
+							"DataParameterDrivenTwice",
+							"Camera node parameter '{0}.{1}' is somehow overriden twice!"),
+						FText::FromName(DataParameter->Target->GetFName()), 
+						FText::FromName(DataParameter->TargetPropertyName)));
+			return false;
+		}
+		return true;
+	}
+
+private:
+
+	FCameraObjectInterfaceBuilder& Owner;
+};
+
+}  // namespace Internal
+
+FCameraObjectInterfaceBuilder::FCameraObjectInterfaceBuilder(FCameraBuildLog& InBuildLog)
+	: BuildLog(InBuildLog)
+{
+}
+
+void FCameraObjectInterfaceBuilder::BuildInterface(UBaseCameraObject* InCameraObject, const FCameraNodeHierarchy& InHierarchy, bool bCollectStrayNodes)
+{
+	TSet<UCameraNode*> CameraNodesToGather(InHierarchy.GetFlattenedHierarchy());
+
+	if (bCollectStrayNodes)
+	{
+		// Get the list of nodes, both connected and disconnected from the root hierarchy.
+		// We could use AllNodeTreeObjects for that, but it only exists in editor builds, and we 
+		// don't want to rely on unit tests or runtime data manipulation to have correctly populated 
+		// it, so we'll try to gather any stray nodes by looking at objects outer'ed to the camera rig.
+		ForEachObjectWithOuter(InCameraObject, [&CameraNodesToGather](UObject* Obj)
+				{
+					if (UCameraNode* CameraNode = Cast<UCameraNode>(Obj))
+					{
+						CameraNodesToGather.Add(CameraNode);
+					}
+				});
+		const int32 NumStrayCameraNodes = (CameraNodesToGather.Num() - InHierarchy.Num());
+		if (NumStrayCameraNodes > 0)
+		{
+			UE_LOG(LogCameraSystem, Verbose, TEXT("Collected %d stray camera nodes while building camera rig '%s'."),
+					NumStrayCameraNodes, *GetPathNameSafe(CameraObject));
+		}
+	}
+
+	BuildInterface(InCameraObject, CameraNodesToGather.Array());
+}
+
+void FCameraObjectInterfaceBuilder::BuildInterface(UBaseCameraObject* InCameraObject, TArrayView<UCameraNode*> InCameraObjectNodes)
+{
+	if (!ensure(InCameraObject))
+	{
+		return;
+	}
+
+	CameraObject = InCameraObject;
+	CameraObjectNodes = InCameraObjectNodes;
+	{
+		BuildInterfaceImpl();
+	}
+	CameraObject = nullptr;
+	CameraObjectNodes.Reset();
+}
+
+void FCameraObjectInterfaceBuilder::BuildInterfaceImpl()
+{
+	GatherOldDrivenParameters();
+	BuildInterfaceParameters();
+	BuildInterfaceParameterBindings();
+	DiscardUnusedParameters();
+}
+
+void FCameraObjectInterfaceBuilder::GatherOldDrivenParameters()
+{
+	// Keep track of which blendable/data parameters were previously overriden with private IDs.
+	// Then clear those private IDs. This is because it's easier to rebuild all this from a blank 
+	// slate than trying to figure out what changed.
+	//
+	// As we rebuild things in BuildInterfaceParameterBindings, we compare to the old state to
+	// figure out if we need to flag anything as modified for the current transaction.
+	//
+	// Note that parameters driven by user-defined variables are left alone.
+
+	OldDrivenBlendableParameters.Reset();
+	OldDrivenDataParameters.Reset();
+
+	for (UCameraNode* CameraNode : CameraObjectNodes)
+	{
+		UClass* CameraNodeClass = CameraNode->GetClass();
+		
+		for (TFieldIterator<FProperty> It(CameraNodeClass); It; ++It)
+		{
+			FProperty* Property(*It);
+
+			// First look for some blendable camera parameters.
+			if (FStructProperty* StructProperty = CastField<FStructProperty>(Property))
+			{
+				bool bFoundCameraParameter = true;
+#define UE_CAMERA_VARIABLE_FOR_TYPE(ValueType, ValueName)\
+				if (StructProperty->Struct == F##ValueName##CameraParameter::StaticStruct())\
+				{\
+					auto* CameraParameterPtr = StructProperty->ContainerPtrToValuePtr<F##ValueName##CameraParameter>(CameraNode);\
+					if (CameraParameterPtr->VariableID.IsValid() && !CameraParameterPtr->Variable)\
+					{\
+						OldDrivenBlendableParameters.Add(\
+								FDrivenParameterKey{ StructProperty->GetFName(), CameraNode },\
+								CameraParameterPtr->VariableID);\
+						CameraParameterPtr->VariableID = FCameraVariableID();\
+					}\
+				}\
+				else if (StructProperty->Struct == F##ValueName##CameraVariableReference::StaticStruct())\
+				{\
+					auto* VariableReferencePtr = StructProperty->ContainerPtrToValuePtr<F##ValueName##CameraVariableReference>(CameraNode);\
+					if (VariableReferencePtr->VariableID.IsValid() && !VariableReferencePtr->Variable)\
+					{\
+						OldDrivenBlendableParameters.Add(\
+								FDrivenParameterKey{ StructProperty->GetFName(), CameraNode },\
+								VariableReferencePtr->VariableID);\
+						VariableReferencePtr->VariableID = FCameraVariableID();\
+					}\
+				}\
+				else
+UE_CAMERA_VARIABLE_FOR_ALL_TYPES()
+#undef UE_CAMERA_VARIABLE_FOR_TYPE
+				{
+					// Other struct type...
+					bFoundCameraParameter = false;
+				}
+
+				if (bFoundCameraParameter)
+				{
+					continue;
+				}
+			}
+
+			// Then look for some data parameters.
+			const FName DataIDPropertyName = FName(It->GetName() + TEXT("DataID"));
+			FStructProperty* DataIDStructProperty = CastField<FStructProperty>(CameraNodeClass->FindPropertyByName(DataIDPropertyName));
+			if (DataIDStructProperty && DataIDStructProperty->Struct == FCameraContextDataID::StaticStruct())
+			{
+				FCameraContextDataID* ExistingDataID = DataIDStructProperty->ContainerPtrToValuePtr<FCameraContextDataID>(CameraNode);
+				if (ExistingDataID->IsValid())
+				{
+					OldDrivenDataParameters.Add(FDrivenParameterKey{ Property->GetFName(), CameraNode }, *ExistingDataID);
+					*ExistingDataID = FCameraContextDataID();
+				}
+			}
+		}
+
+		if (ICustomCameraNodeParameterProvider* CustomParameterProvider = Cast<ICustomCameraNodeParameterProvider>(CameraNode))
+		{
+			FCustomCameraNodeParameterInfos CustomParameters;
+			CustomParameterProvider->GetCustomCameraNodeParameters(CustomParameters);
+
+			for (const FCustomCameraNodeParameterInfos::FBlendableParameterInfo& BlendableParameter : CustomParameters.BlendableParameters)
+			{
+				if (!BlendableParameter.OverrideVariableID)
+				{
+					continue;
+				}
+
+				if (BlendableParameter.OverrideVariableID->IsValid())
+				{
+					OldDrivenBlendableParameters.Add(
+							FDrivenParameterKey{ BlendableParameter.ParameterName, CameraNode },
+							*BlendableParameter.OverrideVariableID);
+					(*BlendableParameter.OverrideVariableID) = FCameraVariableID();
+				}
+			}
+
+			for (const FCustomCameraNodeParameterInfos::FDataParameterInfo& DataParameter : CustomParameters.DataParameters)
+			{
+				if (!DataParameter.OverrideDataID)
+				{
+					continue;
+				}
+
+				if (DataParameter.OverrideDataID->IsValid())
+				{
+					OldDrivenDataParameters.Add(
+							FDrivenParameterKey{ DataParameter.ParameterName, CameraNode },
+							*DataParameter.OverrideDataID);
+					(*DataParameter.OverrideDataID) = FCameraContextDataID();
+				}
+			}
+		}
+	}
+}
+
+void FCameraObjectInterfaceBuilder::BuildInterfaceParameters()
+{
+	// Here we simply validate all blendable/data interface parameters and create IDs for their entries in the
+	// variable and context data tables.
+
+	using namespace Internal;
+
+	for (auto It = CameraObject->Interface.BlendableParameters.CreateIterator(); It; ++It)
+	{
+		UCameraObjectInterfaceBlendableParameter* BlendableParameter(*It);
+
+		// Basic validations.
+		if (!BlendableParameter)
+		{
+			BuildLog.AddMessage(EMessageSeverity::Warning,
+					CameraObject,
+					LOCTEXT("InvalidBlendableParameter", "Invalid interface parameter was found and removed."));
+
+			CameraObject->Modify();
+			It.RemoveCurrent();
+
+			continue;
+		}
+
+		if (BlendableParameter->InterfaceParameterName.IsEmpty())
+		{
+			BuildLog.AddMessage(EMessageSeverity::Error,
+					BlendableParameter,
+					LOCTEXT(
+						"InvalidBlendableParameterName",
+						"Invalid interface parameter name."));
+			continue;
+		}
+
+		// Create a new private variable ID for this interface parameter. Flag the parameter as changed if
+		// the ID is different, generally when it's a new parameter.
+		FCameraVariableID VariableID = FCameraVariableID::FromHashValue(GetTypeHash(BlendableParameter->GetGuid()));
+		if (BlendableParameter->PrivateVariableID != VariableID)
+		{
+			BlendableParameter->Modify();
+			BlendableParameter->PrivateVariableID = VariableID;
+		}
+	}
+
+	for (auto It = CameraObject->Interface.DataParameters.CreateIterator(); It; ++It)
+	{
+		UCameraObjectInterfaceDataParameter* DataParameter(*It);
+
+		// Basic validations.
+		if (!DataParameter)
+		{
+			BuildLog.AddMessage(EMessageSeverity::Warning,
+					CameraObject,
+					LOCTEXT("InvalidDataParameter", "Invalid interface parameter was found and removed."));
+
+			CameraObject->Modify();
+			It.RemoveCurrent();
+
+			continue;
+		}
+
+		if (DataParameter->InterfaceParameterName.IsEmpty())
+		{
+			BuildLog.AddMessage(EMessageSeverity::Error,
+					DataParameter,
+					LOCTEXT(
+						"InvalidDataParameterName",
+						"Invalid interface parameter name."));
+			continue;
+		}
+
+		// Create a new private data ID for this interface parameter. Flag the parameter as changed if
+		// the ID is different, generally when it's a new parameter.
+		FCameraContextDataID DataID = FCameraContextDataID::FromHashValue(GetTypeHash(DataParameter->GetGuid()));
+		if (DataParameter->PrivateDataID != DataID)
+		{
+			DataParameter->Modify();
+			DataParameter->PrivateDataID = DataID;
+		}
+	}
+}
+
+void FCameraObjectInterfaceBuilder::BuildInterfaceParameterBindings()
+{
+	// Now we connect the interface paramters to whatever node property they are supposed to drive.
+	// Each time we need to check for either a custom property (via ICustomCameraNodeParameterProvider),
+	// or a UObject property found with reflection.
+
+	using FBuiltDrivenParameter = TTuple<UObject*, FName>;
+	TSet<FBuiltDrivenParameter> BuiltDrivenParameters;
+
+	const FString CameraObjectName = CameraObject->GetName();
+	const FString CameraObjectPathName = CameraObject->GetPathName();
+
+	for (const UCameraObjectInterfaceBlendableParameter* BlendableParameter : CameraObject->Interface.BlendableParameters)
+	{
+		// Basic validations.
+		if (!BlendableParameter->Target)
+		{
+			BuildLog.AddMessage(EMessageSeverity::Warning,
+					BlendableParameter,
+					LOCTEXT(
+						"InvalidBlendableParameterTarget",
+						"Invalid interface parameter: it has no target node."));
+			continue;
+		}
+		if (BlendableParameter->TargetPropertyName.IsNone())
+		{
+			BuildLog.AddMessage(EMessageSeverity::Error,
+					BlendableParameter,
+					LOCTEXT(
+						"InvalidBlendableParameterTargetPropertyName", 
+						"Invalid interface parameter: it has not target property name."));
+			continue;
+		}
+		if (BlendableParameter->InterfaceParameterName.IsEmpty())
+		{
+			continue;
+		}
+
+		// Check duplicate bindings.
+		FBuiltDrivenParameter NewDrivenParameter(BlendableParameter->Target, BlendableParameter->TargetPropertyName);
+		if (BuiltDrivenParameters.Contains(NewDrivenParameter))
+		{
+			BuildLog.AddMessage(EMessageSeverity::Error,
+					FText::Format(LOCTEXT(
+						"BlendableParameterTargetCollision",
+						"Multiple interface parameters targeting property '{0}' on camera node '{1}'. Ignoring duplicates."),
+						FText::FromName(BlendableParameter->TargetPropertyName),
+						FText::FromName(BlendableParameter->Target->GetFName())));
+			continue;
+		}
+		BuiltDrivenParameters.Add(NewDrivenParameter);
+
+		// See if this interface parameter is overriding a camera node parameter.
+		// Otherwise, maybe it's targeting a camera rig node's override for an inner rig interface parameter.
+		if (SetupCustomBlendableParameterOverride(BlendableParameter))
+		{
+			// Implicit continue.
+		}
+		else if (SetupCameraParameterOrVariableReferenceOverride(BlendableParameter))
+		{
+			// Implicit continue.
+		}
+		else
+		{
+			UObject* Target = BlendableParameter->Target;
+			BuildLog.AddMessage(EMessageSeverity::Error,
+					Target,
+					FText::Format(LOCTEXT(
+						"InvalidBlendableParameterTargetProperty",
+						"Invalid interface parameter '{0}', driving property '{1}' on '{2}', but no such property found."),
+						FText::FromString(BlendableParameter->InterfaceParameterName), 
+						FText::FromName(BlendableParameter->TargetPropertyName),
+						FText::FromName(Target->GetFName())));
+		}
+	}
+
+	for (const UCameraObjectInterfaceDataParameter* DataParameter : CameraObject->Interface.DataParameters)
+	{
+		// Basic validations.
+		if (!DataParameter->Target)
+		{
+			BuildLog.AddMessage(EMessageSeverity::Warning,
+					DataParameter,
+					LOCTEXT(
+						"InvalidDataParameterTarget",
+						"Invalid interface parameter: it has no target node."));
+			continue;
+		}
+		if (DataParameter->TargetPropertyName.IsNone())
+		{
+			BuildLog.AddMessage(EMessageSeverity::Error,
+					DataParameter,
+					LOCTEXT(
+						"InvalidDataParameterTargetPropertyName", 
+						"Invalid interface parameter: it has not target property name."));
+			continue;
+		}
+		if (DataParameter->InterfaceParameterName.IsEmpty())
+		{
+			continue;
+		}
+
+		if (SetupCustomDataParameterOverride(DataParameter))
+		{
+			// Implicit continue.
+		}
+		else if (SetupDataContextPropertyOverride(DataParameter))
+		{
+			// Implicit continue.
+		}
+		else
+		{
+			UObject* Target = DataParameter->Target;
+			BuildLog.AddMessage(EMessageSeverity::Error,
+					Target,
+					FText::Format(LOCTEXT(
+						"InvalidDataParameterTargetProperty",
+						"Invalid interface parameter '{0}', driving property '{1}' on '{2}', but no such property found."),
+						FText::FromString(DataParameter->InterfaceParameterName), 
+						FText::FromName(DataParameter->TargetPropertyName),
+						FText::FromName(Target->GetFName())));
+		}
+	}
+}
+
+bool FCameraObjectInterfaceBuilder::SetupCameraParameterOrVariableReferenceOverride(const UCameraObjectInterfaceBlendableParameter* BlendableParameter)
+{
+	using namespace Internal;
+
+	// Here we hook up interface parameters connected to a camera node property. This property is supposed
+	// to be of one of the camera parameter types (FBooleanCameraParameter, FInteger32CameraParameter, etc.)
+	// so they have both a fixed value (bool, int32, etc.) and a "private variable" which we will set to the 
+	// private variable of the given interface parameter, checking that the types match (UBooleanCameraVariable, 
+	// UInteger32CameraVariable, etc.)
+
+	UObject* Target = BlendableParameter->Target;
+	UClass* TargetClass = Target->GetClass();
+	FProperty* TargetProperty = TargetClass->FindPropertyByName(BlendableParameter->TargetPropertyName);
+	if (!TargetProperty)
+	{
+		// No match, try something else.
+		return false;
+	}
+
+	FStructProperty* TargetStructProperty = CastField<FStructProperty>(TargetProperty);
+	if (!TargetStructProperty)
+	{
+		BuildLog.AddMessage(EMessageSeverity::Error,
+				Target,
+				FText::Format(LOCTEXT(
+						"InvalidCameraNodeParameter",
+						"Invalid interface parameter '{0}', driving property '{1}' on '{2}', but it's not a camera parameter."),
+					FText::FromString(BlendableParameter->InterfaceParameterName), 
+					FText::FromName(BlendableParameter->TargetPropertyName),
+					FText::FromName(Target->GetFName())));
+		return true;
+	}
+
+	// Get the type of the camera parameter by matching the struct against all the types we support,
+	// and create a private camera variable asset to drive its value.
+	FInterfaceParameterBindingBuilder Builder(*this);
+#define UE_CAMERA_VARIABLE_FOR_TYPE(ValueType, ValueName)\
+	if (TargetStructProperty->Struct == F##ValueName##CameraParameter::StaticStruct())\
+	{\
+		auto* CameraParameterPtr = TargetStructProperty->ContainerPtrToValuePtr<F##ValueName##CameraParameter>(Target);\
+		Builder.SetCameraParameterOrVariableReferenceOverride<F##ValueName##CameraParameter>(\
+				BlendableParameter, TargetStructProperty, CameraParameterPtr\
+				);\
+	}\
+	else if (TargetStructProperty->Struct == F##ValueName##CameraVariableReference::StaticStruct())\
+	{\
+		auto* VariableReferencePtr = TargetStructProperty->ContainerPtrToValuePtr<F##ValueName##CameraVariableReference>(Target);\
+		Builder.SetCameraParameterOrVariableReferenceOverride<F##ValueName##CameraVariableReference>(\
+				BlendableParameter, TargetStructProperty, VariableReferencePtr\
+				);\
+	}\
+	else
+	UE_CAMERA_VARIABLE_FOR_ALL_TYPES()
+#undef UE_CAMERA_VARIABLE_FOR_TYPE
+	{
+		BuildLog.AddMessage(EMessageSeverity::Error,
+				BlendableParameter,
+				FText::Format(LOCTEXT(
+						"InvalidCameraNodeParameter",
+						"Invalid interface parameter '{0}', driving property '{1}' on '{2}', but it's not a camera parameter."),
+					FText::FromString(BlendableParameter->InterfaceParameterName), 
+					FText::FromName(BlendableParameter->TargetPropertyName),
+					FText::FromName(Target->GetFName())));
+	}
+
+	return true;
+}
+
+bool FCameraObjectInterfaceBuilder::SetupCustomBlendableParameterOverride(const UCameraObjectInterfaceBlendableParameter* BlendableParameter)
+{
+	using namespace Internal;
+
+	ICustomCameraNodeParameterProvider* Target = Cast<ICustomCameraNodeParameterProvider>(BlendableParameter->Target);
+	if (!Target)
+	{
+		// No match, try something else.
+		return false;
+	}
+
+	// Look for a parameter override matching the target name.
+	// TODO: we're querying the list of custom parameters every time, we may want to cache it for this phase.
+	FCustomCameraNodeParameterInfos CustomParameters;
+	Target->GetCustomCameraNodeParameters(CustomParameters);
+
+	FCustomCameraNodeParameterInfos::FBlendableParameterInfo* TargetCustomParameter = 
+		CustomParameters.BlendableParameters.FindByPredicate(
+			[BlendableParameter](FCustomCameraNodeParameterInfos::FBlendableParameterInfo& CustomParameter)
+			{
+				return CustomParameter.ParameterName == BlendableParameter->TargetPropertyName;
+			});
+	if (!TargetCustomParameter)
+	{
+		// No match, try something else.
+		return false;
+	}
+
+	FInterfaceParameterBindingBuilder Builder(*this);
+	switch (TargetCustomParameter->ParameterType)
+	{
+#define UE_CAMERA_VARIABLE_FOR_TYPE(ValueType, ValueName)\
+		case ECameraVariableType::ValueName:\
+			Builder.SetCustomBlendableParameterOverride<U##ValueName##CameraVariable>(BlendableParameter, *TargetCustomParameter);\
+			break;
+	UE_CAMERA_VARIABLE_FOR_ALL_TYPES()
+#undef UE_CAMERA_VARIABLE_FOR_TYPE
+		case ECameraVariableType::BlendableStruct:
+			Builder.SetCustomBlendableStructParameterOverride(BlendableParameter, *TargetCustomParameter);
+			break;
+	}
+
+	return true;
+}
+
+bool FCameraObjectInterfaceBuilder::SetupDataContextPropertyOverride(const UCameraObjectInterfaceDataParameter* DataParameter)
+{
+	using namespace Internal;
+
+	UObject* Target = DataParameter->Target;
+	UClass* TargetClass = Target->GetClass();
+	FProperty* TargetProperty = TargetClass->FindPropertyByName(DataParameter->TargetPropertyName);
+	if (!TargetProperty)
+	{
+		// No match, try something else.
+		return false;
+	}
+
+	const FName TargetDataIDPropertyName = FName(TargetProperty->GetName() + TEXT("DataID"));
+	FStructProperty* TargetDataIDProperty = CastField<FStructProperty>(TargetClass->FindPropertyByName(TargetDataIDPropertyName));
+	if (!TargetDataIDProperty || TargetDataIDProperty->Struct != FCameraContextDataID::StaticStruct())
+	{
+		UE_LOG(LogCameraSystem, Error,
+				TEXT("Interface parameter '{0}' is driving data context property '{1}' on '{2}' "
+					 "but no FCameraContextDataID property '{3}' was found to store the override ID."),
+				*DataParameter->InterfaceParameterName,
+				*DataParameter->TargetPropertyName.ToString(),
+				*Target->GetName(),
+				*TargetDataIDPropertyName.ToString());
+		return false;
+	}
+
+	FCameraContextDataID* OverrideDataID = TargetDataIDProperty->ContainerPtrToValuePtr<FCameraContextDataID>(Target);
+
+	FInterfaceParameterBindingBuilder Builder(*this);
+	Builder.SetDataContextPropertyOverride(DataParameter, TargetProperty, OverrideDataID);
+
+	return true;
+}
+
+bool FCameraObjectInterfaceBuilder::SetupCustomDataParameterOverride(const UCameraObjectInterfaceDataParameter* DataParameter)
+{
+	using namespace Internal;
+
+	ICustomCameraNodeParameterProvider* Target = Cast<ICustomCameraNodeParameterProvider>(DataParameter->Target);
+	if (!Target)
+	{
+		// No match, try something else.
+		return false;
+	}
+
+	FCustomCameraNodeParameterInfos CustomParameters;
+	Target->GetCustomCameraNodeParameters(CustomParameters);
+
+	FCustomCameraNodeParameterInfos::FDataParameterInfo* TargetCustomParameter =
+		CustomParameters.DataParameters.FindByPredicate(
+				[DataParameter](FCustomCameraNodeParameterInfos::FDataParameterInfo& CustomParameter)
+				{
+					return CustomParameter.ParameterName == DataParameter->TargetPropertyName;
+				});
+	if (!TargetCustomParameter)
+	{
+		// No match, try something else.
+		return false;
+	}
+
+	FInterfaceParameterBindingBuilder Builder(*this);
+	Builder.SetCustomDataParameterOverride(DataParameter, *TargetCustomParameter);
+
+	return true;
+}
+
+void FCameraObjectInterfaceBuilder::DiscardUnusedParameters()
+{
+	// Now that we've rebuilt all exposed parameters, anything left from the old list 
+	// must be discarded. These are nodes and properties that used to be driven by
+	// variables and now aren't, so we need to flag them as modified.
+
+	for (TPair<FDrivenParameterKey, FCameraVariableID> Pair : OldDrivenBlendableParameters)
+	{
+		UObject* Target = Pair.Key.Value;
+		Target->Modify();
+	}
+	OldDrivenBlendableParameters.Reset();
+
+	for (TPair<FDrivenParameterKey, FCameraContextDataID> Pair : OldDrivenDataParameters)
+	{
+		UObject* Target = Pair.Key.Value;
+		Target->Modify();
+	}
+	OldDrivenDataParameters.Reset();
+}
+
+}  // namespace UE::Cameras
+
+#undef LOCTEXT_NAMESPACE
+
diff --git a/Engine/Plugins/Cameras/GameplayCameras/Source/GameplayCameras/Private/Build/CameraObjectInterfaceParameterBuilder.cpp b/Engine/Plugins/Cameras/GameplayCameras/Source/GameplayCameras/Private/Build/CameraObjectInterfaceParameterBuilder.cpp
new file mode 100644
--- /dev/null
+++ b/Engine/Plugins/Cameras/GameplayCameras/Source/GameplayCameras/Private/Build/CameraObjectInterfaceParameterBuilder.cpp
@@ -0,0 +1,398 @@
+// Copyright Epic Games, Inc. All Rights Reserved.
+
+#include "Build/CameraObjectInterfaceParameterBuilder.h"
+
+#include "Core/BaseCameraObject.h"
+#include "Core/CameraNode.h"
+#include "Core/CameraObjectInterfaceParameterDefinition.h"
+#include "Core/CameraParameters.h"
+#include "Core/ICustomCameraNodeParameterProvider.h"
+#include "StructUtils/PropertyBag.h"
+#include "UObject/UnrealType.h"
+
+namespace UE::Cameras
+{
+
+FCameraObjectInterfaceParameterBuilder::FCameraObjectInterfaceParameterBuilder()
+{
+}
+
+void FCameraObjectInterfaceParameterBuilder::BuildParameters(UBaseCameraObject* InCameraObject)
+{
+	CameraObject = InCameraObject;
+	{
+		BuildParametersImpl();
+	}
+	CameraObject = nullptr;
+}
+
+void FCameraObjectInterfaceParameterBuilder::BuildParametersImpl()
+{
+	BuildParameterDefinitions();
+	BuildDefaultParameters();
+}
+
+void FCameraObjectInterfaceParameterBuilder::BuildParameterDefinitions()
+{
+	TArray<FCameraObjectInterfaceParameterDefinition> ParameterDefinitions;
+
+	for (const UCameraObjectInterfaceBlendableParameter* BlendableParameter : CameraObject->Interface.BlendableParameters)
+	{
+		if (BlendableParameter && BlendableParameter->PrivateVariableID)
+		{
+			FCameraObjectInterfaceParameterDefinition Definition;
+			Definition.ParameterName = FName(BlendableParameter->InterfaceParameterName);
+			Definition.ParameterGuid = BlendableParameter->GetGuid();
+			Definition.ParameterType = ECameraObjectInterfaceParameterType::Blendable;
+			Definition.VariableID = BlendableParameter->PrivateVariableID;
+			Definition.VariableType = BlendableParameter->ParameterType;
+			Definition.BlendableStructType = BlendableParameter->BlendableStructType;
+			ParameterDefinitions.Add(Definition);
+		}
+	}
+
+	for (const UCameraObjectInterfaceDataParameter* DataParameter : CameraObject->Interface.DataParameters)
+	{
+		if (DataParameter && DataParameter->PrivateDataID)
+		{
+			FCameraObjectInterfaceParameterDefinition Definition;
+			Definition.ParameterName = FName(DataParameter->InterfaceParameterName);
+			Definition.ParameterGuid = DataParameter->GetGuid();
+			Definition.ParameterType = ECameraObjectInterfaceParameterType::Data;
+			Definition.DataID = DataParameter->PrivateDataID;
+			Definition.DataType = DataParameter->DataType;
+			Definition.DataContainerType = DataParameter->DataContainerType;
+			Definition.DataTypeObject = DataParameter->DataTypeObject;
+			ParameterDefinitions.Add(Definition);
+		}
+	}
+
+	if (ParameterDefinitions != CameraObject->ParameterDefinitions)
+	{
+		CameraObject->Modify();
+		CameraObject->ParameterDefinitions = ParameterDefinitions;
+	}
+}
+
+void FCameraObjectInterfaceParameterBuilder::BuildDefaultParameters()
+{
+	FInstancedPropertyBag DefaultParameters;
+	BuildDefaultParameters(CameraObject, DefaultParameters);
+	if (!DefaultParameters.Identical(&CameraObject->DefaultParameters, 0))
+	{
+		CameraObject->Modify();
+		CameraObject->DefaultParameters = DefaultParameters;
+	}
+}
+
+void FCameraObjectInterfaceParameterBuilder::BuildDefaultParameters(const UBaseCameraObject* CameraObject, FInstancedPropertyBag& OutPropertyBag)
+{
+	TArray<FPropertyBagPropertyDesc> DefaultParameterProperties;
+	AppendDefaultParameterProperties(CameraObject, DefaultParameterProperties);
+	OutPropertyBag.AddProperties(DefaultParameterProperties);
+	SetDefaultParameterValues(CameraObject, OutPropertyBag);
+}
+
+void FCameraObjectInterfaceParameterBuilder::AppendDefaultParameterProperties(const UBaseCameraObject* CameraObject, TArray<FPropertyBagPropertyDesc>& OutProperties)
+{
+	for (const FCameraObjectInterfaceParameterDefinition& Definition : CameraObject->GetParameterDefinitions())
+	{
+		bool bIsValidProperty = true;
+		EPropertyBagPropertyType PropertyType = EPropertyBagPropertyType::Struct;
+		EPropertyBagContainerType ContainerType = EPropertyBagContainerType::None;
+		const UObject* PropertyTypeObject = nullptr;
+		EPropertyFlags PropertyFlags = CPF_None;
+
+		if (Definition.ParameterType == ECameraObjectInterfaceParameterType::Blendable)
+		{
+			switch (Definition.VariableType)
+			{
+#define UE_CAMERA_VARIABLE_FOR_TYPE(ValueType, ValueName)\
+				case ECameraVariableType::ValueName:\
+					PropertyTypeObject = F##ValueName##CameraParameter::StaticStruct();\
+					break;
+				UE_CAMERA_VARIABLE_FOR_ALL_TYPES()
+#undef UE_CAMERA_VARIABLE_FOR_TYPE
+				case ECameraVariableType::BlendableStruct:
+					PropertyTypeObject = Definition.BlendableStructType;
+					PropertyFlags = CPF_Interp;
+					break;
+				default:
+					ensure(false);
+					break;
+			}
+		}
+		else if (Definition.ParameterType == ECameraObjectInterfaceParameterType::Data)
+		{
+			PropertyTypeObject = Definition.DataTypeObject;
+
+			switch (Definition.DataType)
+			{
+				case ECameraContextDataType::Name:
+					PropertyType = EPropertyBagPropertyType::Name;
+					PropertyFlags = CPF_Interp;
+					break;
+				case ECameraContextDataType::String:
+					PropertyType = EPropertyBagPropertyType::String;
+					PropertyFlags = CPF_Interp;
+					break;
+				case ECameraContextDataType::Enum:
+					PropertyType = EPropertyBagPropertyType::Enum;
+					ensure(PropertyTypeObject && PropertyTypeObject->IsA<UEnum>());
+					PropertyFlags = CPF_Interp;
+					break;
+				case ECameraContextDataType::Struct:
+					PropertyType = EPropertyBagPropertyType::Struct;
+					ensure(PropertyTypeObject && PropertyTypeObject->IsA<UScriptStruct>());
+					break;
+				case ECameraContextDataType::Object:
+					PropertyType = EPropertyBagPropertyType::Object;
+					PropertyFlags = CPF_Interp;
+					break;
+				case ECameraContextDataType::Class:
+					PropertyType = EPropertyBagPropertyType::Class;
+					PropertyFlags = CPF_Interp;
+					break;
+				default:
+					bIsValidProperty = false;
+					break;
+			}
+
+			switch (Definition.DataContainerType)
+			{
+				case ECameraContextDataContainerType::Array:
+					ContainerType = EPropertyBagContainerType::Array;
+					break;
+			}
+		}
+		else
+		{
+			bIsValidProperty = false;
+		}
+
+		if (ensure(bIsValidProperty))
+		{
+			FPropertyBagPropertyDesc NewProperty(Definition.ParameterName, ContainerType, PropertyType, PropertyTypeObject);
+			// Make the property bag match the camera interface parameter GUIDs.
+			NewProperty.ID = Definition.ParameterGuid;
+			NewProperty.PropertyFlags |= PropertyFlags;
+
+			OutProperties.Add(NewProperty);
+		}
+	}
+}
+
+void FCameraObjectInterfaceParameterBuilder::SetDefaultParameterValues(const UBaseCameraObject* CameraObject, FInstancedPropertyBag& PropertyBag)
+{
+	uint8* PropertyBagValue = PropertyBag.GetMutableValue().GetMemory();
+	const UPropertyBag* PropertyBagStruct = PropertyBag.GetPropertyBagStruct();
+	if (!ensure(PropertyBagValue && PropertyBagStruct))
+	{
+		return;
+	}
+
+	for (const UCameraObjectInterfaceBlendableParameter* BlendableParameter : CameraObject->Interface.BlendableParameters)
+	{
+		if (!ensure(BlendableParameter))
+		{
+			continue;
+		}
+
+		UCameraNode* CameraNode = BlendableParameter->Target;
+		if (!CameraNode)
+		{
+			continue;
+		}
+
+		const void* RawSourceValuePtr = nullptr;
+
+		// First check if the value is found on a custom parameter.
+		if (ICustomCameraNodeParameterProvider* CustomParameterProvider = Cast<ICustomCameraNodeParameterProvider>(CameraNode))
+		{
+			FCustomCameraNodeParameterInfos CustomParameters;
+			CustomParameterProvider->GetCustomCameraNodeParameters(CustomParameters);
+
+			FCustomCameraNodeParameterInfos::FBlendableParameterInfo* TargetCustomParameter = 
+				CustomParameters.BlendableParameters.FindByPredicate(
+						[BlendableParameter](FCustomCameraNodeParameterInfos::FBlendableParameterInfo& CustomParameter)
+						{
+							return CustomParameter.ParameterName == BlendableParameter->TargetPropertyName;
+						});
+			if (TargetCustomParameter)
+			{
+				RawSourceValuePtr = TargetCustomParameter->DefaultValue;
+			}
+		}
+
+		// If not found, check on a reflected UObject property.
+		if (!RawSourceValuePtr)
+		{
+			const UClass* TargetClass = CameraNode->GetClass();
+			FStructProperty* StructProperty = CastField<FStructProperty>(
+					TargetClass->FindPropertyByName(BlendableParameter->TargetPropertyName));
+			if (StructProperty)
+			{
+				switch (BlendableParameter->ParameterType)
+				{
+#define UE_CAMERA_VARIABLE_FOR_TYPE(ValueType, ValueName)\
+					case ECameraVariableType::ValueName:\
+						{\
+							using CameraParameterType = F##ValueName##CameraParameter;\
+							if (ensure(StructProperty->Struct == CameraParameterType::StaticStruct()))\
+							{\
+								CameraParameterType* CameraParameterPtr = StructProperty->ContainerPtrToValuePtr<CameraParameterType>(CameraNode);\
+								RawSourceValuePtr = static_cast<void*>(&CameraParameterPtr->Value);\
+							}\
+						}\
+						break;
+					UE_CAMERA_VARIABLE_FOR_ALL_TYPES()
+#undef UE_CAMERA_VARIABLE_FOR_TYPE
+					case ECameraVariableType::BlendableStruct:
+						{
+							RawSourceValuePtr = StructProperty->ContainerPtrToValuePtr<void>(CameraNode);
+						}
+						break;
+				}
+			}
+		}
+
+		if (!ensure(RawSourceValuePtr))
+		{
+			continue;
+		}
+
+		// Find the corresponding property on the default parameters' property bag.
+		const FPropertyBagPropertyDesc* PropertyDesc = PropertyBagStruct->FindPropertyDescByID(BlendableParameter->GetGuid());
+		if (!ensure(PropertyDesc && PropertyDesc->CachedProperty))
+		{
+			continue;
+		}
+
+		// This property should be a structure: either a camera parameter for all the standard blendable types,
+		// or a blendable structure.
+		const FStructProperty* DefaultParameterProperty = CastField<FStructProperty>(PropertyDesc->CachedProperty);
+		if (!ensure(DefaultParameterProperty))
+		{
+			continue;
+		}
+
+		switch (BlendableParameter->ParameterType)
+		{
+#define UE_CAMERA_VARIABLE_FOR_TYPE(ValueType, ValueName)\
+			case ECameraVariableType::ValueName:\
+				{\
+					using CameraParameterType = F##ValueName##CameraParameter;\
+					if (ensure(DefaultParameterProperty->Struct == CameraParameterType::StaticStruct()))\
+					{\
+						const ValueType* SourceValuePtr = reinterpret_cast<const ValueType*>(RawSourceValuePtr);\
+						CameraParameterType* DestinationParameter = DefaultParameterProperty->ContainerPtrToValuePtr<CameraParameterType>(PropertyBagValue);\
+						DestinationParameter->Value = *SourceValuePtr;\
+					}\
+				}\
+				break;
+			UE_CAMERA_VARIABLE_FOR_ALL_TYPES()
+#undef UE_CAMERA_VARIABLE_FOR_TYPE
+			case ECameraVariableType::BlendableStruct:
+				if (ensure(DefaultParameterProperty->Struct == BlendableParameter->BlendableStructType))
+				{
+					void* RawDestinationValuePtr = DefaultParameterProperty->ContainerPtrToValuePtr<void>(PropertyBagValue);
+					BlendableParameter->BlendableStructType->CopyScriptStruct(RawDestinationValuePtr, RawSourceValuePtr);
+				}
+				break;
+		}
+	}
+
+	for (const UCameraObjectInterfaceDataParameter* DataParameter : CameraObject->Interface.DataParameters)
+	{
+		if (!ensure(DataParameter && DataParameter->Target))
+		{
+			continue;
+		}
+
+		UCameraNode* CameraNode = DataParameter->Target;
+
+		const uint8* RawSourceValuePtr = nullptr;
+
+		if (ICustomCameraNodeParameterProvider* CustomParameterProvider = Cast<ICustomCameraNodeParameterProvider>(CameraNode))
+		{
+			FCustomCameraNodeParameterInfos CustomParameters;
+			CustomParameterProvider->GetCustomCameraNodeParameters(CustomParameters);
+
+			FCustomCameraNodeParameterInfos::FDataParameterInfo* TargetCustomParameter = 
+				CustomParameters.DataParameters.FindByPredicate(
+						[DataParameter](FCustomCameraNodeParameterInfos::FDataParameterInfo& CustomParameter)
+						{
+							return CustomParameter.ParameterName == DataParameter->TargetPropertyName;
+						});
+			if (TargetCustomParameter)
+			{
+				RawSourceValuePtr = TargetCustomParameter->DefaultValue;
+			}
+		}
+
+		if (!RawSourceValuePtr)
+		{
+			const UClass* TargetClass = DataParameter->Target->GetClass();
+			FProperty* Property = TargetClass->FindPropertyByName(DataParameter->TargetPropertyName);
+			if (Property)
+			{
+				RawSourceValuePtr = (uint8*)Property->ContainerPtrToValuePtr<void>(DataParameter->Target);
+			}
+		}
+
+		if (!ensure(RawSourceValuePtr))
+		{
+			continue;
+		}
+
+		const FPropertyBagPropertyDesc* PropertyDesc = PropertyBagStruct->FindPropertyDescByID(DataParameter->GetGuid());
+		if (!ensure(PropertyDesc && PropertyDesc->CachedProperty))
+		{
+			continue;
+		}
+
+		if (DataParameter->DataContainerType == ECameraContextDataContainerType::None)
+		{
+			void* RawDestinationValuePtr = PropertyDesc->CachedProperty->ContainerPtrToValuePtr<void>(PropertyBagValue);
+			if (ensure(RawDestinationValuePtr))
+			{
+				SetDefaultParameterValue(DataParameter, RawDestinationValuePtr, RawSourceValuePtr);
+			}
+		}
+		else if (DataParameter->DataContainerType == ECameraContextDataContainerType::Array)
+		{
+			// Array properties are empty by default.
+		}
+	}
+}
+
+void FCameraObjectInterfaceParameterBuilder::SetDefaultParameterValue(const UCameraObjectInterfaceDataParameter* DataParameter, void* DestValuePtr, const void* SrcValuePtr)
+{
+	switch (DataParameter->DataType)
+	{
+		case ECameraContextDataType::Name:
+			*((FName*)DestValuePtr) = *((FName*)SrcValuePtr);
+			break;
+		case ECameraContextDataType::String:
+			*((FString*)DestValuePtr) = *((FString*)SrcValuePtr);
+			break;
+		case ECameraContextDataType::Enum:
+			*((uint8*)DestValuePtr) = *((uint8*)SrcValuePtr);
+			break;
+		case ECameraContextDataType::Struct:
+			{
+				const UScriptStruct* StructType = CastChecked<const UScriptStruct>(DataParameter->DataTypeObject);
+				StructType->CopyScriptStruct(DestValuePtr, SrcValuePtr);
+			}
+			break;
+		case ECameraContextDataType::Object:
+			*((FObjectPtr*)DestValuePtr) = *((FObjectPtr*)SrcValuePtr);
+			break;
+		case ECameraContextDataType::Class:
+			*((FObjectPtr*)DestValuePtr) = *((FObjectPtr*)SrcValuePtr);
+			break;
+	}
+}
+
+}  // namespace UE::Cameras
+
diff --git a/Engine/Plugins/Cameras/GameplayCameras/Source/GameplayCameras/Private/Build/CameraRigAssetBuilder.cpp b/Engine/Plugins/Cameras/GameplayCameras/Source/GameplayCameras/Private/Build/CameraRigAssetBuilder.cpp
new file mode 100644
--- /dev/null
+++ b/Engine/Plugins/Cameras/GameplayCameras/Source/GameplayCameras/Private/Build/CameraRigAssetBuilder.cpp
@@ -0,0 +1,92 @@
+// Copyright Epic Games, Inc. All Rights Reserved.
+
+#include "Build/CameraRigAssetBuilder.h"
+
+#include "Build/CameraNodeHierarchyBuilder.h"
+#include "Build/CameraObjectInterfaceBuilder.h"
+#include "Build/CameraObjectInterfaceParameterBuilder.h"
+#include "Core/CameraRigAsset.h"
+#include "GameplayCamerasDelegates.h"
+
+#define LOCTEXT_NAMESPACE "CameraRigAssetBuilder"
+
+namespace UE::Cameras
+{
+
+FCameraRigAssetBuilder::FCameraRigAssetBuilder(FCameraBuildLog& InBuildLog)
+	: BuildLog(InBuildLog)
+{
+}
+
+void FCameraRigAssetBuilder::BuildCameraRig(UCameraRigAsset* InCameraRig)
+{
+	BuildCameraRig(InCameraRig, FCustomBuildStep::CreateLambda([](UCameraRigAsset*, FCameraBuildLog&) {}));
+}
+
+void FCameraRigAssetBuilder::BuildCameraRig(UCameraRigAsset* InCameraRig, FCustomBuildStep InCustomBuildStep)
+{
+	if (!ensure(InCameraRig))
+	{
+		return;
+	}
+
+	CameraRig = InCameraRig;
+
+	BuildLog.SetLoggingPrefix(CameraRig->GetPathName() + TEXT(": "));
+	{
+		BuildCameraRigImpl();
+
+		InCustomBuildStep.ExecuteIfBound(CameraRig, BuildLog);
+
+		CameraRig->EventHandlers.Notify(&ICameraRigAssetEventHandler::OnCameraRigBuilt, CameraRig);
+	}
+	BuildLog.SetLoggingPrefix(FString());
+	UpdateBuildStatus();
+
+	FGameplayCamerasDelegates::OnCameraRigAssetBuilt().Broadcast(CameraRig);
+
+	CameraRig = nullptr;
+}
+
+void FCameraRigAssetBuilder::BuildCameraRigImpl()
+{
+	if (!CameraRig->RootNode)
+	{
+		BuildLog.AddMessage(EMessageSeverity::Error, CameraRig, 
+				FText::Format(LOCTEXT("MissingRootNode", "Camera rig '{0}' has no root node."), 
+					FText::FromString(GetPathNameSafe(CameraRig))));
+		return;
+	}
+
+	FCameraNodeHierarchyBuilder NodeBuilder(BuildLog, CameraRig);
+	NodeBuilder.PreBuild();
+
+	FCameraObjectInterfaceBuilder InterfaceBuilder(BuildLog);
+	InterfaceBuilder.BuildInterface(CameraRig, NodeBuilder.GetHierarchy(), true);
+
+	NodeBuilder.Build();
+
+	FCameraObjectInterfaceParameterBuilder ParameterBuilder;
+	ParameterBuilder.BuildParameters(CameraRig);
+}
+
+void FCameraRigAssetBuilder::UpdateBuildStatus()
+{
+	ECameraBuildStatus BuildStatus = ECameraBuildStatus::Clean;
+	if (BuildLog.HasErrors())
+	{
+		BuildStatus = ECameraBuildStatus::WithErrors;
+	}
+	else if (BuildLog.HasWarnings())
+	{
+		BuildStatus = ECameraBuildStatus::CleanWithWarnings;
+	}
+
+	// Don't modify the camera rig: BuildStatus is transient.
+	CameraRig->BuildStatus = BuildStatus;
+}
+
+}  // namespace UE::Cameras
+
+#undef LOCTEXT_NAMESPACE
+
diff --git a/Engine/Plugins/Cameras/GameplayCameras/Source/GameplayCameras/Private/Core/BaseCameraObjectReference.cpp b/Engine/Plugins/Cameras/GameplayCameras/Source/GameplayCameras/Private/Core/BaseCameraObjectReference.cpp
new file mode 100644
--- /dev/null
+++ b/Engine/Plugins/Cameras/GameplayCameras/Source/GameplayCameras/Private/Core/BaseCameraObjectReference.cpp
@@ -0,0 +1,290 @@
+// Copyright Epic Games, Inc. All Rights Reserved.
+
+#include "Core/BaseCameraObjectReference.h"
+
+#include "Core/BaseCameraObject.h"
+#include "Core/CameraNodeEvaluator.h"
+#include "Core/CameraParameters.h"
+#include "Core/ICustomCameraNodeParameterProvider.h"
+
+#include UE_INLINE_GENERATED_CPP_BY_NAME(BaseCameraObjectReference)
+
+const FCameraObjectInterfaceParameterMetaData* FBaseCameraObjectReference::FindMetaData(const FGuid& PropertyID) const
+{
+	return ParameterMetaData.FindByPredicate(
+			[PropertyID](const FCameraObjectInterfaceParameterMetaData& Item)
+			{
+				return Item.ParameterGuid == PropertyID;
+			});
+}
+
+FCameraObjectInterfaceParameterMetaData& FBaseCameraObjectReference::FindOrAddMetaData(const FGuid& PropertyID)
+{
+	FCameraObjectInterfaceParameterMetaData* ExistingMetaData = ParameterMetaData.FindByPredicate(
+			[PropertyID](FCameraObjectInterfaceParameterMetaData& Item)
+			{
+				return Item.ParameterGuid == PropertyID;
+			});
+	if (ExistingMetaData)
+	{
+		return *ExistingMetaData;
+	}
+
+	FCameraObjectInterfaceParameterMetaData& NewMetaData = ParameterMetaData.Emplace_GetRef();
+	NewMetaData.ParameterGuid = PropertyID;
+	return NewMetaData;
+}
+
+bool FBaseCameraObjectReference::IsParameterOverridden(const FGuid& PropertyID) const
+{
+	if (const FCameraObjectInterfaceParameterMetaData* MetaData = FindMetaData(PropertyID))
+	{
+		return MetaData->bIsOverridden;
+	}
+	return false;
+}
+
+void FBaseCameraObjectReference::SetParameterOverridden(const FGuid& PropertyID, bool bIsOverridden)
+{
+	FCameraObjectInterfaceParameterMetaData& MetaData = FindOrAddMetaData(PropertyID);
+	MetaData.bIsOverridden = bIsOverridden;
+}
+
+bool FBaseCameraObjectReference::IsParameterAnimated(const FGuid& PropertyID) const
+{
+	if (const FCameraObjectInterfaceParameterMetaData* MetaData = FindMetaData(PropertyID))
+	{
+		return MetaData->bIsAnimated;
+	}
+	return false;
+}
+
+void FBaseCameraObjectReference::SetParameterAnimated(const FGuid& PropertyID, bool bIsAnimated)
+{
+	FCameraObjectInterfaceParameterMetaData& MetaData = FindOrAddMetaData(PropertyID);
+	MetaData.bIsAnimated = bIsAnimated;
+}
+
+bool FBaseCameraObjectReference::NeedsRebuildParameters() const
+{
+	const UBaseCameraObject* CameraObject = GetCameraObject();
+
+	if ((!CameraObject && Parameters.IsValid()) || (CameraObject && !Parameters.IsValid()))
+	{
+		return true;
+	}
+
+	if (CameraObject)
+	{
+		const UPropertyBag* AssetParametersType = CameraObject->GetDefaultParameters().GetPropertyBagStruct();
+		const UPropertyBag* ReferenceParametersType = Parameters.GetPropertyBagStruct();
+		if (AssetParametersType != ReferenceParametersType)
+		{
+			return true;
+		}
+	}
+
+	return false;
+}
+
+bool FBaseCameraObjectReference::RebuildParametersIfNeeded()
+{
+	if (NeedsRebuildParameters())
+	{
+		RebuildParameters();
+		return true;
+	}
+	return false;
+}
+
+void FBaseCameraObjectReference::RebuildParameters()
+{
+	const UBaseCameraObject* CameraObject = GetCameraObject();
+
+	if (CameraObject)
+	{
+		TArray<FGuid> ParameterOverrideGuids;
+		GetOverriddenParameterGuids(ParameterOverrideGuids);
+		Parameters.MigrateToNewBagInstanceWithOverrides(CameraObject->GetDefaultParameters(), ParameterOverrideGuids);
+
+		if (const UPropertyBag* ParametersType = Parameters.GetPropertyBagStruct())
+		{
+			// Remove metadata for parameters that don't exist anymore, and add default metadata for
+			// new parameters.
+			TSet<FGuid> ExistingMetaDataIDs;
+			for (TArray<FCameraObjectInterfaceParameterMetaData>::TIterator It = ParameterMetaData.CreateIterator(); It; ++It)
+			{
+				ExistingMetaDataIDs.Add(It->ParameterGuid);
+			}
+
+			TSet<FGuid> WantedMetaDataIDs;
+			TMap<FGuid, ECameraObjectInterfaceParameterType> WantedParameterTypes;
+			for (const FCameraObjectInterfaceParameterDefinition& Definition : CameraObject->GetParameterDefinitions())
+			{
+				WantedMetaDataIDs.Add(Definition.ParameterGuid);
+				WantedParameterTypes.Add(Definition.ParameterGuid, Definition.ParameterType);
+			}
+
+			TSet<FGuid> RemovedMetaDataIDs = ExistingMetaDataIDs.Difference(WantedMetaDataIDs);
+			for (TArray<FCameraObjectInterfaceParameterMetaData>::TIterator It = ParameterMetaData.CreateIterator(); It; ++It)
+			{
+				if (RemovedMetaDataIDs.Contains(It->ParameterGuid))
+				{
+					It.RemoveCurrentSwap();
+				}
+			}
+
+			TSet<FGuid> AddedMetaDataIDs = WantedMetaDataIDs.Difference(ExistingMetaDataIDs);
+			for (const FGuid& ParameterGuid : AddedMetaDataIDs)
+			{
+				FCameraObjectInterfaceParameterMetaData NewMetaData;
+				NewMetaData.ParameterGuid = ParameterGuid;
+				ParameterMetaData.Add(NewMetaData);
+			}
+		}
+		else
+		{
+			ParameterMetaData.Reset();
+		}
+	}
+	else
+	{
+		Parameters.Reset();
+		ParameterMetaData.Reset();
+	}
+}
+
+void FBaseCameraObjectReference::GetCustomCameraNodeParameters(FCustomCameraNodeParameterInfos& OutParameterInfos)
+{
+	RebuildParametersIfNeeded();
+
+	const UBaseCameraObject* CameraObject = GetCameraObject();
+	if (!CameraObject)
+	{
+		return;
+	}
+
+	const UPropertyBag* ParametersStruct = Parameters.GetPropertyBagStruct();
+	const FInstancedPropertyBag& DefaultParameters = CameraObject->GetDefaultParameters();
+	if (!ensure(ParametersStruct && ParametersStruct == DefaultParameters.GetPropertyBagStruct()))
+	{
+		return;
+	}
+
+	for (const FCameraObjectInterfaceParameterDefinition& Definition : CameraObject->GetParameterDefinitions())
+	{
+		const FPropertyBagPropertyDesc* PropertyDesc = ParametersStruct->FindPropertyDescByID(Definition.ParameterGuid);
+		if (!ensure(PropertyDesc))
+		{
+			continue;
+		}
+
+		if (Definition.ParameterType == ECameraObjectInterfaceParameterType::Blendable)
+		{
+			if (!ensure(PropertyDesc->ValueType == EPropertyBagPropertyType::Struct))
+			{
+				continue;
+			}
+
+			switch (Definition.VariableType)
+			{
+#define UE_CAMERA_VARIABLE_FOR_TYPE(ValueType, ValueName)\
+				case ECameraVariableType::ValueName:\
+					if (ensure(PropertyDesc->ValueTypeObject == F##ValueName##CameraParameter::StaticStruct()))\
+					{\
+						using CameraParameterType = F##ValueName##CameraParameter;\
+						TValueOrError<CameraParameterType*, EPropertyBagResult> PropertyValue =\
+							Parameters.GetValueStruct<CameraParameterType>(*PropertyDesc);\
+						if (ensure(PropertyValue.HasValue() && !PropertyValue.HasError()))\
+						{\
+							CameraParameterType* CameraParameter = PropertyValue.GetValue();\
+							check(CameraParameter);\
+							OutParameterInfos.AddBlendableParameter(\
+									Definition.ParameterName,\
+									Definition.VariableType,\
+									nullptr,\
+									reinterpret_cast<uint8*>(&CameraParameter->Value),\
+									&CameraParameter->VariableID);\
+						}\
+					}\
+					break;
+				UE_CAMERA_VARIABLE_FOR_ALL_TYPES()
+#undef UE_CAMERA_VARIABLE_FOR_TYPE
+				case ECameraVariableType::BlendableStruct:
+					{
+						TValueOrError<FStructView, EPropertyBagResult> PropertyValue =
+							Parameters.GetValueStruct(*PropertyDesc, Definition.BlendableStructType);
+						if (ensure(PropertyValue.HasValue() && !PropertyValue.HasError()))
+						{
+							FStructView& StructValue = PropertyValue.GetValue();
+							check(StructValue.IsValid());
+
+							FCameraObjectInterfaceParameterMetaData& MetaData = FindOrAddMetaData(Definition.ParameterGuid);
+							OutParameterInfos.AddBlendableParameter(
+									Definition.ParameterName,
+									Definition.VariableType,
+									Definition.BlendableStructType,
+									StructValue.GetMemory(),
+									&MetaData.OverrideVariableID);
+						}
+					}
+					break;
+			}
+		}
+		else if (Definition.ParameterType == ECameraObjectInterfaceParameterType::Data)
+		{
+			FCameraObjectInterfaceParameterMetaData& MetaData = FindOrAddMetaData(Definition.ParameterGuid);
+
+			switch (Definition.DataType)
+			{
+				case ECameraContextDataType::Name:
+					if (ensure(PropertyDesc->ValueType == EPropertyBagPropertyType::Name))
+					{
+						const uint8* DefaultValue = reinterpret_cast<uint8*>(DefaultParameters.GetValueName(*PropertyDesc).TryGetValue());
+						OutParameterInfos.AddDataParameter(Definition.ParameterName, ECameraContextDataType::Name, ECameraContextDataContainerType::None, nullptr, DefaultValue, &MetaData.OverrideDataID);
+					}
+					break;
+				case ECameraContextDataType::String:
+					if (ensure(PropertyDesc->ValueType == EPropertyBagPropertyType::String))
+					{
+						const uint8* DefaultValue = reinterpret_cast<uint8*>(DefaultParameters.GetValueString(*PropertyDesc).TryGetValue());
+						OutParameterInfos.AddDataParameter(Definition.ParameterName, ECameraContextDataType::String, ECameraContextDataContainerType::None, nullptr, DefaultValue, &MetaData.OverrideDataID);
+					}
+					break;
+				case ECameraContextDataType::Enum:
+					if (ensure(PropertyDesc->ValueType == EPropertyBagPropertyType::Enum &&
+								PropertyDesc->ValueTypeObject == Definition.DataTypeObject))
+					{
+						const UEnum* EnumType = CastChecked<const UEnum>(Definition.DataTypeObject);
+						const uint8* DefaultValue = reinterpret_cast<uint8*>(DefaultParameters.GetValueEnum(*PropertyDesc, EnumType).TryGetValue());
+						OutParameterInfos.AddDataParameter(Definition.ParameterName, ECameraContextDataType::Enum, ECameraContextDataContainerType::None, EnumType, DefaultValue, &MetaData.OverrideDataID);
+					}
+					break;
+				case ECameraContextDataType::Struct:
+					if (ensure(PropertyDesc->ValueType == EPropertyBagPropertyType::Struct &&
+								PropertyDesc->ValueTypeObject == Definition.DataTypeObject))
+					{
+						const UScriptStruct* DataType = CastChecked<const UScriptStruct>(Definition.DataTypeObject);
+						const uint8* DefaultValue = reinterpret_cast<uint8*>(DefaultParameters.GetValueStruct(*PropertyDesc, DataType).TryGetValue());
+						OutParameterInfos.AddDataParameter(Definition.ParameterName, ECameraContextDataType::Struct, ECameraContextDataContainerType::None, DataType, DefaultValue, &MetaData.OverrideDataID);
+					}
+					break;
+				case ECameraContextDataType::Object:
+					if (ensure(PropertyDesc->ValueType == EPropertyBagPropertyType::Object))
+					{
+						const uint8* DefaultValue = reinterpret_cast<uint8*>(DefaultParameters.GetValueObject(*PropertyDesc).TryGetValue());
+						OutParameterInfos.AddDataParameter(Definition.ParameterName, ECameraContextDataType::Object, ECameraContextDataContainerType::None, Definition.DataTypeObject, DefaultValue, &MetaData.OverrideDataID);
+					}
+					break;
+				case ECameraContextDataType::Class:
+					if (ensure(PropertyDesc->ValueType == EPropertyBagPropertyType::Class))
+					{
+						const uint8* DefaultValue = reinterpret_cast<uint8*>(DefaultParameters.GetValueClass(*PropertyDesc).TryGetValue());
+						OutParameterInfos.AddDataParameter(Definition.ParameterName, ECameraContextDataType::Class, ECameraContextDataContainerType::None, Definition.DataTypeObject, DefaultValue, &MetaData.OverrideDataID);
+					}
+					break;
+			}
+		}
+	}
+}
+
diff --git a/Engine/Plugins/Cameras/GameplayCameras/Source/GameplayCameras/Private/Core/BlendStackCameraNode.cpp b/Engine/Plugins/Cameras/GameplayCameras/Source/GameplayCameras/Private/Core/BlendStackCameraNode.cpp
--- a/Engine/Plugins/Cameras/GameplayCameras/Source/GameplayCameras/Private/Core/BlendStackCameraNode.cpp
+++ b/Engine/Plugins/Cameras/GameplayCameras/Source/GameplayCameras/Private/Core/BlendStackCameraNode.cpp
@@ -17,7 +17,7 @@
 #include "Debug/CameraPoseDebugBlock.h"
 #include "Debug/VariableTableDebugBlock.h"
 #include "HAL/IConsoleManager.h"
-#include "Helpers/CameraParameterOverrideHelper.h"
+#include "Helpers/CameraObjectInterfaceParameterOverrideHelper.h"
 #include "IGameplayCamerasLiveEditManager.h"
 #include "IGameplayCamerasModule.h"
 #include "Math/ColorList.h"
@@ -97,7 +97,7 @@ bool FBlendStackCameraNodeEvaluator::InitializeEntry(
 	NewEntry.Result.ContextDataTable.Initialize(CameraRig->AllocationInfo.ContextDataTableInfo);
 
 	// Set the default values, so that pre-blending from default values works.
-	FCameraParameterOverrideHelper::ApplyDefaultBlendableParameters(CameraRig, NewEntry.Result.VariableTable);
+	FCameraObjectInterfaceParameterOverrideHelper::ApplyDefaultBlendableParameters(CameraRig, NewEntry.Result.VariableTable);
 
 	// Set all the data from the context.
 	const FCameraNodeEvaluationResult& ContextResult = EvaluationContext->GetInitialResult();
diff --git a/Engine/Plugins/Cameras/GameplayCameras/Source/GameplayCameras/Private/Core/BlendStackRootCameraNode.cpp b/Engine/Plugins/Cameras/GameplayCameras/Source/GameplayCameras/Private/Core/BlendStackRootCameraNode.cpp
--- a/Engine/Plugins/Cameras/GameplayCameras/Source/GameplayCameras/Private/Core/BlendStackRootCameraNode.cpp
+++ b/Engine/Plugins/Cameras/GameplayCameras/Source/GameplayCameras/Private/Core/BlendStackRootCameraNode.cpp
@@ -10,8 +10,8 @@
 #include "Debug/CameraDebugBlock.h"
 #include "Debug/CameraDebugBlockBuilder.h"
 #include "Debug/CameraDebugRenderer.h"
-#include "Helpers/CameraParameterOverrideHelper.h"
-#include "Helpers/CameraRigParameterOverrideEvaluator.h"
+#include "Helpers/CameraObjectInterfaceParameterOverrideHelper.h"
+#include "Helpers/CameraObjectReferenceParameterOverrideEvaluator.h"
 #include "Nodes/Common/CameraRigCameraNode.h"
 
 #include UE_INLINE_GENERATED_CPP_BY_NAME(BlendStackRootCameraNode)
@@ -148,7 +148,7 @@ void FBlendStackRootCameraNodeEvaluator::MergeCameraRig(const FCameraNodeEvaluat
 
 	BuildNestedPrefabTrail(CameraRig, BlendedParameterOverrides.PrefabTrail);
 
-	const FCameraRigAllocationInfo& AllocationInfo = BlendablePrefabCameraRig->AllocationInfo;
+	const FCameraObjectAllocationInfo& AllocationInfo = BlendablePrefabCameraRig->AllocationInfo;
 	BlendedParameterOverrides.Result.VariableTable.Initialize(AllocationInfo.VariableTableInfo);
 
 	if (Blend)
@@ -209,7 +209,7 @@ void FBlendStackRootCameraNodeEvaluator::InitializeBlendedParameterOverridesStac
 
 	BuildNestedPrefabTrail(OriginalCameraRig, InitialParameterOverrides.PrefabTrail);
 
-	const FCameraRigAllocationInfo& AllocationInfo = BlendablePrefabCameraRig->AllocationInfo;
+	const FCameraObjectAllocationInfo& AllocationInfo = BlendablePrefabCameraRig->AllocationInfo;
 	InitialParameterOverrides.Result.VariableTable.Initialize(AllocationInfo.VariableTableInfo);
 
 	BlendedParameterOverridesStack.Add(MoveTemp(InitialParameterOverrides));
@@ -261,13 +261,13 @@ void FBlendStackRootCameraNodeEvaluator::RunBlendedParameterOverridesStack(const
 
 		// Start by setting the default values of all parameters. If we don't do this, parameter overrides
 		// wouldn't have a base value to blend from.
-		FCameraParameterOverrideHelper::ApplyDefaultBlendableParameters(BlendablePrefabCameraRig, CurResult.VariableTable);
+		FCameraObjectInterfaceParameterOverrideHelper::ApplyDefaultBlendableParameters(BlendablePrefabCameraRig, CurResult.VariableTable);
 
 		// Next, override the defaults with the specific values of this entry, applied bottoms up.
 		for (int32 Index = BlendedParameterOverrides.PrefabTrail.Num() - 1; Index >= 0; --Index)
 		{
 			const UCameraRigCameraNode* CurPrefabNode = BlendedParameterOverrides.PrefabTrail[Index];
-			FCameraRigParameterOverrideEvaluator OverrideEvaluator(CurPrefabNode->CameraRigReference);
+			FCameraObjectReferenceParameterOverrideEvaluator OverrideEvaluator(CurPrefabNode->CameraRigReference);
 			OverrideEvaluator.ApplyParameterOverrides(CurResult.VariableTable, false);
 		}
 
diff --git a/Engine/Plugins/Cameras/GameplayCameras/Source/GameplayCameras/Private/Core/CameraAsset.cpp b/Engine/Plugins/Cameras/GameplayCameras/Source/GameplayCameras/Private/Core/CameraAsset.cpp
--- a/Engine/Plugins/Cameras/GameplayCameras/Source/GameplayCameras/Private/Core/CameraAsset.cpp
+++ b/Engine/Plugins/Cameras/GameplayCameras/Source/GameplayCameras/Private/Core/CameraAsset.cpp
@@ -2,8 +2,8 @@
 
 #include "Core/CameraAsset.h"
 
-#include "Core/CameraAssetBuilder.h"
-#include "Core/CameraBuildLog.h"
+#include "Build/CameraAssetBuilder.h"
+#include "Build/CameraBuildLog.h"
 #include "Core/CameraDirector.h"
 #include "Core/CameraRigAsset.h"
 #include "UObject/ObjectRedirector.h"
diff --git a/Engine/Plugins/Cameras/GameplayCameras/Source/GameplayCameras/Private/Core/CameraNode.cpp b/Engine/Plugins/Cameras/GameplayCameras/Source/GameplayCameras/Private/Core/CameraNode.cpp
--- a/Engine/Plugins/Cameras/GameplayCameras/Source/GameplayCameras/Private/Core/CameraNode.cpp
+++ b/Engine/Plugins/Cameras/GameplayCameras/Source/GameplayCameras/Private/Core/CameraNode.cpp
@@ -65,7 +65,7 @@ void UCameraNode::PreBuild(FCameraBuildLog& BuildLog)
 	OnPreBuild(BuildLog);
 }
 
-void UCameraNode::Build(FCameraRigBuildContext& BuildContext)
+void UCameraNode::Build(FCameraObjectBuildContext& BuildContext)
 {
 	OnBuild(BuildContext);
 }
diff --git a/Engine/Plugins/Cameras/GameplayCameras/Source/GameplayCameras/Private/Core/CameraNodeHierarchy.cpp b/Engine/Plugins/Cameras/GameplayCameras/Source/GameplayCameras/Private/Core/CameraNodeHierarchy.cpp
--- a/Engine/Plugins/Cameras/GameplayCameras/Source/GameplayCameras/Private/Core/CameraNodeHierarchy.cpp
+++ b/Engine/Plugins/Cameras/GameplayCameras/Source/GameplayCameras/Private/Core/CameraNodeHierarchy.cpp
@@ -2,7 +2,7 @@
 
 #include "Core/CameraNodeHierarchy.h"
 
-#include "Core/CameraRigAsset.h"
+#include "Core/BaseCameraObject.h"
 
 namespace UE::Cameras
 {
@@ -11,9 +11,9 @@ FCameraNodeHierarchy::FCameraNodeHierarchy()
 {
 }
 
-FCameraNodeHierarchy::FCameraNodeHierarchy(UCameraRigAsset* InCameraRig)
+FCameraNodeHierarchy::FCameraNodeHierarchy(UBaseCameraObject* InCameraObject)
 {
-	Build(InCameraRig);
+	Build(InCameraObject);
 }
 
 TArrayView<UCameraNode* const> FCameraNodeHierarchy::GetFlattenedHierarchy() const
@@ -26,9 +26,9 @@ int32 FCameraNodeHierarchy::Num() const
 	return FlattenedHierarchy.Num();
 }
 
-void FCameraNodeHierarchy::Build(UCameraRigAsset* InCameraRig)
+void FCameraNodeHierarchy::Build(UBaseCameraObject* InCameraObject)
 {
-	Build(InCameraRig ? InCameraRig->RootNode : nullptr);
+	Build(InCameraObject ? InCameraObject->GetRootNode() : nullptr);
 }
 
 void FCameraNodeHierarchy::Build(UCameraNode* InRootCameraNode)
diff --git a/Engine/Plugins/Cameras/GameplayCameras/Source/GameplayCameras/Private/Core/CameraRigAsset.cpp b/Engine/Plugins/Cameras/GameplayCameras/Source/GameplayCameras/Private/Core/CameraRigAsset.cpp
--- a/Engine/Plugins/Cameras/GameplayCameras/Source/GameplayCameras/Private/Core/CameraRigAsset.cpp
+++ b/Engine/Plugins/Cameras/GameplayCameras/Source/GameplayCameras/Private/Core/CameraRigAsset.cpp
@@ -2,165 +2,15 @@
 
 #include "Core/CameraRigAsset.h"
 
-#include "Core/CameraBuildLog.h"
+#include "Build/CameraBuildLog.h"
+#include "Build/CameraRigAssetBuilder.h"
 #include "Core/CameraNode.h"
-#include "Core/CameraRigAssetBuilder.h"
 #include "Core/CameraVariableAssets.h"
 #include "UObject/AssetRegistryTagsContext.h"
 #include "UObject/ObjectSaveContext.h"
 
 #include UE_INLINE_GENERATED_CPP_BY_NAME(CameraRigAsset)
 
-void FCameraRigAllocationInfo::Append(const FCameraRigAllocationInfo& OtherAllocationInfo)
-{
-	const FCameraNodeEvaluatorAllocationInfo& OtherEvaluatorInfo(OtherAllocationInfo.EvaluatorInfo);
-	EvaluatorInfo.MaxAlignof = FMath::Max(EvaluatorInfo.MaxAlignof, OtherEvaluatorInfo.MaxAlignof);
-	EvaluatorInfo.TotalSizeof = Align(EvaluatorInfo.TotalSizeof, OtherEvaluatorInfo.MaxAlignof) + OtherEvaluatorInfo.TotalSizeof;
-
-	const FCameraVariableTableAllocationInfo& OtherVariableTableInfo(OtherAllocationInfo.VariableTableInfo);
-	VariableTableInfo.VariableDefinitions.Append(OtherVariableTableInfo.VariableDefinitions);
-
-	const FCameraContextDataTableAllocationInfo& OtherContextDataTableInfo(OtherAllocationInfo.ContextDataTableInfo);
-	ContextDataTableInfo.DataDefinitions.Append(OtherContextDataTableInfo.DataDefinitions);
-}
-
-#if WITH_EDITOR
-
-void UCameraRigInterfaceParameterBase::GetGraphNodePosition(FName InGraphName, int32& NodePosX, int32& NodePosY) const
-{
-	NodePosX = GraphNodePos.X;
-	NodePosY = GraphNodePos.Y;
-}
-
-void UCameraRigInterfaceParameterBase::OnGraphNodeMoved(FName InGraphName, int32 NodePosX, int32 NodePosY, bool bMarkDirty)
-{
-	Modify(bMarkDirty);
-
-	GraphNodePos.X = NodePosX;
-	GraphNodePos.Y = NodePosY;
-}
-
-#endif
-
-void UCameraRigInterfaceParameterBase::PostLoad()
-{
-	if (!Guid.IsValid())
-	{
-		Guid = FGuid::NewGuid();
-	}
-
-	Super::PostLoad();
-}
-
-void UCameraRigInterfaceParameterBase::PostInitProperties()
-{
-	Super::PostInitProperties();
-
-	if (!HasAnyFlags(RF_ClassDefaultObject | RF_ArchetypeObject | RF_NeedLoad | RF_WasLoaded) && 
-			!Guid.IsValid())
-	{
-		Guid = FGuid::NewGuid();
-	}
-}
-
-void UCameraRigInterfaceParameterBase::PostDuplicate(EDuplicateMode::Type DuplicateMode)
-{
-	Super::PostDuplicate(DuplicateMode);
-
-	if (DuplicateMode == EDuplicateMode::Normal)
-	{
-		Guid = FGuid::NewGuid();
-	}
-}
-
-FCameraVariableDefinition UCameraRigBlendableParameter::GetVariableDefinition() const
-{
-	FCameraVariableDefinition Definition;
-	Definition.VariableID = PrivateVariableID;
-	Definition.VariableType = ParameterType;
-	Definition.BlendableStructType = BlendableStructType;
-	Definition.bIsPrivate = true;
-	Definition.bIsInput = bIsPreBlended;
-#if WITH_EDITORONLY_DATA
-	Definition.VariableName = GetVariableName();
-#endif
-	return Definition;
-}
-
-FCameraContextDataDefinition UCameraRigDataParameter::GetDataDefinition() const
-{
-	FCameraContextDataDefinition Definition;
-	Definition.DataID = PrivateDataID;
-	Definition.DataType = DataType;
-	Definition.DataContainerType = DataContainerType;
-	Definition.DataTypeObject = DataTypeObject;
-#if WITH_EDITORONLY_DATA
-	Definition.DataName = GetDataName();
-#endif
-	return Definition;
-}
-
-#if WITH_EDITORONLY_DATA
-
-FString UCameraRigBlendableParameter::GetVariableName() const
-{
-	const UCameraRigAsset* Owner = GetTypedOuter<UCameraRigAsset>();
-	return FString::Printf(TEXT("Override_%s_%s"), *GetNameSafe(Owner), *InterfaceParameterName);
-}
-
-FString UCameraRigDataParameter::GetDataName() const
-{
-	const UCameraRigAsset* Owner = GetTypedOuter<UCameraRigAsset>();
-	return FString::Printf(TEXT("Override_%s_%s"), *GetNameSafe(Owner), *InterfaceParameterName);
-}
-
-#endif  // WITH_EDITORONLY_DATA
-
-UCameraRigBlendableParameter* FCameraRigInterface::FindBlendableParameterByName(const FString& ParameterName) const
-{
-	const TObjectPtr<UCameraRigBlendableParameter>* FoundItem = BlendableParameters.FindByPredicate(
-			[&ParameterName](UCameraRigBlendableParameter* Item)
-			{
-				return Item->InterfaceParameterName == ParameterName;
-			});
-	return FoundItem ? *FoundItem : nullptr;
-}
-
-UCameraRigDataParameter* FCameraRigInterface::FindDataParameterByName(const FString& ParameterName) const
-{
-	const TObjectPtr<UCameraRigDataParameter>* FoundItem = DataParameters.FindByPredicate(
-			[&ParameterName](UCameraRigDataParameter* Item)
-			{
-				return Item->InterfaceParameterName == ParameterName;
-			});
-	return FoundItem ? *FoundItem : nullptr;
-}
-
-UCameraRigBlendableParameter* FCameraRigInterface::FindBlendableParameterByGuid(const FGuid& ParameterGuid) const
-{
-	const TObjectPtr<UCameraRigBlendableParameter>* FoundItem = BlendableParameters.FindByPredicate(
-			[&ParameterGuid](UCameraRigBlendableParameter* Item)
-			{
-				return Item->GetGuid() == ParameterGuid;
-			});
-	return FoundItem ? *FoundItem : nullptr;
-}
-
-UCameraRigDataParameter* FCameraRigInterface::FindDataParameterByGuid(const FGuid& ParameterGuid) const
-{
-	const TObjectPtr<UCameraRigDataParameter>* FoundItem = DataParameters.FindByPredicate(
-			[&ParameterGuid](UCameraRigDataParameter* Item)
-			{
-				return Item->GetGuid() == ParameterGuid;
-			});
-	return FoundItem ? *FoundItem : nullptr;
-}
-
-bool FCameraRigInterface::HasBlendableParameter(const FString& ParameterName) const
-{
-	return FindBlendableParameterByName(ParameterName) != nullptr;
-}
-
 const FName UCameraRigAsset::NodeTreeGraphName(TEXT("NodeTree"));
 const FName UCameraRigAsset::TransitionsGraphName(TEXT("Transitions"));
 
@@ -187,7 +37,7 @@ void UCameraRigAsset::PostLoad()
 			It.RemoveCurrent();
 			continue;
 		}
-		if (UCameraRigInterfaceParameterBase* InterfaceParameter = Cast<UCameraRigInterfaceParameterBase>(Item))
+		if (UCameraObjectInterfaceParameterBase* InterfaceParameter = Cast<UCameraObjectInterfaceParameterBase>(Item))
 		{
 			InterfaceParameter->bHasGraphNode = true;
 			It.RemoveCurrent();
@@ -204,7 +54,7 @@ void UCameraRigAsset::PostLoad()
 
 	// Initialize the ParameterType of blendable parameters that pre-existed the new interface data.
 	// The type defaults to Boolean so check only those.
-	for (UCameraRigBlendableParameter* BlendableParameter : Interface.BlendableParameters)
+	for (UCameraObjectInterfaceBlendableParameter* BlendableParameter : Interface.BlendableParameters)
 	{
 		if (BlendableParameter->ParameterType == ECameraVariableType::Boolean)
 		{
@@ -342,15 +192,7 @@ void UCameraRigAsset::OnGraphNodeMoved(FName InGraphName, int32 NodePosX, int32
 
 EObjectTreeGraphObjectSupportFlags UCameraRigAsset::GetSupportFlags(FName InGraphName) const
 {
-	if (UPackage* OuterPackage = Cast<UPackage>(GetOuter()))
-	{
-		// Can't rename a standalone rig prefab -- you have to rename it in the content browser.
-		return EObjectTreeGraphObjectSupportFlags::CommentText;
-	}
-	else
-	{
-		return EObjectTreeGraphObjectSupportFlags::CommentText | EObjectTreeGraphObjectSupportFlags::CustomRename;
-	}
+	return EObjectTreeGraphObjectSupportFlags::CommentText;
 }
 
 const FString& UCameraRigAsset::GetGraphNodeCommentText(FName InGraphName) const
@@ -444,7 +286,7 @@ void UCameraRigAsset::RemoveConnectableObject(FName InGraphName, UObject* InObje
 void UCameraRigAsset::GetAssetRegistryTags(FAssetRegistryTagsContext Context) const
 {
 	const UEnum* VariableTypeEnum = StaticEnum<ECameraVariableType>();
-	for (const UCameraRigBlendableParameter* BlendableParameter : Interface.BlendableParameters)
+	for (const UCameraObjectInterfaceBlendableParameter* BlendableParameter : Interface.BlendableParameters)
 	{
 		if (BlendableParameter)
 		{
@@ -457,7 +299,7 @@ void UCameraRigAsset::GetAssetRegistryTags(FAssetRegistryTagsContext Context) co
 	}
 
 	const UEnum* ContextDataTypeEnum = StaticEnum<ECameraContextDataType>();
-	for (const UCameraRigDataParameter* DataParameter : Interface.DataParameters)
+	for (const UCameraObjectInterfaceDataParameter* DataParameter : Interface.DataParameters)
 	{
 		if (DataParameter)
 		{
diff --git a/Engine/Plugins/Cameras/GameplayCameras/Source/GameplayCameras/Private/Core/CameraRigAssetReference.cpp b/Engine/Plugins/Cameras/GameplayCameras/Source/GameplayCameras/Private/Core/CameraRigAssetReference.cpp
--- a/Engine/Plugins/Cameras/GameplayCameras/Source/GameplayCameras/Private/Core/CameraRigAssetReference.cpp
+++ b/Engine/Plugins/Cameras/GameplayCameras/Source/GameplayCameras/Private/Core/CameraRigAssetReference.cpp
@@ -2,10 +2,10 @@
 
 #include "Core/CameraRigAssetReference.h"
 
+#include "Build/CameraObjectInterfaceParameterBuilder.h"
 #include "Core/CameraRigAsset.h"
 #include "Core/CameraNodeEvaluator.h"
-#include "Core/ICustomCameraNodeParameterProvider.h"
-#include "Helpers/CameraRigParameterOverrideEvaluator.h"
+#include "Helpers/CameraObjectReferenceParameterOverrideEvaluator.h"
 
 #include UE_INLINE_GENERATED_CPP_BY_NAME(CameraRigAssetReference)
 
@@ -22,296 +22,16 @@ FCameraRigAssetReference::FCameraRigAssetReference(UCameraRigAsset* InCameraRig)
 
 PRAGMA_ENABLE_DEPRECATION_WARNINGS
 
-void FCameraRigAssetReference::ApplyParameterOverrides(UE::Cameras::FCameraNodeEvaluationResult& OutResult, bool bDrivenOverridesOnly)
-{
-	using namespace UE::Cameras;
-	FCameraRigParameterOverrideEvaluator OverrideEvaluator(*this);
-	OverrideEvaluator.ApplyParameterOverrides(OutResult.VariableTable, OutResult.ContextDataTable, bDrivenOverridesOnly);
-}
-
-const FCameraRigAssetReferenceParameterMetaData* FCameraRigAssetReference::FindMetaData(const FGuid& PropertyID) const
-{
-	return ParameterMetaData.FindByPredicate(
-			[PropertyID](const FCameraRigAssetReferenceParameterMetaData& Item)
-			{
-				return Item.ParameterGuid == PropertyID;
-			});
-}
-
-FCameraRigAssetReferenceParameterMetaData& FCameraRigAssetReference::FindOrAddMetaData(const FGuid& PropertyID)
-{
-	FCameraRigAssetReferenceParameterMetaData* ExistingMetaData = ParameterMetaData.FindByPredicate(
-			[PropertyID](FCameraRigAssetReferenceParameterMetaData& Item)
-			{
-				return Item.ParameterGuid == PropertyID;
-			});
-	if (ExistingMetaData)
-	{
-		return *ExistingMetaData;
-	}
-
-	FCameraRigAssetReferenceParameterMetaData& NewMetaData = ParameterMetaData.Emplace_GetRef();
-	NewMetaData.ParameterGuid = PropertyID;
-	return NewMetaData;
-}
-
-void FCameraRigAssetReference::GenerateOverriddenParameterGuidArray(TArray<FGuid>& OutOverriddenIDs) const
-{
-	for (const FCameraRigAssetReferenceParameterMetaData& MetaData : ParameterMetaData)
-	{
-		if (MetaData.bIsOverridden)
-		{
-			OutOverriddenIDs.Add(MetaData.ParameterGuid);
-		}
-	}
-}
-
-bool FCameraRigAssetReference::IsParameterOverridden(const FGuid& PropertyID) const
-{
-	if (const FCameraRigAssetReferenceParameterMetaData* MetaData = FindMetaData(PropertyID))
-	{
-		return MetaData->bIsOverridden;
-	}
-	return false;
-}
-
-void FCameraRigAssetReference::SetParameterOverridden(const FGuid& PropertyID, bool bIsOverridden)
-{
-	FCameraRigAssetReferenceParameterMetaData& MetaData = FindOrAddMetaData(PropertyID);
-	MetaData.bIsOverridden = bIsOverridden;
-}
-
-bool FCameraRigAssetReference::IsParameterAnimated(const FGuid& PropertyID) const
-{
-	if (const FCameraRigAssetReferenceParameterMetaData* MetaData = FindMetaData(PropertyID))
-	{
-		return MetaData->bIsAnimated;
-	}
-	return false;
-}
-
-void FCameraRigAssetReference::SetParameterAnimated(const FGuid& PropertyID, bool bIsAnimated)
-{
-	FCameraRigAssetReferenceParameterMetaData& MetaData = FindOrAddMetaData(PropertyID);
-	MetaData.bIsAnimated = bIsAnimated;
-}
-
-bool FCameraRigAssetReference::NeedsRebuildParameters() const
-{
-	if ((!CameraRig && Parameters.IsValid()) || (CameraRig && !Parameters.IsValid()))
-	{
-		return true;
-	}
-
-	if (CameraRig)
-	{
-		const UPropertyBag* AssetParametersType = CameraRig->GetDefaultParameters().GetPropertyBagStruct();
-		const UPropertyBag* ReferenceParametersType = Parameters.GetPropertyBagStruct();
-		if (AssetParametersType != ReferenceParametersType)
-		{
-			return true;
-		}
-	}
-
-	return false;
-}
-
-bool FCameraRigAssetReference::RebuildParametersIfNeeded()
-{
-	if (NeedsRebuildParameters())
-	{
-		RebuildParameters();
-		return true;
-	}
-	return false;
-}
-
-void FCameraRigAssetReference::RebuildParameters()
+const UBaseCameraObject* FCameraRigAssetReference::GetCameraObject() const
 {
-	if (CameraRig)
-	{
-		TArray<FGuid> ParameterOverrideGuids;
-		GenerateOverriddenParameterGuidArray(ParameterOverrideGuids);
-		Parameters.MigrateToNewBagInstanceWithOverrides(CameraRig->GetDefaultParameters(), ParameterOverrideGuids);
-
-		if (const UPropertyBag* ParametersType = Parameters.GetPropertyBagStruct())
-		{
-			// Remove metadata for parameters that don't exist anymore, and add default metadata for
-			// new parameters.
-			TSet<FGuid> ExistingMetaDataIDs;
-			for (TArray<FCameraRigAssetReferenceParameterMetaData>::TIterator It = ParameterMetaData.CreateIterator(); It; ++It)
-			{
-				ExistingMetaDataIDs.Add(It->ParameterGuid);
-			}
-
-			TSet<FGuid> WantedMetaDataIDs;
-			TMap<FGuid, ECameraRigInterfaceParameterType> WantedParameterTypes;
-			for (const FCameraRigParameterDefinition& Definition : CameraRig->GetParameterDefinitions())
-			{
-				WantedMetaDataIDs.Add(Definition.ParameterGuid);
-				WantedParameterTypes.Add(Definition.ParameterGuid, Definition.ParameterType);
-			}
-
-			TSet<FGuid> RemovedMetaDataIDs = ExistingMetaDataIDs.Difference(WantedMetaDataIDs);
-			for (TArray<FCameraRigAssetReferenceParameterMetaData>::TIterator It = ParameterMetaData.CreateIterator(); It; ++It)
-			{
-				if (RemovedMetaDataIDs.Contains(It->ParameterGuid))
-				{
-					It.RemoveCurrentSwap();
-				}
-			}
-
-			TSet<FGuid> AddedMetaDataIDs = WantedMetaDataIDs.Difference(ExistingMetaDataIDs);
-			for (const FGuid& ParameterGuid : AddedMetaDataIDs)
-			{
-				FCameraRigAssetReferenceParameterMetaData NewMetaData;
-				NewMetaData.ParameterGuid = ParameterGuid;
-				ParameterMetaData.Add(NewMetaData);
-			}
-		}
-		else
-		{
-			ParameterMetaData.Reset();
-		}
-	}
-	else
-	{
-		Parameters.Reset();
-		ParameterMetaData.Reset();
-	}
+	return CameraRig;
 }
 
-void FCameraRigAssetReference::GetCustomCameraNodeParameters(FCustomCameraNodeParameterInfos& OutParameterInfos)
+void FCameraRigAssetReference::ApplyParameterOverrides(UE::Cameras::FCameraNodeEvaluationResult& OutResult, bool bDrivenOverridesOnly)
 {
-	if (!CameraRig)
-	{
-		return;
-	}
-
-	RebuildParametersIfNeeded();
-
-	const UPropertyBag* ParametersStruct = Parameters.GetPropertyBagStruct();
-	const FInstancedPropertyBag& DefaultParameters = CameraRig->GetDefaultParameters();
-	if (!ensure(ParametersStruct && ParametersStruct == DefaultParameters.GetPropertyBagStruct()))
-	{
-		return;
-	}
-
-	for (const FCameraRigParameterDefinition& Definition : CameraRig->GetParameterDefinitions())
-	{
-		const FPropertyBagPropertyDesc* PropertyDesc = ParametersStruct->FindPropertyDescByID(Definition.ParameterGuid);
-		if (!ensure(PropertyDesc))
-		{
-			continue;
-		}
-
-		if (Definition.ParameterType == ECameraRigInterfaceParameterType::Blendable)
-		{
-			if (!ensure(PropertyDesc->ValueType == EPropertyBagPropertyType::Struct))
-			{
-				continue;
-			}
-
-			switch (Definition.VariableType)
-			{
-#define UE_CAMERA_VARIABLE_FOR_TYPE(ValueType, ValueName)\
-				case ECameraVariableType::ValueName:\
-					if (ensure(PropertyDesc->ValueTypeObject == F##ValueName##CameraParameter::StaticStruct()))\
-					{\
-						using CameraParameterType = F##ValueName##CameraParameter;\
-						TValueOrError<CameraParameterType*, EPropertyBagResult> PropertyValue =\
-							Parameters.GetValueStruct<CameraParameterType>(*PropertyDesc);\
-						if (ensure(PropertyValue.HasValue() && !PropertyValue.HasError()))\
-						{\
-							CameraParameterType* CameraParameter = PropertyValue.GetValue();\
-							check(CameraParameter);\
-							OutParameterInfos.AddBlendableParameter(\
-									Definition.ParameterName,\
-									Definition.VariableType,\
-									nullptr,\
-									reinterpret_cast<uint8*>(&CameraParameter->Value),\
-									&CameraParameter->VariableID);\
-						}\
-					}\
-					break;
-				UE_CAMERA_VARIABLE_FOR_ALL_TYPES()
-#undef UE_CAMERA_VARIABLE_FOR_TYPE
-				case ECameraVariableType::BlendableStruct:
-					{
-						TValueOrError<FStructView, EPropertyBagResult> PropertyValue =
-							Parameters.GetValueStruct(*PropertyDesc, Definition.BlendableStructType);
-						if (ensure(PropertyValue.HasValue() && !PropertyValue.HasError()))
-						{
-							FStructView& StructValue = PropertyValue.GetValue();
-							check(StructValue.IsValid());
-
-							FCameraRigAssetReferenceParameterMetaData& MetaData = FindOrAddMetaData(Definition.ParameterGuid);
-							OutParameterInfos.AddBlendableParameter(
-									Definition.ParameterName,
-									Definition.VariableType,
-									Definition.BlendableStructType,
-									StructValue.GetMemory(),
-									&MetaData.OverrideVariableID);
-						}
-					}
-					break;
-			}
-		}
-		else if (Definition.ParameterType == ECameraRigInterfaceParameterType::Data)
-		{
-			FCameraRigAssetReferenceParameterMetaData& MetaData = FindOrAddMetaData(Definition.ParameterGuid);
-
-			switch (Definition.DataType)
-			{
-				case ECameraContextDataType::Name:
-					if (ensure(PropertyDesc->ValueType == EPropertyBagPropertyType::Name))
-					{
-						const uint8* DefaultValue = reinterpret_cast<uint8*>(DefaultParameters.GetValueName(*PropertyDesc).TryGetValue());
-						OutParameterInfos.AddDataParameter(Definition.ParameterName, ECameraContextDataType::Name, ECameraContextDataContainerType::None, nullptr, DefaultValue, &MetaData.OverrideDataID);
-					}
-					break;
-				case ECameraContextDataType::String:
-					if (ensure(PropertyDesc->ValueType == EPropertyBagPropertyType::String))
-					{
-						const uint8* DefaultValue = reinterpret_cast<uint8*>(DefaultParameters.GetValueString(*PropertyDesc).TryGetValue());
-						OutParameterInfos.AddDataParameter(Definition.ParameterName, ECameraContextDataType::String, ECameraContextDataContainerType::None, nullptr, DefaultValue, &MetaData.OverrideDataID);
-					}
-					break;
-				case ECameraContextDataType::Enum:
-					if (ensure(PropertyDesc->ValueType == EPropertyBagPropertyType::Enum &&
-								PropertyDesc->ValueTypeObject == Definition.DataTypeObject))
-					{
-						const UEnum* EnumType = CastChecked<const UEnum>(Definition.DataTypeObject);
-						const uint8* DefaultValue = reinterpret_cast<uint8*>(DefaultParameters.GetValueEnum(*PropertyDesc, EnumType).TryGetValue());
-						OutParameterInfos.AddDataParameter(Definition.ParameterName, ECameraContextDataType::Enum, ECameraContextDataContainerType::None, EnumType, DefaultValue, &MetaData.OverrideDataID);
-					}
-					break;
-				case ECameraContextDataType::Struct:
-					if (ensure(PropertyDesc->ValueType == EPropertyBagPropertyType::Struct &&
-								PropertyDesc->ValueTypeObject == Definition.DataTypeObject))
-					{
-						const UScriptStruct* DataType = CastChecked<const UScriptStruct>(Definition.DataTypeObject);
-						const uint8* DefaultValue = reinterpret_cast<uint8*>(DefaultParameters.GetValueStruct(*PropertyDesc, DataType).TryGetValue());
-						OutParameterInfos.AddDataParameter(Definition.ParameterName, ECameraContextDataType::Struct, ECameraContextDataContainerType::None, DataType, DefaultValue, &MetaData.OverrideDataID);
-					}
-					break;
-				case ECameraContextDataType::Object:
-					if (ensure(PropertyDesc->ValueType == EPropertyBagPropertyType::Object))
-					{
-						const uint8* DefaultValue = reinterpret_cast<uint8*>(DefaultParameters.GetValueObject(*PropertyDesc).TryGetValue());
-						OutParameterInfos.AddDataParameter(Definition.ParameterName, ECameraContextDataType::Object, ECameraContextDataContainerType::None, Definition.DataTypeObject, DefaultValue, &MetaData.OverrideDataID);
-					}
-					break;
-				case ECameraContextDataType::Class:
-					if (ensure(PropertyDesc->ValueType == EPropertyBagPropertyType::Class))
-					{
-						const uint8* DefaultValue = reinterpret_cast<uint8*>(DefaultParameters.GetValueClass(*PropertyDesc).TryGetValue());
-						OutParameterInfos.AddDataParameter(Definition.ParameterName, ECameraContextDataType::Class, ECameraContextDataContainerType::None, Definition.DataTypeObject, DefaultValue, &MetaData.OverrideDataID);
-					}
-					break;
-			}
-		}
-	}
+	using namespace UE::Cameras;
+	FCameraObjectReferenceParameterOverrideEvaluator OverrideEvaluator(*this);
+	OverrideEvaluator.ApplyParameterOverrides(OutResult.VariableTable, OutResult.ContextDataTable, bDrivenOverridesOnly);
 }
 
 bool FCameraRigAssetReference::SerializeFromMismatchedTag(struct FPropertyTag const& Tag, FStructuredArchive::FSlot Slot)
@@ -333,7 +53,7 @@ void FCameraRigAssetReference::PostSerialize(const FArchive& Ar)
 	// Make a property bag with the legacy overrides, and then set the values in it.
 	bool bHasAnyLegacyOverride = false;
 	TArray<FPropertyBagPropertyDesc> LegacyParameterProperties;
-	TArray<FCameraRigAssetReferenceParameterMetaData> LegacyParameterMetaData;
+	TArray<FCameraObjectInterfaceParameterMetaData> LegacyParameterMetaData;
 
 #define UE_CAMERA_VARIABLE_FOR_TYPE(ValueType, ValueName)\
 	for (F##ValueName##CameraRigParameterOverride& ParameterOverride : ParameterOverrides_DEPRECATED.ValueName##Overrides)\
@@ -344,7 +64,7 @@ void FCameraRigAssetReference::PostSerialize(const FArchive& Ar)
 		FPropertyBagPropertyDesc LegacyParameterProperty(PropertyName, PropertyType, PropertyTypeObject);\
 		LegacyParameterProperty.ID = ParameterOverride.InterfaceParameterGuid;\
 		LegacyParameterProperties.Add(LegacyParameterProperty);\
-		FCameraRigAssetReferenceParameterMetaData MetaData;\
+		FCameraObjectInterfaceParameterMetaData MetaData;\
 		MetaData.ParameterGuid = ParameterOverride.InterfaceParameterGuid;\
 		MetaData.bIsOverridden = true;\
 		LegacyParameterMetaData.Add(MetaData);\
@@ -378,7 +98,7 @@ void FCameraRigAssetReference::PostSerialize(const FArchive& Ar)
 	{
 		for (const FGuid& Guid : ParameterOverrideGuids_DEPRECATED)
 		{
-			FCameraRigAssetReferenceParameterMetaData MetaData;
+			FCameraObjectInterfaceParameterMetaData MetaData;
 			MetaData.ParameterGuid = Guid;
 			MetaData.bIsOverridden = true;
 			ParameterMetaData.Add(MetaData);
diff --git a/Engine/Plugins/Cameras/GameplayCameras/Source/GameplayCameras/Private/Core/CameraRigCombinationRegistry.cpp b/Engine/Plugins/Cameras/GameplayCameras/Source/GameplayCameras/Private/Core/CameraRigCombinationRegistry.cpp
--- a/Engine/Plugins/Cameras/GameplayCameras/Source/GameplayCameras/Private/Core/CameraRigCombinationRegistry.cpp
+++ b/Engine/Plugins/Cameras/GameplayCameras/Source/GameplayCameras/Private/Core/CameraRigCombinationRegistry.cpp
@@ -3,7 +3,7 @@
 #include "Core/CameraRigCombinationRegistry.h"
 
 #include "Core/CameraRigAsset.h"
-#include "Helpers/CameraRigParameterOverrideEvaluator.h"
+#include "Helpers/CameraObjectReferenceParameterOverrideEvaluator.h"
 
 namespace UE::Cameras
 {
@@ -123,7 +123,7 @@ void FCombinedCameraRigsCameraNodeEvaluator::ApplyParameterOverrides(FCameraVari
 	const UCombinedCameraRigsCameraNode* CombinedRigsNode = GetCameraNodeAs<UCombinedCameraRigsCameraNode>();
 	for (const FCameraRigAssetReference& IndividualCameraRigReference : CombinedRigsNode->CameraRigReferences)
 	{
-		FCameraRigParameterOverrideEvaluator OverrideEvaluator(IndividualCameraRigReference);
+		FCameraObjectReferenceParameterOverrideEvaluator OverrideEvaluator(IndividualCameraRigReference);
 		OverrideEvaluator.ApplyParameterOverrides(OutVariableTable, bDrivenOnly);
 	}
 }
diff --git a/Engine/Plugins/Cameras/GameplayCameras/Source/GameplayCameras/Private/Core/CameraRigTransition.cpp b/Engine/Plugins/Cameras/GameplayCameras/Source/GameplayCameras/Private/Core/CameraRigTransition.cpp
--- a/Engine/Plugins/Cameras/GameplayCameras/Source/GameplayCameras/Private/Core/CameraRigTransition.cpp
+++ b/Engine/Plugins/Cameras/GameplayCameras/Source/GameplayCameras/Private/Core/CameraRigTransition.cpp
@@ -2,9 +2,9 @@
 
 #include "Core/CameraRigTransition.h"
 
+#include "Build/CameraBuildLog.h"
+#include "Build/CameraObjectBuildContext.h"
 #include "Core/BlendCameraNode.h"
-#include "Core/CameraBuildLog.h"
-#include "Core/CameraRigBuildContext.h"
 #include "Logging/TokenizedMessage.h"
 #include "Nodes/Blends/SimpleBlendCameraNode.h"
 
@@ -34,7 +34,7 @@ bool UCameraRigTransitionCondition::TransitionMatches(const FCameraRigTransition
 	return OnTransitionMatches(Params);
 }
 
-void UCameraRigTransitionCondition::Build(FCameraRigBuildContext& BuildContext)
+void UCameraRigTransitionCondition::Build(FCameraObjectBuildContext& BuildContext)
 {
 	OnBuild(BuildContext);
 }
@@ -99,7 +99,7 @@ bool UCameraRigTransition::AllConditionsMatch(const FCameraRigTransitionConditio
 	return true;
 }
 
-void UCameraRigTransition::Build(FCameraRigBuildContext& BuildContext)
+void UCameraRigTransition::Build(FCameraObjectBuildContext& BuildContext)
 {
 	if (Blend)
 	{
diff --git a/Engine/Plugins/Cameras/GameplayCameras/Source/GameplayCameras/Private/Core/TransientBlendStackCameraNode.cpp b/Engine/Plugins/Cameras/GameplayCameras/Source/GameplayCameras/Private/Core/TransientBlendStackCameraNode.cpp
--- a/Engine/Plugins/Cameras/GameplayCameras/Source/GameplayCameras/Private/Core/TransientBlendStackCameraNode.cpp
+++ b/Engine/Plugins/Cameras/GameplayCameras/Source/GameplayCameras/Private/Core/TransientBlendStackCameraNode.cpp
@@ -118,7 +118,7 @@ FBlendStackEntryID FTransientBlendStackCameraNodeEvaluator::PushNewEntry(const F
 
 		if (Params.CameraRig)
 		{
-			const FCameraRigAllocationInfo& AllocationInfo = Params.CameraRig->AllocationInfo;
+			const FCameraObjectAllocationInfo& AllocationInfo = Params.CameraRig->AllocationInfo;
 			UnblendedResult.VariableTable.Initialize(AllocationInfo.VariableTableInfo);
 			UnblendedResult.ContextDataTable.Initialize(AllocationInfo.ContextDataTableInfo);
 		}
diff --git a/Engine/Plugins/Cameras/GameplayCameras/Source/GameplayCameras/Private/Directors/BlueprintCameraDirector.cpp b/Engine/Plugins/Cameras/GameplayCameras/Source/GameplayCameras/Private/Directors/BlueprintCameraDirector.cpp
--- a/Engine/Plugins/Cameras/GameplayCameras/Source/GameplayCameras/Private/Directors/BlueprintCameraDirector.cpp
+++ b/Engine/Plugins/Cameras/GameplayCameras/Source/GameplayCameras/Private/Directors/BlueprintCameraDirector.cpp
@@ -2,9 +2,9 @@
 
 #include "Directors/BlueprintCameraDirector.h"
 
+#include "Build/CameraBuildLog.h"
 #include "Components/ActorComponent.h"
 #include "Core/CameraAsset.h"
-#include "Core/CameraBuildLog.h"
 #include "Core/CameraDirectorEvaluator.h"
 #include "Core/CameraEvaluationContext.h"
 #include "Core/CameraRigAsset.h"
diff --git a/Engine/Plugins/Cameras/GameplayCameras/Source/GameplayCameras/Private/Directors/SingleCameraDirector.cpp b/Engine/Plugins/Cameras/GameplayCameras/Source/GameplayCameras/Private/Directors/SingleCameraDirector.cpp
--- a/Engine/Plugins/Cameras/GameplayCameras/Source/GameplayCameras/Private/Directors/SingleCameraDirector.cpp
+++ b/Engine/Plugins/Cameras/GameplayCameras/Source/GameplayCameras/Private/Directors/SingleCameraDirector.cpp
@@ -2,8 +2,8 @@
 
 #include "Directors/SingleCameraDirector.h"
 
+#include "Build/CameraBuildLog.h"
 #include "Core/CameraAsset.h"
-#include "Core/CameraBuildLog.h"
 #include "Logging/TokenizedMessage.h"
 
 #include UE_INLINE_GENERATED_CPP_BY_NAME(SingleCameraDirector)
diff --git a/Engine/Plugins/Cameras/GameplayCameras/Source/GameplayCameras/Private/GameFramework/CameraRigParameterInterop.cpp b/Engine/Plugins/Cameras/GameplayCameras/Source/GameplayCameras/Private/GameFramework/CameraRigParameterInterop.cpp
--- a/Engine/Plugins/Cameras/GameplayCameras/Source/GameplayCameras/Private/GameFramework/CameraRigParameterInterop.cpp
+++ b/Engine/Plugins/Cameras/GameplayCameras/Source/GameplayCameras/Private/GameFramework/CameraRigParameterInterop.cpp
@@ -36,7 +36,7 @@ void SetCameraRigBlendableParameter(FBlueprintCameraVariableTable VariableTable,
 
 FCameraVariableID GetParameterPrivateVariableID(UCameraRigAsset* CameraRig, const FString& ParameterName, const UScriptStruct* BlendableStructType = nullptr)
 {
-	UCameraRigBlendableParameter* BlendableParameter = CameraRig->Interface.FindBlendableParameterByName(ParameterName);
+	UCameraObjectInterfaceBlendableParameter* BlendableParameter = CameraRig->Interface.FindBlendableParameterByName(ParameterName);
 	if (!BlendableParameter)
 	{
 		const FText Text = LOCTEXT("NoSuchBlendableParameter", "No parameter '{0}' found on camera rig '{1}'. Setting this camera variable table value will most probably accomplish nothing.");
@@ -78,7 +78,7 @@ bool ValidateSetCameraRigDataParameter(FBlueprintCameraContextDataTable& Context
 
 FCameraContextDataID GetParameterPrivateDataID(UCameraRigAsset* CameraRig, const FString& ParameterName)
 {
-	UCameraRigDataParameter* DataParameter = CameraRig->Interface.FindDataParameterByName(ParameterName);
+	UCameraObjectInterfaceDataParameter* DataParameter = CameraRig->Interface.FindDataParameterByName(ParameterName);
 	if (!DataParameter)
 	{
 		const FText Text = LOCTEXT("NoSuchDataParameter", "No parameter '{0}' found on camera rig '{1}'. Setting this data will most probably accomplish nothing.");
