Cameras: extract most of the camera rig logic into a base "camera object" class so it can be shared with the upcoming camera shakes

diff --git a/Engine/Plugins/Cameras/GameplayCameras/Source/GameplayCameras/Private/GameFramework/GameplayCameraComponent.cpp b/Engine/Plugins/Cameras/GameplayCameras/Source/GameplayCameras/Private/GameFramework/GameplayCameraComponent.cpp
--- a/Engine/Plugins/Cameras/GameplayCameras/Source/GameplayCameras/Private/GameFramework/GameplayCameraComponent.cpp
+++ b/Engine/Plugins/Cameras/GameplayCameras/Source/GameplayCameras/Private/GameFramework/GameplayCameraComponent.cpp
@@ -2,10 +2,10 @@
 
 #include "GameFramework/GameplayCameraComponent.h"
 
+#include "Build/CameraAssetBuilder.h"
+#include "Build/CameraBuildLog.h"
 #include "CineCameraComponent.h"
 #include "Core/CameraAsset.h"
-#include "Core/CameraAssetBuilder.h"
-#include "Core/CameraBuildLog.h"
 #include "Core/CameraSystemEvaluator.h"
 #include "Core/RootCameraNode.h"
 #include "Debug/CameraDebugRenderer.h"
diff --git a/Engine/Plugins/Cameras/GameplayCameras/Source/GameplayCameras/Private/Helpers/CameraAssetParameterOverrideEvaluator.cpp b/Engine/Plugins/Cameras/GameplayCameras/Source/GameplayCameras/Private/Helpers/CameraAssetParameterOverrideEvaluator.cpp
--- a/Engine/Plugins/Cameras/GameplayCameras/Source/GameplayCameras/Private/Helpers/CameraAssetParameterOverrideEvaluator.cpp
+++ b/Engine/Plugins/Cameras/GameplayCameras/Source/GameplayCameras/Private/Helpers/CameraAssetParameterOverrideEvaluator.cpp
@@ -4,8 +4,9 @@
 
 #include "Core/CameraAsset.h"
 #include "Core/CameraAssetReference.h"
+#include "Core/CameraObjectInterfaceParameterDefinition.h"
 #include "Core/CameraRigAsset.h"
-#include "Helpers/CameraParameterOverrideHelper.h"
+#include "Helpers/CameraObjectInterfaceParameterOverrideHelper.h"
 
 namespace UE::Cameras
 {
@@ -37,17 +38,17 @@ void FCameraAssetParameterOverrideEvaluator::ApplyParameterOverrides(FCameraVari
 		return;
 	}
 
-	TConstArrayView<FCameraRigParameterDefinition> ParameterDefinitions = CameraAsset->GetParameterDefinitions();
+	TConstArrayView<FCameraObjectInterfaceParameterDefinition> ParameterDefinitions = CameraAsset->GetParameterDefinitions();
 	const TArray<TObjectPtr<const UCameraRigAsset>>& ParameterOwners = CameraAsset->ParameterOwners;
 	ensure(ParameterDefinitions.Num() == ParameterOwners.Num());
 
-	FCameraParameterOverrideHelper Helper(OutVariableTable, OutContextDataTable);
+	FCameraObjectInterfaceParameterOverrideHelper Helper(OutVariableTable, OutContextDataTable);
 
 	for (int32 Index = 0, MaxIndex = FMath::Min(ParameterDefinitions.Num(), ParameterOwners.Num()); Index < MaxIndex; ++Index)
 	{
-		const FCameraRigParameterDefinition& Definition(ParameterDefinitions[Index]);
+		const FCameraObjectInterfaceParameterDefinition& Definition(ParameterDefinitions[Index]);
 
-		if (!OutContextDataTable && Definition.ParameterType == ECameraRigInterfaceParameterType::Data)
+		if (!OutContextDataTable && Definition.ParameterType == ECameraObjectInterfaceParameterType::Data)
 		{
 			continue;
 		}
diff --git a/Engine/Plugins/Cameras/GameplayCameras/Source/GameplayCameras/Private/Helpers/CameraObjectInterfaceParameterOverrideHelper.h b/Engine/Plugins/Cameras/GameplayCameras/Source/GameplayCameras/Private/Helpers/CameraObjectInterfaceParameterOverrideHelper.h
new file mode 100644
--- /dev/null
+++ b/Engine/Plugins/Cameras/GameplayCameras/Source/GameplayCameras/Private/Helpers/CameraObjectInterfaceParameterOverrideHelper.h
@@ -0,0 +1,63 @@
+// Copyright Epic Games, Inc. All Rights Reserved.
+
+#pragma once
+
+#include "StructUtils/PropertyBag.h"
+
+class UBaseCameraObject;
+struct FCameraObjectInterfaceParameterDefinition;
+
+namespace UE::Cameras
+{
+
+class FCameraContextDataTable;
+class FCameraVariableTable;
+
+/**
+ * A helper class for applying camera object interface parameter overrides from a property bag, 
+ * such as with camera asset references, camera rig asset references, and camera shake asset
+ * references.
+ */
+struct FCameraObjectInterfaceParameterOverrideHelper
+{
+public:
+
+	/** Sets default values of blendable interface parameters in the given variable table. */
+	static void ApplyDefaultBlendableParameters(const UBaseCameraObject* CameraObject, FCameraVariableTable& OutVariableTable);
+
+public:
+
+	/** 
+	 * Creates a new helper instance.
+	 *
+	 * The given variable or context data tables can be null, in which case blendable or data
+	 * interface parameters will be skipped.
+	 */
+	FCameraObjectInterfaceParameterOverrideHelper(FCameraVariableTable* OutVariableTable, FCameraContextDataTable* OutContextDataTable);
+
+	/** Sets overriden values of interface parameters in the given variable and context data tables. */
+	void ApplyParameterOverrides(
+			const UBaseCameraObject* CameraObject,
+			const FInstancedPropertyBag& ParameterOverrides,
+			const TSet<FGuid>& OverridenParameterGuids,
+			const TSet<FGuid>& AnimatedParameterGuids,
+			FCameraVariableTable* OutVariableTable,
+			FCameraContextDataTable* OutContextDataTable,
+			bool bDrivenOverridesOnly);
+
+	/** Sets a given override value for an interface parameter. */
+	void ApplyParameterOverride(
+			const UBaseCameraObject* CameraObject,
+			const FCameraObjectInterfaceParameterDefinition& ParameterDefinition,
+			const FInstancedPropertyBag& PropertyBag,
+			const FPropertyBagPropertyDesc& PropertyBagPropertyDesc,
+			bool bDrivenOverridesOnly);
+
+private:
+
+	FCameraVariableTable* VariableTable;
+	FCameraContextDataTable* ContextDataTable;
+};
+
+}  // namespace UE::Cameras
+
diff --git a/Engine/Plugins/Cameras/GameplayCameras/Source/GameplayCameras/Private/Helpers/CameraObjectReferenceParameterOverrideEvaluator.h b/Engine/Plugins/Cameras/GameplayCameras/Source/GameplayCameras/Private/Helpers/CameraObjectReferenceParameterOverrideEvaluator.h
new file mode 100644
--- /dev/null
+++ b/Engine/Plugins/Cameras/GameplayCameras/Source/GameplayCameras/Private/Helpers/CameraObjectReferenceParameterOverrideEvaluator.h
@@ -0,0 +1,84 @@
+// Copyright Epic Games, Inc. All Rights Reserved.
+
+#pragma once
+
+#include "Core/BaseCameraObjectReference.h"
+#include "Helpers/CameraObjectInterfaceParameterOverrideHelper.h"
+
+namespace UE::Cameras
+{
+
+class FCameraContextDataTable;
+class FCameraVariableTable;
+
+/**
+ * Utility class for applying interface parameter overrides to a camera object via a
+ * given variable table and/or context data table.
+ */
+class FCameraObjectReferenceParameterOverrideEvaluator
+{
+public:
+
+	/** Creates a new parameter override evaluator. */
+	FCameraObjectReferenceParameterOverrideEvaluator(const FBaseCameraObjectReference& InObjectReference)
+		: ObjectReference(InObjectReference)
+	{}
+
+	/** 
+	 * Applies override values to the given variable table.
+	 *
+	 * @param OutVariableTable      The variable table in which to set the override values.
+	 * @param bDrivenOverridesOnly  Whether only overrides driven by variables should be applied.
+	 */
+	void ApplyParameterOverrides(FCameraVariableTable& OutVariableTable, bool bDrivenOverridesOnly = false)
+	{
+		TSet<FGuid> OverridenParameterGuids;
+		ObjectReference.GetOverriddenParameterGuids(OverridenParameterGuids);
+
+		TSet<FGuid> AnimatedParameterGuids;
+		ObjectReference.GetAnimatedParameterGuids(AnimatedParameterGuids);
+
+		FCameraObjectInterfaceParameterOverrideHelper Helper(&OutVariableTable, nullptr);
+		Helper.ApplyParameterOverrides(
+				ObjectReference.GetCameraObject(),
+				ObjectReference.GetParameters(),
+				OverridenParameterGuids,
+				AnimatedParameterGuids,
+				&OutVariableTable, 
+				nullptr, 
+				bDrivenOverridesOnly);
+	}
+
+	/** 
+	 * Applies override values to the given variable and context data tables.
+	 *
+	 * @param OutVariableTable      The variable table in which to set the override values.
+	 * @param OutContextDataTable   The context data table in which to set the override values.
+	 * @param bDrivenOverridesOnly  Whether only overrides driven by variables should be applied.
+	 */
+	void ApplyParameterOverrides(FCameraVariableTable& OutVariableTable, FCameraContextDataTable& OutContextDataTable, bool bDrivenOverridesOnly = false)
+	{
+		TSet<FGuid> OverridenParameterGuids;
+		ObjectReference.GetOverriddenParameterGuids(OverridenParameterGuids);
+
+		TSet<FGuid> AnimatedParameterGuids;
+		ObjectReference.GetAnimatedParameterGuids(AnimatedParameterGuids);
+
+		FCameraObjectInterfaceParameterOverrideHelper Helper(&OutVariableTable, nullptr);
+		Helper.ApplyParameterOverrides(
+				ObjectReference.GetCameraObject(),
+				ObjectReference.GetParameters(),
+				OverridenParameterGuids,
+				AnimatedParameterGuids,
+				&OutVariableTable, 
+				&OutContextDataTable, 
+				bDrivenOverridesOnly);
+	}
+
+private:
+
+	const FBaseCameraObjectReference& ObjectReference;
+};
+
+}  // namespace UE::Cameras
+
diff --git a/Engine/Plugins/Cameras/GameplayCameras/Source/GameplayCameras/Private/MovieScene/MovieSceneCameraParameterInstantiator.cpp b/Engine/Plugins/Cameras/GameplayCameras/Source/GameplayCameras/Private/MovieScene/MovieSceneCameraParameterInstantiator.cpp
--- a/Engine/Plugins/Cameras/GameplayCameras/Source/GameplayCameras/Private/MovieScene/MovieSceneCameraParameterInstantiator.cpp
+++ b/Engine/Plugins/Cameras/GameplayCameras/Source/GameplayCameras/Private/MovieScene/MovieSceneCameraParameterInstantiator.cpp
@@ -4,7 +4,7 @@
 
 #include "Core/CameraAsset.h"
 #include "Core/CameraAssetReference.h"
-#include "Core/CameraRigParameterDefinition.h"
+#include "Core/CameraObjectInterfaceParameterDefinition.h"
 #include "EntitySystem/BuiltInComponentTypes.h"
 #include "EntitySystem/MovieSceneEntityManager.h"
 #include "EntitySystem/MovieSceneEntitySystemTask.h"
@@ -131,7 +131,7 @@ struct FSetupCameraParameterOverrideTask
 			}
 		}
 
-		for (const FCameraRigParameterDefinition& ParameterDefintion : CameraAsset->GetParameterDefinitions())
+		for (const FCameraObjectInterfaceParameterDefinition& ParameterDefintion : CameraAsset->GetParameterDefinitions())
 		{
 			if (ParameterDefintion.ParameterName == ParameterName)
 			{
diff --git a/Engine/Plugins/Cameras/GameplayCameras/Source/GameplayCameras/Private/Nodes/Common/CameraRigCameraNode.cpp b/Engine/Plugins/Cameras/GameplayCameras/Source/GameplayCameras/Private/Nodes/Common/CameraRigCameraNode.cpp
--- a/Engine/Plugins/Cameras/GameplayCameras/Source/GameplayCameras/Private/Nodes/Common/CameraRigCameraNode.cpp
+++ b/Engine/Plugins/Cameras/GameplayCameras/Source/GameplayCameras/Private/Nodes/Common/CameraRigCameraNode.cpp
@@ -2,11 +2,11 @@
 
 #include "Nodes/Common/CameraRigCameraNode.h"
 
-#include "Core/CameraBuildLog.h"
+#include "Build/CameraBuildLog.h"
+#include "Build/CameraObjectBuildContext.h"
 #include "Core/CameraNodeEvaluator.h"
 #include "Core/CameraRigAsset.h"
-#include "Core/CameraRigBuildContext.h"
-#include "Helpers/CameraRigParameterOverrideEvaluator.h"
+#include "Helpers/CameraObjectReferenceParameterOverrideEvaluator.h"
 
 #include UE_INLINE_GENERATED_CPP_BY_NAME(CameraRigCameraNode)
 
@@ -63,7 +63,7 @@ void FCameraRigCameraNodeEvaluator::ApplyParameterOverrides(FCameraVariableTable
 {
 	const UCameraRigCameraNode* PrefabNode = GetCameraNodeAs<UCameraRigCameraNode>();
 
-	FCameraRigParameterOverrideEvaluator OverrideEvaluator(PrefabNode->CameraRigReference);
+	FCameraObjectReferenceParameterOverrideEvaluator OverrideEvaluator(PrefabNode->CameraRigReference);
 	OverrideEvaluator.ApplyParameterOverrides(OutVariableTable, bDrivenOnly);
 }
 
@@ -71,7 +71,7 @@ void FCameraRigCameraNodeEvaluator::ApplyParameterOverrides(FCameraVariableTable
 {
 	const UCameraRigCameraNode* PrefabNode = GetCameraNodeAs<UCameraRigCameraNode>();
 
-	FCameraRigParameterOverrideEvaluator OverrideEvaluator(PrefabNode->CameraRigReference);
+	FCameraObjectReferenceParameterOverrideEvaluator OverrideEvaluator(PrefabNode->CameraRigReference);
 	OverrideEvaluator.ApplyParameterOverrides(OutVariableTable, OutContextDataTable, bDrivenOnly);
 }
 
@@ -90,7 +90,7 @@ void UCameraRigCameraNode::OnPreBuild(FCameraBuildLog& BuildLog)
 	CameraRigReference.RebuildParametersIfNeeded();
 }
 
-void UCameraRigCameraNode::OnBuild(FCameraRigBuildContext& BuildContext)
+void UCameraRigCameraNode::OnBuild(FCameraObjectBuildContext& BuildContext)
 {
 	using namespace UE::Cameras;
 
@@ -120,7 +120,7 @@ FCameraNodeEvaluatorPtr UCameraRigCameraNode::OnBuildEvaluator(FCameraNodeEvalua
 
 #if WITH_EDITOR
 
-void UCameraRigCameraNode::PostEditChangeProperty( struct FPropertyChangedEvent& PropertyChangedEvent)
+void UCameraRigCameraNode::PostEditChangeProperty(struct FPropertyChangedEvent& PropertyChangedEvent)
 {
 	Super::PostEditChangeProperty(PropertyChangedEvent);
 
diff --git a/Engine/Plugins/Cameras/GameplayCameras/Source/GameplayCameras/Private/Nodes/Input/AutoRotateInput2DCameraNode.cpp b/Engine/Plugins/Cameras/GameplayCameras/Source/GameplayCameras/Private/Nodes/Input/AutoRotateInput2DCameraNode.cpp
--- a/Engine/Plugins/Cameras/GameplayCameras/Source/GameplayCameras/Private/Nodes/Input/AutoRotateInput2DCameraNode.cpp
+++ b/Engine/Plugins/Cameras/GameplayCameras/Source/GameplayCameras/Private/Nodes/Input/AutoRotateInput2DCameraNode.cpp
@@ -2,12 +2,12 @@
 
 #include "Nodes/Input/AutoRotateInput2DCameraNode.h"
 
+#include "Build/CameraObjectBuildContext.h"
 #include "Core/BuiltInCameraVariables.h"
 #include "Core/CameraEvaluationContext.h"
 #include "Core/CameraOperation.h"
 #include "Core/CameraParameterReader.h"
 #include "Core/CameraRigAsset.h"
-#include "Core/CameraRigBuildContext.h"
 #include "Core/CameraValueInterpolator.h"
 #include "Core/CameraVariableTableFwd.h"
 #include "Debug/CameraDebugBlock.h"
@@ -322,7 +322,7 @@ FCameraNodeChildrenView UAutoRotateInput2DCameraNode::OnGetChildren()
 	return FCameraNodeChildrenView({ InputNode });
 }
 
-void UAutoRotateInput2DCameraNode::OnBuild(FCameraRigBuildContext& BuildContext)
+void UAutoRotateInput2DCameraNode::OnBuild(FCameraObjectBuildContext& BuildContext)
 {
 	using namespace UE::Cameras;
 	const FBuiltInCameraVariables& BuiltInVariables = FBuiltInCameraVariables::Get();
diff --git a/Engine/Plugins/Cameras/GameplayCameras/Source/GameplayCameras/Private/Nodes/Input/CameraRigInput1DSlot.cpp b/Engine/Plugins/Cameras/GameplayCameras/Source/GameplayCameras/Private/Nodes/Input/CameraRigInput1DSlot.cpp
--- a/Engine/Plugins/Cameras/GameplayCameras/Source/GameplayCameras/Private/Nodes/Input/CameraRigInput1DSlot.cpp
+++ b/Engine/Plugins/Cameras/GameplayCameras/Source/GameplayCameras/Private/Nodes/Input/CameraRigInput1DSlot.cpp
@@ -2,10 +2,10 @@
 
 #include "Nodes/Input/CameraRigInput1DSlot.h"
 
-#include "Core/CameraBuildLog.h"
+#include "Build/CameraBuildLog.h"
+#include "Build/CameraObjectBuildContext.h"
 #include "Core/CameraOperation.h"
 #include "Core/CameraRigAsset.h"
-#include "Core/CameraRigBuildContext.h"
 #include "Core/CameraVariableAssets.h"
 #include "Core/CameraVariableTable.h"
 #include "Nodes/Input/Input1DCameraNode.h"
@@ -90,7 +90,7 @@ void FCameraRigInput1DSlotEvaluator::OnSerialize(const FCameraNodeEvaluatorSeria
 
 }  // namespace UE::Cameras
 
-void UCameraRigInput1DSlot::OnBuild(FCameraRigBuildContext& BuildContext) 
+void UCameraRigInput1DSlot::OnBuild(FCameraObjectBuildContext& BuildContext) 
 {
 	using namespace UE::Cameras;
 
diff --git a/Engine/Plugins/Cameras/GameplayCameras/Source/GameplayCameras/Private/Nodes/Input/CameraRigInput2DSlot.cpp b/Engine/Plugins/Cameras/GameplayCameras/Source/GameplayCameras/Private/Nodes/Input/CameraRigInput2DSlot.cpp
--- a/Engine/Plugins/Cameras/GameplayCameras/Source/GameplayCameras/Private/Nodes/Input/CameraRigInput2DSlot.cpp
+++ b/Engine/Plugins/Cameras/GameplayCameras/Source/GameplayCameras/Private/Nodes/Input/CameraRigInput2DSlot.cpp
@@ -2,10 +2,10 @@
 
 #include "Nodes/Input/CameraRigInput2DSlot.h"
 
-#include "Core/CameraBuildLog.h"
+#include "Build/CameraBuildLog.h"
+#include "Build/CameraObjectBuildContext.h"
 #include "Core/CameraOperation.h"
 #include "Core/CameraRigAsset.h"
-#include "Core/CameraRigBuildContext.h"
 #include "Core/CameraVariableAssets.h"
 #include "Core/CameraVariableTable.h"
 #include "Nodes/Input/Input2DCameraNode.h"
@@ -97,7 +97,7 @@ void FCameraRigInput2DSlotEvaluator::OnSerialize(const FCameraNodeEvaluatorSeria
 
 }  // namespace UE::Cameras
 
-void UCameraRigInput2DSlot::OnBuild(FCameraRigBuildContext& BuildContext)
+void UCameraRigInput2DSlot::OnBuild(FCameraObjectBuildContext& BuildContext)
 {
 	using namespace UE::Cameras;
 
diff --git a/Engine/Plugins/Cameras/GameplayCameras/Source/GameplayCameras/Private/Nodes/Utility/BlueprintCameraNode.cpp b/Engine/Plugins/Cameras/GameplayCameras/Source/GameplayCameras/Private/Nodes/Utility/BlueprintCameraNode.cpp
--- a/Engine/Plugins/Cameras/GameplayCameras/Source/GameplayCameras/Private/Nodes/Utility/BlueprintCameraNode.cpp
+++ b/Engine/Plugins/Cameras/GameplayCameras/Source/GameplayCameras/Private/Nodes/Utility/BlueprintCameraNode.cpp
@@ -2,12 +2,12 @@
 
 #include "Nodes/Utility/BlueprintCameraNode.h"
 
+#include "Build/CameraBuildLog.h"
+#include "Build/CameraObjectBuildContext.h"
 #include "Components/ActorComponent.h"
-#include "Core/CameraBuildLog.h"
 #include "Core/CameraContextDataTableFwd.h"
 #include "Core/CameraEvaluationContext.h"
 #include "Core/CameraNodeEvaluator.h"
-#include "Core/CameraRigBuildContext.h"
 #include "Core/CameraSystemEvaluator.h"
 #include "Core/CameraVariableAssets.h"
 #include "Core/CameraVariableTable.h"
@@ -550,7 +550,7 @@ void UBlueprintCameraNode::OnPreBuild(FCameraBuildLog& BuildLog)
 	RebuildOverrides();
 }
 
-void UBlueprintCameraNode::OnBuild(FCameraRigBuildContext& BuildContext)
+void UBlueprintCameraNode::OnBuild(FCameraObjectBuildContext& BuildContext)
 {
 	if (!CameraNodeEvaluatorTemplate)
 	{
diff --git a/Engine/Plugins/Cameras/GameplayCameras/Source/GameplayCameras/Private/Tests/CameraRigAssetBuilderTests.cpp b/Engine/Plugins/Cameras/GameplayCameras/Source/GameplayCameras/Private/Tests/CameraRigAssetBuilderTests.cpp
--- a/Engine/Plugins/Cameras/GameplayCameras/Source/GameplayCameras/Private/Tests/CameraRigAssetBuilderTests.cpp
+++ b/Engine/Plugins/Cameras/GameplayCameras/Source/GameplayCameras/Private/Tests/CameraRigAssetBuilderTests.cpp
@@ -72,7 +72,7 @@ bool FCameraRigAssetBuilderSimpleParameterTest::RunTest(const FString& Parameter
 
 	CameraRig->BuildCameraRig();
 
-	UCameraRigBlendableParameter* Parameter = CameraRig->Interface.BlendableParameters[0];
+	UCameraObjectInterfaceBlendableParameter* Parameter = CameraRig->Interface.BlendableParameters[0];
 	UTEST_EQUAL("Test parameter", Parameter->InterfaceParameterName, TEXT("Test"));
 	UTEST_TRUE("Test parameter variable ID", Parameter->PrivateVariableID.IsValid());
 	UTEST_EQUAL("Test node parameter", OffsetNode->TranslationOffset.VariableID, Parameter->PrivateVariableID);
@@ -159,8 +159,8 @@ bool FCameraRigAssetBuilderDrivenOverridesTest::RunTest(const FString& Parameter
 
 	OuterCameraRig->BuildCameraRig();
 
-	UCameraRigBlendableParameter* OffsetParam = InnerCameraRig->Interface.BlendableParameters[0];
-	UCameraRigBlendableParameter* FocalLengthParam = InnerCameraRig->Interface.BlendableParameters[1];
+	UCameraObjectInterfaceBlendableParameter* OffsetParam = InnerCameraRig->Interface.BlendableParameters[0];
+	UCameraObjectInterfaceBlendableParameter* FocalLengthParam = InnerCameraRig->Interface.BlendableParameters[1];
 
 	// Test that the inner nodes are driven by the interface parameters.
 	{
@@ -193,7 +193,7 @@ bool FCameraRigAssetBuilderDrivenOverridesTest::RunTest(const FString& Parameter
 		UTEST_EQUAL_EXPR(FocalLengthParamOverride->Value, 25.f);
 	}
 
-	UCameraRigBlendableParameter* MiddleOffsetParam = MiddleCameraRig->Interface.BlendableParameters[0];
+	UCameraObjectInterfaceBlendableParameter* MiddleOffsetParam = MiddleCameraRig->Interface.BlendableParameters[0];
 	{
 		FInstancedPropertyBag& ParameterOverrides = MiddlePrefabNode->CameraRigReference.GetParameters();
 
diff --git a/Engine/Plugins/Cameras/GameplayCameras/Source/GameplayCameras/Private/Tests/GameplayCamerasTestBuilder.h b/Engine/Plugins/Cameras/GameplayCameras/Source/GameplayCameras/Private/Tests/GameplayCamerasTestBuilder.h
--- a/Engine/Plugins/Cameras/GameplayCameras/Source/GameplayCameras/Private/Tests/GameplayCamerasTestBuilder.h
+++ b/Engine/Plugins/Cameras/GameplayCameras/Source/GameplayCameras/Private/Tests/GameplayCamerasTestBuilder.h
@@ -487,7 +487,7 @@ public:
 	 */
 	ThisType& AddBlendableParameter(const FString& ParameterName, ECameraVariableType ParameterType, UCameraNode* Target, FName TargetPropertyName)
 	{
-		UCameraRigBlendableParameter* BlendableParameter = NewObject<UCameraRigBlendableParameter>(CameraRig);
+		UCameraObjectInterfaceBlendableParameter* BlendableParameter = NewObject<UCameraObjectInterfaceBlendableParameter>(CameraRig);
 		BlendableParameter->InterfaceParameterName = ParameterName;
 		BlendableParameter->ParameterType = ParameterType;
 		BlendableParameter->Target = Target;
diff --git a/Engine/Plugins/Cameras/GameplayCameras/Source/GameplayCameras/Public/Build/CameraNodeHierarchyBuilder.h b/Engine/Plugins/Cameras/GameplayCameras/Source/GameplayCameras/Public/Build/CameraNodeHierarchyBuilder.h
new file mode 100644
--- /dev/null
+++ b/Engine/Plugins/Cameras/GameplayCameras/Source/GameplayCameras/Public/Build/CameraNodeHierarchyBuilder.h
@@ -0,0 +1,47 @@
+// Copyright Epic Games, Inc. All Rights Reserved.
+
+#pragma once
+
+#include "Build/CameraBuildLog.h"
+#include "Core/CameraNodeHierarchy.h"
+
+class UBaseCameraObject;
+
+namespace UE::Cameras
+{
+
+struct FCameraObjectBuildContext;
+
+/**
+ * A helper class that can build a hierarchy of camera nodes.
+ */
+class GAMEPLAYCAMERAS_API FCameraNodeHierarchyBuilder
+{
+public:
+
+	/** Creates a new camera node hierarchy builder. */
+	FCameraNodeHierarchyBuilder(FCameraBuildLog& InBuildLog, UBaseCameraObject* InCameraObject);
+
+	/** Gets the camera node hierarchy. */
+	const FCameraNodeHierarchy& GetHierarchy() const { return CameraNodeHierarchy; }
+
+	/** Calls PreBuild on all the camera nodes. */
+	void PreBuild();
+
+	/** Calls Build on all the camera nodes and computes the allocation info. */
+	void Build();
+
+private:
+
+	void CallBuild(FCameraObjectBuildContext& BuildContext, UCameraNode* CameraNode);
+	void BuildParametersAllocationInfo(FCameraObjectBuildContext& BuildContext);
+
+private:
+
+	FCameraBuildLog& BuildLog;
+	UBaseCameraObject* CameraObject = nullptr;
+	FCameraNodeHierarchy CameraNodeHierarchy;
+};
+
+}  // namespace UE::Cameras
+
diff --git a/Engine/Plugins/Cameras/GameplayCameras/Source/GameplayCameras/Public/Build/CameraObjectBuildContext.h b/Engine/Plugins/Cameras/GameplayCameras/Source/GameplayCameras/Public/Build/CameraObjectBuildContext.h
new file mode 100644
--- /dev/null
+++ b/Engine/Plugins/Cameras/GameplayCameras/Source/GameplayCameras/Public/Build/CameraObjectBuildContext.h
@@ -0,0 +1,29 @@
+// Copyright Epic Games, Inc. All Rights Reserved.
+
+#pragma once
+
+#include "Core/BaseCameraObject.h"
+
+namespace UE::Cameras
+{
+
+class FCameraBuildLog;
+
+/**
+ * Camera object build context.
+ */
+struct FCameraObjectBuildContext
+{
+	FCameraObjectBuildContext(FCameraBuildLog& InBuildLog)
+		: BuildLog(InBuildLog)
+	{}
+
+	/** The build log for emitting messages. */
+	FCameraBuildLog& BuildLog;
+
+	/** The allocation information for the camera rig. */
+	FCameraObjectAllocationInfo AllocationInfo;
+};
+
+}  // namespace UE::Cameras
+
diff --git a/Engine/Plugins/Cameras/GameplayCameras/Source/GameplayCameras/Public/Build/CameraObjectInterfaceBuilder.h b/Engine/Plugins/Cameras/GameplayCameras/Source/GameplayCameras/Public/Build/CameraObjectInterfaceBuilder.h
new file mode 100644
--- /dev/null
+++ b/Engine/Plugins/Cameras/GameplayCameras/Source/GameplayCameras/Public/Build/CameraObjectInterfaceBuilder.h
@@ -0,0 +1,64 @@
+// Copyright Epic Games, Inc. All Rights Reserved.
+
+#pragma once
+
+#include "Build/CameraBuildLog.h"
+#include "Containers/ArrayView.h"
+#include "Core/CameraContextDataTableFwd.h"
+#include "Core/CameraVariableTableFwd.h"
+
+class UBaseCameraObject;
+class UCameraObjectInterfaceBlendableParameter;
+class UCameraObjectInterfaceDataParameter;
+class UCameraNode;
+struct FCameraObjectInterface;
+
+namespace UE::Cameras
+{
+
+class FCameraNodeHierarchy;
+
+namespace Internal { struct FInterfaceParameterBindingBuilder; }
+
+class GAMEPLAYCAMERAS_API FCameraObjectInterfaceBuilder
+{
+public:
+
+	FCameraObjectInterfaceBuilder(FCameraBuildLog& InBuildLog);
+
+	void BuildInterface(UBaseCameraObject* InCameraObject, const FCameraNodeHierarchy& InHierarchy, bool bCollectStrayNodes);
+	void BuildInterface(UBaseCameraObject* InCameraObject, TArrayView<UCameraNode*> InCameraObjectNodes);
+
+private:
+
+	void BuildInterfaceImpl();
+
+	void GatherOldDrivenParameters();
+	void BuildInterfaceParameters();
+	void BuildInterfaceParameterBindings();
+	void DiscardUnusedParameters();
+
+private:
+
+	bool SetupCameraParameterOrVariableReferenceOverride(const UCameraObjectInterfaceBlendableParameter* BlendableParameter);
+	bool SetupCustomBlendableParameterOverride(const UCameraObjectInterfaceBlendableParameter* BlendableParameter);
+
+	bool SetupDataContextPropertyOverride(const UCameraObjectInterfaceDataParameter* DataParameter);
+	bool SetupCustomDataParameterOverride(const UCameraObjectInterfaceDataParameter* DataParameter);
+
+private:
+
+	FCameraBuildLog& BuildLog;
+
+	UBaseCameraObject* CameraObject = nullptr;
+	TArray<UCameraNode*> CameraObjectNodes;
+
+	using FDrivenParameterKey = TTuple<FName, UObject*>;
+	TMap<FDrivenParameterKey, FCameraVariableID> OldDrivenBlendableParameters;
+	TMap<FDrivenParameterKey, FCameraContextDataID> OldDrivenDataParameters;
+
+	friend struct Internal::FInterfaceParameterBindingBuilder;
+};
+
+}  // namespace UE::Cameras
+
diff --git a/Engine/Plugins/Cameras/GameplayCameras/Source/GameplayCameras/Public/Build/CameraObjectInterfaceParameterBuilder.h b/Engine/Plugins/Cameras/GameplayCameras/Source/GameplayCameras/Public/Build/CameraObjectInterfaceParameterBuilder.h
new file mode 100644
--- /dev/null
+++ b/Engine/Plugins/Cameras/GameplayCameras/Source/GameplayCameras/Public/Build/CameraObjectInterfaceParameterBuilder.h
@@ -0,0 +1,52 @@
+// Copyright Epic Games, Inc. All Rights Reserved.
+
+#pragma once
+
+#include "Containers/Array.h"
+#include "StructUtils/PropertyBag.h"
+
+class UBaseCameraObject;
+class UCameraObjectInterfaceDataParameter;
+
+namespace UE::Cameras
+{
+
+/**
+ * A helper class for building an FInstancedPropertyBag from a list of camera rig
+ * parameter definitions.
+ */
+/** Builds the parameter definitions for the given camera object. */
+/**
+ * Builds the property bag that contains a property for each exposed parameter on the given camera object.
+ * Each property's value is set to the default value of the corresponding parameter.
+ */
+class GAMEPLAYCAMERAS_API FCameraObjectInterfaceParameterBuilder
+{
+public:
+
+	FCameraObjectInterfaceParameterBuilder();
+
+	void BuildParameters(UBaseCameraObject* InCameraObject);
+
+public:
+
+	static void BuildDefaultParameters(const UBaseCameraObject* CameraObject, FInstancedPropertyBag& OutPropertyBag);
+	static void AppendDefaultParameterProperties(const UBaseCameraObject* CameraObject, TArray<FPropertyBagPropertyDesc>& OutProperties);
+	static void SetDefaultParameterValues(const UBaseCameraObject* CameraObject, FInstancedPropertyBag& PropertyBag);
+
+private:
+
+	static void SetDefaultParameterValue(const UCameraObjectInterfaceDataParameter* DataParameter, void* DestValuePtr, const void* SrcValuePtr);
+
+	void BuildParametersImpl();
+
+	void BuildParameterDefinitions();
+	void BuildDefaultParameters();
+
+private:
+
+	UBaseCameraObject* CameraObject = nullptr;
+};
+
+}  // namespace UE::Cameras
+
diff --git a/Engine/Plugins/Cameras/GameplayCameras/Source/GameplayCameras/Public/Core/BaseCameraObjectReference.h b/Engine/Plugins/Cameras/GameplayCameras/Source/GameplayCameras/Public/Core/BaseCameraObjectReference.h
new file mode 100644
--- /dev/null
+++ b/Engine/Plugins/Cameras/GameplayCameras/Source/GameplayCameras/Public/Core/BaseCameraObjectReference.h
@@ -0,0 +1,138 @@
+// Copyright Epic Games, Inc. All Rights Reserved.
+
+#pragma once
+
+#include "Core/CameraContextDataTableFwd.h"
+#include "Core/CameraVariableTableFwd.h"
+#include "StructUtils/PropertyBag.h"
+
+#include "BaseCameraObjectReference.generated.h"
+
+class UBaseCameraObject;
+struct FCustomCameraNodeParameterInfos;
+
+namespace UE::Cameras
+{
+	struct FCameraNodeEvaluationResult;
+}
+
+/**
+ * Metadata for a referenced camera object's interface parameters.
+ */
+USTRUCT()
+struct FCameraObjectInterfaceParameterMetaData
+{
+	GENERATED_BODY()
+
+	/** The GUID of the parameter. */
+	UPROPERTY()
+	FGuid ParameterGuid;
+
+	/** The ID to use for overriding a blendable parameter. */
+	UPROPERTY()
+	FCameraVariableID OverrideVariableID;
+
+	/** The ID to use for overriding a data parameter. */
+	UPROPERTY()
+	FCameraContextDataID OverrideDataID;
+
+	/** Whether this parameter has an override value. */
+	UPROPERTY()
+	bool bIsOverridden = false;
+
+	/** Whether this parameter is being animated. */
+	UPROPERTY()
+	bool bIsAnimated = false;
+};
+
+USTRUCT(BlueprintType)
+struct FBaseCameraObjectReference
+{
+	GENERATED_BODY()
+
+public:
+
+	virtual ~FBaseCameraObjectReference() {}
+
+	/** Gets the parameters for this camera rig, some of which containing overrides. */
+	const FInstancedPropertyBag& GetParameters() const
+	{
+		return Parameters;
+	}
+
+	/** Gets the parameters for this camera rig, some of which containing overrides. */
+	FInstancedPropertyBag& GetParameters()
+	{
+		return Parameters;
+	}
+
+public:
+
+	/** Gets the camera object that this reference points to. */
+	GAMEPLAYCAMERAS_API virtual const UBaseCameraObject* GetCameraObject() const { return nullptr; }
+
+	/** Returns whether the override parameters structure needs to be rebuilt. */
+	GAMEPLAYCAMERAS_API bool NeedsRebuildParameters() const;
+	/** Rebuilds the override parameters structure, if needed. */
+	GAMEPLAYCAMERAS_API bool RebuildParametersIfNeeded();
+	/** Rebuilds the override parameters structure. */
+	GAMEPLAYCAMERAS_API void RebuildParameters();
+
+public:
+
+	// Internal API.
+
+	GAMEPLAYCAMERAS_API bool IsParameterOverridden(const FGuid& PropertyID) const;
+	GAMEPLAYCAMERAS_API void SetParameterOverridden(const FGuid& PropertyID, bool bIsOverridden);
+
+	GAMEPLAYCAMERAS_API bool IsParameterAnimated(const FGuid& PropertyID) const;
+	GAMEPLAYCAMERAS_API void SetParameterAnimated(const FGuid& PropertyID, bool bIsAnimated);
+
+	GAMEPLAYCAMERAS_API void GetCustomCameraNodeParameters(FCustomCameraNodeParameterInfos& OutParameterInfos);
+	
+	template<typename ContainerType>
+	void GetOverriddenParameterGuids(ContainerType& OutOverriddenIDs) const;
+
+	template<typename ContainerType>
+	void GetAnimatedParameterGuids(ContainerType& OutAnimatedIDs) const;
+
+private:
+
+	const FCameraObjectInterfaceParameterMetaData* FindMetaData(const FGuid& PropertyID) const;
+	FCameraObjectInterfaceParameterMetaData& FindOrAddMetaData(const FGuid& PropertyID);
+
+protected:
+
+	/** The camera rig's parameters. */
+	UPROPERTY(EditAnywhere, Category="", meta=(FixedLayout))
+	FInstancedPropertyBag Parameters;
+
+	/** Metadata for the parameters. */
+	UPROPERTY()
+	TArray<FCameraObjectInterfaceParameterMetaData> ParameterMetaData;
+};
+
+template<typename ContainerType>
+void FBaseCameraObjectReference::GetOverriddenParameterGuids(ContainerType& OutOverriddenIDs) const
+{
+	for (const FCameraObjectInterfaceParameterMetaData& MetaData : ParameterMetaData)
+	{
+		if (MetaData.bIsOverridden)
+		{
+			OutOverriddenIDs.Add(MetaData.ParameterGuid);
+		}
+	}
+}
+
+template<typename ContainerType>
+void FBaseCameraObjectReference::GetAnimatedParameterGuids(ContainerType& OutAnimatedIDs) const
+{
+	for (const FCameraObjectInterfaceParameterMetaData& MetaData : ParameterMetaData)
+	{
+		if (MetaData.bIsAnimated)
+		{
+			OutAnimatedIDs.Add(MetaData.ParameterGuid);
+		}
+	}
+}
+
diff --git a/Engine/Plugins/Cameras/GameplayCameras/Source/GameplayCameras/Public/Core/CameraAsset.h b/Engine/Plugins/Cameras/GameplayCameras/Source/GameplayCameras/Public/Core/CameraAsset.h
--- a/Engine/Plugins/Cameras/GameplayCameras/Source/GameplayCameras/Public/Core/CameraAsset.h
+++ b/Engine/Plugins/Cameras/GameplayCameras/Source/GameplayCameras/Public/Core/CameraAsset.h
@@ -2,10 +2,10 @@
 
 #pragma once
 
-#include "Core/CameraBuildStatus.h"
+#include "Build/CameraBuildStatus.h"
 #include "Core/CameraContextDataTableAllocationInfo.h"
 #include "Core/CameraEventHandler.h"
-#include "Core/CameraRigParameterDefinition.h"
+#include "Core/CameraObjectInterfaceParameterDefinition.h"
 #include "Core/CameraRigTransition.h"
 #include "Core/CameraVariableTableAllocationInfo.h"
 #include "Core/ObjectTreeGraphObject.h"
@@ -103,7 +103,7 @@ public:
 	FInstancedPropertyBag& GetDefaultParameters() { return DefaultParameters; }
 
 	/** Gets the definitions of parameters exposed on this camera asset. */
-	TConstArrayView<FCameraRigParameterDefinition> GetParameterDefinitions() const { return ParameterDefinitions; }
+	TConstArrayView<FCameraObjectInterfaceParameterDefinition> GetParameterDefinitions() const { return ParameterDefinitions; }
 
 public:
 
@@ -188,7 +188,7 @@ private:
 
 	/** Generated list of all the camera rigs' parameters. */
 	UPROPERTY()
-	TArray<FCameraRigParameterDefinition> ParameterDefinitions;
+	TArray<FCameraObjectInterfaceParameterDefinition> ParameterDefinitions;
 
 	/** Owner rigs for each parameter definition. */
 	UPROPERTY()
diff --git a/Engine/Plugins/Cameras/GameplayCameras/Source/GameplayCameras/Public/Core/CameraNode.h b/Engine/Plugins/Cameras/GameplayCameras/Source/GameplayCameras/Public/Core/CameraNode.h
--- a/Engine/Plugins/Cameras/GameplayCameras/Source/GameplayCameras/Public/Core/CameraNode.h
+++ b/Engine/Plugins/Cameras/GameplayCameras/Source/GameplayCameras/Public/Core/CameraNode.h
@@ -13,7 +13,7 @@
 namespace UE::Cameras
 {
 	class FCameraBuildLog;
-	struct FCameraRigBuildContext;
+	struct FCameraObjectBuildContext;
 }
 
 /** View on a camera node's children. */
@@ -42,7 +42,7 @@ class GAMEPLAYCAMERAS_API UCameraNode
 public:
 	
 	using FCameraBuildLog = UE::Cameras::FCameraBuildLog;
-	using FCameraRigBuildContext = UE::Cameras::FCameraRigBuildContext;
+	using FCameraObjectBuildContext = UE::Cameras::FCameraObjectBuildContext;
 	using FCameraNodeEvaluatorBuilder = UE::Cameras::FCameraNodeEvaluatorBuilder;
 
 	/** Get the list of children under this node. */
@@ -52,7 +52,7 @@ public:
 	void PreBuild(FCameraBuildLog& BuildLog);
 
 	/** Gets optional info about this node's required allocations at runtime. */
-	void Build(FCameraRigBuildContext& BuildContext);
+	void Build(FCameraObjectBuildContext& BuildContext);
 
 	/** Builds the evaluator for this node. */
 	FCameraNodeEvaluatorPtr BuildEvaluator(FCameraNodeEvaluatorBuilder& Builder) const;
@@ -74,7 +74,7 @@ protected:
 	virtual void OnPreBuild(FCameraBuildLog& BuildLog) {}
 
 	/** Gets optional info about this node's required allocations at runtime. */
-	virtual void OnBuild(FCameraRigBuildContext& BuildContext) {}
+	virtual void OnBuild(FCameraObjectBuildContext& BuildContext) {}
 
 	/** Builds the evaluator for this node. */
 	virtual FCameraNodeEvaluatorPtr OnBuildEvaluator(FCameraNodeEvaluatorBuilder& Builder) const { return nullptr; }
diff --git a/Engine/Plugins/Cameras/GameplayCameras/Source/GameplayCameras/Public/Core/CameraNodeHierarchy.h b/Engine/Plugins/Cameras/GameplayCameras/Source/GameplayCameras/Public/Core/CameraNodeHierarchy.h
--- a/Engine/Plugins/Cameras/GameplayCameras/Source/GameplayCameras/Public/Core/CameraNodeHierarchy.h
+++ b/Engine/Plugins/Cameras/GameplayCameras/Source/GameplayCameras/Public/Core/CameraNodeHierarchy.h
@@ -5,7 +5,7 @@
 #include "Containers/Array.h"
 #include "Core/CameraNode.h"
 
-class UCameraRigAsset;
+class UBaseCameraObject;
 
 namespace UE::Cameras
 {
@@ -20,8 +20,8 @@ public:
 
 	/** Build an empty hierarchy. */
 	FCameraNodeHierarchy();
-	/** Build a hierarchy starting from the given camera rig's root node. */
-	FCameraNodeHierarchy(UCameraRigAsset* InCameraRig);
+	/** Build a hierarchy starting from the given camera object's root node. */
+	FCameraNodeHierarchy(UBaseCameraObject* InCameraObject);
 	/** Build a hierarchy starting from the given root node. */
 	FCameraNodeHierarchy(UCameraNode* InRootCameraNode);
 
@@ -33,8 +33,8 @@ public:
 
 public:
 
-	/** Build a hierarchy starting from the given camera rig's root node. */
-	void Build(UCameraRigAsset* InCameraRig);
+	/** Build a hierarchy starting from the given camera object's root node. */
+	void Build(UBaseCameraObject* InCameraObject);
 	/** Build a hierarchy starting from the given root node. */
 	void Build(UCameraNode* InRootCameraNode);
 	/** Resets this object to an empty hierarchy. */
diff --git a/Engine/Plugins/Cameras/GameplayCameras/Source/GameplayCameras/Public/Core/CameraRigAsset.h b/Engine/Plugins/Cameras/GameplayCameras/Source/GameplayCameras/Public/Core/CameraRigAsset.h
--- a/Engine/Plugins/Cameras/GameplayCameras/Source/GameplayCameras/Public/Core/CameraRigAsset.h
+++ b/Engine/Plugins/Cameras/GameplayCameras/Source/GameplayCameras/Public/Core/CameraRigAsset.h
@@ -2,20 +2,16 @@
 
 #pragma once
 
-#include "Core/CameraBuildStatus.h"
-#include "Core/CameraContextDataTableAllocationInfo.h"
-#include "Core/CameraContextDataTableFwd.h"
+#include "Build/CameraBuildStatus.h"
+#include "Core/BaseCameraObject.h"
 #include "Core/CameraEventHandler.h"
-#include "Core/CameraNodeEvaluatorFwd.h"
-#include "Core/CameraRigParameterDefinition.h"
+#include "Core/CameraObjectInterfaceParameterDefinition.h"
 #include "Core/CameraRigTransition.h"
-#include "Core/CameraVariableTableAllocationInfo.h"
 #include "Core/ObjectTreeGraphObject.h"
 #include "Core/ObjectTreeGraphRootObject.h"
 #include "CoreTypes.h"
 #include "GameplayTagAssetInterface.h"
 #include "GameplayTagContainer.h"
-#include "StructUtils/PropertyBag.h"
 #include "UObject/ObjectPtr.h"
 
 #include "CameraRigAsset.generated.h"
@@ -40,9 +36,6 @@ namespace UE::Cameras
 		/** Called when the camera rig asset has been built. */
 		virtual void OnCameraRigBuilt(const UCameraRigAsset* CameraRigAsset) {}
 
-		/** Called when the camera rig's interface has changed. */
-		virtual void OnCameraRigInterfaceChanged() {}
-
 #if WITH_EDITOR
 		virtual void OnObjectAddedToGraph(const FName GraphName, UObject* Object) {}
 		virtual void OnObjectRemovedFromGraph(const FName GraphName, UObject* Object) {}
@@ -50,233 +43,6 @@ namespace UE::Cameras
 	};
 }
 
-/**
- * Structure describing various allocations needed by a camera rig.
- */
-USTRUCT()
-struct FCameraRigAllocationInfo
-{
-	GENERATED_BODY()
-
-	/** Allocation info for node evaluators. */
-	UPROPERTY()
-	FCameraNodeEvaluatorAllocationInfo EvaluatorInfo;
-
-	/** Allocation info for the variable table. */
-	UPROPERTY()
-	FCameraVariableTableAllocationInfo VariableTableInfo;
-
-	/** Allocation info for the context data table. */
-	UPROPERTY()
-	FCameraContextDataTableAllocationInfo ContextDataTableInfo;
-
-public:
-
-	GAMEPLAYCAMERAS_API void Append(const FCameraRigAllocationInfo& OtherAllocationInfo);
-
-	bool operator==(const FCameraRigAllocationInfo& Other) const = default;
-};
-
-template<>
-struct TStructOpsTypeTraits<FCameraRigAllocationInfo> : public TStructOpsTypeTraitsBase2<FCameraRigAllocationInfo>
-{
-	enum
-	{
-		WithCopy = true,
-		WithIdenticalViaEquality = true
-	};
-};
-
-/**
- * Base class for interface parameters on a camera rig asset.
- */
-UCLASS(MinimalAPI, meta=(
-			ObjectTreeGraphSelfPinDirection="Output"))
-class UCameraRigInterfaceParameterBase 
-	: public UObject
-	, public IObjectTreeGraphObject
-{
-	GENERATED_BODY()
-
-public:
-
-	/** The exposed name for this parameter. */
-	UPROPERTY(EditAnywhere, Category=Camera)
-	FString InterfaceParameterName;
-
-	/** The camera node this parameter is connected. */
-	UPROPERTY(meta=(ObjectTreeGraphHidden=true))
-	TObjectPtr<UCameraNode> Target;
-
-	/**
-	 * The name of the property this parameter is connected to on the target camera node.
-	 * This may be an actual UObject property, but it may be something else, like the name
-	 * of an interface parameter on a nested camera rig, or the name of a Blueprint property
-	 * on the evaluator class of a Blueprint camera node.
-	 */
-	UPROPERTY()
-	FName TargetPropertyName;
-
-#if WITH_EDITORONLY_DATA
-
-	/** Whether this parameter has been added to the node graph in the editor. */
-	UPROPERTY()
-	bool bHasGraphNode = false;
-
-#endif  // WITH_EDITORONLY_DATA
-
-public:
-
-	/** Gets this parameter's unique ID. */
-	const FGuid& GetGuid() const { return Guid; }
-
-protected:
-
-	/** The Guid of this parameter. */
-	UPROPERTY()
-	FGuid Guid;
-
-protected:
-
-	// IObjectTreeGraphObject interface.
-#if WITH_EDITOR
-	virtual void GetGraphNodePosition(FName InGraphName, int32& NodePosX, int32& NodePosY) const override;
-	virtual void OnGraphNodeMoved(FName InGraphName, int32 NodePosX, int32 NodePosY, bool bMarkDirty) override;
-#endif
-
-	// UObject interface.
-	virtual void PostLoad() override;
-	virtual void PostInitProperties() override;
-	virtual void PostDuplicate(EDuplicateMode::Type DuplicateMode) override;
-
-private:
-
-#if WITH_EDITORONLY_DATA
-
-	UPROPERTY()
-	FIntVector2 GraphNodePos = FIntVector2::ZeroValue;
-
-#endif  // WITH_EDITORONLY_DATA
-};
-
-/**
- * An exposed camera rig parameter that drives a specific parameter on one of
- * its camera nodes.
- */
-UCLASS(MinimalAPI)
-class UCameraRigBlendableParameter : public UCameraRigInterfaceParameterBase
-{
-	GENERATED_BODY()
-
-public:
-
-	/** The type of this parameter. */
-	UPROPERTY()
-	ECameraVariableType ParameterType = ECameraVariableType::Boolean;
-
-	/** The struct type of this parameter if it is a blendable struct. */
-	UPROPERTY()
-	TObjectPtr<const UScriptStruct> BlendableStructType;
-
-	/**
-	 * Whether this parameter's value should be pre-blended.
-	 *
-	 * Pre-blending means that if two blending camera rigs share this parameter, 
-	 * each of their values will be blended in a first evaluation pass, and then
-	 * both camera rigs will evaluate with the same blended value.
-	 */
-	UPROPERTY()
-	bool bIsPreBlended = true;
-
-	// Built on save/cook.
-
-	/** The ID to use to access the underlying variable value in the variable table. */
-	UPROPERTY()
-	FCameraVariableID PrivateVariableID;
-
-
-	// Deprecated.
-
-	UPROPERTY()
-	TObjectPtr<UCameraVariableAsset> PrivateVariable_DEPRECATED;
-
-public:
-
-	/** Gets the camera variable definition for this parameter. */
-	GAMEPLAYCAMERAS_API FCameraVariableDefinition GetVariableDefinition() const;
-
-#if WITH_EDITORONLY_DATA
-	GAMEPLAYCAMERAS_API FString GetVariableName() const;
-#endif
-};
-
-UCLASS(MinimalAPI)
-class UCameraRigDataParameter : public UCameraRigInterfaceParameterBase
-{
-	GENERATED_BODY()
-
-public:
-
-	/** The type of this parameter. */
-	UPROPERTY()
-	ECameraContextDataType DataType = ECameraContextDataType::Name;
-
-	/** The type of container for this parameter. */
-	UPROPERTY()
-	ECameraContextDataContainerType DataContainerType = ECameraContextDataContainerType::None;
-
-	/** An additional type object for this parameter. */
-	UPROPERTY()
-	TObjectPtr<const UObject> DataTypeObject;
-
-	// Built on save/cook.
-
-	/** The ID to use to access the underlying data in the context data table. */
-	UPROPERTY()
-	FCameraContextDataID PrivateDataID;
-
-public:
-
-	/** Gets the camera context data definition for this parameter. */
-	GAMEPLAYCAMERAS_API FCameraContextDataDefinition GetDataDefinition() const;
-
-#if WITH_EDITORONLY_DATA
-	GAMEPLAYCAMERAS_API FString GetDataName() const;
-#endif
-};
-
-/**
- * Structure defining the public data interface of a camera rig asset.
- */
-USTRUCT()
-struct FCameraRigInterface
-{
-	GENERATED_BODY()
-
-public:
-
-	/** The list of exposed blendable parameters on the camera rig. */
-	UPROPERTY(Instanced)
-	TArray<TObjectPtr<UCameraRigBlendableParameter>> BlendableParameters;
-
-	/** The list of exposed data parameters on the camera rig. */
-	UPROPERTY(Instanced)
-	TArray<TObjectPtr<UCameraRigDataParameter>> DataParameters;
-
-public:
-	
-	/** Finds an exposed parameter by name. */
-	GAMEPLAYCAMERAS_API UCameraRigBlendableParameter* FindBlendableParameterByName(const FString& ParameterName) const;
-	GAMEPLAYCAMERAS_API UCameraRigDataParameter* FindDataParameterByName(const FString& ParameterName) const;
-
-	/** Finds an exposed parameter by Guid. */
-	GAMEPLAYCAMERAS_API UCameraRigBlendableParameter* FindBlendableParameterByGuid(const FGuid& ParameterGuid) const;
-	GAMEPLAYCAMERAS_API UCameraRigDataParameter* FindDataParameterByGuid(const FGuid& ParameterGuid) const;
-
-	/** Returns whether an exposed parameter with the given name exists. */
-	GAMEPLAYCAMERAS_API bool HasBlendableParameter(const FString& ParameterName) const;
-};
-
 /**
  * List of packages that contain the definition of a camera rig.
  * In most cases there's only one, but with nested assets there could be more.
@@ -289,7 +55,7 @@ using FCameraRigPackages = TArray<const UPackage*, TInlineAllocator<4>>;
  */
 UCLASS(MinimalAPI)
 class UCameraRigAsset
-	: public UObject
+	: public UBaseCameraObject
 	, public IGameplayTagAssetInterface
 	, public IHasCameraBuildStatus
 	, public IObjectTreeGraphObject
@@ -313,10 +79,6 @@ public:
 	UPROPERTY(EditAnywhere, Category=GameplayTags)
 	FGameplayTagContainer GameplayTags;
 
-	/** The public data interface of this camera rig. */
-	UPROPERTY()
-	FCameraRigInterface Interface;
-
 	/** List of enter transitions for this camera rig. */
 	UPROPERTY(Instanced)
 	TArray<TObjectPtr<UCameraRigTransition>> EnterTransitions;
@@ -329,22 +91,9 @@ public:
 	UPROPERTY(EditAnywhere, Category="Transition")
 	ECameraRigInitialOrientation InitialOrientation = ECameraRigInitialOrientation::None;
 
-	/** Allocation information for all the nodes and variables in this camera rig. */
-	UPROPERTY()
-	FCameraRigAllocationInfo AllocationInfo;
-
 	/** Gets the camera rig's unique ID. */
 	const FGuid& GetGuid() const { return Guid; }
 
-	/** Gets the default values for the parameters exposed on this camera rig. */
-	const FInstancedPropertyBag& GetDefaultParameters() const { return DefaultParameters; }
-
-	/** Gets the default values for the parameters exposed on this camera rig. */
-	FInstancedPropertyBag& GetDefaultParameters() { return DefaultParameters; }
-
-	/** Gets the definitions of parameters exposed on this camera rig. */
-	TConstArrayView<FCameraRigParameterDefinition> GetParameterDefinitions() const { return ParameterDefinitions; }
-
 public:
 
 	/** The current build state of this camera rig. */
@@ -365,6 +114,9 @@ public:
 
 public:
 
+	// BaseCameraObject interface.
+	virtual UCameraNode* GetRootNode() override { return RootNode; }
+
 	// IGameplayTagAssetInterface interface.
 	virtual void GetOwnedGameplayTags(FGameplayTagContainer& TagContainer) const override;
 
@@ -416,14 +168,6 @@ private:
 	UPROPERTY()
 	FGuid Guid;
 
-	/** The default interface parameter values, generated during build. */
-	UPROPERTY()
-	FInstancedPropertyBag DefaultParameters;
-
-	/** The definitions of parameters exposed on this camera rig. */
-	UPROPERTY()
-	TArray<FCameraRigParameterDefinition> ParameterDefinitions;
-
 #if WITH_EDITORONLY_DATA
 
 	/** Position of the camera rig node in the node graph editor. */
diff --git a/Engine/Plugins/Cameras/GameplayCameras/Source/GameplayCameras/Public/Core/CameraRigAssetReference.h b/Engine/Plugins/Cameras/GameplayCameras/Source/GameplayCameras/Public/Core/CameraRigAssetReference.h
--- a/Engine/Plugins/Cameras/GameplayCameras/Source/GameplayCameras/Public/Core/CameraRigAssetReference.h
+++ b/Engine/Plugins/Cameras/GameplayCameras/Source/GameplayCameras/Public/Core/CameraRigAssetReference.h
@@ -2,14 +2,14 @@
 
 #pragma once
 
-#include "Core/CameraContextDataTableFwd.h"
+#include "BaseCameraObjectReference.h"
+#include "Core/CameraObjectInterfaceParameterDefinition.h"
 #include "Core/CameraParameters.h"
-#include "Core/CameraRigParameterDefinition.h"
-#include "StructUtils/PropertyBag.h"
 #include "UObject/ObjectPtr.h"
 
 #include "CameraRigAssetReference.generated.h"
 
+class UBaseCameraObject;
 class UCameraRigAsset;
 struct FCameraRigAssetReference;
 struct FCustomCameraNodeParameterInfos;
@@ -17,12 +17,10 @@ struct FPropertyTag;
 
 namespace UE::Cameras
 {
-	class FCameraRigAssetBuilder;
 	class FCameraRigAssetReferenceDetailsCustomization;
 	struct FCameraNodeEvaluationResult;
 }
 
-
 struct UE_DEPRECATED(5.7, "Camera rig references now use instanced property bags.") FCameraRigParameterOverrideBase;
 
 struct UE_DEPRECATED(5.7, "Camera rig references now use instanced property bags.") FBooleanCameraRigParameterOverride;
@@ -237,41 +235,12 @@ private:
 	friend struct FCameraRigAssetReference;
 };
 
-/**
- * Metadata for a referenced camera rig's parameters.
- */
-USTRUCT()
-struct FCameraRigAssetReferenceParameterMetaData
-{
-	GENERATED_BODY()
-
-	/** The GUID of the parameter. */
-	UPROPERTY()
-	FGuid ParameterGuid;
-
-	/** The ID to use for overriding a blendable parameter. */
-	UPROPERTY()
-	FCameraVariableID OverrideVariableID;
-
-	/** The ID to use for overriding a data parameter. */
-	UPROPERTY()
-	FCameraContextDataID OverrideDataID;
-
-	/** Whether this parameter has an override value. */
-	UPROPERTY()
-	bool bIsOverridden = false;
-
-	/** Whether this parameter is being animated. */
-	UPROPERTY()
-	bool bIsAnimated = false;
-};
-
 /**
  * A structure holding a reference to a camera rig asset, along with the interface parameter
  * override values.
  */
 USTRUCT(BlueprintType)
-struct FCameraRigAssetReference
+struct FCameraRigAssetReference : public FBaseCameraObjectReference
 {
 	GENERATED_BODY()
 
@@ -308,61 +277,25 @@ public:
 		}
 	}
 
-	/** Gets the parameters for this camera rig, some of which containing overrides. */
-	const FInstancedPropertyBag& GetParameters() const
-	{
-		return Parameters;
-	}
-
-	/** Gets the parameters for this camera rig, some of which containing overrides. */
-	FInstancedPropertyBag& GetParameters()
-	{
-		return Parameters;
-	}
-
 	/** Applies the parameter override values to the given evaluation result. */
 	void ApplyParameterOverrides(UE::Cameras::FCameraNodeEvaluationResult& OutResult, bool bDrivenOverridesOnly);
 
 public:
 
-	// Internal API.
+	// FBaseCameraObjectReference interface.
+	GAMEPLAYCAMERAS_API virtual const UBaseCameraObject* GetCameraObject() const override;
+
+public:
 
 	bool SerializeFromMismatchedTag(FPropertyTag const& Tag, FStructuredArchive::FSlot Slot);
 	void PostSerialize(const FArchive& Ar);
 
-	GAMEPLAYCAMERAS_API bool IsParameterOverridden(const FGuid& PropertyID) const;
-	GAMEPLAYCAMERAS_API void SetParameterOverridden(const FGuid& PropertyID, bool bIsOverridden);
-
-	GAMEPLAYCAMERAS_API bool IsParameterAnimated(const FGuid& PropertyID) const;
-	GAMEPLAYCAMERAS_API void SetParameterAnimated(const FGuid& PropertyID, bool bIsAnimated);
-
-	GAMEPLAYCAMERAS_API bool NeedsRebuildParameters() const;
-	GAMEPLAYCAMERAS_API bool RebuildParametersIfNeeded();
-	GAMEPLAYCAMERAS_API void RebuildParameters();
-
-	GAMEPLAYCAMERAS_API void GetCustomCameraNodeParameters(FCustomCameraNodeParameterInfos& OutParameterInfos);
-
-private:
-
-	const FCameraRigAssetReferenceParameterMetaData* FindMetaData(const FGuid& PropertyID) const;
-	FCameraRigAssetReferenceParameterMetaData& FindOrAddMetaData(const FGuid& PropertyID);
-	
-	void GenerateOverriddenParameterGuidArray(TArray<FGuid>& OutOverriddenIDs) const;
-
 private:
 
 	/** The referenced camera rig. */
 	UPROPERTY(EditAnywhere, Category="")
 	TObjectPtr<UCameraRigAsset> CameraRig;
 
-	/** The camera rig's parameters. */
-	UPROPERTY(EditAnywhere, Category="", meta=(FixedLayout))
-	FInstancedPropertyBag Parameters;
-
-	/** Metadata for the parameters. */
-	UPROPERTY()
-	TArray<FCameraRigAssetReferenceParameterMetaData> ParameterMetaData;
-
 
 	// Deprecated
 
diff --git a/Engine/Plugins/Cameras/GameplayCameras/Source/GameplayCameras/Public/Core/CameraRigTransition.h b/Engine/Plugins/Cameras/GameplayCameras/Source/GameplayCameras/Public/Core/CameraRigTransition.h
--- a/Engine/Plugins/Cameras/GameplayCameras/Source/GameplayCameras/Public/Core/CameraRigTransition.h
+++ b/Engine/Plugins/Cameras/GameplayCameras/Source/GameplayCameras/Public/Core/CameraRigTransition.h
@@ -15,7 +15,7 @@ class UCameraRigAsset;
 
 namespace UE::Cameras
 {
-	struct FCameraRigBuildContext;
+	struct FCameraObjectBuildContext;
 }
 
 /**
@@ -49,13 +49,13 @@ class GAMEPLAYCAMERAS_API UCameraRigTransitionCondition
 
 public:
 
-	using FCameraRigBuildContext = UE::Cameras::FCameraRigBuildContext;
+	using FCameraObjectBuildContext = UE::Cameras::FCameraObjectBuildContext;
 
 	/** Evaluates whether this transition should be used. */
 	bool TransitionMatches(const FCameraRigTransitionConditionMatchParams& Params) const;
 
 	/** Build process callback for this transition. */
-	void Build(FCameraRigBuildContext& BuildContext);
+	void Build(FCameraObjectBuildContext& BuildContext);
 
 protected:
 
@@ -63,7 +63,7 @@ protected:
 	virtual bool OnTransitionMatches(const FCameraRigTransitionConditionMatchParams& Params) const { return false; }
 
 	/** Build process callback for this transition. */
-	virtual void OnBuild(FCameraRigBuildContext& BuildContext) {}
+	virtual void OnBuild(FCameraObjectBuildContext& BuildContext) {}
 
 protected:
 
@@ -167,13 +167,13 @@ public:
 
 public:
 
-	using FCameraRigBuildContext = UE::Cameras::FCameraRigBuildContext;
+	using FCameraObjectBuildContext = UE::Cameras::FCameraObjectBuildContext;
 
 	/** Returns whether all transition condition matches the given parameters. */
 	bool AllConditionsMatch(const FCameraRigTransitionConditionMatchParams& Params) const;
 
 	/** Build process callback for this transition. */
-	void Build(FCameraRigBuildContext& BuildContext);
+	void Build(FCameraObjectBuildContext& BuildContext);
 
 protected:
 
diff --git a/Engine/Plugins/Cameras/GameplayCameras/Source/GameplayCameras/Public/Core/ICustomCameraNodeParameterProvider.h b/Engine/Plugins/Cameras/GameplayCameras/Source/GameplayCameras/Public/Core/ICustomCameraNodeParameterProvider.h
--- a/Engine/Plugins/Cameras/GameplayCameras/Source/GameplayCameras/Public/Core/ICustomCameraNodeParameterProvider.h
+++ b/Engine/Plugins/Cameras/GameplayCameras/Source/GameplayCameras/Public/Core/ICustomCameraNodeParameterProvider.h
@@ -13,8 +13,9 @@ class UCameraNode;
 
 namespace UE::Cameras
 {
-	class FCameraRigAssetBuilder;
-	class FCameraRigParameterBuilder;
+	class FCameraNodeHierarchyBuilder;
+	class FCameraObjectInterfaceBuilder;
+	class FCameraObjectInterfaceParameterBuilder;
 	namespace Internal { struct FInterfaceParameterBindingBuilder; }
 }
 
@@ -171,8 +172,9 @@ private:
 	TArray<FBlendableParameterInfo> BlendableParameters;
 	TArray<FDataParameterInfo> DataParameters;
 
-	friend class UE::Cameras::FCameraRigAssetBuilder;
-	friend class UE::Cameras::FCameraRigParameterBuilder;
+	friend class UE::Cameras::FCameraNodeHierarchyBuilder;
+	friend class UE::Cameras::FCameraObjectInterfaceBuilder;
+	friend class UE::Cameras::FCameraObjectInterfaceParameterBuilder;
 	friend struct UE::Cameras::Internal::FInterfaceParameterBindingBuilder;
 };
 
diff --git a/Engine/Plugins/Cameras/GameplayCameras/Source/GameplayCameras/Public/Nodes/Common/CameraRigCameraNode.h b/Engine/Plugins/Cameras/GameplayCameras/Source/GameplayCameras/Public/Nodes/Common/CameraRigCameraNode.h
--- a/Engine/Plugins/Cameras/GameplayCameras/Source/GameplayCameras/Public/Nodes/Common/CameraRigCameraNode.h
+++ b/Engine/Plugins/Cameras/GameplayCameras/Source/GameplayCameras/Public/Nodes/Common/CameraRigCameraNode.h
@@ -23,7 +23,7 @@ protected:
 
 	// UCameraNode interface.
 	virtual void OnPreBuild(FCameraBuildLog& BuildLog) override;
-	virtual void OnBuild(FCameraRigBuildContext& BuildContext) override;
+	virtual void OnBuild(FCameraObjectBuildContext& BuildContext) override;
 	virtual FCameraNodeEvaluatorPtr OnBuildEvaluator(FCameraNodeEvaluatorBuilder& Builder) const override;
 
 	// ICustomCameraNodeParameterProvider interface.
diff --git a/Engine/Plugins/Cameras/GameplayCameras/Source/GameplayCameras/Public/Nodes/Input/AutoRotateInput2DCameraNode.h b/Engine/Plugins/Cameras/GameplayCameras/Source/GameplayCameras/Public/Nodes/Input/AutoRotateInput2DCameraNode.h
--- a/Engine/Plugins/Cameras/GameplayCameras/Source/GameplayCameras/Public/Nodes/Input/AutoRotateInput2DCameraNode.h
+++ b/Engine/Plugins/Cameras/GameplayCameras/Source/GameplayCameras/Public/Nodes/Input/AutoRotateInput2DCameraNode.h
@@ -75,7 +75,7 @@ protected:
 
 	// UCameraNode interface.
 	virtual FCameraNodeChildrenView OnGetChildren() override;
-	virtual void OnBuild(FCameraRigBuildContext& BuildContext) override;
+	virtual void OnBuild(FCameraObjectBuildContext& BuildContext) override;
 	virtual FCameraNodeEvaluatorPtr OnBuildEvaluator(FCameraNodeEvaluatorBuilder& Builder) const override;
 };
 
diff --git a/Engine/Plugins/Cameras/GameplayCameras/Source/GameplayCamerasEditor/Private/Customizations/CameraRigAssetReferenceDetailsCustomization.cpp b/Engine/Plugins/Cameras/GameplayCameras/Source/GameplayCamerasEditor/Private/Customizations/CameraRigAssetReferenceDetailsCustomization.cpp
--- a/Engine/Plugins/Cameras/GameplayCameras/Source/GameplayCamerasEditor/Private/Customizations/CameraRigAssetReferenceDetailsCustomization.cpp
+++ b/Engine/Plugins/Cameras/GameplayCameras/Source/GameplayCamerasEditor/Private/Customizations/CameraRigAssetReferenceDetailsCustomization.cpp
@@ -4,6 +4,7 @@
 
 #include "Core/CameraRigAsset.h"
 #include "Core/CameraRigAssetReference.h"
+#include "Customizations/CameraObjectInterfaceParameterOverrideDataDetails.h"
 #include "GameplayCamerasDelegates.h"
 #include "IDetailChildrenBuilder.h"
 #include "IPropertyUtilities.h"
@@ -15,83 +16,7 @@
 namespace UE::Cameras
 {
 
-class FCameraRigAssetParameterOverrideDataDetails : public FPropertyBagInstanceDataDetails
-{
-public:
-
-	FCameraRigAssetParameterOverrideDataDetails(
-			TSharedPtr<IPropertyHandle> InStructPropertyHandle,
-			TSharedPtr<IPropertyHandle> InParametersPropertyHandle,
-			TSharedPtr<IPropertyUtilities>& InPropertyUtilities)
-		: FPropertyBagInstanceDataDetails(InParametersPropertyHandle, InPropertyUtilities, true)
-		, StructPropertyHandle(InStructPropertyHandle)
-	{
-	}
-
-protected:
-
-	struct FCameraRigAssetReferenceOverrideProvider : public IPropertyBagOverrideProvider
-	{
-		FCameraRigAssetReferenceOverrideProvider(FCameraRigAssetReference& InCameraRigAssetReference)
-			: CameraRigAssetReference(InCameraRigAssetReference)
-		{
-		}
-		
-		virtual bool IsPropertyOverridden(const FGuid PropertyID) const override
-		{
-			return CameraRigAssetReference.IsParameterOverridden(PropertyID);
-		}
-		
-		virtual void SetPropertyOverride(const FGuid PropertyID, const bool bIsOverridden) const override
-		{
-			CameraRigAssetReference.SetParameterOverridden(PropertyID, bIsOverridden);
-		}
-
-	private:
-		FCameraRigAssetReference& CameraRigAssetReference;
-	};
-
-	virtual bool HasPropertyOverrides() const override
-	{
-		return true;
-	}
-
-	virtual void PreChangeOverrides() override
-	{
-		StructPropertyHandle->NotifyPreChange();
-	}
-
-	virtual void PostChangeOverrides() override
-	{
-		StructPropertyHandle->NotifyPostChange(EPropertyChangeType::ValueSet);
-		StructPropertyHandle->NotifyFinishedChangingProperties();
-	}
-
-	virtual void EnumeratePropertyBags(TSharedPtr<IPropertyHandle> PropertyBagHandle, const EnumeratePropertyBagFuncRef& Func) const override
-	{
-		StructPropertyHandle->EnumerateRawData([Func](void* RawData, const int32 DataIndex, const int32 NumDatas)
-		{
-			if (FCameraRigAssetReference* CameraRigAssetReference = static_cast<FCameraRigAssetReference*>(RawData))
-			{
-				if (const UCameraRigAsset* CameraRigAsset = CameraRigAssetReference->GetCameraRig())
-				{
-					const FInstancedPropertyBag& DefaultParameters = CameraRigAsset->GetDefaultParameters();
-					FInstancedPropertyBag& Parameters = CameraRigAssetReference->GetParameters();
-					FCameraRigAssetReferenceOverrideProvider OverrideProvider(*CameraRigAssetReference);
-					if (!Func(DefaultParameters, Parameters, OverrideProvider))
-					{
-						return false;
-					}
-				}
-			}
-			return true;
-		});
-	}
-
-private:
-	
-	TSharedPtr<IPropertyHandle> StructPropertyHandle;
-};
+using FCameraRigAssetParameterOverrideDataDetails = TCameraObjectInterfaceParameterOverrideDataDetails<FCameraRigAssetReference>;
 
 TSharedRef<IPropertyTypeCustomization> FCameraRigAssetReferenceDetailsCustomization::MakeInstance()
 {
diff --git a/Engine/Plugins/Cameras/GameplayCameras/Source/GameplayCamerasEditor/Private/Editors/CameraNodeGraphDragDropOp.cpp b/Engine/Plugins/Cameras/GameplayCameras/Source/GameplayCamerasEditor/Private/Editors/CameraNodeGraphDragDropOp.cpp
--- a/Engine/Plugins/Cameras/GameplayCameras/Source/GameplayCamerasEditor/Private/Editors/CameraNodeGraphDragDropOp.cpp
+++ b/Engine/Plugins/Cameras/GameplayCameras/Source/GameplayCamerasEditor/Private/Editors/CameraNodeGraphDragDropOp.cpp
@@ -9,7 +9,7 @@
 
 #define LOCTEXT_NAMESPACE "CameraNodeGraphDragDropOp"
 
-TSharedRef<FCameraNodeGraphInterfaceParameterDragDropOp> FCameraNodeGraphInterfaceParameterDragDropOp::New(UCameraRigInterfaceParameterBase* InInterfaceParameter)
+TSharedRef<FCameraNodeGraphInterfaceParameterDragDropOp> FCameraNodeGraphInterfaceParameterDragDropOp::New(UCameraObjectInterfaceParameterBase* InInterfaceParameter)
 {
 	TSharedRef<FCameraNodeGraphInterfaceParameterDragDropOp> Operation = MakeShared<FCameraNodeGraphInterfaceParameterDragDropOp>();
 	Operation->InterfaceParameter = InInterfaceParameter;
diff --git a/Engine/Plugins/Cameras/GameplayCameras/Source/GameplayCamerasEditor/Private/Editors/CameraNodeGraphSchema.cpp b/Engine/Plugins/Cameras/GameplayCameras/Source/GameplayCamerasEditor/Private/Editors/CameraNodeGraphSchema.cpp
--- a/Engine/Plugins/Cameras/GameplayCameras/Source/GameplayCamerasEditor/Private/Editors/CameraNodeGraphSchema.cpp
+++ b/Engine/Plugins/Cameras/GameplayCameras/Source/GameplayCamerasEditor/Private/Editors/CameraNodeGraphSchema.cpp
@@ -2,25 +2,19 @@
 
 #include "Editors/CameraNodeGraphSchema.h"
 
-#include "Core/BlendCameraNode.h"
+#include "Core/BaseCameraObject.h"
 #include "Core/CameraNode.h"
-#include "Core/CameraNodeHierarchy.h"
 #include "Core/CameraParameters.h"
-#include "Core/CameraRigAsset.h"
 #include "Core/CameraVariableReferences.h"
 #include "Core/ICustomCameraNodeParameterProvider.h"
 #include "EdGraph/EdGraphPin.h"
 #include "EdGraphSchema_K2.h"
 #include "Editors/CameraNodeGraphNode.h"
-#include "Editors/CameraRigInterfaceParameterGraphNode.h"
+#include "Editors/CameraObjectInterfaceParameterGraphNode.h"
 #include "Editors/ObjectTreeGraph.h"
 #include "Editors/ObjectTreeGraphConfig.h"
 #include "Editors/ObjectTreeGraphNode.h"
-#include "Framework/Notifications/NotificationManager.h"
 #include "GameplayCamerasEditorSettings.h"
-#include "Nodes/Common/ArrayCameraNode.h"
-#include "Widgets/Notifications/SNotificationList.h"
-
 #include "ScopedTransaction.h"
 
 #include UE_INLINE_GENERATED_CPP_BY_NAME(CameraNodeGraphSchema)
@@ -37,41 +31,21 @@ UCameraNodeGraphSchema::UCameraNodeGraphSchema(const FObjectInitializer& ObjInit
 	PinColors.Initialize();
 }
 
-FObjectTreeGraphConfig UCameraNodeGraphSchema::BuildGraphConfig() const
+void UCameraNodeGraphSchema::BuildBaseGraphConfig(FObjectTreeGraphConfig& OutGraphConfig) const
 {
 	const UGameplayCamerasEditorSettings* Settings = GetDefault<UGameplayCamerasEditorSettings>();
 
-	FObjectTreeGraphConfig GraphConfig;
-	GraphConfig.GraphName = UCameraRigAsset::NodeTreeGraphName;
-	GraphConfig.ConnectableObjectClasses.Add(UCameraRigAsset::StaticClass());
-	GraphConfig.ConnectableObjectClasses.Add(UCameraNode::StaticClass());
-	GraphConfig.NonConnectableObjectClasses.Add(UBlendCameraNode::StaticClass());
-	GraphConfig.GraphDisplayInfo.PlainName = LOCTEXT("NodeGraphPlainName", "CameraNodes");
-	GraphConfig.GraphDisplayInfo.DisplayName = LOCTEXT("NodeGraphDisplayName", "Camera Nodes");
-	GraphConfig.DefaultSelfPinName = NAME_None;
-	GraphConfig.ObjectClassConfigs.Emplace(UCameraRigAsset::StaticClass())
-		.OnlyAsRoot()
-		.HasSelfPin(false)
-		.NodeTitleUsesObjectName(true)
-		.NodeTitleColor(Settings->CameraRigAssetTitleColor);
-	GraphConfig.ObjectClassConfigs.Emplace(UCameraNode::StaticClass())
+	OutGraphConfig.GraphDisplayInfo.PlainName = LOCTEXT("NodeGraphPlainName", "CameraNodes");
+	OutGraphConfig.GraphDisplayInfo.DisplayName = LOCTEXT("NodeGraphDisplayName", "Camera Nodes");
+	OutGraphConfig.DefaultSelfPinName = NAME_None;
+	OutGraphConfig.ObjectClassConfigs.Emplace(UCameraNode::StaticClass())
 		.StripDisplayNameSuffix(TEXT("Camera Node"))
 		.CreateCategoryMetaData(TEXT("CameraNodeCategories"))
 		.NodeTitleColor(Settings->CameraNodeTitleColor)
 		.GraphNodeClass(UCameraNodeGraphNode::StaticClass());
-	GraphConfig.ObjectClassConfigs.Emplace(UArrayCameraNode::StaticClass())
-		.OnSetupNewObject(FOnSetupNewObject::CreateLambda([](UObject* NewObject)
-				{
-					// Add two new pins by default.
-					UArrayCameraNode* ArrayNode = CastChecked<UArrayCameraNode>(NewObject);
-					ArrayNode->Children.AddDefaulted();
-					ArrayNode->Children.AddDefaulted();
-				}));
 
 	// Note that we don't add the interface parameter types to the config, we will manage
 	// them ourselves.
-
-	return GraphConfig;
 }
 
 void UCameraNodeGraphSchema::CollectAllObjects(UObjectTreeGraph* InGraph, TSet<UObject*>& OutAllObjects) const
@@ -80,34 +54,6 @@ void UCameraNodeGraphSchema::CollectAllObjects(UObjectTreeGraph* InGraph, TSet<U
 
 	// Only get the graph objects from the root interface.
 	CollectAllConnectableObjectsFromRootInterface(InGraph, OutAllObjects, false);
-
-	// See if we are missing objects from AllNodeTreeObjects... if so, add them and notify the user.
-	UCameraRigAsset* CameraRig = Cast<UCameraRigAsset>(InGraph->GetRootObject());
-	if (CameraRig)
-	{
-		FCameraNodeHierarchy Hierarchy(CameraRig);
-
-		TSet<UObject*> AllNodeTreeObjects;
-		((IObjectTreeGraphRootObject*)CameraRig)->GetConnectableObjects(UCameraRigAsset::NodeTreeGraphName, AllNodeTreeObjects);
-
-		TSet<UObject*> MissingNodeTreeObjects;
-		if (Hierarchy.FindMissingConnectableObjects(AllNodeTreeObjects, MissingNodeTreeObjects))
-		{
-			FNotificationInfo NotificationInfo(
-					FText::Format(
-						LOCTEXT("AllNodeTreeObjectsMismatch", 
-							"Found {0} nodes missing from the internal list. Please re-save the asset."),
-						MissingNodeTreeObjects.Num()));
-			NotificationInfo.ExpireDuration = 4.0f;
-			FSlateNotificationManager::Get().AddNotification(NotificationInfo);
-
-			for (UObject* MissingObject : MissingNodeTreeObjects)
-			{
-				((IObjectTreeGraphRootObject*)CameraRig)->AddConnectableObject(UCameraRigAsset::NodeTreeGraphName, MissingObject);
-				OutAllObjects.Add(MissingObject);
-			}
-		}
-	}
 }
 
 void UCameraNodeGraphSchema::OnCreateAllNodes(UObjectTreeGraph* InGraph, const FCreatedNodes& InCreatedNodes) const
@@ -120,8 +66,8 @@ void UCameraNodeGraphSchema::OnCreateAllNodes(UObjectTreeGraph* InGraph, const F
 		return;
 	}
 
-	UCameraRigAsset* CameraRig = Cast<UCameraRigAsset>(InGraph->GetRootObject());
-	if (!ensure(CameraRig))
+	UBaseCameraObject* CameraObject = Cast<UBaseCameraObject>(InGraph->GetRootObject());
+	if (!ensure(CameraObject))
 	{
 		return;
 	}
@@ -129,18 +75,18 @@ void UCameraNodeGraphSchema::OnCreateAllNodes(UObjectTreeGraph* InGraph, const F
 	// Add nodes for all interface parameters that have been added to the graph.
 	// These nodes are UObjectTreeGraphNode instances, but they are "unmanaged" by the UObjectTreeGraphSchema since
 	// their object types are not in the ConnectableObjectClasses. Instead, we manage them ourselves in this schema.
-	TArray<TObjectPtr<UCameraRigInterfaceParameterBase>> InterfaceParameters;
-	InterfaceParameters.Append(CameraRig->Interface.BlendableParameters);
-	InterfaceParameters.Append(CameraRig->Interface.DataParameters);
+	TArray<TObjectPtr<UCameraObjectInterfaceParameterBase>> InterfaceParameters;
+	InterfaceParameters.Append(CameraObject->Interface.BlendableParameters);
+	InterfaceParameters.Append(CameraObject->Interface.DataParameters);
 
-	for (UCameraRigInterfaceParameterBase* InterfaceParameter : InterfaceParameters)
+	for (UCameraObjectInterfaceParameterBase* InterfaceParameter : InterfaceParameters)
 	{
 		if (!InterfaceParameter->bHasGraphNode)
 		{
 			continue;
 		}
 		
-		UCameraRigInterfaceParameterGraphNode* InterfaceParameterNode = CreateInterfaceParameterNode(InGraph, InterfaceParameter);
+		UCameraObjectInterfaceParameterGraphNode* InterfaceParameterNode = CreateInterfaceParameterNode(InGraph, InterfaceParameter);
 		UObjectTreeGraphNode* CameraNodeNode = InCreatedNodes.CreatedNodes.FindRef(InterfaceParameter->Target);
 		if (CameraNodeNode)
 		{
@@ -156,11 +102,11 @@ void UCameraNodeGraphSchema::OnCreateAllNodes(UObjectTreeGraph* InGraph, const F
 			else
 			{
 				FName ErrorPinCategory = NAME_None;
-				if (InterfaceParameter->IsA<UCameraRigBlendableParameter>())
+				if (InterfaceParameter->IsA<UCameraObjectInterfaceBlendableParameter>())
 				{
 					ErrorPinCategory = PC_CameraParameter;
 				}
-				else if (InterfaceParameter->IsA<UCameraRigDataParameter>())
+				else if (InterfaceParameter->IsA<UCameraObjectInterfaceDataParameter>())
 				{
 					ErrorPinCategory = PC_CameraContextData;
 				}
@@ -174,10 +120,10 @@ void UCameraNodeGraphSchema::OnCreateAllNodes(UObjectTreeGraph* InGraph, const F
 	}
 }
 
-UCameraRigInterfaceParameterGraphNode* UCameraNodeGraphSchema::CreateInterfaceParameterNode(UEdGraph* InGraph, UCameraRigInterfaceParameterBase* InterfaceParameter) const
+UCameraObjectInterfaceParameterGraphNode* UCameraNodeGraphSchema::CreateInterfaceParameterNode(UEdGraph* InGraph, UCameraObjectInterfaceParameterBase* InterfaceParameter) const
 {
-	FGraphNodeCreator<UCameraRigInterfaceParameterGraphNode> GraphNodeCreator(*InGraph);
-	UCameraRigInterfaceParameterGraphNode* InterfaceParameterNode = GraphNodeCreator.CreateNode(false);
+	FGraphNodeCreator<UCameraObjectInterfaceParameterGraphNode> GraphNodeCreator(*InGraph);
+	UCameraObjectInterfaceParameterGraphNode* InterfaceParameterNode = GraphNodeCreator.CreateNode(false);
 	InterfaceParameterNode->Initialize(InterfaceParameter);
 	GraphNodeCreator.Finalize();
 	return InterfaceParameterNode;
@@ -216,8 +162,8 @@ void UCameraNodeGraphSchema::GetGraphContextActions(FGraphContextMenuBuilder& Co
 			TSharedRef<FCameraNodeGraphSchemaAction_NewInterfaceParameterNode> Action = 
 				MakeShared<FCameraNodeGraphSchemaAction_NewInterfaceParameterNode>(
 						FText::GetEmpty(),
-						LOCTEXT("NewInterfaceParameterAction", "Camera Rig Parameter"),
-						LOCTEXT("NewInterfaceParameterActionToolTip", "Exposes this parameter on the camera rig"));
+						LOCTEXT("NewInterfaceParameterAction", "Camera Interface Parameter"),
+						LOCTEXT("NewInterfaceParameterActionToolTip", "Exposes this parameter on the camera object"));
 
 			if (DraggedPin->PinType.PinCategory == PC_CameraParameter ||
 				DraggedPin->PinType.PinCategory == PC_CameraVariableReference)
@@ -256,8 +202,8 @@ void UCameraNodeGraphSchema::GetGraphContextActions(FGraphContextMenuBuilder& Co
 					return;
 				}
 
-				FCameraRigParameterDefinition NewParameterDefinition;
-				NewParameterDefinition.ParameterType = ECameraRigInterfaceParameterType::Blendable;
+				FCameraObjectInterfaceParameterDefinition NewParameterDefinition;
+				NewParameterDefinition.ParameterType = ECameraObjectInterfaceParameterType::Blendable;
 				NewParameterDefinition.VariableType = VariableType;
 				NewParameterDefinition.BlendableStructType = BlendableStructType;
 				Action->ParameterDefinition = NewParameterDefinition;
@@ -325,8 +271,8 @@ void UCameraNodeGraphSchema::GetGraphContextActions(FGraphContextMenuBuilder& Co
 					return;
 				}
 
-				FCameraRigParameterDefinition NewParameterDefinition;
-				NewParameterDefinition.ParameterType = ECameraRigInterfaceParameterType::Data;
+				FCameraObjectInterfaceParameterDefinition NewParameterDefinition;
+				NewParameterDefinition.ParameterType = ECameraObjectInterfaceParameterType::Data;
 				NewParameterDefinition.DataType = DataType;
 				NewParameterDefinition.DataContainerType = DataContainerType;
 				NewParameterDefinition.DataTypeObject = DataTypeObject;
@@ -351,16 +297,16 @@ const FPinConnectionResponse UCameraNodeGraphSchema::CanCreateConnection(const U
 			B->PinType.PinCategory == PC_Self &&
 			!A->bOrphanedPin)
 	{
-		UCameraRigInterfaceParameterGraphNode* NodeB = Cast<UCameraRigInterfaceParameterGraphNode>(B->GetOwningNode());
+		UCameraObjectInterfaceParameterGraphNode* NodeB = Cast<UCameraObjectInterfaceParameterGraphNode>(B->GetOwningNode());
 		if (NodeB)
 		{
-			UCameraRigBlendableParameter* BlendableParameter = NodeB->CastObject<UCameraRigBlendableParameter>();
+			UCameraObjectInterfaceBlendableParameter* BlendableParameter = NodeB->CastObject<UCameraObjectInterfaceBlendableParameter>();
 			if (BlendableParameter && 
 					A->PinType.PinSubCategory == UEnum::GetValueAsName(BlendableParameter->ParameterType))
 			{
 				return FPinConnectionResponse(CONNECT_RESPONSE_BREAK_OTHERS_AB, TEXT("Compatible pin types"));
 			}
-			UCameraRigDataParameter* DataParameter = NodeB->CastObject<UCameraRigDataParameter>();
+			UCameraObjectInterfaceDataParameter* DataParameter = NodeB->CastObject<UCameraObjectInterfaceDataParameter>();
 			if (DataParameter && 
 					A->PinType.PinSubCategory == UEnum::GetValueAsName(DataParameter->DataType) &&
 					A->PinType.PinSubCategoryObject == DataParameter->DataTypeObject)
@@ -375,16 +321,16 @@ const FPinConnectionResponse UCameraNodeGraphSchema::CanCreateConnection(const U
 			 B->PinType.PinCategory == PC_CameraContextData) &&
 			!B->bOrphanedPin)
 	{
-		UCameraRigInterfaceParameterGraphNode* NodeA = Cast<UCameraRigInterfaceParameterGraphNode>(A->GetOwningNode());
+		UCameraObjectInterfaceParameterGraphNode* NodeA = Cast<UCameraObjectInterfaceParameterGraphNode>(A->GetOwningNode());
 		if (NodeA)
 		{
-			UCameraRigBlendableParameter* BlendableParameter = NodeA->CastObject<UCameraRigBlendableParameter>();
+			UCameraObjectInterfaceBlendableParameter* BlendableParameter = NodeA->CastObject<UCameraObjectInterfaceBlendableParameter>();
 			if (BlendableParameter && 
 					B->PinType.PinSubCategory == UEnum::GetValueAsName(BlendableParameter->ParameterType))
 			{
 				return FPinConnectionResponse(CONNECT_RESPONSE_BREAK_OTHERS_AB, TEXT("Compatible pin types"));
 			}
-			UCameraRigDataParameter* DataParameter = NodeA->CastObject<UCameraRigDataParameter>();
+			UCameraObjectInterfaceDataParameter* DataParameter = NodeA->CastObject<UCameraObjectInterfaceDataParameter>();
 			if (DataParameter && 
 					B->PinType.PinSubCategory == UEnum::GetValueAsName(DataParameter->DataType) &&
 					B->PinType.PinSubCategoryObject == DataParameter->DataTypeObject)
@@ -402,7 +348,7 @@ bool UCameraNodeGraphSchema::OnTryCreateCustomConnection(UEdGraphPin* A, UEdGrap
 	// See if we are in the situation of connecting an interface parameter to a camera node property.
 	UEdGraphPin* TargetPin = nullptr;
 	UObjectTreeGraphNode* TargetNode = nullptr;
-	UCameraRigInterfaceParameterGraphNode* InterfaceParameterNode = nullptr;
+	UCameraObjectInterfaceParameterGraphNode* InterfaceParameterNode = nullptr;
 
 	if ((A->PinType.PinCategory == PC_CameraParameter || 
 				A->PinType.PinCategory == PC_CameraVariableReference ||
@@ -411,14 +357,14 @@ bool UCameraNodeGraphSchema::OnTryCreateCustomConnection(UEdGraphPin* A, UEdGrap
 	{
 		TargetPin = A;
 		TargetNode = Cast<UObjectTreeGraphNode>(A->GetOwningNode());
-		InterfaceParameterNode = Cast<UCameraRigInterfaceParameterGraphNode>(B->GetOwningNode());
+		InterfaceParameterNode = Cast<UCameraObjectInterfaceParameterGraphNode>(B->GetOwningNode());
 	}
 	else if (A->PinType.PinCategory == PC_Self && 
 			(B->PinType.PinCategory == PC_CameraParameter || 
 			 B->PinType.PinCategory == PC_CameraVariableReference ||
 			 B->PinType.PinCategory == PC_CameraContextData))
 	{
-		InterfaceParameterNode = Cast<UCameraRigInterfaceParameterGraphNode>(A->GetOwningNode());
+		InterfaceParameterNode = Cast<UCameraObjectInterfaceParameterGraphNode>(A->GetOwningNode());
 		TargetNode = Cast<UObjectTreeGraphNode>(B->GetOwningNode());
 		TargetPin = B;
 	}
@@ -426,7 +372,7 @@ bool UCameraNodeGraphSchema::OnTryCreateCustomConnection(UEdGraphPin* A, UEdGrap
 	if (TargetNode && TargetPin && InterfaceParameterNode)
 	{
 		UCameraNode* Target = TargetNode->CastObject<UCameraNode>();
-		UCameraRigInterfaceParameterBase* InterfaceParameter = InterfaceParameterNode->GetInterfaceParameter();
+		UCameraObjectInterfaceParameterBase* InterfaceParameter = InterfaceParameterNode->GetInterfaceParameter();
 		if (Target && InterfaceParameter)
 		{
 			InterfaceParameter->Modify();
@@ -445,7 +391,7 @@ bool UCameraNodeGraphSchema::OnBreakCustomPinLinks(UEdGraphPin& TargetPin) const
 {
 	// See if we are in the situation of an interface parameter node being disconnected from 
 	// a camera node property pin.
-	UCameraRigInterfaceParameterGraphNode* InterfaceParameterNode = nullptr;
+	UCameraObjectInterfaceParameterGraphNode* InterfaceParameterNode = nullptr;
 
 	if (TargetPin.PinType.PinCategory == PC_CameraParameter ||
 			TargetPin.PinType.PinCategory == PC_CameraVariableReference ||
@@ -453,17 +399,17 @@ bool UCameraNodeGraphSchema::OnBreakCustomPinLinks(UEdGraphPin& TargetPin) const
 	{
 		if (TargetPin.LinkedTo.Num() > 0)
 		{
-			InterfaceParameterNode = Cast<UCameraRigInterfaceParameterGraphNode>(TargetPin.LinkedTo[0]->GetOwningNode());
+			InterfaceParameterNode = Cast<UCameraObjectInterfaceParameterGraphNode>(TargetPin.LinkedTo[0]->GetOwningNode());
 		}
 	}
 	else if (TargetPin.PinType.PinCategory == PC_Self)
 	{
-		InterfaceParameterNode = Cast<UCameraRigInterfaceParameterGraphNode>(TargetPin.GetOwningNode());
+		InterfaceParameterNode = Cast<UCameraObjectInterfaceParameterGraphNode>(TargetPin.GetOwningNode());
 	}
 
 	if (InterfaceParameterNode)
 	{
-		UCameraRigInterfaceParameterBase* InterfaceParameter = InterfaceParameterNode->GetInterfaceParameter();
+		UCameraObjectInterfaceParameterBase* InterfaceParameter = InterfaceParameterNode->GetInterfaceParameter();
 		if (InterfaceParameter)
 		{
 			InterfaceParameter->Modify();
@@ -482,19 +428,19 @@ bool UCameraNodeGraphSchema::OnBreakSingleCustomPinLink(UEdGraphPin* SourcePin,
 {
 	// See if we are in the situation of an interface parameter node being disconnected from 
 	// a camera node property pin.
-	UCameraRigInterfaceParameterGraphNode* InterfaceParameterNode = nullptr;
+	UCameraObjectInterfaceParameterGraphNode* InterfaceParameterNode = nullptr;
 	if (SourcePin->PinType.PinCategory == PC_Self)
 	{
-		InterfaceParameterNode = Cast<UCameraRigInterfaceParameterGraphNode>(SourcePin->GetOwningNode());
+		InterfaceParameterNode = Cast<UCameraObjectInterfaceParameterGraphNode>(SourcePin->GetOwningNode());
 	}
 	else if (TargetPin->PinType.PinCategory == PC_Self)
 	{
-		InterfaceParameterNode = Cast<UCameraRigInterfaceParameterGraphNode>(TargetPin->GetOwningNode());
+		InterfaceParameterNode = Cast<UCameraObjectInterfaceParameterGraphNode>(TargetPin->GetOwningNode());
 	}
 
 	if (InterfaceParameterNode)
 	{
-		UCameraRigInterfaceParameterBase* InterfaceParameter = InterfaceParameterNode->GetInterfaceParameter();
+		UCameraObjectInterfaceParameterBase* InterfaceParameter = InterfaceParameterNode->GetInterfaceParameter();
 		if (InterfaceParameter)
 		{
 			InterfaceParameter->Modify();
@@ -530,9 +476,9 @@ bool UCameraNodeGraphSchema::SafeDeleteNodeFromGraph(UEdGraph* Graph, UEdGraphNo
 
 	// Deleting an interface parameter node simply removes its bHasGraphNode flag.
 	// To actually delete the parameter, the user needs to remove it from the "parameters" panel.
-	if (UCameraRigInterfaceParameterGraphNode* InterfaceParameterNode = Cast<UCameraRigInterfaceParameterGraphNode>(Node))
+	if (UCameraObjectInterfaceParameterGraphNode* InterfaceParameterNode = Cast<UCameraObjectInterfaceParameterGraphNode>(Node))
 	{
-		if (UCameraRigInterfaceParameterBase* InterfaceParameter = InterfaceParameterNode->GetInterfaceParameter())
+		if (UCameraObjectInterfaceParameterBase* InterfaceParameter = InterfaceParameterNode->GetInterfaceParameter())
 		{
 			InterfaceParameter->Modify();
 			InterfaceParameter->bHasGraphNode = false;
@@ -574,8 +520,8 @@ UEdGraphNode* FCameraNodeGraphSchemaAction_NewInterfaceParameterNode::PerformAct
 		return nullptr;
 	}
 
-	UCameraRigAsset* CameraRig = Cast<UCameraRigAsset>(ObjectTreeGraph->GetRootObject());
-	if (!ensure(CameraRig))
+	UBaseCameraObject* CameraObject = Cast<UBaseCameraObject>(ObjectTreeGraph->GetRootObject());
+	if (!ensure(CameraObject))
 	{
 		return nullptr;
 	}
@@ -584,25 +530,25 @@ UEdGraphNode* FCameraNodeGraphSchemaAction_NewInterfaceParameterNode::PerformAct
 
 	const UCameraNodeGraphSchema* Schema = CastChecked<UCameraNodeGraphSchema>(ParentGraph->GetSchema());
 
-	CameraRig->Modify();
+	CameraObject->Modify();
 
 	// Create a new interface parameter and set it up based on the pin we're creating it from, if any.
-	UCameraRigInterfaceParameterBase* NewInterfaceParameter = nullptr;
-	if (ParameterDefinition.ParameterType == ECameraRigInterfaceParameterType::Blendable)
+	UCameraObjectInterfaceParameterBase* NewInterfaceParameter = nullptr;
+	if (ParameterDefinition.ParameterType == ECameraObjectInterfaceParameterType::Blendable)
 	{
-		UCameraRigBlendableParameter* NewBlendableParameter = NewObject<UCameraRigBlendableParameter>(CameraRig, NAME_None, RF_Transactional);
+		UCameraObjectInterfaceBlendableParameter* NewBlendableParameter = NewObject<UCameraObjectInterfaceBlendableParameter>(CameraObject, NAME_None, RF_Transactional);
 		NewBlendableParameter->ParameterType = ParameterDefinition.VariableType;
 		NewBlendableParameter->BlendableStructType = ParameterDefinition.BlendableStructType;
-		CameraRig->Interface.BlendableParameters.Add(NewBlendableParameter);
+		CameraObject->Interface.BlendableParameters.Add(NewBlendableParameter);
 		NewInterfaceParameter = NewBlendableParameter;
 	}
-	else if (ParameterDefinition.ParameterType == ECameraRigInterfaceParameterType::Data)
+	else if (ParameterDefinition.ParameterType == ECameraObjectInterfaceParameterType::Data)
 	{
-		UCameraRigDataParameter* NewDataParameter = NewObject<UCameraRigDataParameter>(CameraRig, NAME_None, RF_Transactional);
+		UCameraObjectInterfaceDataParameter* NewDataParameter = NewObject<UCameraObjectInterfaceDataParameter>(CameraObject, NAME_None, RF_Transactional);
 		NewDataParameter->DataType = ParameterDefinition.DataType;
 		NewDataParameter->DataContainerType = ParameterDefinition.DataContainerType;
 		NewDataParameter->DataTypeObject = ParameterDefinition.DataTypeObject;
-		CameraRig->Interface.DataParameters.Add(NewDataParameter);
+		CameraObject->Interface.DataParameters.Add(NewDataParameter);
 		NewInterfaceParameter = NewDataParameter;
 	}
 
@@ -617,7 +563,7 @@ UEdGraphNode* FCameraNodeGraphSchemaAction_NewInterfaceParameterNode::PerformAct
 
 	ObjectTreeGraph->Modify();
 
-	UCameraRigInterfaceParameterGraphNode* NewGraphNode = Schema->CreateInterfaceParameterNode(ObjectTreeGraph, NewInterfaceParameter);
+	UCameraObjectInterfaceParameterGraphNode* NewGraphNode = Schema->CreateInterfaceParameterNode(ObjectTreeGraph, NewInterfaceParameter);
 
 	NewGraphNode->NodePosX = Location.X;
 	NewGraphNode->NodePosY = Location.Y;
@@ -625,7 +571,7 @@ UEdGraphNode* FCameraNodeGraphSchemaAction_NewInterfaceParameterNode::PerformAct
 
 	NewGraphNode->AutowireNewNode(FromPin);
 
-	CameraRig->EventHandlers.Notify(&UE::Cameras::ICameraRigAssetEventHandler::OnCameraRigInterfaceChanged);
+	CameraObject->EventHandlers.Notify(&UE::Cameras::ICameraObjectEventHandler::OnCameraObjectInterfaceChanged);
 
 	return NewGraphNode;
 }
@@ -652,8 +598,8 @@ UEdGraphNode* FCameraNodeGraphSchemaAction_AddInterfaceParameterNode::PerformAct
 		return nullptr;
 	}
 
-	UCameraRigAsset* CameraRig = Cast<UCameraRigAsset>(ObjectTreeGraph->GetRootObject());
-	if (!ensure(CameraRig))
+	UBaseCameraObject* CameraObject = Cast<UBaseCameraObject>(ObjectTreeGraph->GetRootObject());
+	if (!ensure(CameraObject))
 	{
 		return nullptr;
 	}
@@ -668,7 +614,7 @@ UEdGraphNode* FCameraNodeGraphSchemaAction_AddInterfaceParameterNode::PerformAct
 
 	ParentGraph->Modify();
 	
-	UCameraRigInterfaceParameterGraphNode* NewGraphNode = Schema->CreateInterfaceParameterNode(ParentGraph, InterfaceParameter);
+	UCameraObjectInterfaceParameterGraphNode* NewGraphNode = Schema->CreateInterfaceParameterNode(ParentGraph, InterfaceParameter);
 
 	NewGraphNode->NodePosX = Location.X;
 	NewGraphNode->NodePosY = Location.Y;
@@ -676,7 +622,7 @@ UEdGraphNode* FCameraNodeGraphSchemaAction_AddInterfaceParameterNode::PerformAct
 
 	NewGraphNode->AutowireNewNode(FromPin);
 
-	CameraRig->EventHandlers.Notify(&UE::Cameras::ICameraRigAssetEventHandler::OnCameraRigInterfaceChanged);
+	CameraObject->EventHandlers.Notify(&UE::Cameras::ICameraObjectEventHandler::OnCameraObjectInterfaceChanged);
 
 	return NewGraphNode;
 }
diff --git a/Engine/Plugins/Cameras/GameplayCameras/Source/GameplayCamerasEditor/Private/Editors/ObjectTreeGraphConfig.cpp b/Engine/Plugins/Cameras/GameplayCameras/Source/GameplayCamerasEditor/Private/Editors/ObjectTreeGraphConfig.cpp
--- a/Engine/Plugins/Cameras/GameplayCameras/Source/GameplayCamerasEditor/Private/Editors/ObjectTreeGraphConfig.cpp
+++ b/Engine/Plugins/Cameras/GameplayCameras/Source/GameplayCamerasEditor/Private/Editors/ObjectTreeGraphConfig.cpp
@@ -22,8 +22,8 @@ FObjectTreeGraphClassConfig::FObjectTreeGraphClassConfig()
 
 FObjectTreeGraphClassConfig& FObjectTreeGraphClassConfig::OnlyAsRoot()
 {
-	_CanCreateNew = false;
-	_CanDelete = false;
+	CanCreateNew(false);
+	CanDelete(false);
 	return *this;
 }
 
diff --git a/Engine/Plugins/Cameras/GameplayCameras/Source/GameplayCamerasEditor/Private/Editors/SCameraRigAssetEditor.cpp b/Engine/Plugins/Cameras/GameplayCameras/Source/GameplayCamerasEditor/Private/Editors/SCameraRigAssetEditor.cpp
--- a/Engine/Plugins/Cameras/GameplayCameras/Source/GameplayCamerasEditor/Private/Editors/SCameraRigAssetEditor.cpp
+++ b/Engine/Plugins/Cameras/GameplayCameras/Source/GameplayCamerasEditor/Private/Editors/SCameraRigAssetEditor.cpp
@@ -6,8 +6,8 @@
 #include "Core/CameraRigAsset.h"
 #include "Core/CameraRigTransition.h"
 #include "EdGraph/EdGraphPin.h"
-#include "Editors/CameraNodeGraphSchema.h"
-#include "Editors/CameraRigInterfaceParameterGraphNode.h"
+#include "Editors/CameraObjectInterfaceParameterGraphNode.h"
+#include "Editors/CameraRigCameraNodeGraphSchema.h"
 #include "Editors/CameraRigTransitionGraphSchema.h"
 #include "Editors/ObjectTreeGraph.h"
 #include "Editors/ObjectTreeGraphConfig.h"
@@ -32,7 +32,7 @@ void SCameraRigAssetEditor::Construct(const FArguments& InArgs)
 
 	CreateGraphEditors();
 
-	CameraRigAsset->EventHandlers.Register(EventHandler, this);
+	CameraRigAsset->UBaseCameraObject::EventHandlers.Register(EventHandler, this);
 
 	ChildSlot
 	[
@@ -65,7 +65,7 @@ void SCameraRigAssetEditor::SetCameraRigAsset(UCameraRigAsset* InCameraRig)
 
 		SetEditorModeImpl(CurrentMode, true);
 
-		CameraRigAsset->EventHandlers.Register(EventHandler, this);
+		CameraRigAsset->UBaseCameraObject::EventHandlers.Register(EventHandler, this);
 	}
 }
 
@@ -77,8 +77,8 @@ void SCameraRigAssetEditor::CreateGraphEditors()
 
 void SCameraRigAssetEditor::CreateNodeGraphEditor()
 {
-	UClass* SchemaClass = UCameraNodeGraphSchema::StaticClass();
-	UCameraNodeGraphSchema* DefaultSchemaObject = Cast<UCameraNodeGraphSchema>(SchemaClass->GetDefaultObject());
+	UClass* SchemaClass = UCameraRigCameraNodeGraphSchema::StaticClass();
+	UCameraRigCameraNodeGraphSchema* DefaultSchemaObject = Cast<UCameraRigCameraNodeGraphSchema>(SchemaClass->GetDefaultObject());
 	FObjectTreeGraphConfig GraphConfig = DefaultSchemaObject->BuildGraphConfig();
 
 	NodeGraph = NewObject<UObjectTreeGraph>(GetTransientPackage(), NAME_None, RF_Transactional | RF_Standalone);
@@ -304,18 +304,18 @@ void SCameraRigAssetEditor::RemoveOnAnyGraphChanged(FDelegateUserObjectConst InU
 	OnAnyGraphChanged.RemoveAll(InUserObject);
 }
 
-void SCameraRigAssetEditor::OnCameraRigInterfaceChanged()
+void SCameraRigAssetEditor::OnCameraObjectInterfaceChanged()
 {
 	// List all the interface parameters that want a node.
-	TSet<UCameraRigInterfaceParameterBase*> InterfaceParametersWithNodes;
-	for (UCameraRigInterfaceParameterBase* InterfaceParameter : CameraRigAsset->Interface.BlendableParameters)
+	TSet<UCameraObjectInterfaceParameterBase*> InterfaceParametersWithNodes;
+	for (UCameraObjectInterfaceParameterBase* InterfaceParameter : CameraRigAsset->Interface.BlendableParameters)
 	{
 		if (InterfaceParameter->bHasGraphNode)
 		{
 			InterfaceParametersWithNodes.Add(InterfaceParameter);
 		}
 	}
-	for (UCameraRigInterfaceParameterBase* InterfaceParameter : CameraRigAsset->Interface.DataParameters)
+	for (UCameraObjectInterfaceParameterBase* InterfaceParameter : CameraRigAsset->Interface.DataParameters)
 	{
 		if (InterfaceParameter->bHasGraphNode)
 		{
@@ -324,12 +324,12 @@ void SCameraRigAssetEditor::OnCameraRigInterfaceChanged()
 	}
 
 	// Find all the interface parameter nodes that already exist.
-	TArray<UCameraRigInterfaceParameterGraphNode*> InterfaceParameterNodes;
+	TArray<UCameraObjectInterfaceParameterGraphNode*> InterfaceParameterNodes;
 	NodeGraph->GetNodesOfClass(InterfaceParameterNodes);
-	TMap<UCameraRigInterfaceParameterBase*, UCameraRigInterfaceParameterGraphNode*> InterfaceParameterToNodeMap;
+	TMap<UCameraObjectInterfaceParameterBase*, UCameraObjectInterfaceParameterGraphNode*> InterfaceParameterToNodeMap;
 
 	// Remove nodes that aren't needed anymore.
-	for (UCameraRigInterfaceParameterGraphNode* InterfaceParameterNode : InterfaceParameterNodes)
+	for (UCameraObjectInterfaceParameterGraphNode* InterfaceParameterNode : InterfaceParameterNodes)
 	{
 		InterfaceParameterToNodeMap.Add(InterfaceParameterNode->GetInterfaceParameter(), InterfaceParameterNode);
 		if (!InterfaceParametersWithNodes.Contains(InterfaceParameterNode->GetInterfaceParameter()))
@@ -340,8 +340,8 @@ void SCameraRigAssetEditor::OnCameraRigInterfaceChanged()
 	}
 
 	// Add nodes that are newly needed.
-	const UCameraNodeGraphSchema* Schema = CastChecked<UCameraNodeGraphSchema>(NodeGraph->GetSchema());
-	for (UCameraRigInterfaceParameterBase* InterfaceParameter : InterfaceParametersWithNodes)
+	const UCameraRigCameraNodeGraphSchema* Schema = CastChecked<UCameraRigCameraNodeGraphSchema>(NodeGraph->GetSchema());
+	for (UCameraObjectInterfaceParameterBase* InterfaceParameter : InterfaceParametersWithNodes)
 	{
 		if (!InterfaceParameterToNodeMap.Contains(InterfaceParameter))
 		{
diff --git a/Engine/Plugins/Cameras/GameplayCameras/Source/GameplayCamerasEditor/Private/Sequencer/GameplayCameraComponentTrackEditor.cpp b/Engine/Plugins/Cameras/GameplayCameras/Source/GameplayCamerasEditor/Private/Sequencer/GameplayCameraComponentTrackEditor.cpp
--- a/Engine/Plugins/Cameras/GameplayCameras/Source/GameplayCamerasEditor/Private/Sequencer/GameplayCameraComponentTrackEditor.cpp
+++ b/Engine/Plugins/Cameras/GameplayCameras/Source/GameplayCamerasEditor/Private/Sequencer/GameplayCameraComponentTrackEditor.cpp
@@ -3,8 +3,8 @@
 #include "Sequencer/GameplayCameraComponentTrackEditor.h"
 
 #include "Core/CameraAsset.h"
+#include "Core/CameraObjectInterfaceParameterDefinition.h"
 #include "Core/CameraParameters.h"
-#include "Core/CameraRigParameterDefinition.h"
 #include "EventHandlers/MovieSceneDataEventContainer.h"
 #include "Framework/MultiBox/MultiBoxBuilder.h"
 #include "GameFramework/Actor.h"
diff --git a/Engine/Plugins/Cameras/GameplayCameras/Source/GameplayCamerasEditor/Private/Sequencer/GameplayCameraComponentTrackEditor.h b/Engine/Plugins/Cameras/GameplayCameras/Source/GameplayCamerasEditor/Private/Sequencer/GameplayCameraComponentTrackEditor.h
--- a/Engine/Plugins/Cameras/GameplayCameras/Source/GameplayCamerasEditor/Private/Sequencer/GameplayCameraComponentTrackEditor.h
+++ b/Engine/Plugins/Cameras/GameplayCameras/Source/GameplayCamerasEditor/Private/Sequencer/GameplayCameraComponentTrackEditor.h
@@ -9,7 +9,7 @@
 class ISequencer;
 class UGameplayCameraComponent;
 class UMovieScenePropertyTrack;
-struct FCameraRigParameterDefinition;
+struct FCameraObjectInterfaceParameterDefinition;
 
 namespace UE::Sequencer
 {
diff --git a/Engine/Plugins/Cameras/GameplayCameras/Source/GameplayCamerasEditor/Private/Toolkits/BuildButtonToolkit.cpp b/Engine/Plugins/Cameras/GameplayCameras/Source/GameplayCamerasEditor/Private/Toolkits/BuildButtonToolkit.cpp
--- a/Engine/Plugins/Cameras/GameplayCameras/Source/GameplayCamerasEditor/Private/Toolkits/BuildButtonToolkit.cpp
+++ b/Engine/Plugins/Cameras/GameplayCameras/Source/GameplayCamerasEditor/Private/Toolkits/BuildButtonToolkit.cpp
@@ -2,7 +2,7 @@
 
 #include "Toolkits/BuildButtonToolkit.h"
 
-#include "Core/CameraBuildStatus.h"
+#include "Build/CameraBuildStatus.h"
 #include "Styles/GameplayCamerasEditorStyle.h"
 #include "Textures/SlateIcon.h"
 #include "ToolMenuEntry.h"
diff --git a/Engine/Plugins/Cameras/GameplayCameras/Source/GameplayCamerasEditor/Private/Toolkits/CameraAssetEditorToolkit.cpp b/Engine/Plugins/Cameras/GameplayCameras/Source/GameplayCamerasEditor/Private/Toolkits/CameraAssetEditorToolkit.cpp
--- a/Engine/Plugins/Cameras/GameplayCameras/Source/GameplayCamerasEditor/Private/Toolkits/CameraAssetEditorToolkit.cpp
+++ b/Engine/Plugins/Cameras/GameplayCameras/Source/GameplayCamerasEditor/Private/Toolkits/CameraAssetEditorToolkit.cpp
@@ -4,10 +4,10 @@
 
 #include "AssetRegistry/AssetRegistryModule.h"
 #include "AssetTools/CameraAssetEditor.h"
+#include "Build/CameraAssetBuilder.h"
+#include "Build/CameraBuildLog.h"
 #include "Commands/CameraAssetEditorCommands.h"
 #include "Core/CameraAsset.h"
-#include "Core/CameraAssetBuilder.h"
-#include "Core/CameraBuildLog.h"
 #include "Core/CameraDirector.h"
 #include "Core/CameraRigAsset.h"
 #include "Editors/ObjectTreeGraphConfig.h"
diff --git a/Engine/Plugins/Cameras/GameplayCameras/Source/GameplayCamerasEditor/Private/Toolkits/CameraBuildLogToolkit.cpp b/Engine/Plugins/Cameras/GameplayCameras/Source/GameplayCamerasEditor/Private/Toolkits/CameraBuildLogToolkit.cpp
--- a/Engine/Plugins/Cameras/GameplayCameras/Source/GameplayCamerasEditor/Private/Toolkits/CameraBuildLogToolkit.cpp
+++ b/Engine/Plugins/Cameras/GameplayCameras/Source/GameplayCamerasEditor/Private/Toolkits/CameraBuildLogToolkit.cpp
@@ -2,7 +2,7 @@
 
 #include "Toolkits/CameraBuildLogToolkit.h"
 
-#include "Core/CameraBuildLog.h"
+#include "Build/CameraBuildLog.h"
 #include "IMessageLogListing.h"
 #include "Logging/TokenizedMessage.h"
 #include "MessageLogModule.h"
