Cameras: extract most of the camera rig logic into a base "camera object" class so it can be shared with the upcoming camera shakes

diff --git a/Engine/Plugins/Cameras/GameplayCameras/Source/GameplayCamerasEditor/Private/Toolkits/CameraObjectInterfaceParametersToolkit.cpp b/Engine/Plugins/Cameras/GameplayCameras/Source/GameplayCamerasEditor/Private/Toolkits/CameraObjectInterfaceParametersToolkit.cpp
new file mode 100644
--- /dev/null
+++ b/Engine/Plugins/Cameras/GameplayCameras/Source/GameplayCamerasEditor/Private/Toolkits/CameraObjectInterfaceParametersToolkit.cpp
@@ -0,0 +1,1001 @@
+// Copyright Epic Games, Inc. All Rights Reserved.
+
+#include "Toolkits/CameraObjectInterfaceParametersToolkit.h"
+
+#include "Core/BaseCameraObject.h"
+#include "EdGraph/EdGraphNode.h"
+#include "Editor.h"
+#include "Editors/CameraNodeGraphDragDropOp.h"
+#include "Framework/MultiBox/MultiBoxBuilder.h"
+#include "IGameplayCamerasModule.h"
+#include "ScopedTransaction.h"
+#include "SPinTypeSelector.h"
+#include "Styles/GameplayCamerasEditorStyle.h"
+#include "Styling/SlateIconFinder.h"
+#include "ToolMenus.h"
+#include "Types/SlateEnums.h"
+#include "Widgets/Input/SButton.h"
+#include "Widgets/Input/SCheckBox.h"
+#include "Widgets/Input/SComboButton.h"
+#include "Widgets/Layout/SBox.h"
+#include "Widgets/Layout/SSplitter.h"
+#include "Widgets/SBoxPanel.h"
+#include "Widgets/SCompoundWidget.h"
+#include "Widgets/Text/SInlineEditableTextBlock.h"
+#include "Widgets/Text/SRichTextBlock.h"
+#include "Widgets/Text/STextBlock.h"
+#include "Widgets/Views/SListView.h"
+
+#define LOCTEXT_NAMESPACE "CameraObjectInterfaceParametersToolkit"
+
+namespace UE::Cameras
+{
+
+static const FName ParameterTypeColumn(TEXT("ParameterType"));
+static const FName ParameterNameColumn(TEXT("ParameterName"));
+static const FName ParameterMessageColumn(TEXT("ParameterMessage"));
+static const FName ParameterIsPreBlendedColumn(TEXT("ParameterIsPreBlended"));
+
+/**
+ * List entry for any interface parameter panel.
+ */
+template<typename ParameterType>
+class SCameraObjectInterfaceParameterTableRowBase : public SMultiColumnTableRow<TObjectPtr<ParameterType>>
+{
+public:
+
+	SLATE_BEGIN_ARGS(SCameraObjectInterfaceParameterTableRowBase<ParameterType>)
+	{}
+		SLATE_ARGUMENT(TObjectPtr<ParameterType>, Item)
+	SLATE_END_ARGS()
+
+	typedef SMultiColumnTableRow<TObjectPtr<ParameterType>> FSuperRowType;
+	typedef typename STableRow<TObjectPtr<ParameterType>>::FArguments FTableRowArgs;
+		
+	void Construct(const FArguments& Args, const TSharedRef<STableViewBase>& OwnerTable)
+	{
+		Item = Args._Item;
+
+		FSuperRowType::Construct(
+				FTableRowArgs(),
+				OwnerTable);
+	}
+
+	void EnterNameEditingMode()
+	{
+		NameTextBlock->EnterEditingMode();
+	}
+
+protected:
+
+	virtual TSharedRef<SWidget> GenerateWidgetForColumn(const FName& InColumnName) override
+	{
+		if (InColumnName == ParameterNameColumn)
+		{
+			return SAssignNew(NameTextBlock, SInlineEditableTextBlock)
+				.IsSelected(this, &SCameraObjectInterfaceParameterTableRowBase<ParameterType>::IsSelected)
+				.Text_Lambda([this]() { return FText::FromString(Item->InterfaceParameterName); })
+				.OnTextCommitted(this, &SCameraObjectInterfaceParameterTableRowBase<ParameterType>::OnParameterNameTextCommitted);
+		}
+		else if (InColumnName == ParameterMessageColumn)
+		{
+			TSharedRef<FGameplayCamerasEditorStyle> GameplayCamerasEditorStyle = FGameplayCamerasEditorStyle::Get();
+			const FTextBlockStyle& MessageStyle = GameplayCamerasEditorStyle->GetWidgetStyle<FTextBlockStyle>("CameraObjectEditor.InterfaceParameter.Message");
+
+			return SNew(STextBlock)
+				.TextStyle(&MessageStyle)
+				.Text(this, &SCameraObjectInterfaceParameterTableRowBase<ParameterType>::GetParameterMessageText)
+				.ToolTipText(this, &SCameraObjectInterfaceParameterTableRowBase<ParameterType>::GetParameterMessageToolTip);
+		}
+
+		return SNullWidget::NullWidget;
+	}
+
+	virtual FReply OnDragDetected(const FGeometry& MyGeometry, const FPointerEvent& MouseEvent) override
+	{
+		if (MouseEvent.IsMouseButtonDown(EKeys::LeftMouseButton))
+		{
+			return FReply::Handled().BeginDragDrop(FCameraNodeGraphInterfaceParameterDragDropOp::New(Item));
+		}
+		return SMultiColumnTableRow<TObjectPtr<ParameterType>>::OnDragDetected(MyGeometry, MouseEvent);
+	}
+
+	void OnParameterNameTextCommitted(const FText& Text, ETextCommit::Type CommitType)
+	{
+		const FScopedTransaction Transaction(LOCTEXT("RenameInterfaceParameterTransaction", "Rename Interface Parameter"));
+
+		Item->Modify();
+		Item->InterfaceParameterName = Text.ToString();
+	}
+
+	FText GetParameterMessageText() const
+	{
+		if (!Item->Target || Item->TargetPropertyName.IsNone())
+		{
+			return LOCTEXT("UnboundInterfaceParameterMessage", "Unbound");
+		}
+		return FText::GetEmpty();
+	}
+
+	FText GetParameterMessageToolTip() const
+	{
+		if (!Item->Target || Item->TargetPropertyName.IsNone())
+		{
+			return LOCTEXT(
+					"UnboundInterfaceParameterMessageToolTip", 
+					"This interface parameter is not bound to any camera node. Setting this parameter will have no effect.");
+		}
+		return FText::GetEmpty();
+	}
+
+protected:
+
+	TObjectPtr<ParameterType> Item;
+	TSharedPtr<SInlineEditableTextBlock> NameTextBlock;
+};
+
+/**
+ * List entry for the blendable parameters panel.
+ */
+class SCameraObjectInterfaceBlendableParameterTableRow : public SCameraObjectInterfaceParameterTableRowBase<UCameraObjectInterfaceBlendableParameter>
+{
+public:
+
+	void Construct(const FArguments& Args, const TSharedRef<STableViewBase>& OwnerTable)
+	{
+		SCameraObjectInterfaceParameterTableRowBase<UCameraObjectInterfaceBlendableParameter>::Construct(Args, OwnerTable);
+	}
+
+protected:
+
+	virtual TSharedRef<SWidget> GenerateWidgetForColumn(const FName& InColumnName) override
+	{
+		if (InColumnName == ParameterTypeColumn)
+		{
+			TSharedRef<FGameplayCamerasEditorStyle> GameplayCamerasEditorStyle = FGameplayCamerasEditorStyle::Get();
+
+			const UEdGraphSchema_K2* K2Schema = GetDefault<UEdGraphSchema_CameraNodeK2>();
+
+			return SNew(SBox)
+				.MinDesiredWidth(125.f)
+				.Padding(this, &SCameraObjectInterfaceBlendableParameterTableRow::GetPinTypeSelectorPadding)
+				[
+					SNew(SPinTypeSelector, FGetPinTypeTree::CreateSP(this, &SCameraObjectInterfaceBlendableParameterTableRow::GetPinTypeTreeInfos))
+					.OnPinTypeChanged(this, &SCameraObjectInterfaceBlendableParameterTableRow::OnBlendableParameterPinTypeChanged)
+					.TargetPinType(this, &SCameraObjectInterfaceBlendableParameterTableRow::GetBlendableParameterPinType)
+					.ReadOnly(this, &SCameraObjectInterfaceBlendableParameterTableRow::IsPinTypeSelectorReadOnly)
+					.Schema(K2Schema)
+					.bAllowArrays(false)
+				];
+		}
+		else if (InColumnName == ParameterIsPreBlendedColumn)
+		{
+			return SNew(SCheckBox)
+				.IsChecked(this, &SCameraObjectInterfaceBlendableParameterTableRow::IsBlendableParameterPreBlended)
+				.OnCheckStateChanged(this, &SCameraObjectInterfaceBlendableParameterTableRow::OnBlendableParameterPreBlendedChanged);
+		}
+
+		return SCameraObjectInterfaceParameterTableRowBase<UCameraObjectInterfaceBlendableParameter>::GenerateWidgetForColumn(InColumnName);
+	}
+
+	FMargin GetPinTypeSelectorPadding() const
+	{
+		// Add some horizontal margin when the pin type selector is read-only, so that the transition on hover is seamless.
+		return IsHovered() ? FMargin(0) : FMargin(7, 0);
+	}
+
+	void GetPinTypeTreeInfos(TArray<FPinTypeTreeItem>& TypeTree, ETypeTreeFilter TypeTreeFilter) const
+	{
+		using FPinTypeTreeInfo = UEdGraphSchema_K2::FPinTypeTreeInfo;
+
+		const UEdGraphSchema_K2* K2Schema = GetDefault<UEdGraphSchema_CameraNodeK2>();
+
+		TypeTree.Reset();
+
+		TypeTree.Add(MakeShared<FPinTypeTreeInfo>(UEdGraphSchema_K2::GetCategoryText(UEdGraphSchema_K2::PC_Boolean, true), UEdGraphSchema_K2::PC_Boolean, K2Schema, LOCTEXT("BooleanType", "True or false value")));
+		TypeTree.Add(MakeShared<FPinTypeTreeInfo>(UEdGraphSchema_K2::GetCategoryText(UEdGraphSchema_K2::PC_Int, true), UEdGraphSchema_K2::PC_Int, K2Schema, LOCTEXT("IntegerType", "Integer number")));
+		TypeTree.Add(MakeShared<FPinTypeTreeInfo>(UEdGraphSchema_K2::GetCategoryText(UEdGraphSchema_K2::PC_Float, true), UEdGraphSchema_K2::PC_Float, K2Schema, LOCTEXT("FloatType", "Floating point number")));
+		TypeTree.Add(MakeShared<FPinTypeTreeInfo>(UEdGraphSchema_K2::GetCategoryText(UEdGraphSchema_K2::PC_Double, true), UEdGraphSchema_K2::PC_Double, K2Schema, LOCTEXT("DoubleType", "Double precision floating point number")));
+
+		TypeTree.Add(MakeShared<FPinTypeTreeInfo>(UEdGraphSchema_K2::PC_Struct, TVariantStructure<FVector2f>::Get(), LOCTEXT("Vector2fType", "A 2D vector")));
+		TypeTree.Add(MakeShared<FPinTypeTreeInfo>(UEdGraphSchema_K2::PC_Struct, TBaseStructure<FVector2D>::Get(), LOCTEXT("Vector2dType", "A double precision 2D vector")));
+		TypeTree.Add(MakeShared<FPinTypeTreeInfo>(UEdGraphSchema_K2::PC_Struct, TVariantStructure<FVector3f>::Get(), LOCTEXT("Vector3fType", "A 3D vector")));
+		TypeTree.Add(MakeShared<FPinTypeTreeInfo>(UEdGraphSchema_K2::PC_Struct, TBaseStructure<FVector>::Get(), LOCTEXT("Vector3dType", "A double precision 3D vector")));
+		TypeTree.Add(MakeShared<FPinTypeTreeInfo>(UEdGraphSchema_K2::PC_Struct, TVariantStructure<FVector4f>::Get(), LOCTEXT("Vector4fType", "A 4D vector")));
+		TypeTree.Add(MakeShared<FPinTypeTreeInfo>(UEdGraphSchema_K2::PC_Struct, TBaseStructure<FVector4>::Get(), LOCTEXT("Vector4dType", "A double precision 4D vector")));
+		TypeTree.Add(MakeShared<FPinTypeTreeInfo>(UEdGraphSchema_K2::PC_Struct, TVariantStructure<FRotator3f>::Get(), LOCTEXT("Rotator3fType", "A 3D rotation")));
+		TypeTree.Add(MakeShared<FPinTypeTreeInfo>(UEdGraphSchema_K2::PC_Struct, TBaseStructure<FRotator>::Get(), LOCTEXT("Rotator3dType", "A double precision 3D rotation")));
+		TypeTree.Add(MakeShared<FPinTypeTreeInfo>(UEdGraphSchema_K2::PC_Struct, TVariantStructure<FTransform3f>::Get(), LOCTEXT("Transform3fType", "A 3D transformation")));
+		TypeTree.Add(MakeShared<FPinTypeTreeInfo>(UEdGraphSchema_K2::PC_Struct, TBaseStructure<FTransform>::Get(), LOCTEXT("Transform3dType", "A double precision 3D transformation")));
+
+		TSharedPtr<FPinTypeTreeInfo> Structs = MakeShared<FPinTypeTreeInfo>(
+				LOCTEXT("BlendableStructPinTypeLabel", "Blendable Structures"),
+				UEdGraphSchema_K2::PC_Struct,
+				K2Schema,
+				LOCTEXT("BlendableStructPinTypeToolTip", "Blendable structure types"),
+				true);
+		IGameplayCamerasModule& GameplayCamerasModule = IGameplayCamerasModule::Get();
+		for (const FBlendableStructInfo& BlendableStruct : GameplayCamerasModule.GetBlendableStructs())
+		{
+			if (const UScriptStruct* StructType = BlendableStruct.StructType)
+			{
+				Structs->Children.Add(
+						MakeShared<UEdGraphSchema_K2::FPinTypeTreeInfo>(
+							UEdGraphSchema_K2::PC_Struct,
+							const_cast<UScriptStruct*>(StructType),
+							StructType->GetToolTipText(),
+							false,
+							(uint8)EObjectReferenceType::NotAnObject));
+			}
+		}
+		TypeTree.Add(Structs);
+	}
+
+	FEdGraphPinType GetBlendableParameterPinType() const
+	{
+		FEdGraphPinType PinType;
+		switch (Item->ParameterType)
+		{
+			case ECameraVariableType::Boolean:
+				PinType.PinCategory = UEdGraphSchema_K2::PC_Boolean;
+				break;
+			case ECameraVariableType::Integer32:
+				PinType.PinCategory = UEdGraphSchema_K2::PC_Int;
+				break;
+			case ECameraVariableType::Float:
+				PinType.PinCategory = UEdGraphSchema_K2::PC_Float;
+				break;
+			case ECameraVariableType::Double:
+				PinType.PinCategory = UEdGraphSchema_K2::PC_Double;
+				break;
+			case ECameraVariableType::Vector2f:
+				PinType.PinCategory = UEdGraphSchema_K2::PC_Struct;
+				PinType.PinSubCategoryObject = TVariantStructure<FVector2f>::Get();
+				break;
+			case ECameraVariableType::Vector2d:
+				PinType.PinCategory = UEdGraphSchema_K2::PC_Struct;
+				PinType.PinSubCategoryObject = TBaseStructure<FVector2D>::Get();
+				break;
+			case ECameraVariableType::Vector3f:
+				PinType.PinCategory = UEdGraphSchema_K2::PC_Struct;
+				PinType.PinSubCategoryObject = TVariantStructure<FVector3f>::Get();
+				break;
+			case ECameraVariableType::Vector3d:
+				PinType.PinCategory = UEdGraphSchema_K2::PC_Struct;
+				PinType.PinSubCategoryObject = TBaseStructure<FVector>::Get();
+				break;
+			case ECameraVariableType::Vector4f:
+				PinType.PinCategory = UEdGraphSchema_K2::PC_Struct;
+				PinType.PinSubCategoryObject = TVariantStructure<FVector4f>::Get();
+				break;
+			case ECameraVariableType::Vector4d:
+				PinType.PinCategory = UEdGraphSchema_K2::PC_Struct;
+				PinType.PinSubCategoryObject = TBaseStructure<FVector4>::Get();
+				break;
+			case ECameraVariableType::Rotator3f:
+				PinType.PinCategory = UEdGraphSchema_K2::PC_Struct;
+				PinType.PinSubCategoryObject = TVariantStructure<FRotator3f>::Get();
+				break;
+			case ECameraVariableType::Rotator3d:
+				PinType.PinCategory = UEdGraphSchema_K2::PC_Struct;
+				PinType.PinSubCategoryObject = TBaseStructure<FRotator>::Get();
+				break;
+			case ECameraVariableType::Transform3f:
+				PinType.PinCategory = UEdGraphSchema_K2::PC_Struct;
+				PinType.PinSubCategoryObject = TVariantStructure<FTransform3f>::Get();
+				break;
+			case ECameraVariableType::Transform3d:
+				PinType.PinCategory = UEdGraphSchema_K2::PC_Struct;
+				PinType.PinSubCategoryObject = TBaseStructure<FTransform>::Get();
+				break;
+			case ECameraVariableType::BlendableStruct:
+				PinType.PinCategory = UEdGraphSchema_K2::PC_Struct;
+				PinType.PinSubCategoryObject = const_cast<UScriptStruct*>(Item->BlendableStructType.Get());
+				break;
+		}
+		return PinType;
+	}
+
+	void OnBlendableParameterPinTypeChanged(const FEdGraphPinType& PinType)
+	{
+		bool bIsValidType = true;
+		ECameraVariableType NewParameterType = ECameraVariableType::Boolean;
+
+		if (PinType.PinCategory == UEdGraphSchema_K2::PC_Boolean)
+		{
+			NewParameterType = ECameraVariableType::Boolean;
+		}
+		else if (PinType.PinCategory == UEdGraphSchema_K2::PC_Int)
+		{
+			NewParameterType = ECameraVariableType::Integer32;
+		}
+		else if (PinType.PinCategory == UEdGraphSchema_K2::PC_Float)
+		{
+			NewParameterType = ECameraVariableType::Float;
+		}
+		else if (PinType.PinCategory == UEdGraphSchema_K2::PC_Double)
+		{
+			NewParameterType = ECameraVariableType::Double;
+		}
+		else if (PinType.PinCategory == UEdGraphSchema_K2::PC_Struct)
+		{
+			const UObject* TypeObject = PinType.PinSubCategoryObject.Get();
+			if (TypeObject == TVariantStructure<FVector2f>::Get())
+			{
+				NewParameterType = ECameraVariableType::Vector2f;
+			}
+			else if (TypeObject == TBaseStructure<FVector2D>::Get())
+			{
+				NewParameterType = ECameraVariableType::Vector2d;
+			}
+			else if (TypeObject == TVariantStructure<FVector3f>::Get())
+			{
+				NewParameterType = ECameraVariableType::Vector3f;
+			}
+			else if (TypeObject == TBaseStructure<FVector>::Get())
+			{
+				NewParameterType = ECameraVariableType::Vector3d;
+			}
+			else if (TypeObject == TVariantStructure<FVector4f>::Get())
+			{
+				NewParameterType = ECameraVariableType::Vector4f;
+			}
+			else if (TypeObject == TBaseStructure<FVector4>::Get())
+			{
+				NewParameterType = ECameraVariableType::Vector4d;
+			}
+			else if (TypeObject == TVariantStructure<FRotator3f>::Get())
+			{
+				NewParameterType = ECameraVariableType::Rotator3f;
+			}
+			else if (TypeObject == TBaseStructure<FRotator>::Get())
+			{
+				NewParameterType = ECameraVariableType::Rotator3d;
+			}
+			else if (TypeObject == TVariantStructure<FTransform3f>::Get())
+			{
+				NewParameterType = ECameraVariableType::Transform3f;
+			}
+			else if (TypeObject == TBaseStructure<FTransform>::Get())
+			{
+				NewParameterType = ECameraVariableType::Transform3d;
+			}
+			else
+			{
+				bIsValidType = false;
+			}
+		}
+		else
+		{
+			bIsValidType = false;
+		}
+
+		if (ensure(bIsValidType) && Item->ParameterType != NewParameterType)
+		{
+			const FScopedTransaction Transaction(LOCTEXT("ChangeBlendableParameterType", "Change Blendable Parameter Type"));
+
+			Item->Modify();
+			Item->ParameterType = NewParameterType;
+		}
+	}
+
+	ECheckBoxState IsBlendableParameterPreBlended() const
+	{
+		return Item->bIsPreBlended ? ECheckBoxState::Checked : ECheckBoxState::Unchecked;
+	}
+	
+	void OnBlendableParameterPreBlendedChanged(ECheckBoxState CheckState)
+	{
+		const bool bIsPreBlended = (CheckState == ECheckBoxState::Checked);
+		if (Item->bIsPreBlended != bIsPreBlended)
+		{
+			const FScopedTransaction Transaction(LOCTEXT("ChangeBlendableParameterIsPreBlended", "Change Blendable Parameter Pre-Blending"));
+
+			Item->Modify();
+			Item->bIsPreBlended = bIsPreBlended;
+		}
+	}
+
+	bool IsPinTypeSelectorReadOnly() const
+	{
+		return !IsHovered();
+	}
+};
+
+class FDataParameterPinTypeSelectorFilter : public IPinTypeSelectorFilter
+{
+	virtual bool ShouldShowPinTypeTreeItem(FPinTypeTreeItem InItem) const override
+	{
+		// Remove types that are for blendable parameters.
+		FEdGraphPinType PinType = InItem->GetPinType(false);
+		if (PinType.PinCategory == UEdGraphSchema_K2::PC_Boolean ||
+				PinType.PinCategory == UEdGraphSchema_K2::PC_Int ||
+				PinType.PinCategory == UEdGraphSchema_K2::PC_Float ||
+				PinType.PinCategory == UEdGraphSchema_K2::PC_Double ||
+				PinType.PinCategory == UEdGraphSchema_K2::PC_Real)
+		{
+			return false;
+		}
+
+		if (PinType.PinCategory == UEdGraphSchema_K2::PC_Struct)
+		{
+			const UObject* TypeObject = PinType.PinSubCategoryObject.Get();
+			if (TypeObject == TVariantStructure<FVector2f>::Get() ||
+					TypeObject == TBaseStructure<FVector2D>::Get() ||
+					TypeObject == TVariantStructure<FVector3f>::Get() ||
+					TypeObject == TBaseStructure<FVector>::Get() ||
+					TypeObject == TVariantStructure<FVector4f>::Get() ||
+					TypeObject == TBaseStructure<FVector4>::Get() ||
+					TypeObject == TVariantStructure<FRotator3f>::Get() ||
+					TypeObject == TBaseStructure<FRotator>::Get() ||
+					TypeObject == TVariantStructure<FTransform3f>::Get() ||
+					TypeObject == TBaseStructure<FTransform>::Get())
+			{
+				return false;
+			}
+		}
+		
+		return true;
+	};
+};
+
+/**
+ * List entry for the data parameters panel.
+ */
+class SCameraObjectInterfaceDataParameterTableRow : public SCameraObjectInterfaceParameterTableRowBase<UCameraObjectInterfaceDataParameter>
+{
+protected:
+
+	virtual TSharedRef<SWidget> GenerateWidgetForColumn(const FName& InColumnName) override
+	{
+		if (InColumnName == ParameterTypeColumn)
+		{
+			TSharedRef<FGameplayCamerasEditorStyle> GameplayCamerasEditorStyle = FGameplayCamerasEditorStyle::Get();
+
+			const UEdGraphSchema_K2* K2Schema = GetDefault<UEdGraphSchema_CameraNodeK2>();
+			TArray<TSharedPtr<IPinTypeSelectorFilter>> PinTypeSelectorFilters{ MakeShared<FDataParameterPinTypeSelectorFilter>() };
+
+			return SNew(SBox)
+				.MinDesiredWidth(125.f)
+				.Padding(this, &SCameraObjectInterfaceDataParameterTableRow::GetPinTypeSelectorPadding)
+				[
+					SNew(SPinTypeSelector, FGetPinTypeTree::CreateUObject(K2Schema, &UEdGraphSchema_K2::GetVariableTypeTree))
+					.OnPinTypeChanged(this, &SCameraObjectInterfaceDataParameterTableRow::OnDataParameterPinTypeChanged)
+					.TargetPinType(this, &SCameraObjectInterfaceDataParameterTableRow::GetDataParameterPinType)
+					.ReadOnly(this, &SCameraObjectInterfaceDataParameterTableRow::IsPinTypeSelectorReadOnly)
+					.Schema(K2Schema)
+					.CustomFilters(PinTypeSelectorFilters)
+					.bAllowArrays(true)
+				];
+		}
+
+		return SCameraObjectInterfaceParameterTableRowBase<UCameraObjectInterfaceDataParameter>::GenerateWidgetForColumn(InColumnName);
+	}
+
+	FMargin GetPinTypeSelectorPadding() const
+	{
+		// Add some horizontal margin when the pin type selector is read-only, so that the transition on hover is seamless.
+		return IsHovered() ? FMargin(0) : FMargin(7, 0);
+	}
+
+	FEdGraphPinType GetDataParameterPinType() const
+	{
+		FEdGraphPinType PinType;
+		PinType.PinSubCategoryObject = const_cast<UObject*>(Item->DataTypeObject.Get());
+
+		switch (Item->DataType)
+		{
+			case ECameraContextDataType::Name:
+				PinType.PinCategory = UEdGraphSchema_K2::PC_Name;
+				break;
+			case ECameraContextDataType::String:
+				PinType.PinCategory = UEdGraphSchema_K2::PC_String;
+				break;
+			case ECameraContextDataType::Enum:
+				PinType.PinCategory = UEdGraphSchema_K2::PC_Enum;
+				break;
+			case ECameraContextDataType::Struct:
+				PinType.PinCategory = UEdGraphSchema_K2::PC_Struct;
+				break;
+			case ECameraContextDataType::Object:
+				PinType.PinCategory = UEdGraphSchema_K2::PC_Object;
+				break;
+			case ECameraContextDataType::Class:
+				PinType.PinCategory = UEdGraphSchema_K2::PC_Class;
+				break;
+		}
+
+		switch (Item->DataContainerType)
+		{
+			case ECameraContextDataContainerType::Array:
+				PinType.ContainerType = EPinContainerType::Array;
+				break;
+		}
+
+		return PinType;
+	}
+
+	void OnDataParameterPinTypeChanged(const FEdGraphPinType& PinType)
+	{
+		bool bIsValidType = true;
+		ECameraContextDataType NewDataType = ECameraContextDataType::Name;
+		ECameraContextDataContainerType NewDataContainerType = ECameraContextDataContainerType::None;
+		const UObject* NewDataTypeObject = PinType.PinSubCategoryObject.Get();
+
+		if (PinType.PinCategory == UEdGraphSchema_K2::PC_Name)
+		{
+			NewDataType = ECameraContextDataType::Name;
+		}
+		else if (PinType.PinCategory == UEdGraphSchema_K2::PC_String)
+		{
+			NewDataType = ECameraContextDataType::String;
+		}
+		else if (PinType.PinCategory == UEdGraphSchema_K2::PC_Enum)
+		{
+			NewDataType = ECameraContextDataType::Enum;
+		}
+		else if (PinType.PinCategory == UEdGraphSchema_K2::PC_Struct)
+		{
+			NewDataType = ECameraContextDataType::Struct;
+		}
+		else if (PinType.PinCategory == UEdGraphSchema_K2::PC_Object)
+		{
+			NewDataType = ECameraContextDataType::Object;
+		}
+		else if (PinType.PinCategory == UEdGraphSchema_K2::PC_Class)
+		{
+			NewDataType = ECameraContextDataType::Class;
+		}
+		else
+		{
+			bIsValidType = false;
+		}
+
+		switch (PinType.ContainerType)
+		{
+			case EPinContainerType::None:
+				NewDataContainerType = ECameraContextDataContainerType::None;
+				break;
+			case EPinContainerType::Array:
+				NewDataContainerType = ECameraContextDataContainerType::Array;
+				break;
+			default:
+				bIsValidType = false;
+				break;
+		}
+
+		if (ensure(bIsValidType) && (NewDataType != Item->DataType || NewDataContainerType != Item->DataContainerType || NewDataTypeObject != Item->DataTypeObject))
+		{
+			const FScopedTransaction Transaction(LOCTEXT("ChangeDataParameterType", "Change Data Parameter Type"));
+
+			Item->DataType = NewDataType;
+			Item->DataTypeObject = NewDataTypeObject;
+		}
+	}
+
+	bool IsPinTypeSelectorReadOnly() const
+	{
+		return !IsHovered();
+	}
+};
+
+/**
+ * The overall interface parameters panel, showing two sub-panels, one for blendable parameters, and
+ * one for data parameters.
+ */
+class SCameraObjectInterfaceParametersPanel 
+	: public SCompoundWidget
+	, public ICameraObjectEventHandler
+{
+public:
+
+	SLATE_BEGIN_ARGS(SCameraObjectInterfaceParametersPanel)
+	{}
+	SLATE_END_ARGS()
+
+	void Construct(const FArguments& Args, FCameraObjectInterfaceParametersToolkit* OwnerToolkit);
+
+	void RequestListRefresh();
+
+protected:
+
+	// SWidget interface.
+	virtual void Tick(const FGeometry& AllottedGeometry, const double InCurrentTime, const float InDeltaTime) override;
+
+	// ICameraObjectEventHandler interface.
+	virtual void OnCameraObjectInterfaceChanged() override;
+
+private:
+
+	TSharedRef<ITableRow> OnGenerateBlendableParameterRow(TObjectPtr<UCameraObjectInterfaceBlendableParameter> Item, const TSharedRef<STableViewBase>& OwnerTable);
+	TSharedRef<ITableRow> OnGenerateDataParameterRow(TObjectPtr<UCameraObjectInterfaceDataParameter> Item, const TSharedRef<STableViewBase>& OwnerTable);
+
+	void OnBlendableSelectionChanged(TObjectPtr<UCameraObjectInterfaceBlendableParameter> Item, ESelectInfo::Type Type);
+	void OnDataParameterSelectionChanged(TObjectPtr<UCameraObjectInterfaceDataParameter> Item, ESelectInfo::Type Type);
+
+	TSharedPtr<SWidget> OnBlendableParameterContextMenuOpening();
+	TSharedPtr<SWidget> OnDataParameterContextMenuOpening();
+
+	template<typename ItemType>
+	TSharedPtr<SWidget> OnInterfaceParameterContextMenuOpening(
+			TSharedPtr<SListView<TObjectPtr<ItemType>>> ListView, TObjectPtr<ItemType> Item);
+	template<typename ItemType>
+	void OnRenameInterfaceParameter(
+			TSharedPtr<SListView<TObjectPtr<ItemType>>> ListView, TObjectPtr<ItemType> Item);
+	template<typename ItemType>
+	void OnDeleteInterfaceParameter(
+			TSharedPtr<SListView<TObjectPtr<ItemType>>> ListView, TObjectPtr<ItemType> Item);
+
+	FReply OnAddBlendableParameter();
+	FReply OnAddDataParameter();
+
+private:
+
+	UBaseCameraObject* CameraObject = nullptr;
+	FCameraObjectInterfaceParametersToolkit* Toolkit = nullptr;
+
+	TCameraEventHandler<ICameraObjectEventHandler> EventHandler;
+
+	TSharedPtr<SListView<TObjectPtr<UCameraObjectInterfaceBlendableParameter>>> BlendableParametersListView;
+	TSharedPtr<SListView<TObjectPtr<UCameraObjectInterfaceDataParameter>>> DataParametersListView;
+
+	bool bListRefreshRequested = false;
+};
+
+void SCameraObjectInterfaceParametersPanel::Construct(const FArguments& Args, FCameraObjectInterfaceParametersToolkit* OwnerToolkit)
+{
+	CameraObject = OwnerToolkit->GetCameraObject();
+	CameraObject->EventHandlers.Register(EventHandler, this);
+
+	Toolkit = OwnerToolkit;
+
+	ChildSlot
+	[
+		SNew(SSplitter)
+		.Orientation(Orient_Vertical)
+
+		// The blendable parameters panel.
+		+SSplitter::Slot()
+		[
+			SNew(SVerticalBox)
+			+SVerticalBox::Slot()
+			.AutoHeight()
+			[
+				SNew(SBorder)
+				.BorderImage(FAppStyle::GetBrush("ToolPanel.GroupBorder"))
+				.Padding(FMargin(3, 5))
+				[
+					SNew(SHorizontalBox)
+					+SHorizontalBox::Slot()
+					.FillWidth(1.f)
+					.VAlign(VAlign_Center)
+					.Padding(5.f)
+					[
+						SNew(SRichTextBlock)
+						.Text(LOCTEXT("BlendableParametersPanelTitle", "Blendable Parameters"))
+						.TransformPolicy(ETextTransformPolicy::ToUpper)
+						.DecoratorStyleSet(&FAppStyle::Get())
+						.TextStyle(FAppStyle::Get(), "DetailsView.CategoryTextStyle")
+					]
+					+SHorizontalBox::Slot()
+					.AutoWidth()
+					.VAlign(VAlign_Center)
+					[
+						SNew(SButton)
+						.ButtonStyle(FAppStyle::Get(), "SimpleButton")
+						.ContentPadding(FMargin(1, 0))
+						.ToolTipText(LOCTEXT("AddBlendableToolTip", "Add a blendable parameter"))
+						.OnClicked(this, &SCameraObjectInterfaceParametersPanel::OnAddBlendableParameter)
+						[
+							SNew(SImage)
+							.Image(FAppStyle::Get().GetBrush("Icons.PlusCircle"))
+							.ColorAndOpacity(FSlateColor::UseForeground())
+						]
+					]
+				]
+			]
+			+SVerticalBox::Slot()
+			.FillHeight(1.f)
+			[
+				SAssignNew(BlendableParametersListView, SListView<TObjectPtr<UCameraObjectInterfaceBlendableParameter>>)
+				.ListItemsSource(&CameraObject->Interface.BlendableParameters)
+				.OnGenerateRow(this, &SCameraObjectInterfaceParametersPanel::OnGenerateBlendableParameterRow)
+				.OnSelectionChanged(this, &SCameraObjectInterfaceParametersPanel::OnBlendableSelectionChanged)
+				.OnContextMenuOpening(this, &SCameraObjectInterfaceParametersPanel::OnBlendableParameterContextMenuOpening)
+				.HeaderRow
+				(
+					SNew(SHeaderRow)
+
+					+SHeaderRow::Column(ParameterTypeColumn)
+					.FillWidth(0.3f)
+					.DefaultLabel(LOCTEXT("ParameterTypeColumnLabel", "Type"))
+
+					+SHeaderRow::Column(ParameterNameColumn)
+					.FillWidth(0.5f)
+					.DefaultLabel(LOCTEXT("ParameterNameColumnLabel", "Name"))
+
+					+SHeaderRow::Column(ParameterIsPreBlendedColumn)
+					.ManualWidth(60)
+					.DefaultLabel(LOCTEXT("ParameterIsPreBlendedColumnLabel", "Pre-Blend"))
+
+					+SHeaderRow::Column(ParameterMessageColumn)
+					.FillWidth(0.2f)
+					.DefaultLabel(LOCTEXT("ParameterMessageColumnLabel", "Note"))
+				)
+			]
+		]
+
+		// The data parameters panel.
+		+SSplitter::Slot()
+		[
+			SNew(SVerticalBox)
+			+SVerticalBox::Slot()
+			.AutoHeight()
+			[
+				SNew(SBorder)
+				.BorderImage(FAppStyle::GetBrush("ToolPanel.GroupBorder"))
+				.Padding(FMargin(3, 5))
+				[
+					SNew(SHorizontalBox)
+					+SHorizontalBox::Slot()
+					.FillWidth(1.f)
+					.VAlign(VAlign_Center)
+					.Padding(5.f)
+					[
+						SNew(SRichTextBlock)
+						.Text(LOCTEXT("DataParametersPanelTitle", "Data Parameters"))
+						.TransformPolicy(ETextTransformPolicy::ToUpper)
+						.DecoratorStyleSet(&FAppStyle::Get())
+						.TextStyle(FAppStyle::Get(), "DetailsView.CategoryTextStyle")
+					]
+					+SHorizontalBox::Slot()
+					.AutoWidth()
+					.VAlign(VAlign_Center)
+					[
+						SNew(SButton)
+						.ButtonStyle(FAppStyle::Get(), "SimpleButton")
+						.ContentPadding(FMargin(1, 0))
+						.ToolTipText(LOCTEXT("AddDataParameterToolTip", "Add a data parameter"))
+						.OnClicked(this, &SCameraObjectInterfaceParametersPanel::OnAddDataParameter)
+						[
+							SNew(SImage)
+							.Image(FAppStyle::Get().GetBrush("Icons.PlusCircle"))
+							.ColorAndOpacity(FSlateColor::UseForeground())
+						]
+					]
+				]
+			]
+			+SVerticalBox::Slot()
+			.FillHeight(1.f)
+			[
+				SAssignNew(DataParametersListView, SListView<TObjectPtr<UCameraObjectInterfaceDataParameter>>)
+				.ListItemsSource(&CameraObject->Interface.DataParameters)
+				.OnGenerateRow(this, &SCameraObjectInterfaceParametersPanel::OnGenerateDataParameterRow)
+				.OnSelectionChanged(this, &SCameraObjectInterfaceParametersPanel::OnDataParameterSelectionChanged)
+				.OnContextMenuOpening(this, &SCameraObjectInterfaceParametersPanel::OnDataParameterContextMenuOpening)
+				.HeaderRow
+				(
+					SNew(SHeaderRow)
+
+					+SHeaderRow::Column(ParameterTypeColumn)
+					.FillWidth(0.3f)
+					.DefaultLabel(LOCTEXT("ParameterTypeColumnLabel", "Type"))
+
+					+SHeaderRow::Column(ParameterNameColumn)
+					.FillWidth(0.5f)
+					.DefaultLabel(LOCTEXT("ParameterNameColumnLabel", "Name"))
+
+					+SHeaderRow::Column(ParameterMessageColumn)
+					.FillWidth(0.2f)
+					.DefaultLabel(LOCTEXT("ParameterMessageColumnLabel", "Note"))
+				)
+			]
+		]
+	];
+}
+
+void SCameraObjectInterfaceParametersPanel::RequestListRefresh()
+{
+	bListRefreshRequested = true;
+}
+
+void SCameraObjectInterfaceParametersPanel::Tick(const FGeometry& AllottedGeometry, const double InCurrentTime, const float InDeltaTime)
+{
+	SCompoundWidget::Tick(AllottedGeometry, InCurrentTime, InDeltaTime);
+
+	if (bListRefreshRequested)
+	{
+		bListRefreshRequested = false;
+
+		BlendableParametersListView->RequestListRefresh();
+		DataParametersListView->RequestListRefresh();
+	}
+}
+
+TSharedRef<ITableRow> SCameraObjectInterfaceParametersPanel::OnGenerateBlendableParameterRow(TObjectPtr<UCameraObjectInterfaceBlendableParameter> Item, const TSharedRef<STableViewBase>& OwnerTable)
+{
+	return SNew(SCameraObjectInterfaceBlendableParameterTableRow, OwnerTable)
+		.Item(Item);
+}
+
+TSharedRef<ITableRow> SCameraObjectInterfaceParametersPanel::OnGenerateDataParameterRow(TObjectPtr<UCameraObjectInterfaceDataParameter> Item, const TSharedRef<STableViewBase>& OwnerTable)
+{
+	return SNew(SCameraObjectInterfaceDataParameterTableRow, OwnerTable)
+		.Item(Item);
+}
+
+void SCameraObjectInterfaceParametersPanel::OnBlendableSelectionChanged(TObjectPtr<UCameraObjectInterfaceBlendableParameter> Item, ESelectInfo::Type Type)
+{
+	Toolkit->OnInterfaceParameterSelected().Broadcast(Item);
+}
+
+void SCameraObjectInterfaceParametersPanel::OnDataParameterSelectionChanged(TObjectPtr<UCameraObjectInterfaceDataParameter> Item, ESelectInfo::Type Type)
+{
+	Toolkit->OnInterfaceParameterSelected().Broadcast(Item);
+}
+
+TSharedPtr<SWidget> SCameraObjectInterfaceParametersPanel::OnBlendableParameterContextMenuOpening()
+{
+	TArray<TObjectPtr<UCameraObjectInterfaceBlendableParameter>> SelectedItems = BlendableParametersListView->GetSelectedItems();
+	if (SelectedItems.Num() > 0)
+	{
+		return OnInterfaceParameterContextMenuOpening(BlendableParametersListView, SelectedItems[0]);
+	}
+	return nullptr;
+}
+
+TSharedPtr<SWidget> SCameraObjectInterfaceParametersPanel::OnDataParameterContextMenuOpening()
+{
+	TArray<TObjectPtr<UCameraObjectInterfaceDataParameter>> SelectedItems = DataParametersListView->GetSelectedItems();
+	if (SelectedItems.Num() > 0)
+	{
+		return OnInterfaceParameterContextMenuOpening(DataParametersListView, SelectedItems[0]);
+	}
+	return nullptr;
+}
+
+template<typename ItemType>
+TSharedPtr<SWidget> SCameraObjectInterfaceParametersPanel::OnInterfaceParameterContextMenuOpening(
+		TSharedPtr<SListView<TObjectPtr<ItemType>>> ListView, TObjectPtr<ItemType> Item)
+{
+	FMenuBuilder MenuBuilder(true, nullptr);
+
+	MenuBuilder.AddMenuEntry(
+			LOCTEXT("RenameInterfaceParameter", "Rename"),
+			LOCTEXT("RenameInterfaceParameterToolTip", "Renames this interface parameter"),
+			FSlateIcon(),
+			FExecuteAction::CreateSP(this, &SCameraObjectInterfaceParametersPanel::OnRenameInterfaceParameter, ListView, Item));
+	MenuBuilder.AddMenuEntry(
+			LOCTEXT("DeleteInterfaceParameter", "Delete"),
+			LOCTEXT("DeleteInterfaceParameterToolTip", "Deletes this interface parameter"),
+			FSlateIcon(),
+			FExecuteAction::CreateSP(this, &SCameraObjectInterfaceParametersPanel::OnDeleteInterfaceParameter, ListView, Item));
+
+	return MenuBuilder.MakeWidget();
+}
+
+template<typename ItemType>
+void SCameraObjectInterfaceParametersPanel::OnRenameInterfaceParameter(
+		TSharedPtr<SListView<TObjectPtr<ItemType>>> ListView, TObjectPtr<ItemType> Item)
+{
+	TSharedPtr<ITableRow> RowWidget = ListView->WidgetFromItem(Item);
+	if (!ensure(RowWidget))
+	{
+		return;
+	}
+
+	TSharedPtr<SCameraObjectInterfaceParameterTableRowBase<ItemType>> TypedRowWidget = 
+		StaticCastSharedPtr<SCameraObjectInterfaceParameterTableRowBase<ItemType>>(RowWidget);
+	TypedRowWidget->EnterNameEditingMode();
+}
+
+template<typename ItemType>
+void SCameraObjectInterfaceParametersPanel::OnDeleteInterfaceParameter(
+		TSharedPtr<SListView<TObjectPtr<ItemType>>> ListView, TObjectPtr<ItemType> Item)
+{
+	const FScopedTransaction Transaction(LOCTEXT("RemoveInterfaceParameter", "Remove Interface Parameter"));
+
+	CameraObject->Modify();
+
+	if constexpr (std::is_same_v<ItemType, UCameraObjectInterfaceBlendableParameter>)
+	{
+		const int32 NumRemoved = CameraObject->Interface.BlendableParameters.Remove(Item);
+		ensure(NumRemoved == 1);
+	}
+	else if constexpr (std::is_same_v<ItemType, UCameraObjectInterfaceDataParameter>)
+	{
+		const int32 NumRemoved = CameraObject->Interface.DataParameters.Remove(Item);
+		ensure(NumRemoved == 1);
+	}
+
+	CameraObject->EventHandlers.Notify(&ICameraObjectEventHandler::OnCameraObjectInterfaceChanged);
+
+	ListView->RequestListRefresh();
+}
+
+FReply SCameraObjectInterfaceParametersPanel::OnAddBlendableParameter()
+{
+	const FScopedTransaction Transaction(LOCTEXT("AddBlendableParameter", "Add Blendable Parameter"));
+
+	UCameraObjectInterfaceBlendableParameter* NewBlendableParameter = NewObject<UCameraObjectInterfaceBlendableParameter>(CameraObject, NAME_None, RF_Transactional);
+	NewBlendableParameter->InterfaceParameterName = NewBlendableParameter->GetName();
+
+	CameraObject->Modify();
+	CameraObject->Interface.BlendableParameters.Add(NewBlendableParameter);
+	CameraObject->EventHandlers.Notify(&ICameraObjectEventHandler::OnCameraObjectInterfaceChanged);
+
+	return FReply::Handled();
+}
+
+FReply SCameraObjectInterfaceParametersPanel::OnAddDataParameter()
+{
+	const FScopedTransaction Transaction(LOCTEXT("AddDataParameter", "Add Data Parameter"));
+
+	UCameraObjectInterfaceDataParameter* NewDataParameter = NewObject<UCameraObjectInterfaceDataParameter>(CameraObject, NAME_None, RF_Transactional);
+	NewDataParameter->InterfaceParameterName = NewDataParameter->GetName();
+
+	CameraObject->Modify();
+	CameraObject->Interface.DataParameters.Add(NewDataParameter);
+	CameraObject->EventHandlers.Notify(&ICameraObjectEventHandler::OnCameraObjectInterfaceChanged);
+
+	return FReply::Handled();
+}
+
+void SCameraObjectInterfaceParametersPanel::OnCameraObjectInterfaceChanged()
+{
+	bListRefreshRequested = true;
+}
+
+FCameraObjectInterfaceParametersToolkit::FCameraObjectInterfaceParametersToolkit()
+{
+	SAssignNew(PanelContainer, SBox);
+
+	if (GEditor)
+	{
+		GEditor->RegisterForUndo(this);
+	}
+}
+
+FCameraObjectInterfaceParametersToolkit::~FCameraObjectInterfaceParametersToolkit()
+{
+	if (GEditor)
+	{
+		GEditor->UnregisterForUndo(this);
+	}
+}
+
+void FCameraObjectInterfaceParametersToolkit::SetCameraObject(UBaseCameraObject* InCameraObject)
+{
+	if (CameraObject != InCameraObject)
+	{
+		PanelContainer->SetContent(SNullWidget::NullWidget);
+
+		CameraObject = InCameraObject;
+
+		if (CameraObject)
+		{
+			Panel = SNew(SCameraObjectInterfaceParametersPanel, this);
+			PanelContainer->SetContent(Panel.ToSharedRef());
+		}
+	}
+}
+
+TSharedPtr<SWidget> FCameraObjectInterfaceParametersToolkit::GetInterfaceParametersPanel() const
+{
+	return PanelContainer;
+}
+
+void FCameraObjectInterfaceParametersToolkit::PostUndo(bool bSuccess)
+{
+	Panel->RequestListRefresh();
+}
+
+void FCameraObjectInterfaceParametersToolkit::PostRedo(bool bSuccess)
+{
+	Panel->RequestListRefresh();
+}
+
+}  // namespace UE::Cameras
+
+#undef LOCTEXT_NAMESPACE
+
diff --git a/Engine/Plugins/Cameras/GameplayCameras/Source/GameplayCamerasEditor/Private/Toolkits/CameraRigAssetEditorToolkit.cpp b/Engine/Plugins/Cameras/GameplayCameras/Source/GameplayCamerasEditor/Private/Toolkits/CameraRigAssetEditorToolkit.cpp
--- a/Engine/Plugins/Cameras/GameplayCameras/Source/GameplayCamerasEditor/Private/Toolkits/CameraRigAssetEditorToolkit.cpp
+++ b/Engine/Plugins/Cameras/GameplayCameras/Source/GameplayCamerasEditor/Private/Toolkits/CameraRigAssetEditorToolkit.cpp
@@ -2,18 +2,17 @@
 
 #include "Toolkits/CameraRigAssetEditorToolkit.h"
 
+#include "Build/CameraBuildLog.h"
+#include "Build/CameraRigAssetBuilder.h"
 #include "Commands/CameraRigAssetEditorCommands.h"
-#include "Core/CameraBuildLog.h"
 #include "Core/CameraRigAsset.h"
-#include "Core/CameraRigAssetBuilder.h"
 #include "Customizations/RichCurveDetailsCustomizations.h"
-#include "Editors/CameraNodeGraphSchema.h"
+#include "Editors/CameraRigCameraNodeGraphSchema.h"
 #include "Editors/CameraRigTransitionGraphSchema.h"
 #include "Editors/SCameraRigAssetEditor.h"
 #include "Editors/SFindInObjectTreeGraph.h"
 #include "Framework/Docking/LayoutExtender.h"
 #include "Framework/Docking/TabManager.h"
-#include "GraphEditAction.h"
 #include "Helpers/AssetTypeMenuOverlayHelper.h"
 #include "IGameplayCamerasEditorModule.h"
 #include "IGameplayCamerasLiveEditManager.h"
@@ -23,8 +22,8 @@
 #include "ToolMenus.h"
 #include "Toolkits/BuildButtonToolkit.h"
 #include "Toolkits/CameraBuildLogToolkit.h"
+#include "Toolkits/CameraObjectInterfaceParametersToolkit.h"
 #include "Toolkits/CameraRigAssetEditorToolkitBase.h"
-#include "Toolkits/CameraRigInterfaceParametersToolkit.h"
 #include "Toolkits/CurveEditorToolkit.h"
 #include "Toolkits/StandardToolkitLayout.h"
 #include "Widgets/Docking/SDockTab.h"
@@ -48,7 +47,7 @@ FCameraRigAssetEditorToolkit::FCameraRigAssetEditorToolkit(UAssetEditor* InOwnin
 	BuildButtonToolkit = MakeShared<FBuildButtonToolkit>();
 	BuildLogToolkit = MakeShared<FCameraBuildLogToolkit>();
 	CurveEditorToolkit = MakeShared<FCurveEditorToolkit>();
-	InterfaceParametersToolkit = MakeShared<FCameraRigInterfaceParametersToolkit>();
+	InterfaceParametersToolkit = MakeShared<FCameraObjectInterfaceParametersToolkit>();
 
 	// Override base class default layout.
 	TSharedPtr<FStandardToolkitLayout> StandardLayout = Impl->GetStandardLayout();
@@ -60,8 +59,8 @@ FCameraRigAssetEditorToolkit::FCameraRigAssetEditorToolkit(UAssetEditor* InOwnin
 	}
 	StandaloneDefaultLayout = StandardLayout->GetLayout();
 
-	UClass* NodeGraphSchemaClass = UCameraNodeGraphSchema::StaticClass();
-	UCameraNodeGraphSchema* DefaultNodeGraphSchema = Cast<UCameraNodeGraphSchema>(NodeGraphSchemaClass->GetDefaultObject());
+	UClass* NodeGraphSchemaClass = UCameraRigCameraNodeGraphSchema::StaticClass();
+	UCameraRigCameraNodeGraphSchema* DefaultNodeGraphSchema = Cast<UCameraRigCameraNodeGraphSchema>(NodeGraphSchemaClass->GetDefaultObject());
 	NodeGraphConfig = DefaultNodeGraphSchema->BuildGraphConfig();
 
 	UClass* TransitionSchemaClass = UCameraRigTransitionGraphSchema::StaticClass();
@@ -80,7 +79,7 @@ void FCameraRigAssetEditorToolkit::SetCameraRigAsset(UCameraRigAsset* InCameraRi
 
 	Impl->SetCameraRigAsset(InCameraRig);
 	BuildButtonToolkit->SetTarget(InCameraRig);
-	InterfaceParametersToolkit->SetCameraRigAsset(InCameraRig);
+	InterfaceParametersToolkit->SetCameraObject(InCameraRig);
 
 	if (InCameraRig)
 	{
@@ -225,7 +224,7 @@ void FCameraRigAssetEditorToolkit::CreateWidgets()
 	BuildLogToolkit->Initialize("CameraRigAssetBuildMessages");
 
 	// Hook-up the selection of interface parameters.
-	InterfaceParametersToolkit->OnInterfaceParameterSelected().AddSP(this, &FCameraRigAssetEditorToolkit::OnCameraRigInterfaceParameterSelected);
+	InterfaceParametersToolkit->OnInterfaceParameterSelected().AddSP(this, &FCameraRigAssetEditorToolkit::OnCameraObjectInterfaceParameterSelected);
 }
 
 void FCameraRigAssetEditorToolkit::RegisterToolbar()
@@ -307,7 +306,7 @@ void FCameraRigAssetEditorToolkit::PostRegenerateMenusAndToolbars()
 	SetMenuOverlay(FAssetTypeMenuOverlayHelper::CreateMenuOverlay(UCameraRigAsset::StaticClass()));
 }
 
-void FCameraRigAssetEditorToolkit::OnCameraRigInterfaceParameterSelected(UCameraRigInterfaceParameterBase* Object)
+void FCameraRigAssetEditorToolkit::OnCameraObjectInterfaceParameterSelected(UCameraObjectInterfaceParameterBase* Object)
 {
 	OnJumpToObject(Object, NAME_None);
 }
diff --git a/Engine/Plugins/Cameras/GameplayCameras/Source/GameplayCamerasEditor/Private/Toolkits/CameraRigTransitionEditorToolkitBase.cpp b/Engine/Plugins/Cameras/GameplayCameras/Source/GameplayCamerasEditor/Private/Toolkits/CameraRigTransitionEditorToolkitBase.cpp
--- a/Engine/Plugins/Cameras/GameplayCameras/Source/GameplayCamerasEditor/Private/Toolkits/CameraRigTransitionEditorToolkitBase.cpp
+++ b/Engine/Plugins/Cameras/GameplayCameras/Source/GameplayCamerasEditor/Private/Toolkits/CameraRigTransitionEditorToolkitBase.cpp
@@ -2,8 +2,8 @@
 
 #include "Toolkits/CameraRigTransitionEditorToolkitBase.h"
 
+#include "Build/CameraBuildStatus.h"
 #include "Commands/CameraRigTransitionEditorCommands.h"
-#include "Core/CameraBuildStatus.h"
 #include "Editors/CameraRigTransitionGraphSchemaBase.h"
 #include "Editors/SCameraRigTransitionEditor.h"
 #include "Editors/SObjectTreeGraphToolbox.h"
diff --git a/Engine/Plugins/Cameras/GameplayCameras/Source/GameplayCamerasEditor/Public/Editors/CameraNodeGraphDragDropOp.h b/Engine/Plugins/Cameras/GameplayCameras/Source/GameplayCamerasEditor/Public/Editors/CameraNodeGraphDragDropOp.h
--- a/Engine/Plugins/Cameras/GameplayCameras/Source/GameplayCamerasEditor/Public/Editors/CameraNodeGraphDragDropOp.h
+++ b/Engine/Plugins/Cameras/GameplayCameras/Source/GameplayCamerasEditor/Public/Editors/CameraNodeGraphDragDropOp.h
@@ -5,7 +5,7 @@
 #include "DragAndDrop/DecoratedDragDropOp.h"
 
 class SGraphEditor;
-class UCameraRigInterfaceParameterBase;
+class UCameraObjectInterfaceParameterBase;
 
 class FCameraNodeGraphInterfaceParameterDragDropOp : public FDecoratedDragDropOp
 {
@@ -13,13 +13,13 @@ public:
 	
 	DRAG_DROP_OPERATOR_TYPE(FCameraNodeGraphInterfaceParameterDragDropOp, FDecoratedDragDropOp)
 
-	static TSharedRef<FCameraNodeGraphInterfaceParameterDragDropOp> New(UCameraRigInterfaceParameterBase* InInterfaceParameter);
+	static TSharedRef<FCameraNodeGraphInterfaceParameterDragDropOp> New(UCameraObjectInterfaceParameterBase* InInterfaceParameter);
 
 	FReply ExecuteDragOver(TSharedPtr<SGraphEditor> GraphEditor);
 	FReply ExecuteDrop(TSharedPtr<SGraphEditor> GraphEditor, const FVector2D& NewLocation);
 
 private:
 
-	UCameraRigInterfaceParameterBase* InterfaceParameter;
+	UCameraObjectInterfaceParameterBase* InterfaceParameter;
 };
 
diff --git a/Engine/Plugins/Cameras/GameplayCameras/Source/GameplayCamerasEditor/Public/Editors/CameraNodeGraphSchema.h b/Engine/Plugins/Cameras/GameplayCameras/Source/GameplayCamerasEditor/Public/Editors/CameraNodeGraphSchema.h
--- a/Engine/Plugins/Cameras/GameplayCameras/Source/GameplayCamerasEditor/Public/Editors/CameraNodeGraphSchema.h
+++ b/Engine/Plugins/Cameras/GameplayCameras/Source/GameplayCamerasEditor/Public/Editors/CameraNodeGraphSchema.h
@@ -2,9 +2,8 @@
 
 #pragma once
 
-#include "Core/CameraContextDataTableFwd.h"
-#include "Core/CameraRigParameterDefinition.h"
-#include "Core/CameraVariableTableFwd.h"
+#include "Core/CameraObjectInterface.h"
+#include "Core/CameraObjectInterfaceParameterDefinition.h"
 #include "EdGraph/EdGraphPin.h"
 #include "Editors/CameraNodeGraphPinColors.h"
 #include "Editors/ObjectTreeGraphSchema.h"
@@ -12,8 +11,8 @@
 #include "CameraNodeGraphSchema.generated.h"
 
 class UCameraNode;
-class UCameraRigInterfaceParameterBase;
-class UCameraRigInterfaceParameterGraphNode;
+class UCameraObjectInterfaceParameterBase;
+class UCameraObjectInterfaceParameterGraphNode;
 struct FObjectTreeGraphConfig;
 
 /**
@@ -32,9 +31,7 @@ public:
 
 	UCameraNodeGraphSchema(const FObjectInitializer& ObjInit);
 
-	FObjectTreeGraphConfig BuildGraphConfig() const;
-
-	UCameraRigInterfaceParameterGraphNode* CreateInterfaceParameterNode(UEdGraph* InGraph, UCameraRigInterfaceParameterBase* InterfaceParameter) const;
+	UCameraObjectInterfaceParameterGraphNode* CreateInterfaceParameterNode(UEdGraph* InGraph, UCameraObjectInterfaceParameterBase* InterfaceParameter) const;
 
 protected:
 
@@ -51,6 +48,10 @@ protected:
 	virtual void CollectAllObjects(UObjectTreeGraph* InGraph, TSet<UObject*>& OutAllObjects) const override;
 	virtual void OnCreateAllNodes(UObjectTreeGraph* InGraph, const FCreatedNodes& InCreatedNodes) const override;
 
+protected:
+
+	void BuildBaseGraphConfig(FObjectTreeGraphConfig& OutGraphConfig) const;
+
 private:
 
 	UEdGraphPin* FindPin(UEdGraphNode* InNode, const FName& InPinName, const FName& InPinCategoryName) const;
@@ -70,7 +71,7 @@ public:
 
 	/** The new parameter's definition. */
 	UPROPERTY()
-	FCameraRigParameterDefinition ParameterDefinition;
+	FCameraObjectInterfaceParameterDefinition ParameterDefinition;
 
 public:
 
@@ -91,7 +92,7 @@ struct FCameraNodeGraphSchemaAction_AddInterfaceParameterNode : public FEdGraphS
 	GENERATED_BODY()
 	
 	UPROPERTY()
-	TObjectPtr<UCameraRigInterfaceParameterBase> InterfaceParameter;
+	TObjectPtr<UCameraObjectInterfaceParameterBase> InterfaceParameter;
 
 public:
 	
diff --git a/Engine/Plugins/Cameras/GameplayCameras/Source/GameplayCamerasEditor/Public/Editors/SCameraRigAssetEditor.h b/Engine/Plugins/Cameras/GameplayCameras/Source/GameplayCamerasEditor/Public/Editors/SCameraRigAssetEditor.h
--- a/Engine/Plugins/Cameras/GameplayCameras/Source/GameplayCamerasEditor/Public/Editors/SCameraRigAssetEditor.h
+++ b/Engine/Plugins/Cameras/GameplayCameras/Source/GameplayCamerasEditor/Public/Editors/SCameraRigAssetEditor.h
@@ -37,7 +37,7 @@ enum class ECameraRigAssetEditorMode
  */
 class SCameraRigAssetEditor 
 	: public SCompoundWidget
-	, public ICameraRigAssetEventHandler
+	, public ICameraObjectEventHandler
 {
 public:
 
@@ -87,8 +87,8 @@ public:
 
 protected:
 
-	// ICameraRigAssetEventHandler interface.
-	virtual void OnCameraRigInterfaceChanged() override;
+	// ICameraObjectEventHandler interface.
+	virtual void OnCameraObjectInterfaceChanged() override;
 
 protected:
 
@@ -110,7 +110,7 @@ private:
 	TObjectPtr<UCameraRigAsset> CameraRigAsset;
 
 	/** Event handler */
-	TCameraEventHandler<ICameraRigAssetEventHandler> EventHandler;
+	TCameraEventHandler<ICameraObjectEventHandler> EventHandler;
 
 	/** Reference to the details view */
 	TSharedPtr<IDetailsView> DetailsView;
diff --git a/Engine/Plugins/Cameras/GameplayCameras/Source/GameplayCamerasEditor/Public/IGameplayCamerasEditorModule.h b/Engine/Plugins/Cameras/GameplayCameras/Source/GameplayCamerasEditor/Public/IGameplayCamerasEditorModule.h
--- a/Engine/Plugins/Cameras/GameplayCameras/Source/GameplayCamerasEditor/Public/IGameplayCamerasEditorModule.h
+++ b/Engine/Plugins/Cameras/GameplayCameras/Source/GameplayCamerasEditor/Public/IGameplayCamerasEditorModule.h
@@ -2,7 +2,7 @@
 
 #pragma once
 
-#include "Core/CameraBuildLog.h"
+#include "Build/CameraBuildLog.h"
 #include "CoreTypes.h"
 #include "Delegates/Delegate.h"
 #include "Modules/ModuleInterface.h"
diff --git a/Engine/Plugins/Cameras/GameplayCameras/Source/GameplayCamerasEditor/Public/Toolkits/CameraRigAssetEditorToolkit.h b/Engine/Plugins/Cameras/GameplayCameras/Source/GameplayCamerasEditor/Public/Toolkits/CameraRigAssetEditorToolkit.h
--- a/Engine/Plugins/Cameras/GameplayCameras/Source/GameplayCamerasEditor/Public/Toolkits/CameraRigAssetEditorToolkit.h
+++ b/Engine/Plugins/Cameras/GameplayCameras/Source/GameplayCamerasEditor/Public/Toolkits/CameraRigAssetEditorToolkit.h
@@ -19,8 +19,8 @@ namespace UE::Cameras
 
 class FBuildButtonToolkit;
 class FCameraBuildLogToolkit;
+class FCameraObjectInterfaceParametersToolkit;
 class FCameraRigAssetEditorToolkitBase;
-class FCameraRigInterfaceParametersToolkit;
 class FCurveEditorToolkit;
 class IGameplayCamerasLiveEditManager;
 
@@ -65,7 +65,7 @@ private:
 
 	void OnCurvesTabClosed(TSharedRef<SDockTab> InTab);
 
-	void OnCameraRigInterfaceParameterSelected(UCameraRigInterfaceParameterBase* Object);
+	void OnCameraObjectInterfaceParameterSelected(UCameraObjectInterfaceParameterBase* Object);
 
 	void OnBuild();
 	void OnFindInCameraRig();
@@ -97,7 +97,7 @@ private:
 	/** The curve editor */
 	TSharedPtr<FCurveEditorToolkit> CurveEditorToolkit;
 	/** The interface parameters panel */
-	TSharedPtr<FCameraRigInterfaceParametersToolkit> InterfaceParametersToolkit;
+	TSharedPtr<FCameraObjectInterfaceParametersToolkit> InterfaceParametersToolkit;
 
 	/** Search widget */
 	TSharedPtr<SFindInObjectTreeGraph> SearchWidget;
diff --git a/Engine/Plugins/Cameras/GameplayCameras/Source/GameplayCamerasUncookedOnly/Private/BlueprintGraph/K2Node_CameraRigBase.cpp b/Engine/Plugins/Cameras/GameplayCameras/Source/GameplayCamerasUncookedOnly/Private/BlueprintGraph/K2Node_CameraRigBase.cpp
--- a/Engine/Plugins/Cameras/GameplayCameras/Source/GameplayCamerasUncookedOnly/Private/BlueprintGraph/K2Node_CameraRigBase.cpp
+++ b/Engine/Plugins/Cameras/GameplayCameras/Source/GameplayCamerasUncookedOnly/Private/BlueprintGraph/K2Node_CameraRigBase.cpp
@@ -131,7 +131,7 @@ void UK2Node_CameraRigBase::OnCameraRigAssetBuilt(const UCameraRigAsset* InBuilt
 	}
 }
 
-FEdGraphPinType UK2Node_CameraRigBase::MakeBlendableParameterPinType(const UCameraRigBlendableParameter* BlendableParameter)
+FEdGraphPinType UK2Node_CameraRigBase::MakeBlendableParameterPinType(const UCameraObjectInterfaceBlendableParameter* BlendableParameter)
 {
 	return MakeBlendableParameterPinType(BlendableParameter->ParameterType, BlendableParameter->BlendableStructType);
 }
@@ -191,7 +191,7 @@ FEdGraphPinType UK2Node_CameraRigBase::MakeBlendableParameterPinType(ECameraVari
 	return PinType;
 }
 
-FEdGraphPinType UK2Node_CameraRigBase::MakeDataParameterPinType(const UCameraRigDataParameter* DataParameter)
+FEdGraphPinType UK2Node_CameraRigBase::MakeDataParameterPinType(const UCameraObjectInterfaceDataParameter* DataParameter)
 {
 	return MakeDataParameterPinType(DataParameter->DataType, DataParameter->DataContainerType, DataParameter->DataTypeObject);
 }
@@ -238,7 +238,7 @@ FEdGraphPinType UK2Node_CameraRigBase::MakeDataParameterPinType(ECameraContextDa
 	return PinType;
 }
 
-FName UK2Node_CameraRigBase::GetBlendableParameterInteropSettingFunctionName(const UCameraRigBlendableParameter* BlendableParameter)
+FName UK2Node_CameraRigBase::GetBlendableParameterInteropSettingFunctionName(const UCameraObjectInterfaceBlendableParameter* BlendableParameter)
 {
 	if (ensure(BlendableParameter && BlendableParameter->PrivateVariableID))
 	{
@@ -288,7 +288,7 @@ FName UK2Node_CameraRigBase::GetBlendableParameterInteropSettingFunctionName(ECa
 	return CallSetParameterFuncName;
 }
 
-FName UK2Node_CameraRigBase::GetDataParameterInteropSettingFunctionName(const UCameraRigDataParameter* DataParameter)
+FName UK2Node_CameraRigBase::GetDataParameterInteropSettingFunctionName(const UCameraObjectInterfaceDataParameter* DataParameter)
 {
 	if (ensure(DataParameter && DataParameter->PrivateDataID.IsValid()))
 	{
diff --git a/Engine/Plugins/Cameras/GameplayCameras/Source/GameplayCamerasUncookedOnly/Private/BlueprintGraph/K2Node_CameraRigBase.h b/Engine/Plugins/Cameras/GameplayCameras/Source/GameplayCamerasUncookedOnly/Private/BlueprintGraph/K2Node_CameraRigBase.h
--- a/Engine/Plugins/Cameras/GameplayCameras/Source/GameplayCamerasUncookedOnly/Private/BlueprintGraph/K2Node_CameraRigBase.h
+++ b/Engine/Plugins/Cameras/GameplayCameras/Source/GameplayCamerasUncookedOnly/Private/BlueprintGraph/K2Node_CameraRigBase.h
@@ -12,8 +12,8 @@
 #include "K2Node_CameraRigBase.generated.h"
 
 class UCameraRigAsset;
-class UCameraRigBlendableParameter;
-class UCameraRigDataParameter;
+class UCameraObjectInterfaceBlendableParameter;
+class UCameraObjectInterfaceDataParameter;
 
 /**
  * Utility base class for Blueprint nodes that can set camera rig parameters.
@@ -44,14 +44,14 @@ public:
 
 public:
 
-	static FEdGraphPinType MakeBlendableParameterPinType(const UCameraRigBlendableParameter* BlendableParameter);
+	static FEdGraphPinType MakeBlendableParameterPinType(const UCameraObjectInterfaceBlendableParameter* BlendableParameter);
 	static FEdGraphPinType MakeBlendableParameterPinType(ECameraVariableType CameraVariableType, const UScriptStruct* BlendableStructType);
-	static FEdGraphPinType MakeDataParameterPinType(const UCameraRigDataParameter* DataParameter);
+	static FEdGraphPinType MakeDataParameterPinType(const UCameraObjectInterfaceDataParameter* DataParameter);
 	static FEdGraphPinType MakeDataParameterPinType(ECameraContextDataType CameraContextDataType, ECameraContextDataContainerType CameraContextDataContainerType, const UObject* CameraContextDataTypeObject);
 
-	static FName GetBlendableParameterInteropSettingFunctionName(const UCameraRigBlendableParameter* BlendableParameter);
+	static FName GetBlendableParameterInteropSettingFunctionName(const UCameraObjectInterfaceBlendableParameter* BlendableParameter);
 	static FName GetBlendableParameterInteropSettingFunctionName(ECameraVariableType CameraVariableType);
-	static FName GetDataParameterInteropSettingFunctionName(const UCameraRigDataParameter* DataParameter);
+	static FName GetDataParameterInteropSettingFunctionName(const UCameraObjectInterfaceDataParameter* DataParameter);
 	static FName GetDataParameterInteropSettingFunctionName(ECameraContextDataType CameraContextDataType, ECameraContextDataContainerType CameraContextDataContainerType, const UObject* CameraContextDataTypeObject);
 
 protected:
diff --git a/Engine/Plugins/Cameras/GameplayCameras/Source/GameplayCamerasUncookedOnly/Private/BlueprintGraph/K2Node_SetCameraRigParameter.cpp b/Engine/Plugins/Cameras/GameplayCameras/Source/GameplayCamerasUncookedOnly/Private/BlueprintGraph/K2Node_SetCameraRigParameter.cpp
--- a/Engine/Plugins/Cameras/GameplayCameras/Source/GameplayCamerasUncookedOnly/Private/BlueprintGraph/K2Node_SetCameraRigParameter.cpp
+++ b/Engine/Plugins/Cameras/GameplayCameras/Source/GameplayCamerasUncookedOnly/Private/BlueprintGraph/K2Node_SetCameraRigParameter.cpp
@@ -28,11 +28,11 @@ void UK2Node_SetCameraRigParameter::Initialize(const FAssetData& UnloadedCameraR
 	UCameraRigAsset* LoadedCameraRig = Cast<UCameraRigAsset>(UnloadedCameraRig.GetAsset());
 	if (ensure(LoadedCameraRig))
 	{
-		if (const UCameraRigBlendableParameter* BlendableParameter = LoadedCameraRig->Interface.FindBlendableParameterByName(InCameraParameterName))
+		if (const UCameraObjectInterfaceBlendableParameter* BlendableParameter = LoadedCameraRig->Interface.FindBlendableParameterByName(InCameraParameterName))
 		{
 			Initialize(LoadedCameraRig, InCameraParameterName, BlendableParameter->ParameterType, BlendableParameter->BlendableStructType);
 		}
-		else if (const UCameraRigDataParameter* DataParameter = LoadedCameraRig->Interface.FindDataParameterByName(InCameraParameterName))
+		else if (const UCameraObjectInterfaceDataParameter* DataParameter = LoadedCameraRig->Interface.FindDataParameterByName(InCameraParameterName))
 		{
 			Initialize(LoadedCameraRig, InCameraParameterName, DataParameter->DataType, DataParameter->DataContainerType, DataParameter->DataTypeObject);
 		}
diff --git a/Engine/Plugins/Cameras/GameplayCameras/Source/GameplayCamerasUncookedOnly/Private/BlueprintGraph/K2Node_SetCameraRigParameters.cpp b/Engine/Plugins/Cameras/GameplayCameras/Source/GameplayCamerasUncookedOnly/Private/BlueprintGraph/K2Node_SetCameraRigParameters.cpp
--- a/Engine/Plugins/Cameras/GameplayCameras/Source/GameplayCamerasUncookedOnly/Private/BlueprintGraph/K2Node_SetCameraRigParameters.cpp
+++ b/Engine/Plugins/Cameras/GameplayCameras/Source/GameplayCamerasUncookedOnly/Private/BlueprintGraph/K2Node_SetCameraRigParameters.cpp
@@ -51,7 +51,7 @@ void UK2Node_SetCameraRigParameters::ReallocatePinsDuringReconstruction(TArray<U
 	if (CameraRig)
 	{
 		PreloadObject(CameraRig);
-		for (UCameraRigBlendableParameter* BlendableParameter : CameraRig->Interface.BlendableParameters)
+		for (UCameraObjectInterfaceBlendableParameter* BlendableParameter : CameraRig->Interface.BlendableParameters)
 		{
 			PreloadObject(BlendableParameter);
 			if (BlendableParameter)
@@ -59,7 +59,7 @@ void UK2Node_SetCameraRigParameters::ReallocatePinsDuringReconstruction(TArray<U
 				PreloadObject(const_cast<UScriptStruct*>(BlendableParameter->BlendableStructType.Get()));
 			}
 		}
-		for (UCameraRigDataParameter* DataParameter : CameraRig->Interface.DataParameters)
+		for (UCameraObjectInterfaceDataParameter* DataParameter : CameraRig->Interface.DataParameters)
 		{
 			PreloadObject(DataParameter);
 			if (DataParameter)
@@ -145,7 +145,7 @@ void UK2Node_SetCameraRigParameters::ExpandNode(FKismetCompilerContext& Compiler
 	FindBlendableParameterPins(BlendableParameterPins);
 	for (UEdGraphPin* RigParameterPin : BlendableParameterPins)
 	{
-		const UCameraRigBlendableParameter* BlendableParameter = CameraRig->Interface.FindBlendableParameterByName(RigParameterPin->GetName());
+		const UCameraObjectInterfaceBlendableParameter* BlendableParameter = CameraRig->Interface.FindBlendableParameterByName(RigParameterPin->GetName());
 		if (!BlendableParameter)
 		{
 			CompilerContext.MessageLog.Error(*LOCTEXT("ErrorMissingParameter", "SetCameraRigParameters node @@ is trying to set parameter @@ but camera rig @@ has no such parameter.").ToString(), this, *RigParameterPin->GetName(), CameraRig);
@@ -213,7 +213,7 @@ void UK2Node_SetCameraRigParameters::ExpandNode(FKismetCompilerContext& Compiler
 	FindDataParameterPins(DataParameterPins);
 	for (UEdGraphPin* RigParameterPin : DataParameterPins)
 	{
-		const UCameraRigDataParameter* DataParameter = CameraRig->Interface.FindDataParameterByName(RigParameterPin->GetName());
+		const UCameraObjectInterfaceDataParameter* DataParameter = CameraRig->Interface.FindDataParameterByName(RigParameterPin->GetName());
 		if (!DataParameter)
 		{
 			CompilerContext.MessageLog.Error(*LOCTEXT("ErrorMissingParameter", "SetCameraRigParameters node @@ is trying to set parameter @@ but camera rig @@ has no such parameter.").ToString(), this, *RigParameterPin->GetName(), CameraRig);
@@ -323,7 +323,7 @@ void UK2Node_SetCameraRigParameters::CreateParameterPins()
 		return;
 	}
 
-	for (const UCameraRigBlendableParameter* BlendableParameter : CameraRig->Interface.BlendableParameters)
+	for (const UCameraObjectInterfaceBlendableParameter* BlendableParameter : CameraRig->Interface.BlendableParameters)
 	{
 		if (!ensure(BlendableParameter))
 		{
@@ -347,7 +347,7 @@ void UK2Node_SetCameraRigParameters::CreateParameterPins()
 		BlendableParameterPinNames.Add(NewPin->PinName);
 	}
 
-	for (const UCameraRigDataParameter* DataParameter : CameraRig->Interface.DataParameters)
+	for (const UCameraObjectInterfaceDataParameter* DataParameter : CameraRig->Interface.DataParameters)
 	{
 		if (!ensure(DataParameter))
 		{
