Cameras: add auto-reset option for context data

diff --git a/Engine/Plugins/Cameras/GameplayCameras/Source/GameplayCameras/Private/Core/CameraContextDataTable.cpp b/Engine/Plugins/Cameras/GameplayCameras/Source/GameplayCameras/Private/Core/CameraContextDataTable.cpp
--- a/Engine/Plugins/Cameras/GameplayCameras/Source/GameplayCameras/Private/Core/CameraContextDataTable.cpp
+++ b/Engine/Plugins/Cameras/GameplayCameras/Source/GameplayCameras/Private/Core/CameraContextDataTable.cpp
@@ -116,6 +116,10 @@ void FCameraContextDataTable::Initialize(const FCameraContextDataTableAllocation
 		NewEntry.TypeObject = DataDefinition.DataTypeObject;
 		NewEntry.Offset = NewEntryOffset;
 		NewEntry.Flags = EEntryFlags::None;
+		if (DataDefinition.bAutoReset)
+		{
+			NewEntry.Flags |= EEntryFlags::AutoReset;
+		}
 #if WITH_EDITORONLY_DATA
 		NewEntry.DebugName = DataDefinition.DataName;
 #endif
@@ -166,6 +170,10 @@ void FCameraContextDataTable::AddData(const FCameraContextDataDefinition& DataDe
 	NewEntry.TypeObject = DataDefinition.DataTypeObject;
 	NewEntry.Offset = DataPtr - Memory;
 	NewEntry.Flags = EEntryFlags::None;
+	if (DataDefinition.bAutoReset)
+	{
+		NewEntry.Flags |= EEntryFlags::AutoReset;
+	}
 #if WITH_EDITORONLY_DATA
 	NewEntry.DebugName = DataDefinition.DataName;
 #endif
@@ -879,6 +887,17 @@ void FCameraContextDataTable::ClearAllWrittenThisFrameFlags()
 	}
 }
 
+void FCameraContextDataTable::AutoResetValues()
+{
+	for (FEntry& Entry : Entries)
+	{
+		if (EnumHasAnyFlags(Entry.Flags, EEntryFlags::AutoReset))
+		{
+			EnumRemoveFlags(Entry.Flags, EEntryFlags::Written | EEntryFlags::WrittenThisFrame);
+		}
+	}
+}
+
 void FCameraContextDataTable::OverrideAll(const FCameraContextDataTable& OtherTable)
 {
 	InternalOverride(OtherTable, ECameraContextDataTableFilter::None);
diff --git a/Engine/Plugins/Cameras/GameplayCameras/Source/GameplayCameras/Private/Core/CameraEvaluationContext.cpp b/Engine/Plugins/Cameras/GameplayCameras/Source/GameplayCameras/Private/Core/CameraEvaluationContext.cpp
--- a/Engine/Plugins/Cameras/GameplayCameras/Source/GameplayCameras/Private/Core/CameraEvaluationContext.cpp
+++ b/Engine/Plugins/Cameras/GameplayCameras/Source/GameplayCameras/Private/Core/CameraEvaluationContext.cpp
@@ -143,8 +143,11 @@ void FCameraEvaluationContext::AddReferencedObjects(FReferenceCollector& Collect
 void FCameraEvaluationContext::OnEndCameraSystemUpdate()
 {
 	InitialResult.CameraPose.ClearAllChangedFlags();
+
 	InitialResult.VariableTable.AutoResetValues();
 	InitialResult.VariableTable.ClearAllWrittenThisFrameFlags();
+
+	InitialResult.ContextDataTable.AutoResetValues();
 	InitialResult.ContextDataTable.ClearAllWrittenThisFrameFlags();
 
 	for (FConditionalResults::ElementType& Pair : ConditionalResults)
@@ -153,6 +156,7 @@ void FCameraEvaluationContext::OnEndCameraSystemUpdate()
 		Result.CameraPose.ClearAllChangedFlags();
 		Result.VariableTable.AutoResetValues();
 		Result.VariableTable.ClearAllWrittenThisFrameFlags();
+		Result.ContextDataTable.AutoResetValues();
 		Result.ContextDataTable.ClearAllWrittenThisFrameFlags();
 	}
 
diff --git a/Engine/Plugins/Cameras/GameplayCameras/Source/GameplayCameras/Private/Core/CameraSystemEvaluator.cpp b/Engine/Plugins/Cameras/GameplayCameras/Source/GameplayCameras/Private/Core/CameraSystemEvaluator.cpp
--- a/Engine/Plugins/Cameras/GameplayCameras/Source/GameplayCameras/Private/Core/CameraSystemEvaluator.cpp
+++ b/Engine/Plugins/Cameras/GameplayCameras/Source/GameplayCameras/Private/Core/CameraSystemEvaluator.cpp
@@ -210,6 +210,7 @@ void FCameraSystemEvaluator::UpdateImpl(float DeltaTime, ECameraNodeEvaluationTy
 
 	// Reset variables and data.
 	RootNodeResult.VariableTable.AutoResetValues();
+	RootNodeResult.ContextDataTable.AutoResetValues();
 
 	// Pre-update all services.
 	PreUpdateServices(DeltaTime, ECameraEvaluationServiceFlags::None);
diff --git a/Engine/Plugins/Cameras/GameplayCameras/Source/GameplayCameras/Public/Core/CameraContextDataTable.h b/Engine/Plugins/Cameras/GameplayCameras/Source/GameplayCameras/Public/Core/CameraContextDataTable.h
--- a/Engine/Plugins/Cameras/GameplayCameras/Source/GameplayCameras/Public/Core/CameraContextDataTable.h
+++ b/Engine/Plugins/Cameras/GameplayCameras/Source/GameplayCameras/Public/Core/CameraContextDataTable.h
@@ -203,11 +203,14 @@ public:
 	bool IsValueWrittenThisFrame(FCameraContextDataID InID) const;
 	void ClearAllWrittenThisFrameFlags();
 
+	void AutoResetValues();
+
 private:
 
 	enum class EEntryFlags : uint8
 	{
 		None = 0,
+		AutoReset = 1 << 1,
 		Written = 1 << 2,
 		WrittenThisFrame = 1 << 3
 	};
diff --git a/Engine/Plugins/Cameras/GameplayCameras/Source/GameplayCameras/Public/Core/CameraContextDataTableAllocationInfo.h b/Engine/Plugins/Cameras/GameplayCameras/Source/GameplayCameras/Public/Core/CameraContextDataTableAllocationInfo.h
--- a/Engine/Plugins/Cameras/GameplayCameras/Source/GameplayCameras/Public/Core/CameraContextDataTableAllocationInfo.h
+++ b/Engine/Plugins/Cameras/GameplayCameras/Source/GameplayCameras/Public/Core/CameraContextDataTableAllocationInfo.h
@@ -35,6 +35,10 @@ struct FCameraContextDataDefinition
 	UPROPERTY()
 	TObjectPtr<const UObject> DataTypeObject;
 
+	/** Whether the data should auto-reset to an "unset" state after every evaluation. */
+	UPROPERTY()
+	bool bAutoReset = false;
+
 #if WITH_EDITORONLY_DATA
 	/** The name of the data, for debugging purposes. */
 	UPROPERTY()
diff --git a/Engine/Plugins/Cameras/GameplayCameras/Source/GameplayCameras/Public/Core/CameraVariableTable.h b/Engine/Plugins/Cameras/GameplayCameras/Source/GameplayCameras/Public/Core/CameraVariableTable.h
--- a/Engine/Plugins/Cameras/GameplayCameras/Source/GameplayCameras/Public/Core/CameraVariableTable.h
+++ b/Engine/Plugins/Cameras/GameplayCameras/Source/GameplayCameras/Public/Core/CameraVariableTable.h
@@ -234,7 +234,7 @@ private:
 		ECameraVariableType Type;
 		const UScriptStruct* StructType = nullptr;
 		uint32 Offset;
-		mutable EEntryFlags Flags;
+		EEntryFlags Flags;
 #if WITH_EDITORONLY_DATA
 		FString DebugName;
 #endif
diff --git a/Engine/Plugins/Cameras/GameplayCameras/Source/GameplayCameras/Public/Core/CameraVariableTableAllocationInfo.h b/Engine/Plugins/Cameras/GameplayCameras/Source/GameplayCameras/Public/Core/CameraVariableTableAllocationInfo.h
--- a/Engine/Plugins/Cameras/GameplayCameras/Source/GameplayCameras/Public/Core/CameraVariableTableAllocationInfo.h
+++ b/Engine/Plugins/Cameras/GameplayCameras/Source/GameplayCameras/Public/Core/CameraVariableTableAllocationInfo.h
@@ -44,11 +44,7 @@ struct FCameraVariableDefinition
 	UPROPERTY()
 	bool bIsInput = false;
 
-	/**
-	 * Whether the variable should auto-reset to an "unset" state after every
-	 * evaluation. Doing so resets it to its default value if bound to a
-	 * camera variable or camera parameter, or to a non-existent state otherwise.
-	 */
+	/** Whether the variable should auto-reset to an "unset" state after every evaluation. */
 	UPROPERTY()
 	bool bAutoReset = false;
 
