Cameras: add camera shake assets, editor, and nodes

Camera shake assets are another type of "camera object", like camera rigs. They have interface parameters like camera rigs, but they have a shake camera node as their root node, and simple fixed-time blends for ease-in and ease-out.

The camera shake asset editor has a graph editor like the camera rig, without the "transitions" editor mode.

A few shake camera nodes are implemented, like perlin noise, composite shakes, and envelope.

A shake prefab node allows running a camera shake in the visual stack with pre-blended parameters.

diff --git a/Engine/Plugins/Cameras/GameplayCameras/Source/GameplayCameras/Private/Build/CameraShakeAssetBuilder.cpp b/Engine/Plugins/Cameras/GameplayCameras/Source/GameplayCameras/Private/Build/CameraShakeAssetBuilder.cpp
new file mode 100644
--- /dev/null
+++ b/Engine/Plugins/Cameras/GameplayCameras/Source/GameplayCameras/Private/Build/CameraShakeAssetBuilder.cpp
@@ -0,0 +1,83 @@
+// Copyright Epic Games, Inc. All Rights Reserved.
+
+#include "Build/CameraShakeAssetBuilder.h"
+
+#include "Build/CameraNodeHierarchyBuilder.h"
+#include "Build/CameraObjectInterfaceBuilder.h"
+#include "Build/CameraObjectInterfaceParameterBuilder.h"
+#include "Core/CameraShakeAsset.h"
+#include "GameplayCamerasDelegates.h"
+
+#define LOCTEXT_NAMESPACE "CameraShakeAssetBuilder"
+
+namespace UE::Cameras
+{
+
+FCameraShakeAssetBuilder::FCameraShakeAssetBuilder(FCameraBuildLog& InBuildLog)
+	: BuildLog(InBuildLog)
+{
+}
+
+void FCameraShakeAssetBuilder::BuildCameraShake(UCameraShakeAsset* InCameraShake)
+{
+	if (!ensure(InCameraShake))
+	{
+		return;
+	}
+
+	CameraShake = InCameraShake;
+
+	BuildLog.SetLoggingPrefix(CameraShake->GetPathName() + TEXT(": "));
+	{
+		BuildCameraShakeImpl();
+	}
+	BuildLog.SetLoggingPrefix(FString());
+	UpdateBuildStatus();
+
+	FGameplayCamerasDelegates::OnCameraShakeAssetBuilt().Broadcast(CameraShake);
+
+	CameraShake = nullptr;
+}
+
+void FCameraShakeAssetBuilder::BuildCameraShakeImpl()
+{
+	if (!CameraShake->RootNode)
+	{
+		BuildLog.AddMessage(EMessageSeverity::Error, CameraShake, 
+				FText::Format(LOCTEXT("MissingRootNode", "Camera shake '{0}' has no root node."), 
+					FText::FromString(GetPathNameSafe(CameraShake))));
+		return;
+	}
+
+	FCameraNodeHierarchyBuilder NodeBuilder(BuildLog, CameraShake);
+	NodeBuilder.PreBuild();
+
+	FCameraObjectInterfaceBuilder InterfaceBuilder(BuildLog);
+	InterfaceBuilder.BuildInterface(CameraShake, NodeBuilder.GetHierarchy(), true);
+
+	NodeBuilder.Build();
+
+	FCameraObjectInterfaceParameterBuilder ParameterBuilder;
+	ParameterBuilder.BuildParameters(CameraShake);
+}
+
+void FCameraShakeAssetBuilder::UpdateBuildStatus()
+{
+	ECameraBuildStatus BuildStatus = ECameraBuildStatus::Clean;
+	if (BuildLog.HasErrors())
+	{
+		BuildStatus = ECameraBuildStatus::WithErrors;
+	}
+	else if (BuildLog.HasWarnings())
+	{
+		BuildStatus = ECameraBuildStatus::CleanWithWarnings;
+	}
+
+	// Don't modify the camera shake: BuildStatus is transient.
+	CameraShake->BuildStatus = BuildStatus;
+}
+
+}  // namespace UE::Cameras
+
+#undef LOCTEXT_NAMESPACE
+
diff --git a/Engine/Plugins/Cameras/GameplayCameras/Source/GameplayCameras/Private/Core/CameraShakeAsset.cpp b/Engine/Plugins/Cameras/GameplayCameras/Source/GameplayCameras/Private/Core/CameraShakeAsset.cpp
new file mode 100644
--- /dev/null
+++ b/Engine/Plugins/Cameras/GameplayCameras/Source/GameplayCameras/Private/Core/CameraShakeAsset.cpp
@@ -0,0 +1,142 @@
+// Copyright Epic Games, Inc. All Rights Reserved.
+
+#include "Core/CameraShakeAsset.h"
+
+#include "Build/CameraBuildLog.h"
+#include "Build/CameraShakeAssetBuilder.h"
+#include "Core/ShakeCameraNode.h"
+#include "UObject/AssetRegistryTagsContext.h"
+#include "UObject/ObjectSaveContext.h"
+
+#include UE_INLINE_GENERATED_CPP_BY_NAME(CameraShakeAsset)
+
+void UCameraShakeAsset::PostInitProperties()
+{
+	Super::PostInitProperties();
+
+	if (!HasAnyFlags(RF_ClassDefaultObject | RF_ArchetypeObject | RF_NeedLoad | RF_WasLoaded) && 
+			!Guid.IsValid())
+	{
+		Guid = FGuid::NewGuid();
+	}
+}
+
+void UCameraShakeAsset::PostDuplicate(EDuplicateMode::Type DuplicateMode)
+{
+	Super::PostDuplicate(DuplicateMode);
+
+	if (DuplicateMode == EDuplicateMode::Normal)
+	{
+		Guid = FGuid::NewGuid();
+	}
+}
+
+void UCameraShakeAsset::BuildCameraShake()
+{
+	using namespace UE::Cameras;
+
+	FCameraBuildLog BuildLog;
+	BuildLog.SetForwardMessagesToLogging(true);
+	BuildCameraShake(BuildLog);
+}
+
+void UCameraShakeAsset::BuildCameraShake(UE::Cameras::FCameraBuildLog& InBuildLog)
+{
+	using namespace UE::Cameras;
+
+	FCameraShakeAssetBuilder Builder(InBuildLog);
+	Builder.BuildCameraShake(this);
+}
+
+UCameraNode* UCameraShakeAsset::GetRootNode()
+{
+	return RootNode;
+}
+
+void UCameraShakeAsset::DirtyBuildStatus()
+{
+	BuildStatus = ECameraBuildStatus::Dirty;
+}
+
+void UCameraShakeAsset::PreSave(FObjectPreSaveContext ObjectSaveContext)
+{
+#if WITH_EDITOR
+
+	const bool bIsUserObject = !HasAnyFlags(RF_ClassDefaultObject | RF_ArchetypeObject);
+	const bool bIsEditorAutoSave = ((ObjectSaveContext.GetSaveFlags() & SAVE_FromAutosave) != 0);
+	if (bIsUserObject && !bIsEditorAutoSave)
+	{
+		// Build when saving/cooking.
+		BuildCameraShake();
+	}
+
+#endif
+
+	Super::PreSave(ObjectSaveContext);
+}
+
+#if WITH_EDITOR
+
+void UCameraShakeAsset::GetGraphNodePosition(FName InGraphName, int32& NodePosX, int32& NodePosY) const
+{
+	NodePosX = GraphNodePos.X;
+	NodePosY = GraphNodePos.Y;
+}
+
+void UCameraShakeAsset::OnGraphNodeMoved(FName InGraphName, int32 NodePosX, int32 NodePosY, bool bMarkDirty)
+{
+	Modify(bMarkDirty);
+
+	GraphNodePos.X = NodePosX;
+	GraphNodePos.Y = NodePosY;
+}
+
+EObjectTreeGraphObjectSupportFlags UCameraShakeAsset::GetSupportFlags(FName InGraphName) const
+{
+	return EObjectTreeGraphObjectSupportFlags::CommentText;
+}
+
+const FString& UCameraShakeAsset::GetGraphNodeCommentText(FName InGraphName) const
+{
+	return GraphNodeComment;
+}
+
+void UCameraShakeAsset::OnUpdateGraphNodeCommentText(FName InGraphName, const FString& NewComment)
+{
+	Modify();
+
+	GraphNodeComment = NewComment;
+}
+
+void UCameraShakeAsset::GetGraphNodeName(FName InGraphName, FText& OutName) const
+{
+	OutName = FText::FromString(GetName());
+}
+
+void UCameraShakeAsset::GetConnectableObjects(FName InGraphName, TSet<UObject*>& OutObjects) const
+{
+	OutObjects.Append(AllNodeObjects);
+}
+
+void UCameraShakeAsset::AddConnectableObject(FName InGraphName, UObject* InObject)
+{
+	using namespace UE::Cameras;
+
+	Modify();
+
+	const int32 Index = AllNodeObjects.AddUnique(InObject);
+	ensure(Index == AllNodeObjects.Num() - 1);
+}
+
+void UCameraShakeAsset::RemoveConnectableObject(FName InGraphName, UObject* InObject)
+{
+	using namespace UE::Cameras;
+
+	Modify();
+
+	const int32 NumRemoved = AllNodeObjects.Remove(InObject);
+	ensure(NumRemoved == 1);
+}
+
+#endif  // WITH_EDITOR
+
diff --git a/Engine/Plugins/Cameras/GameplayCameras/Source/GameplayCameras/Private/Core/CameraShakeAssetReference.cpp b/Engine/Plugins/Cameras/GameplayCameras/Source/GameplayCameras/Private/Core/CameraShakeAssetReference.cpp
new file mode 100644
--- /dev/null
+++ b/Engine/Plugins/Cameras/GameplayCameras/Source/GameplayCameras/Private/Core/CameraShakeAssetReference.cpp
@@ -0,0 +1,30 @@
+// Copyright Epic Games, Inc. All Rights Reserved.
+
+#include "Core/CameraShakeAssetReference.h"
+
+#include "Core/CameraShakeAsset.h"
+#include "Helpers/CameraObjectReferenceParameterOverrideEvaluator.h"
+
+#include UE_INLINE_GENERATED_CPP_BY_NAME(CameraShakeAssetReference)
+
+FCameraShakeAssetReference::FCameraShakeAssetReference()
+{
+}
+
+FCameraShakeAssetReference::FCameraShakeAssetReference(UCameraShakeAsset* InCameraShake)
+	: CameraShake(InCameraShake)
+{
+}
+
+const UBaseCameraObject* FCameraShakeAssetReference::GetCameraObject() const
+{
+	return CameraShake;
+}
+
+void FCameraShakeAssetReference::ApplyParameterOverrides(UE::Cameras::FCameraNodeEvaluationResult& OutResult, bool bDrivenOverridesOnly)
+{
+	using namespace UE::Cameras;
+	FCameraObjectReferenceParameterOverrideEvaluator OverrideEvaluator(*this);
+	OverrideEvaluator.ApplyParameterOverrides(OutResult.VariableTable, OutResult.ContextDataTable, bDrivenOverridesOnly);
+}
+
diff --git a/Engine/Plugins/Cameras/GameplayCameras/Source/GameplayCameras/Private/Core/CameraSystemEvaluator.cpp b/Engine/Plugins/Cameras/GameplayCameras/Source/GameplayCameras/Private/Core/CameraSystemEvaluator.cpp
--- a/Engine/Plugins/Cameras/GameplayCameras/Source/GameplayCameras/Private/Core/CameraSystemEvaluator.cpp
+++ b/Engine/Plugins/Cameras/GameplayCameras/Source/GameplayCameras/Private/Core/CameraSystemEvaluator.cpp
@@ -16,6 +16,7 @@
 #include "GameplayCamerasSettings.h"
 #include "Math/ColorList.h"
 #include "Services/CameraModifierService.h"
+#include "Services/CameraShakeService.h"
 #include "Services/OrientationInitializationService.h"
 #include "UObject/Package.h"
 #include "UObject/UObjectGlobals.h"
@@ -93,6 +94,7 @@ void FCameraSystemEvaluator::Initialize(const FCameraSystemEvaluatorCreateParams
 	}
 
 	RegisterEvaluationService(MakeShared<FCameraModifierService>());
+	RegisterEvaluationService(MakeShared<FCameraShakeService>());
 	RegisterEvaluationService(MakeShared<FOrientationInitializationService>());
 
 	CameraRigCombinationRegistry = MakeShared<FCameraRigCombinationRegistry>();
diff --git a/Engine/Plugins/Cameras/GameplayCameras/Source/GameplayCameras/Private/Core/DefaultRootCameraNode.cpp b/Engine/Plugins/Cameras/GameplayCameras/Source/GameplayCameras/Private/Core/DefaultRootCameraNode.cpp
--- a/Engine/Plugins/Cameras/GameplayCameras/Source/GameplayCameras/Private/Core/DefaultRootCameraNode.cpp
+++ b/Engine/Plugins/Cameras/GameplayCameras/Source/GameplayCameras/Private/Core/DefaultRootCameraNode.cpp
@@ -193,6 +193,11 @@ void FDefaultRootCameraNodeEvaluator::OnGetCameraRigInfo(const FCameraRigInstanc
 	}
 }
 
+const FCameraVariableTable* FDefaultRootCameraNodeEvaluator::OnGetBlendedParameters() const
+{
+	return &MainLayer->GetBlendedParameters();
+}
+
 void FDefaultRootCameraNodeEvaluator::OnBuildSingleCameraRigHierarchy(const FSingleCameraRigHierarchyBuildParams& Params, FCameraNodeEvaluatorHierarchy& OutHierarchy)
 {
 	OutHierarchy.Build(BaseLayer);
diff --git a/Engine/Plugins/Cameras/GameplayCameras/Source/GameplayCameras/Private/Core/RootCameraNode.cpp b/Engine/Plugins/Cameras/GameplayCameras/Source/GameplayCameras/Private/Core/RootCameraNode.cpp
--- a/Engine/Plugins/Cameras/GameplayCameras/Source/GameplayCameras/Private/Core/RootCameraNode.cpp
+++ b/Engine/Plugins/Cameras/GameplayCameras/Source/GameplayCameras/Private/Core/RootCameraNode.cpp
@@ -43,6 +43,11 @@ bool FRootCameraNodeEvaluator::HasAnyActiveCameraRig() const
 	return RigInfo.RootEvaluator != nullptr;
 }
 
+const FCameraVariableTable* FRootCameraNodeEvaluator::GetBlendedParameters() const
+{
+	return OnGetBlendedParameters();
+}
+
 void FRootCameraNodeEvaluator::BuildSingleCameraRigHierarchy(const FSingleCameraRigHierarchyBuildParams& Params, FCameraNodeEvaluatorHierarchy& OutHierarchy)
 {
 	OnBuildSingleCameraRigHierarchy(Params, OutHierarchy);
diff --git a/Engine/Plugins/Cameras/GameplayCameras/Source/GameplayCameras/Private/Core/ShakeCameraNode.cpp b/Engine/Plugins/Cameras/GameplayCameras/Source/GameplayCameras/Private/Core/ShakeCameraNode.cpp
new file mode 100644
--- /dev/null
+++ b/Engine/Plugins/Cameras/GameplayCameras/Source/GameplayCameras/Private/Core/ShakeCameraNode.cpp
@@ -0,0 +1,23 @@
+// Copyright Epic Games, Inc. All Rights Reserved.
+
+#include "Core/ShakeCameraNode.h"
+
+#include UE_INLINE_GENERATED_CPP_BY_NAME(ShakeCameraNode)
+
+namespace UE::Cameras
+{
+
+UE_DEFINE_CAMERA_NODE_EVALUATOR(FShakeCameraNodeEvaluator)
+
+void FShakeCameraNodeEvaluator::ShakeResult(const FCameraNodeShakeParams& Params, FCameraNodeShakeResult& OutResult)
+{
+	OnShakeResult(Params, OutResult);
+}
+
+void FShakeCameraNodeEvaluator::RestartShake(const FCameraNodeShakeRestartParams& Params)
+{
+	OnRestartShake(Params);
+}
+
+}  // namespace UE::Cameras
+
diff --git a/Engine/Plugins/Cameras/GameplayCameras/Source/GameplayCameras/Private/GameplayCamerasDelegates.cpp b/Engine/Plugins/Cameras/GameplayCameras/Source/GameplayCameras/Private/GameplayCamerasDelegates.cpp
--- a/Engine/Plugins/Cameras/GameplayCameras/Source/GameplayCameras/Private/GameplayCamerasDelegates.cpp
+++ b/Engine/Plugins/Cameras/GameplayCameras/Source/GameplayCameras/Private/GameplayCamerasDelegates.cpp
@@ -7,6 +7,7 @@ namespace UE::Cameras
 
 FOnCameraAssetBuilt FGameplayCamerasDelegates::OnCameraAssetBuiltDelegate;
 FOnCameraRigAssetBuilt FGameplayCamerasDelegates::OnCameraRigAssetBuiltDelegate;
+FOnCameraShakeAssetBuilt FGameplayCamerasDelegates::OnCameraShakeAssetBuiltDelegate;
 FOnCameraNodeChanged FGameplayCamerasDelegates::OnCustomCameraNodeParametersChangedDelegate;
 
 }  // namespace UE::Cameras
diff --git a/Engine/Plugins/Cameras/GameplayCameras/Source/GameplayCameras/Private/Math/PerlinNoise.cpp b/Engine/Plugins/Cameras/GameplayCameras/Source/GameplayCameras/Private/Math/PerlinNoise.cpp
new file mode 100644
--- /dev/null
+++ b/Engine/Plugins/Cameras/GameplayCameras/Source/GameplayCameras/Private/Math/PerlinNoise.cpp
@@ -0,0 +1,159 @@
+// Copyright Epic Games, Inc. All Rights Reserved.
+
+#include "Math/PerlinNoise.h"
+
+#include "Math/Interpolation.h"
+#include "Math/UnrealMathUtility.h"
+#include "Serialization/Archive.h"
+
+#include UE_INLINE_GENERATED_CPP_BY_NAME(PerlinNoise)
+
+namespace UE::Cameras
+{
+
+FPerlinNoise::FPerlinNoise()
+{
+	Initialize(1.f);
+}
+
+FPerlinNoise::FPerlinNoise(const FPerlinNoiseData& InData, uint8 InOctaves)
+	: Amplitude(InData.Amplitude)
+	, NumOctaves(FMath::Clamp(InOctaves, 1, MAX_OCTAVES))
+{
+	Initialize(1.f);
+}
+
+FPerlinNoise::FPerlinNoise(float InAmplitude, float InFrequency, uint8 InOctaves)
+	: Amplitude(InAmplitude)
+	, NumOctaves(FMath::Clamp(InOctaves, 1, MAX_OCTAVES))
+{
+	Initialize(InFrequency);
+}
+
+void FPerlinNoise::Initialize(float InFrequency)
+{
+	float OctaveFrequency = InFrequency;
+
+	for (int32 Index = 0; Index < NumOctaves; ++Index)
+	{
+		FSinglePerlinNoise& Octave(Octaves[Index]);
+		Octave.Frequency = OctaveFrequency;
+		Octave.Prev = FMath::FRandRange(-1.f, 1.f);
+		Octave.Next = FMath::FRandRange(-1.f, 1.f);
+
+		OctaveFrequency *= Lacunarity;
+	}
+}
+
+void FPerlinNoise::SetFrequency(float InFrequency)
+{
+	if (InFrequency == Octaves[0].Frequency)
+	{
+		return;
+	}
+
+	// Move our current time to be in the same relative place inside the new interval
+	// compared to the old interval. So for instance if we were at 70% between the 
+	// noise peaks using the old frequency, let's move ourselves to be at 70% between
+	// the noise peaks of the new frequency.
+	// This loses the amount of accumulated time, but we don't really need it.
+
+	float NewOctaveFrequency = InFrequency;
+
+	for (int32 Index = 0; Index < NumOctaves; ++Index)
+	{
+		FSinglePerlinNoise& Octave(Octaves[Index]);
+
+		const float OldInterval = (Octave.Frequency > 0 ? 1.f / Octave.Frequency : 1.f);
+		const float OldIntervalFactor = FMath::Fractional(Octave.CurTime / OldInterval);
+
+		Octave.Frequency = NewOctaveFrequency;
+
+		const float NewInterval = (Octave.Frequency > 0 ? 1.f / Octave.Frequency : 1.f);
+		Octave.CurTime = NewInterval * OldIntervalFactor;
+
+		NewOctaveFrequency *= Lacunarity;
+	}
+}
+
+void FPerlinNoise::SetNumOctaves(uint8 InNumOctaves)
+{
+	InNumOctaves = FMath::Clamp(InNumOctaves, 1, MAX_OCTAVES);
+	if (InNumOctaves < NumOctaves)
+	{
+		NumOctaves = InNumOctaves;
+	}
+	else if (InNumOctaves > NumOctaves)
+	{
+		float OctaveFrequency = Octaves[NumOctaves - 1].Frequency * Lacunarity;
+
+		for (int32 Index = NumOctaves; Index < InNumOctaves; ++Index)
+		{
+			FSinglePerlinNoise& Octave(Octaves[Index]);
+			Octave.Frequency = OctaveFrequency;
+			Octave.CurTime = 0.f;
+			Octave.Prev = FMath::FRandRange(-1.f, 1.f);
+			Octave.Next = FMath::FRandRange(-1.f, 1.f);
+
+			OctaveFrequency *= Lacunarity;
+		}
+
+		NumOctaves = InNumOctaves;
+	}
+}
+
+float FPerlinNoise::GenerateValue(float DeltaTime)
+{
+	float Value = 0.f;
+	float OctaveAmplitude = Amplitude;
+
+	for (int32 Index = 0; Index < NumOctaves; ++Index)
+	{
+		FSinglePerlinNoise& Octave(Octaves[Index]);
+
+		const float Interval = (Octave.Frequency > 0 ? 1.f / Octave.Frequency : 1.f);
+
+		const float PrevNumIntervals = (Octave.CurTime / Interval);
+		const float NextNumIntervals = ((Octave.CurTime + DeltaTime) / Interval);
+
+		// If we are going over the end of the current interval, generate
+		// some new value for the next interval.
+		if ((int32)NextNumIntervals > (int32)PrevNumIntervals)
+		{
+			Octave.Prev = Octave.Next;
+			Octave.Next = FMath::FRandRange(-1.f, 1.f);
+		}
+
+		Octave.CurTime += DeltaTime;
+
+		const float IntervalFactor = (NextNumIntervals - FMath::TruncToFloat(NextNumIntervals));
+		const float InterpFactor = SmootherStep(IntervalFactor);
+
+		Value += OctaveAmplitude * FMath::Lerp(Octave.Prev, Octave.Next, InterpFactor);
+
+		OctaveAmplitude *= OctaveGain;
+	}
+
+	return Value;
+}
+
+void FPerlinNoise::Serialize(FArchive& Ar)
+{
+	Ar << Amplitude;
+	Ar << Lacunarity;
+	Ar << OctaveGain;
+
+	Ar << NumOctaves;
+
+	for (int32 Index = 0; Index < MAX_OCTAVES; ++Index)
+	{
+		FSinglePerlinNoise& Octave(Octaves[Index]);
+		Ar << Octave.Frequency;
+		Ar << Octave.CurTime;
+		Ar << Octave.Prev;
+		Ar << Octave.Next;
+	}
+}
+
+}  // namespace UE::Cameras
+
diff --git a/Engine/Plugins/Cameras/GameplayCameras/Source/GameplayCameras/Private/Nodes/Shakes/CameraShakeCameraNode.cpp b/Engine/Plugins/Cameras/GameplayCameras/Source/GameplayCameras/Private/Nodes/Shakes/CameraShakeCameraNode.cpp
new file mode 100644
--- /dev/null
+++ b/Engine/Plugins/Cameras/GameplayCameras/Source/GameplayCameras/Private/Nodes/Shakes/CameraShakeCameraNode.cpp
@@ -0,0 +1,194 @@
+// Copyright Epic Games, Inc. All Rights Reserved.
+
+#include "Nodes/Shakes/CameraShakeCameraNode.h"
+
+#include "Build/CameraBuildLog.h"
+#include "Build/CameraObjectBuildContext.h"
+#include "Core/CameraNodeEvaluator.h"
+#include "Core/CameraSystemEvaluator.h"
+#include "Core/ShakeCameraNode.h"
+#include "Helpers/CameraObjectReferenceParameterOverrideEvaluator.h"
+#include "Services/CameraShakeService.h"
+
+#include UE_INLINE_GENERATED_CPP_BY_NAME(CameraShakeCameraNode)
+
+#define LOCTEXT_NAMESPACE "CameraShakeCameraNode"
+
+namespace UE::Cameras
+{
+
+class FCameraShakeCameraNodeEvaluator : public FCameraNodeEvaluator
+{
+	UE_DECLARE_CAMERA_NODE_EVALUATOR(GAMEPLAYCAMERAS_API, FCameraShakeCameraNodeEvaluator)
+
+public:
+
+	FCameraShakeCameraNodeEvaluator();
+
+protected:
+
+	// FCameraNodeEvaluator interface.
+	virtual FCameraNodeEvaluatorChildrenView OnGetChildren() override;
+	virtual void OnInitialize(const FCameraNodeEvaluatorInitializeParams& Params, FCameraNodeEvaluationResult& OutResult) override;
+	virtual void OnUpdateParameters(const FCameraBlendedParameterUpdateParams& Params, FCameraBlendedParameterUpdateResult& OutResult) override;
+	virtual void OnBuild(const FCameraNodeEvaluatorBuildParams& Params) override;
+	virtual void OnRun(const FCameraNodeEvaluationParams& Params, FCameraNodeEvaluationResult& OutResult) override;
+
+private:
+
+	void ApplyParameterOverrides(FCameraVariableTable& OutVariableTable, bool bDrivenOnly);
+	void ApplyParameterOverrides(FCameraVariableTable& OutVariableTable, FCameraContextDataTable& OutContextDataTable, bool bDrivenOnly);
+
+private:
+
+	ECameraShakeEvaluationMode EvaluationMode = ECameraShakeEvaluationMode::VisualLayer;
+
+	FCameraNodeEvaluator* CameraShakeRootEvaluator = nullptr;
+
+	TSharedPtr<FCameraShakeService> CameraShakeService;
+};
+
+UE_DEFINE_CAMERA_NODE_EVALUATOR(FCameraShakeCameraNodeEvaluator)
+
+FCameraShakeCameraNodeEvaluator::FCameraShakeCameraNodeEvaluator()
+{
+	AddNodeEvaluatorFlags(ECameraNodeEvaluatorFlags::NeedsParameterUpdate);
+}
+
+FCameraNodeEvaluatorChildrenView FCameraShakeCameraNodeEvaluator::OnGetChildren()
+{
+	return FCameraNodeEvaluatorChildrenView({ CameraShakeRootEvaluator });
+}
+
+void FCameraShakeCameraNodeEvaluator::OnBuild(const FCameraNodeEvaluatorBuildParams& Params)
+{
+	const UCameraShakeCameraNode* CameraShakeNode = GetCameraNodeAs<UCameraShakeCameraNode>();
+	EvaluationMode = CameraShakeNode->EvaluationMode;
+
+	// If evaluating the shake "inline", build our children evaluators.
+	if (CameraShakeNode->EvaluationMode == ECameraShakeEvaluationMode::Inline)
+	{
+		const UCameraShakeAsset* CameraShake = CameraShakeNode->CameraShakeReference.GetCameraShake();
+		if (CameraShake && CameraShake->RootNode)
+		{
+			CameraShakeRootEvaluator = Params.BuildEvaluator(CameraShake->RootNode);
+		}
+	}
+}
+
+void FCameraShakeCameraNodeEvaluator::OnInitialize(const FCameraNodeEvaluatorInitializeParams& Params, FCameraNodeEvaluationResult& OutResult)
+{
+	// Apply overrides right away.
+	ApplyParameterOverrides(OutResult.VariableTable, OutResult.ContextDataTable, false);
+
+	// If evaluating the shake later in the visual layer, acquire the shake service we will use to
+	// keep that shake alive.
+	if (EvaluationMode == ECameraShakeEvaluationMode::VisualLayer)
+	{
+		CameraShakeService = Params.Evaluator->FindEvaluationService<FCameraShakeService>();
+		ensure(CameraShakeService);
+	}
+}
+
+void FCameraShakeCameraNodeEvaluator::OnUpdateParameters(const FCameraBlendedParameterUpdateParams& Params, FCameraBlendedParameterUpdateResult& OutResult)
+{
+	// Keep applying overrides in case they are driven by a variable.
+	ApplyParameterOverrides(OutResult.VariableTable, false);
+}
+
+void FCameraShakeCameraNodeEvaluator::OnRun(const FCameraNodeEvaluationParams& Params, FCameraNodeEvaluationResult& OutResult)
+{
+	// If evaluating the shake "inline", simply run it.
+	if (CameraShakeRootEvaluator)
+	{
+		CameraShakeRootEvaluator->Run(Params, OutResult);
+	}
+	// If evaluating the shake later in the visual layer, keep requesting that the shake service
+	// maintains this shake alive. We already put the shake's parameters inside the vairable table so
+	// its values will blend with whoever else wants this shake, and it will eventually run with the
+	// blended values.
+	else if (CameraShakeService)
+	{
+		const UCameraShakeCameraNode* CameraShakeNode = GetCameraNodeAs<UCameraShakeCameraNode>();
+
+		FStartCameraShakeParams StartParams;
+		StartParams.CameraShake = CameraShakeNode->CameraShakeReference.GetCameraShake();
+		CameraShakeService->RequestCameraShakeThisFrame(StartParams);
+	}
+}
+
+void FCameraShakeCameraNodeEvaluator::ApplyParameterOverrides(FCameraVariableTable& OutVariableTable, bool bDrivenOnly)
+{
+	const UCameraShakeCameraNode* PrefabNode = GetCameraNodeAs<UCameraShakeCameraNode>();
+
+	FCameraObjectReferenceParameterOverrideEvaluator OverrideEvaluator(PrefabNode->CameraShakeReference);
+	OverrideEvaluator.ApplyParameterOverrides(OutVariableTable, bDrivenOnly);
+}
+
+void FCameraShakeCameraNodeEvaluator::ApplyParameterOverrides(FCameraVariableTable& OutVariableTable, FCameraContextDataTable& OutContextDataTable, bool bDrivenOnly)
+{
+	const UCameraShakeCameraNode* PrefabNode = GetCameraNodeAs<UCameraShakeCameraNode>();
+
+	FCameraObjectReferenceParameterOverrideEvaluator OverrideEvaluator(PrefabNode->CameraShakeReference);
+	OverrideEvaluator.ApplyParameterOverrides(OutVariableTable, OutContextDataTable, bDrivenOnly);
+}
+
+}  // namespace UE::Cameras
+
+void UCameraShakeCameraNode::OnPreBuild(FCameraBuildLog& BuildLog)
+{
+	// Build the inner camera shake. Silently skip it if it's not set... but we will
+	// report an error in OnBuild about it.
+	if (UCameraShakeAsset* CameraShake = CameraShakeReference.GetCameraShake())
+	{
+		CameraShake->BuildCameraShake(BuildLog);
+	}
+
+	// Make sure the property bag of the camera shake reference is up to date.
+	CameraShakeReference.RebuildParametersIfNeeded();
+}
+
+void UCameraShakeCameraNode::OnBuild(FCameraObjectBuildContext& BuildContext)
+{
+	using namespace UE::Cameras;
+
+	UCameraShakeAsset* CameraShake = CameraShakeReference.GetCameraShake();
+	if (!CameraShake)
+	{
+		BuildContext.BuildLog.AddMessage(EMessageSeverity::Error, this, 
+				LOCTEXT("MissingCameraShake", "No camera shake specified on camera shake node."));
+		return;
+	}
+
+	// Whatever allocations our inner camera shake needs for its evaluators and
+	// their camera variables, we add that to our camera shake's allocation info.
+	BuildContext.AllocationInfo.Append(CameraShake->AllocationInfo);
+}
+
+void UCameraShakeCameraNode::GetCustomCameraNodeParameters(FCustomCameraNodeParameterInfos& OutParameterInfos)
+{
+	CameraShakeReference.GetCustomCameraNodeParameters(OutParameterInfos);
+}
+
+FCameraNodeEvaluatorPtr UCameraShakeCameraNode::OnBuildEvaluator(FCameraNodeEvaluatorBuilder& Builder) const
+{
+	using namespace UE::Cameras;
+	return Builder.BuildEvaluator<FCameraShakeCameraNodeEvaluator>();
+}
+
+#if WITH_EDITOR
+
+void UCameraShakeCameraNode::PostEditChangeProperty(struct FPropertyChangedEvent& PropertyChangedEvent)
+{
+	Super::PostEditChangeProperty(PropertyChangedEvent);
+
+	if (PropertyChangedEvent.GetPropertyName() == GET_MEMBER_NAME_CHECKED(UCameraShakeCameraNode, CameraShakeReference))
+	{
+		OnCustomCameraNodeParametersChanged(this);
+	}
+}
+
+#endif
+
+#undef LOCTEXT_NAMESPACE
+
diff --git a/Engine/Plugins/Cameras/GameplayCameras/Source/GameplayCameras/Private/Nodes/Shakes/CompositeShakeCameraNode.cpp b/Engine/Plugins/Cameras/GameplayCameras/Source/GameplayCameras/Private/Nodes/Shakes/CompositeShakeCameraNode.cpp
new file mode 100644
--- /dev/null
+++ b/Engine/Plugins/Cameras/GameplayCameras/Source/GameplayCameras/Private/Nodes/Shakes/CompositeShakeCameraNode.cpp
@@ -0,0 +1,95 @@
+// Copyright Epic Games, Inc. All Rights Reserved.
+
+#include "Nodes/Shakes/CompositeShakeCameraNode.h"
+
+#include "Core/CameraNodeEvaluator.h"
+
+namespace UE::Cameras
+{
+
+class FCompositeShakeCameraNodeEvaluator : public FShakeCameraNodeEvaluator
+{
+	UE_DECLARE_SHAKE_CAMERA_NODE_EVALUATOR(GAMEPLAYCAMERAS_API, FCompositeShakeCameraNodeEvaluator)
+
+public:
+
+protected:
+
+	// FShakeCameraNodeEvaluator interface.
+	virtual void OnBuild(const FCameraNodeEvaluatorBuildParams& Params) override;
+	virtual FCameraNodeEvaluatorChildrenView OnGetChildren() override;
+	virtual void OnRun(const FCameraNodeEvaluationParams& Params, FCameraNodeEvaluationResult& OutResult) override;
+	virtual void OnShakeResult(const FCameraNodeShakeParams& Params, FCameraNodeShakeResult& OutResult) override;
+	virtual void OnRestartShake(const FCameraNodeShakeRestartParams& Params) override;
+
+private:
+
+	TArray<FShakeCameraNodeEvaluator*> ShakeEvaluators;
+};
+
+UE_DEFINE_SHAKE_CAMERA_NODE_EVALUATOR(FCompositeShakeCameraNodeEvaluator)
+
+void FCompositeShakeCameraNodeEvaluator::OnBuild(const FCameraNodeEvaluatorBuildParams& Params)
+{
+	const UCompositeShakeCameraNode* CompositeShake = GetCameraNodeAs<UCompositeShakeCameraNode>();
+	for (const UShakeCameraNode* Shake : CompositeShake->Shakes)
+	{
+		FShakeCameraNodeEvaluator* ShakeEvaluator = Params.BuildEvaluatorAs<FShakeCameraNodeEvaluator>(Shake);
+		if (ShakeEvaluator)
+		{
+			ShakeEvaluators.Add(ShakeEvaluator);
+		}
+	}
+}
+
+FCameraNodeEvaluatorChildrenView FCompositeShakeCameraNodeEvaluator::OnGetChildren()
+{
+	return FCameraNodeEvaluatorChildrenView{ TArrayView<FShakeCameraNodeEvaluator*>(ShakeEvaluators) };
+}
+
+void FCompositeShakeCameraNodeEvaluator::OnRun(const FCameraNodeEvaluationParams& Params, FCameraNodeEvaluationResult& OutResult)
+{
+	for (FShakeCameraNodeEvaluator* ShakeEvaluator : ShakeEvaluators)
+	{
+		ShakeEvaluator->Run(Params, OutResult);
+	}
+}
+
+void FCompositeShakeCameraNodeEvaluator::OnShakeResult(const FCameraNodeShakeParams& Params, FCameraNodeShakeResult& OutResult)
+{
+	float MaxTimeLeft = 0.f;
+	bool bAnyInfiniteShake = false;
+
+	for (FShakeCameraNodeEvaluator* ShakeEvaluator : ShakeEvaluators)
+	{
+		ShakeEvaluator->ShakeResult(Params, OutResult);
+
+		if (OutResult.ShakeTimeLeft >= 0.f)
+		{
+			MaxTimeLeft = FMath::Max(MaxTimeLeft, OutResult.ShakeTimeLeft);
+		}
+		else
+		{
+			bAnyInfiniteShake = true;
+		}
+	}
+
+	OutResult.ShakeTimeLeft = (bAnyInfiniteShake ? -1.f : MaxTimeLeft);
+}
+
+void FCompositeShakeCameraNodeEvaluator::OnRestartShake(const FCameraNodeShakeRestartParams& Params)
+{
+	for (FShakeCameraNodeEvaluator* ShakeEvaluator : ShakeEvaluators)
+	{
+		ShakeEvaluator->RestartShake(Params);
+	}
+}
+
+}  // namespace UE::Cameras
+
+FCameraNodeEvaluatorPtr UCompositeShakeCameraNode::OnBuildEvaluator(FCameraNodeEvaluatorBuilder& Builder) const
+{
+	using namespace UE::Cameras;
+	return Builder.BuildEvaluator<FCompositeShakeCameraNodeEvaluator>();
+}
+
diff --git a/Engine/Plugins/Cameras/GameplayCameras/Source/GameplayCameras/Private/Nodes/Shakes/EnvelopeShakeCameraNode.cpp b/Engine/Plugins/Cameras/GameplayCameras/Source/GameplayCameras/Private/Nodes/Shakes/EnvelopeShakeCameraNode.cpp
new file mode 100644
--- /dev/null
+++ b/Engine/Plugins/Cameras/GameplayCameras/Source/GameplayCameras/Private/Nodes/Shakes/EnvelopeShakeCameraNode.cpp
@@ -0,0 +1,158 @@
+// Copyright Epic Games, Inc. All Rights Reserved.
+
+#include "Nodes/Shakes/EnvelopeShakeCameraNode.h"
+
+#include "Core/CameraNodeEvaluator.h"
+#include "Math/Interpolation.h"
+#include "Math/PerlinNoise.h"
+
+namespace UE::Cameras
+{
+
+class FEnvelopeShakeCameraNodeEvaluator : public FShakeCameraNodeEvaluator
+{
+	UE_DECLARE_SHAKE_CAMERA_NODE_EVALUATOR(GAMEPLAYCAMERAS_API, FEnvelopeShakeCameraNodeEvaluator)
+
+public:
+
+protected:
+
+	// FShakeCameraNodeEvaluator interface.
+	virtual void OnBuild(const FCameraNodeEvaluatorBuildParams& Params) override;
+	virtual FCameraNodeEvaluatorChildrenView OnGetChildren() override;
+	virtual void OnInitialize(const FCameraNodeEvaluatorInitializeParams& Params, FCameraNodeEvaluationResult& OutResult) override;
+	virtual void OnRun(const FCameraNodeEvaluationParams& Params, FCameraNodeEvaluationResult& OutResult) override;
+	virtual void OnShakeResult(const FCameraNodeShakeParams& Params, FCameraNodeShakeResult& OutResult) override;
+	virtual void OnRestartShake(const FCameraNodeShakeRestartParams& Params) override;
+	virtual void OnSerialize(const FCameraNodeEvaluatorSerializeParams& Params, FArchive& Ar) override;
+
+private:
+
+	FShakeCameraNodeEvaluator* ShakeEvaluator = nullptr;
+
+	float EaseInTime;
+	float EaseOutTime;
+	float TotalTime;
+
+	float CurrentTime = 0.f;
+};
+
+UE_DEFINE_SHAKE_CAMERA_NODE_EVALUATOR(FEnvelopeShakeCameraNodeEvaluator)
+
+void FEnvelopeShakeCameraNodeEvaluator::OnBuild(const FCameraNodeEvaluatorBuildParams& Params)
+{
+	const UEnvelopeShakeCameraNode* EnvelopeNode = GetCameraNodeAs<UEnvelopeShakeCameraNode>();
+	ShakeEvaluator = Params.BuildEvaluatorAs<FShakeCameraNodeEvaluator>(EnvelopeNode->Shake);
+}
+
+FCameraNodeEvaluatorChildrenView FEnvelopeShakeCameraNodeEvaluator::OnGetChildren()
+{
+	return FCameraNodeEvaluatorChildrenView{ ShakeEvaluator };
+}
+
+void FEnvelopeShakeCameraNodeEvaluator::OnInitialize(const FCameraNodeEvaluatorInitializeParams& Params, FCameraNodeEvaluationResult& OutResult)
+{
+	const UEnvelopeShakeCameraNode* EnvelopeNode = GetCameraNodeAs<UEnvelopeShakeCameraNode>();
+
+	EaseInTime = EnvelopeNode->EaseInTime.GetValue(OutResult.VariableTable);
+	EaseOutTime = EnvelopeNode->EaseOutTime.GetValue(OutResult.VariableTable);
+	TotalTime = EnvelopeNode->TotalTime.GetValue(OutResult.VariableTable);
+
+	EaseInTime = FMath::Max(0, EaseInTime);
+	EaseOutTime = FMath::Max(0, EaseOutTime);
+	TotalTime = FMath::Max(0, TotalTime);
+	if (EaseInTime + EaseOutTime > TotalTime)
+	{
+		const float HalfTotalTime = TotalTime / 2.f;
+		EaseInTime = EaseOutTime = HalfTotalTime;
+	}
+
+	CurrentTime = 0.f;
+}
+
+void FEnvelopeShakeCameraNodeEvaluator::OnRun(const FCameraNodeEvaluationParams& Params, FCameraNodeEvaluationResult& OutResult)
+{
+	if (CurrentTime >= TotalTime)
+	{
+		return;
+	}
+
+	CurrentTime += Params.DeltaTime;
+
+	if (ShakeEvaluator)
+	{
+		ShakeEvaluator->Run(Params, OutResult);
+	}
+}
+
+void FEnvelopeShakeCameraNodeEvaluator::OnShakeResult(const FCameraNodeShakeParams& Params, FCameraNodeShakeResult& OutResult)
+{
+	if (CurrentTime >= TotalTime || !ShakeEvaluator)
+	{
+		OutResult.ShakeTimeLeft = 0.f;
+		return;
+	}
+
+	float Alpha = 1.f;
+	if (EaseInTime > 0 && CurrentTime < EaseInTime)
+	{
+		Alpha = SmoothStep(CurrentTime / EaseInTime);
+	}
+	else if (EaseOutTime > 0 && CurrentTime > TotalTime - EaseOutTime)
+	{
+		Alpha = SmoothStep((TotalTime - CurrentTime) / EaseOutTime);
+	}
+
+	OutResult.ShakeTimeLeft = (TotalTime - CurrentTime);
+	{
+		FCameraNodeShakeParams ChildParams(Params);
+		ChildParams.ShakeScale *= Alpha;
+
+		ShakeEvaluator->ShakeResult(ChildParams, OutResult);
+	}
+	OutResult.ShakeTimeLeft = FMath::Clamp(OutResult.ShakeTimeLeft, 0, TotalTime - CurrentTime);
+}
+
+void FEnvelopeShakeCameraNodeEvaluator::OnRestartShake(const FCameraNodeShakeRestartParams& Params)
+{
+	if (EaseInTime > 0 && CurrentTime < EaseInTime)
+	{
+		// Keep easing in, but extend TotalTime by the time we already did.
+		TotalTime += CurrentTime;
+	}
+	else if (EaseOutTime > 0 && CurrentTime > TotalTime - EaseOutTime)
+	{
+		// Stop easing out and restart easing in at an equivalent time.
+		const float Alpha = (TotalTime - CurrentTime) / EaseOutTime;
+		CurrentTime = Alpha * EaseInTime;
+	}
+	else
+	{
+		// We're fully shaking. Keep doing that for another whole duration.
+		TotalTime += (TotalTime - CurrentTime);
+	}
+
+	if (ShakeEvaluator)
+	{
+		ShakeEvaluator->RestartShake(Params);
+	}
+}
+
+void FEnvelopeShakeCameraNodeEvaluator::OnSerialize(const FCameraNodeEvaluatorSerializeParams& Params, FArchive& Ar)
+{
+	Super::OnSerialize(Params, Ar);
+
+	Ar << EaseInTime;
+	Ar << EaseOutTime;
+	Ar << TotalTime;
+	Ar << CurrentTime;
+}
+
+}  // namespace UE::Cameras
+
+FCameraNodeEvaluatorPtr UEnvelopeShakeCameraNode::OnBuildEvaluator(FCameraNodeEvaluatorBuilder& Builder) const
+{
+	using namespace UE::Cameras;
+	return Builder.BuildEvaluator<FEnvelopeShakeCameraNodeEvaluator>();
+}
+
diff --git a/Engine/Plugins/Cameras/GameplayCameras/Source/GameplayCameras/Private/Nodes/Shakes/PerlinNoiseLocationShakeCameraNode.cpp b/Engine/Plugins/Cameras/GameplayCameras/Source/GameplayCameras/Private/Nodes/Shakes/PerlinNoiseLocationShakeCameraNode.cpp
new file mode 100644
--- /dev/null
+++ b/Engine/Plugins/Cameras/GameplayCameras/Source/GameplayCameras/Private/Nodes/Shakes/PerlinNoiseLocationShakeCameraNode.cpp
@@ -0,0 +1,150 @@
+// Copyright Epic Games, Inc. All Rights Reserved.
+
+#include "Nodes/Shakes/PerlinNoiseLocationShakeCameraNode.h"
+
+#include "Core/CameraNodeEvaluator.h"
+#include "Core/CameraParameterReader.h"
+#include "Debug/CameraDebugBlock.h"
+#include "Debug/CameraDebugBlockBuilder.h"
+#include "Debug/CameraDebugGraph.h"
+#include "Debug/CameraDebugRenderer.h"
+#include "Math/PerlinNoise.h"
+
+#define LOCTEXT_NAMESPACE "PerlinNoiseLocationShakeCameraNode"
+
+namespace UE::Cameras
+{
+
+class FPerlinNoiseLocationShakeCameraNodeEvaluator : public FShakeCameraNodeEvaluator
+{
+	UE_DECLARE_SHAKE_CAMERA_NODE_EVALUATOR(GAMEPLAYCAMERAS_API, FPerlinNoiseLocationShakeCameraNodeEvaluator)
+
+public:
+
+protected:
+
+	// FShakeCameraNodeEvaluator interface.
+	virtual void OnInitialize(const FCameraNodeEvaluatorInitializeParams& Params, FCameraNodeEvaluationResult& OutResult) override;
+	virtual void OnRun(const FCameraNodeEvaluationParams& Params, FCameraNodeEvaluationResult& OutResult) override;
+	virtual void OnShakeResult(const FCameraNodeShakeParams& Params, FCameraNodeShakeResult& OutResult) override;
+	virtual void OnSerialize(const FCameraNodeEvaluatorSerializeParams& Params, FArchive& Ar) override;
+#if UE_GAMEPLAY_CAMERAS_DEBUG
+	virtual void OnBuildDebugBlocks(const FCameraDebugBlockBuildParams& Params, FCameraDebugBlockBuilder& Builder) override;
+#endif  // UE_GAMEPLAY_CAMERAS_DEBUG
+
+private:
+
+	TCameraParameterReader<float> AmplitudeMultiplier;
+	TCameraParameterReader<float> FrequencyMultiplier;
+	TCameraParameterReader<int32> Octaves;
+
+	FPerlinNoise GeneratorX;
+	FPerlinNoise GeneratorY;
+	FPerlinNoise GeneratorZ;
+
+	FVector3d NoiseValue;
+
+#if UE_GAMEPLAY_CAMERAS_DEBUG
+	TCameraDebugGraph<3> NoiseValues;
+#endif
+};
+
+UE_DECLARE_CAMERA_DEBUG_BLOCK_START(GAMEPLAYCAMERAS_API, FPerlinNoiseLocationShakeDebugBlock)
+	UE_DECLARE_CAMERA_DEBUG_BLOCK_FIELD(TCameraDebugGraph<3>, NoiseValues)
+UE_DECLARE_CAMERA_DEBUG_BLOCK_END()
+
+UE_DEFINE_CAMERA_DEBUG_BLOCK_WITH_FIELDS(FPerlinNoiseLocationShakeDebugBlock)
+
+UE_DEFINE_SHAKE_CAMERA_NODE_EVALUATOR(FPerlinNoiseLocationShakeCameraNodeEvaluator)
+
+void FPerlinNoiseLocationShakeCameraNodeEvaluator::OnInitialize(const FCameraNodeEvaluatorInitializeParams& Params, FCameraNodeEvaluationResult& OutResult)
+{
+	const UPerlinNoiseLocationShakeCameraNode* ShakeNode = GetCameraNodeAs<UPerlinNoiseLocationShakeCameraNode>();
+
+	AmplitudeMultiplier.Initialize(ShakeNode->AmplitudeMultiplier);
+	FrequencyMultiplier.Initialize(ShakeNode->FrequencyMultiplier);
+	Octaves.Initialize(ShakeNode->Octaves);
+
+	GeneratorX = FPerlinNoise(
+			ShakeNode->X.Amplitude * AmplitudeMultiplier.Get(OutResult.VariableTable),
+			ShakeNode->X.Frequency * FrequencyMultiplier.Get(OutResult.VariableTable),
+			Octaves.Get(OutResult.VariableTable));
+	GeneratorY = FPerlinNoise(
+			ShakeNode->Y.Amplitude * AmplitudeMultiplier.Get(OutResult.VariableTable),
+			ShakeNode->Y.Frequency * FrequencyMultiplier.Get(OutResult.VariableTable),
+			Octaves.Get(OutResult.VariableTable));
+	GeneratorY = FPerlinNoise(
+			ShakeNode->Y.Amplitude * AmplitudeMultiplier.Get(OutResult.VariableTable),
+			ShakeNode->Y.Frequency * FrequencyMultiplier.Get(OutResult.VariableTable),
+			Octaves.Get(OutResult.VariableTable));
+}
+
+void FPerlinNoiseLocationShakeCameraNodeEvaluator::OnRun(const FCameraNodeEvaluationParams& Params, FCameraNodeEvaluationResult& OutResult)
+{
+	const UPerlinNoiseLocationShakeCameraNode* ShakeNode = GetCameraNodeAs<UPerlinNoiseLocationShakeCameraNode>();
+
+	GeneratorX.SetAmplitude(ShakeNode->X.Amplitude * AmplitudeMultiplier.Get(OutResult.VariableTable));
+	GeneratorX.SetFrequency(ShakeNode->X.Frequency * FrequencyMultiplier.Get(OutResult.VariableTable));
+	GeneratorX.SetNumOctaves(Octaves.Get(OutResult.VariableTable));
+
+	GeneratorY.SetAmplitude(ShakeNode->Y.Amplitude * AmplitudeMultiplier.Get(OutResult.VariableTable));
+	GeneratorY.SetFrequency(ShakeNode->Y.Frequency * FrequencyMultiplier.Get(OutResult.VariableTable));
+	GeneratorY.SetNumOctaves(Octaves.Get(OutResult.VariableTable));
+
+	GeneratorZ.SetAmplitude(ShakeNode->Z.Amplitude * AmplitudeMultiplier.Get(OutResult.VariableTable));
+	GeneratorZ.SetFrequency(ShakeNode->Z.Frequency * FrequencyMultiplier.Get(OutResult.VariableTable));
+	GeneratorZ.SetNumOctaves(Octaves.Get(OutResult.VariableTable));
+
+	const float ValueX = GeneratorX.GenerateValue(Params.DeltaTime);
+	const float ValueY = GeneratorY.GenerateValue(Params.DeltaTime);
+	const float ValueZ = GeneratorZ.GenerateValue(Params.DeltaTime);
+
+	NoiseValue = FVector3d(ValueX, ValueY, ValueZ);
+
+#if UE_GAMEPLAY_CAMERAS_DEBUG
+	NoiseValues.Add(Params.DeltaTime, ValueX, ValueY, ValueZ);
+#endif
+}
+
+void FPerlinNoiseLocationShakeCameraNodeEvaluator::OnShakeResult(const FCameraNodeShakeParams& Params, FCameraNodeShakeResult& OutResult)
+{
+	FVector3d Location = OutResult.ShakenResult.CameraPose.GetLocation();
+	Location += NoiseValue;
+	OutResult.ShakenResult.CameraPose.SetLocation(Location);
+	OutResult.ShakeTimeLeft = -1.f;
+}
+
+void FPerlinNoiseLocationShakeCameraNodeEvaluator::OnSerialize(const FCameraNodeEvaluatorSerializeParams& Params, FArchive& Ar)
+{
+	Super::OnSerialize(Params, Ar);
+
+	Ar << GeneratorX;
+	Ar << GeneratorY;
+	Ar << GeneratorZ;
+}
+
+#if UE_GAMEPLAY_CAMERAS_DEBUG
+
+void FPerlinNoiseLocationShakeCameraNodeEvaluator::OnBuildDebugBlocks(const FCameraDebugBlockBuildParams& Params, FCameraDebugBlockBuilder& Builder)
+{
+	FPerlinNoiseLocationShakeDebugBlock& DebugBlock = Builder.AttachDebugBlock<FPerlinNoiseLocationShakeDebugBlock>();
+	DebugBlock.NoiseValues = NoiseValues;
+}
+
+void FPerlinNoiseLocationShakeDebugBlock::OnDebugDraw(const FCameraDebugBlockDrawParams& Params, FCameraDebugRenderer& Renderer)
+{
+	Renderer.DrawGraph(NoiseValues, LOCTEXT("DebugGraphName", "Location Noise"));
+}
+
+#endif  // UE_GAMEPLAY_CAMERAS_DEBUG
+
+}  // namespace UE::Cameras
+
+FCameraNodeEvaluatorPtr UPerlinNoiseLocationShakeCameraNode::OnBuildEvaluator(FCameraNodeEvaluatorBuilder& Builder) const
+{
+	using namespace UE::Cameras;
+	return Builder.BuildEvaluator<FPerlinNoiseLocationShakeCameraNodeEvaluator>();
+}
+
+#undef LOCTEXT_NAMESPACE
+
diff --git a/Engine/Plugins/Cameras/GameplayCameras/Source/GameplayCameras/Private/Nodes/Shakes/PerlinNoiseRotationShakeCameraNode.cpp b/Engine/Plugins/Cameras/GameplayCameras/Source/GameplayCameras/Private/Nodes/Shakes/PerlinNoiseRotationShakeCameraNode.cpp
new file mode 100644
--- /dev/null
+++ b/Engine/Plugins/Cameras/GameplayCameras/Source/GameplayCameras/Private/Nodes/Shakes/PerlinNoiseRotationShakeCameraNode.cpp
@@ -0,0 +1,149 @@
+// Copyright Epic Games, Inc. All Rights Reserved.
+
+#include "Nodes/Shakes/PerlinNoiseRotationShakeCameraNode.h"
+
+#include "Core/CameraNodeEvaluator.h"
+#include "Core/CameraParameterReader.h"
+#include "Debug/CameraDebugBlock.h"
+#include "Debug/CameraDebugBlockBuilder.h"
+#include "Debug/CameraDebugGraph.h"
+#include "Debug/CameraDebugRenderer.h"
+#include "Math/PerlinNoise.h"
+
+#define LOCTEXT_NAMESPACE "PerlinNoiseRotationShakeCameraNode"
+
+namespace UE::Cameras
+{
+
+class FPerlinNoiseRotationShakeCameraNodeEvaluator : public FShakeCameraNodeEvaluator
+{
+	UE_DECLARE_SHAKE_CAMERA_NODE_EVALUATOR(GAMEPLAYCAMERAS_API, FPerlinNoiseRotationShakeCameraNodeEvaluator)
+
+public:
+
+protected:
+
+	// FShakeCameraNodeEvaluator interface.
+	virtual void OnInitialize(const FCameraNodeEvaluatorInitializeParams& Params, FCameraNodeEvaluationResult& OutResult) override;
+	virtual void OnRun(const FCameraNodeEvaluationParams& Params, FCameraNodeEvaluationResult& OutResult) override;
+	virtual void OnShakeResult(const FCameraNodeShakeParams& Params, FCameraNodeShakeResult& OutResult) override;
+	virtual void OnSerialize(const FCameraNodeEvaluatorSerializeParams& Params, FArchive& Ar) override;
+#if UE_GAMEPLAY_CAMERAS_DEBUG
+	virtual void OnBuildDebugBlocks(const FCameraDebugBlockBuildParams& Params, FCameraDebugBlockBuilder& Builder) override;
+#endif  // UE_GAMEPLAY_CAMERAS_DEBUG
+
+private:
+
+	TCameraParameterReader<float> AmplitudeMultiplier;
+	TCameraParameterReader<float> FrequencyMultiplier;
+	TCameraParameterReader<int32> Octaves;
+
+	FPerlinNoise GeneratorYaw;
+	FPerlinNoise GeneratorPitch;
+	FPerlinNoise GeneratorRoll;
+
+	FRotator3d NoiseValue;
+
+#if UE_GAMEPLAY_CAMERAS_DEBUG
+	TCameraDebugGraph<3> NoiseValues;
+#endif
+};
+
+UE_DECLARE_CAMERA_DEBUG_BLOCK_START(GAMEPLAYCAMERAS_API, FPerlinNoiseRotationShakeDebugBlock)
+	UE_DECLARE_CAMERA_DEBUG_BLOCK_FIELD(TCameraDebugGraph<3>, NoiseValues)
+UE_DECLARE_CAMERA_DEBUG_BLOCK_END()
+
+UE_DEFINE_CAMERA_DEBUG_BLOCK_WITH_FIELDS(FPerlinNoiseRotationShakeDebugBlock)
+
+UE_DEFINE_SHAKE_CAMERA_NODE_EVALUATOR(FPerlinNoiseRotationShakeCameraNodeEvaluator)
+
+void FPerlinNoiseRotationShakeCameraNodeEvaluator::OnInitialize(const FCameraNodeEvaluatorInitializeParams& Params, FCameraNodeEvaluationResult& OutResult)
+{
+	const UPerlinNoiseRotationShakeCameraNode* ShakeNode = GetCameraNodeAs<UPerlinNoiseRotationShakeCameraNode>();
+
+	AmplitudeMultiplier.Initialize(ShakeNode->AmplitudeMultiplier);
+	FrequencyMultiplier.Initialize(ShakeNode->FrequencyMultiplier);
+	Octaves.Initialize(ShakeNode->Octaves);
+
+	GeneratorYaw = FPerlinNoise(
+			ShakeNode->Yaw.Amplitude * AmplitudeMultiplier.Get(OutResult.VariableTable),
+			ShakeNode->Yaw.Frequency * FrequencyMultiplier.Get(OutResult.VariableTable),
+			Octaves.Get(OutResult.VariableTable));
+	GeneratorPitch = FPerlinNoise(
+			ShakeNode->Pitch.Amplitude * AmplitudeMultiplier.Get(OutResult.VariableTable),
+			ShakeNode->Pitch.Frequency * FrequencyMultiplier.Get(OutResult.VariableTable),
+			Octaves.Get(OutResult.VariableTable));
+	GeneratorRoll = FPerlinNoise(
+			ShakeNode->Roll.Amplitude * AmplitudeMultiplier.Get(OutResult.VariableTable),
+			ShakeNode->Roll.Frequency * FrequencyMultiplier.Get(OutResult.VariableTable),
+			Octaves.Get(OutResult.VariableTable));
+}
+
+void FPerlinNoiseRotationShakeCameraNodeEvaluator::OnRun(const FCameraNodeEvaluationParams& Params, FCameraNodeEvaluationResult& OutResult)
+{
+	const UPerlinNoiseRotationShakeCameraNode* ShakeNode = GetCameraNodeAs<UPerlinNoiseRotationShakeCameraNode>();
+
+	GeneratorYaw.SetAmplitude(ShakeNode->Yaw.Amplitude * AmplitudeMultiplier.Get(OutResult.VariableTable));
+	GeneratorYaw.SetFrequency(ShakeNode->Yaw.Frequency * FrequencyMultiplier.Get(OutResult.VariableTable));
+	GeneratorYaw.SetNumOctaves(Octaves.Get(OutResult.VariableTable));
+
+	GeneratorPitch.SetAmplitude(ShakeNode->Pitch.Amplitude * AmplitudeMultiplier.Get(OutResult.VariableTable));
+	GeneratorPitch.SetFrequency(ShakeNode->Pitch.Frequency * FrequencyMultiplier.Get(OutResult.VariableTable));
+	GeneratorPitch.SetNumOctaves(Octaves.Get(OutResult.VariableTable));
+
+	GeneratorRoll.SetAmplitude(ShakeNode->Roll.Amplitude * AmplitudeMultiplier.Get(OutResult.VariableTable));
+	GeneratorRoll.SetFrequency(ShakeNode->Roll.Frequency * FrequencyMultiplier.Get(OutResult.VariableTable));
+	GeneratorRoll.SetNumOctaves(Octaves.Get(OutResult.VariableTable));
+
+	const float ValueYaw = GeneratorYaw.GenerateValue(Params.DeltaTime);
+	const float ValuePitch = GeneratorPitch.GenerateValue(Params.DeltaTime);
+	const float ValueRoll = GeneratorRoll.GenerateValue(Params.DeltaTime);
+	NoiseValue = FRotator3d(ValuePitch, ValueYaw, ValueRoll);
+
+#if UE_GAMEPLAY_CAMERAS_DEBUG
+	NoiseValues.Add(Params.DeltaTime, ValueYaw, ValuePitch, ValueRoll);
+#endif
+}
+
+void FPerlinNoiseRotationShakeCameraNodeEvaluator::OnShakeResult(const FCameraNodeShakeParams& Params, FCameraNodeShakeResult& OutResult)
+{
+	FRotator3d Rotation = OutResult.ShakenResult.CameraPose.GetRotation();
+	Rotation += NoiseValue;
+	OutResult.ShakenResult.CameraPose.SetRotation(Rotation);
+	OutResult.ShakeTimeLeft = -1.f;
+}
+
+void FPerlinNoiseRotationShakeCameraNodeEvaluator::OnSerialize(const FCameraNodeEvaluatorSerializeParams& Params, FArchive& Ar)
+{
+	Super::OnSerialize(Params, Ar);
+
+	Ar << GeneratorYaw;
+	Ar << GeneratorPitch;
+	Ar << GeneratorRoll;
+}
+
+#if UE_GAMEPLAY_CAMERAS_DEBUG
+
+void FPerlinNoiseRotationShakeCameraNodeEvaluator::OnBuildDebugBlocks(const FCameraDebugBlockBuildParams& Params, FCameraDebugBlockBuilder& Builder)
+{
+	FPerlinNoiseRotationShakeDebugBlock& DebugBlock = Builder.AttachDebugBlock<FPerlinNoiseRotationShakeDebugBlock>();
+	DebugBlock.NoiseValues = NoiseValues;
+}
+
+void FPerlinNoiseRotationShakeDebugBlock::OnDebugDraw(const FCameraDebugBlockDrawParams& Params, FCameraDebugRenderer& Renderer)
+{
+	Renderer.DrawGraph(NoiseValues, LOCTEXT("DebugGraphName", "Rotation Noise"));
+}
+
+#endif  // UE_GAMEPLAY_CAMERAS_DEBUG
+
+}  // namespace UE::Cameras
+
+FCameraNodeEvaluatorPtr UPerlinNoiseRotationShakeCameraNode::OnBuildEvaluator(FCameraNodeEvaluatorBuilder& Builder) const
+{
+	using namespace UE::Cameras;
+	return Builder.BuildEvaluator<FPerlinNoiseRotationShakeCameraNodeEvaluator>();
+}
+
+#undef LOCTEXT_NAMESPACE
+
diff --git a/Engine/Plugins/Cameras/GameplayCameras/Source/GameplayCameras/Private/Services/CameraShakeService.cpp b/Engine/Plugins/Cameras/GameplayCameras/Source/GameplayCameras/Private/Services/CameraShakeService.cpp
new file mode 100644
--- /dev/null
+++ b/Engine/Plugins/Cameras/GameplayCameras/Source/GameplayCameras/Private/Services/CameraShakeService.cpp
@@ -0,0 +1,441 @@
+// Copyright Epic Games, Inc. All Rights Reserved.
+
+#include "Services/CameraShakeService.h"
+
+#include "Core/BlendCameraNode.h"
+#include "Core/CameraEvaluationContext.h"
+#include "Core/CameraNodeEvaluatorHierarchy.h"
+#include "Core/CameraNodeEvaluatorStorage.h"
+#include "Core/CameraRigAsset.h"
+#include "Core/CameraShakeAsset.h"
+#include "Core/CameraSystemEvaluator.h"
+#include "Core/RootCameraNode.h"
+#include "Core/ShakeCameraNode.h"
+#include "UObject/ObjectMacros.h"
+
+#include UE_INLINE_GENERATED_CPP_BY_NAME(CameraShakeService)
+
+namespace UE::Cameras
+{
+
+class FCameraShakeServiceCameraNodeEvaluator : public FCameraNodeEvaluator
+{
+	UE_DECLARE_CAMERA_NODE_EVALUATOR(GAMEPLAYCAMERAS_API, FCameraShakeServiceCameraNodeEvaluator)
+
+public:
+
+	void StartCameraShake(const FStartCameraShakeParams& Params);
+	void RequestCameraShakeThisFrame(const FStartCameraShakeParams& Params);
+
+protected:
+
+	// FCameraNodeEvaluator interface.
+	virtual FCameraNodeEvaluatorChildrenView OnGetChildren() override;
+	virtual void OnInitialize(const FCameraNodeEvaluatorInitializeParams& Params, FCameraNodeEvaluationResult& OutResult) override;
+	virtual void OnRun(const FCameraNodeEvaluationParams& Params, FCameraNodeEvaluationResult& OutResult) override;
+
+private:
+
+	struct FShakeEntry;
+
+	FShakeEntry* AddCameraShake(const FStartCameraShakeParams& Params);
+
+	bool InitializeEntry(
+		FShakeEntry& NewEntry, 
+		const UCameraShakeAsset* CameraShake,
+		TSharedPtr<const FCameraEvaluationContext> EvaluationContext);
+
+	void PopEntry(int32 EntryIndex);
+
+private:
+
+	struct FShakeEntry
+	{
+		TWeakPtr<const FCameraEvaluationContext> EvaluationContext;
+		TObjectPtr<const UCameraShakeAsset> CameraShake;
+		FCameraNodeEvaluatorStorage EvaluatorStorage;
+		FBlendCameraNodeEvaluator* BlendInEvaluator = nullptr;
+		FBlendCameraNodeEvaluator* BlendOutEvaluator = nullptr;
+		FShakeCameraNodeEvaluator* RootEvaluator = nullptr;
+		FCameraNodeEvaluatorHierarchy EvaluatorHierarchy;
+		FCameraNodeEvaluationResult Result;
+		float CurrentTime = 0.f;
+		float ShakeScale = 1.f;
+		ECameraShakePlaySpace PlaySpace = ECameraShakePlaySpace::CameraLocal;
+		FMatrix UserPlaySpaceMatrix;
+		uint8 NumRequests = 0;
+		bool bPersistentRequest = false;
+		bool bIsFirstFrame = false;
+	};
+
+	FCameraSystemEvaluator* OwningEvaluator = nullptr;
+	TSharedPtr<const FCameraEvaluationContext> ShakeContext;
+	const FCameraVariableTable* BlendedParameters = nullptr;;
+
+	TArray<FShakeEntry> Entries;
+};
+
+UE_DEFINE_CAMERA_NODE_EVALUATOR(FCameraShakeServiceCameraNodeEvaluator)
+
+FCameraNodeEvaluatorChildrenView FCameraShakeServiceCameraNodeEvaluator::OnGetChildren()
+{
+	FCameraNodeEvaluatorChildrenView ChildrenView;
+	for (FShakeEntry& Entry : Entries)
+	{
+		ChildrenView.Add(Entry.RootEvaluator);
+	}
+	return ChildrenView;
+}
+
+void FCameraShakeServiceCameraNodeEvaluator::StartCameraShake(const FStartCameraShakeParams& Params)
+{
+	if (!Params.CameraShake)
+	{
+		return;
+	}
+
+	// If this shake wants to only have a single instance active at a time, look for a running
+	// one and restart it.
+	if (Params.CameraShake->bIsSingleInstance)
+	{
+		FShakeEntry* ExistingEntry = nullptr;
+		for (FShakeEntry& Entry : Entries)
+		{
+			if (Entry.CameraShake == Params.CameraShake)
+			{
+				ExistingEntry = &Entry;
+				break;
+			}
+		}
+		if (ExistingEntry && ensure(ExistingEntry->RootEvaluator))
+		{
+			FCameraNodeShakeRestartParams RestartParams;
+			ExistingEntry->RootEvaluator->RestartShake(RestartParams);
+			return;
+		}
+	}
+
+	FShakeEntry* NewEntry = AddCameraShake(Params);
+	if (NewEntry)
+	{
+		NewEntry->bPersistentRequest = true;
+	}
+}
+
+void FCameraShakeServiceCameraNodeEvaluator::RequestCameraShakeThisFrame(const FStartCameraShakeParams& Params)
+{
+	if (!Params.CameraShake)
+	{
+		return;
+	}
+
+	// Record this request on a running camera shake, if any.
+	bool bFoundShake = false;
+	for (FShakeEntry& Entry : Entries)
+	{
+		if (Entry.CameraShake == Params.CameraShake && !Entry.bPersistentRequest)
+		{
+			++Entry.NumRequests;
+			bFoundShake = true;
+			break;
+		}
+	}
+
+	if (bFoundShake)
+	{
+		return;
+	}
+
+	// Create a new camera shake if there wasn't any, and record this first request.
+	FShakeEntry* NewEntry = AddCameraShake(Params);
+	if (NewEntry)
+	{
+		NewEntry->bPersistentRequest = false;
+		NewEntry->NumRequests = 1;
+	}
+}
+
+FCameraShakeServiceCameraNodeEvaluator::FShakeEntry* FCameraShakeServiceCameraNodeEvaluator::AddCameraShake(const FStartCameraShakeParams& Params)
+{
+	ensure(Params.CameraShake);
+
+	FShakeEntry NewEntry;
+	const bool bInitialized = InitializeEntry(NewEntry, Params.CameraShake, ShakeContext);
+	if (!bInitialized)
+	{
+		return nullptr;
+	}
+
+	NewEntry.ShakeScale = Params.ShakeScale;
+	NewEntry.PlaySpace = Params.PlaySpace;
+	NewEntry.UserPlaySpaceMatrix = Params.UserPlaySpaceMatrix;
+
+	const int32 AddedIndex = Entries.Add(MoveTemp(NewEntry));
+	return &Entries[AddedIndex];
+}
+
+void FCameraShakeServiceCameraNodeEvaluator::OnInitialize(const FCameraNodeEvaluatorInitializeParams& Params, FCameraNodeEvaluationResult& OutResult)
+{
+	OwningEvaluator = Params.Evaluator;
+	ShakeContext = Params.EvaluationContext;
+	BlendedParameters = OwningEvaluator->GetRootNodeEvaluator()->GetBlendedParameters();
+}
+
+void FCameraShakeServiceCameraNodeEvaluator::OnRun(const FCameraNodeEvaluationParams& Params, FCameraNodeEvaluationResult& OutResult)
+{
+	TArray<int32> EntriesToRemove;
+
+	for (int32 Index = 0; Index < Entries.Num(); ++Index)
+	{
+		FShakeEntry& Entry(Entries[Index]);
+
+		// See if anybody still cares about this shake.
+		if (!Entry.bPersistentRequest && Entry.NumRequests == 0)
+		{
+			EntriesToRemove.Add(Index);
+			continue;
+		}
+
+		// Set us up for updating this shake.
+		TSharedPtr<const FCameraEvaluationContext> CurContext = Entry.EvaluationContext.Pin();
+
+		FCameraNodeEvaluationParams CurParams(Params);
+		CurParams.EvaluationContext = CurContext;
+		CurParams.bIsFirstFrame = Entry.bIsFirstFrame;
+
+		FCameraNodeEvaluationResult& CurResult(Entry.Result);
+
+		// Start with the input given to us.
+		{
+			CurResult.CameraPose = OutResult.CameraPose;
+			CurResult.VariableTable.OverrideAll(OutResult.VariableTable);
+			CurResult.ContextDataTable.OverrideAll(OutResult.ContextDataTable);
+			CurResult.CameraRigJoints.OverrideAll(OutResult.CameraRigJoints);
+			CurResult.PostProcessSettings.OverrideAll(OutResult.PostProcessSettings);
+
+			const FCameraNodeEvaluationResult& ContextResult(CurContext->GetInitialResult());
+			CurResult.bIsCameraCut = OutResult.bIsCameraCut || ContextResult.bIsCameraCut;
+			CurResult.bIsValid = true;
+		}
+
+		// Add any parameters coming from the main blend stack.
+		if (BlendedParameters)
+		{
+			CurResult.VariableTable.Override(*BlendedParameters, ECameraVariableTableFilter::KnownOnly);
+		}
+
+		// Update timing.
+		Entry.CurrentTime += Params.DeltaTime;
+		
+		// Run the shake!
+		float CurTimeLeft = 0.f;
+		if (Entry.RootEvaluator)
+		{
+			Entry.RootEvaluator->Run(CurParams, CurResult);
+
+			FCameraNodeShakeParams ShakeParams(CurParams);
+			ShakeParams.ShakeScale = Entry.ShakeScale;
+			ShakeParams.PlaySpace = Entry.PlaySpace;
+			ShakeParams.UserPlaySpaceMatrix = Entry.UserPlaySpaceMatrix;
+
+			FCameraNodeShakeResult ShakeResult(CurResult);
+
+			Entry.RootEvaluator->ShakeResult(ShakeParams, ShakeResult);
+
+			CurTimeLeft = ShakeResult.ShakeTimeLeft;
+		}
+
+		// We are done with this shake this frame, so clear requests.
+		Entry.NumRequests = 0;
+
+		// If it says it's finished, schedule it for removal.
+		if (CurTimeLeft == 0)
+		{
+			EntriesToRemove.Add(Index);
+			continue;
+		}
+
+		// Update blends.
+		if (Entry.CameraShake->BlendIn && 
+				Entry.CurrentTime < Entry.CameraShake->BlendIn->BlendTime &&
+				ensure(Entry.BlendInEvaluator))
+		{
+			Entry.BlendInEvaluator->Run(CurParams, CurResult);
+
+			FCameraNodeBlendParams BlendParams(Params, Entry.Result);
+			FCameraNodeBlendResult BlendResult(OutResult);
+			Entry.BlendInEvaluator->BlendResults(BlendParams, BlendResult);
+		}
+		else if (Entry.CameraShake->BlendOut &&
+				CurTimeLeft >= 0.f && CurTimeLeft < Entry.CameraShake->BlendOut->BlendTime &&
+				ensure(Entry.BlendOutEvaluator))
+		{
+			Entry.BlendOutEvaluator->Run(CurParams, CurResult);
+
+			FCameraNodeBlendParams BlendParams(Params, Entry.Result);
+			FCameraNodeBlendResult BlendResult(OutResult);
+			Entry.BlendOutEvaluator->BlendResults(BlendParams, BlendResult);
+		}
+		else
+		{
+			OutResult.OverrideAll(CurResult);
+		}
+	}
+
+	// Remove any finished shakes.
+	for (int32 Index = EntriesToRemove.Num() - 1; Index >= 0; --Index)
+	{
+		PopEntry(EntriesToRemove[Index]);
+	}
+}
+
+bool FCameraShakeServiceCameraNodeEvaluator::InitializeEntry(
+	FShakeEntry& NewEntry, 
+	const UCameraShakeAsset* CameraShake,
+	TSharedPtr<const FCameraEvaluationContext> EvaluationContext)
+{
+	// Generate the hierarchy of node evaluators inside our storage buffer.
+	FCameraNodeEvaluatorTreeBuildParams BuildParams;
+	BuildParams.RootCameraNode = CameraShake->RootNode;
+	BuildParams.AllocationInfo = &CameraShake->AllocationInfo.EvaluatorInfo;
+	FCameraNodeEvaluator* RootEvaluator = NewEntry.EvaluatorStorage.BuildEvaluatorTree(BuildParams);
+	if (!ensureMsgf(RootEvaluator, TEXT("No root evaluator was created for new camera rig!")))
+	{
+		return false;
+	}
+
+	// Generate the blend-in and blend-out evaluators.
+	FBlendCameraNodeEvaluator* BlendInEvaluator = nullptr;
+	if (CameraShake->BlendIn)
+	{
+		FCameraNodeEvaluatorTreeBuildParams BlendBuildParams;
+		BlendBuildParams.RootCameraNode = CameraShake->BlendIn;
+		BlendInEvaluator = NewEntry.EvaluatorStorage.BuildEvaluatorTree(BlendBuildParams)->CastThis<FBlendCameraNodeEvaluator>();
+	}
+	FBlendCameraNodeEvaluator* BlendOutEvaluator = nullptr;
+	if (CameraShake->BlendOut)
+	{
+		FCameraNodeEvaluatorTreeBuildParams BlendBuildParams;
+		BlendBuildParams.RootCameraNode = CameraShake->BlendOut;
+		BlendOutEvaluator = NewEntry.EvaluatorStorage.BuildEvaluatorTree(BlendBuildParams)->CastThis<FBlendCameraNodeEvaluator>();
+
+		const bool bReversed = BlendOutEvaluator->SetReversed(true);
+		ensure(bReversed);  // TODO: if the blend can't play in reverse, wrap it in a reverse blend evaluator.
+	}
+
+	// Allocate variable table and context data table.
+	NewEntry.Result.VariableTable.Initialize(CameraShake->AllocationInfo.VariableTableInfo);
+	NewEntry.Result.ContextDataTable.Initialize(CameraShake->AllocationInfo.ContextDataTableInfo);
+
+	// Set all the data from the context.
+	const FCameraNodeEvaluationResult& ContextResult = EvaluationContext->GetInitialResult();
+	NewEntry.Result.VariableTable.OverrideAll(ContextResult.VariableTable, true);
+	NewEntry.Result.ContextDataTable.OverrideAll(ContextResult.ContextDataTable);
+
+	// Initialize the node evaluators.
+	FCameraNodeEvaluatorInitializeParams InitParams(&NewEntry.EvaluatorHierarchy);
+	InitParams.Evaluator = OwningEvaluator;
+	InitParams.EvaluationContext = EvaluationContext;
+	RootEvaluator->Initialize(InitParams, NewEntry.Result);
+
+	// Wrap up!
+	NewEntry.EvaluationContext = EvaluationContext;
+	NewEntry.CameraShake = CameraShake;
+	NewEntry.RootEvaluator = RootEvaluator->CastThisChecked<FShakeCameraNodeEvaluator>();
+	NewEntry.BlendInEvaluator = BlendInEvaluator;
+	NewEntry.BlendOutEvaluator = BlendOutEvaluator;
+	NewEntry.bIsFirstFrame = true;
+
+	return true;
+}
+
+void FCameraShakeServiceCameraNodeEvaluator::PopEntry(int32 EntryIndex)
+{
+	if (!ensure(Entries.IsValidIndex(EntryIndex)))
+	{
+		return;
+	}
+
+	Entries.RemoveAt(EntryIndex);
+}
+
+UE_DEFINE_CAMERA_EVALUATION_SERVICE(FCameraShakeService)
+
+void FCameraShakeService::OnInitialize(const FCameraEvaluationServiceInitializeParams& Params)
+{
+	ensure(Evaluator == nullptr);
+	Evaluator = Params.Evaluator;
+}
+
+void FCameraShakeService::OnTeardown(const FCameraEvaluationServiceTeardownParams& Params)
+{
+	ensure(Evaluator != nullptr);
+	Evaluator = nullptr;
+}
+
+void FCameraShakeService::StartCameraShake(const FStartCameraShakeParams& Params)
+{
+	EnsureShakeContextCreated();
+
+	if (ensure(ShakeEvaluator))
+	{
+		ShakeEvaluator->StartCameraShake(Params);
+	}
+}
+
+void FCameraShakeService::RequestCameraShakeThisFrame(const FStartCameraShakeParams& Params)
+{
+	EnsureShakeContextCreated();
+
+	if (ensure(ShakeEvaluator))
+	{
+		ShakeEvaluator->RequestCameraShakeThisFrame(Params);
+	}
+}
+
+void FCameraShakeService::EnsureShakeContextCreated()
+{
+	// Create the evaluation context, which is a "null" context with no particular logic.
+	if (!ShakeContext)
+	{
+		ShakeContext = MakeShared<FCameraEvaluationContext>();
+	}
+
+	// Create the camera rig that will contain and run all the camera shakes.
+	if (!ShakeContainerRig)
+	{
+		ShakeContainerRig = NewObject<UCameraRigAsset>(GetTransientPackage(), TEXT("CameraShakeContainerRig"), RF_Transient);
+		ShakeContainerRig->RootNode = NewObject<UCameraShakeServiceCameraNode>(ShakeContainerRig, NAME_None, RF_Transient);
+		ShakeContainerRig->BuildCameraRig();
+	}
+
+	// Instantiate the "container" camera rig inside the visual layer.
+	if (!ShakeEvaluator)
+	{
+		FRootCameraNodeEvaluator* RootEvaluator = Evaluator->GetRootNodeEvaluator();
+
+		FActivateCameraRigParams ActivateParams;
+		ActivateParams.EvaluationContext = ShakeContext;
+		ActivateParams.CameraRig = ShakeContainerRig;
+		ActivateParams.Layer = ECameraRigLayer::Visual;
+		FCameraRigInstanceID InstanceID = RootEvaluator->ActivateCameraRig(ActivateParams);
+
+		FCameraRigEvaluationInfo ShakeContainerRigInfo;
+		RootEvaluator->GetCameraRigInfo(InstanceID, ShakeContainerRigInfo);
+
+		if (ShakeContainerRigInfo.RootEvaluator)
+		{
+			ShakeEvaluator = ShakeContainerRigInfo.RootEvaluator->CastThis<FCameraShakeServiceCameraNodeEvaluator>();
+		}
+		ensure(ShakeEvaluator);
+	}
+}
+
+}  // namespace UE::Cameras
+
+FCameraNodeEvaluatorPtr UCameraShakeServiceCameraNode::OnBuildEvaluator(FCameraNodeEvaluatorBuilder& Builder) const
+{
+	using namespace UE::Cameras;
+	return Builder.BuildEvaluator<FCameraShakeServiceCameraNodeEvaluator>();
+}
+
diff --git a/Engine/Plugins/Cameras/GameplayCameras/Source/GameplayCameras/Public/Build/CameraShakeAssetBuilder.h b/Engine/Plugins/Cameras/GameplayCameras/Source/GameplayCameras/Public/Build/CameraShakeAssetBuilder.h
new file mode 100644
--- /dev/null
+++ b/Engine/Plugins/Cameras/GameplayCameras/Source/GameplayCameras/Public/Build/CameraShakeAssetBuilder.h
@@ -0,0 +1,37 @@
+// Copyright Epic Games, Inc. All Rights Reserved.
+
+#pragma once
+
+#include "Build/CameraBuildLog.h"
+#include "Core/CameraNodeHierarchy.h"
+
+class UCameraShakeAsset;
+
+namespace UE::Cameras
+{
+
+class GAMEPLAYCAMERAS_API FCameraShakeAssetBuilder
+{
+public:
+
+	FCameraShakeAssetBuilder(FCameraBuildLog& InBuildLog);
+
+	void BuildCameraShake(UCameraShakeAsset* InCameraShake);
+
+private:
+
+	void BuildCameraShakeImpl();
+
+	void UpdateBuildStatus();
+
+private:
+
+	FCameraBuildLog& BuildLog;
+
+	UCameraShakeAsset* CameraShake = nullptr;
+
+	FCameraNodeHierarchy CameraNodeHierarchy;
+};
+
+}  // namespace UE::Cameras
+
diff --git a/Engine/Plugins/Cameras/GameplayCameras/Source/GameplayCameras/Public/Core/CameraShakeAsset.h b/Engine/Plugins/Cameras/GameplayCameras/Source/GameplayCameras/Public/Core/CameraShakeAsset.h
new file mode 100644
--- /dev/null
+++ b/Engine/Plugins/Cameras/GameplayCameras/Source/GameplayCameras/Public/Core/CameraShakeAsset.h
@@ -0,0 +1,129 @@
+// Copyright Epic Games, Inc. All Rights Reserved.
+
+#pragma once
+
+#include "Build/CameraBuildStatus.h"
+#include "Core/BaseCameraObject.h"
+#include "Core/ObjectTreeGraphObject.h"
+#include "Core/ObjectTreeGraphRootObject.h"
+#include "Nodes/Blends/SimpleBlendCameraNode.h"
+
+#include "CameraShakeAsset.generated.h"
+
+class UCameraShakeAsset;
+class UShakeCameraNode;
+
+namespace UE::Cameras
+{
+	class FCameraBuildLog;
+	class FCameraShakeAssetBuilder;
+}
+
+UCLASS(MinimalAPI)
+class UCameraShakeAsset
+	: public UBaseCameraObject
+	, public IHasCameraBuildStatus
+	, public IObjectTreeGraphObject
+	, public IObjectTreeGraphRootObject
+{
+	GENERATED_BODY()
+
+public:
+
+	/** Root camera node. */
+	UPROPERTY(Instanced)
+	TObjectPtr<UShakeCameraNode> RootNode;
+
+	/** The blend to use for easing into the shake. */
+	UPROPERTY(Instanced)
+	TObjectPtr<USimpleFixedTimeBlendCameraNode> BlendIn;
+
+	/** The blend to use for easing out of the shake. */
+	UPROPERTY(Instanced)
+	TObjectPtr<USimpleFixedTimeBlendCameraNode> BlendOut;
+
+	/**
+	 * Whether only one instance of this shake can be started via the camera shake functions.
+	 * Note that this doesn't prevent this shake from being used "inline" inside a camera rig.
+	 */
+	UPROPERTY(EditAnywhere, Category="Common")
+	bool bIsSingleInstance = false;
+
+public:
+
+	/** Gets the camera shake's unique ID. */
+	const FGuid& GetGuid() const { return Guid; }
+
+public:
+
+	/** The current build state of this camera shake. */
+	UPROPERTY(Transient)
+	ECameraBuildStatus BuildStatus = ECameraBuildStatus::Dirty;
+
+	/** Builds this camera shake asset. */
+	GAMEPLAYCAMERAS_API void BuildCameraShake();
+
+	/** Builds this camera shake asset. */
+	GAMEPLAYCAMERAS_API void BuildCameraShake(UE::Cameras::FCameraBuildLog& InBuildLog);
+
+public:
+
+	// BaseCameraObject interface.
+	GAMEPLAYCAMERAS_API virtual UCameraNode* GetRootNode() override;
+
+	// IHasCameraBuildStatus interface.
+	virtual ECameraBuildStatus GetBuildStatus() const override { return BuildStatus; }
+	virtual void DirtyBuildStatus() override;
+
+protected:
+
+	// IObjectTreeGraphObject interface.
+#if WITH_EDITOR
+	virtual void GetGraphNodePosition(FName InGraphName, int32& NodePosX, int32& NodePosY) const override;
+	virtual void OnGraphNodeMoved(FName InGraphName, int32 NodePosX, int32 NodePosY, bool bMarkDirty) override;
+	virtual EObjectTreeGraphObjectSupportFlags GetSupportFlags(FName InGraphName) const override;
+	virtual const FString& GetGraphNodeCommentText(FName InGraphName) const override;
+	virtual void OnUpdateGraphNodeCommentText(FName InGraphName, const FString& NewComment) override;
+	virtual void GetGraphNodeName(FName InGraphName, FText& OutName) const override;
+#endif
+
+	// IObjectTreeGraphRootObject interface.
+#if WITH_EDITOR
+	virtual void GetConnectableObjects(FName InGraphName, TSet<UObject*>& OutObjects) const override;
+	virtual void AddConnectableObject(FName InGraphName, UObject* InObject) override;
+	virtual void RemoveConnectableObject(FName InGraphName, UObject* InObject) override;
+#endif
+
+	// UObject interface.
+	virtual void PostInitProperties() override;
+	virtual void PostDuplicate(EDuplicateMode::Type DuplicateMode) override;
+	virtual void PreSave(FObjectPreSaveContext ObjectSaveContext) override;
+
+private:
+
+	/** The camera shake's unique ID. */
+	UPROPERTY()
+	FGuid Guid;
+
+#if WITH_EDITORONLY_DATA
+
+	/** Position of the camera shake node in the node graph editor. */
+	UPROPERTY()
+	FIntVector2 GraphNodePos = FIntVector2::ZeroValue;
+
+	/** User-written comment in the node graph editor. */
+	UPROPERTY()
+	FString GraphNodeComment;
+
+	/** 
+	 * A list of all the camera nodes, including the 'loose' ones that aren't connected
+	 * to the root node, and therefore would be GC'ed if we didn't hold them here.
+	 */
+	UPROPERTY(Instanced, meta=(ObjectTreeGraphHidden=true))
+	TArray<TObjectPtr<UObject>> AllNodeObjects;
+
+#endif  // WITH_EDITORONLY_DATA
+
+	friend class UE::Cameras::FCameraShakeAssetBuilder;
+};
+
diff --git a/Engine/Plugins/Cameras/GameplayCameras/Source/GameplayCameras/Public/Core/CameraShakeAssetReference.h b/Engine/Plugins/Cameras/GameplayCameras/Source/GameplayCameras/Public/Core/CameraShakeAssetReference.h
new file mode 100644
--- /dev/null
+++ b/Engine/Plugins/Cameras/GameplayCameras/Source/GameplayCameras/Public/Core/CameraShakeAssetReference.h
@@ -0,0 +1,78 @@
+// Copyright Epic Games, Inc. All Rights Reserved.
+
+#pragma once
+
+#include "BaseCameraObjectReference.h"
+#include "UObject/ObjectPtr.h"
+
+#include "CameraShakeAssetReference.generated.h"
+
+class UBaseCameraObject;
+class UCameraShakeAsset;
+struct FCustomCameraNodeParameterInfos;
+
+namespace UE::Cameras
+{
+	class FCameraShakeAssetReferenceDetailsCustomization;
+	struct FCameraNodeEvaluationResult;
+}
+
+/**
+ * A structure holding a reference to a camera shake asset, along with the interface parameter
+ * override values.
+ */
+USTRUCT(BlueprintType)
+struct FCameraShakeAssetReference : public FBaseCameraObjectReference
+{
+	GENERATED_BODY()
+
+public:
+
+	FCameraShakeAssetReference();
+	FCameraShakeAssetReference(UCameraShakeAsset* InCameraShake);
+
+	/** Returns whether this reference points to a valid camera shake. */
+	bool IsValid() const
+	{
+		return CameraShake != nullptr;
+	}
+
+	/** Gets the referenced camera shake. */
+	UCameraShakeAsset* GetCameraShake()
+	{
+		return CameraShake;
+	}
+
+	/** Gets the referenced camera shake. */
+	const UCameraShakeAsset* GetCameraShake() const
+	{
+		return CameraShake;
+	}
+
+	/** Sets the referenced camera shake. */
+	void SetCameraShake(UCameraShakeAsset* InCameraShake)
+	{
+		if (CameraShake != InCameraShake)
+		{
+			CameraShake = InCameraShake;
+			RebuildParameters();
+		}
+	}
+
+	/** Applies the parameter override values to the given evaluation result. */
+	void ApplyParameterOverrides(UE::Cameras::FCameraNodeEvaluationResult& OutResult, bool bDrivenOverridesOnly);
+
+public:
+
+	// FBaseCameraObjectReference interface.
+	GAMEPLAYCAMERAS_API virtual const UBaseCameraObject* GetCameraObject() const override;
+
+private:
+
+	/** The referenced camera shake. */
+	UPROPERTY(EditAnywhere, Category="")
+	TObjectPtr<UCameraShakeAsset> CameraShake;
+
+	friend class UE::Cameras::FCameraShakeAssetReferenceDetailsCustomization;
+};
+
diff --git a/Engine/Plugins/Cameras/GameplayCameras/Source/GameplayCameras/Public/Core/DefaultRootCameraNode.h b/Engine/Plugins/Cameras/GameplayCameras/Source/GameplayCameras/Public/Core/DefaultRootCameraNode.h
--- a/Engine/Plugins/Cameras/GameplayCameras/Source/GameplayCameras/Public/Core/DefaultRootCameraNode.h
+++ b/Engine/Plugins/Cameras/GameplayCameras/Source/GameplayCameras/Public/Core/DefaultRootCameraNode.h
@@ -69,6 +69,7 @@ protected:
 	virtual void OnDeactivateCameraRig(const FDeactivateCameraRigParams& Params) override;
 	virtual void OnGetActiveCameraRigInfo(FCameraRigEvaluationInfo& OutCameraRigInfo) const override;
 	virtual void OnGetCameraRigInfo(const FCameraRigInstanceID InstanceID, FCameraRigEvaluationInfo& OutCameraRigInfo) const override;
+	virtual const FCameraVariableTable* OnGetBlendedParameters() const override;
 	virtual void OnBuildSingleCameraRigHierarchy(const FSingleCameraRigHierarchyBuildParams& Params, FCameraNodeEvaluatorHierarchy& OutHierarchy) override;
 	virtual void OnRunSingleCameraRig(const FSingleCameraRigEvaluationParams& Params, FCameraNodeEvaluationResult& OutResult) override;
 
diff --git a/Engine/Plugins/Cameras/GameplayCameras/Source/GameplayCameras/Public/Core/ObjectChildrenView.h b/Engine/Plugins/Cameras/GameplayCameras/Source/GameplayCameras/Public/Core/ObjectChildrenView.h
--- a/Engine/Plugins/Cameras/GameplayCameras/Source/GameplayCameras/Public/Core/ObjectChildrenView.h
+++ b/Engine/Plugins/Cameras/GameplayCameras/Source/GameplayCameras/Public/Core/ObjectChildrenView.h
@@ -9,6 +9,7 @@
 #include "UObject/ObjectPtr.h"
 
 #include <initializer_list>
+#include <type_traits>
 
 namespace UE::Cameras
 {
@@ -44,6 +45,13 @@ public:
 		Storage.template Set<FArray>(FArray(InChildren));
 	}
 
+	/** Sets the view to pointer storage and adds the given list of children pointers. */
+	template<typename OtherType, bool V = std::is_convertible_v<OtherType, ChildType>>
+	TObjectChildrenView(TArrayView<OtherType> InRange)
+	{
+		Storage.template Set<FArray>(FArray(InRange));
+	}
+
 	/**
 	 * Sets the view to pointer storage (if not already done) and adds the given
 	 * pointer to the list.
diff --git a/Engine/Plugins/Cameras/GameplayCameras/Source/GameplayCameras/Public/Core/RootCameraNode.h b/Engine/Plugins/Cameras/GameplayCameras/Source/GameplayCameras/Public/Core/RootCameraNode.h
--- a/Engine/Plugins/Cameras/GameplayCameras/Source/GameplayCameras/Public/Core/RootCameraNode.h
+++ b/Engine/Plugins/Cameras/GameplayCameras/Source/GameplayCameras/Public/Core/RootCameraNode.h
@@ -131,6 +131,9 @@ public:
 	/** Gets whether any camera rig is running. */
 	bool HasAnyActiveCameraRig() const;
 
+	/** Gets the variable table containing the blended camera object interfaces parameters. */
+	const FCameraVariableTable* GetBlendedParameters() const;
+
 	/**
 	 * Builds the hierarchy of the system for a given single camera rig.
 	 * This is expected to return the nodes of all the layers, except for the main layer which
@@ -168,6 +171,9 @@ protected:
 	/** Gets information about a specified camera rig. */
 	virtual void OnGetCameraRigInfo(const FCameraRigInstanceID InstanceID, FCameraRigEvaluationInfo& OutCameraRigInfo) const {}
 
+	/** Gets the variable table containing the blended camera object interfaces parameters. */
+	virtual const FCameraVariableTable* OnGetBlendedParameters() const { return nullptr; }
+
 	/* Builds the hierarchy of the system for a given single camera rig. */
 	virtual void OnBuildSingleCameraRigHierarchy(const FSingleCameraRigHierarchyBuildParams& Params, FCameraNodeEvaluatorHierarchy& OutHierarchy) {}
 
diff --git a/Engine/Plugins/Cameras/GameplayCameras/Source/GameplayCameras/Public/Core/ShakeCameraNode.h b/Engine/Plugins/Cameras/GameplayCameras/Source/GameplayCameras/Public/Core/ShakeCameraNode.h
new file mode 100644
--- /dev/null
+++ b/Engine/Plugins/Cameras/GameplayCameras/Source/GameplayCameras/Public/Core/ShakeCameraNode.h
@@ -0,0 +1,104 @@
+// Copyright Epic Games, Inc. All Rights Reserved.
+
+#pragma once
+
+#include "Camera/CameraTypes.h"
+#include "Core/CameraNode.h"
+#include "Core/CameraNodeEvaluator.h"
+
+#include "ShakeCameraNode.generated.h"
+
+/**
+ * Base class for shake camera nodes.
+ */
+UCLASS(MinimalAPI, Abstract, meta=(CameraNodeCategories="Shakes"))
+class UShakeCameraNode : public UCameraNode
+{
+	GENERATED_BODY()
+};
+
+namespace UE::Cameras
+{
+
+/**
+ * Parameters for applying a shake to a camera result.
+ */
+struct FCameraNodeShakeParams
+{
+	FCameraNodeShakeParams(const FCameraNodeEvaluationParams& InChildParams)
+		: ChildParams(InChildParams) 
+	{}
+
+	/** The parameters that the shake received during the evaluation. */
+	const FCameraNodeEvaluationParams& ChildParams;
+	/** The intensity to use for the camera shake. */
+	float ShakeScale = 1.f;
+	/** The play space to modify the result by */
+	ECameraShakePlaySpace PlaySpace = ECameraShakePlaySpace::CameraLocal;
+	/** The custom space to use for the shake. Only used when PlaySpace is UserDefined. */
+	FMatrix UserPlaySpaceMatrix;
+};
+
+/**
+ * Result structure for applying a shake to a camera result.
+ */
+struct FCameraNodeShakeResult
+{
+	FCameraNodeShakeResult(FCameraNodeEvaluationResult& InShakenResult)
+		: ShakenResult(InShakenResult)
+	{}
+
+	/** The result that should be shaken. */
+	FCameraNodeEvaluationResult& ShakenResult;
+
+	/** The time left in this shake, if applicable. Negative values indicate an infinite shake. */
+	float ShakeTimeLeft = 0.f;
+};
+
+/**
+ * Parameters for restarting a running camera shake.
+ */
+struct FCameraNodeShakeRestartParams
+{
+};
+
+class FShakeCameraNodeEvaluator : public FCameraNodeEvaluator
+{
+	UE_DECLARE_CAMERA_NODE_EVALUATOR(GAMEPLAYCAMERAS_API, FShakeCameraNodeEvaluator)
+
+public:
+
+	/** Applies the shake to the given result. */
+	GAMEPLAYCAMERAS_API void ShakeResult(const FCameraNodeShakeParams& Params, FCameraNodeShakeResult& OutResult);
+
+	/** Restart a running camera shake. */
+	GAMEPLAYCAMERAS_API void RestartShake(const FCameraNodeShakeRestartParams& Params);
+
+protected:
+
+	/** Applies the shake to the given result. */
+	virtual void OnShakeResult(const FCameraNodeShakeParams& Params, FCameraNodeShakeResult& OutResult) {}
+
+	/** Called when the intensity of the shake changes, but before ShakeScale is set. */
+	virtual void OnSetShakeScale(float InShakeScale) {}
+
+	/** Restart a running camera shake. */
+	virtual void OnRestartShake(const FCameraNodeShakeRestartParams& Params) {}
+};
+
+}  // namespace UE::Cameras
+
+// Macros for declaring and defining new shake node evaluators. They are the same
+// as the base ones for generic node evaluators, but the first one prevents you
+// from having to specify FShakeCameraNodeEvaluator as the base class, which saves
+// a little bit of typing.
+//
+#define UE_DECLARE_SHAKE_CAMERA_NODE_EVALUATOR(ApiDeclSpec, ClassName)\
+	UE_DECLARE_CAMERA_NODE_EVALUATOR_EX(ApiDeclSpec, ClassName, ::UE::Cameras::FShakeCameraNodeEvaluator)
+
+#define UE_DECLARE_SHAKE_CAMERA_NODE_EVALUATOR_EX(ApiDeclSpec, ClassName, BaseClassName)\
+	UE_DECLARE_CAMERA_NODE_EVALUATOR_EX(ApiDeclSpec, ClassName, BaseClassName)
+
+#define UE_DEFINE_SHAKE_CAMERA_NODE_EVALUATOR(ClassName)\
+	UE_DEFINE_CAMERA_NODE_EVALUATOR(ClassName)
+
diff --git a/Engine/Plugins/Cameras/GameplayCameras/Source/GameplayCameras/Public/Core/TransientBlendStackCameraNode.h b/Engine/Plugins/Cameras/GameplayCameras/Source/GameplayCameras/Public/Core/TransientBlendStackCameraNode.h
--- a/Engine/Plugins/Cameras/GameplayCameras/Source/GameplayCameras/Public/Core/TransientBlendStackCameraNode.h
+++ b/Engine/Plugins/Cameras/GameplayCameras/Source/GameplayCameras/Public/Core/TransientBlendStackCameraNode.h
@@ -64,6 +64,9 @@ public:
 	/** Freeze all camera rigs that belong to a given evaluation context. */
 	void FreezeAll(TSharedPtr<FCameraEvaluationContext> EvaluationContext);
 
+	/** Gets the variable table containing the blended camera object interfaces parameters. */
+	const FCameraVariableTable& GetBlendedParameters() const { return PreBlendVariableTable; }
+
 protected:
 
 	// FCameraNodeEvaluator interface.
diff --git a/Engine/Plugins/Cameras/GameplayCameras/Source/GameplayCameras/Public/GameplayCamerasDelegates.h b/Engine/Plugins/Cameras/GameplayCameras/Source/GameplayCameras/Public/GameplayCamerasDelegates.h
--- a/Engine/Plugins/Cameras/GameplayCameras/Source/GameplayCameras/Public/GameplayCamerasDelegates.h
+++ b/Engine/Plugins/Cameras/GameplayCameras/Source/GameplayCameras/Public/GameplayCamerasDelegates.h
@@ -8,6 +8,7 @@
 class UCameraAsset;
 class UCameraNode;
 class UCameraRigAsset;
+class UCameraShakeAsset;
 
 namespace UE::Cameras
 {
@@ -16,6 +17,7 @@ class FCameraBuildLog;
 
 DECLARE_MULTICAST_DELEGATE_OneParam(FOnCameraAssetBuilt, const UCameraAsset*);
 DECLARE_MULTICAST_DELEGATE_OneParam(FOnCameraRigAssetBuilt, const UCameraRigAsset*);
+DECLARE_MULTICAST_DELEGATE_OneParam(FOnCameraShakeAssetBuilt, const UCameraShakeAsset*);
 DECLARE_MULTICAST_DELEGATE_OneParam(FOnCameraNodeChanged, const UCameraNode*);
 
 /**
@@ -37,6 +39,12 @@ public:
 		return OnCameraRigAssetBuiltDelegate;
 	}
 
+	/** Broadcast when a camera shake has been built. */
+	static inline FOnCameraShakeAssetBuilt& OnCameraShakeAssetBuilt()
+	{
+		return OnCameraShakeAssetBuiltDelegate;
+	}
+
 	/** Broadcast when a custom camera parameter provider node changes it parameters. */
 	static inline FOnCameraNodeChanged& OnCustomCameraNodeParametersChanged()
 	{
@@ -47,6 +55,7 @@ private:
 
 	static FOnCameraAssetBuilt OnCameraAssetBuiltDelegate;
 	static FOnCameraRigAssetBuilt OnCameraRigAssetBuiltDelegate;
+	static FOnCameraShakeAssetBuilt OnCameraShakeAssetBuiltDelegate;
 	static FOnCameraNodeChanged OnCustomCameraNodeParametersChangedDelegate;
 };
 
diff --git a/Engine/Plugins/Cameras/GameplayCameras/Source/GameplayCameras/Public/Math/PerlinNoise.h b/Engine/Plugins/Cameras/GameplayCameras/Source/GameplayCameras/Public/Math/PerlinNoise.h
new file mode 100644
--- /dev/null
+++ b/Engine/Plugins/Cameras/GameplayCameras/Source/GameplayCameras/Public/Math/PerlinNoise.h
@@ -0,0 +1,100 @@
+// Copyright Epic Games, Inc. All Rights Reserved.
+
+#pragma once
+
+#include "CoreTypes.h"
+
+#include "PerlinNoise.generated.h"
+
+class FArchive;
+
+USTRUCT()
+struct FPerlinNoiseData
+{
+	GENERATED_BODY()
+
+	UPROPERTY(EditAnywhere, Category="Perlin Noise")
+	float Amplitude = 1.f;
+
+	UPROPERTY(EditAnywhere, Category="Perlin Noise")
+	float Frequency = 1.f;
+};
+
+namespace UE::Cameras
+{
+
+/**
+ * Implements a 1-dimensional perlin noise generator.
+ *
+ * Unlike the FMath::PerlinNoise1D method, this one doesn't use a hard-coded lookup
+ * table of pre-generated numbers. Since camera noise doesn't need to generate
+ * numbers for thousands or millions of pixels at a time (only a handful every frame),
+ * we can afford to get slightly better noise quality here.
+ *
+ * Other differences include:
+ *
+ *   - Multi-octave noise.
+ *   - The ability to dynamically change the frequency.
+ *   - Input is a delta-time, instead of an absolute undefined value.
+ */
+struct GAMEPLAYCAMERAS_API FPerlinNoise
+{
+public:
+
+	FPerlinNoise();
+	FPerlinNoise(const FPerlinNoiseData& InData, uint8 InOctaves = 1);
+	FPerlinNoise(float InAmplitude, float InFrequency, uint8 InOctaves = 1);
+
+	float GetAmplitude() const { return Amplitude; }
+	void SetAmplitude(float InAmplitude) { Amplitude = InAmplitude; }
+
+	float GetFrequency() const { return Octaves[0].Frequency; }
+	void SetFrequency(float InFrequency);
+
+	float GetLacunarity() const { return Lacunarity; }
+	void SetLacunarity(float InLacunarity) { Lacunarity = InLacunarity; }
+
+	float GetOctaveGain() const { return OctaveGain; }
+	void SetOctaveGain(float InOctaveGain) { OctaveGain = InOctaveGain; }
+
+	uint8 GetNumOctaves() const { return NumOctaves; }
+	void SetNumOctaves(uint8 InNumOctaves);
+
+	float GenerateValue(float DeltaTime);
+
+	void Serialize(FArchive& Ar);
+
+	friend FArchive& operator<< (FArchive& Ar, FPerlinNoise& This)
+	{
+		This.Serialize(Ar);
+		return Ar;
+	}
+
+private:
+
+	void Initialize(float InFrequency);
+
+private:
+
+	static const uint8 MAX_OCTAVES = 4;
+
+	struct FSinglePerlinNoise
+	{
+		float Frequency = 1.f;
+		float CurTime = 0.f;
+		float Prev = 0.f;
+		float Next = 0.f;
+	};
+
+	float Amplitude = 1.f;
+
+	float Lacunarity = 2.f;
+	float OctaveGain = 0.5f;
+
+	uint8 NumOctaves = 1;
+
+	FSinglePerlinNoise Octaves[MAX_OCTAVES];
+};
+
+}  // namespace UE::Cameras
+
diff --git a/Engine/Plugins/Cameras/GameplayCameras/Source/GameplayCameras/Public/Nodes/Shakes/CameraShakeCameraNode.h b/Engine/Plugins/Cameras/GameplayCameras/Source/GameplayCameras/Public/Nodes/Shakes/CameraShakeCameraNode.h
new file mode 100644
--- /dev/null
+++ b/Engine/Plugins/Cameras/GameplayCameras/Source/GameplayCameras/Public/Nodes/Shakes/CameraShakeCameraNode.h
@@ -0,0 +1,53 @@
+// Copyright Epic Games, Inc. All Rights Reserved.
+
+#pragma once
+
+#include "Core/CameraNode.h"
+#include "Core/CameraShakeAsset.h"
+#include "Core/CameraShakeAssetReference.h"
+#include "Core/ICustomCameraNodeParameterProvider.h"
+
+#include "CameraShakeCameraNode.generated.h"
+
+UENUM(BlueprintType)
+enum class ECameraShakeEvaluationMode : uint8
+{
+	Inline,
+	VisualLayer
+};
+
+/**
+ * A camera node that runs a camera shake asset.
+ */
+UCLASS(MinimalAPI, meta=(DisplayName="Camera Shake Prefab", CameraNodeCategories="Shakes"))
+class UCameraShakeCameraNode 
+	: public UCameraNode
+	, public ICustomCameraNodeParameterProvider
+{
+	GENERATED_BODY()
+
+protected:
+
+	// UCameraNode interface.
+	virtual void OnPreBuild(FCameraBuildLog& BuildLog) override;
+	virtual void OnBuild(FCameraObjectBuildContext& BuildContext) override;
+	virtual FCameraNodeEvaluatorPtr OnBuildEvaluator(FCameraNodeEvaluatorBuilder& Builder) const override;
+
+	// ICustomCameraNodeParameterProvider interface.
+	virtual void GetCustomCameraNodeParameters(FCustomCameraNodeParameterInfos& OutParameterInfos) override;
+
+	// UObject interface.
+#if WITH_EDITOR
+	virtual void PostEditChangeProperty( struct FPropertyChangedEvent& PropertyChangedEvent) override;
+#endif  // WITH_EDITOR
+
+public:
+
+	/** The camera shake to run. */
+	UPROPERTY(EditAnywhere, Category=Common, meta=(ObjectTreeGraphHidden=true))
+	FCameraShakeAssetReference CameraShakeReference;
+
+	UPROPERTY(EditAnywhere, Category=Common)
+	ECameraShakeEvaluationMode EvaluationMode = ECameraShakeEvaluationMode::VisualLayer;
+};
+
diff --git a/Engine/Plugins/Cameras/GameplayCameras/Source/GameplayCameras/Public/Nodes/Shakes/CompositeShakeCameraNode.h b/Engine/Plugins/Cameras/GameplayCameras/Source/GameplayCameras/Public/Nodes/Shakes/CompositeShakeCameraNode.h
new file mode 100644
--- /dev/null
+++ b/Engine/Plugins/Cameras/GameplayCameras/Source/GameplayCameras/Public/Nodes/Shakes/CompositeShakeCameraNode.h
@@ -0,0 +1,27 @@
+// Copyright Epic Games, Inc. All Rights Reserved.
+
+#pragma once
+
+#include "Containers/Array.h"
+#include "Core/ShakeCameraNode.h"
+#include "UObject/ObjectPtr.h"
+
+#include "CompositeShakeCameraNode.generated.h"
+
+UCLASS()
+class UCompositeShakeCameraNode : public UShakeCameraNode
+{
+	GENERATED_BODY()
+
+protected:
+
+	// UCameraNode interface.
+	virtual FCameraNodeEvaluatorPtr OnBuildEvaluator(FCameraNodeEvaluatorBuilder& Builder) const;
+
+public:
+
+	/** The shakes to run simultaneously, in order. */
+	UPROPERTY()
+	TArray<TObjectPtr<UShakeCameraNode>> Shakes;
+};
+
diff --git a/Engine/Plugins/Cameras/GameplayCameras/Source/GameplayCameras/Public/Nodes/Shakes/EnvelopeShakeCameraNode.h b/Engine/Plugins/Cameras/GameplayCameras/Source/GameplayCameras/Public/Nodes/Shakes/EnvelopeShakeCameraNode.h
new file mode 100644
--- /dev/null
+++ b/Engine/Plugins/Cameras/GameplayCameras/Source/GameplayCameras/Public/Nodes/Shakes/EnvelopeShakeCameraNode.h
@@ -0,0 +1,35 @@
+// Copyright Epic Games, Inc. All Rights Reserved.
+
+#pragma once
+
+#include "Core/CameraParameters.h"
+#include "Core/ShakeCameraNode.h"
+#include "UObject/ObjectPtr.h"
+
+#include "EnvelopeShakeCameraNode.generated.h"
+
+UCLASS()
+class UEnvelopeShakeCameraNode : public UShakeCameraNode
+{
+	GENERATED_BODY()
+
+protected:
+
+	// UCameraNode interface.
+	virtual FCameraNodeEvaluatorPtr OnBuildEvaluator(FCameraNodeEvaluatorBuilder& Builder) const;
+
+public:
+
+	UPROPERTY(EditAnywhere, Category="Easing")
+	FFloatCameraParameter EaseInTime = 0.2f;
+
+	UPROPERTY(EditAnywhere, Category="Easing")
+	FFloatCameraParameter EaseOutTime = 0.2f;
+
+	UPROPERTY(EditAnywhere, Category="Time")
+	FFloatCameraParameter TotalTime = 1.f;
+
+	UPROPERTY(EditAnywhere, Category="Common")
+	TObjectPtr<UShakeCameraNode> Shake;
+};
+
diff --git a/Engine/Plugins/Cameras/GameplayCameras/Source/GameplayCameras/Public/Nodes/Shakes/PerlinNoiseLocationShakeCameraNode.h b/Engine/Plugins/Cameras/GameplayCameras/Source/GameplayCameras/Public/Nodes/Shakes/PerlinNoiseLocationShakeCameraNode.h
new file mode 100644
--- /dev/null
+++ b/Engine/Plugins/Cameras/GameplayCameras/Source/GameplayCameras/Public/Nodes/Shakes/PerlinNoiseLocationShakeCameraNode.h
@@ -0,0 +1,41 @@
+// Copyright Epic Games, Inc. All Rights Reserved.
+
+#pragma once
+
+#include "Core/CameraParameters.h"
+#include "Core/ShakeCameraNode.h"
+#include "Math/PerlinNoise.h"
+
+#include "PerlinNoiseLocationShakeCameraNode.generated.h"
+
+UCLASS()
+class UPerlinNoiseLocationShakeCameraNode : public UShakeCameraNode
+{
+	GENERATED_BODY()
+
+protected:
+
+	// UCameraNode interface.
+	virtual FCameraNodeEvaluatorPtr OnBuildEvaluator(FCameraNodeEvaluatorBuilder& Builder) const;
+
+public:
+
+	UPROPERTY(EditAnywhere, Category="Perlin Noise")
+	FFloatCameraParameter AmplitudeMultiplier = 1.f;
+
+	UPROPERTY(EditAnywhere, Category="Perlin Noise")
+	FFloatCameraParameter FrequencyMultiplier = 1.f;
+
+	UPROPERTY(EditAnywhere, Category="Perlin Noise")
+	FInteger32CameraParameter Octaves = 1;
+
+	UPROPERTY(EditAnywhere, Category="Common")
+	FPerlinNoiseData X;
+
+	UPROPERTY(EditAnywhere, Category="Common")
+	FPerlinNoiseData Y;
+
+	UPROPERTY(EditAnywhere, Category="Common")
+	FPerlinNoiseData Z;
+};
+
diff --git a/Engine/Plugins/Cameras/GameplayCameras/Source/GameplayCameras/Public/Nodes/Shakes/PerlinNoiseRotationShakeCameraNode.h b/Engine/Plugins/Cameras/GameplayCameras/Source/GameplayCameras/Public/Nodes/Shakes/PerlinNoiseRotationShakeCameraNode.h
new file mode 100644
--- /dev/null
+++ b/Engine/Plugins/Cameras/GameplayCameras/Source/GameplayCameras/Public/Nodes/Shakes/PerlinNoiseRotationShakeCameraNode.h
@@ -0,0 +1,41 @@
+// Copyright Epic Games, Inc. All Rights Reserved.
+
+#pragma once
+
+#include "Core/CameraParameters.h"
+#include "Core/ShakeCameraNode.h"
+#include "Math/PerlinNoise.h"
+
+#include "PerlinNoiseRotationShakeCameraNode.generated.h"
+
+UCLASS()
+class UPerlinNoiseRotationShakeCameraNode : public UShakeCameraNode
+{
+	GENERATED_BODY()
+
+protected:
+
+	// UCameraNode interface.
+	virtual FCameraNodeEvaluatorPtr OnBuildEvaluator(FCameraNodeEvaluatorBuilder& Builder) const;
+
+public:
+
+	UPROPERTY(EditAnywhere, Category="Perlin Noise")
+	FFloatCameraParameter AmplitudeMultiplier = 1.f;
+
+	UPROPERTY(EditAnywhere, Category="Perlin Noise")
+	FFloatCameraParameter FrequencyMultiplier = 1.f;
+
+	UPROPERTY(EditAnywhere, Category="Perlin Noise")
+	FInteger32CameraParameter Octaves = 1;
+
+	UPROPERTY(EditAnywhere, Category="Common")
+	FPerlinNoiseData Yaw;
+
+	UPROPERTY(EditAnywhere, Category="Common")
+	FPerlinNoiseData Pitch;
+
+	UPROPERTY(EditAnywhere, Category="Common")
+	FPerlinNoiseData Roll;
+};
+
diff --git a/Engine/Plugins/Cameras/GameplayCameras/Source/GameplayCameras/Public/Services/CameraShakeService.h b/Engine/Plugins/Cameras/GameplayCameras/Source/GameplayCameras/Public/Services/CameraShakeService.h
new file mode 100644
--- /dev/null
+++ b/Engine/Plugins/Cameras/GameplayCameras/Source/GameplayCameras/Public/Services/CameraShakeService.h
@@ -0,0 +1,85 @@
+// Copyright Epic Games, Inc. All Rights Reserved.
+
+#pragma once
+
+#include "Camera/CameraTypes.h"
+#include "Core/CameraNode.h"
+#include "Core/CameraEvaluationService.h"
+
+#include "CameraShakeService.generated.h"
+
+class UCameraRigAsset;
+class UCameraShakeAsset;
+
+namespace UE::Cameras
+{
+
+class FCameraEvaluationContext;
+class FCameraShakeServiceCameraNodeEvaluator;
+
+/**
+ * Parameters for starting a camera shake.
+ */
+struct FStartCameraShakeParams
+{
+	/** The camera shake to play. */
+	const UCameraShakeAsset* CameraShake = nullptr;
+	/** The intensity to use for the camera shake. */
+	float ShakeScale = 1.f;
+	/** The play space to modify the result by */
+	ECameraShakePlaySpace PlaySpace = ECameraShakePlaySpace::CameraLocal;
+	/** The custom space to use for the shake. Only used when PlaySpace is UserDefined. */
+	FMatrix UserPlaySpaceMatrix;
+};
+
+/**
+ * A camera system service that handles running camera shakes in the visual layer.
+ */
+class FCameraShakeService : public FCameraEvaluationService
+{
+	UE_DECLARE_CAMERA_EVALUATION_SERVICE(GAMEPLAYCAMERAS_API, FCameraShakeService)
+
+public:
+
+	/** Starts a new camera shake. */
+	GAMEPLAYCAMERAS_API void StartCameraShake(const FStartCameraShakeParams& Params);
+
+	/**
+	 * Requests that a given camera shake runs. Multiple requests for the same
+	 * camera shake are combined, and the camera shake keeps running as long as
+	 * there are requests active. Requests are cleared every frame.
+	 */
+	GAMEPLAYCAMERAS_API void RequestCameraShakeThisFrame(const FStartCameraShakeParams& Params);
+
+protected:
+
+	// FCameraEvaluationService interface.
+	virtual void OnInitialize(const FCameraEvaluationServiceInitializeParams& Params) override;
+	virtual void OnTeardown(const FCameraEvaluationServiceTeardownParams& Params) override;
+
+private:
+
+	void EnsureShakeContextCreated();
+
+private:
+
+	FCameraSystemEvaluator* Evaluator = nullptr;
+
+	TSharedPtr<FCameraEvaluationContext> ShakeContext;
+	TObjectPtr<UCameraRigAsset> ShakeContainerRig = nullptr;
+	FCameraShakeServiceCameraNodeEvaluator* ShakeEvaluator = nullptr;
+};
+
+}  // namespace UE::Cameras
+
+UCLASS(MinimalAPI, Hidden)
+class UCameraShakeServiceCameraNode : public UCameraNode
+{
+	GENERATED_BODY()
+
+protected:
+
+	// UCameraNode interface.
+	virtual FCameraNodeEvaluatorPtr OnBuildEvaluator(FCameraNodeEvaluatorBuilder& Builder) const override;
+};
+
diff --git a/Engine/Plugins/Cameras/GameplayCameras/Source/GameplayCamerasEditor/Private/AssetTools/AssetDefinition_CameraRigAsset.cpp b/Engine/Plugins/Cameras/GameplayCameras/Source/GameplayCamerasEditor/Private/AssetTools/AssetDefinition_CameraRigAsset.cpp
--- a/Engine/Plugins/Cameras/GameplayCameras/Source/GameplayCamerasEditor/Private/AssetTools/AssetDefinition_CameraRigAsset.cpp
+++ b/Engine/Plugins/Cameras/GameplayCameras/Source/GameplayCamerasEditor/Private/AssetTools/AssetDefinition_CameraRigAsset.cpp
@@ -3,7 +3,6 @@
 #include "AssetTools/AssetDefinition_CameraRigAsset.h"
 
 #include "IGameplayCamerasEditorModule.h"
-#include "Styles/GameplayCamerasEditorStyle.h"
 #include "Toolkits/CameraRigAssetEditorToolkit.h"
 #include "Toolkits/IToolkit.h"
 #include "Modules/ModuleManager.h"
diff --git a/Engine/Plugins/Cameras/GameplayCameras/Source/GameplayCamerasEditor/Private/AssetTools/AssetDefinition_CameraShakeAsset.cpp b/Engine/Plugins/Cameras/GameplayCameras/Source/GameplayCamerasEditor/Private/AssetTools/AssetDefinition_CameraShakeAsset.cpp
new file mode 100644
--- /dev/null
+++ b/Engine/Plugins/Cameras/GameplayCameras/Source/GameplayCamerasEditor/Private/AssetTools/AssetDefinition_CameraShakeAsset.cpp
@@ -0,0 +1,55 @@
+// Copyright Epic Games, Inc. All Rights Reserved.
+
+#include "AssetTools/AssetDefinition_CameraShakeAsset.h"
+
+#include "IGameplayCamerasEditorModule.h"
+#include "Toolkits/CameraShakeAssetEditorToolkit.h"
+#include "Toolkits/IToolkit.h"
+#include "Modules/ModuleManager.h"
+
+#define LOCTEXT_NAMESPACE "AssetDefinition_CameraShakeAsset"
+
+#include UE_INLINE_GENERATED_CPP_BY_NAME(AssetDefinition_CameraShakeAsset)
+
+FText UAssetDefinition_CameraShakeAsset::GetAssetDisplayName() const
+{
+	return LOCTEXT("AssetDisplayName", "Camera Shake");
+}
+
+FLinearColor UAssetDefinition_CameraShakeAsset::GetAssetColor() const
+{
+	return FLinearColor(FColor(200, 80, 80));
+}
+
+TSoftClassPtr<UObject> UAssetDefinition_CameraShakeAsset::GetAssetClass() const
+{
+	return UCameraShakeAsset::StaticClass();
+}
+
+TConstArrayView<FAssetCategoryPath> UAssetDefinition_CameraShakeAsset::GetAssetCategories() const
+{
+	static const auto Categories = { FAssetCategoryPath(EAssetCategoryPaths::Gameplay) };
+	return Categories;
+}
+
+FAssetOpenSupport UAssetDefinition_CameraShakeAsset::GetAssetOpenSupport(const FAssetOpenSupportArgs& OpenSupportArgs) const
+{
+	return FAssetOpenSupport(
+			OpenSupportArgs.OpenMethod,
+			OpenSupportArgs.OpenMethod == EAssetOpenMethod::Edit,
+			EToolkitMode::Standalone); 
+}
+
+EAssetCommandResult UAssetDefinition_CameraShakeAsset::OpenAssets(const FAssetOpenArgs& OpenArgs) const
+{
+	IGameplayCamerasEditorModule& GameplayCamerasEditorModule = FModuleManager::LoadModuleChecked<IGameplayCamerasEditorModule>("GameplayCamerasEditor");
+	for (UCameraShakeAsset* CameraShake : OpenArgs.LoadObjects<UCameraShakeAsset>())
+	{
+		GameplayCamerasEditorModule.CreateCameraShakeEditor(OpenArgs.GetToolkitMode(), OpenArgs.ToolkitHost, CameraShake);
+	}	
+
+	return EAssetCommandResult::Handled;
+}
+
+#undef LOCTEXT_NAMESPACE
+
diff --git a/Engine/Plugins/Cameras/GameplayCameras/Source/GameplayCamerasEditor/Private/AssetTools/AssetDefinition_CameraShakeAsset.h b/Engine/Plugins/Cameras/GameplayCameras/Source/GameplayCamerasEditor/Private/AssetTools/AssetDefinition_CameraShakeAsset.h
new file mode 100644
--- /dev/null
+++ b/Engine/Plugins/Cameras/GameplayCameras/Source/GameplayCamerasEditor/Private/AssetTools/AssetDefinition_CameraShakeAsset.h
@@ -0,0 +1,25 @@
+// Copyright Epic Games, Inc. All Rights Reserved.
+
+#pragma once
+
+#include "AssetDefinitionDefault.h"
+#include "Core/CameraShakeAsset.h"
+
+#include "AssetDefinition_CameraShakeAsset.generated.h"
+
+UCLASS()
+class UAssetDefinition_CameraShakeAsset : public UAssetDefinitionDefault
+{
+	GENERATED_BODY()
+
+public:
+
+	// UAssetDefinition interface
+	virtual FText GetAssetDisplayName() const override;
+	virtual FLinearColor GetAssetColor() const override;
+	virtual TSoftClassPtr<UObject> GetAssetClass() const override;
+	virtual TConstArrayView<FAssetCategoryPath> GetAssetCategories() const override;
+	virtual FAssetOpenSupport GetAssetOpenSupport(const FAssetOpenSupportArgs& OpenSupportArgs) const override;
+	virtual EAssetCommandResult OpenAssets(const FAssetOpenArgs& OpenArgs) const override;
+};
+
diff --git a/Engine/Plugins/Cameras/GameplayCameras/Source/GameplayCamerasEditor/Private/AssetTools/CameraShakeAssetEditor.cpp b/Engine/Plugins/Cameras/GameplayCameras/Source/GameplayCamerasEditor/Private/AssetTools/CameraShakeAssetEditor.cpp
new file mode 100644
--- /dev/null
+++ b/Engine/Plugins/Cameras/GameplayCameras/Source/GameplayCamerasEditor/Private/AssetTools/CameraShakeAssetEditor.cpp
@@ -0,0 +1,29 @@
+// Copyright Epic Games, Inc. All Rights Reserved.
+
+#include "AssetTools/CameraShakeAssetEditor.h"
+
+#include "Core/CameraShakeAsset.h"
+#include "Toolkits/CameraShakeAssetEditorToolkit.h"
+
+#include UE_INLINE_GENERATED_CPP_BY_NAME(CameraShakeAssetEditor)
+
+void UCameraShakeAssetEditor::Initialize(TObjectPtr<UCameraShakeAsset> InCameraShakeAsset)
+{
+	CameraShakeAsset = InCameraShakeAsset;
+
+	Super::Initialize();
+}
+
+void UCameraShakeAssetEditor::GetObjectsToEdit(TArray<UObject*>& InObjectsToEdit)
+{
+	InObjectsToEdit.Add(CameraShakeAsset.Get());
+}
+
+TSharedPtr<FBaseAssetToolkit> UCameraShakeAssetEditor::CreateToolkit()
+{
+	using namespace UE::Cameras;
+	TSharedPtr<FCameraShakeAssetEditorToolkit> Toolkit = MakeShared<FCameraShakeAssetEditorToolkit>(this);
+	Toolkit->SetCameraShakeAsset(CameraShakeAsset);
+	return Toolkit;
+}
+
diff --git a/Engine/Plugins/Cameras/GameplayCameras/Source/GameplayCamerasEditor/Private/Commands/CameraShakeAssetEditorCommands.cpp b/Engine/Plugins/Cameras/GameplayCameras/Source/GameplayCamerasEditor/Private/Commands/CameraShakeAssetEditorCommands.cpp
new file mode 100644
--- /dev/null
+++ b/Engine/Plugins/Cameras/GameplayCameras/Source/GameplayCamerasEditor/Private/Commands/CameraShakeAssetEditorCommands.cpp
@@ -0,0 +1,43 @@
+// Copyright Epic Games, Inc. All Rights Reserved.
+
+#include "Commands/CameraShakeAssetEditorCommands.h"
+
+#include "Framework/Commands/Commands.h"
+#include "Framework/Commands/InputChord.h"
+#include "Framework/Commands/UICommandInfo.h"
+#include "GenericPlatform/GenericApplication.h"
+#include "Styles/GameplayCamerasEditorStyle.h"
+#include "Styling/AppStyle.h"
+
+#define LOCTEXT_NAMESPACE "CameraShakeAssetEditorCommands"
+
+namespace UE::Cameras
+{
+
+FCameraShakeAssetEditorCommands::FCameraShakeAssetEditorCommands()
+	: TCommands<FCameraShakeAssetEditorCommands>(
+			"CameraShakeAssetEditor",
+			LOCTEXT("CameraShakeAssetEditor", "Camera Shake Asset Editor"),
+			NAME_None,
+			FGameplayCamerasEditorStyle::Get()->GetStyleSetName()
+		)
+{
+}
+
+void FCameraShakeAssetEditorCommands::RegisterCommands()
+{
+	UI_COMMAND(Build, "Build", "Builds the asset and refreshes it in PIE",
+			EUserInterfaceActionType::Button, FInputChord(EKeys::F7));
+
+	UI_COMMAND(ShowMessages, "Messages", "Shows the message log for this camera rig",
+			EUserInterfaceActionType::ToggleButton, FInputChord());
+	UI_COMMAND(FindInCameraShake, "Search", "Searches for nodes in this camera rig",
+			EUserInterfaceActionType::Button, FInputChord(EKeys::F, EModifierKey::Control));
+	UI_COMMAND(FocusHome, "Home", "Focuses the graph canvas back on the root node",
+			EUserInterfaceActionType::Button, FInputChord());
+}
+
+}  // namespace UE::Cameras
+
+#undef LOCTEXT_NAMESPACE
+
diff --git a/Engine/Plugins/Cameras/GameplayCameras/Source/GameplayCamerasEditor/Private/Customizations/CameraShakeAssetReferenceDetailsCustomization.cpp b/Engine/Plugins/Cameras/GameplayCameras/Source/GameplayCamerasEditor/Private/Customizations/CameraShakeAssetReferenceDetailsCustomization.cpp
new file mode 100644
--- /dev/null
+++ b/Engine/Plugins/Cameras/GameplayCameras/Source/GameplayCamerasEditor/Private/Customizations/CameraShakeAssetReferenceDetailsCustomization.cpp
@@ -0,0 +1,99 @@
+// Copyright Epic Games, Inc. All Rights Reserved.
+
+#include "Customizations/CameraShakeAssetReferenceDetailsCustomization.h"
+
+#include "Core/CameraShakeAsset.h"
+#include "Core/CameraShakeAssetReference.h"
+#include "Customizations/CameraObjectInterfaceParameterOverrideDataDetails.h"
+#include "GameplayCamerasDelegates.h"
+#include "IDetailChildrenBuilder.h"
+#include "IPropertyUtilities.h"
+#include "PropertyBagDetails.h"
+#include "PropertyCustomizationHelpers.h"
+
+#define LOCTEXT_NAMESPACE "CameraShakeAssetReferenceDetailsCustomization"
+
+namespace UE::Cameras
+{
+
+using FCameraShakeAssetParameterOverrideDataDetails = TCameraObjectInterfaceParameterOverrideDataDetails<FCameraShakeAssetReference>;
+
+TSharedRef<IPropertyTypeCustomization> FCameraShakeAssetReferenceDetailsCustomization::MakeInstance()
+{
+	return MakeShared<FCameraShakeAssetReferenceDetailsCustomization>();
+}
+
+FCameraShakeAssetReferenceDetailsCustomization::~FCameraShakeAssetReferenceDetailsCustomization()
+{
+	FGameplayCamerasDelegates::OnCameraShakeAssetBuilt().RemoveAll(this);
+}
+
+void FCameraShakeAssetReferenceDetailsCustomization::CustomizeHeader(TSharedRef<IPropertyHandle> InStructPropertyHandle, FDetailWidgetRow& InHeaderRow, IPropertyTypeCustomizationUtils& InCustomizationUtils)
+{
+	StructPropertyHandle = InStructPropertyHandle;
+	PropertyUtilities = InCustomizationUtils.GetPropertyUtilities();
+
+	CameraShakeAssetPropertyHandle = StructPropertyHandle->GetChildHandle(GET_MEMBER_NAME_CHECKED(FCameraShakeAssetReference, CameraShake));
+	ParametersPropertyHandle = StructPropertyHandle->GetChildHandle(GET_MEMBER_NAME_CHECKED(FCameraShakeAssetReference, Parameters));
+
+	CameraShakeAssetPropertyHandle->SetOnPropertyValueChanged(
+			FSimpleDelegate::CreateSP(this, &FCameraShakeAssetReferenceDetailsCustomization::RebuildParametersIfNeeded));
+
+	if (!GIsTransacting)
+	{
+		RebuildParametersIfNeeded();
+	}
+
+	InHeaderRow
+	.ShouldAutoExpand(true)
+	.NameContent()
+	[
+		CameraShakeAssetPropertyHandle->CreatePropertyNameWidget()
+	]
+	.ValueContent()
+	[
+		CameraShakeAssetPropertyHandle->CreatePropertyValueWidgetWithCustomization(nullptr)
+	];
+	
+	FGameplayCamerasDelegates::OnCameraShakeAssetBuilt().AddSP(this, &FCameraShakeAssetReferenceDetailsCustomization::OnCameraShakeAssetBuilt);
+}
+
+void FCameraShakeAssetReferenceDetailsCustomization::CustomizeChildren(TSharedRef<IPropertyHandle> InStructPropertyHandle, IDetailChildrenBuilder& InChildrenBuilder, IPropertyTypeCustomizationUtils& InCustomizationUtils)
+{
+	const TSharedRef<FCameraShakeAssetParameterOverrideDataDetails> ParameterOverrideDataDetails = 
+		MakeShared<FCameraShakeAssetParameterOverrideDataDetails>(StructPropertyHandle, ParametersPropertyHandle, PropertyUtilities);
+	InChildrenBuilder.AddCustomBuilder(ParameterOverrideDataDetails);
+}
+
+void FCameraShakeAssetReferenceDetailsCustomization::OnCameraShakeAssetBuilt(const UCameraShakeAsset* CameraShake)
+{
+	RebuildParametersIfNeeded();
+}
+
+void FCameraShakeAssetReferenceDetailsCustomization::RebuildParametersIfNeeded()
+{
+	TArray<void*> RawData;
+	StructPropertyHandle->AccessRawData(RawData);
+
+	bool bRebuiltAny = false;
+
+	for (int32 Index = 0; Index < RawData.Num(); ++Index)
+	{
+		FCameraShakeAssetReference* CameraShakeAssetReference = static_cast<FCameraShakeAssetReference*>(RawData[Index]);
+		if (CameraShakeAssetReference)
+		{
+			const bool bRebuiltOne = CameraShakeAssetReference->RebuildParametersIfNeeded();
+			bRebuiltAny |= bRebuiltOne;
+		}
+	}
+
+	if (bRebuiltAny && PropertyUtilities)
+	{
+		PropertyUtilities->RequestRefresh();
+	}
+}
+
+}  // namespace UE::Cameras
+
+#undef LOCTEXT_NAMESPACE
+
diff --git a/Engine/Plugins/Cameras/GameplayCameras/Source/GameplayCamerasEditor/Private/Editors/CameraShakeCameraNodeGraphSchema.cpp b/Engine/Plugins/Cameras/GameplayCameras/Source/GameplayCamerasEditor/Private/Editors/CameraShakeCameraNodeGraphSchema.cpp
new file mode 100644
--- /dev/null
+++ b/Engine/Plugins/Cameras/GameplayCameras/Source/GameplayCamerasEditor/Private/Editors/CameraShakeCameraNodeGraphSchema.cpp
@@ -0,0 +1,41 @@
+// Copyright Epic Games, Inc. All Rights Reserved.
+
+#include "Editors/CameraShakeCameraNodeGraphSchema.h"
+
+#include "Core/CameraShakeAsset.h"
+#include "Core/ShakeCameraNode.h"
+#include "Editors/ObjectTreeGraphConfig.h"
+#include "GameplayCamerasEditorSettings.h"
+#include "Nodes/Blends/SimpleBlendCameraNode.h"
+
+#include UE_INLINE_GENERATED_CPP_BY_NAME(CameraShakeCameraNodeGraphSchema)
+
+#define LOCTEXT_NAMESPACE "CameraShakeCameraNodeGraphSchema"
+
+UCameraShakeCameraNodeGraphSchema::UCameraShakeCameraNodeGraphSchema(const FObjectInitializer& ObjInit)
+	: Super(ObjInit)
+{
+}
+
+FObjectTreeGraphConfig UCameraShakeCameraNodeGraphSchema::BuildGraphConfig() const
+{
+	const UGameplayCamerasEditorSettings* Settings = GetDefault<UGameplayCamerasEditorSettings>();
+
+	FObjectTreeGraphConfig GraphConfig;
+
+	BuildBaseGraphConfig(GraphConfig);
+
+	GraphConfig.ConnectableObjectClasses.Add(UShakeCameraNode::StaticClass());
+	GraphConfig.ConnectableObjectClasses.Add(USimpleFixedTimeBlendCameraNode::StaticClass());
+	GraphConfig.ConnectableObjectClasses.Add(UCameraShakeAsset::StaticClass());
+	GraphConfig.ObjectClassConfigs.Emplace(UCameraShakeAsset::StaticClass())
+		.OnlyAsRoot()
+		.HasSelfPin(false)
+		.NodeTitleUsesObjectName(true)
+		.NodeTitleColor(Settings->CameraShakeAssetTitleColor);
+
+	return GraphConfig;
+}
+
+#undef LOCTEXT_NAMESPACE
+
diff --git a/Engine/Plugins/Cameras/GameplayCameras/Source/GameplayCamerasEditor/Private/Factories/CameraShakeAssetFactory.cpp b/Engine/Plugins/Cameras/GameplayCameras/Source/GameplayCamerasEditor/Private/Factories/CameraShakeAssetFactory.cpp
new file mode 100644
--- /dev/null
+++ b/Engine/Plugins/Cameras/GameplayCameras/Source/GameplayCamerasEditor/Private/Factories/CameraShakeAssetFactory.cpp
@@ -0,0 +1,37 @@
+// Copyright Epic Games, Inc. All Rights Reserved.
+
+#include "Factories/CameraShakeAssetFactory.h"
+
+#include "Core/CameraShakeAsset.h"
+
+#include UE_INLINE_GENERATED_CPP_BY_NAME(CameraShakeAssetFactory)
+
+#define LOCTEXT_NAMESPACE "CameraShakeAssetFactory"
+
+UCameraShakeAssetFactory::UCameraShakeAssetFactory(const FObjectInitializer& ObjectInitializer)
+	: Super(ObjectInitializer)
+{
+	bCreateNew = true;
+	bEditAfterNew = true;
+	SupportedClass = UCameraShakeAsset::StaticClass();
+}
+
+UObject* UCameraShakeAssetFactory::FactoryCreateNew(UClass* Class, UObject* Parent, FName Name, EObjectFlags Flags, UObject* Context, FFeedbackContext* Warn)
+{
+	UCameraShakeAsset* NewCameraShake = NewObject<UCameraShakeAsset>(Parent, Class, Name, Flags | RF_Transactional);
+	return NewCameraShake;
+}
+
+bool UCameraShakeAssetFactory::ConfigureProperties()
+{
+	return true;
+}
+
+bool UCameraShakeAssetFactory::ShouldShowInNewMenu() const
+{
+	return true;
+}
+
+#undef LOCTEXT_NAMESPACE
+
+
diff --git a/Engine/Plugins/Cameras/GameplayCameras/Source/GameplayCamerasEditor/Private/Factories/CameraShakeAssetFactory.h b/Engine/Plugins/Cameras/GameplayCameras/Source/GameplayCamerasEditor/Private/Factories/CameraShakeAssetFactory.h
new file mode 100644
--- /dev/null
+++ b/Engine/Plugins/Cameras/GameplayCameras/Source/GameplayCamerasEditor/Private/Factories/CameraShakeAssetFactory.h
@@ -0,0 +1,26 @@
+// Copyright Epic Games, Inc. All Rights Reserved.
+
+#pragma once
+
+#include "Factories/Factory.h"
+
+#include "CameraShakeAssetFactory.generated.h"
+
+class UCameraShakeAsset;
+
+/**
+ * Implements a factory for UCameraShakeAsset objects.
+ */
+UCLASS(hidecategories=Object)
+class UCameraShakeAssetFactory : public UFactory
+{
+	GENERATED_BODY()
+
+	UCameraShakeAssetFactory(const FObjectInitializer& ObjectInitializer);
+
+	// UFactory Interface
+	virtual UObject* FactoryCreateNew(UClass* Class, UObject* Parent, FName Name, EObjectFlags Flags, UObject* Context, FFeedbackContext* Warn) override;
+	virtual bool ConfigureProperties() override;
+	virtual bool ShouldShowInNewMenu() const override;
+};
+
diff --git a/Engine/Plugins/Cameras/GameplayCameras/Source/GameplayCamerasEditor/Private/GameplayCamerasEditorModule.cpp b/Engine/Plugins/Cameras/GameplayCameras/Source/GameplayCamerasEditor/Private/GameplayCamerasEditorModule.cpp
--- a/Engine/Plugins/Cameras/GameplayCameras/Source/GameplayCamerasEditor/Private/GameplayCamerasEditorModule.cpp
+++ b/Engine/Plugins/Cameras/GameplayCameras/Source/GameplayCamerasEditor/Private/GameplayCamerasEditorModule.cpp
@@ -5,10 +5,12 @@
 #include "AssetTools/CameraAssetEditor.h"
 #include "AssetTools/CameraRigAssetEditor.h"
 #include "AssetTools/CameraRigProxyAssetEditor.h"
+#include "AssetTools/CameraShakeAssetEditor.h"
 #include "AssetTools/CameraVariableCollectionEditor.h"
 #include "Commands/CameraAssetEditorCommands.h"
 #include "Commands/CameraRigAssetEditorCommands.h"
 #include "Commands/CameraRigTransitionEditorCommands.h"
+#include "Commands/CameraShakeAssetEditorCommands.h"
 #include "Commands/CameraVariableCollectionEditorCommands.h"
 #include "Commands/GameplayCamerasDebuggerCommands.h"
 #include "Commands/ObjectTreeGraphEditorCommands.h"
@@ -17,6 +19,7 @@
 #include "Customizations/CameraParameterDetailsCustomizations.h"
 #include "Customizations/CameraProxyTableDetailsCustomization.h"
 #include "Customizations/CameraRigAssetReferenceDetailsCustomization.h"
+#include "Customizations/CameraShakeAssetReferenceDetailsCustomization.h"
 #include "Customizations/CameraVariableReferenceDetailsCustomizations.h"
 #include "Customizations/FilmbackCameraNodeDetailsCustomization.h"
 #include "Customizations/RichCurveDetailsCustomizations.h"
@@ -52,6 +55,7 @@
 #include "Toolkits/BlueprintCameraDirectorAssetEditorMode.h"
 #include "Toolkits/CameraAssetEditorToolkit.h"
 #include "Toolkits/CameraRigAssetEditorToolkit.h"
+#include "Toolkits/CameraShakeAssetEditorToolkit.h"
 #include "Toolkits/SingleCameraDirectorAssetEditorMode.h"
 #include "Trace/CameraSystemRewindDebuggerExtension.h"
 #include "Trace/CameraSystemRewindDebuggerTrack.h"
@@ -117,6 +121,7 @@ public:
 		FCameraAssetEditorCommands::Unregister();
 		FCameraRigAssetEditorCommands::Unregister();
 		FCameraRigTransitionEditorCommands::Unregister();
+		FCameraShakeAssetEditorCommands::Unregister();
 		FCameraVariableCollectionEditorCommands::Unregister();
 		FGameplayCamerasDebuggerCommands::Unregister();
 		FObjectTreeGraphEditorCommands::Unregister();
@@ -160,6 +165,14 @@ public:
 		return AssetEditor;
 	}
 
+	virtual UCameraShakeAssetEditor* CreateCameraShakeEditor(const EToolkitMode::Type Mode, const TSharedPtr<IToolkitHost>& InitToolkitHost, UCameraShakeAsset* CameraShake) override
+	{
+		UAssetEditorSubsystem* AssetEditorSubsystem = GEditor->GetEditorSubsystem<UAssetEditorSubsystem>();
+		UCameraShakeAssetEditor* AssetEditor = NewObject<UCameraShakeAssetEditor>(AssetEditorSubsystem, NAME_None, RF_Transient);
+		AssetEditor->Initialize(CameraShake);
+		return AssetEditor;
+	}
+
 	virtual UCameraVariableCollectionEditor* CreateCameraVariableCollectionEditor(const EToolkitMode::Type Mode, const TSharedPtr<IToolkitHost>& InitToolkitHost, UCameraVariableCollection* VariableCollection) override
 	{
 		UAssetEditorSubsystem* AssetEditorSubsystem = GEditor->GetEditorSubsystem<UAssetEditorSubsystem>();
@@ -382,6 +395,7 @@ private:
 		FCameraAssetEditorCommands::Register();
 		FCameraRigAssetEditorCommands::Register();	
 		FCameraRigTransitionEditorCommands::Register();
+		FCameraShakeAssetEditorCommands::Register();	
 		FCameraVariableCollectionEditorCommands::Register();
 		FGameplayCamerasDebuggerCommands::Register();
 		FObjectTreeGraphEditorCommands::Register();
@@ -435,6 +449,10 @@ private:
 				"CameraRigAssetReference", 
 				FOnGetPropertyTypeCustomizationInstance::CreateStatic(
 					&FCameraRigAssetReferenceDetailsCustomization::MakeInstance));
+		PropertyEditorModule.RegisterCustomPropertyTypeLayout(
+				"CameraShakeAssetReference", 
+				FOnGetPropertyTypeCustomizationInstance::CreateStatic(
+					&FCameraShakeAssetReferenceDetailsCustomization::MakeInstance));
 
 		PropertyEditorModule.RegisterCustomClassLayout(
 				"FilmbackCameraNode", 
@@ -457,6 +475,7 @@ private:
 			PropertyEditorModule->UnregisterCustomPropertyTypeLayout("CameraRigProxyTableEntry");
 			PropertyEditorModule->UnregisterCustomPropertyTypeLayout("CameraRigAsset");
 			PropertyEditorModule->UnregisterCustomPropertyTypeLayout("CameraRigAssetReference");
+			PropertyEditorModule->UnregisterCustomPropertyTypeLayout("CameraShakeAssetReference");
 
 			PropertyEditorModule->UnregisterCustomClassLayout("FilmbackCameraNode");
 		}
diff --git a/Engine/Plugins/Cameras/GameplayCameras/Source/GameplayCamerasEditor/Private/GameplayCamerasEditorSettings.cpp b/Engine/Plugins/Cameras/GameplayCameras/Source/GameplayCamerasEditor/Private/GameplayCamerasEditorSettings.cpp
--- a/Engine/Plugins/Cameras/GameplayCameras/Source/GameplayCamerasEditor/Private/GameplayCamerasEditorSettings.cpp
+++ b/Engine/Plugins/Cameras/GameplayCameras/Source/GameplayCamerasEditor/Private/GameplayCamerasEditorSettings.cpp
@@ -10,6 +10,7 @@ UGameplayCamerasEditorSettings::UGameplayCamerasEditorSettings(const FObjectInit
 	CameraNodeTitleColor = FLinearColor(0.190525f, 0.583898f, 1.0f, 1.0f);           // Light blue
 	CameraAssetTitleColor = FLinearColor(1.0f, 0.170000f, 0.0f, 1.0f);				 // Orange
 	CameraRigAssetTitleColor = FLinearColor(1.0f, 0.170000f, 0.0f, 1.0f);            // Orange
+	CameraShakeAssetTitleColor = FLinearColor(1.0f, 0.170000f, 0.0f, 1.0f);          // Orange
 	CameraRigTransitionTitleColor = FLinearColor(1.0f, 0.65f, 0.4f, 1.0f);           // Beige
 	CameraRigTransitionConditionTitleColor = FLinearColor(0.8f, 0.4f, 0.4f, 1.0f);   // Salmon
 	CameraBlendNodeTitleColor = FLinearColor(0.6f, 0.0f, 1.0f, 1.0f);                // Purple
diff --git a/Engine/Plugins/Cameras/GameplayCameras/Source/GameplayCamerasEditor/Private/Styles/GameplayCamerasEditorStyle.cpp b/Engine/Plugins/Cameras/GameplayCameras/Source/GameplayCamerasEditor/Private/Styles/GameplayCamerasEditorStyle.cpp
--- a/Engine/Plugins/Cameras/GameplayCameras/Source/GameplayCamerasEditor/Private/Styles/GameplayCamerasEditorStyle.cpp
+++ b/Engine/Plugins/Cameras/GameplayCameras/Source/GameplayCamerasEditor/Private/Styles/GameplayCamerasEditorStyle.cpp
@@ -82,22 +82,29 @@ FGameplayCamerasEditorStyle::FGameplayCamerasEditorStyle()
 	Set("CameraRigAssetEditor.Tabs.Curves", new IMAGE_BRUSH_SVG("Icons/CurveEditor", Icon16x16));
 	Set("CameraRigAssetEditor.Tabs.InterfaceParameters", new IMAGE_BRUSH_SVG("Icons/CameraRig-InterfaceParameters", Icon16x16));
 
-	Set("CameraRigAssetEditor.Build", new IMAGE_BRUSH_SVG("Icons/CameraRig-BuildStatus_Background", Icon20x20));
-	Set("CameraRigAssetEditor.BuildStatus.Background", new IMAGE_BRUSH_SVG("Icons/CameraRig-BuildStatus_Background", Icon20x20));
-	Set("CameraRigAssetEditor.BuildStatus.Overlay.Error", new IMAGE_BRUSH_SVG("Icons/CameraRig-BuildStatus_Fail_Badge", Icon20x20, FStyleColors::Error));
-	Set("CameraRigAssetEditor.BuildStatus.Overlay.Good", new IMAGE_BRUSH_SVG("Icons/CameraRig-BuildStatus_Good_Badge", Icon20x20, FStyleColors::AccentGreen));
-	Set("CameraRigAssetEditor.BuildStatus.Overlay.Unknown", new IMAGE_BRUSH_SVG("Icons/CameraRig-BuildStatus_Unknown_Badge", Icon20x20, FStyleColors::AccentYellow));
-	Set("CameraRigAssetEditor.BuildStatus.Overlay.Warning", new IMAGE_BRUSH_SVG("Icons/CameraRig-BuildStatus_Warning_Badge", Icon20x20, FStyleColors::Warning));
-
 	Set("CameraRigAssetEditor.ShowNodeHierarchy", new IMAGE_BRUSH_SVG("Icons/CameraRig-NodeHierarchy", Icon20x20));
 	Set("CameraRigAssetEditor.ShowTransitions", new IMAGE_BRUSH_SVG("Icons/CameraRig-Transitions", Icon20x20));
 	Set("CameraRigAssetEditor.FocusHome", new IMAGE_BRUSH_SVG("Icons/GraphEditor-Home", Icon20x20));
-
-	Set("CameraRigAssetEditor.LiveUpdate", new IMAGE_BRUSH_SVG("Icons/CameraRig-LiveUpdate", Icon20x20));
-	Set("CameraRigAssetEditor.Apply", new IMAGE_BRUSH_SVG("CameraEditor-Apply", Icon20x20));
 	Set("CameraRigAssetEditor.FindInCameraRig", new CORE_IMAGE_BRUSH_SVG("Starship/Common/Search", Icon20x20));
 
-	Set("CameraRigAssetEditor.InterfaceParameter.Message", FTextBlockStyle(NormalText).SetFont(DEFAULT_FONT("Italic", 10)));
+	// Camera shake editor icons.
+	Set("CameraShakeAssetEditor.Tabs.Toolbox", new IMAGE_BRUSH_SVG("Icons/CameraRig-Toolbox", Icon16x16));
+	Set("CameraShakeAssetEditor.Tabs.Search", new CORE_IMAGE_BRUSH_SVG("Starship/Common/search", Icon16x16));
+	Set("CameraShakeAssetEditor.Tabs.Messages", new CORE_IMAGE_BRUSH_SVG("Starship/Common/OutputLog", Icon16x16));
+	Set("CameraShakeAssetEditor.Tabs.InterfaceParameters", new IMAGE_BRUSH_SVG("Icons/CameraRig-InterfaceParameters", Icon16x16));
+
+	Set("CameraShakeAssetEditor.FocusHome", new IMAGE_BRUSH_SVG("Icons/GraphEditor-Home", Icon20x20));
+	Set("CameraShakeAssetEditor.FindInCameraShake", new CORE_IMAGE_BRUSH_SVG("Starship/Common/Search", Icon20x20));
+
+	// Camera object build statuses.
+	Set("CameraObjectEditor.BuildStatus.Background", new IMAGE_BRUSH_SVG("Icons/CameraRig-BuildStatus_Background", Icon20x20));
+	Set("CameraObjectEditor.BuildStatus.Overlay.Error", new IMAGE_BRUSH_SVG("Icons/CameraRig-BuildStatus_Fail_Badge", Icon20x20, FStyleColors::Error));
+	Set("CameraObjectEditor.BuildStatus.Overlay.Good", new IMAGE_BRUSH_SVG("Icons/CameraRig-BuildStatus_Good_Badge", Icon20x20, FStyleColors::AccentGreen));
+	Set("CameraObjectEditor.BuildStatus.Overlay.Unknown", new IMAGE_BRUSH_SVG("Icons/CameraRig-BuildStatus_Unknown_Badge", Icon20x20, FStyleColors::AccentYellow));
+	Set("CameraObjectEditor.BuildStatus.Overlay.Warning", new IMAGE_BRUSH_SVG("Icons/CameraRig-BuildStatus_Warning_Badge", Icon20x20, FStyleColors::Warning));
+
+	// Camera object editor icons.
+	Set("CameraObjectEditor.InterfaceParameter.Message", FTextBlockStyle(NormalText).SetFont(DEFAULT_FONT("Italic", 10)));
 
 	// Camera parameters icons.
 	Set("CameraParameter.VariableBrowser", new IMAGE_BRUSH_SVG("Icons/CameraParameter-Variable", Icon16x16));
diff --git a/Engine/Plugins/Cameras/GameplayCameras/Source/GameplayCamerasEditor/Private/Toolkits/BuildButtonToolkit.cpp b/Engine/Plugins/Cameras/GameplayCameras/Source/GameplayCamerasEditor/Private/Toolkits/BuildButtonToolkit.cpp
--- a/Engine/Plugins/Cameras/GameplayCameras/Source/GameplayCamerasEditor/Private/Toolkits/BuildButtonToolkit.cpp
+++ b/Engine/Plugins/Cameras/GameplayCameras/Source/GameplayCamerasEditor/Private/Toolkits/BuildButtonToolkit.cpp
@@ -36,11 +36,11 @@ FToolMenuEntry FBuildButtonToolkit::MakeToolbarButton(TSharedPtr<FUICommandInfo>
 
 FSlateIcon FBuildButtonToolkit::GetBuildButtonIcon() const
 {
-	static const FName BuildStatusBackground("CameraRigAssetEditor.BuildStatus.Background");
-	static const FName BuildStatusError("CameraRigAssetEditor.BuildStatus.Overlay.Error");
-	static const FName BuildStatusGood("CameraRigAssetEditor.BuildStatus.Overlay.Good");
-	static const FName BuildStatusUnknown("CameraRigAssetEditor.BuildStatus.Overlay.Unknown");
-	static const FName BuildStatusWarning("CameraRigAssetEditor.BuildStatus.Overlay.Warning");
+	static const FName BuildStatusBackground("CameraObjectEditor.BuildStatus.Background");
+	static const FName BuildStatusError("CameraObjectEditor.BuildStatus.Overlay.Error");
+	static const FName BuildStatusGood("CameraObjectEditor.BuildStatus.Overlay.Good");
+	static const FName BuildStatusUnknown("CameraObjectEditor.BuildStatus.Overlay.Unknown");
+	static const FName BuildStatusWarning("CameraObjectEditor.BuildStatus.Overlay.Warning");
 
 	const FName CamerasStyleSetName = FGameplayCamerasEditorStyle::Get()->GetStyleSetName();
 
diff --git a/Engine/Plugins/Cameras/GameplayCameras/Source/GameplayCamerasEditor/Private/Toolkits/CameraShakeAssetEditorToolkit.cpp b/Engine/Plugins/Cameras/GameplayCameras/Source/GameplayCamerasEditor/Private/Toolkits/CameraShakeAssetEditorToolkit.cpp
new file mode 100644
--- /dev/null
+++ b/Engine/Plugins/Cameras/GameplayCameras/Source/GameplayCamerasEditor/Private/Toolkits/CameraShakeAssetEditorToolkit.cpp
@@ -0,0 +1,437 @@
+// Copyright Epic Games, Inc. All Rights Reserved.
+
+#include "Toolkits/CameraShakeAssetEditorToolkit.h"
+
+#include "Build/CameraBuildLog.h"
+#include "Build/CameraShakeAssetBuilder.h"
+#include "Commands/CameraShakeAssetEditorCommands.h"
+#include "Core/CameraShakeAsset.h"
+#include "Editors/CameraShakeCameraNodeGraphSchema.h"
+#include "Editors/ObjectTreeGraph.h"
+#include "Editors/ObjectTreeGraphNode.h"
+#include "Editors/SCameraNodeGraphEditor.h"
+#include "Editors/SFindInObjectTreeGraph.h"
+#include "Editors/SObjectTreeGraphEditor.h"
+#include "Editors/SObjectTreeGraphToolbox.h"
+#include "Framework/Docking/LayoutExtender.h"
+#include "Framework/Docking/TabManager.h"
+#include "Helpers/AssetTypeMenuOverlayHelper.h"
+#include "IGameplayCamerasEditorModule.h"
+#include "Modules/ModuleManager.h"
+#include "PropertyEditorModule.h"
+#include "Styles/GameplayCamerasEditorStyle.h"
+#include "ToolMenus.h"
+#include "Toolkits/BuildButtonToolkit.h"
+#include "Toolkits/CameraBuildLogToolkit.h"
+#include "Toolkits/CameraObjectInterfaceParametersToolkit.h"
+#include "Toolkits/StandardToolkitLayout.h"
+#include "Widgets/Docking/SDockTab.h"
+
+#include UE_INLINE_GENERATED_CPP_BY_NAME(CameraShakeAssetEditorToolkit)
+
+#define LOCTEXT_NAMESPACE "CameraShakeAssetEditorToolkit"
+
+namespace UE::Cameras
+{
+
+const FName FCameraShakeAssetEditorToolkit::CameraShakeEditorTabId(TEXT("CameraShakeAssetEditor_GraphEditor"));
+const FName FCameraShakeAssetEditorToolkit::DetailsViewTabId(TEXT("CameraShakeAssetEditor_DetailsView"));
+const FName FCameraShakeAssetEditorToolkit::SearchTabId(TEXT("CameraShakeAssetEditor_Search"));
+const FName FCameraShakeAssetEditorToolkit::MessagesTabId(TEXT("CameraShakeAssetEditor_Messages"));
+const FName FCameraShakeAssetEditorToolkit::ToolboxTabId(TEXT("CameraShakeAssetEditor_Toolbox"));
+const FName FCameraShakeAssetEditorToolkit::InterfaceParametersTabId(TEXT("CameraShakeAssetEditor_InterfaceParameters"));
+
+FCameraShakeAssetEditorToolkit::FCameraShakeAssetEditorToolkit(UAssetEditor* InOwningAssetEditor)
+	: FBaseAssetToolkit(InOwningAssetEditor)
+{
+	BuildButtonToolkit = MakeShared<FBuildButtonToolkit>();
+	BuildLogToolkit = MakeShared<FCameraBuildLogToolkit>();
+	InterfaceParametersToolkit = MakeShared<FCameraObjectInterfaceParametersToolkit>();
+
+	StandardLayout = MakeShared<FStandardToolkitLayout>(TEXT("CameraShakeAssetEditor_Layout_v1"));
+	{
+		StandardLayout->AddCenterTab(CameraShakeEditorTabId);
+
+		StandardLayout->AddRightTab(DetailsViewTabId);
+
+		StandardLayout->AddBottomTab(SearchTabId);
+		StandardLayout->AddBottomTab(MessagesTabId);
+
+		StandardLayout->AddLeftTab(ToolboxTabId);
+		StandardLayout->AddLeftTab(InterfaceParametersTabId, ETabState::OpenedTab);
+	}
+	StandaloneDefaultLayout = StandardLayout->GetLayout();
+}
+
+FCameraShakeAssetEditorToolkit::~FCameraShakeAssetEditorToolkit()
+{
+}
+
+void FCameraShakeAssetEditorToolkit::AddReferencedObjects(FReferenceCollector& Collector)
+{
+	Collector.AddReferencedObject(CameraShakeAsset);
+}
+
+FString FCameraShakeAssetEditorToolkit::GetReferencerName() const
+{
+	return TEXT("FCameraRigAssetEditorToolkitBase");
+}
+
+void FCameraShakeAssetEditorToolkit::SetCameraShakeAsset(UCameraShakeAsset* InCameraShake)
+{
+	CameraShakeAsset = InCameraShake;
+
+	BuildButtonToolkit->SetTarget(InCameraShake);
+	InterfaceParametersToolkit->SetCameraObject(InCameraShake);
+}
+
+FText FCameraShakeAssetEditorToolkit::GetCameraShakeAssetName() const
+{
+	if (CameraShakeAsset)
+	{
+		return NodeGraphConfig.GetDisplayNameText(CameraShakeAsset);
+	}
+	return LOCTEXT("NoCameraShake", "No Camera Shake");
+}
+
+bool FCameraShakeAssetEditorToolkit::IsGraphEditorEnabled() const
+{
+	return CameraShakeAsset != nullptr;
+}
+
+void FCameraShakeAssetEditorToolkit::RegisterTabSpawners(const TSharedRef<class FTabManager>& InTabManager)
+{
+	// Skip FBaseAssetToolkit here because we don't want a viewport tab.
+	FAssetEditorToolkit::RegisterTabSpawners(InTabManager);
+
+	const FName CamerasStyleSetName = FGameplayCamerasEditorStyle::Get()->GetStyleSetName();
+
+	InTabManager->RegisterTabSpawner(ToolboxTabId, FOnSpawnTab::CreateSP(this, &FCameraShakeAssetEditorToolkit::SpawnTab_Toolbox))
+		.SetDisplayName(LOCTEXT("Toolbox", "Toolbox"))
+		.SetGroup(AssetEditorTabsCategory.ToSharedRef())
+		.SetIcon(FSlateIcon(CamerasStyleSetName, "CameraShakeAssetEditor.Tabs.Toolbox"));
+
+	InTabManager->RegisterTabSpawner(CameraShakeEditorTabId, FOnSpawnTab::CreateSP(this, &FCameraShakeAssetEditorToolkit::SpawnTab_CameraShakeEditor))
+		.SetDisplayName(LOCTEXT("CameraShakeEditor", "Camera Shake"))
+		.SetGroup(AssetEditorTabsCategory.ToSharedRef());
+
+	InTabManager->RegisterTabSpawner(DetailsViewTabId, FOnSpawnTab::CreateSP(this, &FCameraShakeAssetEditorToolkit::SpawnTab_Details))
+		.SetDisplayName(LOCTEXT("Details", "Details"))
+		.SetGroup(AssetEditorTabsCategory.ToSharedRef())
+		.SetIcon(FSlateIcon(FAppStyle::GetAppStyleSetName(), "LevelEditor.Tabs.Details"));
+
+	InTabManager->RegisterTabSpawner(SearchTabId, FOnSpawnTab::CreateSP(this, &FCameraShakeAssetEditorToolkit::SpawnTab_Search))
+		.SetDisplayName(LOCTEXT("Search", "Search"))
+		.SetGroup(AssetEditorTabsCategory.ToSharedRef())
+		.SetIcon(FSlateIcon(CamerasStyleSetName, "CameraShakeAssetEditor.Tabs.Search"));
+
+	InTabManager->RegisterTabSpawner(MessagesTabId, FOnSpawnTab::CreateSP(this, &FCameraShakeAssetEditorToolkit::SpawnTab_Messages))
+		.SetDisplayName(LOCTEXT("Messages", "Messages"))
+		.SetGroup(AssetEditorTabsCategory.ToSharedRef())
+		.SetIcon(FSlateIcon(CamerasStyleSetName, "CameraShakeAssetEditor.Tabs.Messages"));
+
+	InTabManager->RegisterTabSpawner(InterfaceParametersTabId, FOnSpawnTab::CreateSP(this, &FCameraShakeAssetEditorToolkit::SpawnTab_InterfaceParameters))
+		.SetDisplayName(LOCTEXT("InterfaceParameters", "InterfaceParameters"))
+		.SetGroup(AssetEditorTabsCategory.ToSharedRef())
+		.SetIcon(FSlateIcon(CamerasStyleSetName, "CameraShakeAssetEditor.Tabs.InterfaceParameters"));
+}
+
+TSharedRef<SDockTab> FCameraShakeAssetEditorToolkit::SpawnTab_CameraShakeEditor(const FSpawnTabArgs& Args)
+{
+	TSharedPtr<SDockTab> CameraShakeEditorTab = SNew(SDockTab)
+		.Label(LOCTEXT("CameraShakeEditorTabTitle", "Camera Shake Editor"))
+		[
+			NodeGraphEditor.ToSharedRef()
+		];
+
+	return CameraShakeEditorTab.ToSharedRef();
+}
+
+TSharedRef<SDockTab> FCameraShakeAssetEditorToolkit::SpawnTab_Details(const FSpawnTabArgs& Args)
+{
+	TSharedPtr<SDockTab> DetailsTab = SNew(SDockTab)
+		.Label(LOCTEXT("DetailsTitle", "Details"))
+		[
+			DetailsView.ToSharedRef()
+		];
+
+	return DetailsTab.ToSharedRef();
+}
+
+TSharedRef<SDockTab> FCameraShakeAssetEditorToolkit::SpawnTab_Search(const FSpawnTabArgs& Args)
+{
+	TSharedPtr<SDockTab> SearchTab = SNew(SDockTab)
+		.Label(LOCTEXT("SearchTabTitle", "Search"))
+		[
+			SearchWidget.ToSharedRef()
+		];
+
+	return SearchTab.ToSharedRef();
+}
+
+TSharedRef<SDockTab> FCameraShakeAssetEditorToolkit::SpawnTab_Messages(const FSpawnTabArgs& Args)
+{
+	TSharedPtr<SDockTab> MessagesTab = SNew(SDockTab)
+		.Label(LOCTEXT("MessagesTabTitle", "Messages"))
+		[
+			BuildLogToolkit->GetMessagesWidget().ToSharedRef()
+		];
+
+	return MessagesTab.ToSharedRef();
+}
+
+TSharedRef<SDockTab> FCameraShakeAssetEditorToolkit::SpawnTab_Toolbox(const FSpawnTabArgs& Args)
+{
+	TSharedPtr<SDockTab> ToolboxTab = SNew(SDockTab)
+		.Label(LOCTEXT("ToolboxTabTitle", "Toolbox"))
+		[
+			ToolboxWidget.ToSharedRef()
+		];
+
+	return ToolboxTab.ToSharedRef();
+}
+
+TSharedRef<SDockTab> FCameraShakeAssetEditorToolkit::SpawnTab_InterfaceParameters(const FSpawnTabArgs& Args)
+{
+	TSharedPtr<SDockTab> InterfaceParametersTab = SNew(SDockTab)
+		.Label(LOCTEXT("InterfaceParametersTabTitle", "Parameters"))
+		[
+			InterfaceParametersToolkit->GetInterfaceParametersPanel().ToSharedRef()
+		];
+
+	return InterfaceParametersTab.ToSharedRef();
+}
+
+void FCameraShakeAssetEditorToolkit::UnregisterTabSpawners(const TSharedRef<FTabManager>& InTabManager)
+{
+	// Skip FBaseAssetToolkit here because we don't want a viewport tab.
+	FAssetEditorToolkit::UnregisterTabSpawners(InTabManager);
+	
+	InTabManager->UnregisterTabSpawner(CameraShakeEditorTabId);
+	InTabManager->UnregisterTabSpawner(DetailsViewTabId);
+	InTabManager->UnregisterTabSpawner(SearchTabId);
+	InTabManager->UnregisterTabSpawner(MessagesTabId);
+	InTabManager->UnregisterTabSpawner(ToolboxTabId);
+	InTabManager->UnregisterTabSpawner(InterfaceParametersTabId);
+}
+
+void FCameraShakeAssetEditorToolkit::CreateWidgets()
+{
+	// Skip FBaseAssetToolkit here because we don't want a viewport tab, and our base class
+	// has its own details view in order to get a notify hook.
+	// ...no up-call...
+
+	RegisterToolbar();
+	CreateEditorModeManager();
+	LayoutExtender = MakeShared<FLayoutExtender>();
+
+	// Now do our custom stuff.
+
+	// Create the details view.
+	FPropertyEditorModule& PropertyEditorModule = FModuleManager::GetModuleChecked<FPropertyEditorModule>("PropertyEditor");
+	FDetailsViewArgs DetailsViewArgs;
+	DetailsViewArgs.NameAreaSettings = FDetailsViewArgs::HideNameArea;
+	DetailsViewArgs.bHideSelectionTip = true;
+	DetailsViewArgs.NotifyHook = this;
+	DetailsView = PropertyEditorModule.CreateDetailView(DetailsViewArgs);
+
+	// Create the camera rig editor.
+	CreateNodeGraphEditor();
+
+	// Create the toolbox, default to the rig editor items.
+	ToolboxWidget = SNew(SObjectTreeGraphToolbox)
+		.GraphConfig(NodeGraphConfig);
+
+	// Create the search panel.
+	SearchWidget = SNew(SFindInObjectTreeGraph)
+		.OnGetRootObjectsToSearch(this, &FCameraShakeAssetEditorToolkit::OnGetRootObjectsToSearch)
+		.OnJumpToObjectRequested(this, &FCameraShakeAssetEditorToolkit::OnJumpToObject);
+
+	// Create the message log.
+	BuildLogToolkit->Initialize("CameraShakeAssetBuildMessages");
+
+	// Hook-up the selection of interface parameters.
+	InterfaceParametersToolkit->OnInterfaceParameterSelected().AddSP(this, &FCameraShakeAssetEditorToolkit::OnCameraObjectInterfaceParameterSelected);
+}
+
+void FCameraShakeAssetEditorToolkit::CreateNodeGraphEditor()
+{
+	UClass* SchemaClass = UCameraShakeCameraNodeGraphSchema::StaticClass();
+	UCameraShakeCameraNodeGraphSchema* DefaultSchemaObject = Cast<UCameraShakeCameraNodeGraphSchema>(SchemaClass->GetDefaultObject());
+	NodeGraphConfig = DefaultSchemaObject->BuildGraphConfig();
+
+	NodeGraph = NewObject<UObjectTreeGraph>(GetTransientPackage(), NAME_None, RF_Transactional | RF_Standalone);
+	NodeGraph->Schema = SchemaClass;
+	NodeGraph->Reset(CameraShakeAsset, NodeGraphConfig);
+
+	FGraphAppearanceInfo Appearance;
+	Appearance.CornerText = LOCTEXT("CameraShakeGraphText", "CAMERA SHAKE");
+
+	NodeGraphEditor = SNew(SCameraNodeGraphEditor)
+		.Appearance(Appearance)
+		.DetailsView(DetailsView)
+		.GraphTitle(this, &FCameraShakeAssetEditorToolkit::GetCameraShakeAssetName)
+		.IsEnabled(this, &FCameraShakeAssetEditorToolkit::IsGraphEditorEnabled)
+		.GraphToEdit(NodeGraph)
+		.AssetEditorToolkit(SharedThis(this));
+}
+
+void FCameraShakeAssetEditorToolkit::RegisterToolbar()
+{
+	FName ParentName;
+	const FName MenuName = GetToolMenuToolbarName(ParentName);
+	UToolMenus* ToolMenus = UToolMenus::Get();
+	if (!ToolMenus->IsMenuRegistered(MenuName))
+	{
+		FToolMenuOwnerScoped ToolMenuOwnerScope(this);
+		FToolMenuInsert InsertAfterAssetSection("Asset", EToolMenuInsertType::After);
+
+		UToolMenu* ToolbarMenu = UToolMenus::Get()->RegisterMenu(
+				MenuName, ParentName, EMultiBoxType::ToolBar);
+
+		ToolbarMenu->AddDynamicSection("Tools", FNewToolMenuDelegate::CreateLambda(
+				[](UToolMenu* InMenu)
+				{
+					UCameraShakeAssetEditorMenuContext* Context = InMenu->FindContext<UCameraShakeAssetEditorMenuContext>();
+					FCameraShakeAssetEditorToolkit* This = Context ? Context->Toolkit.Pin().Get() : nullptr;
+					if (!ensure(This))
+					{
+						return;
+					}
+
+					const FCameraShakeAssetEditorCommands& Commands = FCameraShakeAssetEditorCommands::Get();
+
+					FToolMenuSection& ToolsSection = InMenu->AddSection("Tools");
+					ToolsSection.AddEntry(This->BuildButtonToolkit->MakeToolbarButton(Commands.Build));
+					ToolsSection.AddEntry(FToolMenuEntry::InitToolBarButton(Commands.FindInCameraShake));
+				}),
+				InsertAfterAssetSection);
+
+		const FCameraShakeAssetEditorCommands& Commands = FCameraShakeAssetEditorCommands::Get();
+
+		FToolMenuSection& GraphsSection = ToolbarMenu->AddSection("Graphs", TAttribute<FText>(), InsertAfterAssetSection);
+
+		GraphsSection.AddEntry(FToolMenuEntry::InitToolBarButton(Commands.FocusHome));
+	}
+}
+
+void FCameraShakeAssetEditorToolkit::InitToolMenuContext(FToolMenuContext& MenuContext)
+{
+	FBaseAssetToolkit::InitToolMenuContext(MenuContext);
+
+	UCameraShakeAssetEditorMenuContext* Context = NewObject<UCameraShakeAssetEditorMenuContext>();
+	Context->Toolkit = SharedThis(this);
+	MenuContext.AddObject(Context);
+}
+
+void FCameraShakeAssetEditorToolkit::PostInitAssetEditor()
+{
+	const FCameraShakeAssetEditorCommands& Commands = FCameraShakeAssetEditorCommands::Get();
+
+	ToolkitCommands->MapAction(
+		Commands.Build,
+		FExecuteAction::CreateSP(this, &FCameraShakeAssetEditorToolkit::OnBuild));
+
+	ToolkitCommands->MapAction(
+		Commands.FocusHome,
+		FExecuteAction::CreateSP(this, &FCameraShakeAssetEditorToolkit::OnFocusHome));
+
+	ToolkitCommands->MapAction(
+		Commands.FindInCameraShake,
+		FExecuteAction::CreateSP(this, &FCameraShakeAssetEditorToolkit::OnFindInCameraShake));
+
+	BuildLogToolkit->OnRequestJumpToObject().BindSPLambda(this, [this](UObject* Object)
+		{
+			OnJumpToObject(Object, NAME_None);
+		});
+
+	RegenerateMenusAndToolbars();
+}
+
+void FCameraShakeAssetEditorToolkit::PostRegenerateMenusAndToolbars()
+{
+	SetMenuOverlay(FAssetTypeMenuOverlayHelper::CreateMenuOverlay(UCameraShakeAsset::StaticClass()));
+}
+
+void FCameraShakeAssetEditorToolkit::NotifyPostChange(const FPropertyChangedEvent& PropertyChangedEvent, FProperty* PropertyThatChanged)
+{
+	// Called when something is modified in the details view.
+	if (CameraShakeAsset)
+	{
+		CameraShakeAsset->BuildStatus = ECameraBuildStatus::Dirty;
+	}
+}
+
+void FCameraShakeAssetEditorToolkit::OnCameraObjectInterfaceParameterSelected(UCameraObjectInterfaceParameterBase* Object)
+{
+	OnJumpToObject(Object, NAME_None);
+}
+
+void FCameraShakeAssetEditorToolkit::OnBuild()
+{
+	if (!CameraShakeAsset)
+	{
+		return;
+	}
+
+	FCameraBuildLog BuildLog;
+	FCameraShakeAssetBuilder Builder(BuildLog);
+	Builder.BuildCameraShake(CameraShakeAsset);
+
+	BuildLogToolkit->PopulateMessageListing(BuildLog);
+
+	if (CameraShakeAsset->BuildStatus != ECameraBuildStatus::Clean)
+	{
+		TabManager->TryInvokeTab(MessagesTabId);
+	}
+}
+
+void FCameraShakeAssetEditorToolkit::OnFindInCameraShake()
+{
+	TabManager->TryInvokeTab(SearchTabId);
+	SearchWidget->FocusSearchEditBox();
+}
+
+void FCameraShakeAssetEditorToolkit::OnGetRootObjectsToSearch(TArray<FFindInObjectTreeGraphSource>& OutSources)
+{
+	OutSources.Add(FFindInObjectTreeGraphSource{ CameraShakeAsset, &NodeGraphConfig });
+}
+
+void FCameraShakeAssetEditorToolkit::OnFocusHome()
+{
+	OnJumpToObject(CameraShakeAsset, NAME_None);
+}
+
+void FCameraShakeAssetEditorToolkit::OnJumpToObject(UObject* Object, FName PropertyName)
+{
+	if (UObjectTreeGraphNode* NodeGraphObjectNode = NodeGraph->FindObjectNode(Object))
+	{
+		NodeGraphEditor->JumpToNode(NodeGraphObjectNode);
+	}
+}
+
+FText FCameraShakeAssetEditorToolkit::GetBaseToolkitName() const
+{
+	return LOCTEXT("AppLabel", "Camera Shake Asset");
+}
+
+FName FCameraShakeAssetEditorToolkit::GetToolkitFName() const
+{
+	static FName ToolkitName("CameraShakeAssetEditor");
+	return ToolkitName;
+}
+
+FString FCameraShakeAssetEditorToolkit::GetWorldCentricTabPrefix() const
+{
+	return LOCTEXT("WorldCentricTabPrefix", "Camera Shake Asset ").ToString();
+}
+
+FLinearColor FCameraShakeAssetEditorToolkit::GetWorldCentricTabColorScale() const
+{
+	return FLinearColor(0.7f, 0.0f, 0.0f, 0.5f);
+}
+
+}  // namespace UE::Cameras
+
+#undef LOCTEXT_NAMESPACE
+
diff --git a/Engine/Plugins/Cameras/GameplayCameras/Source/GameplayCamerasEditor/Public/AssetTools/CameraShakeAssetEditor.h b/Engine/Plugins/Cameras/GameplayCameras/Source/GameplayCamerasEditor/Public/AssetTools/CameraShakeAssetEditor.h
new file mode 100644
--- /dev/null
+++ b/Engine/Plugins/Cameras/GameplayCameras/Source/GameplayCamerasEditor/Public/AssetTools/CameraShakeAssetEditor.h
@@ -0,0 +1,37 @@
+// Copyright Epic Games, Inc. All Rights Reserved.
+
+#pragma once
+
+#include "Tools/UAssetEditor.h"
+
+#include "CameraShakeAssetEditor.generated.h"
+
+class FBaseAssetToolkit;
+class UCameraShakeAsset;
+
+/**
+ * Editor for a camera rig asset.
+ */
+UCLASS(Transient)
+class UCameraShakeAssetEditor : public UAssetEditor
+{
+	GENERATED_BODY()
+
+public:
+
+	void Initialize(TObjectPtr<UCameraShakeAsset> InCameraShakeAsset);
+
+	UCameraShakeAsset* GetCameraShakeAsset() const { return CameraShakeAsset; }
+
+public:
+
+	// UAssetEditor interface
+	virtual void GetObjectsToEdit(TArray<UObject*>& InObjectsToEdit) override;
+	virtual TSharedPtr<FBaseAssetToolkit> CreateToolkit() override;
+
+private:
+
+	UPROPERTY()
+	TObjectPtr<UCameraShakeAsset> CameraShakeAsset;
+};
+
diff --git a/Engine/Plugins/Cameras/GameplayCameras/Source/GameplayCamerasEditor/Public/Commands/CameraShakeAssetEditorCommands.h b/Engine/Plugins/Cameras/GameplayCameras/Source/GameplayCamerasEditor/Public/Commands/CameraShakeAssetEditorCommands.h
new file mode 100644
--- /dev/null
+++ b/Engine/Plugins/Cameras/GameplayCameras/Source/GameplayCamerasEditor/Public/Commands/CameraShakeAssetEditorCommands.h
@@ -0,0 +1,33 @@
+// Copyright Epic Games, Inc. All Rights Reserved.
+
+#pragma once
+
+#include "CoreTypes.h"
+#include "Framework/Commands/Commands.h"
+
+class FUICommandInfo;
+class UClass;
+struct FInputChord;
+
+namespace UE::Cameras
+{
+
+class FCameraShakeAssetEditorCommands : public TCommands<FCameraShakeAssetEditorCommands>
+{
+public:
+
+	FCameraShakeAssetEditorCommands();
+
+	virtual void RegisterCommands() override;
+	
+public:
+
+	TSharedPtr<FUICommandInfo> Build;
+	
+	TSharedPtr<FUICommandInfo> ShowMessages;
+	TSharedPtr<FUICommandInfo> FindInCameraShake;
+	TSharedPtr<FUICommandInfo> FocusHome;
+};
+
+}  // namespace UE::Cameras
+
diff --git a/Engine/Plugins/Cameras/GameplayCameras/Source/GameplayCamerasEditor/Public/Customizations/CameraShakeAssetReferenceDetailsCustomization.h b/Engine/Plugins/Cameras/GameplayCameras/Source/GameplayCamerasEditor/Public/Customizations/CameraShakeAssetReferenceDetailsCustomization.h
new file mode 100644
--- /dev/null
+++ b/Engine/Plugins/Cameras/GameplayCameras/Source/GameplayCamerasEditor/Public/Customizations/CameraShakeAssetReferenceDetailsCustomization.h
@@ -0,0 +1,42 @@
+// Copyright Epic Games, Inc. All Rights Reserved.
+
+#pragma once
+
+#include "IPropertyTypeCustomization.h"
+
+class IPropertyHandle;
+class IPropertyUtilities;
+class UCameraShakeAsset;
+
+namespace UE::Cameras
+{
+
+class FCameraShakeAssetReferenceDetailsCustomization : public IPropertyTypeCustomization
+{
+public:
+
+	static TSharedRef<IPropertyTypeCustomization> MakeInstance();
+
+public:
+
+	~FCameraShakeAssetReferenceDetailsCustomization();
+
+	// IPropertyTypeCustomization interface.
+	virtual void CustomizeHeader(TSharedRef<IPropertyHandle> InStructPropertyHandle, FDetailWidgetRow& InHeaderRow, IPropertyTypeCustomizationUtils& InCustomizationUtils) override;
+	virtual void CustomizeChildren(TSharedRef<IPropertyHandle> InStructPropertyHandle, IDetailChildrenBuilder& InChildrenBuilder, IPropertyTypeCustomizationUtils& InCustomizationUtils) override;
+
+private:
+
+	void OnCameraShakeAssetBuilt(const UCameraShakeAsset* CameraShake);
+	void RebuildParametersIfNeeded();
+
+private:
+
+	TSharedPtr<IPropertyHandle> StructPropertyHandle;
+	TSharedPtr<IPropertyHandle> CameraShakeAssetPropertyHandle;
+	TSharedPtr<IPropertyHandle> ParametersPropertyHandle;
+	TSharedPtr<IPropertyUtilities> PropertyUtilities;
+};
+
+}  // namespace UE::Cameras
+
diff --git a/Engine/Plugins/Cameras/GameplayCameras/Source/GameplayCamerasEditor/Public/Editors/CameraShakeCameraNodeGraphSchema.h b/Engine/Plugins/Cameras/GameplayCameras/Source/GameplayCamerasEditor/Public/Editors/CameraShakeCameraNodeGraphSchema.h
new file mode 100644
--- /dev/null
+++ b/Engine/Plugins/Cameras/GameplayCameras/Source/GameplayCamerasEditor/Public/Editors/CameraShakeCameraNodeGraphSchema.h
@@ -0,0 +1,23 @@
+// Copyright Epic Games, Inc. All Rights Reserved.
+
+#pragma once
+
+#include "Editors/CameraNodeGraphSchema.h"
+
+#include "CameraShakeCameraNodeGraphSchema.generated.h"
+
+/**
+ * Schema class for camera node graph.
+ */
+UCLASS()
+class UCameraShakeCameraNodeGraphSchema : public UCameraNodeGraphSchema
+{
+	GENERATED_BODY()
+
+public:
+
+	UCameraShakeCameraNodeGraphSchema(const FObjectInitializer& ObjInit);
+
+	FObjectTreeGraphConfig BuildGraphConfig() const;
+};
+
diff --git a/Engine/Plugins/Cameras/GameplayCameras/Source/GameplayCamerasEditor/Public/GameplayCamerasEditorSettings.h b/Engine/Plugins/Cameras/GameplayCameras/Source/GameplayCamerasEditor/Public/GameplayCamerasEditorSettings.h
--- a/Engine/Plugins/Cameras/GameplayCameras/Source/GameplayCamerasEditor/Public/GameplayCamerasEditorSettings.h
+++ b/Engine/Plugins/Cameras/GameplayCameras/Source/GameplayCamerasEditor/Public/GameplayCamerasEditorSettings.h
@@ -31,6 +31,9 @@ public:
 	UPROPERTY(EditAnywhere, Config, Category=NodeTitleColors)
 	FLinearColor CameraRigAssetTitleColor;
 
+	UPROPERTY(EditAnywhere, Config, Category=NodeTitleColors)
+	FLinearColor CameraShakeAssetTitleColor;
+
 	UPROPERTY(EditAnywhere, Config, Category=NodeTitleColors)
 	FLinearColor CameraRigTransitionTitleColor;
 
diff --git a/Engine/Plugins/Cameras/GameplayCameras/Source/GameplayCamerasEditor/Public/IGameplayCamerasEditorModule.h b/Engine/Plugins/Cameras/GameplayCameras/Source/GameplayCamerasEditor/Public/IGameplayCamerasEditorModule.h
--- a/Engine/Plugins/Cameras/GameplayCameras/Source/GameplayCamerasEditor/Public/IGameplayCamerasEditorModule.h
+++ b/Engine/Plugins/Cameras/GameplayCameras/Source/GameplayCamerasEditor/Public/IGameplayCamerasEditorModule.h
@@ -15,6 +15,8 @@ class UCameraRigAsset;
 class UCameraRigAssetEditor;
 class UCameraRigProxyAsset;
 class UCameraRigProxyAssetEditor;
+class UCameraShakeAsset;
+class UCameraShakeAssetEditor;
 class UCameraVariableCollection;
 class UCameraVariableCollectionEditor;
 
@@ -65,6 +67,9 @@ public:
 	/** Creates an editor for the given camera rig proxy asset */
 	virtual UCameraRigProxyAssetEditor* CreateCameraRigProxyEditor(const EToolkitMode::Type Mode, const TSharedPtr<IToolkitHost>& InitToolkitHost, UCameraRigProxyAsset* CameraRigProxy) = 0;
 
+	/** Creates an editor for the given camera shake asset */
+	virtual UCameraShakeAssetEditor* CreateCameraShakeEditor(const EToolkitMode::Type Mode, const TSharedPtr<IToolkitHost>& InitToolkitHost, UCameraShakeAsset* CameraShake) = 0;
+
 	/** Creates an editor for the given variable collection */
 	virtual UCameraVariableCollectionEditor* CreateCameraVariableCollectionEditor(const EToolkitMode::Type Mode, const TSharedPtr<IToolkitHost>& InitToolkitHost, UCameraVariableCollection* VariableCollection) = 0;
 
diff --git a/Engine/Plugins/Cameras/GameplayCameras/Source/GameplayCamerasEditor/Public/Toolkits/CameraShakeAssetEditorToolkit.h b/Engine/Plugins/Cameras/GameplayCameras/Source/GameplayCamerasEditor/Public/Toolkits/CameraShakeAssetEditorToolkit.h
new file mode 100644
--- /dev/null
+++ b/Engine/Plugins/Cameras/GameplayCameras/Source/GameplayCamerasEditor/Public/Toolkits/CameraShakeAssetEditorToolkit.h
@@ -0,0 +1,138 @@
+// Copyright Epic Games, Inc. All Rights Reserved.
+
+#pragma once
+
+#include "Core/CameraShakeAsset.h"
+#include "CoreTypes.h"
+#include "Editors/ObjectTreeGraphConfig.h"
+#include "Misc/NotifyHook.h"
+#include "Tools/BaseAssetToolkit.h"
+#include "UObject/GCObject.h"
+
+#include "CameraShakeAssetEditorToolkit.generated.h"
+
+class SFindInObjectTreeGraph;
+class SObjectTreeGraphEditor;
+class SObjectTreeGraphToolbox;
+class UCameraShakeAsset;
+struct FEdGraphEditAction;
+struct FFindInObjectTreeGraphSource;
+
+namespace UE::Cameras
+{
+
+class FBuildButtonToolkit;
+class FCameraBuildLogToolkit;
+class FCameraObjectInterfaceParametersToolkit;
+class FStandardToolkitLayout;
+class IGameplayCamerasLiveEditManager;
+
+class FCameraShakeAssetEditorToolkit
+	: public FBaseAssetToolkit
+	, public FGCObject
+	, public FNotifyHook
+{
+public:
+
+	FCameraShakeAssetEditorToolkit(UAssetEditor* InOwningAssetEditor);
+	~FCameraShakeAssetEditorToolkit();
+
+	void SetCameraShakeAsset(UCameraShakeAsset* InCameraShake);
+
+protected:
+
+	// FBaseAssetToolkit interface
+	virtual void RegisterTabSpawners(const TSharedRef<class FTabManager>& InTabManager) override;
+	virtual void UnregisterTabSpawners(const TSharedRef<FTabManager>& InTabManager) override;
+	virtual void CreateWidgets() override;
+	virtual void RegisterToolbar() override;
+	virtual void InitToolMenuContext(FToolMenuContext& MenuContext) override;
+	virtual void PostInitAssetEditor() override;
+	virtual void PostRegenerateMenusAndToolbars() override;
+
+	// IToolkit interface
+	virtual FText GetBaseToolkitName() const override;
+	virtual FName GetToolkitFName() const override;
+	virtual FString GetWorldCentricTabPrefix() const override;
+	virtual FLinearColor GetWorldCentricTabColorScale() const override;
+
+	// FGCObject interface
+	virtual void AddReferencedObjects(FReferenceCollector& Collector) override;
+	virtual FString GetReferencerName() const override;
+
+	// FNotifyHook interface
+	virtual void NotifyPostChange(const FPropertyChangedEvent& PropertyChangedEvent, FProperty* PropertyThatChanged) override;
+
+private:
+
+	TSharedRef<SDockTab> SpawnTab_CameraShakeEditor(const FSpawnTabArgs& Args);
+	TSharedRef<SDockTab> SpawnTab_Details(const FSpawnTabArgs& Args);
+	TSharedRef<SDockTab> SpawnTab_Search(const FSpawnTabArgs& Args);
+	TSharedRef<SDockTab> SpawnTab_Messages(const FSpawnTabArgs& Args);
+	TSharedRef<SDockTab> SpawnTab_Toolbox(const FSpawnTabArgs& Args);
+	TSharedRef<SDockTab> SpawnTab_InterfaceParameters(const FSpawnTabArgs& Args);
+
+	void CreateNodeGraphEditor();
+	FText GetCameraShakeAssetName() const;
+	bool IsGraphEditorEnabled() const;
+
+	void OnCameraObjectInterfaceParameterSelected(UCameraObjectInterfaceParameterBase* Object);
+
+	void OnBuild();
+	void OnFindInCameraShake();
+
+	void OnGetRootObjectsToSearch(TArray<FFindInObjectTreeGraphSource>& OutSources);
+	void OnFocusHome();
+	void OnJumpToObject(UObject* Object, FName PropertyName);
+
+private:
+
+	static const FName CameraShakeEditorTabId;
+
+	static const FName DetailsViewTabId;
+
+	static const FName SearchTabId;
+	static const FName MessagesTabId;
+
+	static const FName ToolboxTabId;
+	static const FName InterfaceParametersTabId;
+
+	/** The asset being edited */
+	TObjectPtr<UCameraShakeAsset> CameraShakeAsset;
+
+	/** The layout for this editor toolkit */
+	TSharedPtr<FStandardToolkitLayout> StandardLayout;
+
+	/** Cached config for the node graph */
+	FObjectTreeGraphConfig NodeGraphConfig;
+
+	/** The node hierarchy graph */
+	TObjectPtr<UObjectTreeGraph> NodeGraph;
+	/** The node hierarchy graph editor */
+	TSharedPtr<SObjectTreeGraphEditor> NodeGraphEditor;
+
+	/** The build button */
+	TSharedPtr<FBuildButtonToolkit> BuildButtonToolkit;
+	/** The output log */
+	TSharedPtr<FCameraBuildLogToolkit> BuildLogToolkit;
+	/** Toolbox widget */
+	TSharedPtr<SObjectTreeGraphToolbox> ToolboxWidget;
+	/** The interface parameters panel */
+	TSharedPtr<FCameraObjectInterfaceParametersToolkit> InterfaceParametersToolkit;
+
+	/** Search widget */
+	TSharedPtr<SFindInObjectTreeGraph> SearchWidget;
+};
+
+}  // namespace UE::Cameras
+
+UCLASS()
+class UCameraShakeAssetEditorMenuContext : public UObject
+{
+	GENERATED_BODY()
+
+public:
+
+	TWeakPtr<UE::Cameras::FCameraShakeAssetEditorToolkit> Toolkit;
+};
+
