Cameras: add missing redirector for camera rig allocation info

diff --git a/Engine/Plugins/Cameras/GameplayCameras/Config/BaseGameplayCameras.ini b/Engine/Plugins/Cameras/GameplayCameras/Config/BaseGameplayCameras.ini
--- a/Engine/Plugins/Cameras/GameplayCameras/Config/BaseGameplayCameras.ini
+++ b/Engine/Plugins/Cameras/GameplayCameras/Config/BaseGameplayCameras.ini
@@ -11,11 +11,12 @@
 +FunctionRedirects=(OldName="BlueprintCameraDirectorEvaluator.ActivateCameraRigPrefab", NewName="ActivateCameraRig")
 +PropertyRedirects=(OldName="/Script/GameplayCameras.CameraRigInterface.InterfaceParameters", NewName="BlendableParameters")
 
-+EnumRedirects=(OldName="/Script/GameplayCameras.ECameraRigInterfaceParameterType",NewName="ECameraObjectInterfaceCameraParameterType")
-+StructRedirects=(OldName="/Script/GameplayCameras.CameraRigInterface",NewName="CameraObjectInterface")
-+ClassRedirects=(OldName="/Script/GameplayCameras.CameraRigInterfaceParameterBase",NewName="CameraObjectInterfaceParameterBase")
 +ClassRedirects=(OldName="/Script/GameplayCameras.CameraRigBlendableParameter",NewName="CameraObjectInterfaceBlendableParameter")
 +ClassRedirects=(OldName="/Script/GameplayCameras.CameraRigDataParameter",NewName="CameraObjectInterfaceDataParameter")
-+StructRedirects=(OldName="/Script/GameplayCameras.CameraRigParameterDefinition",NewName="CameraObjectInterfaceParameterDefinition")
++ClassRedirects=(OldName="/Script/GameplayCameras.CameraRigInterfaceParameterBase",NewName="CameraObjectInterfaceParameterBase")
 +StructRedirects=(OldName="/Script/GameplayCameras.CameraRigAssetReferenceParameterMetaData",NewName="CameraObjectInterfaceParameterMetaData")
++StructRedirects=(OldName="/Script/GameplayCameras.CameraRigInterface",NewName="CameraObjectInterface")
++StructRedirects=(OldName="/Script/GameplayCameras.CameraRigParameterDefinition",NewName="CameraObjectInterfaceParameterDefinition")
++StructRedirects=(OldName="/Script/GameplayCameras.CameraRigAllocationInfo",NewName="CameraObjectAllocationInfo")
++EnumRedirects=(OldName="/Script/GameplayCameras.ECameraRigInterfaceParameterType",NewName="ECameraObjectInterfaceCameraParameterType")
 
