Cameras: remove errors and ensures that don't apply to editor preview evaluation

diff --git a/Engine/Plugins/Cameras/GameplayCameras/Source/GameplayCameras/Private/Build/CameraObjectInterfaceParameterBuilder.cpp b/Engine/Plugins/Cameras/GameplayCameras/Source/GameplayCameras/Private/Build/CameraObjectInterfaceParameterBuilder.cpp
--- a/Engine/Plugins/Cameras/GameplayCameras/Source/GameplayCameras/Private/Build/CameraObjectInterfaceParameterBuilder.cpp
+++ b/Engine/Plugins/Cameras/GameplayCameras/Source/GameplayCameras/Private/Build/CameraObjectInterfaceParameterBuilder.cpp
@@ -256,7 +256,7 @@ void FCameraObjectInterfaceParameterBuilder::SetDefaultParameterValues(const UBa
 			}
 		}
 
-		if (!ensure(RawSourceValuePtr))
+		if (!RawSourceValuePtr)
 		{
 			continue;
 		}
diff --git a/Engine/Plugins/Cameras/GameplayCameras/Source/GameplayCameras/Private/Nodes/Input/InputAxisBinding2DCameraNode.cpp b/Engine/Plugins/Cameras/GameplayCameras/Source/GameplayCameras/Private/Nodes/Input/InputAxisBinding2DCameraNode.cpp
--- a/Engine/Plugins/Cameras/GameplayCameras/Source/GameplayCameras/Private/Nodes/Input/InputAxisBinding2DCameraNode.cpp
+++ b/Engine/Plugins/Cameras/GameplayCameras/Source/GameplayCameras/Private/Nodes/Input/InputAxisBinding2DCameraNode.cpp
@@ -5,6 +5,7 @@
 #include "Components/InputComponent.h"
 #include "Core/CameraEvaluationContext.h"
 #include "Core/CameraParameterReader.h"
+#include "Core/CameraSystemEvaluator.h"
 #include "EnhancedInputComponent.h"
 #include "GameFramework/Actor.h"
 #include "GameplayCameras.h"
@@ -65,7 +66,7 @@ void FInputAxisBinding2DCameraNodeEvaluator::OnInitialize(const FCameraNodeEvalu
 			AxisValueBindings.Add(AxisValueBinding);
 		}
 	}
-	else
+	else if (Params.Evaluator->GetRole() == ECameraSystemEvaluatorRole::Game)
 	{
 		UE_LOG(LogCameraSystem, Error, TEXT("No input component found on context owner '%s' for node '%s' in '%s'."),
 				*GetNameSafe(ContextOwner), 
