Cameras: fix redirectors for truly old 5.5 content

diff --git a/Engine/Plugins/Cameras/GameplayCameras/Config/BaseGameplayCameras.ini b/Engine/Plugins/Cameras/GameplayCameras/Config/BaseGameplayCameras.ini
--- a/Engine/Plugins/Cameras/GameplayCameras/Config/BaseGameplayCameras.ini
+++ b/Engine/Plugins/Cameras/GameplayCameras/Config/BaseGameplayCameras.ini
@@ -6,7 +6,7 @@
 +FunctionRedirects=(OldName="ActivateCameraRigFunctions.ActivateVisualCameraRig", NewName="ActivatePersistentVisualCameraRig")
 
 ; UE 5.6 redirects.
-+ClassRedirects=(OldName="/Script/GameplayCameras.CameraRigInterfaceParameter", NewName="CameraRigBlendableParameter")
++ClassRedirects=(OldName="/Script/GameplayCameras.CameraRigInterfaceParameter", NewName="CameraObjectInterfaceBlendableParameter")
 +StructRedirects=(OldName="/Script/GameplayCameras.CameraTargetInfo",NewName="CameraActorTargetInfo")
 +FunctionRedirects=(OldName="BlueprintCameraDirectorEvaluator.ActivateCameraRigPrefab", NewName="ActivateCameraRig")
 +PropertyRedirects=(OldName="/Script/GameplayCameras.CameraRigInterface.InterfaceParameters", NewName="BlendableParameters")
@@ -18,5 +18,6 @@
 +StructRedirects=(OldName="/Script/GameplayCameras.CameraRigInterface",NewName="CameraObjectInterface")
 +StructRedirects=(OldName="/Script/GameplayCameras.CameraRigParameterDefinition",NewName="CameraObjectInterfaceParameterDefinition")
 +StructRedirects=(OldName="/Script/GameplayCameras.CameraRigAllocationInfo",NewName="CameraObjectAllocationInfo")
++PropertyRedirects=(OldName="/Script/GameplayCameras.CameraObjectInterface.InterfaceParameters", NewName="BlendableParameters")
 +EnumRedirects=(OldName="/Script/GameplayCameras.ECameraRigInterfaceParameterType",NewName="ECameraObjectInterfaceCameraParameterType")
 
