Cameras: fix context data overrides not applying correctly in some situations

diff --git a/Engine/Plugins/Cameras/GameplayCameras/Source/GameplayCameras/Private/Helpers/CameraObjectReferenceParameterOverrideEvaluator.h b/Engine/Plugins/Cameras/GameplayCameras/Source/GameplayCameras/Private/Helpers/CameraObjectReferenceParameterOverrideEvaluator.h
--- a/Engine/Plugins/Cameras/GameplayCameras/Source/GameplayCameras/Private/Helpers/CameraObjectReferenceParameterOverrideEvaluator.h
+++ b/Engine/Plugins/Cameras/GameplayCameras/Source/GameplayCameras/Private/Helpers/CameraObjectReferenceParameterOverrideEvaluator.h
@@ -64,7 +64,7 @@ public:
 		TSet<FGuid> AnimatedParameterGuids;
 		ObjectReference.GetAnimatedParameterGuids(AnimatedParameterGuids);
 
-		FCameraObjectInterfaceParameterOverrideHelper Helper(&OutVariableTable, nullptr);
+		FCameraObjectInterfaceParameterOverrideHelper Helper(&OutVariableTable, &OutContextDataTable);
 		Helper.ApplyParameterOverrides(
 				ObjectReference.GetCameraObject(),
 				ObjectReference.GetParameters(),
