Cameras: add new actor for single-rig cameras

diff --git a/Engine/Plugins/Cameras/GameplayCameras/Source/GameplayCameras/Private/GameFramework/GameplayCameraActor.cpp b/Engine/Plugins/Cameras/GameplayCameras/Source/GameplayCameras/Private/GameFramework/GameplayCameraActor.cpp
--- a/Engine/Plugins/Cameras/GameplayCameras/Source/GameplayCameras/Private/GameFramework/GameplayCameraActor.cpp
+++ b/Engine/Plugins/Cameras/GameplayCameras/Source/GameplayCameras/Private/GameFramework/GameplayCameraActor.cpp
@@ -2,11 +2,9 @@
 
 #include "GameFramework/GameplayCameraActor.h"
 
-#include "CineCameraComponent.h"
 #include "Engine/EngineTypes.h"
 #include "Engine/World.h"
 #include "GameFramework/GameplayCameraComponent.h"
-#include "GameFramework/PlayerController.h"
 #include "Kismet/GameplayStatics.h"
 
 #include UE_INLINE_GENERATED_CPP_BY_NAME(GameplayCameraActor)
diff --git a/Engine/Plugins/Cameras/GameplayCameras/Source/GameplayCameras/Private/GameFramework/GameplayCameraComponent.cpp b/Engine/Plugins/Cameras/GameplayCameras/Source/GameplayCameras/Private/GameFramework/GameplayCameraComponent.cpp
--- a/Engine/Plugins/Cameras/GameplayCameras/Source/GameplayCameras/Private/GameFramework/GameplayCameraComponent.cpp
+++ b/Engine/Plugins/Cameras/GameplayCameras/Source/GameplayCameras/Private/GameFramework/GameplayCameraComponent.cpp
@@ -2,29 +2,8 @@
 
 #include "GameFramework/GameplayCameraComponent.h"
 
-#include "Build/CameraAssetBuilder.h"
-#include "Build/CameraBuildLog.h"
-#include "CineCameraComponent.h"
 #include "Core/CameraAsset.h"
-#include "Core/CameraSystemEvaluator.h"
-#include "Core/RootCameraNode.h"
-#include "Debug/CameraDebugRenderer.h"
-#include "Engine/Canvas.h"
-#include "Engine/EngineTypes.h"
-#include "Engine/World.h"
-#include "GameFramework/GameplayCameraSystemActor.h"
-#include "GameFramework/GameplayCameraSystemHost.h"
 #include "GameplayCamerasDelegates.h"
-#include "Kismet/GameplayStatics.h"
-#include "Misc/AssertionMacros.h"
-#include "PrimitiveDrawInterface.h"
-#include "SceneView.h"
-#include "ShowFlags.h"
-
-#if WITH_EDITOR
-#include "Editor.h"
-#include "LevelEditorViewport.h"
-#endif
 
 #include UE_INLINE_GENERATED_CPP_BY_NAME(GameplayCameraComponent)
 
@@ -33,13 +12,6 @@
 UGameplayCameraComponent::UGameplayCameraComponent(const FObjectInitializer& ObjectInit)
 	: Super(ObjectInit)
 {
-	bTickInEditor = true;
-	bWantsOnUpdateTransform = true;
-
-	PrimaryComponentTick.bCanEverTick = true;
-
-	OutputCameraComponent = ObjectInit.CreateDefaultSubobject<UCineCameraComponent>(this, TEXT("OutputCameraComponent"), true);
-	OutputCameraComponent->SetupAttachment(this);
 }
 
 void UGameplayCameraComponent::PostLoad()
@@ -53,223 +25,6 @@ void UGameplayCameraComponent::PostLoad()
 	}
 }
 
-void UGameplayCameraComponent::BeginDestroy()
-{
-	Super::BeginDestroy();
-
-#if WITH_EDITOR
-
-	if (EditorPreviewEvaluator)
-	{
-		EditorPreviewEvaluator.Reset();
-	}
-
-	if (EvaluationContext)
-	{
-		EvaluationContext.Reset();
-	}
-
-#endif  // WITH_EDITOR
-}
-
-void UGameplayCameraComponent::AddReferencedObjects(UObject* InThis, FReferenceCollector& Collector)
-{
-	Super::AddReferencedObjects(InThis, Collector);
-
-	UGameplayCameraComponent* This = CastChecked<UGameplayCameraComponent>(InThis);
-	if (This->EvaluationContext.IsValid())
-	{
-		This->EvaluationContext->AddReferencedObjects(Collector);
-	}
-
-#if WITH_EDITOR
-
-	if (This->EditorPreviewEvaluator)
-	{
-		This->EditorPreviewEvaluator->AddReferencedObjects(Collector);
-	}
-
-#endif  // WITH_EDITOR
-}
-
-TSharedPtr<const UE::Cameras::FCameraEvaluationContext> UGameplayCameraComponent::GetEvaluationContext() const
-{
-	return EvaluationContext;
-}
-
-TSharedPtr<UE::Cameras::FCameraEvaluationContext> UGameplayCameraComponent::GetEvaluationContext()
-{
-	return EvaluationContext;
-}
-
-APlayerController* UGameplayCameraComponent::GetPlayerController() const
-{
-	if (CameraSystemHost)
-	{
-		return CameraSystemHost->GetPlayerController();
-	}
-	return nullptr;
-}
-
-void UGameplayCameraComponent::ActivateCameraForPlayerIndex(int32 PlayerIndex)
-{
-	ActivateCameraEvaluationContext(PlayerIndex);
-}
-
-void UGameplayCameraComponent::ActivateCameraForPlayerController(APlayerController* PlayerController)
-{
-	ActivateCameraEvaluationContext(PlayerController);
-}
-
-void UGameplayCameraComponent::DeactivateCamera()
-{
-	DeactivateCameraEvaluationContext();
-}
-
-void UGameplayCameraComponent::ActivateCameraEvaluationContext(int32 PlayerIndex)
-{
-	DeactivateCameraEvaluationContext();
-
-	APlayerController* PlayerController = UGameplayStatics::GetPlayerController(this, PlayerIndex);
-	ActivateCameraEvaluationContext(PlayerController);
-}
-
-void UGameplayCameraComponent::DeactivateCameraEvaluationContext()
-{
-	using namespace UE::Cameras;
-
-	if (!CameraSystemHost)
-	{
-		return;
-	}
-
-	if (EvaluationContext.IsValid())
-	{
-		TSharedPtr<FCameraSystemEvaluator> Evaluator = CameraSystemHost->GetCameraSystemEvaluator();
-		Evaluator->RemoveEvaluationContext(EvaluationContext.ToSharedRef());
-	}
-
-	if (OutputCameraComponent)
-	{
-		OutputCameraComponent->SetRelativeTransform(FTransform());
-	}
-
-	// Don't deactivate the component: we still need to update our evaluation context while any
-	// running camera rigs blend out.
-}
-
-void UGameplayCameraComponent::ActivateCameraEvaluationContext(APlayerController* PlayerController)
-{
-	using namespace UE::Cameras;
-
-	if (!PlayerController)
-	{
-		FFrame::KismetExecutionMessage(
-				TEXT("Can't activate gameplay camera component: invalid player controller!"),
-				ELogVerbosity::Error);
-		return;
-	}
-
-	if (!CameraReference.IsValid())
-	{
-		FFrame::KismetExecutionMessage(
-				TEXT("Can't activate gameplay camera component: no camera asset was set!"),
-				ELogVerbosity::Error);
-		return;
-	}
-	
-	CameraSystemHost = UGameplayCameraSystemHost::FindOrCreateHost(PlayerController);
-	if (!CameraSystemHost)
-	{
-		FFrame::KismetExecutionMessage(
-				TEXT("Can't activate gameplay camera component: no camera system host found!"),
-				ELogVerbosity::Error);
-		return;
-	}
-
-	AGameplayCameraSystemActor::AutoManageActiveViewTarget(PlayerController);
-
-	EnsureCameraEvaluationContextCreated(PlayerController);
-
-	TSharedPtr<FCameraSystemEvaluator> CameraSystemEvaluator = CameraSystemHost->GetCameraSystemEvaluator();
-	CameraSystemEvaluator->PushEvaluationContext(EvaluationContext.ToSharedRef());
-
-	// Make sure the component is active so it receives tick updates to maintain the evaluation context.
-	Activate();
-}
-
-void UGameplayCameraComponent::EnsureCameraEvaluationContextCreated(APlayerController* PlayerController)
-{
-	using namespace UE::Cameras;
-
-	if (!EvaluationContext.IsValid())
-	{
-		EvaluationContext = MakeShared<FGameplayCameraComponentEvaluationContext>();
-
-		FCameraEvaluationContextInitializeParams InitParams;
-		InitParams.Owner = this;
-		InitParams.CameraAsset = CameraReference.GetCameraAsset();
-		InitParams.PlayerController = PlayerController;
-		EvaluationContext->Initialize(InitParams);
-
-		UpdateCameraEvaluationContext(true);
-	}
-}
-
-#define UE_PRIVATE_GAMEPLAY_CAMERA_COMPONENT_VALIDATE_EVALUATION_CONTEXT(ErrorMsg, ErrorResult)\
-	using namespace UE::Cameras;\
-	if (!EvaluationContext)\
-	{\
-		FFrame::KismetExecutionMessage(\
-				*FString::Format(\
-					TEXT(#ErrorResult " on Gameplay Camera component '{0}': it isn't active."),\
-					{ *GetNameSafe(this) }),\
-				ELogVerbosity::Error);\
-		return ErrorResult;\
-	}
-
-FBlueprintCameraNodeEvaluationResult UGameplayCameraComponent::GetInitialResult() const
-{
-	UE_PRIVATE_GAMEPLAY_CAMERA_COMPONENT_VALIDATE_EVALUATION_CONTEXT("Can't get shared camera data", FBlueprintCameraNodeEvaluationResult());
-
-	return FBlueprintCameraNodeEvaluationResult(&EvaluationContext->GetInitialResult());
-}
-
-FBlueprintCameraNodeEvaluationResult UGameplayCameraComponent::GetConditionalResult(ECameraEvaluationDataCondition Condition) const
-{
-	UE_PRIVATE_GAMEPLAY_CAMERA_COMPONENT_VALIDATE_EVALUATION_CONTEXT("Can't get conditional camera data", FBlueprintCameraNodeEvaluationResult());
-
-	return FBlueprintCameraNodeEvaluationResult(&EvaluationContext->GetOrAddConditionalResult(Condition));
-}
-
-FBlueprintCameraPose UGameplayCameraComponent::GetInitialPose() const
-{
-	UE_PRIVATE_GAMEPLAY_CAMERA_COMPONENT_VALIDATE_EVALUATION_CONTEXT("Can't get initial camera pose", FBlueprintCameraPose());
-
-	return FBlueprintCameraPose::FromCameraPose(EvaluationContext->GetInitialResult().CameraPose);
-}
-
-bool UGameplayCameraComponent::SetInitialPose(const FBlueprintCameraPose& CameraPose)
-{
-	UE_PRIVATE_GAMEPLAY_CAMERA_COMPONENT_VALIDATE_EVALUATION_CONTEXT("Can't set initial camera pose", false);
-
-	FCameraPose InitialPose = EvaluationContext->GetInitialResult().CameraPose;
-	CameraPose.ApplyTo(InitialPose);
-	return true;
-}
-
-FBlueprintCameraVariableTable UGameplayCameraComponent::GetInitialVariableTable() const
-{
-	using namespace UE::Cameras;
-
-	UE_PRIVATE_GAMEPLAY_CAMERA_COMPONENT_VALIDATE_EVALUATION_CONTEXT("Can't get initial camera variable table", FBlueprintCameraVariableTable());
-
-	FCameraVariableTable& VariableTable = EvaluationContext->GetInitialResult().VariableTable;
-	return FBlueprintCameraVariableTable(&VariableTable);
-}
-
-#undef UE_PRIVATE_GAMEPLAY_CAMERA_COMPONENT_VALIDATE_EVALUATION_CONTEXT
-
 void UGameplayCameraComponent::OnRegister()
 {
 	using namespace UE::Cameras;
@@ -278,53 +33,9 @@ void UGameplayCameraComponent::OnRegister()
 
 #if WITH_EDITOR
 
-	UWorld* World = GetWorld();
-	bIsEditorWorld = (World && (World->WorldType == EWorldType::Editor || World->WorldType == EWorldType::EditorPreview));
-
 	FGameplayCamerasDelegates::OnCameraAssetBuilt().AddUObject(this, &UGameplayCameraComponent::OnCameraAssetBuilt);
 
-	const TCHAR* ShowFlagName = TEXT("GameplayCameras");
-	CustomShowFlag = FEngineShowFlags::FindIndexByName(ShowFlagName);
-
-#endif  // WITH_EDITOR
-}
-
-void UGameplayCameraComponent::BeginPlay()
-{
-	using namespace UE::Cameras;
-
-	Super::BeginPlay();
-
-#if WITH_EDITOR
-
-	if (CameraReference.IsValid())
-	{
-		UWorld* World = GetWorld();
-		if (World && World->WorldType == EWorldType::PIE)
-		{
-			// Auto-build the camera asset on begin play to make sure we've got the latest user edits.
-			FCameraBuildLog BuildLog;
-			FCameraAssetBuilder Builder(BuildLog);
-			Builder.BuildCamera(CameraReference.GetCameraAsset());
-		}
-
-		CameraReference.RebuildParametersIfNeeded();
-	}
-
-#endif  // WITH_EDITOR
-
-	if (IsActive() && AutoActivateForPlayer != EAutoReceiveInput::Disabled && GetNetMode() != NM_DedicatedServer)
-	{
-		const int32 PlayerIndex = AutoActivateForPlayer.GetIntValue() - 1;
-		ActivateCameraForPlayerIndex(PlayerIndex);
-	}
-}
-
-void UGameplayCameraComponent::EndPlay(const EEndPlayReason::Type EndPlayReason)
-{
-	DeactivateCameraEvaluationContext();
-
-	Super::EndPlay(EndPlayReason);
+#endif
 }
 
 void UGameplayCameraComponent::OnUnregister()
@@ -340,211 +51,35 @@ void UGameplayCameraComponent::OnUnregister()
 	Super::OnUnregister();
 }
 
-void UGameplayCameraComponent::TickComponent(float DeltaTime, ELevelTick TickType, FActorComponentTickFunction *ThisTickFunction)
+UCameraAsset* UGameplayCameraComponent::GetCameraAsset()
 {
-	using namespace UE::Cameras;
-
-	Super::TickComponent(DeltaTime, TickType, ThisTickFunction);
-
-#if WITH_EDITOR
-
-	// Make sure things are setup (or not) if we want to run the camera logic in editor (or not).
-	AutoManageEditorPreviewEvaluator();
-
-#endif  // WITH_EDITOR
-
-	if (EvaluationContext)
-	{
-		UpdateCameraEvaluationContext(false);
-
-#if WITH_EDITOR
-		
-		UpdateEditorPreviewEvaluator(DeltaTime);
-
-#endif  // WITH_EDITOR
-
-		UpdateOutputCameraComponent();
-	}
+	return CameraReference.GetCameraAsset();
 }
 
-void UGameplayCameraComponent::UpdateCameraEvaluationContext(bool bForceApplyParameterOverrides)
+bool UGameplayCameraComponent::OnValidateCameraEvaluationContextActivation()
 {
-	using namespace UE::Cameras;
-
-	FCameraNodeEvaluationResult& InitialResult = EvaluationContext->GetInitialResult();
-
-	const FTransform& OwnerTransform = GetComponentTransform();
-	InitialResult.CameraPose.SetTransform(OwnerTransform, true);
-	InitialResult.bIsCameraCut = false;
-	InitialResult.bIsValid = true;
-
-	if (bIsCameraCutNextFrame)
-	{
-		InitialResult.bIsCameraCut = true;
-		bIsCameraCutNextFrame = false;
-	}
-
-	const bool bApplyDrivenParametersOnly = !bForceApplyParameterOverrides;
-	CameraReference.ApplyParameterOverrides(InitialResult, bApplyDrivenParametersOnly);
-
-#if WITH_EDITOR
-	EvaluationContext->UpdateForEditorPreview();
-#endif  // WITH_EDITOR
-}
-
-void UGameplayCameraComponent::UpdateOutputCameraComponent()
-{
-	using namespace UE::Cameras;
-
-	if (!OutputCameraComponent)
-	{
-		return;
-	}
-
-	TSharedPtr<FCameraSystemEvaluator> CameraSystemEvaluator;
-
-	if (CameraSystemHost)
-	{
-		CameraSystemEvaluator = CameraSystemHost->GetCameraSystemEvaluator();
-	}
-#if WITH_EDITOR
-	else if (EditorPreviewEvaluator)
-	{
-		CameraSystemEvaluator = EditorPreviewEvaluator;
-	}
-#endif  // WITH_EDITOR
-
-	bool bGotValidTransform = false;
-	if (CameraSystemEvaluator)
+	if (!CameraReference.IsValid())
 	{
-		FRootCameraNodeEvaluator* RootNodeEvaluator = CameraSystemEvaluator->GetRootNodeEvaluator();
-		if (RootNodeEvaluator && RootNodeEvaluator->HasAnyActiveCameraRig())
-		{
-			const FCameraSystemEvaluationResult& Result = CameraSystemEvaluator->GetEvaluatedResult();
-
-			OutputCameraComponent->SetWorldTransform(Result.CameraPose.GetTransform());
-			OutputCameraComponent->SetFieldOfView(Result.CameraPose.GetEffectiveFieldOfView());
-			OutputCameraComponent->CurrentAperture = Result.CameraPose.GetAperture();
-			OutputCameraComponent->Filmback.SensorWidth = Result.CameraPose.GetSensorWidth();
-			OutputCameraComponent->Filmback.SensorHeight = Result.CameraPose.GetSensorHeight();
-			OutputCameraComponent->bConstrainAspectRatio = Result.CameraPose.GetConstrainAspectRatio();
-			OutputCameraComponent->bOverrideAspectRatioAxisConstraint = Result.CameraPose.GetOverrideAspectRatioAxisConstraint();
-			OutputCameraComponent->AspectRatioAxisConstraint = Result.CameraPose.GetAspectRatioAxisConstraint();
-
-			OutputCameraComponent->FocusSettings.ManualFocusDistance = Result.CameraPose.GetFocusDistance();
-			OutputCameraComponent->FocusSettings.FocusMethod = (Result.CameraPose.GetEnablePhysicalCamera() ? ECameraFocusMethod::Manual : ECameraFocusMethod::Disable);
-
-			bGotValidTransform = true;
-		}
+		FFrame::KismetExecutionMessage(
+				TEXT("Can't activate gameplay camera component: no camera asset was set!"),
+				ELogVerbosity::Error);
+		return false;
 	}
 	
-	if (!bGotValidTransform)
-	{
-		OutputCameraComponent->SetRelativeTransform(FTransform());
-	}
-}
-
-void UGameplayCameraComponent::OnComponentDestroyed(bool bDestroyingHierarchy)
-{
-	Super::OnComponentDestroyed(bDestroyingHierarchy);
-}
-
-void UGameplayCameraComponent::OnUpdateTransform(EUpdateTransformFlags UpdateTransformFlags, ETeleportType Teleport)
-{
-	Super::OnUpdateTransform(UpdateTransformFlags, Teleport);
-
-	if (EvaluationContext && Teleport != ETeleportType::None)
-	{
-		bIsCameraCutNextFrame = true;
-	}
-
-#if WITH_EDITOR
-
-	if (bIsEditorWorld && EvaluationContext)
-	{
-		UpdateCameraEvaluationContext(false);
-	}
-
-#endif  // WITH_EDITOR
+	return true;
 }
 
-#if WITH_EDITOR
-
-void UGameplayCameraComponent::AutoManageEditorPreviewEvaluator()
+void UGameplayCameraComponent::OnUpdateCameraEvaluationContext(bool bForceApplyParameterOverrides)
 {
 	using namespace UE::Cameras;
 
-	if (!bIsEditorWorld)
-	{
-		return;
-	}
-	
-	if (bRunInEditor && !(EditorPreviewEvaluator && EvaluationContext))
-	{
-		// We want to run the camera logic in the editor but we haven't set things up for that.
-		// Let's create the preview evaluator and the evaluation context.
-		if (!EditorPreviewEvaluator)
-		{
-			EditorPreviewEvaluator = MakeShared<FCameraSystemEvaluator>();
-
-			FCameraSystemEvaluatorCreateParams CreateParams;
-			CreateParams.Owner = this;
-			CreateParams.Role = ECameraSystemEvaluatorRole::EditorPreview;
-			EditorPreviewEvaluator->Initialize(CreateParams);
-		}
-		if (!EvaluationContext)
-		{
-			EnsureCameraEvaluationContextCreated(nullptr);
-			EditorPreviewEvaluator->PushEvaluationContext(EvaluationContext.ToSharedRef());
-			EvaluationContext->SetEditorPreviewCameraRigIndex(EditorPreviewCameraRigIndex);
-		}
-	}
-	else if (!bRunInEditor && (EditorPreviewEvaluator || EvaluationContext))
-	{
-		// We don't want to run the camera logic in the editor anymore. Let's tear things down.
-		EditorPreviewEvaluator = nullptr;
-		EvaluationContext = nullptr;
-	}
-}
-
-void UGameplayCameraComponent::OnCameraAssetReferenceChanged()
-{
-	if (!bIsEditorWorld)
-	{
-		return;
-	}
+	FCameraNodeEvaluationResult& InitialResult = GetEvaluationContext()->GetInitialResult();
 
-	if (bRunInEditor && EditorPreviewEvaluator && EvaluationContext)
-	{
-		if (EvaluationContext->GetCameraAsset() != CameraReference.GetCameraAsset())
-		{
-			// The camera asset has changed! Recreate the context.
-			EditorPreviewEvaluator->RemoveEvaluationContext(EvaluationContext.ToSharedRef());
-			EvaluationContext = nullptr;
-
-			EnsureCameraEvaluationContextCreated(nullptr);
-			EditorPreviewEvaluator->PushEvaluationContext(EvaluationContext.ToSharedRef());
-		}
-		else
-		{
-			// Otherwise, maybe one of the parameter overrides has changed. Re-apply them.
-			UpdateCameraEvaluationContext(true);
-		}
-	}
+	const bool bApplyDrivenParametersOnly = !bForceApplyParameterOverrides;
+	CameraReference.ApplyParameterOverrides(InitialResult, bApplyDrivenParametersOnly);
 }
 
-void UGameplayCameraComponent::OnEditorPreviewCameraRigIndexChanged()
-{
-	if (!bIsEditorWorld)
-	{
-		return;
-	}
-
-	if (bRunInEditor && EditorPreviewEvaluator && EvaluationContext)
-	{
-		EvaluationContext->SetEditorPreviewCameraRigIndex(EditorPreviewCameraRigIndex);
-	}
-}
+#if WITH_EDITOR
 
 void UGameplayCameraComponent::OnCameraAssetBuilt(const UCameraAsset* InCameraAsset)
 {
@@ -557,10 +92,10 @@ void UGameplayCameraComponent::OnCameraAssetBuilt(const UCameraAsset* InCameraAs
 
 	// If our camera asset was just built, it may have some new parameters. We need to rebuild
 	// our variable table and context data table, and re-apply overrides.
-	if (bRunInEditor && EditorPreviewEvaluator && EvaluationContext)
+	if (HasActiveEditorPreview())
 	{
 		const FCameraAssetAllocationInfo& AllocationInfo = InCameraAsset->GetAllocationInfo();
-		FCameraNodeEvaluationResult& InitialResult = EvaluationContext->GetInitialResult();
+		FCameraNodeEvaluationResult& InitialResult = GetEvaluationContext()->GetInitialResult();
 		InitialResult.VariableTable.Initialize(AllocationInfo.VariableTableInfo);
 		InitialResult.ContextDataTable.Initialize(AllocationInfo.ContextDataTableInfo);
 
@@ -570,50 +105,6 @@ void UGameplayCameraComponent::OnCameraAssetBuilt(const UCameraAsset* InCameraAs
 	}
 }
 
-bool UGameplayCameraComponent::GetEditorPreviewInfo(float DeltaTime, FMinimalViewInfo& ViewOut)
-{
-	if (OutputCameraComponent)
-	{
-		OutputCameraComponent->GetEditorPreviewInfo(DeltaTime, ViewOut);
-		return true;
-	}
-	return false;
-}
-
-void UGameplayCameraComponent::UpdateEditorPreviewEvaluator(float DeltaTime)
-{
-	using namespace UE::Cameras;
-
-	if (EditorPreviewEvaluator)
-	{
-		FCameraSystemEvaluationParams Params;
-		Params.DeltaTime = DeltaTime;
-		EditorPreviewEvaluator->Update(Params);
-	}
-}
-
-void UGameplayCameraComponent::OnDrawVisualizationHUD(const FViewport* Viewport, const FSceneView* SceneView, FCanvas* Canvas) const
-{
-	using namespace UE::Cameras;
-
-	const bool bHasShowFlag = SceneView->Family->EngineShowFlags.GetSingleFlag(CustomShowFlag);
-	if (bHasShowFlag && bRunInEditor && EditorPreviewEvaluator && EvaluationContext)
-	{
-		const AActor* OwnerActor = GetOwner();
-
-		const AActor* ViewActor = SceneView->ViewActor.Get();
-		const bool bIsLockedToCamera = (ViewActor == OwnerActor);
-
-		FCameraSystemEditorPreviewParams Params;
-		Params.Canvas = Canvas;
-		Params.SceneView = SceneView;
-		Params.bIsLockedToCamera = bIsLockedToCamera;
-		Params.bDrawWorldDebug = false;
-
-		EditorPreviewEvaluator->DrawEditorPreview(Params);
-	}
-}
-
 void UGameplayCameraComponent::PostEditChangeProperty( struct FPropertyChangedEvent& PropertyChangedEvent)
 {
 	using namespace UE::Cameras;
@@ -623,47 +114,23 @@ void UGameplayCameraComponent::PostEditChangeProperty( struct FPropertyChangedEv
 	const FName MemberPropertyName = PropertyChangedEvent.GetMemberPropertyName();
 	if (MemberPropertyName == GET_MEMBER_NAME_CHECKED(UGameplayCameraComponent, CameraReference))
 	{
-		OnCameraAssetReferenceChanged();
-	}
-	else if (MemberPropertyName == GET_MEMBER_NAME_CHECKED(UGameplayCameraComponent, bRunInEditor))
-	{
-		AutoManageEditorPreviewEvaluator();
-	}
-	else if (MemberPropertyName == GET_MEMBER_NAME_CHECKED(UGameplayCameraComponent, EditorPreviewCameraRigIndex))
-	{
-		OnEditorPreviewCameraRigIndexChanged();
-	}
-}
-
-#endif  // WITH_EDITOR
-
-namespace UE::Cameras
-{
-
-UE_DEFINE_CAMERA_EVALUATION_CONTEXT(FGameplayCameraComponentEvaluationContext)
-
-#if WITH_EDITOR
-
-void FGameplayCameraComponentEvaluationContext::UpdateForEditorPreview()
-{
-	FCameraSystemEvaluator* ActiveEvaluator = GetCameraSystemEvaluator();
-	if (ActiveEvaluator && ActiveEvaluator->GetRole() == ECameraSystemEvaluatorRole::EditorPreview)
-	{
-		if (GCurrentLevelEditingViewportClient && GCurrentLevelEditingViewportClient->Viewport)
-		{
-			FIntPoint ViewportSize = GCurrentLevelEditingViewportClient->Viewport->GetSizeXY();
-			OverrideViewportSize = ViewportSize;
-		}
-		else
+		if (HasActiveEditorPreview())
 		{
-			OverrideViewportSize.Reset();
+			if (PropertyChangedEvent.GetPropertyName() == TEXT("CameraAsset"))
+			{
+				// The camera asset has changed! Recreate the context.
+				RecreateEditorPreviewEvaluationContext();
+			}
+			else
+			{
+				// Otherwise, maybe one of the parameter overrides has changed. Re-apply them.
+				UpdateCameraEvaluationContext(true);
+			}
 		}
 	}
 }
 
 #endif  // WITH_EDITOR
 
-}  // namespace UE::Cameras
-
 #undef LOCTEXT_NAMESPACE
 
diff --git a/Engine/Plugins/Cameras/GameplayCameras/Source/GameplayCameras/Private/GameFramework/GameplayCameraComponentBase.cpp b/Engine/Plugins/Cameras/GameplayCameras/Source/GameplayCameras/Private/GameFramework/GameplayCameraComponentBase.cpp
new file mode 100644
--- /dev/null
+++ b/Engine/Plugins/Cameras/GameplayCameras/Source/GameplayCameras/Private/GameFramework/GameplayCameraComponentBase.cpp
@@ -0,0 +1,626 @@
+// Copyright Epic Games, Inc. All Rights Reserved.
+
+#include "GameFramework/GameplayCameraComponentBase.h"
+
+#include "Build/CameraAssetBuilder.h"
+#include "Build/CameraBuildLog.h"
+#include "CineCameraComponent.h"
+#include "Core/CameraAsset.h"
+#include "Core/CameraSystemEvaluator.h"
+#include "Core/RootCameraNode.h"
+#include "Debug/CameraDebugRenderer.h"
+#include "Engine/Canvas.h"
+#include "Engine/EngineTypes.h"
+#include "Engine/World.h"
+#include "GameFramework/GameplayCameraSystemActor.h"
+#include "GameFramework/GameplayCameraSystemHost.h"
+#include "GameplayCamerasDelegates.h"
+#include "Kismet/GameplayStatics.h"
+#include "Misc/AssertionMacros.h"
+#include "PrimitiveDrawInterface.h"
+#include "SceneView.h"
+#include "ShowFlags.h"
+
+#if WITH_EDITOR
+#include "Editor.h"
+#include "LevelEditorViewport.h"
+#endif
+
+#include UE_INLINE_GENERATED_CPP_BY_NAME(GameplayCameraComponentBase)
+
+#define LOCTEXT_NAMESPACE "GameplayCameraComponentBase"
+
+UGameplayCameraComponentBase::UGameplayCameraComponentBase(const FObjectInitializer& ObjectInit)
+	: Super(ObjectInit)
+{
+	bTickInEditor = true;
+	bWantsOnUpdateTransform = true;
+
+	PrimaryComponentTick.bCanEverTick = true;
+
+	OutputCameraComponent = ObjectInit.CreateDefaultSubobject<UCineCameraComponent>(this, TEXT("OutputCameraComponent"), true);
+	OutputCameraComponent->SetupAttachment(this);
+}
+
+void UGameplayCameraComponentBase::BeginDestroy()
+{
+	Super::BeginDestroy();
+
+#if WITH_EDITOR
+
+	if (EditorPreviewEvaluator)
+	{
+		EditorPreviewEvaluator.Reset();
+	}
+
+	if (EvaluationContext)
+	{
+		EvaluationContext.Reset();
+	}
+
+#endif  // WITH_EDITOR
+}
+
+void UGameplayCameraComponentBase::AddReferencedObjects(UObject* InThis, FReferenceCollector& Collector)
+{
+	Super::AddReferencedObjects(InThis, Collector);
+
+	UGameplayCameraComponentBase* This = CastChecked<UGameplayCameraComponentBase>(InThis);
+	if (This->EvaluationContext.IsValid())
+	{
+		This->EvaluationContext->AddReferencedObjects(Collector);
+	}
+
+#if WITH_EDITOR
+
+	if (This->EditorPreviewEvaluator)
+	{
+		This->EditorPreviewEvaluator->AddReferencedObjects(Collector);
+	}
+
+#endif  // WITH_EDITOR
+}
+
+TSharedPtr<const UE::Cameras::FCameraEvaluationContext> UGameplayCameraComponentBase::GetEvaluationContext() const
+{
+	return EvaluationContext;
+}
+
+TSharedPtr<UE::Cameras::FCameraEvaluationContext> UGameplayCameraComponentBase::GetEvaluationContext()
+{
+	return EvaluationContext;
+}
+
+APlayerController* UGameplayCameraComponentBase::GetPlayerController() const
+{
+	if (CameraSystemHost)
+	{
+		return CameraSystemHost->GetPlayerController();
+	}
+	return nullptr;
+}
+
+void UGameplayCameraComponentBase::ActivateCameraForPlayerIndex(int32 PlayerIndex)
+{
+	ActivateCameraEvaluationContext(PlayerIndex);
+}
+
+void UGameplayCameraComponentBase::ActivateCameraForPlayerController(APlayerController* PlayerController)
+{
+	ActivateCameraEvaluationContext(PlayerController);
+}
+
+void UGameplayCameraComponentBase::DeactivateCamera()
+{
+	DeactivateCameraEvaluationContext();
+}
+
+void UGameplayCameraComponentBase::ActivateCameraEvaluationContext(int32 PlayerIndex)
+{
+	DeactivateCameraEvaluationContext();
+
+	APlayerController* PlayerController = UGameplayStatics::GetPlayerController(this, PlayerIndex);
+	ActivateCameraEvaluationContext(PlayerController);
+}
+
+void UGameplayCameraComponentBase::DeactivateCameraEvaluationContext()
+{
+	using namespace UE::Cameras;
+
+	if (!CameraSystemHost)
+	{
+		return;
+	}
+
+	if (EvaluationContext.IsValid())
+	{
+		TSharedPtr<FCameraSystemEvaluator> Evaluator = CameraSystemHost->GetCameraSystemEvaluator();
+		Evaluator->RemoveEvaluationContext(EvaluationContext.ToSharedRef());
+	}
+
+	if (OutputCameraComponent)
+	{
+		OutputCameraComponent->SetRelativeTransform(FTransform());
+	}
+
+	// Don't deactivate the component: we still need to update our evaluation context while any
+	// running camera rigs blend out.
+}
+
+void UGameplayCameraComponentBase::ActivateCameraEvaluationContext(APlayerController* PlayerController)
+{
+	using namespace UE::Cameras;
+
+	if (!PlayerController)
+	{
+		FFrame::KismetExecutionMessage(
+				TEXT("Can't activate gameplay camera component: invalid player controller!"),
+				ELogVerbosity::Error);
+		return;
+	}
+
+	CameraSystemHost = UGameplayCameraSystemHost::FindOrCreateHost(PlayerController);
+	if (!CameraSystemHost)
+	{
+		FFrame::KismetExecutionMessage(
+				TEXT("Can't activate gameplay camera component: no camera system host found!"),
+				ELogVerbosity::Error);
+		return;
+	}
+
+	if (!OnValidateCameraEvaluationContextActivation())
+	{
+		return;
+	}
+
+	AGameplayCameraSystemActor::AutoManageActiveViewTarget(PlayerController);
+
+	EnsureCameraEvaluationContextCreated(PlayerController);
+
+	TSharedPtr<FCameraSystemEvaluator> CameraSystemEvaluator = CameraSystemHost->GetCameraSystemEvaluator();
+	CameraSystemEvaluator->PushEvaluationContext(EvaluationContext.ToSharedRef());
+
+	// Make sure the component is active so it receives tick updates to maintain the evaluation context.
+	Activate();
+}
+
+void UGameplayCameraComponentBase::EnsureCameraEvaluationContextCreated(APlayerController* PlayerController)
+{
+	using namespace UE::Cameras;
+
+	if (!EvaluationContext.IsValid())
+	{
+		const UCameraAsset* CameraAsset = GetCameraAsset();
+		if (ensureMsgf(
+				CameraAsset,
+				TEXT("GameplayCameraComponent sub-class '%s' didn't provide us with a valid camera asset"),
+				*GetNameSafe(GetClass())))
+		{
+			EvaluationContext = MakeShared<FGameplayCameraComponentEvaluationContext>();
+
+			FCameraEvaluationContextInitializeParams InitParams;
+			InitParams.Owner = this;
+			InitParams.CameraAsset = CameraAsset;
+			InitParams.PlayerController = PlayerController;
+			EvaluationContext->Initialize(InitParams);
+
+			UpdateCameraEvaluationContext(true);
+		}
+	}
+}
+
+#define UE_PRIVATE_GAMEPLAY_CAMERA_COMPONENT_VALIDATE_EVALUATION_CONTEXT(ErrorMsg, ErrorResult)\
+	using namespace UE::Cameras;\
+	if (!EvaluationContext)\
+	{\
+		FFrame::KismetExecutionMessage(\
+				*FString::Format(\
+					TEXT(#ErrorResult " on Gameplay Camera component '{0}': it isn't active."),\
+					{ *GetNameSafe(this) }),\
+				ELogVerbosity::Error);\
+		return ErrorResult;\
+	}
+
+FBlueprintCameraNodeEvaluationResult UGameplayCameraComponentBase::GetInitialResult() const
+{
+	UE_PRIVATE_GAMEPLAY_CAMERA_COMPONENT_VALIDATE_EVALUATION_CONTEXT("Can't get shared camera data", FBlueprintCameraNodeEvaluationResult());
+
+	return FBlueprintCameraNodeEvaluationResult(&EvaluationContext->GetInitialResult());
+}
+
+FBlueprintCameraNodeEvaluationResult UGameplayCameraComponentBase::GetConditionalResult(ECameraEvaluationDataCondition Condition) const
+{
+	UE_PRIVATE_GAMEPLAY_CAMERA_COMPONENT_VALIDATE_EVALUATION_CONTEXT("Can't get conditional camera data", FBlueprintCameraNodeEvaluationResult());
+
+	return FBlueprintCameraNodeEvaluationResult(&EvaluationContext->GetOrAddConditionalResult(Condition));
+}
+
+FBlueprintCameraPose UGameplayCameraComponentBase::GetInitialPose() const
+{
+	UE_PRIVATE_GAMEPLAY_CAMERA_COMPONENT_VALIDATE_EVALUATION_CONTEXT("Can't get initial camera pose", FBlueprintCameraPose());
+
+	return FBlueprintCameraPose::FromCameraPose(EvaluationContext->GetInitialResult().CameraPose);
+}
+
+bool UGameplayCameraComponentBase::SetInitialPose(const FBlueprintCameraPose& CameraPose)
+{
+	UE_PRIVATE_GAMEPLAY_CAMERA_COMPONENT_VALIDATE_EVALUATION_CONTEXT("Can't set initial camera pose", false);
+
+	FCameraPose InitialPose = EvaluationContext->GetInitialResult().CameraPose;
+	CameraPose.ApplyTo(InitialPose);
+	return true;
+}
+
+FBlueprintCameraVariableTable UGameplayCameraComponentBase::GetInitialVariableTable() const
+{
+	using namespace UE::Cameras;
+
+	UE_PRIVATE_GAMEPLAY_CAMERA_COMPONENT_VALIDATE_EVALUATION_CONTEXT("Can't get initial camera variable table", FBlueprintCameraVariableTable());
+
+	FCameraVariableTable& VariableTable = EvaluationContext->GetInitialResult().VariableTable;
+	return FBlueprintCameraVariableTable(&VariableTable);
+}
+
+#undef UE_PRIVATE_GAMEPLAY_CAMERA_COMPONENT_VALIDATE_EVALUATION_CONTEXT
+
+void UGameplayCameraComponentBase::OnRegister()
+{
+	using namespace UE::Cameras;
+
+	Super::OnRegister();
+
+#if WITH_EDITOR
+
+	UWorld* World = GetWorld();
+	bIsEditorWorld = (World && (World->WorldType == EWorldType::Editor || World->WorldType == EWorldType::EditorPreview));
+
+	const TCHAR* ShowFlagName = TEXT("GameplayCameras");
+	CustomShowFlag = FEngineShowFlags::FindIndexByName(ShowFlagName);
+
+#endif  // WITH_EDITOR
+}
+
+void UGameplayCameraComponentBase::BeginPlay()
+{
+	using namespace UE::Cameras;
+
+	Super::BeginPlay();
+
+#if WITH_EDITOR
+
+	UCameraAsset* CameraAsset = GetCameraAsset();
+	if (CameraAsset)
+	{
+		UWorld* World = GetWorld();
+		if (World && World->WorldType == EWorldType::PIE)
+		{
+			// Auto-build the camera asset on begin play to make sure we've got the latest user edits.
+			FCameraBuildLog BuildLog;
+			FCameraAssetBuilder Builder(BuildLog);
+			Builder.BuildCamera(CameraAsset);
+		}
+	}
+
+#endif  // WITH_EDITOR
+
+	if (IsActive() && AutoActivateForPlayer != EAutoReceiveInput::Disabled && GetNetMode() != NM_DedicatedServer)
+	{
+		const int32 PlayerIndex = AutoActivateForPlayer.GetIntValue() - 1;
+		ActivateCameraForPlayerIndex(PlayerIndex);
+	}
+}
+
+void UGameplayCameraComponentBase::EndPlay(const EEndPlayReason::Type EndPlayReason)
+{
+	DeactivateCameraEvaluationContext();
+
+	Super::EndPlay(EndPlayReason);
+}
+
+void UGameplayCameraComponentBase::OnUnregister()
+{
+	using namespace UE::Cameras;
+
+#if WITH_EDITOR
+
+	FGameplayCamerasDelegates::OnCameraAssetBuilt().RemoveAll(this);
+
+#endif  // WITH_EDITOR
+
+	Super::OnUnregister();
+}
+
+void UGameplayCameraComponentBase::TickComponent(float DeltaTime, ELevelTick TickType, FActorComponentTickFunction *ThisTickFunction)
+{
+	using namespace UE::Cameras;
+
+	Super::TickComponent(DeltaTime, TickType, ThisTickFunction);
+
+#if WITH_EDITOR
+
+	// Make sure things are setup (or not) if we want to run the camera logic in editor (or not).
+	AutoManageEditorPreviewEvaluator();
+
+#endif  // WITH_EDITOR
+
+	if (EvaluationContext)
+	{
+		UpdateCameraEvaluationContext(false);
+
+#if WITH_EDITOR
+		
+		UpdateEditorPreviewEvaluator(DeltaTime);
+
+#endif  // WITH_EDITOR
+
+		UpdateOutputCameraComponent();
+	}
+}
+
+void UGameplayCameraComponentBase::UpdateCameraEvaluationContext(bool bForceApplyParameterOverrides)
+{
+	using namespace UE::Cameras;
+
+	FCameraNodeEvaluationResult& InitialResult = EvaluationContext->GetInitialResult();
+
+	const FTransform& OwnerTransform = GetComponentTransform();
+	InitialResult.CameraPose.SetTransform(OwnerTransform, true);
+	InitialResult.bIsCameraCut = false;
+	InitialResult.bIsValid = true;
+
+	if (bIsCameraCutNextFrame)
+	{
+		InitialResult.bIsCameraCut = true;
+		bIsCameraCutNextFrame = false;
+	}
+
+	OnUpdateCameraEvaluationContext(bForceApplyParameterOverrides);
+
+#if WITH_EDITOR
+	EvaluationContext->UpdateForEditorPreview();
+#endif  // WITH_EDITOR
+}
+
+void UGameplayCameraComponentBase::UpdateOutputCameraComponent()
+{
+	using namespace UE::Cameras;
+
+	if (!OutputCameraComponent)
+	{
+		return;
+	}
+
+	TSharedPtr<FCameraSystemEvaluator> CameraSystemEvaluator;
+
+	if (CameraSystemHost)
+	{
+		CameraSystemEvaluator = CameraSystemHost->GetCameraSystemEvaluator();
+	}
+#if WITH_EDITOR
+	else if (EditorPreviewEvaluator)
+	{
+		CameraSystemEvaluator = EditorPreviewEvaluator;
+	}
+#endif  // WITH_EDITOR
+
+	bool bGotValidTransform = false;
+	if (CameraSystemEvaluator)
+	{
+		FRootCameraNodeEvaluator* RootNodeEvaluator = CameraSystemEvaluator->GetRootNodeEvaluator();
+		if (RootNodeEvaluator && RootNodeEvaluator->HasAnyActiveCameraRig())
+		{
+			const FCameraSystemEvaluationResult& Result = CameraSystemEvaluator->GetEvaluatedResult();
+
+			OutputCameraComponent->SetWorldTransform(Result.CameraPose.GetTransform());
+			OutputCameraComponent->SetFieldOfView(Result.CameraPose.GetEffectiveFieldOfView());
+			OutputCameraComponent->CurrentAperture = Result.CameraPose.GetAperture();
+			OutputCameraComponent->Filmback.SensorWidth = Result.CameraPose.GetSensorWidth();
+			OutputCameraComponent->Filmback.SensorHeight = Result.CameraPose.GetSensorHeight();
+			OutputCameraComponent->bConstrainAspectRatio = Result.CameraPose.GetConstrainAspectRatio();
+			OutputCameraComponent->bOverrideAspectRatioAxisConstraint = Result.CameraPose.GetOverrideAspectRatioAxisConstraint();
+			OutputCameraComponent->AspectRatioAxisConstraint = Result.CameraPose.GetAspectRatioAxisConstraint();
+
+			OutputCameraComponent->FocusSettings.ManualFocusDistance = Result.CameraPose.GetFocusDistance();
+			OutputCameraComponent->FocusSettings.FocusMethod = (Result.CameraPose.GetEnablePhysicalCamera() ? ECameraFocusMethod::Manual : ECameraFocusMethod::Disable);
+
+			bGotValidTransform = true;
+		}
+	}
+	
+	if (!bGotValidTransform)
+	{
+		OutputCameraComponent->SetRelativeTransform(FTransform());
+	}
+}
+
+void UGameplayCameraComponentBase::OnComponentDestroyed(bool bDestroyingHierarchy)
+{
+	Super::OnComponentDestroyed(bDestroyingHierarchy);
+}
+
+void UGameplayCameraComponentBase::OnUpdateTransform(EUpdateTransformFlags UpdateTransformFlags, ETeleportType Teleport)
+{
+	Super::OnUpdateTransform(UpdateTransformFlags, Teleport);
+
+	if (EvaluationContext && Teleport != ETeleportType::None)
+	{
+		bIsCameraCutNextFrame = true;
+	}
+
+#if WITH_EDITOR
+
+	if (bIsEditorWorld && EvaluationContext)
+	{
+		UpdateCameraEvaluationContext(false);
+	}
+
+#endif  // WITH_EDITOR
+}
+
+#if WITH_EDITOR
+
+void UGameplayCameraComponentBase::AutoManageEditorPreviewEvaluator()
+{
+	using namespace UE::Cameras;
+
+	if (!bIsEditorWorld)
+	{
+		return;
+	}
+	
+	if (bRunInEditor && !(EditorPreviewEvaluator && EvaluationContext))
+	{
+		// We want to run the camera logic in the editor but we haven't set things up for that.
+		// Let's create the preview evaluator and the evaluation context.
+		if (!EditorPreviewEvaluator)
+		{
+			EditorPreviewEvaluator = MakeShared<FCameraSystemEvaluator>();
+
+			FCameraSystemEvaluatorCreateParams CreateParams;
+			CreateParams.Owner = this;
+			CreateParams.Role = ECameraSystemEvaluatorRole::EditorPreview;
+			EditorPreviewEvaluator->Initialize(CreateParams);
+		}
+		if (!EvaluationContext)
+		{
+			EnsureCameraEvaluationContextCreated(nullptr);
+			EditorPreviewEvaluator->PushEvaluationContext(EvaluationContext.ToSharedRef());
+			EvaluationContext->SetEditorPreviewCameraRigIndex(EditorPreviewCameraRigIndex);
+		}
+	}
+	else if (!bRunInEditor && (EditorPreviewEvaluator || EvaluationContext))
+	{
+		// We don't want to run the camera logic in the editor anymore. Let's tear things down.
+		EditorPreviewEvaluator = nullptr;
+		EvaluationContext = nullptr;
+	}
+}
+
+void UGameplayCameraComponentBase::OnEditorPreviewCameraRigIndexChanged()
+{
+	if (!bIsEditorWorld)
+	{
+		return;
+	}
+
+	if (bRunInEditor && EditorPreviewEvaluator && EvaluationContext)
+	{
+		EvaluationContext->SetEditorPreviewCameraRigIndex(EditorPreviewCameraRigIndex);
+	}
+}
+
+bool UGameplayCameraComponentBase::HasActiveEditorPreview() const
+{
+	return (
+			bIsEditorWorld &&
+			bRunInEditor &&
+			EditorPreviewEvaluator &&
+			EvaluationContext);
+}
+
+void UGameplayCameraComponentBase::RecreateEditorPreviewEvaluationContext()
+{
+	if (bIsEditorWorld && bRunInEditor && EditorPreviewEvaluator && EvaluationContext)
+	{
+		EditorPreviewEvaluator->RemoveEvaluationContext(EvaluationContext.ToSharedRef());
+		EvaluationContext = nullptr;
+
+		EnsureCameraEvaluationContextCreated(nullptr);
+		EditorPreviewEvaluator->PushEvaluationContext(EvaluationContext.ToSharedRef());
+	}
+}
+
+bool UGameplayCameraComponentBase::GetEditorPreviewInfo(float DeltaTime, FMinimalViewInfo& ViewOut)
+{
+	if (OutputCameraComponent)
+	{
+		OutputCameraComponent->GetEditorPreviewInfo(DeltaTime, ViewOut);
+		return true;
+	}
+	return false;
+}
+
+void UGameplayCameraComponentBase::UpdateEditorPreviewEvaluator(float DeltaTime)
+{
+	using namespace UE::Cameras;
+
+	if (EditorPreviewEvaluator)
+	{
+		FCameraSystemEvaluationParams Params;
+		Params.DeltaTime = DeltaTime;
+		EditorPreviewEvaluator->Update(Params);
+	}
+}
+
+void UGameplayCameraComponentBase::OnDrawVisualizationHUD(const FViewport* Viewport, const FSceneView* SceneView, FCanvas* Canvas) const
+{
+	using namespace UE::Cameras;
+
+	const bool bHasShowFlag = SceneView->Family->EngineShowFlags.GetSingleFlag(CustomShowFlag);
+	if (bHasShowFlag && bRunInEditor && EditorPreviewEvaluator && EvaluationContext)
+	{
+		const AActor* OwnerActor = GetOwner();
+
+		const AActor* ViewActor = SceneView->ViewActor.Get();
+		const bool bIsLockedToCamera = (ViewActor == OwnerActor);
+
+		FCameraSystemEditorPreviewParams Params;
+		Params.Canvas = Canvas;
+		Params.SceneView = SceneView;
+		Params.bIsLockedToCamera = bIsLockedToCamera;
+		Params.bDrawWorldDebug = false;
+
+		EditorPreviewEvaluator->DrawEditorPreview(Params);
+	}
+}
+
+void UGameplayCameraComponentBase::PostEditChangeProperty(struct FPropertyChangedEvent& PropertyChangedEvent)
+{
+	using namespace UE::Cameras;
+
+	Super::PostEditChangeProperty(PropertyChangedEvent);
+
+	const FName MemberPropertyName = PropertyChangedEvent.GetMemberPropertyName();
+	if (MemberPropertyName == GET_MEMBER_NAME_CHECKED(UGameplayCameraComponentBase, bRunInEditor))
+	{
+		AutoManageEditorPreviewEvaluator();
+	}
+	else if (MemberPropertyName == GET_MEMBER_NAME_CHECKED(UGameplayCameraComponentBase, EditorPreviewCameraRigIndex))
+	{
+		OnEditorPreviewCameraRigIndexChanged();
+	}
+}
+
+#endif  // WITH_EDITOR
+
+namespace UE::Cameras
+{
+
+UE_DEFINE_CAMERA_EVALUATION_CONTEXT(FGameplayCameraComponentEvaluationContext)
+
+#if WITH_EDITOR
+
+void FGameplayCameraComponentEvaluationContext::UpdateForEditorPreview()
+{
+	FCameraSystemEvaluator* ActiveEvaluator = GetCameraSystemEvaluator();
+	if (ActiveEvaluator && ActiveEvaluator->GetRole() == ECameraSystemEvaluatorRole::EditorPreview)
+	{
+		if (GCurrentLevelEditingViewportClient && GCurrentLevelEditingViewportClient->Viewport)
+		{
+			FIntPoint ViewportSize = GCurrentLevelEditingViewportClient->Viewport->GetSizeXY();
+			OverrideViewportSize = ViewportSize;
+		}
+		else
+		{
+			OverrideViewportSize.Reset();
+		}
+	}
+}
+
+#endif  // WITH_EDITOR
+
+}  // namespace UE::Cameras
+
+#undef LOCTEXT_NAMESPACE
+
+
diff --git a/Engine/Plugins/Cameras/GameplayCameras/Source/GameplayCameras/Private/GameFramework/GameplayCameraRigActor.cpp b/Engine/Plugins/Cameras/GameplayCameras/Source/GameplayCameras/Private/GameFramework/GameplayCameraRigActor.cpp
new file mode 100644
--- /dev/null
+++ b/Engine/Plugins/Cameras/GameplayCameras/Source/GameplayCameras/Private/GameFramework/GameplayCameraRigActor.cpp
@@ -0,0 +1,42 @@
+// Copyright Epic Games, Inc. All Rights Reserved.
+
+#include "GameFramework/GameplayCameraRigActor.h"
+
+#include "Engine/EngineTypes.h"
+#include "Engine/World.h"
+#include "GameFramework/GameplayCameraRigComponent.h"
+#include "Kismet/GameplayStatics.h"
+
+#include UE_INLINE_GENERATED_CPP_BY_NAME(GameplayCameraRigActor)
+
+#define LOCTEXT_NAMESPACE "GameplayCameraRigActor"
+
+AGameplayCameraRigActor::AGameplayCameraRigActor(const FObjectInitializer& ObjectInit)
+	: Super(ObjectInit)
+{
+	CameraRigComponent = CreateDefaultSubobject<UGameplayCameraRigComponent>(TEXT("CameraRigComponent"));
+	RootComponent = CameraRigComponent;
+}
+
+USceneComponent* AGameplayCameraRigActor::GetDefaultAttachComponent() const
+{
+	return CameraRigComponent;
+}
+
+void AGameplayCameraRigActor::Tick(float DeltaTime)
+{
+	Super::Tick(DeltaTime);
+}
+
+bool AGameplayCameraRigActor::ShouldTickIfViewportsOnly() const
+{
+	return true;
+}
+
+void AGameplayCameraRigActor::DisplayDebug(UCanvas* Canvas, const FDebugDisplayInfo& DebugDisplay, float& YL, float& YPos)
+{
+	Super::DisplayDebug(Canvas, DebugDisplay, YL, YPos);
+}
+
+#undef LOCTEXT_NAMESPACE
+
diff --git a/Engine/Plugins/Cameras/GameplayCameras/Source/GameplayCameras/Private/GameFramework/GameplayCameraRigComponent.cpp b/Engine/Plugins/Cameras/GameplayCameras/Source/GameplayCameras/Private/GameFramework/GameplayCameraRigComponent.cpp
new file mode 100644
--- /dev/null
+++ b/Engine/Plugins/Cameras/GameplayCameras/Source/GameplayCameras/Private/GameFramework/GameplayCameraRigComponent.cpp
@@ -0,0 +1,135 @@
+// Copyright Epic Games, Inc. All Rights Reserved.
+
+#include "GameFramework/GameplayCameraRigComponent.h"
+
+#include "Core/CameraAsset.h"
+#include "Core/CameraRigAsset.h"
+#include "Directors/SingleCameraDirector.h"
+#include "GameplayCamerasDelegates.h"
+
+#include UE_INLINE_GENERATED_CPP_BY_NAME(GameplayCameraRigComponent)
+
+UGameplayCameraRigComponent::UGameplayCameraRigComponent(const FObjectInitializer& ObjectInit)
+	: Super(ObjectInit)
+{
+}
+
+void UGameplayCameraRigComponent::OnRegister()
+{
+	using namespace UE::Cameras;
+
+	Super::OnRegister();
+
+#if WITH_EDITOR
+
+	FGameplayCamerasDelegates::OnCameraRigAssetBuilt().AddUObject(this, &UGameplayCameraRigComponent::OnCameraRigAssetBuilt);
+
+#endif
+}
+
+void UGameplayCameraRigComponent::OnUnregister()
+{
+	using namespace UE::Cameras;
+
+#if WITH_EDITOR
+
+	FGameplayCamerasDelegates::OnCameraRigAssetBuilt().RemoveAll(this);
+
+#endif  // WITH_EDITOR
+
+	Super::OnUnregister();
+}
+
+UCameraAsset* UGameplayCameraRigComponent::GetCameraAsset()
+{
+	if (!GeneratedCameraAsset)
+	{
+		USingleCameraDirector* SingleDirector = NewObject<USingleCameraDirector>(this, TEXT("GeneratedCameraDirector"), RF_Transient);
+		SingleDirector->CameraRig = CameraRigReference.GetCameraRig();
+
+		GeneratedCameraAsset = NewObject<UCameraAsset>(this, TEXT("GeneratedCameraAsset"), RF_Transient);
+		GeneratedCameraAsset->SetCameraDirector(SingleDirector);
+
+		GeneratedCameraAsset->BuildCamera();
+	}
+
+	return GeneratedCameraAsset;
+}
+
+bool UGameplayCameraRigComponent::OnValidateCameraEvaluationContextActivation()
+{
+	if (!CameraRigReference.IsValid())
+	{
+		FFrame::KismetExecutionMessage(
+				TEXT("Can't activate gameplay camera rig component: no camera rig asset was set!"),
+				ELogVerbosity::Error);
+		return false;
+	}
+	
+	return true;
+}
+
+void UGameplayCameraRigComponent::OnUpdateCameraEvaluationContext(bool bForceApplyParameterOverrides)
+{
+	using namespace UE::Cameras;
+
+	FCameraNodeEvaluationResult& InitialResult = GetEvaluationContext()->GetInitialResult();
+
+	const bool bApplyDrivenParametersOnly = !bForceApplyParameterOverrides;
+	CameraRigReference.ApplyParameterOverrides(InitialResult, bApplyDrivenParametersOnly);
+}
+
+#if WITH_EDITOR
+
+void UGameplayCameraRigComponent::OnCameraRigAssetBuilt(const UCameraRigAsset* InCameraRigAsset)
+{
+	using namespace UE::Cameras;
+
+	if (InCameraRigAsset != CameraRigReference.GetCameraRig())
+	{
+		return;
+	}
+
+	// If our camera rig asset was just built, it may have some new parameters. We need to rebuild
+	// our variable table and context data table, and re-apply overrides.
+	if (HasActiveEditorPreview())
+	{
+		const FCameraObjectAllocationInfo& AllocationInfo = InCameraRigAsset->AllocationInfo;
+		FCameraNodeEvaluationResult& InitialResult = GetEvaluationContext()->GetInitialResult();
+		InitialResult.VariableTable.Initialize(AllocationInfo.VariableTableInfo);
+		InitialResult.ContextDataTable.Initialize(AllocationInfo.ContextDataTableInfo);
+
+		CameraRigReference.RebuildParametersIfNeeded();
+
+		UpdateCameraEvaluationContext(true);
+	}
+}
+
+void UGameplayCameraRigComponent::PostEditChangeProperty( struct FPropertyChangedEvent& PropertyChangedEvent)
+{
+	using namespace UE::Cameras;
+
+	Super::PostEditChangeProperty(PropertyChangedEvent);
+
+	const FName MemberPropertyName = PropertyChangedEvent.GetMemberPropertyName();
+	if (MemberPropertyName == GET_MEMBER_NAME_CHECKED(UGameplayCameraRigComponent, CameraRigReference))
+	{
+		if (HasActiveEditorPreview())
+		{
+			if (PropertyChangedEvent.GetPropertyName() == TEXT("CameraRig"))
+			{
+				// The camera rig asset has changed! Recreate the context.
+				GeneratedCameraAsset = nullptr;
+				RecreateEditorPreviewEvaluationContext();
+			}
+			else
+			{
+				// Otherwise, maybe one of the parameter overrides has changed. Re-apply them.
+				UpdateCameraEvaluationContext(true);
+			}
+		}
+	}
+}
+
+#endif  // WITH_EDITOR
+
diff --git a/Engine/Plugins/Cameras/GameplayCameras/Source/GameplayCameras/Private/GameFramework/GameplayCamerasPlayerCameraManager.cpp b/Engine/Plugins/Cameras/GameplayCameras/Source/GameplayCameras/Private/GameFramework/GameplayCamerasPlayerCameraManager.cpp
--- a/Engine/Plugins/Cameras/GameplayCameras/Source/GameplayCameras/Private/GameFramework/GameplayCamerasPlayerCameraManager.cpp
+++ b/Engine/Plugins/Cameras/GameplayCameras/Source/GameplayCameras/Private/GameFramework/GameplayCamerasPlayerCameraManager.cpp
@@ -14,6 +14,7 @@
 #include "Engine/Engine.h"
 #include "GameFramework/ActorCameraEvaluationContext.h"
 #include "GameFramework/GameplayCameraComponent.h"
+#include "GameFramework/GameplayCameraRigComponent.h"
 #include "GameFramework/GameplayCameraSystemActor.h"
 #include "GameFramework/GameplayCameraSystemComponent.h"
 #include "GameFramework/GameplayCameraSystemHost.h"
@@ -267,6 +268,10 @@ void AGameplayCamerasPlayerCameraManager::SetViewTarget(AActor* NewViewTarget, F
 	{
 		GameplayCameraComponent->ActivateCameraForPlayerController(PCOwner);
 	}
+	else if (UGameplayCameraRigComponent* GameplayCameraRigComponent = NewViewTarget->FindComponentByClass<UGameplayCameraRigComponent>())
+	{
+		GameplayCameraRigComponent->ActivateCameraForPlayerController(PCOwner);
+	}
 	else if (UCameraComponent* CameraComponent = NewViewTarget->FindComponentByClass<UCameraComponent>())
 	{
 		TSharedRef<FActorCameraEvaluationContext> NewContext = MakeShared<FActorCameraEvaluationContext>(CameraComponent);
diff --git a/Engine/Plugins/Cameras/GameplayCameras/Source/GameplayCameras/Private/MovieScene/MovieSceneCameraParameterInstantiator.cpp b/Engine/Plugins/Cameras/GameplayCameras/Source/GameplayCameras/Private/MovieScene/MovieSceneCameraParameterInstantiator.cpp
--- a/Engine/Plugins/Cameras/GameplayCameras/Source/GameplayCameras/Private/MovieScene/MovieSceneCameraParameterInstantiator.cpp
+++ b/Engine/Plugins/Cameras/GameplayCameras/Source/GameplayCameras/Private/MovieScene/MovieSceneCameraParameterInstantiator.cpp
@@ -4,6 +4,8 @@
 
 #include "Core/CameraAsset.h"
 #include "Core/CameraAssetReference.h"
+#include "Core/CameraRigAsset.h"
+#include "Core/CameraRigAssetReference.h"
 #include "Core/CameraObjectInterfaceParameterDefinition.h"
 #include "EntitySystem/BuiltInComponentTypes.h"
 #include "EntitySystem/MovieSceneEntityManager.h"
@@ -13,6 +15,7 @@
 #include "Evaluation/PreAnimatedState/MovieScenePreAnimatedStateStorage.h"
 #include "Evaluation/PreAnimatedState/MovieScenePreAnimatedStorageID.inl"
 #include "GameFramework/GameplayCameraComponent.h"
+#include "GameFramework/GameplayCameraRigComponent.h"
 #include "MovieScene/MovieSceneGameplayCamerasComponentTypes.h"
 #include "Tracks/MovieScenePropertyTrack.h"
 
@@ -26,22 +29,36 @@ struct FPreAnimatedCameraParameterStateTraits : UE::MovieScene::FPreAnimatedStat
 	using KeyType = TTuple<FObjectKey, FGuid>;
 	using StorageType = bool;
 
-	bool CachePreAnimatedValue(UGameplayCameraComponent* CameraComponent, const FGuid& ParameterGuid)
+	bool CachePreAnimatedValue(UGameplayCameraComponentBase* CameraComponentBase, const FGuid& ParameterGuid)
 	{
-		if (CameraComponent && ParameterGuid.IsValid())
+		if (ParameterGuid.IsValid())
 		{
-			return CameraComponent->CameraReference.IsParameterAnimated(ParameterGuid);
+			if (UGameplayCameraComponent* CameraComponent = Cast<UGameplayCameraComponent>(CameraComponentBase))
+			{
+				return CameraComponent->CameraReference.IsParameterAnimated(ParameterGuid);
+			}
+			else if (UGameplayCameraRigComponent* CameraRigComponent = Cast<UGameplayCameraRigComponent>(CameraComponentBase))
+			{
+				return CameraRigComponent->CameraRigReference.IsParameterAnimated(ParameterGuid);
+			}
 		}
 		return false;
 	}
 
 	void RestorePreAnimatedValue(const KeyType& InKey, bool bWasAnimated, const UE::MovieScene::FRestoreStateParams& Params)
 	{
-		UGameplayCameraComponent* CameraComponent = Cast<UGameplayCameraComponent>(InKey.Get<0>().ResolveObjectPtr());
+		UGameplayCameraComponentBase* CameraComponentBase = Cast<UGameplayCameraComponentBase>(InKey.Get<0>().ResolveObjectPtr());
 		const FGuid& ParameterGuid = InKey.Get<1>();
-		if (CameraComponent && ParameterGuid.IsValid())
+		if (ParameterGuid.IsValid())
 		{
-			CameraComponent->CameraReference.SetParameterAnimated(ParameterGuid, bWasAnimated);
+			if (UGameplayCameraComponent* CameraComponent = Cast<UGameplayCameraComponent>(CameraComponentBase))
+			{
+				CameraComponent->CameraReference.SetParameterAnimated(ParameterGuid, bWasAnimated);
+			}
+			else if (UGameplayCameraRigComponent* CameraRigComponent = Cast<UGameplayCameraRigComponent>(CameraComponentBase))
+			{
+				CameraRigComponent->CameraRigReference.SetParameterAnimated(ParameterGuid, bWasAnimated);
+			}
 		}
 	}
 };
@@ -70,43 +87,76 @@ struct FSetupCameraParameterOverrideTask
 	{
 		using namespace UE::MovieScene;
 
-		UGameplayCameraComponent* CameraComponent = Cast<UGameplayCameraComponent>(BoundObject);
-		if (!ensure(CameraComponent))
+		UGameplayCameraComponentBase* CameraComponentBase = Cast<UGameplayCameraComponentBase>(BoundObject);
+		if (!ensure(CameraComponentBase))
 		{
 			return;
 		}
 
-		const FGuid ParameterGuid = GetParameterGuid(CameraComponent->CameraReference, PropertyBinding);
+		FGuid ParameterGuid;
+		bool bWasAnimated;
+
+		if (UGameplayCameraComponent* CameraComponent = Cast<UGameplayCameraComponent>(CameraComponentBase))
+		{
+			ParameterGuid = GetParameterGuid(CameraComponent->CameraReference, PropertyBinding);
+			bWasAnimated = CameraComponent->CameraReference.IsParameterAnimated(ParameterGuid);
+		}
+		else if (UGameplayCameraRigComponent* CameraRigComponent = Cast<UGameplayCameraRigComponent>(CameraComponentBase))
+		{
+			ParameterGuid = GetParameterGuid(CameraRigComponent->CameraRigReference, PropertyBinding);
+			bWasAnimated = CameraRigComponent->CameraRigReference.IsParameterAnimated(ParameterGuid);
+		}
+
 		if (!ParameterGuid.IsValid())
 		{
 			return;
 		}
 
-		const bool bWasAnimated = CameraComponent->CameraReference.IsParameterAnimated(ParameterGuid);
-
 		using PreAnimatedKeyType = TTuple<FObjectKey, FGuid>;
 
-		PreAnimatedStorage->BeginTrackingEntity(EntityID, true, RootInstanceHandle, CameraComponent, ParameterGuid);
+		PreAnimatedStorage->BeginTrackingEntity(EntityID, true, RootInstanceHandle, CameraComponentBase, ParameterGuid);
 		PreAnimatedStorage->CachePreAnimatedValue(
-				PreAnimatedKeyType(CameraComponent, ParameterGuid),
+				PreAnimatedKeyType(CameraComponentBase, ParameterGuid),
 				[bWasAnimated](const PreAnimatedKeyType& InKey) { return bWasAnimated; });
 
 		CameraParameterOverrideID = ParameterGuid;
-		CameraComponent->CameraReference.SetParameterAnimated(ParameterGuid, true);
+
+		if (UGameplayCameraComponent* CameraComponent = Cast<UGameplayCameraComponent>(CameraComponentBase))
+		{
+			CameraComponent->CameraReference.SetParameterAnimated(ParameterGuid, true);
+		}
+		else if (UGameplayCameraRigComponent* CameraRigComponent = Cast<UGameplayCameraRigComponent>(CameraComponentBase))
+		{
+			CameraRigComponent->CameraRigReference.SetParameterAnimated(ParameterGuid, true);
+		}
 	}
 
 	FGuid GetParameterGuid(const FCameraAssetReference& CameraReference, const FMovieScenePropertyBinding& PropertyBinding)
 	{
 		const UCameraAsset* CameraAsset = CameraReference.GetCameraAsset();
-		if (!ensure(CameraAsset))
+		if (ensure(CameraAsset))
 		{
-			return FGuid();
+			return GetParameterGuid(CameraAsset->GetParameterDefinitions(), TEXT("CameraReference.Parameters.Value"), PropertyBinding);
 		}
+		return FGuid();
+	}
 
+	FGuid GetParameterGuid(const FCameraRigAssetReference& CameraRigReference, const FMovieScenePropertyBinding& PropertyBinding)
+	{
+		const UCameraRigAsset* CameraRigAsset = CameraRigReference.GetCameraRig();
+		if (ensure(CameraRigAsset))
+		{
+			return GetParameterGuid(CameraRigAsset->GetParameterDefinitions(), TEXT("CameraRigReference.Parameters.Value"), PropertyBinding);
+		}
+		return FGuid();
+	}
+
+	FGuid GetParameterGuid(TConstArrayView<FCameraObjectInterfaceParameterDefinition> ParameterDefinitions, const FString& PropertyPathStartsWith, const FMovieScenePropertyBinding& PropertyBinding)
+	{
 		// This isn't ideal but we know all camera parameters are bound to the "Parameters" property bag
 		// of the camera component, so use that to figure out the parameter name and find is ID.
 		const FString PropertyPath = PropertyBinding.PropertyPath.ToString();
-		if (!ensure(PropertyPath.StartsWith(TEXT("CameraReference.Parameters.Value"))))
+		if (!ensure(PropertyPath.StartsWith(PropertyPathStartsWith)))
 		{
 			return FGuid();
 		}
@@ -131,7 +181,7 @@ struct FSetupCameraParameterOverrideTask
 			}
 		}
 
-		for (const FCameraObjectInterfaceParameterDefinition& ParameterDefintion : CameraAsset->GetParameterDefinitions())
+		for (const FCameraObjectInterfaceParameterDefinition& ParameterDefintion : ParameterDefinitions)
 		{
 			if (ParameterDefintion.ParameterName == ParameterName)
 			{
diff --git a/Engine/Plugins/Cameras/GameplayCameras/Source/GameplayCameras/Public/GameFramework/GameplayCameraActor.h b/Engine/Plugins/Cameras/GameplayCameras/Source/GameplayCameras/Public/GameFramework/GameplayCameraActor.h
--- a/Engine/Plugins/Cameras/GameplayCameras/Source/GameplayCameras/Public/GameFramework/GameplayCameraActor.h
+++ b/Engine/Plugins/Cameras/GameplayCameras/Source/GameplayCameras/Public/GameFramework/GameplayCameraActor.h
@@ -9,8 +9,6 @@
 
 #include "GameplayCameraActor.generated.h"
 
-class APlayerController;
-class UCineCameraComponent;
 class UGameplayCameraComponent;
 
 /**
diff --git a/Engine/Plugins/Cameras/GameplayCameras/Source/GameplayCameras/Public/GameFramework/GameplayCameraComponent.h b/Engine/Plugins/Cameras/GameplayCameras/Source/GameplayCameras/Public/GameFramework/GameplayCameraComponent.h
--- a/Engine/Plugins/Cameras/GameplayCameras/Source/GameplayCameras/Public/GameFramework/GameplayCameraComponent.h
+++ b/Engine/Plugins/Cameras/GameplayCameras/Source/GameplayCameras/Public/GameFramework/GameplayCameraComponent.h
@@ -2,38 +2,16 @@
 
 #pragma once
 
-#include "Components/SceneComponent.h"
 #include "Core/CameraAssetReference.h"
-#include "Core/CameraEvaluationContext.h"
-#include "GameFramework/BlueprintCameraNodeEvaluationResult.h"
-#include "GameFramework/BlueprintCameraPose.h"
-#include "GameFramework/BlueprintCameraVariableTable.h"
-#include "UObject/ObjectMacros.h"
+#include "GameFramework/GameplayCameraComponentBase.h"
 
 #include "GameplayCameraComponent.generated.h"
 
-class APlayerController;
-class FPrimitiveDrawInterface;
-class FSceneView;
-class FViewport;
-class UCameraAsset;
-class UCanvas;
-class UCineCameraComponent;
-class UGameplayCameraSystemHost;
-
-namespace UE::Cameras
-{
-
-class FCameraSystemEvaluator;
-class FGameplayCameraComponentEvaluationContext;
-
-}  // namespace UE::Cameras
-
 /**
  * A component that can run a camera asset inside its own camera evaluation context.
  */
-UCLASS(Blueprintable, MinimalAPI, ClassGroup=Camera, HideCategories=(Mobility, Rendering, LOD), meta=(BlueprintSpawnableComponent))
-class UGameplayCameraComponent : public USceneComponent
+UCLASS(Blueprintable, MinimalAPI)
+class UGameplayCameraComponent : public UGameplayCameraComponentBase
 {
 	GENERATED_BODY()
 
@@ -42,150 +20,36 @@ public:
 	/** Create a new camera component. */
 	GAMEPLAYCAMERAS_API UGameplayCameraComponent(const FObjectInitializer& ObjectInit);
 
-	/** Get the camera evaluation context used by this component. */
-	GAMEPLAYCAMERAS_API TSharedPtr<const UE::Cameras::FCameraEvaluationContext> GetEvaluationContext() const;
-
-	/** Get the camera evaluation context used by this component. */
-	GAMEPLAYCAMERAS_API TSharedPtr<UE::Cameras::FCameraEvaluationContext> GetEvaluationContext();
-
-	/** Get the player controller this component is currently activated for (if any). */
-	GAMEPLAYCAMERAS_API APlayerController* GetPlayerController() const;
-
-public:
-
-	/** Gets the child camera component used as the "output" for the gameplay/procedural camera. */
-	UFUNCTION(BlueprintGetter, Category=Camera)
-	UCineCameraComponent* GetOutputCameraComponent() const { return OutputCameraComponent; }
-
-	/** Activates the camera for the given player. */
-	UFUNCTION(BlueprintCallable, Category=Camera)
-	GAMEPLAYCAMERAS_API void ActivateCameraForPlayerIndex(int32 PlayerIndex);
-
-	/** Activates the camera for the given player. */
-	UFUNCTION(BlueprintCallable, Category=Camera)
-	GAMEPLAYCAMERAS_API void ActivateCameraForPlayerController(APlayerController* PlayerController);
-
-	/** Deactivates the camera for the last player it was activated for. */
-	UFUNCTION(BlueprintCallable, Category=Camera)
-	GAMEPLAYCAMERAS_API void DeactivateCamera();
-
-	/** Gets the shared camera evaluation data for this component's evaluation context. */
-	UFUNCTION(BlueprintPure, Category=Camera, meta=(DisplayName="Get Shared Camera Data"))
-	GAMEPLAYCAMERAS_API FBlueprintCameraNodeEvaluationResult GetInitialResult() const;
-	
-	/** Gets the camera evaluation data for a given sub-set of camera rigs in this component's evaluation context. */
-	UFUNCTION(BlueprintPure, Category=Camera, meta=(DisplayName="Get Conditional Camera Data"))
-	GAMEPLAYCAMERAS_API FBlueprintCameraNodeEvaluationResult GetConditionalResult(ECameraEvaluationDataCondition Condition) const;
-
-	/** Gets the initial camera pose for this component's camera evaluation context. */
-	UFUNCTION(BlueprintPure, Category=Camera, meta=(DeprecatedFunction, DeprecationMessage="Please use GetSharedCameraData"))
-	GAMEPLAYCAMERAS_API FBlueprintCameraPose GetInitialPose() const;
-
-	/** Sets the initial camera pose for this component's camera evaluation context. */
-	UFUNCTION(BlueprintCallable, Category=Camera, meta=(DeprecatedFunction, DeprecationMessage="Please use GetSharedCameraData"))
-	GAMEPLAYCAMERAS_API bool SetInitialPose(const FBlueprintCameraPose& CameraPose);
-
-	/** Gets the initial camera variable table for this component's camera evaluation context. */
-	UFUNCTION(BlueprintPure, Category=Camera, meta=(DeprecatedFunction, DeprecationMessage="Please use GetSharedCameraData"))
-	GAMEPLAYCAMERAS_API FBlueprintCameraVariableTable GetInitialVariableTable() const;
-
 public:
 
-	// UActorComponent interface
+	// UActorComponent interface.
 	virtual void OnRegister() override;
-	virtual void BeginPlay() override;
-	virtual void EndPlay(const EEndPlayReason::Type EndPlayReason) override;
 	virtual void OnUnregister() override;
-	virtual void TickComponent(float DeltaTime, ELevelTick TickType, FActorComponentTickFunction *ThisTickFunction) override;
-	virtual void OnComponentDestroyed(bool bDestroyingHierarchy) override;
-#if WITH_EDITOR
-	virtual bool GetEditorPreviewInfo(float DeltaTime, FMinimalViewInfo& ViewOut) override;
-#endif 
-
-	// USceneComponent interface.
-	virtual void OnUpdateTransform(EUpdateTransformFlags UpdateTransformFlags, ETeleportType Teleport) override;
 
 	// UObject interface.
 	virtual void PostLoad() override;
-	virtual void BeginDestroy() override;
 #if WITH_EDITOR
 	virtual void PostEditChangeProperty( struct FPropertyChangedEvent& PropertyChangedEvent) override;
 #endif
 
-	static void AddReferencedObjects(UObject* InThis, FReferenceCollector& Collector);
-
-public:
-
-#if WITH_EDITOR
-	GAMEPLAYCAMERAS_API void OnDrawVisualizationHUD(const FViewport* Viewport, const FSceneView* SceneView, FCanvas* Canvas) const;
-#endif
-
-private:
-
-	void ActivateCameraEvaluationContext(int32 PlayerIndex);
-	void ActivateCameraEvaluationContext(APlayerController* PlayerController);
-	void EnsureCameraEvaluationContextCreated(APlayerController* PlayerController);
-	void UpdateCameraEvaluationContext(bool bForceApplyParameterOverrides);
-	void UpdateOutputCameraComponent();
-	void DeactivateCameraEvaluationContext();
-
-#if WITH_EDITOR
-	void AutoManageEditorPreviewEvaluator();
-	void OnCameraAssetReferenceChanged();
-	void OnEditorPreviewCameraRigIndexChanged();
-	void OnCameraAssetBuilt(const UCameraAsset* InCameraAsset);
-
-	void UpdateEditorPreviewEvaluator(float DeltaTime);
-#endif  // WITH_EDITOR
-
 public:
 
 	/** The camera asset to run. */
 	UPROPERTY(BlueprintReadWrite, EditAnywhere, Category=Camera, meta=(SequencerHideProperty=true))
 	FCameraAssetReference CameraReference;
 
-	/**
-	 * If AutoActivate is set, auto-activates this component's camera for the given player.
-	 * This is equivalent to calling ActivateCamera on BeginPlay.
-	 */
-	UPROPERTY(BlueprintReadWrite, EditAnywhere, Category=Activation, meta=(EditCondition="bAutoActivate"))
-	TEnumAsByte<EAutoReceiveInput::Type> AutoActivateForPlayer;
-
-#if WITH_EDITORONLY_DATA
-
-	UPROPERTY(EditAnywhere, Category=Camera)
-	bool bRunInEditor = true;
-
-	UPROPERTY(EditAnywhere, Category=Camera, meta=(EditCondition="bRunInEditor"))
-	int32 EditorPreviewCameraRigIndex = 0;
-
-#endif  // WITH_EDITORONLY_DATA
-
-protected:
-
-	UPROPERTY(Transient)
-	TObjectPtr<UCineCameraComponent> OutputCameraComponent;
-
-	UPROPERTY(Transient)
-	TObjectPtr<UGameplayCameraSystemHost> CameraSystemHost;
-
 protected:
 
-	using FGameplayCameraComponentEvaluationContext = UE::Cameras::FGameplayCameraComponentEvaluationContext;
+	// UGameplayCameraComponentBase interface.
+	virtual UCameraAsset* GetCameraAsset() override;
+	virtual bool OnValidateCameraEvaluationContextActivation() override;
+	virtual void OnUpdateCameraEvaluationContext(bool bForceApplyParameterOverrides) override;
 
-	TSharedPtr<FGameplayCameraComponentEvaluationContext> EvaluationContext;
-
-	bool bIsCameraCutNextFrame = false;
+private:
 
 #if WITH_EDITOR
-	
-	TSharedPtr<UE::Cameras::FCameraSystemEvaluator> EditorPreviewEvaluator;
-
-	bool bIsEditorWorld = false;
-
-	int32 CustomShowFlag = INDEX_NONE;
-
-#endif  // WITH_EDITOR
+	void OnCameraAssetBuilt(const UCameraAsset* InCameraAsset);
+#endif
 
 private:
 
@@ -193,24 +57,3 @@ private:
 	TObjectPtr<UCameraAsset> Camera_DEPRECATED;
 };
 
-namespace UE::Cameras
-{
-
-/**
- * Evaluation context for the gameplay camera component.
- */
-class FGameplayCameraComponentEvaluationContext : public FCameraEvaluationContext
-{
-	UE_DECLARE_CAMERA_EVALUATION_CONTEXT(GAMEPLAYCAMERAS_API, FGameplayCameraComponentEvaluationContext)
-
-#if WITH_EDITOR
-
-public:
-
-	void UpdateForEditorPreview();
-
-#endif  // WITH_EDITOR
-};
-
-}  // namespace UE::Cameras
-
diff --git a/Engine/Plugins/Cameras/GameplayCameras/Source/GameplayCameras/Public/GameFramework/GameplayCameraComponentBase.h b/Engine/Plugins/Cameras/GameplayCameras/Source/GameplayCameras/Public/GameFramework/GameplayCameraComponentBase.h
new file mode 100644
--- /dev/null
+++ b/Engine/Plugins/Cameras/GameplayCameras/Source/GameplayCameras/Public/GameFramework/GameplayCameraComponentBase.h
@@ -0,0 +1,218 @@
+// Copyright Epic Games, Inc. All Rights Reserved.
+
+#pragma once
+
+#include "Components/SceneComponent.h"
+#include "Core/CameraAssetReference.h"
+#include "Core/CameraEvaluationContext.h"
+#include "GameFramework/BlueprintCameraNodeEvaluationResult.h"
+#include "GameFramework/BlueprintCameraPose.h"
+#include "GameFramework/BlueprintCameraVariableTable.h"
+#include "UObject/ObjectMacros.h"
+
+#include "GameplayCameraComponentBase.generated.h"
+
+class APlayerController;
+class FPrimitiveDrawInterface;
+class FSceneView;
+class FViewport;
+class UCameraAsset;
+class UCanvas;
+class UCineCameraComponent;
+class UGameplayCameraSystemHost;
+
+namespace UE::Cameras
+{
+
+class FCameraSystemEvaluator;
+class FGameplayCameraComponentEvaluationContext;
+
+}  // namespace UE::Cameras
+
+/**
+ * A component that can run a camera asset inside its own camera evaluation context.
+ */
+UCLASS(Blueprintable, MinimalAPI, Abstract,
+		ClassGroup=Camera, 
+		HideCategories=(Mobility, Rendering, LOD), 
+		meta=(BlueprintSpawnableComponent))
+class UGameplayCameraComponentBase : public USceneComponent
+{
+	GENERATED_BODY()
+
+public:
+
+	/** Create a new camera component. */
+	GAMEPLAYCAMERAS_API UGameplayCameraComponentBase(const FObjectInitializer& ObjectInit);
+
+	/** Get the camera evaluation context used by this component. */
+	GAMEPLAYCAMERAS_API TSharedPtr<const UE::Cameras::FCameraEvaluationContext> GetEvaluationContext() const;
+
+	/** Get the camera evaluation context used by this component. */
+	GAMEPLAYCAMERAS_API TSharedPtr<UE::Cameras::FCameraEvaluationContext> GetEvaluationContext();
+
+	/** Get the player controller this component is currently activated for (if any). */
+	GAMEPLAYCAMERAS_API APlayerController* GetPlayerController() const;
+
+public:
+
+	/** Gets the child camera component used as the "output" for the gameplay/procedural camera. */
+	UFUNCTION(BlueprintGetter, Category=Camera)
+	UCineCameraComponent* GetOutputCameraComponent() const { return OutputCameraComponent; }
+
+	/** Activates the camera for the given player. */
+	UFUNCTION(BlueprintCallable, Category=Camera)
+	GAMEPLAYCAMERAS_API void ActivateCameraForPlayerIndex(int32 PlayerIndex);
+
+	/** Activates the camera for the given player. */
+	UFUNCTION(BlueprintCallable, Category=Camera)
+	GAMEPLAYCAMERAS_API void ActivateCameraForPlayerController(APlayerController* PlayerController);
+
+	/** Deactivates the camera for the last player it was activated for. */
+	UFUNCTION(BlueprintCallable, Category=Camera)
+	GAMEPLAYCAMERAS_API void DeactivateCamera();
+
+	/** Gets the shared camera evaluation data for this component's evaluation context. */
+	UFUNCTION(BlueprintPure, Category=Camera, meta=(DisplayName="Get Shared Camera Data"))
+	GAMEPLAYCAMERAS_API FBlueprintCameraNodeEvaluationResult GetInitialResult() const;
+	
+	/** Gets the camera evaluation data for a given sub-set of camera rigs in this component's evaluation context. */
+	UFUNCTION(BlueprintPure, Category=Camera, meta=(DisplayName="Get Conditional Camera Data"))
+	GAMEPLAYCAMERAS_API FBlueprintCameraNodeEvaluationResult GetConditionalResult(ECameraEvaluationDataCondition Condition) const;
+
+	/** Gets the initial camera pose for this component's camera evaluation context. */
+	UFUNCTION(BlueprintPure, Category=Camera, meta=(DeprecatedFunction, DeprecationMessage="Please use GetSharedCameraData"))
+	GAMEPLAYCAMERAS_API FBlueprintCameraPose GetInitialPose() const;
+
+	/** Sets the initial camera pose for this component's camera evaluation context. */
+	UFUNCTION(BlueprintCallable, Category=Camera, meta=(DeprecatedFunction, DeprecationMessage="Please use GetSharedCameraData"))
+	GAMEPLAYCAMERAS_API bool SetInitialPose(const FBlueprintCameraPose& CameraPose);
+
+	/** Gets the initial camera variable table for this component's camera evaluation context. */
+	UFUNCTION(BlueprintPure, Category=Camera, meta=(DeprecatedFunction, DeprecationMessage="Please use GetSharedCameraData"))
+	GAMEPLAYCAMERAS_API FBlueprintCameraVariableTable GetInitialVariableTable() const;
+
+public:
+
+	// UActorComponent interface
+	virtual void OnRegister() override;
+	virtual void BeginPlay() override;
+	virtual void EndPlay(const EEndPlayReason::Type EndPlayReason) override;
+	virtual void OnUnregister() override;
+	virtual void TickComponent(float DeltaTime, ELevelTick TickType, FActorComponentTickFunction *ThisTickFunction) override;
+	virtual void OnComponentDestroyed(bool bDestroyingHierarchy) override;
+#if WITH_EDITOR
+	virtual bool GetEditorPreviewInfo(float DeltaTime, FMinimalViewInfo& ViewOut) override;
+#endif 
+
+	// USceneComponent interface.
+	virtual void OnUpdateTransform(EUpdateTransformFlags UpdateTransformFlags, ETeleportType Teleport) override;
+
+	// UObject interface.
+	virtual void BeginDestroy() override;
+#if WITH_EDITOR
+	virtual void PostEditChangeProperty( struct FPropertyChangedEvent& PropertyChangedEvent) override;
+#endif
+
+	static void AddReferencedObjects(UObject* InThis, FReferenceCollector& Collector);
+
+public:
+
+#if WITH_EDITOR
+	GAMEPLAYCAMERAS_API void OnDrawVisualizationHUD(const FViewport* Viewport, const FSceneView* SceneView, FCanvas* Canvas) const;
+#endif
+
+protected:
+
+	// UGameplayCameraComponentBase interface.
+	virtual UCameraAsset* GetCameraAsset() PURE_VIRTUAL(UGameplayCameraComponentBase::GetCameraAsset, return nullptr;)
+	virtual bool OnValidateCameraEvaluationContextActivation() { return true; }
+	virtual void OnUpdateCameraEvaluationContext(bool bForceApplyParameterOverrides) {}
+
+	void UpdateCameraEvaluationContext(bool bForceApplyParameterOverrides);
+#if WITH_EDITOR
+	bool HasActiveEditorPreview() const;
+	void RecreateEditorPreviewEvaluationContext();
+	void UpdateEditorPreviewEvaluator(float DeltaTime);
+#endif  // WITH_EDITOR
+
+private:
+
+	void ActivateCameraEvaluationContext(int32 PlayerIndex);
+	void ActivateCameraEvaluationContext(APlayerController* PlayerController);
+	void EnsureCameraEvaluationContextCreated(APlayerController* PlayerController);
+	void UpdateOutputCameraComponent();
+	void DeactivateCameraEvaluationContext();
+
+#if WITH_EDITOR
+	void AutoManageEditorPreviewEvaluator();
+	void OnEditorPreviewCameraRigIndexChanged();
+#endif  // WITH_EDITOR
+
+public:
+
+	/**
+	 * If AutoActivate is set, auto-activates this component's camera for the given player.
+	 * This is equivalent to calling ActivateCamera on BeginPlay.
+	 */
+	UPROPERTY(BlueprintReadWrite, EditAnywhere, Category=Activation, meta=(EditCondition="bAutoActivate"))
+	TEnumAsByte<EAutoReceiveInput::Type> AutoActivateForPlayer;
+
+#if WITH_EDITORONLY_DATA
+
+	UPROPERTY(EditAnywhere, Category=Camera)
+	bool bRunInEditor = true;
+
+	UPROPERTY(EditAnywhere, Category=Camera, meta=(EditCondition="bRunInEditor"))
+	int32 EditorPreviewCameraRigIndex = 0;
+
+#endif  // WITH_EDITORONLY_DATA
+
+private:
+
+	UPROPERTY(Transient)
+	TObjectPtr<UCineCameraComponent> OutputCameraComponent;
+
+	UPROPERTY(Transient)
+	TObjectPtr<UGameplayCameraSystemHost> CameraSystemHost;
+
+private:
+
+	using FGameplayCameraComponentEvaluationContext = UE::Cameras::FGameplayCameraComponentEvaluationContext;
+
+	TSharedPtr<FGameplayCameraComponentEvaluationContext> EvaluationContext;
+
+	bool bIsCameraCutNextFrame = false;
+
+#if WITH_EDITOR
+	
+	TSharedPtr<UE::Cameras::FCameraSystemEvaluator> EditorPreviewEvaluator;
+
+	bool bIsEditorWorld = false;
+
+	int32 CustomShowFlag = INDEX_NONE;
+
+#endif  // WITH_EDITOR
+};
+
+namespace UE::Cameras
+{
+
+/**
+ * Evaluation context for the gameplay camera component.
+ */
+class FGameplayCameraComponentEvaluationContext : public FCameraEvaluationContext
+{
+	UE_DECLARE_CAMERA_EVALUATION_CONTEXT(GAMEPLAYCAMERAS_API, FGameplayCameraComponentEvaluationContext)
+
+#if WITH_EDITOR
+
+public:
+
+	void UpdateForEditorPreview();
+
+#endif  // WITH_EDITOR
+};
+
+}  // namespace UE::Cameras
+
diff --git a/Engine/Plugins/Cameras/GameplayCameras/Source/GameplayCameras/Public/GameFramework/GameplayCameraRigActor.h b/Engine/Plugins/Cameras/GameplayCameras/Source/GameplayCameras/Public/GameFramework/GameplayCameraRigActor.h
new file mode 100644
--- /dev/null
+++ b/Engine/Plugins/Cameras/GameplayCameras/Source/GameplayCameras/Public/GameFramework/GameplayCameraRigActor.h
@@ -0,0 +1,45 @@
+// Copyright Epic Games, Inc. All Rights Reserved.
+
+#pragma once
+
+#include "CoreMinimal.h"
+#include "Engine/EngineTypes.h"
+#include "GameFramework/Actor.h"
+#include "UObject/ObjectMacros.h"
+
+#include "GameplayCameraRigActor.generated.h"
+
+class UGameplayCameraRigComponent;
+
+/**
+ * An actor that can run a camera asset.
+ */
+UCLASS(BlueprintType, MinimalAPI, ClassGroup=Camera, HideCategories=(Input, Rendering))
+class AGameplayCameraRigActor : public AActor
+{
+	GENERATED_BODY()
+
+public:
+
+	AGameplayCameraRigActor(const FObjectInitializer& ObjectInit);
+
+public:
+
+	/** Gets the camera component. */
+	UFUNCTION(BlueprintGetter, Category=Camera)
+	UGameplayCameraRigComponent* GetCameraRigComponent() const { return CameraRigComponent; }
+
+public:
+
+	// AActor interface.
+	virtual USceneComponent* GetDefaultAttachComponent() const override;
+	virtual void Tick(float DeltaTime) override;
+	virtual bool ShouldTickIfViewportsOnly() const override;
+	virtual void DisplayDebug(UCanvas* Canvas, const FDebugDisplayInfo& DebugDisplay, float& YL, float& YPos) override;
+
+private:
+
+	UPROPERTY(VisibleAnywhere, Category=Camera, BlueprintGetter="GetCameraRigComponent", meta=(ExposeFunctionCategories="Camera"))
+	TObjectPtr<UGameplayCameraRigComponent> CameraRigComponent;
+};
+
diff --git a/Engine/Plugins/Cameras/GameplayCameras/Source/GameplayCameras/Public/GameFramework/GameplayCameraRigComponent.h b/Engine/Plugins/Cameras/GameplayCameras/Source/GameplayCameras/Public/GameFramework/GameplayCameraRigComponent.h
new file mode 100644
--- /dev/null
+++ b/Engine/Plugins/Cameras/GameplayCameras/Source/GameplayCameras/Public/GameFramework/GameplayCameraRigComponent.h
@@ -0,0 +1,58 @@
+// Copyright Epic Games, Inc. All Rights Reserved.
+
+#pragma once
+
+#include "Core/CameraRigAssetReference.h"
+#include "GameFramework/GameplayCameraComponentBase.h"
+
+#include "GameplayCameraRigComponent.generated.h"
+
+/**
+ * A component that can run a camera rig asset inside its own camera evaluation context.
+ */
+UCLASS(Blueprintable, MinimalAPI)
+class UGameplayCameraRigComponent : public UGameplayCameraComponentBase
+{
+	GENERATED_BODY()
+
+public:
+
+	/** Create a new camera component. */
+	GAMEPLAYCAMERAS_API UGameplayCameraRigComponent(const FObjectInitializer& ObjectInit);
+
+public:
+
+	// UActorComponent interface.
+	virtual void OnRegister() override;
+	virtual void OnUnregister() override;
+
+	// UObject interface.
+#if WITH_EDITOR
+	virtual void PostEditChangeProperty( struct FPropertyChangedEvent& PropertyChangedEvent) override;
+#endif
+
+public:
+
+	/** The camera rig asset to run. */
+	UPROPERTY(BlueprintReadWrite, EditAnywhere, Category=Camera, meta=(SequencerHideProperty=true))
+	FCameraRigAssetReference CameraRigReference;
+
+protected:
+
+	// UGameplayCameraComponentBase interface.
+	virtual UCameraAsset* GetCameraAsset() override;
+	virtual bool OnValidateCameraEvaluationContextActivation() override;
+	virtual void OnUpdateCameraEvaluationContext(bool bForceApplyParameterOverrides) override;
+
+private:
+
+#if WITH_EDITOR
+	void OnCameraRigAssetBuilt(const UCameraRigAsset* InCameraRigAsset);
+#endif
+
+private:
+
+	UPROPERTY()
+	TObjectPtr<UCameraAsset> GeneratedCameraAsset;
+};
+
diff --git a/Engine/Plugins/Cameras/GameplayCameras/Source/GameplayCamerasEditor/Private/ComponentVisualizers/GameplayCameraComponentVisualizer.cpp b/Engine/Plugins/Cameras/GameplayCameras/Source/GameplayCamerasEditor/Private/ComponentVisualizers/GameplayCameraComponentVisualizer.cpp
--- a/Engine/Plugins/Cameras/GameplayCameras/Source/GameplayCamerasEditor/Private/ComponentVisualizers/GameplayCameraComponentVisualizer.cpp
+++ b/Engine/Plugins/Cameras/GameplayCameras/Source/GameplayCamerasEditor/Private/ComponentVisualizers/GameplayCameraComponentVisualizer.cpp
@@ -2,7 +2,7 @@
 
 #include "ComponentVisualizers/GameplayCameraComponentVisualizer.h"
 
-#include "GameFramework/GameplayCameraComponent.h"
+#include "GameFramework/GameplayCameraComponentBase.h"
 
 namespace UE::Cameras
 {
@@ -13,7 +13,7 @@ void FGameplayCameraComponentVisualizer::DrawVisualization(const UActorComponent
 
 void FGameplayCameraComponentVisualizer::DrawVisualizationHUD(const UActorComponent* Component, const FViewport* Viewport, const FSceneView* View, FCanvas* Canvas)
 {
-	if (const UGameplayCameraComponent* GameplayCameraComponent = Cast<const UGameplayCameraComponent>(Component))
+	if (const UGameplayCameraComponentBase* GameplayCameraComponent = Cast<const UGameplayCameraComponentBase>(Component))
 	{
 		GameplayCameraComponent->OnDrawVisualizationHUD(Viewport, View, Canvas);
 	}
diff --git a/Engine/Plugins/Cameras/GameplayCameras/Source/GameplayCamerasEditor/Private/GameplayCamerasEditorModule.cpp b/Engine/Plugins/Cameras/GameplayCameras/Source/GameplayCamerasEditor/Private/GameplayCamerasEditorModule.cpp
--- a/Engine/Plugins/Cameras/GameplayCameras/Source/GameplayCamerasEditor/Private/GameplayCamerasEditorModule.cpp
+++ b/Engine/Plugins/Cameras/GameplayCameras/Source/GameplayCamerasEditor/Private/GameplayCamerasEditorModule.cpp
@@ -559,7 +559,7 @@ private:
 
 		if (GUnrealEd)
 		{
-			GUnrealEd->RegisterComponentVisualizer(UGameplayCameraComponent::StaticClass()->GetFName(), MakeShared<FGameplayCameraComponentVisualizer>());
+			GUnrealEd->RegisterComponentVisualizer(UGameplayCameraComponentBase::StaticClass()->GetFName(), MakeShared<FGameplayCameraComponentVisualizer>());
 		}
 	}
 
@@ -567,7 +567,7 @@ private:
 	{
 		if (GUnrealEd)
 		{
-			GUnrealEd->UnregisterComponentVisualizer(UGameplayCameraComponent::StaticClass()->GetFName());
+			GUnrealEd->UnregisterComponentVisualizer(UGameplayCameraComponentBase::StaticClass()->GetFName());
 		}
 	}
 	
diff --git a/Engine/Plugins/Cameras/GameplayCameras/Source/GameplayCamerasEditor/Private/Sequencer/GameplayCameraComponentTrackEditor.cpp b/Engine/Plugins/Cameras/GameplayCameras/Source/GameplayCamerasEditor/Private/Sequencer/GameplayCameraComponentTrackEditor.cpp
--- a/Engine/Plugins/Cameras/GameplayCameras/Source/GameplayCamerasEditor/Private/Sequencer/GameplayCameraComponentTrackEditor.cpp
+++ b/Engine/Plugins/Cameras/GameplayCameras/Source/GameplayCamerasEditor/Private/Sequencer/GameplayCameraComponentTrackEditor.cpp
@@ -9,6 +9,7 @@
 #include "Framework/MultiBox/MultiBoxBuilder.h"
 #include "GameFramework/Actor.h"
 #include "GameFramework/GameplayCameraComponent.h"
+#include "GameFramework/GameplayCameraRigComponent.h"
 #include "ISequencer.h"
 #include "KeyPropertyParams.h"
 #include "MVVM/ViewModels/ObjectBindingModel.h"
@@ -68,32 +69,49 @@ void GetKeyablePropertyPathsImpl(TSharedPtr<ISequencer> Sequencer, const UClass*
 	}
 }
 
-void GetKeyablePropertyPaths(TSharedPtr<ISequencer> Sequencer, const UGameplayCameraComponent* CameraComponent, TArray<FPropertyPath>& KeyablePropertyPaths)
+void GetKeyablePropertyPaths(TSharedPtr<ISequencer> Sequencer, const UGameplayCameraComponentBase* CameraComponentBase, TArray<FPropertyPath>& KeyablePropertyPaths)
 {
-	const UCameraAsset* CameraAsset = CameraComponent->CameraReference.GetCameraAsset();
-	if (!CameraAsset)
+	// Start us off with the property path of the parameters struct, and then get all keyable property paths from there.
+	FPropertyPath PropertyPath;
+	const UPropertyBag* CameraParametersStruct = nullptr;
+	const uint8* CameraParametersMemory = nullptr;
+
+	if (const UGameplayCameraComponent* CameraComponent = Cast<UGameplayCameraComponent>(CameraComponentBase))
 	{
-		return;
+		const UCameraAsset* CameraAsset = CameraComponent->CameraReference.GetCameraAsset();
+		if (CameraAsset)
+		{
+			const FInstancedPropertyBag& CameraParameters = CameraComponent->CameraReference.GetParameters();
+			CameraParametersStruct = CameraParameters.GetPropertyBagStruct();
+			CameraParametersMemory = CameraParameters.GetValue().GetMemory();
+
+			PropertyPath.AddProperty(FPropertyInfo(UGameplayCameraComponent::StaticClass()->FindPropertyByName(TEXT("CameraReference"))));
+			PropertyPath.AddProperty(FPropertyInfo(FCameraAssetReference::StaticStruct()->FindPropertyByName(TEXT("Parameters"))));
+		}
+	}
+	else if (const UGameplayCameraRigComponent* CameraRigComponent = Cast<UGameplayCameraRigComponent>(CameraComponentBase))
+	{
+		const UCameraRigAsset* CameraRigAsset = CameraRigComponent->CameraRigReference.GetCameraRig();
+		if (CameraRigAsset)
+		{
+			const FInstancedPropertyBag& CameraParameters = CameraRigComponent->CameraRigReference.GetParameters();
+			CameraParametersStruct = CameraParameters.GetPropertyBagStruct();
+			CameraParametersMemory = CameraParameters.GetValue().GetMemory();
+
+			PropertyPath.AddProperty(FPropertyInfo(UGameplayCameraRigComponent::StaticClass()->FindPropertyByName(TEXT("CameraRigReference"))));
+			PropertyPath.AddProperty(FPropertyInfo(FCameraRigAssetReference::StaticStruct()->FindPropertyByName(TEXT("Parameters"))));
+		}
 	}
 
-	const FInstancedPropertyBag& CameraParameters = CameraComponent->CameraReference.GetParameters();
-	const UPropertyBag* CameraParametersStruct = CameraParameters.GetPropertyBagStruct();
-	const uint8* CameraParametersMemory = CameraParameters.GetValue().GetMemory();
 	if (!CameraParametersStruct || !CameraParametersMemory)
 	{
 		return;
 	}
 
-	// Start us off with the property path of the parameters struct.
-	const UClass* ComponentClass = UGameplayCameraComponent::StaticClass();
-	const UStruct* CameraAssetReferenceStruct = FCameraAssetReference::StaticStruct();
 	const UStruct* PropertyBagStruct = FInstancedPropertyBag::StaticStruct();
-
-	FPropertyPath PropertyPath;
-	PropertyPath.AddProperty(FPropertyInfo(ComponentClass->FindPropertyByName(TEXT("CameraReference"))));
-	PropertyPath.AddProperty(FPropertyInfo(CameraAssetReferenceStruct->FindPropertyByName(TEXT("Parameters"))));
 	PropertyPath.AddProperty(FPropertyInfo(PropertyBagStruct->FindPropertyByName(TEXT("Value"))));
 
+	const UClass* ComponentClass = CameraComponentBase->GetClass();
 	GetKeyablePropertyPathsImpl(Sequencer, ComponentClass, CameraParametersStruct, CameraParametersMemory, PropertyPath, KeyablePropertyPaths);
 }
 
@@ -162,7 +180,7 @@ void FGameplayCameraComponentTrackEditor::BuildTrackContextMenu(FMenuBuilder& Me
 
 void FGameplayCameraComponentTrackEditor::ExtendObjectBindingTrackMenu(TSharedRef<FExtender> Extender, const TArray<FGuid>& ObjectBindings, const UClass* ObjectClass)
 {
-	if (ObjectClass && ObjectClass->IsChildOf<UGameplayCameraComponent>())
+	if (ObjectClass && ObjectClass->IsChildOf<UGameplayCameraComponentBase>())
 	{
 		Extender->AddMenuExtension(
 				TEXT("Tracks"), EExtensionHook::After, nullptr, 
@@ -177,7 +195,7 @@ void FGameplayCameraComponentTrackEditor::OnExtendObjectBindingTrackMenu(FMenuBu
 		return;
 	}
 
-	const UGameplayCameraComponent* CameraComponent = GetCameraComponentForBinding(ObjectBindings[0]);
+	const UGameplayCameraComponentBase* CameraComponent = GetCameraComponentForBinding(ObjectBindings[0]);
 	if (CameraComponent)
 	{
 		using namespace UE::Cameras::Internal;
@@ -240,13 +258,13 @@ FReply FGameplayCameraComponentTrackEditor::OnDrop(const FDragDropEvent& DragDro
 	return FReply::Unhandled();
 }
 
-UGameplayCameraComponent* FGameplayCameraComponentTrackEditor::GetCameraComponentForBinding(const FGuid& ObjectBinding) const
+UGameplayCameraComponentBase* FGameplayCameraComponentTrackEditor::GetCameraComponentForBinding(const FGuid& ObjectBinding) const
 {
 	TSharedPtr<ISequencer> SequencerPtr = GetSequencer();
 	if (SequencerPtr.IsValid())
 	{
 		UObject* BoundObject = SequencerPtr->FindSpawnedObjectOrTemplate(ObjectBinding);
-		return Cast<UGameplayCameraComponent>(BoundObject);
+		return Cast<UGameplayCameraComponentBase>(BoundObject);
 	}
 	return nullptr;
 }
diff --git a/Engine/Plugins/Cameras/GameplayCameras/Source/GameplayCamerasEditor/Private/Sequencer/GameplayCameraComponentTrackEditor.h b/Engine/Plugins/Cameras/GameplayCameras/Source/GameplayCamerasEditor/Private/Sequencer/GameplayCameraComponentTrackEditor.h
--- a/Engine/Plugins/Cameras/GameplayCameras/Source/GameplayCamerasEditor/Private/Sequencer/GameplayCameraComponentTrackEditor.h
+++ b/Engine/Plugins/Cameras/GameplayCameras/Source/GameplayCamerasEditor/Private/Sequencer/GameplayCameraComponentTrackEditor.h
@@ -7,7 +7,7 @@
 #include "Templates/SharedPointerFwd.h"
 
 class ISequencer;
-class UGameplayCameraComponent;
+class UGameplayCameraComponentBase;
 class UMovieScenePropertyTrack;
 struct FCameraObjectInterfaceParameterDefinition;
 
@@ -44,7 +44,7 @@ protected:
 
 private:
 
-	UGameplayCameraComponent* GetCameraComponentForBinding(const FGuid& ObjectBinding) const;
+	UGameplayCameraComponentBase* GetCameraComponentForBinding(const FGuid& ObjectBinding) const;
 
 	void OnExtendObjectBindingTrackMenu(FMenuBuilder& MenuBuilder, TArray<FGuid> ObjectBindings);
 
