Cameras: fix uninitialized variable warning

diff --git a/Engine/Plugins/Cameras/GameplayCameras/Source/GameplayCameras/Private/MovieScene/MovieSceneCameraParameterInstantiator.cpp b/Engine/Plugins/Cameras/GameplayCameras/Source/GameplayCameras/Private/MovieScene/MovieSceneCameraParameterInstantiator.cpp
--- a/Engine/Plugins/Cameras/GameplayCameras/Source/GameplayCameras/Private/MovieScene/MovieSceneCameraParameterInstantiator.cpp
+++ b/Engine/Plugins/Cameras/GameplayCameras/Source/GameplayCameras/Private/MovieScene/MovieSceneCameraParameterInstantiator.cpp
@@ -94,7 +94,7 @@ struct FSetupCameraParameterOverrideTask
 		}
 
 		FGuid ParameterGuid;
-		bool bWasAnimated;
+		bool bWasAnimated = false;
 
 		if (UGameplayCameraComponent* CameraComponent = Cast<UGameplayCameraComponent>(CameraComponentBase))
 		{
