GameplayCameras: Add support for orthographic projection mode

diff --git a/Engine/Plugins/Cameras/GameplayCameras/Source/GameplayCameras/Private/Core/CameraSystemEvaluator.cpp b/Engine/Plugins/Cameras/GameplayCameras/Source/GameplayCameras/Private/Core/CameraSystemEvaluator.cpp
--- a/Engine/Plugins/Cameras/GameplayCameras/Source/GameplayCameras/Private/Core/CameraSystemEvaluator.cpp
+++ b/Engine/Plugins/Cameras/GameplayCameras/Source/GameplayCameras/Private/Core/CameraSystemEvaluator.cpp
@@ -361,7 +361,12 @@ void FCameraSystemEvaluator::GetEvaluatedCameraView(FMinimalViewInfo& DesiredVie
 	DesiredView.AspectRatioAxisConstraint = CameraPose.GetOverrideAspectRatioAxisConstraint() ?
 		CameraPose.GetAspectRatioAxisConstraint() : TOptional<EAspectRatioAxisConstraint>();
 
-	// TODO: add support for ortho cameras.
+	DesiredView.ProjectionMode = CameraPose.GetProjectionMode();
+	if (CameraPose.GetProjectionMode() == ECameraProjectionMode::Orthographic)
+	{
+		DesiredView.OrthoWidth = CameraPose.GetOrthographicWidth();
+	}
+
 	DesiredView.PerspectiveNearClipPlane = CameraPose.GetNearClippingPlane();
 
 	const FPostProcessSettingsCollection& PostProcessSettings = RootNodeResult.PostProcessSettings;
diff --git a/Engine/Plugins/Cameras/GameplayCameras/Source/GameplayCameras/Public/Core/CameraPose.h b/Engine/Plugins/Cameras/GameplayCameras/Source/GameplayCameras/Public/Core/CameraPose.h
--- a/Engine/Plugins/Cameras/GameplayCameras/Source/GameplayCameras/Public/Core/CameraPose.h
+++ b/Engine/Plugins/Cameras/GameplayCameras/Source/GameplayCameras/Public/Core/CameraPose.h
@@ -4,6 +4,7 @@
 
 #include "CoreTypes.h"
 #include "Engine/EngineTypes.h"
+#include "EngineDefines.h"
 #include "Math/MathFwd.h"
 #include "Math/Transform.h"
 
@@ -12,6 +13,7 @@
 class FArchive;
 struct FPostProcessSettings;
 enum EAspectRatioAxisConstraint : int;
+namespace ECameraProjectionMode { enum Type : int; };
 
 #define UE_CAMERA_POSE_FOR_TRANSFORM_PROPERTIES()\
 	UE_CAMERA_POSE_FOR_PROPERTY(FVector, Location)\
@@ -29,7 +31,8 @@ enum EAspectRatioAxisConstraint : int;
 	UE_CAMERA_POSE_FOR_PROPERTY(int32,  DiaphragmBladeCount)\
 	UE_CAMERA_POSE_FOR_PROPERTY(float,  PhysicalCameraBlendWeight)\
 	UE_CAMERA_POSE_FOR_PROPERTY(float,  NearClippingPlane)\
-	UE_CAMERA_POSE_FOR_PROPERTY(float,  FarClippingPlane)
+	UE_CAMERA_POSE_FOR_PROPERTY(float,  FarClippingPlane)\
+	UE_CAMERA_POSE_FOR_PROPERTY(float,  OrthographicWidth)
 
 #define UE_CAMERA_POSE_FOR_FOV_PROPERTIES()\
 	UE_CAMERA_POSE_FOR_PROPERTY(float, FieldOfView)\
@@ -39,7 +42,8 @@ enum EAspectRatioAxisConstraint : int;
 	UE_CAMERA_POSE_FOR_PROPERTY(bool, EnablePhysicalCamera)\
 	UE_CAMERA_POSE_FOR_PROPERTY(bool, ConstrainAspectRatio)\
 	UE_CAMERA_POSE_FOR_PROPERTY(bool, OverrideAspectRatioAxisConstraint)\
-	UE_CAMERA_POSE_FOR_PROPERTY(EAspectRatioAxisConstraint, AspectRatioAxisConstraint)
+	UE_CAMERA_POSE_FOR_PROPERTY(EAspectRatioAxisConstraint, AspectRatioAxisConstraint)\
+	UE_CAMERA_POSE_FOR_PROPERTY(ECameraProjectionMode::Type, ProjectionMode)
 
 #define UE_CAMERA_POSE_FOR_ALL_PROPERTIES()\
 	UE_CAMERA_POSE_FOR_TRANSFORM_PROPERTIES()\
@@ -248,6 +252,12 @@ private:
 	UPROPERTY()
 	float FocalLength = 35.f;
 
+	/**
+	 * The desired width (in world units) of the orthographic view (ignored in Perspective mode) 
+	 */
+	UPROPERTY()
+	float OrthographicWidth = DEFAULT_ORTHOWIDTH;
+
 	/** The aperture of the camera's lens, in f-stops */
 	UPROPERTY()
 	float Aperture = 2.8f;
@@ -314,6 +324,10 @@ private:
 	UPROPERTY()
 	TEnumAsByte<EAspectRatioAxisConstraint> AspectRatioAxisConstraint = EAspectRatioAxisConstraint::AspectRatio_MaintainYFOV;
 
+	/** Projection mode - Perspective or orthographic */
+	UPROPERTY()
+	TEnumAsByte<ECameraProjectionMode::Type> ProjectionMode;
+
 private:
 	
 	/**
diff --git a/Engine/Plugins/Cameras/GameplayCameras/Source/GameplayCameras/Public/Debug/DebugTextRenderer.h b/Engine/Plugins/Cameras/GameplayCameras/Source/GameplayCameras/Public/Debug/DebugTextRenderer.h
--- a/Engine/Plugins/Cameras/GameplayCameras/Source/GameplayCameras/Public/Debug/DebugTextRenderer.h
+++ b/Engine/Plugins/Cameras/GameplayCameras/Source/GameplayCameras/Public/Debug/DebugTextRenderer.h
@@ -3,6 +3,7 @@
 #pragma once
 
 #include "CoreTypes.h"
+#include "Camera/CameraTypes.h"
 #include "Containers/StringView.h"
 #include "Engine/EngineTypes.h"
 #include "GameplayCameras.h"
@@ -59,6 +60,16 @@ inline FString ToDebugString(const EAspectRatioAxisConstraint& FieldValue)
 	}
 	return TEXT("Invalid");
 }
+template<>
+inline FString ToDebugString(const ECameraProjectionMode::Type& FieldValue)
+{
+	switch (FieldValue)
+	{
+	case ECameraProjectionMode::Perspective: return TEXT("Perspective");
+	case ECameraProjectionMode::Orthographic: return TEXT("Orthographic");
+	}
+	return TEXT("Invalid");
+}
 
 /** Command for drawing text on a canvas. */
 struct FDebugTextDrawCommand
diff --git a/Engine/Plugins/Cameras/GameplayCameras/Source/GameplayCameras/Public/GameFramework/BlueprintCameraPose.h b/Engine/Plugins/Cameras/GameplayCameras/Source/GameplayCameras/Public/GameFramework/BlueprintCameraPose.h
--- a/Engine/Plugins/Cameras/GameplayCameras/Source/GameplayCameras/Public/GameFramework/BlueprintCameraPose.h
+++ b/Engine/Plugins/Cameras/GameplayCameras/Source/GameplayCameras/Public/GameFramework/BlueprintCameraPose.h
@@ -48,6 +48,10 @@ public:
 	UPROPERTY(BlueprintReadWrite, EditAnywhere, Category="Camera")
 	float FocalLength = 35.f;
 
+	/** The desired width (in world units) of the orthographic view (ignored in Perspective mode) */
+	UPROPERTY(BlueprintReadWrite, EditAnywhere, Category = "Camera")
+	float OrthographicWidth = DEFAULT_ORTHOWIDTH;
+
 	/** The aperture (f-stop) of the camera's lens. */
 	UPROPERTY(BlueprintReadWrite, EditAnywhere, Category="Camera")
 	float Aperture = 2.8f;
@@ -111,6 +115,10 @@ public:
 	UPROPERTY(BlueprintReadWrite, EditAnywhere, Category="Camera")
 	TEnumAsByte<EAspectRatioAxisConstraint> AspectRatioAxisConstraint = EAspectRatioAxisConstraint::AspectRatio_MaintainYFOV;
 
+	/** Projection mode - Perspective or orthographic */
+	UPROPERTY(BlueprintReadWrite, EditAnywhere, Category = "Camera")
+	TEnumAsByte<ECameraProjectionMode::Type> ProjectionMode;
+
 public:
 
 	/** Creates a new BP camera pose from a raw camera pose structure. */
