Use FSlimHorizontalToolBarBuilder instead of FToolBarBuilder in any places in the engine that customize the curve editor toolbar.
This fixes missing textures in the toolbar and also makes them slim.

diff --git a/Engine/Plugins/Cameras/GameplayCameras/Source/GameplayCamerasEditor/Private/Customizations/RichCurveDetailsCustomizations.cpp b/Engine/Plugins/Cameras/GameplayCameras/Source/GameplayCamerasEditor/Private/Customizations/RichCurveDetailsCustomizations.cpp
--- a/Engine/Plugins/Cameras/GameplayCameras/Source/GameplayCamerasEditor/Private/Customizations/RichCurveDetailsCustomizations.cpp
+++ b/Engine/Plugins/Cameras/GameplayCameras/Source/GameplayCamerasEditor/Private/Customizations/RichCurveDetailsCustomizations.cpp
@@ -639,7 +639,7 @@ TSharedRef<SWidget> FRichCurveDetailsCustomization::CreateCurveEditorPanel(TShar
 	TSharedPtr<FUICommandList> Commands = CurveEditorPanel->GetCommands();
 	TSharedPtr<FExtender> ToolbarExtender = CurveEditorPanel->GetToolbarExtender();
 
-	FToolBarBuilder ToolBarBuilder(Commands, FMultiBoxCustomization::None, ToolbarExtender, true);
+	FSlimHorizontalToolBarBuilder ToolBarBuilder(Commands, FMultiBoxCustomization::None, ToolbarExtender, true);
 	ToolBarBuilder.BeginSection("Asset");
 	ToolBarBuilder.EndSection();
 	TSharedRef<SWidget> ToolBarWidget = ToolBarBuilder.MakeWidget();
diff --git a/Engine/Plugins/Cameras/GameplayCameras/Source/GameplayCamerasEditor/Private/Toolkits/CurveEditorToolkit.cpp b/Engine/Plugins/Cameras/GameplayCameras/Source/GameplayCamerasEditor/Private/Toolkits/CurveEditorToolkit.cpp
--- a/Engine/Plugins/Cameras/GameplayCameras/Source/GameplayCamerasEditor/Private/Toolkits/CurveEditorToolkit.cpp
+++ b/Engine/Plugins/Cameras/GameplayCameras/Source/GameplayCamerasEditor/Private/Toolkits/CurveEditorToolkit.cpp
@@ -51,7 +51,7 @@ void FCurveEditorToolkit::Initialize()
 	TSharedPtr<FUICommandList> Commands = CurveEditorPanel->GetCommands();
 	TSharedPtr<FExtender> ToolbarExtender = CurveEditorPanel->GetToolbarExtender();
 
-	FToolBarBuilder ToolBarBuilder(Commands, FMultiBoxCustomization::None, ToolbarExtender, true);
+	FSlimHorizontalToolBarBuilder ToolBarBuilder(Commands, FMultiBoxCustomization::None, ToolbarExtender, true);
 	ToolBarBuilder.BeginSection("Asset");
 	ToolBarBuilder.EndSection();
 	TSharedRef<SWidget> ToolBarWidget = ToolBarBuilder.MakeWidget();
diff --git a/Engine/Source/Editor/CurveTableEditor/Private/CurveTableEditor.cpp b/Engine/Source/Editor/CurveTableEditor/Private/CurveTableEditor.cpp
--- a/Engine/Source/Editor/CurveTableEditor/Private/CurveTableEditor.cpp
+++ b/Engine/Source/Editor/CurveTableEditor/Private/CurveTableEditor.cpp
@@ -992,7 +992,7 @@ bool FCurveTableEditor::IsCurveViewChecked() const
 TSharedRef<SWidget> FCurveTableEditor::MakeToolbar(TSharedRef<SCurveEditorPanel>& InEditorPanel)
 {
 
-	FToolBarBuilder ToolBarBuilder(InEditorPanel->GetCommands(), FMultiBoxCustomization::None, InEditorPanel->GetToolbarExtender(), true);
+	FSlimHorizontalToolBarBuilder ToolBarBuilder(InEditorPanel->GetCommands(), FMultiBoxCustomization::None, InEditorPanel->GetToolbarExtender(), true);
 	ToolBarBuilder.BeginSection("Asset");
 	ToolBarBuilder.EndSection();
 	// We just use all of the extenders as our toolbar, we don't have a need to create a separate toolbar.
diff --git a/Engine/Source/Editor/Persona/Private/SAnimSequenceCurveEditor.cpp b/Engine/Source/Editor/Persona/Private/SAnimSequenceCurveEditor.cpp
--- a/Engine/Source/Editor/Persona/Private/SAnimSequenceCurveEditor.cpp
+++ b/Engine/Source/Editor/Persona/Private/SAnimSequenceCurveEditor.cpp
@@ -487,7 +487,7 @@ void SAnimSequenceCurveEditor::OnModelHasChanged(const EAnimDataModelNotifyType&
 
 TSharedRef<SWidget> SAnimSequenceCurveEditor::MakeToolbar(TSharedRef<SCurveEditorPanel> InEditorPanel)
 {
-	FToolBarBuilder ToolBarBuilder(InEditorPanel->GetCommands(), FMultiBoxCustomization::None, InEditorPanel->GetToolbarExtender(), true);
+	FSlimHorizontalToolBarBuilder ToolBarBuilder(InEditorPanel->GetCommands(), FMultiBoxCustomization::None, InEditorPanel->GetToolbarExtender(), true);
 	ToolBarBuilder.BeginSection("Asset");
 	ToolBarBuilder.EndSection();
 	// We just use all of the extenders as our toolbar, we don't have a need to create a separate toolbar.
