Cameras: fix crash when trying to setup editor preview for a camera component that has no asset defined

diff --git a/Engine/Plugins/Cameras/GameplayCameras/Source/GameplayCameras/Private/GameFramework/GameplayCameraComponentBase.cpp b/Engine/Plugins/Cameras/GameplayCameras/Source/GameplayCameras/Private/GameFramework/GameplayCameraComponentBase.cpp
--- a/Engine/Plugins/Cameras/GameplayCameras/Source/GameplayCameras/Private/GameFramework/GameplayCameraComponentBase.cpp
+++ b/Engine/Plugins/Cameras/GameplayCameras/Source/GameplayCameras/Private/GameFramework/GameplayCameraComponentBase.cpp
@@ -175,7 +175,22 @@ void UGameplayCameraComponentBase::ActivateCameraEvaluationContext(APlayerContro
 
 	AGameplayCameraSystemActor::AutoManageActiveViewTarget(PlayerController);
 
-	EnsureCameraEvaluationContextCreated(PlayerController);
+	// Make sure the evaluation context has been created. However, this can fail, such as when
+	// we don't have a valid camera asset specified.
+	TryCreateCameraEvaluationContext(PlayerController);
+
+	if (!EvaluationContext.IsValid())
+	{
+		return;
+	}
+
+	if (EvaluationContext->IsActive())
+	{
+		FFrame::KismetExecutionMessage(
+				TEXT("Can't activate gameplay camera component: it is already active!"),
+				ELogVerbosity::Error);
+		return;
+	}
 
 	TSharedPtr<FCameraSystemEvaluator> CameraSystemEvaluator = CameraSystemHost->GetCameraSystemEvaluator();
 	CameraSystemEvaluator->PushEvaluationContext(EvaluationContext.ToSharedRef());
@@ -184,17 +199,13 @@ void UGameplayCameraComponentBase::ActivateCameraEvaluationContext(APlayerContro
 	Activate();
 }
 
-void UGameplayCameraComponentBase::EnsureCameraEvaluationContextCreated(APlayerController* PlayerController)
+void UGameplayCameraComponentBase::TryCreateCameraEvaluationContext(APlayerController* PlayerController)
 {
 	using namespace UE::Cameras;
 
 	if (!EvaluationContext.IsValid())
 	{
-		const UCameraAsset* CameraAsset = GetCameraAsset();
-		if (ensureMsgf(
-				CameraAsset,
-				TEXT("GameplayCameraComponent sub-class '%s' didn't provide us with a valid camera asset"),
-				*GetNameSafe(GetClass())))
+		if (const UCameraAsset* CameraAsset = GetCameraAsset())
 		{
 			EvaluationContext = MakeShared<FGameplayCameraComponentEvaluationContext>();
 
@@ -483,9 +494,12 @@ void UGameplayCameraComponentBase::AutoManageEditorPreviewEvaluator()
 		}
 		if (!EvaluationContext)
 		{
-			EnsureCameraEvaluationContextCreated(nullptr);
-			EditorPreviewEvaluator->PushEvaluationContext(EvaluationContext.ToSharedRef());
-			EvaluationContext->SetEditorPreviewCameraRigIndex(EditorPreviewCameraRigIndex);
+			TryCreateCameraEvaluationContext(nullptr);
+			if (EvaluationContext.IsValid())
+			{
+				EditorPreviewEvaluator->PushEvaluationContext(EvaluationContext.ToSharedRef());
+				EvaluationContext->SetEditorPreviewCameraRigIndex(EditorPreviewCameraRigIndex);
+			}
 		}
 	}
 	else if (!bRunInEditor && (EditorPreviewEvaluator || EvaluationContext))
@@ -525,8 +539,12 @@ void UGameplayCameraComponentBase::RecreateEditorPreviewEvaluationContext()
 		EditorPreviewEvaluator->RemoveEvaluationContext(EvaluationContext.ToSharedRef());
 		EvaluationContext = nullptr;
 
-		EnsureCameraEvaluationContextCreated(nullptr);
-		EditorPreviewEvaluator->PushEvaluationContext(EvaluationContext.ToSharedRef());
+		TryCreateCameraEvaluationContext(nullptr);
+		if (EvaluationContext.IsValid())
+		{
+			EditorPreviewEvaluator->PushEvaluationContext(EvaluationContext.ToSharedRef());
+			EvaluationContext->SetEditorPreviewCameraRigIndex(EditorPreviewCameraRigIndex);
+		}
 	}
 }
 
diff --git a/Engine/Plugins/Cameras/GameplayCameras/Source/GameplayCameras/Public/Core/CameraEvaluationContext.h b/Engine/Plugins/Cameras/GameplayCameras/Source/GameplayCameras/Public/Core/CameraEvaluationContext.h
--- a/Engine/Plugins/Cameras/GameplayCameras/Source/GameplayCameras/Public/Core/CameraEvaluationContext.h
+++ b/Engine/Plugins/Cameras/GameplayCameras/Source/GameplayCameras/Public/Core/CameraEvaluationContext.h
@@ -133,6 +133,9 @@ public:
 	 */
 	void Deactivate(const FCameraEvaluationContextDeactivateParams& Params);
 
+	/** Whether this context is active. */
+	bool IsActive() const { return bActivated; }
+
 public:
 
 	// Internal API.
diff --git a/Engine/Plugins/Cameras/GameplayCameras/Source/GameplayCameras/Public/GameFramework/GameplayCameraComponentBase.h b/Engine/Plugins/Cameras/GameplayCameras/Source/GameplayCameras/Public/GameFramework/GameplayCameraComponentBase.h
--- a/Engine/Plugins/Cameras/GameplayCameras/Source/GameplayCameras/Public/GameFramework/GameplayCameraComponentBase.h
+++ b/Engine/Plugins/Cameras/GameplayCameras/Source/GameplayCameras/Public/GameFramework/GameplayCameraComponentBase.h
@@ -140,7 +140,7 @@ private:
 
 	void ActivateCameraEvaluationContext(int32 PlayerIndex);
 	void ActivateCameraEvaluationContext(APlayerController* PlayerController);
-	void EnsureCameraEvaluationContextCreated(APlayerController* PlayerController);
+	void TryCreateCameraEvaluationContext(APlayerController* PlayerController);
 	void UpdateOutputCameraComponent();
 	void DeactivateCameraEvaluationContext();
 
