Cameras: don't create the editor preview evaluator multiple times

diff --git a/Engine/Plugins/Cameras/GameplayCameras/Source/GameplayCameras/Private/Core/CameraEvaluationContext.cpp b/Engine/Plugins/Cameras/GameplayCameras/Source/GameplayCameras/Private/Core/CameraEvaluationContext.cpp
--- a/Engine/Plugins/Cameras/GameplayCameras/Source/GameplayCameras/Private/Core/CameraEvaluationContext.cpp
+++ b/Engine/Plugins/Cameras/GameplayCameras/Source/GameplayCameras/Private/Core/CameraEvaluationContext.cpp
@@ -200,7 +200,9 @@ void FCameraEvaluationContext::AutoCreateDirectorEvaluator()
 
 void FCameraEvaluationContext::AutoCreateEditorPreviewDirectorEvaluator(const FCameraEvaluationContextActivateParams& Params)
 {
-	if (ensure(Params.Evaluator) && Params.Evaluator->GetRole() == ECameraSystemEvaluatorRole::EditorPreview)
+	if (DirectorEvaluator == nullptr &&
+			ensure(Params.Evaluator) && 
+			Params.Evaluator->GetRole() == ECameraSystemEvaluatorRole::EditorPreview)
 	{
 		FCameraDirectorRigUsageInfo UsageInfo;
 
