Cameras: don't run the attach to player pawn node in editor preview, where there's no player pawn

diff --git a/Engine/Plugins/Cameras/GameplayCameras/Source/GameplayCameras/Private/Nodes/Attach/AttachToPlayerPawnCameraNode.cpp b/Engine/Plugins/Cameras/GameplayCameras/Source/GameplayCameras/Private/Nodes/Attach/AttachToPlayerPawnCameraNode.cpp
--- a/Engine/Plugins/Cameras/GameplayCameras/Source/GameplayCameras/Private/Nodes/Attach/AttachToPlayerPawnCameraNode.cpp
+++ b/Engine/Plugins/Cameras/GameplayCameras/Source/GameplayCameras/Private/Nodes/Attach/AttachToPlayerPawnCameraNode.cpp
@@ -4,6 +4,7 @@
 
 #include "Core/CameraEvaluationContext.h"
 #include "Core/CameraParameterReader.h"
+#include "Core/CameraSystemEvaluator.h"
 #include "GameFramework/Pawn.h"
 #include "GameFramework/PlayerController.h"
 #include "GameplayCameras.h"
@@ -39,6 +40,14 @@ void FAttachToPlayerPawnCameraNodeEvaluator::OnInitialize(const FCameraNodeEvalu
 
 void FAttachToPlayerPawnCameraNodeEvaluator::OnRun(const FCameraNodeEvaluationParams& Params, FCameraNodeEvaluationResult& OutResult)
 {
+#if WITH_EDITOR
+	if (Params.Evaluator && Params.Evaluator->GetRole() == ECameraSystemEvaluatorRole::EditorPreview)
+	{
+		// No player pawn in editor preview.
+		return;
+	}
+#endif
+
 	APlayerController* PlayerController = Params.EvaluationContext->GetPlayerController();
 	if (!ensure(PlayerController))
 	{
