Cameras: add set location/rotation node

diff --git a/Engine/Plugins/Cameras/GameplayCameras/Source/GameplayCameras/Private/Nodes/Common/SetLocationCameraNode.cpp b/Engine/Plugins/Cameras/GameplayCameras/Source/GameplayCameras/Private/Nodes/Common/SetLocationCameraNode.cpp
new file mode 100644
--- /dev/null
+++ b/Engine/Plugins/Cameras/GameplayCameras/Source/GameplayCameras/Private/Nodes/Common/SetLocationCameraNode.cpp
@@ -0,0 +1,88 @@
+// Copyright Epic Games, Inc. All Rights Reserved.
+
+#include "Nodes/Common/SetLocationCameraNode.h"
+
+#include "Core/CameraEvaluationContext.h"
+#include "Core/CameraParameterReader.h"
+#include "GameplayCameras.h"
+
+#include UE_INLINE_GENERATED_CPP_BY_NAME(SetLocationCameraNode)
+
+namespace UE::Cameras
+{
+
+class FSetLocationCameraNodeEvaluator : public FCameraNodeEvaluator
+{
+	UE_DECLARE_CAMERA_NODE_EVALUATOR(GAMEPLAYCAMERAS_API, FSetLocationCameraNodeEvaluator)
+
+protected:
+
+	virtual void OnInitialize(const FCameraNodeEvaluatorInitializeParams& Params, FCameraNodeEvaluationResult& OutResult) override;
+	virtual void OnRun(const FCameraNodeEvaluationParams& Params, FCameraNodeEvaluationResult& OutResult) override;
+
+private:
+
+	TCameraParameterReader<FVector3d> LocationReader;
+};
+
+UE_DEFINE_CAMERA_NODE_EVALUATOR(FSetLocationCameraNodeEvaluator)
+
+void FSetLocationCameraNodeEvaluator::OnInitialize(const FCameraNodeEvaluatorInitializeParams& Params, FCameraNodeEvaluationResult& OutResult)
+{
+	const USetLocationCameraNode* LocationNode = GetCameraNodeAs<USetLocationCameraNode>();
+	LocationReader.Initialize(LocationNode->Location);
+}
+
+void FSetLocationCameraNodeEvaluator::OnRun(const FCameraNodeEvaluationParams& Params, FCameraNodeEvaluationResult& OutResult)
+{
+	const FVector3d NewLocation = LocationReader.Get(OutResult.VariableTable);
+
+	const USetLocationCameraNode* LocationNode = GetCameraNodeAs<USetLocationCameraNode>();
+	switch (LocationNode->OffsetSpace)
+	{
+		case ECameraNodeSpace::CameraPose:
+		default:
+			{
+				FTransform3d Transform = OutResult.CameraPose.GetTransform();
+				Transform = FTransform3d(NewLocation) * Transform;
+				OutResult.CameraPose.SetTransform(Transform);
+			}
+			break;
+		case ECameraNodeSpace::OwningContext:
+			if (Params.EvaluationContext)
+			{ 
+				const FCameraNodeEvaluationResult& InitialResult = Params.EvaluationContext->GetInitialResult();
+				const FTransform3d ContextTransform = InitialResult.CameraPose.GetTransform();
+
+				const FVector3d WorldLocation = ContextTransform.TransformVector(NewLocation);
+
+				FTransform3d Transform = OutResult.CameraPose.GetTransform();
+				Transform.SetLocation(WorldLocation);
+				OutResult.CameraPose.SetTransform(Transform);
+			}
+			else
+			{
+				UE_LOG(LogCameraSystem, Error, 
+						TEXT("SetLocationCameraNode: cannot offset in context space when there is "
+							 "no current context set."));
+				return;
+			}
+			break;
+		case ECameraNodeSpace::World:
+			{
+				FTransform3d Transform = OutResult.CameraPose.GetTransform();
+				Transform.SetLocation(NewLocation);
+				OutResult.CameraPose.SetTransform(Transform);
+			}
+			break;
+	}
+}
+
+}  // namespace UE::Cameras
+
+FCameraNodeEvaluatorPtr USetLocationCameraNode::OnBuildEvaluator(FCameraNodeEvaluatorBuilder& Builder) const
+{
+	using namespace UE::Cameras;
+	return Builder.BuildEvaluator<FSetLocationCameraNodeEvaluator>();
+}
+
diff --git a/Engine/Plugins/Cameras/GameplayCameras/Source/GameplayCameras/Private/Nodes/Common/SetRotationCameraNode.cpp b/Engine/Plugins/Cameras/GameplayCameras/Source/GameplayCameras/Private/Nodes/Common/SetRotationCameraNode.cpp
new file mode 100644
--- /dev/null
+++ b/Engine/Plugins/Cameras/GameplayCameras/Source/GameplayCameras/Private/Nodes/Common/SetRotationCameraNode.cpp
@@ -0,0 +1,88 @@
+// Copyright Epic Games, Inc. All Rights Reserved.
+
+#include "Nodes/Common/SetRotationCameraNode.h"
+
+#include "Core/CameraEvaluationContext.h"
+#include "Core/CameraParameterReader.h"
+#include "GameplayCameras.h"
+
+#include UE_INLINE_GENERATED_CPP_BY_NAME(SetRotationCameraNode)
+
+namespace UE::Cameras
+{
+
+class FSetRotationCameraNodeEvaluator : public FCameraNodeEvaluator
+{
+	UE_DECLARE_CAMERA_NODE_EVALUATOR(GAMEPLAYCAMERAS_API, FSetRotationCameraNodeEvaluator)
+
+protected:
+
+	virtual void OnInitialize(const FCameraNodeEvaluatorInitializeParams& Params, FCameraNodeEvaluationResult& OutResult) override;
+	virtual void OnRun(const FCameraNodeEvaluationParams& Params, FCameraNodeEvaluationResult& OutResult) override;
+
+private:
+
+	TCameraParameterReader<FRotator3d> RotationReader;
+};
+
+UE_DEFINE_CAMERA_NODE_EVALUATOR(FSetRotationCameraNodeEvaluator)
+
+void FSetRotationCameraNodeEvaluator::OnInitialize(const FCameraNodeEvaluatorInitializeParams& Params, FCameraNodeEvaluationResult& OutResult)
+{
+	const USetRotationCameraNode* RotationNode = GetCameraNodeAs<USetRotationCameraNode>();
+	RotationReader.Initialize(RotationNode->Rotation);
+}
+
+void FSetRotationCameraNodeEvaluator::OnRun(const FCameraNodeEvaluationParams& Params, FCameraNodeEvaluationResult& OutResult)
+{
+	const FRotator3d NewRotation = RotationReader.Get(OutResult.VariableTable);
+
+	const USetRotationCameraNode* RotationNode = GetCameraNodeAs<USetRotationCameraNode>();
+	switch (RotationNode->OffsetSpace)
+	{
+		case ECameraNodeSpace::CameraPose:
+		default:
+			{
+				FTransform3d Transform = OutResult.CameraPose.GetTransform();
+				Transform = FTransform3d(NewRotation) * Transform;
+				OutResult.CameraPose.SetTransform(Transform);
+			}
+			break;
+		case ECameraNodeSpace::OwningContext:
+			if (Params.EvaluationContext)
+			{ 
+				const FCameraNodeEvaluationResult& InitialResult = Params.EvaluationContext->GetInitialResult();
+				const FTransform3d ContextTransform = InitialResult.CameraPose.GetTransform();
+
+				const FQuat4d WorldRotation = ContextTransform.TransformRotation(NewRotation.Quaternion());
+
+				FTransform3d Transform = OutResult.CameraPose.GetTransform();
+				Transform.SetRotation(WorldRotation);
+				OutResult.CameraPose.SetTransform(Transform);
+			}
+			else
+			{
+				UE_LOG(LogCameraSystem, Error, 
+						TEXT("SetRotationCameraNode: cannot offset in context space when there is "
+							 "no current context set."));
+				return;
+			}
+			break;
+		case ECameraNodeSpace::World:
+			{
+				FTransform3d Transform = OutResult.CameraPose.GetTransform();
+				Transform.SetRotation(NewRotation.Quaternion());
+				OutResult.CameraPose.SetTransform(Transform);
+			}
+			break;
+	}
+}
+
+}  // namespace UE::Cameras
+
+FCameraNodeEvaluatorPtr USetRotationCameraNode::OnBuildEvaluator(FCameraNodeEvaluatorBuilder& Builder) const
+{
+	using namespace UE::Cameras;
+	return Builder.BuildEvaluator<FSetRotationCameraNodeEvaluator>();
+}
+
diff --git a/Engine/Plugins/Cameras/GameplayCameras/Source/GameplayCameras/Public/Nodes/Common/SetLocationCameraNode.h b/Engine/Plugins/Cameras/GameplayCameras/Source/GameplayCameras/Public/Nodes/Common/SetLocationCameraNode.h
new file mode 100644
--- /dev/null
+++ b/Engine/Plugins/Cameras/GameplayCameras/Source/GameplayCameras/Public/Nodes/Common/SetLocationCameraNode.h
@@ -0,0 +1,34 @@
+// Copyright Epic Games, Inc. All Rights Reserved.
+
+#pragma once
+
+#include "Core/CameraNode.h"
+#include "Core/CameraParameters.h"
+#include "Nodes/CameraNodeTypes.h"
+
+#include "SetLocationCameraNode.generated.h"
+
+/**
+ * A camera node that sets the location of the camera.
+ */
+UCLASS(MinimalAPI, meta=(CameraNodeCategories="Common,Transform"))
+class USetLocationCameraNode : public UCameraNode
+{
+	GENERATED_BODY()
+
+protected:
+
+	// UCameraNode interface.
+	virtual FCameraNodeEvaluatorPtr OnBuildEvaluator(FCameraNodeEvaluatorBuilder& Builder) const override;
+
+public:
+
+	/** The location to set on the camera. */
+	UPROPERTY(EditAnywhere, Category=Common)
+	FVector3dCameraParameter Location;
+
+	/** The space in which to apply the transform. */
+	UPROPERTY(EditAnywhere, Category=Common)
+	ECameraNodeSpace OffsetSpace = ECameraNodeSpace::OwningContext;
+};
+
diff --git a/Engine/Plugins/Cameras/GameplayCameras/Source/GameplayCameras/Public/Nodes/Common/SetRotationCameraNode.h b/Engine/Plugins/Cameras/GameplayCameras/Source/GameplayCameras/Public/Nodes/Common/SetRotationCameraNode.h
new file mode 100644
--- /dev/null
+++ b/Engine/Plugins/Cameras/GameplayCameras/Source/GameplayCameras/Public/Nodes/Common/SetRotationCameraNode.h
@@ -0,0 +1,34 @@
+// Copyright Epic Games, Inc. All Rights Reserved.
+
+#pragma once
+
+#include "Core/CameraNode.h"
+#include "Core/CameraParameters.h"
+#include "Nodes/CameraNodeTypes.h"
+
+#include "SetRotationCameraNode.generated.h"
+
+/**
+ * A camera node that sets the rotation of the camera.
+ */
+UCLASS(MinimalAPI, meta=(CameraNodeCategories="Common,Transform"))
+class USetRotationCameraNode : public UCameraNode
+{
+	GENERATED_BODY()
+
+protected:
+
+	// UCameraNode interface.
+	virtual FCameraNodeEvaluatorPtr OnBuildEvaluator(FCameraNodeEvaluatorBuilder& Builder) const override;
+
+public:
+
+	/** The rotation to set on the camera. */
+	UPROPERTY(EditAnywhere, Category=Common)
+	FRotator3dCameraParameter Rotation;
+
+	/** The space in which to apply the transform. */
+	UPROPERTY(EditAnywhere, Category=Common)
+	ECameraNodeSpace OffsetSpace = ECameraNodeSpace::CameraPose;
+};
+
