Cameras: add comment nodes

- Comment nodes are special objects that map to a subclass of UEdGraphNode_Comment
- This addition required changing some of the graph editor code that was written specifically for UObjectTreeGraphNode only
- There's also a new static registration API for object-tree graph editors, so that creating a comment node can look at any selected nodes and create a comment around them
- Fixed the name of the node creation graph action

diff --git a/Engine/Plugins/Cameras/GameplayCameras/Source/GameplayCameras/Private/Core/ObjectTreeGraphComment.cpp b/Engine/Plugins/Cameras/GameplayCameras/Source/GameplayCameras/Private/Core/ObjectTreeGraphComment.cpp
new file mode 100644
--- /dev/null
+++ b/Engine/Plugins/Cameras/GameplayCameras/Source/GameplayCameras/Private/Core/ObjectTreeGraphComment.cpp
@@ -0,0 +1,41 @@
+// Copyright Epic Games, Inc. All Rights Reserved.
+
+#include "Core/ObjectTreeGraphComment.h"
+
+#include UE_INLINE_GENERATED_CPP_BY_NAME(ObjectTreeGraphComment)
+
+UObjectTreeGraphComment::UObjectTreeGraphComment(const FObjectInitializer& ObjInit)
+	: Super(ObjInit)
+{
+}
+
+#if WITH_EDITOR
+
+void UObjectTreeGraphComment::GetGraphNodePosition(FName InGraphName, int32& NodePosX, int32& NodePosY) const
+{
+	NodePosX = GraphNodePos.X;
+	NodePosY = GraphNodePos.Y;
+}
+
+void UObjectTreeGraphComment::OnGraphNodeMoved(FName InGraphName, int32 NodePosX, int32 NodePosY, bool bMarkDirty)
+{
+	Modify(bMarkDirty);
+
+	GraphNodePos.X = NodePosX;
+	GraphNodePos.Y = NodePosY;
+}
+
+void UObjectTreeGraphComment::GetGraphNodeName(FName InGraphName, FText& OutName) const
+{
+	OutName = FText::FromString(CommentText);
+}
+
+void UObjectTreeGraphComment::OnRenameGraphNode(FName InGraphName, const FString& NewName)
+{
+	Modify();
+
+	CommentText = NewName;
+}
+
+#endif  // WITH_EDITOR
+
diff --git a/Engine/Plugins/Cameras/GameplayCameras/Source/GameplayCameras/Public/Core/ObjectTreeGraphComment.h b/Engine/Plugins/Cameras/GameplayCameras/Source/GameplayCameras/Public/Core/ObjectTreeGraphComment.h
new file mode 100644
--- /dev/null
+++ b/Engine/Plugins/Cameras/GameplayCameras/Source/GameplayCameras/Public/Core/ObjectTreeGraphComment.h
@@ -0,0 +1,58 @@
+// Copyright Epic Games, Inc. All Rights Reserved.
+
+#pragma once
+
+#include "Core/ObjectTreeGraphObject.h"
+#include "UObject/Object.h"
+
+#include "ObjectTreeGraphComment.generated.h"
+
+/**
+ * A generic comment node for object tree graphs.
+ */
+UCLASS(MinimalAPI)
+class UObjectTreeGraphComment 
+	: public UObject
+	, public IObjectTreeGraphObject
+{
+	GENERATED_BODY()
+
+public:
+
+	UObjectTreeGraphComment(const FObjectInitializer& ObjInit);
+
+#if WITH_EDITOR
+
+public:
+
+	// IObjectTreeGraphObject interface.
+	virtual void GetGraphNodePosition(FName InGraphName, int32& NodePosX, int32& NodePosY) const override;
+	virtual void OnGraphNodeMoved(FName InGraphName, int32 NodePosX, int32 NodePosY, bool bMarkDirty) override;
+	virtual void GetGraphNodeName(FName InGraphName, FText& OutName) const override;
+	virtual void OnRenameGraphNode(FName InGraphName, const FString& NewName) override;
+
+#endif  // WITH_EDITOR
+
+#if WITH_EDITORONLY_DATA
+
+public:
+
+	/** The text of the comment node. */
+	UPROPERTY()
+	FString CommentText;
+
+	/** Color of the comment node in the node graph editor. */
+	UPROPERTY()
+	FLinearColor CommentColor = FLinearColor::White;
+
+	/** Position of the comment node in the node graph editor. */
+	UPROPERTY()
+	FIntVector2 GraphNodePos = FIntVector2::ZeroValue;
+
+	/** Size of the comment node in the node graph editor. */
+	UPROPERTY()
+	FIntVector2 GraphNodeSize = FIntVector2(400, 100);
+
+#endif  // WITH_EDITORONLY_DATA
+};
+
diff --git a/Engine/Plugins/Cameras/GameplayCameras/Source/GameplayCameras/Public/Core/ObjectTreeGraphObject.h b/Engine/Plugins/Cameras/GameplayCameras/Source/GameplayCameras/Public/Core/ObjectTreeGraphObject.h
--- a/Engine/Plugins/Cameras/GameplayCameras/Source/GameplayCameras/Public/Core/ObjectTreeGraphObject.h
+++ b/Engine/Plugins/Cameras/GameplayCameras/Source/GameplayCameras/Public/Core/ObjectTreeGraphObject.h
@@ -2,7 +2,7 @@
 
 #pragma once
 
-#include "CoreMinimal.h"
+#include "CoreTypes.h"
 #include "Misc/EnumClassFlags.h"
 #include "UObject/Interface.h"
 
diff --git a/Engine/Plugins/Cameras/GameplayCameras/Source/GameplayCamerasEditor/Private/Editors/CameraNodeGraphSchema.cpp b/Engine/Plugins/Cameras/GameplayCameras/Source/GameplayCamerasEditor/Private/Editors/CameraNodeGraphSchema.cpp
--- a/Engine/Plugins/Cameras/GameplayCameras/Source/GameplayCamerasEditor/Private/Editors/CameraNodeGraphSchema.cpp
+++ b/Engine/Plugins/Cameras/GameplayCameras/Source/GameplayCamerasEditor/Private/Editors/CameraNodeGraphSchema.cpp
@@ -87,7 +87,7 @@ void UCameraNodeGraphSchema::OnCreateAllNodes(UObjectTreeGraph* InGraph, const F
 		}
 		
 		UCameraObjectInterfaceParameterGraphNode* InterfaceParameterNode = CreateInterfaceParameterNode(InGraph, InterfaceParameter);
-		UObjectTreeGraphNode* CameraNodeNode = InCreatedNodes.CreatedNodes.FindRef(InterfaceParameter->Target);
+		UObjectTreeGraphNode* CameraNodeNode = Cast<UObjectTreeGraphNode>(InCreatedNodes.CreatedNodes.FindRef(InterfaceParameter->Target));
 		if (CameraNodeNode)
 		{
 			UEdGraphPin* InterfaceParameterSelfPin = InterfaceParameterNode->GetSelfPin();
diff --git a/Engine/Plugins/Cameras/GameplayCameras/Source/GameplayCamerasEditor/Private/Editors/ObjectTreeDragDropOp.cpp b/Engine/Plugins/Cameras/GameplayCameras/Source/GameplayCamerasEditor/Private/Editors/ObjectTreeDragDropOp.cpp
--- a/Engine/Plugins/Cameras/GameplayCameras/Source/GameplayCamerasEditor/Private/Editors/ObjectTreeDragDropOp.cpp
+++ b/Engine/Plugins/Cameras/GameplayCameras/Source/GameplayCamerasEditor/Private/Editors/ObjectTreeDragDropOp.cpp
@@ -72,7 +72,7 @@ FReply FObjectTreeClassDragDropOp::ExecuteDrop(TSharedPtr<SGraphEditor> GraphEdi
 	FVector2D CurLocation = NewLocation;
 	for (UClass* PlaceableClass : PlaceableClasses)
 	{
-		FObjectGraphSchemaAction_NewNode Action;
+		FObjectTreeGraphSchemaAction_NewNode Action;
 		Action.ObjectClass = PlaceableClass;
 		UEdGraphNode* NewNode = Action.PerformAction(Graph, nullptr, CurLocation, false);
 		GraphEditor->SetNodeSelection(NewNode, true);
diff --git a/Engine/Plugins/Cameras/GameplayCameras/Source/GameplayCamerasEditor/Private/Editors/ObjectTreeGraphCommentNode.cpp b/Engine/Plugins/Cameras/GameplayCameras/Source/GameplayCamerasEditor/Private/Editors/ObjectTreeGraphCommentNode.cpp
new file mode 100644
--- /dev/null
+++ b/Engine/Plugins/Cameras/GameplayCameras/Source/GameplayCamerasEditor/Private/Editors/ObjectTreeGraphCommentNode.cpp
@@ -0,0 +1,145 @@
+// Copyright Epic Games, Inc. All Rights Reserved.
+
+#include "Editors/ObjectTreeGraphCommentNode.h"
+
+#include "Core/ObjectTreeGraphComment.h"
+#include "Editors/ObjectTreeGraph.h"
+#include "Editors/SObjectTreeGraphCommentNode.h"
+#include "Framework/Commands/GenericCommands.h"
+#include "GraphEditorActions.h"
+#include "ScopedTransaction.h"
+#include "ToolMenus.h"
+#include "ToolMenuDelegates.h"
+
+#include UE_INLINE_GENERATED_CPP_BY_NAME(ObjectTreeGraphCommentNode)
+
+#define LOCTEXT_NAMESPACE "ObjectTreeGraphCommentNode"
+
+void UObjectTreeGraphCommentNode::Initialize(UObjectTreeGraphComment* InObject)
+{
+	WeakObject = InObject;
+}
+
+TSharedPtr<SGraphNode> UObjectTreeGraphCommentNode::CreateVisualWidget()
+{
+	return SNew(SObjectTreeGraphCommentNode).GraphNode(this);
+}
+
+void UObjectTreeGraphCommentNode::GetNodeContextMenuActions(UToolMenu* Menu, UGraphNodeContextMenuContext* Context) const
+{
+	const FGraphEditorCommandsImpl& GraphEditorCommands = FGraphEditorCommands::Get();
+	const FGenericCommands& GenericCommands = FGenericCommands::Get();
+
+	// General actions.
+	{
+		FToolMenuSection& Section = Menu->AddSection(
+				"ObjectTreeGraphNodeGenericActions", LOCTEXT("GenericActionsMenuHeader", "General"));
+
+		Section.AddMenuEntry(GenericCommands.Delete);
+		Section.AddMenuEntry(GenericCommands.Cut);
+		Section.AddMenuEntry(GenericCommands.Copy);
+		Section.AddMenuEntry(GenericCommands.Duplicate);
+	}
+
+	// Graph organization.
+	{
+		FToolMenuSection& Section = Menu->AddSection(
+				"ObjectTreeGraphOrganizationActions", LOCTEXT("OrganizationActionsMenuHeader", "Organization"));
+
+		Section.AddSubMenu(
+				"Alignment",
+				LOCTEXT("AlignmentHeader", "Alignment"),
+				FText(),
+				FNewToolMenuDelegate::CreateLambda([&GraphEditorCommands](UToolMenu* InMenu)
+					{
+						FToolMenuSection& SubMenuSection = InMenu->AddSection(
+								"ObjectTreeGraphAlignmentActions", LOCTEXT("AlignmentHeader", "Alignment"));
+						SubMenuSection.AddMenuEntry(GraphEditorCommands.AlignNodesTop);
+						SubMenuSection.AddMenuEntry(GraphEditorCommands.AlignNodesMiddle);
+						SubMenuSection.AddMenuEntry(GraphEditorCommands.AlignNodesBottom);
+						SubMenuSection.AddMenuEntry(GraphEditorCommands.AlignNodesLeft);
+						SubMenuSection.AddMenuEntry(GraphEditorCommands.AlignNodesCenter);
+						SubMenuSection.AddMenuEntry(GraphEditorCommands.AlignNodesRight);
+						SubMenuSection.AddMenuEntry(GraphEditorCommands.StraightenConnections);
+					}));
+
+		Section.AddSubMenu(
+				"Distribution",
+				LOCTEXT("DistributionHeader", "Distribution"),
+				FText(),
+				FNewToolMenuDelegate::CreateLambda([&GraphEditorCommands](UToolMenu* InMenu)
+					{
+						FToolMenuSection& SubMenuSection = InMenu->AddSection(
+								"ObjectTreeGraphDistributionActions", LOCTEXT("DistributionHeader", "Distribution"));
+						SubMenuSection.AddMenuEntry(GraphEditorCommands.DistributeNodesHorizontally);
+						SubMenuSection.AddMenuEntry(GraphEditorCommands.DistributeNodesVertically);
+					}));
+	}
+}
+
+void UObjectTreeGraphCommentNode::PostPlacedNewNode()
+{
+	Super::PostPlacedNewNode();
+
+	MoveMode = ECommentBoxMode::GroupMovement;
+
+	UObjectTreeGraphComment* CommentObject = WeakObject.Get();
+	if (CommentObject)
+	{
+		UObjectTreeGraph* OuterGraph = CastChecked<UObjectTreeGraph>(GetGraph());
+		const FObjectTreeGraphConfig& OuterGraphConfig = OuterGraph->GetConfig();
+
+		CommentObject->GetGraphNodePosition(OuterGraphConfig.GraphName, NodePosX, NodePosY);
+
+		NodeWidth = CommentObject->GraphNodeSize.X;
+		NodeHeight = CommentObject->GraphNodeSize.Y;
+
+		NodeComment = CommentObject->CommentText;
+		CommentColor = CommentObject->CommentColor;
+	}
+}
+
+void UObjectTreeGraphCommentNode::ResizeNode(const FVector2D& NewSize)
+{
+	Super::ResizeNode(NewSize);
+
+	if (UObjectTreeGraphComment* Object = WeakObject.Get())
+	{
+		Object->Modify();
+		Object->GraphNodeSize = FIntVector2((int32)NewSize.X, (int32)NewSize.Y);
+		// Set position as well since may have been resized from the top corner.
+		Object->GraphNodePos = FIntVector2(NodePosX, NodePosY);
+	}
+}
+
+void UObjectTreeGraphCommentNode::OnRenameNode(const FString& NewName)
+{
+	Super::OnRenameNode(NewName);
+
+	UObjectTreeGraphComment* CommentObject = WeakObject.Get();
+	if (CommentObject)
+	{
+		const FScopedTransaction Transaction(LOCTEXT("RenameNode", "Rename Node"));
+
+		UObjectTreeGraph* OuterGraph = CastChecked<UObjectTreeGraph>(GetGraph());
+		const FObjectTreeGraphConfig& OuterGraphConfig = OuterGraph->GetConfig();
+
+		CommentObject->OnRenameGraphNode(OuterGraphConfig.GraphName, NewName);
+	}
+}
+
+void UObjectTreeGraphCommentNode::OnGraphNodeMoved(bool bMarkDirty)
+{
+	UObject* Object = WeakObject.Get();
+	IObjectTreeGraphObject* GraphObject = Cast<IObjectTreeGraphObject>(Object);
+	if (GraphObject)
+	{
+		UObjectTreeGraph* OuterGraph = CastChecked<UObjectTreeGraph>(GetGraph());
+		const FObjectTreeGraphConfig& OuterGraphConfig = OuterGraph->GetConfig();
+
+		GraphObject->OnGraphNodeMoved(OuterGraphConfig.GraphName, NodePosX, NodePosY, bMarkDirty);
+	}
+}
+
+#undef LOCTEXT_NAMESPACE
+
diff --git a/Engine/Plugins/Cameras/GameplayCameras/Source/GameplayCamerasEditor/Private/Editors/ObjectTreeGraphSchema.cpp b/Engine/Plugins/Cameras/GameplayCameras/Source/GameplayCamerasEditor/Private/Editors/ObjectTreeGraphSchema.cpp
--- a/Engine/Plugins/Cameras/GameplayCameras/Source/GameplayCamerasEditor/Private/Editors/ObjectTreeGraphSchema.cpp
+++ b/Engine/Plugins/Cameras/GameplayCameras/Source/GameplayCamerasEditor/Private/Editors/ObjectTreeGraphSchema.cpp
@@ -3,10 +3,13 @@
 #include "Editors/ObjectTreeGraphSchema.h"
 
 #include "Commands/ObjectTreeGraphEditorCommands.h"
+#include "Core/ObjectTreeGraphComment.h"
 #include "Core/ObjectTreeGraphRootObject.h"
 #include "Editors/ObjectTreeConnectionDrawingPolicy.h"
 #include "Editors/ObjectTreeGraph.h"
+#include "Editors/ObjectTreeGraphCommentNode.h"
 #include "Editors/ObjectTreeGraphNode.h"
+#include "Editors/SObjectTreeGraphEditor.h"
 #include "Exporters/Exporter.h"
 #include "Factories.h"
 #include "IGameplayCamerasEditorModule.h"
@@ -220,7 +223,7 @@ void UObjectTreeGraphSchema::CreateAllNodes(UObjectTreeGraph* InGraph) const
 	FCreatedNodes CreatedNodes;
 	for (UObject* Object : AllObjects)
 	{
-		if (UObjectTreeGraphNode* GraphNode = CreateObjectNode(InGraph, Object))
+		if (UEdGraphNode* GraphNode = CreateObjectNode(InGraph, Object))
 		{
 			CreatedNodes.CreatedNodes.Add(Object, GraphNode);
 		}
@@ -231,19 +234,22 @@ void UObjectTreeGraphSchema::CreateAllNodes(UObjectTreeGraph* InGraph) const
 	if (!AllObjects.IsEmpty())
 	{
 		UObject* RootObject = InGraph->GetRootObject();
-		UObjectTreeGraphNode** CreatedRootObjectNode = CreatedNodes.CreatedNodes.Find(RootObject);
+		UEdGraphNode* CreatedRootObjectNode = CreatedNodes.CreatedNodes.FindRef(RootObject);
 		if (ensureMsgf(CreatedRootObjectNode, 
 					TEXT("Can't find root object '%s' in the list of created graph nodes!"),
 					*GetNameSafe(RootObject)))
 		{
-			InGraph->RootObjectNode = *CreatedRootObjectNode;
+			InGraph->RootObjectNode = CastChecked<UObjectTreeGraphNode>(CreatedRootObjectNode);
 		}
 	}
 
 	// Create all the connections.
-	for (TPair<UObject*, UObjectTreeGraphNode*> Pair : CreatedNodes.CreatedNodes)
+	for (TPair<UObject*, UEdGraphNode*> Pair : CreatedNodes.CreatedNodes)
 	{
-		CreateConnections(Pair.Value, CreatedNodes);
+		if (UObjectTreeGraphNode* Node = Cast<UObjectTreeGraphNode>(Pair.Value))
+		{
+			CreateConnections(Node, CreatedNodes);
+		}
 	}
 
 	OnCreateAllNodes(InGraph, CreatedNodes);
@@ -276,14 +282,14 @@ void UObjectTreeGraphSchema::CreateConnections(UObjectTreeGraphNode* InGraphNode
 				continue;
 			}
 
-			UObjectTreeGraphNode* const* ConnectedNode = InCreatedNodes.CreatedNodes.Find(OutConnectedObject);
+			UObjectTreeGraphNode* ConnectedNode = Cast<UObjectTreeGraphNode>(InCreatedNodes.CreatedNodes.FindRef(OutConnectedObject));
 			if (ensure(ConnectedNode))
 			{
 				if (Pin->Direction == EGPD_Input)
 				{
-					(*ConnectedNode)->OverrideSelfPinDirection(EGPD_Output);
+					ConnectedNode->OverrideSelfPinDirection(EGPD_Output);
 				}
-				UEdGraphPin* ConnectedPin = (*ConnectedNode)->GetSelfPin();
+				UEdGraphPin* ConnectedPin = ConnectedNode->GetSelfPin();
 				Pin->MakeLinkTo(ConnectedPin);
 			}
 		}
@@ -309,14 +315,14 @@ void UObjectTreeGraphSchema::CreateConnections(UObjectTreeGraphNode* InGraphNode
 					continue;
 				}
 
-				UObjectTreeGraphNode* const* ConnectedNode = InCreatedNodes.CreatedNodes.Find(ConnectedObject);
+				UObjectTreeGraphNode* ConnectedNode = Cast<UObjectTreeGraphNode>(InCreatedNodes.CreatedNodes.FindRef(ConnectedObject));
 				if (ensure(ConnectedNode))
 				{
 					if (Pin->Direction == EGPD_Input)
 					{
-						(*ConnectedNode)->OverrideSelfPinDirection(EGPD_Output);
+						ConnectedNode->OverrideSelfPinDirection(EGPD_Output);
 					}
-					UEdGraphPin* ConnectedPin = (*ConnectedNode)->GetSelfPin();
+					UEdGraphPin* ConnectedPin = ConnectedNode->GetSelfPin();
 					Pin->MakeLinkTo(ConnectedPin);
 				}
 			}
@@ -328,22 +334,37 @@ void UObjectTreeGraphSchema::OnCreateAllNodes(UObjectTreeGraph* InGraph, const F
 {
 }
 
-UObjectTreeGraphNode* UObjectTreeGraphSchema::CreateObjectNode(UObjectTreeGraph* InGraph, UObject* InObject) const
+UEdGraphNode* UObjectTreeGraphSchema::CreateObjectNode(UObjectTreeGraph* InGraph, UObject* InObject) const
 {
 	if (!InObject)
 	{
 		return nullptr;
 	}
-
-	if (!InGraph->GetConfig().IsConnectable(InObject->GetClass()))
+	
+	if (UObjectTreeGraphComment* Comment = Cast<UObjectTreeGraphComment>(InObject))
 	{
-		return nullptr;
+		return CreateCommentNode(InGraph, Comment);
+	}
+	else if (InGraph->GetConfig().IsConnectable(InObject->GetClass()))
+	{
+		return OnCreateObjectNode(InGraph, InObject);
 	}
 
-	return OnCreateObjectNode(InGraph, InObject);
+	return nullptr;
 }
 
-UObjectTreeGraphNode* UObjectTreeGraphSchema::OnCreateObjectNode(UObjectTreeGraph* InGraph, UObject* InObject) const
+UEdGraphNode* UObjectTreeGraphSchema::CreateCommentNode(UObjectTreeGraph* InGraph, UObjectTreeGraphComment* InComment) const
+{
+	InGraph->Modify();
+
+	FGraphNodeCreator<UObjectTreeGraphCommentNode> GraphNodeCreator(*InGraph);
+	UObjectTreeGraphCommentNode* NewNode = GraphNodeCreator.CreateNode(false);
+	NewNode->Initialize(InComment);
+	GraphNodeCreator.Finalize();
+	return NewNode;
+}
+
+UEdGraphNode* UObjectTreeGraphSchema::OnCreateObjectNode(UObjectTreeGraph* InGraph, UObject* InObject) const
 {
 	const FObjectTreeGraphConfig& Config = InGraph->GetConfig();
 	const FObjectTreeGraphClassConfigs ClassConfigs = Config.GetObjectClassConfigs(InObject->GetClass());
@@ -363,16 +384,13 @@ UObjectTreeGraphNode* UObjectTreeGraphSchema::OnCreateObjectNode(UObjectTreeGrap
 	return NewNode;
 }
 
-void UObjectTreeGraphSchema::AddConnectableObject(UObjectTreeGraph* InGraph, UObjectTreeGraphNode* InNewNode) const
+void UObjectTreeGraphSchema::AddConnectableObject(UObjectTreeGraph* InGraph, UObject* InNewObject) const
 {
-	UObject* Object = InNewNode->GetObject();
-	if (!ensure(Object))
+	if (!ensure(InNewObject))
 	{
 		return;
 	}
-
-	const FObjectTreeGraphConfig& GraphConfig = InGraph->GetConfig();
-	if (!GraphConfig.IsConnectable(Object->GetClass()))
+	if (!ensure(!InNewObject->IsA<UEdGraphNode>()))
 	{
 		return;
 	}
@@ -381,27 +399,25 @@ void UObjectTreeGraphSchema::AddConnectableObject(UObjectTreeGraph* InGraph, UOb
 	IObjectTreeGraphRootObject* RootObjectInterface = Cast<IObjectTreeGraphRootObject>(RootObjectNode->GetObject());
 	if (RootObjectInterface)
 	{
+		const FObjectTreeGraphConfig& GraphConfig = InGraph->GetConfig();
 		const FName GraphName = GraphConfig.GraphName;
-		RootObjectInterface->AddConnectableObject(GraphName, Object);
+		RootObjectInterface->AddConnectableObject(GraphName, InNewObject);
 	}
 
-	OnAddConnectableObject(InGraph, InNewNode);
+	OnAddConnectableObject(InGraph, InNewObject);
 }
 
-void UObjectTreeGraphSchema::OnAddConnectableObject(UObjectTreeGraph* InGraph, UObjectTreeGraphNode* InNewNode) const
+void UObjectTreeGraphSchema::OnAddConnectableObject(UObjectTreeGraph* InGraph, UObject* InNewObject) const
 {
 }
 
-void UObjectTreeGraphSchema::RemoveConnectableObject(UObjectTreeGraph* InGraph, UObjectTreeGraphNode* InRemovedNode) const
+void UObjectTreeGraphSchema::RemoveConnectableObject(UObjectTreeGraph* InGraph, UObject* InRemovedObject) const
 {
-	UObject* Object = InRemovedNode->GetObject();
-	if (!ensure(Object))
+	if (!ensure(InRemovedObject))
 	{
 		return;
 	}
-
-	const FObjectTreeGraphConfig& GraphConfig = InGraph->GetConfig();
-	if (!GraphConfig.IsConnectable(Object->GetClass()))
+	if (!ensure(!InRemovedObject->IsA<UEdGraphNode>()))
 	{
 		return;
 	}
@@ -409,14 +425,15 @@ void UObjectTreeGraphSchema::RemoveConnectableObject(UObjectTreeGraph* InGraph,
 	IObjectTreeGraphRootObject* RootObjectInterface = Cast<IObjectTreeGraphRootObject>(InGraph->GetRootObject());
 	if (RootObjectInterface)
 	{
+		const FObjectTreeGraphConfig& GraphConfig = InGraph->GetConfig();
 		const FName GraphName = GraphConfig.GraphName;
-		RootObjectInterface->RemoveConnectableObject(GraphName, Object);
+		RootObjectInterface->RemoveConnectableObject(GraphName, InRemovedObject);
 	}
 
-	OnRemoveConnectableObject(InGraph, InRemovedNode);
+	OnRemoveConnectableObject(InGraph, InRemovedObject);
 }
 
-void UObjectTreeGraphSchema::OnRemoveConnectableObject(UObjectTreeGraph* InGraph, UObjectTreeGraphNode* InRemovedNode) const
+void UObjectTreeGraphSchema::OnRemoveConnectableObject(UObjectTreeGraph* InGraph, UObject* InRemovedObject) const
 {
 }
 
@@ -538,13 +555,15 @@ void UObjectTreeGraphSchema::GetGraphContextActions(FGraphContextMenuBuilder& Co
 
 			FText KeywordsText(FText::FromString(PossibleObjectClass->GetMetaData(TEXT("Keywords"))));
 
-			TSharedRef<FObjectGraphSchemaAction_NewNode> Action = MakeShared<FObjectGraphSchemaAction_NewNode>(
+			TSharedRef<FObjectTreeGraphSchemaAction_NewNode> Action = MakeShared<FObjectTreeGraphSchemaAction_NewNode>(
 					CategoryText, DisplayName, ToolTipText, Grouping, KeywordsText);
 			Action->ObjectClass = PossibleObjectClass;
 			ContextMenuBuilder.AddAction(StaticCastSharedPtr<FEdGraphSchemaAction>(Action.ToSharedPtr()));
 		}
 	}
 
+	GetCommentAction(ContextMenuBuilder);
+
 	// Don't call the base class, we want to control exactly what can be created.
 }
 
@@ -582,6 +601,22 @@ void UObjectTreeGraphSchema::GetContextMenuActions(UToolMenu* Menu, UGraphNodeCo
 	}
 }
 
+TSharedPtr<FEdGraphSchemaAction> UObjectTreeGraphSchema::GetCreateCommentAction() const
+{
+	return MakeShared<FObjectTreeGraphSchemaAction_NewComment>();
+}
+
+void UObjectTreeGraphSchema::GetCommentAction(FGraphActionMenuBuilder& ActionMenuBuilder) const
+{
+	if (!ActionMenuBuilder.FromPin)
+	{
+		const FText MenuDesc= LOCTEXT("CommentActionDescription", "Add Comment...");
+		const FText MenuToolTip = LOCTEXT("CommentToolTip", "Creates a comment.");
+		TSharedPtr<FObjectTreeGraphSchemaAction_NewComment> NewAction(MakeShared<FObjectTreeGraphSchemaAction_NewComment>(FText::GetEmpty(), MenuDesc, MenuToolTip));
+		ActionMenuBuilder.AddAction(NewAction);
+	}
+}
+
 FName UObjectTreeGraphSchema::GetParentContextMenuName() const
 {
 	// Return NAME_None if we don't want the default menu entries.
@@ -1177,10 +1212,13 @@ bool UObjectTreeGraphSchema::SafeDeleteNodeFromGraph(UEdGraph* Graph, UEdGraphNo
 
 void UObjectTreeGraphSchema::OnDeleteNodeFromGraph(UObjectTreeGraph* Graph, UEdGraphNode* Node) const
 {
-	UObjectTreeGraphNode* ObjectNode = Cast<UObjectTreeGraphNode>(Node);
-	if (ObjectNode)
+	if (UObjectTreeGraphNode* ObjectNode = Cast<UObjectTreeGraphNode>(Node))
+	{
+		RemoveConnectableObject(Graph, ObjectNode->GetObject());
+	}
+	else if (UObjectTreeGraphCommentNode* CommentNode = Cast<UObjectTreeGraphCommentNode>(Node))
 	{
-		RemoveConnectableObject(Graph, ObjectNode);
+		RemoveConnectableObject(Graph, CommentNode->GetObject());
 	}
 }
 
@@ -1225,6 +1263,10 @@ FString UObjectTreeGraphSchema::ExportNodesToText(const FGraphPanelSelectionSet&
 			{
 				ObjectsToExport.Add(ObjectTreeNode->GetObject());
 			}
+			else if (UObjectTreeGraphCommentNode* CommentNode = Cast<UObjectTreeGraphCommentNode>(Node))
+			{
+				ObjectsToExport.Add(CommentNode->GetObject());
+			}
 			else
 			{
 				OtherNodesToExport.Add(Node);
@@ -1360,23 +1402,26 @@ void UObjectTreeGraphSchema::ImportNodesFromText(UObjectTreeGraph* InGraph, cons
 	FCreatedNodes CreatedNodes;
 	for (UObject* Object : ImportedObjects)
 	{
-		if (UObjectTreeGraphNode* GraphNode = CreateObjectNode(InGraph, Object))
+		if (UEdGraphNode* GraphNode = CreateObjectNode(InGraph, Object))
 		{
 			CreatedNodes.CreatedNodes.Add(Object, GraphNode);
 
-			AddConnectableObject(InGraph, GraphNode);
+			AddConnectableObject(InGraph, Object);
 		}
 	}
 
 	// Create all the connections.
-	for (TPair<UObject*, UObjectTreeGraphNode*> Pair : CreatedNodes.CreatedNodes)
+	for (TPair<UObject*, UEdGraphNode*> Pair : CreatedNodes.CreatedNodes)
 	{
-		CreateConnections(Pair.Value, CreatedNodes);
+		if (UObjectTreeGraphNode* Node = Cast<UObjectTreeGraphNode>(Pair.Value))
+		{
+			CreateConnections(Node, CreatedNodes);
+		}
 	}
 
 	OnCreateAllNodes(InGraph, CreatedNodes);
 
-	for (const TTuple<UObject*, UObjectTreeGraphNode*>& Pair : CreatedNodes.CreatedNodes)
+	for (const TTuple<UObject*, UEdGraphNode*>& Pair : CreatedNodes.CreatedNodes)
 	{
 		OutPastedNodes.Add(Pair.Value);
 	}
@@ -1403,16 +1448,16 @@ const FObjectTreeGraphClassConfigs UObjectTreeGraphSchema::GetObjectClassConfigs
 	return InGraph->GetConfig().GetObjectClassConfigs(InObjectClass);
 }
 
-FObjectGraphSchemaAction_NewNode::FObjectGraphSchemaAction_NewNode()
+FObjectTreeGraphSchemaAction_NewNode::FObjectTreeGraphSchemaAction_NewNode()
 {
 }
 
-FObjectGraphSchemaAction_NewNode::FObjectGraphSchemaAction_NewNode(FText InNodeCategory, FText InMenuDesc, FText InToolTip, const int32 InGrouping, FText InKeywords)
+FObjectTreeGraphSchemaAction_NewNode::FObjectTreeGraphSchemaAction_NewNode(FText InNodeCategory, FText InMenuDesc, FText InToolTip, const int32 InGrouping, FText InKeywords)
 	: FEdGraphSchemaAction(InNodeCategory, InMenuDesc, InToolTip, InGrouping, InKeywords)
 {
 }
 
-UEdGraphNode* FObjectGraphSchemaAction_NewNode::PerformAction(UEdGraph* ParentGraph, UEdGraphPin* FromPin, const FVector2D Location, bool bSelectNewNode)
+UEdGraphNode* FObjectTreeGraphSchemaAction_NewNode::PerformAction(UEdGraph* ParentGraph, UEdGraphPin* FromPin, const FVector2D Location, bool bSelectNewNode)
 {
 	UObjectTreeGraph* ObjectTreeGraph = Cast<UObjectTreeGraph>(ParentGraph);
 	if (!ensure(ObjectTreeGraph))
@@ -1452,13 +1497,16 @@ UEdGraphNode* FObjectGraphSchemaAction_NewNode::PerformAction(UEdGraph* ParentGr
 
 		ObjectTreeGraph->Modify();
 
-		UObjectTreeGraphNode* NewGraphNode = Schema->CreateObjectNode(ObjectTreeGraph, NewObject);
+		UEdGraphNode* NewGraphNode = Schema->CreateObjectNode(ObjectTreeGraph, NewObject);
 
-		Schema->AddConnectableObject(ObjectTreeGraph, NewGraphNode);
+		Schema->AddConnectableObject(ObjectTreeGraph, NewObject);
 
 		NewGraphNode->NodePosX = Location.X;
 		NewGraphNode->NodePosY = Location.Y;
-		NewGraphNode->OnGraphNodeMoved(false);
+		if (UObjectTreeGraphNode* NewObjectGraphNode = Cast<UObjectTreeGraphNode>(NewGraphNode))
+		{
+			NewObjectGraphNode->OnGraphNodeMoved(false);
+		}
 
 		AutoSetupNewNode(NewGraphNode, FromPin);
 
@@ -1468,15 +1516,82 @@ UEdGraphNode* FObjectGraphSchemaAction_NewNode::PerformAction(UEdGraph* ParentGr
 	return nullptr;
 }
 
-UObject* FObjectGraphSchemaAction_NewNode::CreateObject()
+UObject* FObjectTreeGraphSchemaAction_NewNode::CreateObject()
 {
 	return NewObject<UObject>(ObjectOuter, ObjectClass, NAME_None, RF_Transactional);
 }
 
-void FObjectGraphSchemaAction_NewNode::AutoSetupNewNode(UObjectTreeGraphNode* NewNode, UEdGraphPin* FromPin)
+void FObjectTreeGraphSchemaAction_NewNode::AutoSetupNewNode(UEdGraphNode* NewNode, UEdGraphPin* FromPin)
 {
 	NewNode->AutowireNewNode(FromPin);
 }
 
+FObjectTreeGraphSchemaAction_NewComment::FObjectTreeGraphSchemaAction_NewComment()
+{
+}
+
+FObjectTreeGraphSchemaAction_NewComment::FObjectTreeGraphSchemaAction_NewComment(FText InNodeCategory, FText InMenuDesc, FText InToolTip, const int32 InGrouping, FText InKeywords)
+	: FEdGraphSchemaAction(InNodeCategory, InMenuDesc, InToolTip, InGrouping, InKeywords)
+{
+}
+
+UEdGraphNode* FObjectTreeGraphSchemaAction_NewComment::PerformAction(UEdGraph* ParentGraph, UEdGraphPin* FromPin, const FVector2D Location, bool bSelectNewNode)
+{
+	UObjectTreeGraph* ObjectTreeGraph = Cast<UObjectTreeGraph>(ParentGraph);
+	if (!ensure(ObjectTreeGraph))
+	{
+		return nullptr;
+	}
+
+	UObject* ObjectOuter = ObjectTreeGraph->GetRootObject();
+	if (!ensure(ObjectOuter))
+	{
+		return nullptr;
+	}
+
+	FSlateRect Bounds;
+	bool bUseBounds = false;
+	if (TSharedPtr<SObjectTreeGraphEditor> GraphEditor = SObjectTreeGraphEditor::FindGraphEditor(ObjectTreeGraph))
+	{
+		bUseBounds = GraphEditor->GetGraphEditor()->GetBoundsForSelectedNodes(Bounds, 50.f);
+	}
+
+	const FScopedTransaction Transaction(LOCTEXT("CreateNewCommentAction", "Create Comment"));
+
+	const UObjectTreeGraphSchema* Schema = CastChecked<UObjectTreeGraphSchema>(ParentGraph->GetSchema());
+
+	UObjectTreeGraphComment* NewComment = NewObject<UObjectTreeGraphComment>(ObjectOuter, NAME_None, RF_Transactional);
+
+	UEdGraphNode* NewGraphNode = Schema->CreateObjectNode(ObjectTreeGraph, NewComment);
+
+	Schema->AddConnectableObject(ObjectTreeGraph, NewComment);
+
+	if (bUseBounds)
+	{
+		NewGraphNode->NodePosX = Bounds.Left;
+		NewGraphNode->NodePosY = Bounds.Top;
+
+		FVector2D BoundsSize = Bounds.GetSize();
+		NewGraphNode->NodeWidth = BoundsSize.X;
+		NewGraphNode->NodeHeight = BoundsSize.Y;
+	}
+	else
+	{
+		NewGraphNode->NodePosX = Location.X;
+		NewGraphNode->NodePosY = Location.Y;
+
+		NewGraphNode->NodeWidth = 400;
+		NewGraphNode->NodeHeight = 400;
+	}
+	if (UObjectTreeGraphCommentNode* NewCommentNode = Cast<UObjectTreeGraphCommentNode>(NewGraphNode))
+	{
+		NewCommentNode->OnGraphNodeMoved(false);
+	}
+
+	NewGraphNode->AutowireNewNode(FromPin);
+
+	return NewGraphNode;
+}
+
 #undef LOCTEXT_NAMESPACE
 
diff --git a/Engine/Plugins/Cameras/GameplayCameras/Source/GameplayCamerasEditor/Private/Editors/SCameraRigAssetEditor.cpp b/Engine/Plugins/Cameras/GameplayCameras/Source/GameplayCamerasEditor/Private/Editors/SCameraRigAssetEditor.cpp
--- a/Engine/Plugins/Cameras/GameplayCameras/Source/GameplayCamerasEditor/Private/Editors/SCameraRigAssetEditor.cpp
+++ b/Engine/Plugins/Cameras/GameplayCameras/Source/GameplayCamerasEditor/Private/Editors/SCameraRigAssetEditor.cpp
@@ -98,6 +98,7 @@ void SCameraRigAssetEditor::CreateNodeGraphEditor()
 		.IsEnabled(this, &SCameraRigAssetEditor::IsGraphEditorEnabled)
 		.GraphToEdit(NodeGraph)
 		.AssetEditorToolkit(AssetEditorToolkit);
+	NodeGraphEditor->RegisterEditor();
 }
 
 void SCameraRigAssetEditor::CreateTransitionGraphEditor()
@@ -123,6 +124,7 @@ void SCameraRigAssetEditor::CreateTransitionGraphEditor()
 		.IsEnabled(this, &SCameraRigAssetEditor::IsGraphEditorEnabled)
 		.GraphToEdit(TransitionGraph)
 		.AssetEditorToolkit(AssetEditorToolkit);
+	TransitionGraphEditor->RegisterEditor();
 }
 
 void SCameraRigAssetEditor::DiscardGraphEditors()
@@ -147,6 +149,9 @@ void SCameraRigAssetEditor::DiscardGraphEditors()
 		}
 	}
 
+	NodeGraphEditor->UnregisterEditor();
+	TransitionGraphEditor->UnregisterEditor();
+
 	NodeGraphChangedHandle.Reset();
 	TransitionGraphChangedHandle.Reset();
 
diff --git a/Engine/Plugins/Cameras/GameplayCameras/Source/GameplayCamerasEditor/Private/Editors/SCameraRigTransitionEditor.cpp b/Engine/Plugins/Cameras/GameplayCameras/Source/GameplayCamerasEditor/Private/Editors/SCameraRigTransitionEditor.cpp
--- a/Engine/Plugins/Cameras/GameplayCameras/Source/GameplayCamerasEditor/Private/Editors/SCameraRigTransitionEditor.cpp
+++ b/Engine/Plugins/Cameras/GameplayCameras/Source/GameplayCamerasEditor/Private/Editors/SCameraRigTransitionEditor.cpp
@@ -87,6 +87,7 @@ void SCameraRigTransitionEditor::CreateTransitionGraphEditor()
 		.GraphTitle(this, &SCameraRigTransitionEditor::GetTransitionOwnerName)
 		.GraphToEdit(TransitionGraph)
 		.AssetEditorToolkit(AssetEditorToolkit);
+	TransitionGraphEditor->RegisterEditor();
 }
 
 void SCameraRigTransitionEditor::DiscardTransitionGraphEditor()
@@ -101,6 +102,7 @@ void SCameraRigTransitionEditor::DiscardTransitionGraphEditor()
 		}
 	}
 
+	TransitionGraphEditor->UnregisterEditor();
 	TransitionGraph = nullptr;
 	TransitionGraphChangedHandle.Reset();
 
diff --git a/Engine/Plugins/Cameras/GameplayCameras/Source/GameplayCamerasEditor/Private/Editors/SObjectTreeGraphCommentNode.cpp b/Engine/Plugins/Cameras/GameplayCameras/Source/GameplayCamerasEditor/Private/Editors/SObjectTreeGraphCommentNode.cpp
new file mode 100644
--- /dev/null
+++ b/Engine/Plugins/Cameras/GameplayCameras/Source/GameplayCamerasEditor/Private/Editors/SObjectTreeGraphCommentNode.cpp
@@ -0,0 +1,36 @@
+// Copyright Epic Games, Inc. All Rights Reserved.
+
+#include "Editors/SObjectTreeGraphCommentNode.h"
+
+#include "Editors/ObjectTreeGraphCommentNode.h"
+#include "Editors/ObjectTreeGraphNode.h"
+
+void SObjectTreeGraphCommentNode::Construct(const FArguments& InArgs)
+{
+	SGraphNodeComment::Construct(SGraphNodeComment::FArguments(), InArgs._GraphNode);
+
+	ObjectGraphNode = InArgs._GraphNode;
+}
+
+void SObjectTreeGraphCommentNode::MoveTo(const FVector2D& NewPosition, FNodeSet& NodeFilter, bool bMarkDirty)
+{
+	SGraphNodeComment::MoveTo(NewPosition, NodeFilter, bMarkDirty);
+
+	if (ObjectGraphNode)
+	{
+		ObjectGraphNode->OnGraphNodeMoved(bMarkDirty);
+
+		if (ObjectGraphNode && ObjectGraphNode->MoveMode == ECommentBoxMode::GroupMovement)
+		{
+			// Also notify any nodes that were moved along.
+			for (FCommentNodeSet::TConstIterator NodeIt(ObjectGraphNode->GetNodesUnderComment()); NodeIt; ++NodeIt)
+			{
+				if (UObjectTreeGraphNode* Node = Cast<UObjectTreeGraphNode>(*NodeIt))
+				{
+					Node->OnGraphNodeMoved(bMarkDirty);
+				}
+			}
+		}
+	}
+}
+
diff --git a/Engine/Plugins/Cameras/GameplayCameras/Source/GameplayCamerasEditor/Private/Editors/SObjectTreeGraphEditor.cpp b/Engine/Plugins/Cameras/GameplayCameras/Source/GameplayCamerasEditor/Private/Editors/SObjectTreeGraphEditor.cpp
--- a/Engine/Plugins/Cameras/GameplayCameras/Source/GameplayCamerasEditor/Private/Editors/SObjectTreeGraphEditor.cpp
+++ b/Engine/Plugins/Cameras/GameplayCameras/Source/GameplayCamerasEditor/Private/Editors/SObjectTreeGraphEditor.cpp
@@ -25,6 +25,25 @@
 
 #define LOCTEXT_NAMESPACE "SObjectTreeGraphEditor"
 
+TMap<UObjectTreeGraph*, TSharedPtr<SObjectTreeGraphEditor>> SObjectTreeGraphEditor::ActiveGraphEditors;
+
+TSharedPtr<SObjectTreeGraphEditor> SObjectTreeGraphEditor::FindGraphEditor(UObjectTreeGraph* InGraph)
+{
+	return ActiveGraphEditors.FindRef(InGraph);
+}
+
+void SObjectTreeGraphEditor::OnBeginEditingGraph(UObjectTreeGraph* InGraph, TSharedRef<SObjectTreeGraphEditor> InGraphEditor)
+{
+	ActiveGraphEditors.Add(InGraph, InGraphEditor);
+}
+
+void SObjectTreeGraphEditor::OnEndEditingGraph(UObjectTreeGraph* InGraph, TSharedRef<SObjectTreeGraphEditor> InGraphEditor)
+{
+	TSharedPtr<SObjectTreeGraphEditor> RemovedGraphEditor;
+	ActiveGraphEditors.RemoveAndCopyValue(InGraph, RemovedGraphEditor);
+	ensure(RemovedGraphEditor == InGraphEditor);
+}
+
 void SObjectTreeGraphEditor::Construct(const FArguments& InArgs)
 {
 	DetailsView = InArgs._DetailsView;
@@ -75,6 +94,22 @@ SObjectTreeGraphEditor::~SObjectTreeGraphEditor()
 	GEditor->UnregisterForUndo(this);
 }
 
+void SObjectTreeGraphEditor::RegisterEditor()
+{
+	if (UObjectTreeGraph* CurrentGraph = Cast<UObjectTreeGraph>(GraphEditor->GetCurrentGraph()))
+	{
+		OnBeginEditingGraph(CurrentGraph, SharedThis(this));
+	}
+}
+
+void SObjectTreeGraphEditor::UnregisterEditor()
+{
+	if (UObjectTreeGraph* CurrentGraph = Cast<UObjectTreeGraph>(GraphEditor->GetCurrentGraph()))
+	{
+		OnEndEditingGraph(CurrentGraph, SharedThis(this));
+	}
+}
+
 void SObjectTreeGraphEditor::InitializeBuiltInCommands()
 {
 	using namespace UE::Cameras;
@@ -351,14 +386,14 @@ bool SObjectTreeGraphEditor::CanImportNodesFromText(const FString& TextToImport)
 	return Schema->CanImportNodesFromText(CurrentGraph, TextToImport);
 }
 
-void SObjectTreeGraphEditor::DeleteNodes(TArrayView<UObjectTreeGraphNode*> NodesToDelete)
+void SObjectTreeGraphEditor::DeleteNodes(TArrayView<UEdGraphNode*> NodesToDelete)
 {
 	UEdGraph* CurrentGraph = GraphEditor->GetCurrentGraph();
 	const UEdGraphSchema* Schema = CurrentGraph->GetSchema();
 
 	const FScopedTransaction Transaction(LOCTEXT("DeleteNode", "Delete Node(s)"));
 
-	for (UObjectTreeGraphNode* Node : NodesToDelete)
+	for (UEdGraphNode* Node : NodesToDelete)
 	{
 		if (Node)
 		{
@@ -381,7 +416,7 @@ bool SObjectTreeGraphEditor::CanSelectAllNodes()
 
 void SObjectTreeGraphEditor::DeleteSelectedNodes()
 {
-	TArray<UObjectTreeGraphNode*> NodesToDelete;
+	TArray<UEdGraphNode*> NodesToDelete;
 	const FGraphPanelSelectionSet SelectedNodes = GraphEditor->GetSelectedNodes();
 
 	for (FGraphPanelSelectionSet::TConstIterator NodeIt(SelectedNodes); NodeIt; ++NodeIt)
@@ -389,7 +424,7 @@ void SObjectTreeGraphEditor::DeleteSelectedNodes()
 		UEdGraphNode* GraphNode = Cast<UEdGraphNode>(*NodeIt);
 		if (GraphNode && GraphNode->CanUserDeleteNode())
 		{
-			NodesToDelete.Add(Cast<UObjectTreeGraphNode>(*NodeIt));
+			NodesToDelete.Add(GraphNode);
 		}
 	}
 	
diff --git a/Engine/Plugins/Cameras/GameplayCameras/Source/GameplayCamerasEditor/Private/Toolkits/CameraShakeAssetEditorToolkit.cpp b/Engine/Plugins/Cameras/GameplayCameras/Source/GameplayCamerasEditor/Private/Toolkits/CameraShakeAssetEditorToolkit.cpp
--- a/Engine/Plugins/Cameras/GameplayCameras/Source/GameplayCamerasEditor/Private/Toolkits/CameraShakeAssetEditorToolkit.cpp
+++ b/Engine/Plugins/Cameras/GameplayCameras/Source/GameplayCamerasEditor/Private/Toolkits/CameraShakeAssetEditorToolkit.cpp
@@ -65,6 +65,10 @@ FCameraShakeAssetEditorToolkit::FCameraShakeAssetEditorToolkit(UAssetEditor* InO
 
 FCameraShakeAssetEditorToolkit::~FCameraShakeAssetEditorToolkit()
 {
+	if (!GExitPurge)
+	{
+		DiscardNodeGraphEditor();
+	}
 }
 
 void FCameraShakeAssetEditorToolkit::AddReferencedObjects(FReferenceCollector& Collector)
@@ -274,6 +278,12 @@ void FCameraShakeAssetEditorToolkit::CreateNodeGraphEditor()
 		.IsEnabled(this, &FCameraShakeAssetEditorToolkit::IsGraphEditorEnabled)
 		.GraphToEdit(NodeGraph)
 		.AssetEditorToolkit(SharedThis(this));
+	NodeGraphEditor->RegisterEditor();
+}
+
+void FCameraShakeAssetEditorToolkit::DiscardNodeGraphEditor()
+{
+	NodeGraphEditor->UnregisterEditor();
 }
 
 void FCameraShakeAssetEditorToolkit::RegisterToolbar()
diff --git a/Engine/Plugins/Cameras/GameplayCameras/Source/GameplayCamerasEditor/Public/Editors/ObjectTreeGraphCommentNode.h b/Engine/Plugins/Cameras/GameplayCameras/Source/GameplayCamerasEditor/Public/Editors/ObjectTreeGraphCommentNode.h
new file mode 100644
--- /dev/null
+++ b/Engine/Plugins/Cameras/GameplayCameras/Source/GameplayCamerasEditor/Public/Editors/ObjectTreeGraphCommentNode.h
@@ -0,0 +1,43 @@
+// Copyright Epic Games, Inc. All Rights Reserved.
+
+#pragma once
+
+#include "EdGraphNode_Comment.h"
+#include "UObject/WeakObjectPtrFwd.h"
+
+#include "ObjectTreeGraphCommentNode.generated.h"
+
+class UObjectTreeGraphComment;
+
+UCLASS(MinimalAPI, Optional)
+class UObjectTreeGraphCommentNode : public UEdGraphNode_Comment
+{
+	GENERATED_BODY()
+
+public:
+
+	/** Initializes this graph node for the given comment object. */
+	void Initialize(UObjectTreeGraphComment* InObject);
+
+	/** Gets the underlying comment object represented by this graph node. */
+	UObjectTreeGraphComment* GetObject() const { return WeakObject.Get(); }
+
+public:
+
+	// UEdGraphNode interface.
+	virtual TSharedPtr<SGraphNode> CreateVisualWidget() override;
+	virtual void GetNodeContextMenuActions(UToolMenu* Menu, UGraphNodeContextMenuContext* Context) const override;
+	virtual void PostPlacedNewNode() override;
+	virtual void ResizeNode(const FVector2D& NewSize) override;
+	virtual void OnRenameNode(const FString& NewName) override;
+
+public:
+
+	void OnGraphNodeMoved(bool bMarkDirty);
+
+private:
+
+	UPROPERTY()
+	TWeakObjectPtr<UObjectTreeGraphComment> WeakObject;
+};
+
diff --git a/Engine/Plugins/Cameras/GameplayCameras/Source/GameplayCamerasEditor/Public/Editors/ObjectTreeGraphSchema.h b/Engine/Plugins/Cameras/GameplayCameras/Source/GameplayCamerasEditor/Public/Editors/ObjectTreeGraphSchema.h
--- a/Engine/Plugins/Cameras/GameplayCameras/Source/GameplayCamerasEditor/Public/Editors/ObjectTreeGraphSchema.h
+++ b/Engine/Plugins/Cameras/GameplayCameras/Source/GameplayCamerasEditor/Public/Editors/ObjectTreeGraphSchema.h
@@ -11,6 +11,7 @@
 class IObjectTreeGraphRootObject;
 class UEdGraph;
 class UObjectTreeGraph;
+class UObjectTreeGraphComment;
 class UObjectTreeGraphNode;
 struct FObjectTreeGraphClassConfigs;
 
@@ -42,13 +43,13 @@ public:
 	void RebuildGraph(UObjectTreeGraph* InGraph) const;
 
 	/** Creates an object graph node for the given object. */
-	UObjectTreeGraphNode* CreateObjectNode(UObjectTreeGraph* InGraph, UObject* InObject) const;
+	UEdGraphNode* CreateObjectNode(UObjectTreeGraph* InGraph, UObject* InObject) const;
 
 	/** Adds an object to the underlying data after it has been added to the graph. */
-	void AddConnectableObject(UObjectTreeGraph* InGraph, UObjectTreeGraphNode* InNewNode) const;
+	void AddConnectableObject(UObjectTreeGraph* InGraph, UObject* InNewObject) const;
 
 	/** Removes an object from the underlying data after it has been removed from the graph. */
-	void RemoveConnectableObject(UObjectTreeGraph* InGraph, UObjectTreeGraphNode* InRemovedNode) const;
+	void RemoveConnectableObject(UObjectTreeGraph* InGraph, UObject* InRemovedObject) const;
 
 	/** Export the given selection into a text suitable for copy/pasting. */
 	FString ExportNodesToText(const FGraphPanelSelectionSet& Nodes, bool bOnlyCanDuplicateNodes, bool bOnlyCanDeleteNodes) const;
@@ -76,6 +77,7 @@ public:
 	// UEdGraphSchema interface.
 	virtual void GetGraphContextActions(FGraphContextMenuBuilder& ContextMenuBuilder) const override;
 	virtual void GetContextMenuActions(class UToolMenu* Menu, class UGraphNodeContextMenuContext* Context) const override;
+	virtual TSharedPtr<FEdGraphSchemaAction> GetCreateCommentAction() const override;
 	virtual FName GetParentContextMenuName() const override;
 	virtual FLinearColor GetPinTypeColor(const FEdGraphPinType& PinType) const override;
 	virtual class FConnectionDrawingPolicy* CreateConnectionDrawingPolicy(int32 InBackLayerID, int32 InFrontLayerID, float InZoomFactor, const FSlateRect& InClippingRect, class FSlateWindowElementList& InDrawElements, UEdGraph* InGraph) const override;
@@ -94,15 +96,15 @@ protected:
 
 	struct FCreatedNodes
 	{
-		TMap<UObject*, UObjectTreeGraphNode*> CreatedNodes;
+		TMap<UObject*, UEdGraphNode*> CreatedNodes;
 	};
 
 	// UObjectTreeGraphSchema interface.
 	virtual void CollectAllObjects(UObjectTreeGraph* InGraph, TSet<UObject*>& OutAllObjects) const;
 	virtual void OnCreateAllNodes(UObjectTreeGraph* InGraph, const FCreatedNodes& InCreatedNodes) const;
-	virtual UObjectTreeGraphNode* OnCreateObjectNode(UObjectTreeGraph* InGraph, UObject* InObject) const;
-	virtual void OnAddConnectableObject(UObjectTreeGraph* InGraph, UObjectTreeGraphNode* InNewNode) const;
-	virtual void OnRemoveConnectableObject(UObjectTreeGraph* InGraph, UObjectTreeGraphNode* InRemovedNode) const;
+	virtual UEdGraphNode* OnCreateObjectNode(UObjectTreeGraph* InGraph, UObject* InObject) const;
+	virtual void OnAddConnectableObject(UObjectTreeGraph* InGraph, UObject* InNewObject) const;
+	virtual void OnRemoveConnectableObject(UObjectTreeGraph* InGraph, UObject* InRemovedObject) const;
 	virtual void CopyNonObjectNodes(TArrayView<UObject*> InObjects, FStringOutputDevice& OutDevice) const;
 	virtual bool OnTryCreateCustomConnection(UEdGraphPin* A, UEdGraphPin* B) const;
 	virtual bool OnBreakCustomPinLinks(UEdGraphPin& TargetPin) const;
@@ -130,13 +132,16 @@ private:
 	void RemoveAllNodes(UObjectTreeGraph* InGraph) const;
 	void CreateAllNodes(UObjectTreeGraph* InGraph) const;
 	void CreateConnections(UObjectTreeGraphNode* InGraphNode, const FCreatedNodes& InCreatedNodes) const;
+
+	UEdGraphNode* CreateCommentNode(UObjectTreeGraph* InGraph, UObjectTreeGraphComment* InComment) const;
+	void GetCommentAction(FGraphActionMenuBuilder& ActionMenuBuilder) const;
 };
 
 /**
  * Graph action to create a new object (and corresponding graph node) of a given class.
  */
 USTRUCT()
-struct FObjectGraphSchemaAction_NewNode : public FEdGraphSchemaAction
+struct FObjectTreeGraphSchemaAction_NewNode : public FEdGraphSchemaAction
 {
 	GENERATED_BODY()
 
@@ -152,19 +157,37 @@ public:
 
 public:
 
-	FObjectGraphSchemaAction_NewNode();
-	FObjectGraphSchemaAction_NewNode(FText InNodeCategory, FText InMenuDesc, FText InToolTip, const int32 InGrouping = 0, FText InKeywords = FText());
+	FObjectTreeGraphSchemaAction_NewNode();
+	FObjectTreeGraphSchemaAction_NewNode(FText InNodeCategory, FText InMenuDesc, FText InToolTip, const int32 InGrouping = 0, FText InKeywords = FText());
 
 public:
 
 	// FEdGraphSchemaAction interface.
-	static FName StaticGetTypeId() { static FName Type("FObjectGraphSchemaAction_NewNode"); return Type; }
+	static FName StaticGetTypeId() { static FName Type("FObjectTreeGraphSchemaAction_NewNode"); return Type; }
 	virtual FName GetTypeId() const override { return StaticGetTypeId(); } 
 	virtual UEdGraphNode* PerformAction(UEdGraph* ParentGraph, UEdGraphPin* FromPin, const FVector2D Location, bool bSelectNewNode = true) override;
 
 protected:
 
 	virtual UObject* CreateObject();
-	virtual void AutoSetupNewNode(UObjectTreeGraphNode* NewNode, UEdGraphPin* FromPin);
+	virtual void AutoSetupNewNode(UEdGraphNode* NewNode, UEdGraphPin* FromPin);
+};
+
+USTRUCT()
+struct FObjectTreeGraphSchemaAction_NewComment : public FEdGraphSchemaAction
+{
+	GENERATED_BODY()
+
+public:
+
+	FObjectTreeGraphSchemaAction_NewComment();
+	FObjectTreeGraphSchemaAction_NewComment(FText InNodeCategory, FText InMenuDesc, FText InToolTip, const int32 InGrouping = 0, FText InKeywords = FText());
+
+public:
+
+	// FEdGraphSchemaAction interface.
+	static FName StaticGetTypeId() { static FName Type("FObjectTreeGraphSchemaAction_NewComment"); return Type; }
+	virtual FName GetTypeId() const override { return StaticGetTypeId(); } 
+	virtual UEdGraphNode* PerformAction(UEdGraph* ParentGraph, UEdGraphPin* FromPin, const FVector2D Location, bool bSelectNewNode = true) override;
 };
 
diff --git a/Engine/Plugins/Cameras/GameplayCameras/Source/GameplayCamerasEditor/Public/Editors/SObjectTreeGraphCommentNode.h b/Engine/Plugins/Cameras/GameplayCameras/Source/GameplayCamerasEditor/Public/Editors/SObjectTreeGraphCommentNode.h
new file mode 100644
--- /dev/null
+++ b/Engine/Plugins/Cameras/GameplayCameras/Source/GameplayCamerasEditor/Public/Editors/SObjectTreeGraphCommentNode.h
@@ -0,0 +1,35 @@
+// Copyright Epic Games, Inc. All Rights Reserved.
+
+#pragma once
+
+#include "SGraphNodeComment.h"
+
+class UObjectTreeGraphCommentNode;
+
+/**
+ * The widget for object tree graph comment nodes.
+ */
+class SObjectTreeGraphCommentNode : public SGraphNodeComment
+{
+public:
+
+	SLATE_BEGIN_ARGS(SObjectTreeGraphCommentNode)
+		: _GraphNode(nullptr)
+	{}
+		SLATE_ARGUMENT(UObjectTreeGraphCommentNode*, GraphNode)
+	SLATE_END_ARGS()
+
+	void Construct(const FArguments& InArgs);
+
+	UObjectTreeGraphCommentNode* GetObjectGraphNode() const { return ObjectGraphNode; }
+
+public:
+
+	// SNodePanel::SNode interface.
+	virtual void MoveTo(const FVector2D& NewPosition, FNodeSet& NodeFilter, bool bMarkDirty=true) override;
+
+protected:
+
+	UObjectTreeGraphCommentNode* ObjectGraphNode;
+};
+
diff --git a/Engine/Plugins/Cameras/GameplayCameras/Source/GameplayCamerasEditor/Public/Editors/SObjectTreeGraphEditor.h b/Engine/Plugins/Cameras/GameplayCameras/Source/GameplayCamerasEditor/Public/Editors/SObjectTreeGraphEditor.h
--- a/Engine/Plugins/Cameras/GameplayCameras/Source/GameplayCamerasEditor/Public/Editors/SObjectTreeGraphEditor.h
+++ b/Engine/Plugins/Cameras/GameplayCameras/Source/GameplayCamerasEditor/Public/Editors/SObjectTreeGraphEditor.h
@@ -9,7 +9,6 @@
 class FAssetEditorToolkit;
 class IDetailsView;
 class UObjectTreeGraph;
-class UObjectTreeGraphNode;
 
 /**
  * A graph editor for an object tree graph.
@@ -18,6 +17,16 @@ class SObjectTreeGraphEditor
 	: public SCompoundWidget
 	, public FEditorUndoClient
 {
+public:
+
+	static TSharedPtr<SObjectTreeGraphEditor> FindGraphEditor(UObjectTreeGraph* InGraph);
+	static void OnBeginEditingGraph(UObjectTreeGraph* InGraph, TSharedRef<SObjectTreeGraphEditor> InGraphEditor);
+	static void OnEndEditingGraph(UObjectTreeGraph* InGraph, TSharedRef<SObjectTreeGraphEditor> InGraphEditor);
+
+private:
+
+	static TMap<UObjectTreeGraph*, TSharedPtr<SObjectTreeGraphEditor>> ActiveGraphEditors;
+
 public:
 
 	SLATE_BEGIN_ARGS(SObjectTreeGraphEditor)
@@ -41,9 +50,14 @@ public:
 	void Construct(const FArguments& InArgs);
 	~SObjectTreeGraphEditor();
 
+	void RegisterEditor();
+	void UnregisterEditor();
+
 	void JumpToNode(UEdGraphNode* InNode);
 	void ResyncDetailsView();
 
+	TSharedPtr<SGraphEditor> GetGraphEditor() { return GraphEditor; }
+
 protected:
 
 	// SWidget interface.
@@ -66,7 +80,7 @@ protected:
 	FString ExportNodesToText(const FGraphPanelSelectionSet& Nodes, bool bOnlyCanDuplicateNodes, bool bOnlyCanDeleteNodes);
 	void ImportNodesFromText(const FVector2D& Location, const FString& TextToImport);
 	bool CanImportNodesFromText(const FString& TextToImport);
-	void DeleteNodes(TArrayView<UObjectTreeGraphNode*> NodesToDelete);
+	void DeleteNodes(TArrayView<UEdGraphNode*> NodesToDelete);
 
 	void SelectAllNodes();
 	bool CanSelectAllNodes();
diff --git a/Engine/Plugins/Cameras/GameplayCameras/Source/GameplayCamerasEditor/Public/Toolkits/CameraShakeAssetEditorToolkit.h b/Engine/Plugins/Cameras/GameplayCameras/Source/GameplayCamerasEditor/Public/Toolkits/CameraShakeAssetEditorToolkit.h
--- a/Engine/Plugins/Cameras/GameplayCameras/Source/GameplayCamerasEditor/Public/Toolkits/CameraShakeAssetEditorToolkit.h
+++ b/Engine/Plugins/Cameras/GameplayCameras/Source/GameplayCamerasEditor/Public/Toolkits/CameraShakeAssetEditorToolkit.h
@@ -73,6 +73,8 @@ private:
 	TSharedRef<SDockTab> SpawnTab_InterfaceParameters(const FSpawnTabArgs& Args);
 
 	void CreateNodeGraphEditor();
+	void DiscardNodeGraphEditor();
+
 	FText GetCameraShakeAssetName() const;
 	bool IsGraphEditorEnabled() const;
 
