Cameras: add suport for overscan and projection offsets

diff --git a/Engine/Plugins/Cameras/GameplayCameras/Source/GameplayCameras/Private/Core/CameraPose.cpp b/Engine/Plugins/Cameras/GameplayCameras/Source/GameplayCameras/Private/Core/CameraPose.cpp
--- a/Engine/Plugins/Cameras/GameplayCameras/Source/GameplayCameras/Private/Core/CameraPose.cpp
+++ b/Engine/Plugins/Cameras/GameplayCameras/Source/GameplayCameras/Private/Core/CameraPose.cpp
@@ -118,12 +118,12 @@ void FCameraPose::SetTransform(FTransform3d Transform, bool bForceSet)
 	SetRotation(Transform.GetRotation().Rotator(), bForceSet);
 }
 
-double FCameraPose::GetEffectiveFieldOfView() const
+double FCameraPose::GetEffectiveFieldOfView(bool bIncludeOverscan) const
 {
-	return GetEffectiveFieldOfView(FocalLength, FieldOfView, SensorWidth, SensorHeight, SqueezeFactor);
+	return GetEffectiveFieldOfView(FocalLength, FieldOfView, SensorWidth, SensorHeight, SqueezeFactor, bIncludeOverscan ? Overscan : 1.0f);
 }
 
-double FCameraPose::GetEffectiveFieldOfView(float FocalLength, float FieldOfView, float SensorWidth, float SensorHeight, float SqueezeFactor)
+double FCameraPose::GetEffectiveFieldOfView(float FocalLength, float FieldOfView, float SensorWidth, float SensorHeight, float SqueezeFactor, float Overscan)
 {
 	const bool bValidFocalLength = (FocalLength > 0.f);
 	const bool bValidFieldOfView = (FieldOfView > 0.f);
@@ -166,7 +166,9 @@ double FCameraPose::GetEffectiveFieldOfView(float FocalLength, float FieldOfView
 			}
 		}
 
-		return FMath::RadiansToDegrees(2.0 * FMath::Atan(CroppedSensorWidth / (2.0 * FocalLength)));
+		const double EffectiveOverscan = (1.0 + Overscan);
+
+		return FMath::RadiansToDegrees(2.0 * FMath::Atan(CroppedSensorWidth * EffectiveOverscan / (2.0 * FocalLength)));
 	}
 	else
 	{
@@ -191,6 +193,44 @@ void FCameraPose::GetDefaultSensorSize(float& OutSensorWidth, float& OutSensorHe
 	OutSensorHeight = 18.67f;
 }
 
+double FCameraPose::GetHorizontalProjectionOffset() const
+{
+	// Compute projection offset with similar code to UCineCameraComponent...
+	double CroppedSensorWidth = SensorWidth * SqueezeFactor;
+	const double AspectRatio = GetSensorAspectRatio(SensorWidth, SensorHeight);
+	if (AspectRatio > 0.0)
+	{
+		double DesqueezeAspectRatio = AspectRatio * SqueezeFactor;
+		if (AspectRatio < DesqueezeAspectRatio)
+		{
+			CroppedSensorWidth *= AspectRatio / DesqueezeAspectRatio;
+		}
+	}
+
+	const double EffectiveOverscan = (1.0f + Overscan);
+
+	return 2.0 * SensorHorizontalOffset / (CroppedSensorWidth * EffectiveOverscan);
+}
+
+double FCameraPose::GetVerticalProjectionOffset() const
+{
+	// Compute projection offset with similar code to UCineCameraComponent...
+	double CroppedSensorHeight = SensorHeight;
+	const double AspectRatio = GetSensorAspectRatio(SensorWidth, SensorHeight);
+	if (AspectRatio > 0.0)
+	{
+		double DesqueezeAspectRatio = AspectRatio * SqueezeFactor;
+		if (DesqueezeAspectRatio < AspectRatio)
+		{
+			CroppedSensorHeight *= DesqueezeAspectRatio / AspectRatio;
+		}
+	}
+
+	const double EffectiveOverscan = (1.0f + Overscan);
+
+	return 2.0 * SensorVerticalOffset / (CroppedSensorHeight * EffectiveOverscan);
+}
+
 bool FCameraPose::ApplyPhysicalCameraSettings(FPostProcessSettings& PostProcessSettings, bool bOverwriteSettings) const
 {
 	if (!EnablePhysicalCamera || PhysicalCameraBlendWeight <= 0.f)
@@ -217,8 +257,8 @@ bool FCameraPose::ApplyPhysicalCameraSettings(FPostProcessSettings& PostProcessS
 	// TODO: support minimum-focus-distance?
 	UE_LERP_PP(DepthOfFieldFocalDistance, FocusDistance);
 
-	// TODO: support overscan?
-	UE_LERP_PP(DepthOfFieldSensorWidth, SensorWidth);
+	const float EffectiveOverscan = (1.0f + Overscan);
+	UE_LERP_PP(DepthOfFieldSensorWidth, SensorWidth * EffectiveOverscan);
 	UE_LERP_PP(DepthOfFieldSqueezeFactor, SqueezeFactor);
 
 #undef UE_LERP_PP
diff --git a/Engine/Plugins/Cameras/GameplayCameras/Source/GameplayCameras/Private/Core/CameraSystemEvaluator.cpp b/Engine/Plugins/Cameras/GameplayCameras/Source/GameplayCameras/Private/Core/CameraSystemEvaluator.cpp
--- a/Engine/Plugins/Cameras/GameplayCameras/Source/GameplayCameras/Private/Core/CameraSystemEvaluator.cpp
+++ b/Engine/Plugins/Cameras/GameplayCameras/Source/GameplayCameras/Private/Core/CameraSystemEvaluator.cpp
@@ -369,11 +369,16 @@ void FCameraSystemEvaluator::GetEvaluatedCameraView(FMinimalViewInfo& DesiredVie
 
 	DesiredView.PerspectiveNearClipPlane = CameraPose.GetNearClippingPlane();
 
+	DesiredView.OffCenterProjectionOffset.X = CameraPose.GetHorizontalProjectionOffset();
+	DesiredView.OffCenterProjectionOffset.Y = CameraPose.GetVerticalProjectionOffset();
+
 	const FPostProcessSettingsCollection& PostProcessSettings = RootNodeResult.PostProcessSettings;
 	DesiredView.PostProcessSettings = PostProcessSettings.Get();
 	DesiredView.PostProcessBlendWeight = 1.f;
 	// Create the physical camera settings if needed. Don't overwrite settings that were set by hand.
 	CameraPose.ApplyPhysicalCameraSettings(DesiredView.PostProcessSettings, false);
+
+	DesiredView.ApplyOverscan(CameraPose.GetOverscan());
 }
 
 #if WITH_EDITOR
diff --git a/Engine/Plugins/Cameras/GameplayCameras/Source/GameplayCameras/Private/GameFramework/GameplayCameraComponentBase.cpp b/Engine/Plugins/Cameras/GameplayCameras/Source/GameplayCameras/Private/GameFramework/GameplayCameraComponentBase.cpp
--- a/Engine/Plugins/Cameras/GameplayCameras/Source/GameplayCameras/Private/GameFramework/GameplayCameraComponentBase.cpp
+++ b/Engine/Plugins/Cameras/GameplayCameras/Source/GameplayCameras/Private/GameFramework/GameplayCameraComponentBase.cpp
@@ -425,8 +425,13 @@ void UGameplayCameraComponentBase::UpdateOutputCameraComponent()
 			OutputCameraComponent->SetWorldTransform(Result.CameraPose.GetTransform());
 			OutputCameraComponent->SetFieldOfView(Result.CameraPose.GetEffectiveFieldOfView());
 			OutputCameraComponent->CurrentAperture = Result.CameraPose.GetAperture();
+
 			OutputCameraComponent->Filmback.SensorWidth = Result.CameraPose.GetSensorWidth();
 			OutputCameraComponent->Filmback.SensorHeight = Result.CameraPose.GetSensorHeight();
+			OutputCameraComponent->Filmback.SensorHorizontalOffset = Result.CameraPose.GetSensorHorizontalOffset();
+			OutputCameraComponent->Filmback.SensorVerticalOffset = Result.CameraPose.GetSensorVerticalOffset();
+
+			OutputCameraComponent->Overscan = Result.CameraPose.GetOverscan();
 			OutputCameraComponent->bConstrainAspectRatio = Result.CameraPose.GetConstrainAspectRatio();
 			OutputCameraComponent->bOverrideAspectRatioAxisConstraint = Result.CameraPose.GetOverrideAspectRatioAxisConstraint();
 			OutputCameraComponent->AspectRatioAxisConstraint = Result.CameraPose.GetAspectRatioAxisConstraint();
diff --git a/Engine/Plugins/Cameras/GameplayCameras/Source/GameplayCameras/Public/Core/CameraPose.h b/Engine/Plugins/Cameras/GameplayCameras/Source/GameplayCameras/Public/Core/CameraPose.h
--- a/Engine/Plugins/Cameras/GameplayCameras/Source/GameplayCameras/Public/Core/CameraPose.h
+++ b/Engine/Plugins/Cameras/GameplayCameras/Source/GameplayCameras/Public/Core/CameraPose.h
@@ -26,8 +26,11 @@ namespace ECameraProjectionMode { enum Type : int; };
 	UE_CAMERA_POSE_FOR_PROPERTY(float,  FocusDistance)\
 	UE_CAMERA_POSE_FOR_PROPERTY(float,  SensorWidth)\
 	UE_CAMERA_POSE_FOR_PROPERTY(float,  SensorHeight)\
+	UE_CAMERA_POSE_FOR_PROPERTY(float,  SensorHorizontalOffset)\
+	UE_CAMERA_POSE_FOR_PROPERTY(float,  SensorVerticalOffset)\
 	UE_CAMERA_POSE_FOR_PROPERTY(float,  ISO)\
 	UE_CAMERA_POSE_FOR_PROPERTY(float,  SqueezeFactor)\
+	UE_CAMERA_POSE_FOR_PROPERTY(float,  Overscan)\
 	UE_CAMERA_POSE_FOR_PROPERTY(int32,  DiaphragmBladeCount)\
 	UE_CAMERA_POSE_FOR_PROPERTY(float,  PhysicalCameraBlendWeight)\
 	UE_CAMERA_POSE_FOR_PROPERTY(float,  NearClippingPlane)\
@@ -155,11 +158,17 @@ public:
 	 * The effective field of view can be driven by the FieldOfView property, or
 	 * the FocalLength property in combination with the sensor size.
 	 */
-	double GetEffectiveFieldOfView() const;
+	double GetEffectiveFieldOfView(bool bIncludeOverscan = true) const;
 
 	/** Gets the aspect ratio of the camera sensor. */
 	double GetSensorAspectRatio() const;
 
+	/** Gets the horizontal projection offset. */
+	double GetHorizontalProjectionOffset() const;
+
+	/** Gets the vertical projection offset. */
+	double GetVerticalProjectionOffset() const;
+
 	/** Gets the aiming ray of the camera. */
 	FRay3d GetAimRay() const;
 
@@ -175,7 +184,7 @@ public:
 public:
 
 	/** Computes the horizontal field of view of a camera. */
-	static double GetEffectiveFieldOfView(float FocalLength, float FieldOfView, float SensorWidth, float SensorHeight, float SqueezeFactor);
+	static double GetEffectiveFieldOfView(float FocalLength, float FieldOfView, float SensorWidth, float SensorHeight, float SqueezeFactor = 1.f, float Overscan = 1.f);
 	
 	/** Computes the aspect ratio of a camera sensor. */
 	static double GetSensorAspectRatio(float SensorWidth, float SensorHeight);
@@ -278,6 +287,14 @@ private:
 	UPROPERTY()
 	float SensorHeight = 18.67f;
 
+	/** Horizontal offset of the sensor, in mm. */
+	UPROPERTY()
+	float SensorHorizontalOffset = 0.f;
+
+	/** Vertical offset of the sensor, in mm. */
+	UPROPERTY()
+	float SensorVerticalOffset = 0.f;
+
 	/** The camera sensor sensitivity in ISO. */
 	UPROPERTY()
 	float ISO = 100.f;
@@ -286,6 +303,10 @@ private:
 	UPROPERTY()
 	float SqueezeFactor = 1.f;
 
+	/** Sensor overscan in percentage */
+	UPROPERTY()
+	float Overscan = 0.f;
+
 	/** Number of blades in the lens diaphragm */
 	UPROPERTY()
 	int32 DiaphragmBladeCount = 8;
diff --git a/Engine/Plugins/Cameras/GameplayCameras/Source/GameplayCameras/Public/GameFramework/BlueprintCameraPose.h b/Engine/Plugins/Cameras/GameplayCameras/Source/GameplayCameras/Public/GameFramework/BlueprintCameraPose.h
--- a/Engine/Plugins/Cameras/GameplayCameras/Source/GameplayCameras/Public/GameFramework/BlueprintCameraPose.h
+++ b/Engine/Plugins/Cameras/GameplayCameras/Source/GameplayCameras/Public/GameFramework/BlueprintCameraPose.h
@@ -72,6 +72,14 @@ public:
 	UPROPERTY(BlueprintReadWrite, EditAnywhere, Category="Camera")
 	float SensorHeight = 18.67f;
 
+	/** Horizontal offset of the sensor, in mm. */
+	UPROPERTY(BlueprintReadWrite, EditAnywhere, Category="Camera")
+	float SensorHorizontalOffset = 0.f;
+
+	/** Vertical offset of the sensor, in mm. */
+	UPROPERTY(BlueprintReadWrite, EditAnywhere, Category="Camera")
+	float SensorVerticalOffset = 0.f;
+
 	/** The camera sensor sensitivity in ISO. */
 	UPROPERTY(BlueprintReadWrite, EditAnywhere, Category="Camera")
 	float ISO = 100.f;
@@ -80,6 +88,10 @@ public:
 	UPROPERTY(BlueprintReadWrite, EditAnywhere, Category="Camera")
 	float SqueezeFactor = 1.f;
 
+	/** Sensor overscan in percentage */
+	UPROPERTY(BlueprintReadWrite, EditAnywhere, Category="Camera")
+	float Overscan = 0.f;
+
 	/** Number of blades in the lens diaphragm */
 	UPROPERTY(BlueprintReadWrite, EditAnywhere, Category="Camera")
 	int32 DiaphragmBladeCount = 8;
