Cameras: preview orthographic mode in editor

diff --git a/Engine/Plugins/Cameras/GameplayCameras/Source/GameplayCameras/Private/GameFramework/GameplayCameraComponentBase.cpp b/Engine/Plugins/Cameras/GameplayCameras/Source/GameplayCameras/Private/GameFramework/GameplayCameraComponentBase.cpp
--- a/Engine/Plugins/Cameras/GameplayCameras/Source/GameplayCameras/Private/GameFramework/GameplayCameraComponentBase.cpp
+++ b/Engine/Plugins/Cameras/GameplayCameras/Source/GameplayCameras/Private/GameFramework/GameplayCameraComponentBase.cpp
@@ -439,6 +439,9 @@ void UGameplayCameraComponentBase::UpdateOutputCameraComponent()
 			OutputCameraComponent->FocusSettings.ManualFocusDistance = Result.CameraPose.GetFocusDistance();
 			OutputCameraComponent->FocusSettings.FocusMethod = (Result.CameraPose.GetEnablePhysicalCamera() ? ECameraFocusMethod::Manual : ECameraFocusMethod::Disable);
 
+			OutputCameraComponent->ProjectionMode = Result.CameraPose.GetProjectionMode();
+			OutputCameraComponent->OrthoWidth = Result.CameraPose.GetOrthographicWidth();
+
 			bGotValidTransform = true;
 		}
 	}
