Cameras: move shutter speed and ISO to new body parameters node, add projection offset settings

diff --git a/Engine/Plugins/Cameras/GameplayCameras/Source/GameplayCameras/Private/Nodes/Common/BodyParametersCameraNode.cpp b/Engine/Plugins/Cameras/GameplayCameras/Source/GameplayCameras/Private/Nodes/Common/BodyParametersCameraNode.cpp
new file mode 100644
--- /dev/null
+++ b/Engine/Plugins/Cameras/GameplayCameras/Source/GameplayCameras/Private/Nodes/Common/BodyParametersCameraNode.cpp
@@ -0,0 +1,60 @@
+// Copyright Epic Games, Inc. All Rights Reserved.
+
+#include "Nodes/Common/BodyParametersCameraNode.h"
+
+#include "Core/CameraParameterReader.h"
+#include "Core/CameraPose.h"
+
+#include UE_INLINE_GENERATED_CPP_BY_NAME(BodyParametersCameraNode)
+
+namespace UE::Cameras
+{
+
+class FBodyParametersCameraNodeEvaluator : public FCameraNodeEvaluator
+{
+	UE_DECLARE_CAMERA_NODE_EVALUATOR(GAMEPLAYCAMERAS_API, FBodyParametersCameraNodeEvaluator)
+
+protected:
+
+	virtual void OnInitialize(const FCameraNodeEvaluatorInitializeParams& Params, FCameraNodeEvaluationResult& OutResult) override;
+	virtual void OnRun(const FCameraNodeEvaluationParams& Params, FCameraNodeEvaluationResult& OutResult) override;
+
+private:
+
+	TCameraParameterReader<float> ShutterSpeedReader;
+	TCameraParameterReader<float> ISOReader;
+};
+
+UE_DEFINE_CAMERA_NODE_EVALUATOR(FBodyParametersCameraNodeEvaluator)
+
+void FBodyParametersCameraNodeEvaluator::OnInitialize(const FCameraNodeEvaluatorInitializeParams& Params, FCameraNodeEvaluationResult& OutResult)
+{
+	const UBodyParametersCameraNode* BodyParametersNode = GetCameraNodeAs<UBodyParametersCameraNode>();
+	ShutterSpeedReader.Initialize(BodyParametersNode->ShutterSpeed);
+	ISOReader.Initialize(BodyParametersNode->ISO);
+}
+
+void FBodyParametersCameraNodeEvaluator::OnRun(const FCameraNodeEvaluationParams& Params, FCameraNodeEvaluationResult& OutResult)
+{
+	FCameraPose& OutPose = OutResult.CameraPose;
+
+	float ShutterSpeed = ShutterSpeedReader.Get(OutResult.VariableTable);
+	if (ShutterSpeed > 0)
+	{
+		OutPose.SetShutterSpeed(ShutterSpeed);
+	}
+	float ISO = ISOReader.Get(OutResult.VariableTable);
+	if (ISO > 0)
+	{
+		OutPose.SetISO(ISO);
+	}
+}
+
+}  // namespace UE::Cameras
+
+FCameraNodeEvaluatorPtr UBodyParametersCameraNode::OnBuildEvaluator(FCameraNodeEvaluatorBuilder& Builder) const
+{
+	using namespace UE::Cameras;
+	return Builder.BuildEvaluator<FBodyParametersCameraNodeEvaluator>();
+}
+
diff --git a/Engine/Plugins/Cameras/GameplayCameras/Source/GameplayCameras/Private/Nodes/Common/FilmbackCameraNode.cpp b/Engine/Plugins/Cameras/GameplayCameras/Source/GameplayCameras/Private/Nodes/Common/FilmbackCameraNode.cpp
--- a/Engine/Plugins/Cameras/GameplayCameras/Source/GameplayCameras/Private/Nodes/Common/FilmbackCameraNode.cpp
+++ b/Engine/Plugins/Cameras/GameplayCameras/Source/GameplayCameras/Private/Nodes/Common/FilmbackCameraNode.cpp
@@ -23,7 +23,9 @@ private:
 
 	TCameraParameterReader<float> SensorWidthReader;
 	TCameraParameterReader<float> SensorHeightReader;
-	TCameraParameterReader<float> ISOReader;
+	TCameraParameterReader<float> SensorHorizontalOffsetReader;
+	TCameraParameterReader<float> SensorVerticalOffsetReader;
+	TCameraParameterReader<float> OverscanReader;
 	TCameraParameterReader<bool> ConstrainAspectRatioReader;
 	TCameraParameterReader<bool> OverrideAspectRatioAxisConstraintReader;
 };
@@ -35,7 +37,9 @@ void FFilmbackCameraNodeEvaluator::OnInitialize(const FCameraNodeEvaluatorInitia
 	const UFilmbackCameraNode* FilmbackNode = GetCameraNodeAs<UFilmbackCameraNode>();
 	SensorWidthReader.Initialize(FilmbackNode->SensorWidth);
 	SensorHeightReader.Initialize(FilmbackNode->SensorHeight);
-	ISOReader.Initialize(FilmbackNode->ISO);
+	SensorHorizontalOffsetReader.Initialize(FilmbackNode->SensorHorizontalOffset);
+	SensorVerticalOffsetReader.Initialize(FilmbackNode->SensorVerticalOffset);
+	OverscanReader.Initialize(FilmbackNode->Overscan);
 	ConstrainAspectRatioReader.Initialize(FilmbackNode->ConstrainAspectRatio);
 	OverrideAspectRatioAxisConstraintReader.Initialize(FilmbackNode->OverrideAspectRatioAxisConstraint);
 }
@@ -54,11 +58,18 @@ void FFilmbackCameraNodeEvaluator::OnRun(const FCameraNodeEvaluationParams& Para
 	{
 		OutPose.SetSensorHeight(SensorHeight);
 	}
-	float ISO = ISOReader.Get(OutResult.VariableTable);
-	if (ISO > 0)
+	float SensorHorizontalOffset = SensorHorizontalOffsetReader.Get(OutResult.VariableTable);
+	if (SensorHorizontalOffset > 0)
 	{
-		OutPose.SetISO(ISO);
+		OutPose.SetSensorHorizontalOffset(SensorHorizontalOffset);
 	}
+	float SensorVerticalOffset = SensorVerticalOffsetReader.Get(OutResult.VariableTable);
+	if (SensorVerticalOffset > 0)
+	{
+		OutPose.SetSensorVerticalOffset(SensorVerticalOffset);
+	}
+
+	OutPose.SetOverscan(OverscanReader.Get(OutResult.VariableTable));
 
 	OutPose.SetConstrainAspectRatio(ConstrainAspectRatioReader.Get(OutResult.VariableTable));
 	OutPose.SetOverrideAspectRatioAxisConstraint(OverrideAspectRatioAxisConstraintReader.Get(OutResult.VariableTable));
diff --git a/Engine/Plugins/Cameras/GameplayCameras/Source/GameplayCameras/Private/Nodes/Common/LensParametersCameraNode.cpp b/Engine/Plugins/Cameras/GameplayCameras/Source/GameplayCameras/Private/Nodes/Common/LensParametersCameraNode.cpp
--- a/Engine/Plugins/Cameras/GameplayCameras/Source/GameplayCameras/Private/Nodes/Common/LensParametersCameraNode.cpp
+++ b/Engine/Plugins/Cameras/GameplayCameras/Source/GameplayCameras/Private/Nodes/Common/LensParametersCameraNode.cpp
@@ -25,7 +25,6 @@ private:
 	TCameraParameterReader<float> FocalLengthReader;
 	TCameraParameterReader<float> FocusDistanceReader;
 	TCameraParameterReader<float> ApertureReader;
-	TCameraParameterReader<float> ShutterSpeedReader;
 	TCameraParameterReader<bool> EnablePhysicalCameraReader;
 };
 
@@ -37,7 +36,6 @@ void FLensParametersCameraNodeEvaluator::OnInitialize(const FCameraNodeEvaluator
 	FocalLengthReader.Initialize(LensParametersNode->FocalLength);
 	FocusDistanceReader.Initialize(LensParametersNode->FocusDistance);
 	ApertureReader.Initialize(LensParametersNode->Aperture);
-	ShutterSpeedReader.Initialize(LensParametersNode->ShutterSpeed);
 	EnablePhysicalCameraReader.Initialize(LensParametersNode->EnablePhysicalCamera);
 }
 
@@ -61,11 +59,6 @@ void FLensParametersCameraNodeEvaluator::OnRun(const FCameraNodeEvaluationParams
 	{
 		OutPose.SetAperture(Aperture);
 	}
-	float ShutterSpeed = ShutterSpeedReader.Get(OutResult.VariableTable);
-	if (ShutterSpeed > 0)
-	{
-		OutPose.SetShutterSpeed(ShutterSpeed);
-	}
 
 	const bool bEnablePhysicalCamera = EnablePhysicalCameraReader.Get(OutResult.VariableTable);
 	OutPose.SetEnablePhysicalCamera(bEnablePhysicalCamera);
diff --git a/Engine/Plugins/Cameras/GameplayCameras/Source/GameplayCameras/Public/Nodes/Common/BodyParametersCameraNode.h b/Engine/Plugins/Cameras/GameplayCameras/Source/GameplayCameras/Public/Nodes/Common/BodyParametersCameraNode.h
new file mode 100644
--- /dev/null
+++ b/Engine/Plugins/Cameras/GameplayCameras/Source/GameplayCameras/Public/Nodes/Common/BodyParametersCameraNode.h
@@ -0,0 +1,33 @@
+// Copyright Epic Games, Inc. All Rights Reserved.
+
+#pragma once
+
+#include "Core/CameraNode.h"
+#include "Core/CameraParameters.h"
+
+#include "BodyParametersCameraNode.generated.h"
+
+/**
+ * A camera node that configures parameters on the camera body.
+ */
+UCLASS(MinimalAPI, meta=(CameraNodeCategories="Common,Body"))
+class UBodyParametersCameraNode : public UCameraNode
+{
+	GENERATED_BODY()
+
+protected:
+
+	// UCameraNode interface.
+	virtual FCameraNodeEvaluatorPtr OnBuildEvaluator(FCameraNodeEvaluatorBuilder& Builder) const override;
+
+public:
+
+	/** Current shutter speed, in 1/seconds */
+	UPROPERTY(EditAnywhere, Category="Filmback", DisplayName="Shutter Speed (1/s)")
+    FFloatCameraParameter ShutterSpeed;
+
+	/** The camera sensor sensitivity in ISO. */
+	UPROPERTY(EditAnywhere, Category="Filmback")
+	FFloatCameraParameter ISO;
+};
+
diff --git a/Engine/Plugins/Cameras/GameplayCameras/Source/GameplayCameras/Public/Nodes/Common/FilmbackCameraNode.h b/Engine/Plugins/Cameras/GameplayCameras/Source/GameplayCameras/Public/Nodes/Common/FilmbackCameraNode.h
--- a/Engine/Plugins/Cameras/GameplayCameras/Source/GameplayCameras/Public/Nodes/Common/FilmbackCameraNode.h
+++ b/Engine/Plugins/Cameras/GameplayCameras/Source/GameplayCameras/Public/Nodes/Common/FilmbackCameraNode.h
@@ -34,9 +34,17 @@ public:
 	UPROPERTY(EditAnywhere, Category="Filmback", meta=(ClampMin="0.001", ForceUnits=mm))
 	FFloatCameraParameter SensorHeight;
 
-	/** The camera sensor sensitivity in ISO. */
-	UPROPERTY(EditAnywhere, Category="Filmback")
-	FFloatCameraParameter ISO;
+	/** Horizontal offset of the sensor, in mm. */
+	UPROPERTY(EditAnywhere, Category="Advanced", meta=(ForceUnits=mm))
+	FFloatCameraParameter SensorHorizontalOffset;
+
+	/** Vertical offset of the sensor, in mm. */
+	UPROPERTY(EditAnywhere, Category="Advanced", meta=(ForceUnits=mm))
+	FFloatCameraParameter SensorVerticalOffset;
+
+	/** Sensor overscan, in percentages (0% for no increase). */
+	UPROPERTY(EditAnywhere, Category="Advanced", meta=(ForceUnits=Percent))
+	FFloatCameraParameter Overscan;
 
 	/** Whether to constrain the aspect ratio of the evaluated camera. */
 	UPROPERTY(EditAnywhere, Category="Filmback")
diff --git a/Engine/Plugins/Cameras/GameplayCameras/Source/GameplayCameras/Public/Nodes/Common/LensParametersCameraNode.h b/Engine/Plugins/Cameras/GameplayCameras/Source/GameplayCameras/Public/Nodes/Common/LensParametersCameraNode.h
--- a/Engine/Plugins/Cameras/GameplayCameras/Source/GameplayCameras/Public/Nodes/Common/LensParametersCameraNode.h
+++ b/Engine/Plugins/Cameras/GameplayCameras/Source/GameplayCameras/Public/Nodes/Common/LensParametersCameraNode.h
@@ -35,10 +35,6 @@ public:
 	UPROPERTY(EditAnywhere, Category="Lens Parameters")
 	FFloatCameraParameter Aperture;
 
-	/** Current shutter speed, in 1/seconds */
-	UPROPERTY(EditAnywhere, Category="Lens Parameters", DisplayName="Shutter Speed (1/s)")
-    FFloatCameraParameter ShutterSpeed;
-
 	/** 
 	 * Whether post-process settings should be automatically applied to reflect
 	 * the lens settings.
