Cameras: add orthographic node

diff --git a/Engine/Plugins/Cameras/GameplayCameras/Source/GameplayCameras/Private/Nodes/Common/OrthographicCameraMode.cpp b/Engine/Plugins/Cameras/GameplayCameras/Source/GameplayCameras/Private/Nodes/Common/OrthographicCameraMode.cpp
new file mode 100644
--- /dev/null
+++ b/Engine/Plugins/Cameras/GameplayCameras/Source/GameplayCameras/Private/Nodes/Common/OrthographicCameraMode.cpp
@@ -0,0 +1,59 @@
+// Copyright Epic Games, Inc. All Rights Reserved.
+
+#include "Camera/CameraTypes.h"
+#include "Nodes/Common/OrthographicCameraNode.h"
+
+#include "Core/CameraEvaluationContext.h"
+#include "Core/CameraParameterReader.h"
+#include "GameplayCameras.h"
+
+#include UE_INLINE_GENERATED_CPP_BY_NAME(OrthographicCameraNode)
+
+namespace UE::Cameras
+{
+
+class FOrthographicCameraNodeEvaluator : public FCameraNodeEvaluator
+{
+	UE_DECLARE_CAMERA_NODE_EVALUATOR(GAMEPLAYCAMERAS_API, FOrthographicCameraNodeEvaluator)
+
+protected:
+
+	virtual void OnInitialize(const FCameraNodeEvaluatorInitializeParams& Params, FCameraNodeEvaluationResult& OutResult) override;
+	virtual void OnRun(const FCameraNodeEvaluationParams& Params, FCameraNodeEvaluationResult& OutResult) override;
+
+private:
+
+	TCameraParameterReader<bool> EnableOrthographicModeReader;
+	TCameraParameterReader<float> OrthographicWidthReader;
+};
+
+UE_DEFINE_CAMERA_NODE_EVALUATOR(FOrthographicCameraNodeEvaluator)
+
+void FOrthographicCameraNodeEvaluator::OnInitialize(const FCameraNodeEvaluatorInitializeParams& Params, FCameraNodeEvaluationResult& OutResult)
+{
+	const UOrthographicCameraNode* OrthographicNode = GetCameraNodeAs<UOrthographicCameraNode>();
+	EnableOrthographicModeReader.Initialize(OrthographicNode->EnableOrthographicMode);
+	OrthographicWidthReader.Initialize(OrthographicNode->OrthographicWidth);
+}
+
+void FOrthographicCameraNodeEvaluator::OnRun(const FCameraNodeEvaluationParams& Params, FCameraNodeEvaluationResult& OutResult)
+{
+	FCameraPose& OutPose = OutResult.CameraPose;
+
+	if (EnableOrthographicModeReader.Get(OutResult.VariableTable))
+	{
+		OutPose.SetProjectionMode(ECameraProjectionMode::Orthographic);
+	}
+
+	const float OrthographicWidth = OrthographicWidthReader.Get(OutResult.VariableTable);
+	OutPose.SetOrthographicWidth(OrthographicWidth);
+}
+
+}  // namespace UE::Cameras
+
+FCameraNodeEvaluatorPtr UOrthographicCameraNode::OnBuildEvaluator(FCameraNodeEvaluatorBuilder& Builder) const
+{
+	using namespace UE::Cameras;
+	return Builder.BuildEvaluator<FOrthographicCameraNodeEvaluator>();
+}
+
diff --git a/Engine/Plugins/Cameras/GameplayCameras/Source/GameplayCameras/Public/Nodes/Common/OrthographicCameraNode.h b/Engine/Plugins/Cameras/GameplayCameras/Source/GameplayCameras/Public/Nodes/Common/OrthographicCameraNode.h
new file mode 100644
--- /dev/null
+++ b/Engine/Plugins/Cameras/GameplayCameras/Source/GameplayCameras/Public/Nodes/Common/OrthographicCameraNode.h
@@ -0,0 +1,34 @@
+// Copyright Epic Games, Inc. All Rights Reserved.
+
+#pragma once
+
+#include "Core/CameraNode.h"
+#include "Core/CameraParameters.h"
+#include "EngineDefines.h"
+
+#include "OrthographicCameraNode.generated.h"
+
+/**
+ * A camera node that sets the perspective mode to orthographic.
+ */
+UCLASS(MinimalAPI, meta=(CameraNodeCategories="Common"))
+class UOrthographicCameraNode : public UCameraNode
+{
+	GENERATED_BODY()
+
+protected:
+
+	// UCameraNode interface.
+	virtual FCameraNodeEvaluatorPtr OnBuildEvaluator(FCameraNodeEvaluatorBuilder& Builder) const override;
+
+public:
+
+	/** Sets the perspective mode to orthographic. */
+	UPROPERTY(EditAnywhere, Category="Orthographic")
+	FBooleanCameraParameter EnableOrthographicMode = true;
+
+	/** The width of the orthographic view. */
+	UPROPERTY(EditAnywhere, Category="Orthographic", meta=(Units=cm))
+	FFloatCameraParameter OrthographicWidth = DEFAULT_ORTHOWIDTH;
+};
+
