Cameras: add post-process settings to overall system result and show debug info about them

diff --git a/Engine/Plugins/Cameras/GameplayCameras/Source/GameplayCameras/Private/Core/CameraSystemEvaluator.cpp b/Engine/Plugins/Cameras/GameplayCameras/Source/GameplayCameras/Private/Core/CameraSystemEvaluator.cpp
--- a/Engine/Plugins/Cameras/GameplayCameras/Source/GameplayCameras/Private/Core/CameraSystemEvaluator.cpp
+++ b/Engine/Plugins/Cameras/GameplayCameras/Source/GameplayCameras/Private/Core/CameraSystemEvaluator.cpp
@@ -32,6 +32,7 @@ void FCameraSystemEvaluationResult::Reset()
 {
 	CameraPose.ClearAllChangedFlags();
 	VariableTable.ClearAllWrittenThisFrameFlags();
+	ContextDataTable.ClearAllWrittenThisFrameFlags();
 	bIsCameraCut = false;
 	bIsValid = false;
 }
@@ -44,6 +45,8 @@ void FCameraSystemEvaluationResult::Reset(const FCameraNodeEvaluationResult& Nod
 	CameraPose = NodeResult.CameraPose;
 
 	VariableTable.OverrideAll(NodeResult.VariableTable);
+	ContextDataTable.OverrideAll(NodeResult.ContextDataTable);
+	PostProcessSettings.OverrideAll(NodeResult.PostProcessSettings);
 
 	bIsCameraCut = NodeResult.bIsCameraCut;
 	bIsValid = true;
diff --git a/Engine/Plugins/Cameras/GameplayCameras/Source/GameplayCameras/Private/Core/PostProcessSettingsCollection.cpp b/Engine/Plugins/Cameras/GameplayCameras/Source/GameplayCameras/Private/Core/PostProcessSettingsCollection.cpp
--- a/Engine/Plugins/Cameras/GameplayCameras/Source/GameplayCameras/Private/Core/PostProcessSettingsCollection.cpp
+++ b/Engine/Plugins/Cameras/GameplayCameras/Source/GameplayCameras/Private/Core/PostProcessSettingsCollection.cpp
@@ -10,13 +10,11 @@ namespace UE::Cameras
 void FPostProcessSettingsCollection::Reset()
 {
 	PostProcessSettings = FPostProcessSettings::GetDefault();
-	bHasAnySetting = false;
 }
 
 void FPostProcessSettingsCollection::OverrideAll(const FPostProcessSettingsCollection& OtherCollection)
 {
 	PostProcessSettings = OtherCollection.PostProcessSettings;
-	bHasAnySetting = OtherCollection.bHasAnySetting;
 }
 
 void FPostProcessSettingsCollection::OverrideChanged(const FPostProcessSettingsCollection& OtherCollection)
@@ -29,8 +27,7 @@ void FPostProcessSettingsCollection::OverrideChanged(const FPostProcessSettings&
 	FPostProcessSettings& ThisPP = PostProcessSettings;
 	const FPostProcessSettings& OtherPP = OtherPostProcessSettings;
 
-	const bool bAnyOverwritten = FPostProcessUtils::OverridePostProcessSettings(ThisPP, OtherPP);
-	bHasAnySetting |= bAnyOverwritten;
+	FPostProcessUtils::OverridePostProcessSettings(ThisPP, OtherPP);
 }
 
 void FPostProcessSettingsCollection::LerpAll(const FPostProcessSettingsCollection& ToCollection, float BlendFactor)
@@ -48,16 +45,13 @@ void FPostProcessSettingsCollection::InternalLerp(const FPostProcessSettings& To
 	FPostProcessSettings& ThisPP = PostProcessSettings;
 	const FPostProcessSettings& ToPP = ToPostProcessSettings;
 
-	const bool bAnyOverwritten = FPostProcessUtils::BlendPostProcessSettings(ThisPP, ToPP, BlendFactor);
-	bHasAnySetting |= bAnyOverwritten;
+	FPostProcessUtils::BlendPostProcessSettings(ThisPP, ToPP, BlendFactor);
 }
 
 void FPostProcessSettingsCollection::Serialize(FArchive& Ar)
 {
 	UScriptStruct* PostProcessSettingsStruct = FPostProcessSettings::StaticStruct();
 	PostProcessSettingsStruct->SerializeItem(Ar, &PostProcessSettings, &FPostProcessSettings::GetDefault());
-
-	Ar << bHasAnySetting;
 }
 
 }  // namespace UE::Cameras
diff --git a/Engine/Plugins/Cameras/GameplayCameras/Source/GameplayCameras/Private/Debug/CameraNodeEvaluationResultDebugBlock.cpp b/Engine/Plugins/Cameras/GameplayCameras/Source/GameplayCameras/Private/Debug/CameraNodeEvaluationResultDebugBlock.cpp
--- a/Engine/Plugins/Cameras/GameplayCameras/Source/GameplayCameras/Private/Debug/CameraNodeEvaluationResultDebugBlock.cpp
+++ b/Engine/Plugins/Cameras/GameplayCameras/Source/GameplayCameras/Private/Debug/CameraNodeEvaluationResultDebugBlock.cpp
@@ -11,6 +11,7 @@
 #include "Debug/CameraRigJointsDebugBlock.h"
 #include "Debug/ContextDataTableDebugBlock.h"
 #include "Debug/DebugTextRenderer.h"
+#include "Debug/PostProcessSettingsDebugBlock.h"
 #include "Debug/VariableTableDebugBlock.h"
 #include "Math/ColorList.h"
 
@@ -27,13 +28,13 @@ FCameraNodeEvaluationResultDebugBlock::FCameraNodeEvaluationResultDebugBlock()
 
 void FCameraNodeEvaluationResultDebugBlock::Initialize(const FCameraNodeEvaluationResult& InResult, FCameraDebugBlockBuilder& Builder)
 {
-	bHasAnyPostProcessSettings = InResult.PostProcessSettings.HasAnyPostProcessSettings();
 	bIsCameraCut = InResult.bIsCameraCut;
 	bIsValid = InResult.bIsValid;
 
 	AddChild(&Builder.BuildDebugBlock<FCameraPoseDebugBlock>(InResult.CameraPose));
 	AddChild(&Builder.BuildDebugBlock<FVariableTableDebugBlock>(InResult.VariableTable));
 	AddChild(&Builder.BuildDebugBlock<FContextDataTableDebugBlock>(InResult.ContextDataTable));
+	AddChild(&Builder.BuildDebugBlock<FPostProcessSettingsDebugBlock>(InResult.PostProcessSettings));
 	AddChild(&Builder.BuildDebugBlock<FCameraRigJointsDebugBlock>(InResult.CameraRigJoints, InResult.VariableTable));
 }
 
@@ -44,7 +45,9 @@ void FCameraNodeEvaluationResultDebugBlock::Initialize(const FCameraSystemEvalua
 
 	AddChild(&Builder.BuildDebugBlock<FCameraPoseDebugBlock>(InResult.CameraPose));
 	AddChild(&Builder.BuildDebugBlock<FVariableTableDebugBlock>(InResult.VariableTable));
-	// No context data or rig joints exposed on the camera system result.
+	AddChild(&Builder.BuildDebugBlock<FContextDataTableDebugBlock>(InResult.ContextDataTable));
+	AddChild(&Builder.BuildDebugBlock<FPostProcessSettingsDebugBlock>(InResult.PostProcessSettings));
+	// No rig joints exposed on the camera system result.
 }
 
 FCameraPoseDebugBlock* FCameraNodeEvaluationResultDebugBlock::GetCameraPoseDebugBlock()
@@ -77,6 +80,16 @@ FContextDataTableDebugBlock* FCameraNodeEvaluationResultDebugBlock::GetContextDa
 	return nullptr;
 }
 
+FPostProcessSettingsDebugBlock* FCameraNodeEvaluationResultDebugBlock::GetPostProcessSettingsDebugBlock()
+{
+	TArrayView<FCameraDebugBlock*> ChildrenView(GetChildren());
+	if (ChildrenView.IsValidIndex(3))
+	{
+		return ChildrenView[3]->CastThis<FPostProcessSettingsDebugBlock>();
+	}
+	return nullptr;
+}
+
 void FCameraNodeEvaluationResultDebugBlock::OnDebugDraw(const FCameraDebugBlockDrawParams& Params, FCameraDebugRenderer& Renderer)
 {
 	const FCameraDebugColors& Colors = FCameraDebugColors::Get();
@@ -95,11 +108,6 @@ void FCameraNodeEvaluationResultDebugBlock::OnDebugDraw(const FCameraDebugBlockD
 		Renderer.AddText(TEXT("  {cam_warning}IsCameraCut"));
 	}
 
-	if (bHasAnyPostProcessSettings)
-	{
-		Renderer.AddText(TEXT("  {cam_notice}post-FX set"));
-	}
-
 	Renderer.NewLine();
 	Renderer.SetTextColor(Colors.Default);
 
@@ -130,12 +138,20 @@ void FCameraNodeEvaluationResultDebugBlock::OnDebugDraw(const FCameraDebugBlockD
 	}
 	if (ChildrenView.IsValidIndex(3))
 	{
-		Renderer.AddText(TEXT("{cam_title}Camera Rig Joints:"));
+		Renderer.AddText(TEXT("{cam_title}Post-Process Settings:"));
 		Renderer.AddIndent();
 		Renderer.SetTextColor(Colors.Default);
 		ChildrenView[3]->DebugDraw(Params, Renderer);
 		Renderer.RemoveIndent();
 	}
+	if (ChildrenView.IsValidIndex(4))
+	{
+		Renderer.AddText(TEXT("{cam_title}Camera Rig Joints:"));
+		Renderer.AddIndent();
+		Renderer.SetTextColor(Colors.Default);
+		ChildrenView[4]->DebugDraw(Params, Renderer);
+		Renderer.RemoveIndent();
+	}
 
 	// We already manually rendered our children.
 	Renderer.SkipAllBlocks();
diff --git a/Engine/Plugins/Cameras/GameplayCameras/Source/GameplayCameras/Private/Debug/PostProcessSettingsDebugBlock.cpp b/Engine/Plugins/Cameras/GameplayCameras/Source/GameplayCameras/Private/Debug/PostProcessSettingsDebugBlock.cpp
new file mode 100644
--- /dev/null
+++ b/Engine/Plugins/Cameras/GameplayCameras/Source/GameplayCameras/Private/Debug/PostProcessSettingsDebugBlock.cpp
@@ -0,0 +1,334 @@
+// Copyright Epic Games, Inc. All Rights Reserved.
+
+#include "Debug/PostProcessSettingsDebugBlock.h"
+
+#include "Containers/UnrealString.h"
+#include "Curves/CurveFloat.h"
+#include "Debug/CameraDebugColors.h"
+#include "Debug/CameraDebugRenderer.h"
+#include "Debug/DebugTextRenderer.h"
+#include "Engine/Texture.h"
+#include "Engine/Texture2D.h"
+#include "Engine/TextureCube.h"
+#include "Math/ColorList.h"
+
+#if UE_GAMEPLAY_CAMERAS_DEBUG
+
+namespace UE::Cameras
+{
+
+UE_DEFINE_CAMERA_DEBUG_BLOCK(FPostProcessSettingsDebugBlock)
+
+FPostProcessSettingsDebugBlock::FPostProcessSettingsDebugBlock()
+{
+}
+
+FPostProcessSettingsDebugBlock::FPostProcessSettingsDebugBlock(const FPostProcessSettingsCollection& InPostProcessSettings)
+	: PostProcessSettings(InPostProcessSettings)
+{
+}
+
+void FPostProcessSettingsDebugBlock::OnDebugDraw(const FCameraDebugBlockDrawParams& Params, FCameraDebugRenderer& Renderer)
+{
+	const FPostProcessSettings& PPS = PostProcessSettings.Get();
+
+#define UE_DRAW_PP(Name)\
+	if (PPS.bOverride_##Name)\
+	{\
+		Renderer.AddText(TEXT("%s  : %s\n"), TEXT(#Name), *ToDebugString(PPS.Name));\
+	}
+
+#define UE_DRAW_PP_NAME(Name)\
+	if (PPS.bOverride_##Name)\
+	{\
+		Renderer.AddText(TEXT("%s  : %s\n"), TEXT(#Name), *PPS.Name->GetName());\
+	}
+
+	{
+		UE_DRAW_PP(TemperatureType);
+		UE_DRAW_PP(WhiteTemp);
+		UE_DRAW_PP(WhiteTint);
+
+		UE_DRAW_PP(ColorSaturation);
+		UE_DRAW_PP(ColorContrast);
+		UE_DRAW_PP(ColorGamma);
+		UE_DRAW_PP(ColorGain);
+		UE_DRAW_PP(ColorOffset);
+
+		UE_DRAW_PP(ColorSaturationShadows);
+		UE_DRAW_PP(ColorContrastShadows);
+		UE_DRAW_PP(ColorGammaShadows);
+		UE_DRAW_PP(ColorGainShadows);
+		UE_DRAW_PP(ColorOffsetShadows);
+
+		UE_DRAW_PP(ColorSaturationMidtones);
+		UE_DRAW_PP(ColorContrastMidtones);
+		UE_DRAW_PP(ColorGammaMidtones);
+		UE_DRAW_PP(ColorGainMidtones);
+		UE_DRAW_PP(ColorOffsetMidtones);
+
+		UE_DRAW_PP(ColorSaturationHighlights);
+		UE_DRAW_PP(ColorContrastHighlights);
+		UE_DRAW_PP(ColorGammaHighlights);
+		UE_DRAW_PP(ColorGainHighlights);
+		UE_DRAW_PP(ColorOffsetHighlights);
+
+		UE_DRAW_PP(ColorCorrectionShadowsMax);
+		UE_DRAW_PP(ColorCorrectionHighlightsMin);
+		UE_DRAW_PP(ColorCorrectionHighlightsMax);
+
+		UE_DRAW_PP(BlueCorrection);
+		UE_DRAW_PP(ExpandGamut);
+		UE_DRAW_PP(ToneCurveAmount);
+
+		UE_DRAW_PP(FilmSlope);
+		UE_DRAW_PP(FilmToe);
+		UE_DRAW_PP(FilmShoulder);
+		UE_DRAW_PP(FilmBlackClip);
+		UE_DRAW_PP(FilmWhiteClip);
+
+		UE_DRAW_PP(SceneColorTint);
+		UE_DRAW_PP(SceneFringeIntensity);
+		UE_DRAW_PP(ChromaticAberrationStartOffset);
+		UE_DRAW_PP(BloomIntensity);
+		UE_DRAW_PP(BloomThreshold);
+		UE_DRAW_PP(Bloom1Tint);
+		UE_DRAW_PP(BloomSizeScale);
+		UE_DRAW_PP(Bloom1Size);
+		UE_DRAW_PP(Bloom2Tint);
+		UE_DRAW_PP(Bloom2Size);
+		UE_DRAW_PP(Bloom3Tint);
+		UE_DRAW_PP(Bloom3Size);
+		UE_DRAW_PP(Bloom4Tint);
+		UE_DRAW_PP(Bloom4Size);
+		UE_DRAW_PP(Bloom5Tint);
+		UE_DRAW_PP(Bloom5Size);
+		UE_DRAW_PP(Bloom6Tint);
+		UE_DRAW_PP(Bloom6Size);
+		UE_DRAW_PP(BloomDirtMaskIntensity);
+		UE_DRAW_PP(BloomDirtMaskTint);
+		UE_DRAW_PP(BloomConvolutionScatterDispersion);
+		UE_DRAW_PP(BloomConvolutionSize);
+		UE_DRAW_PP(BloomConvolutionCenterUV);
+		UE_DRAW_PP(BloomConvolutionPreFilterMin);
+		UE_DRAW_PP(BloomConvolutionPreFilterMax);
+		UE_DRAW_PP(BloomConvolutionPreFilterMult);
+		UE_DRAW_PP(AmbientCubemapIntensity);
+		UE_DRAW_PP(AmbientCubemapTint);
+		UE_DRAW_PP(CameraShutterSpeed);
+		UE_DRAW_PP(CameraISO);
+		UE_DRAW_PP(AutoExposureLowPercent);
+		UE_DRAW_PP(AutoExposureHighPercent);
+		UE_DRAW_PP(AutoExposureMinBrightness);
+		UE_DRAW_PP(AutoExposureMaxBrightness);
+		UE_DRAW_PP(AutoExposureSpeedUp);
+		UE_DRAW_PP(AutoExposureSpeedDown);
+		UE_DRAW_PP(AutoExposureBias);
+		UE_DRAW_PP(HistogramLogMin);
+		UE_DRAW_PP(HistogramLogMax);
+		UE_DRAW_PP(LocalExposureMethod);
+		UE_DRAW_PP(LocalExposureContrastScale_DEPRECATED);
+		UE_DRAW_PP(LocalExposureHighlightContrastScale);
+		UE_DRAW_PP(LocalExposureShadowContrastScale);
+		UE_DRAW_PP(LocalExposureHighlightThreshold);
+		UE_DRAW_PP(LocalExposureShadowThreshold);
+		UE_DRAW_PP(LocalExposureDetailStrength);
+		UE_DRAW_PP(LocalExposureBlurredLuminanceBlend);
+		UE_DRAW_PP(LocalExposureBlurredLuminanceKernelSizePercent);
+		UE_DRAW_PP(LocalExposureHighlightThresholdStrength);
+		UE_DRAW_PP(LocalExposureShadowThresholdStrength);
+		UE_DRAW_PP(LocalExposureMiddleGreyBias);
+		UE_DRAW_PP(LensFlareIntensity);
+		UE_DRAW_PP(LensFlareTint);
+		UE_DRAW_PP(LensFlareBokehSize);
+		UE_DRAW_PP(LensFlareThreshold);
+		UE_DRAW_PP(VignetteIntensity);
+		UE_DRAW_PP(Sharpen);
+		UE_DRAW_PP(FilmGrainIntensity);
+		UE_DRAW_PP(FilmGrainIntensityShadows);
+		UE_DRAW_PP(FilmGrainIntensityMidtones);
+		UE_DRAW_PP(FilmGrainIntensityHighlights);
+		UE_DRAW_PP(FilmGrainShadowsMax);
+		UE_DRAW_PP(FilmGrainHighlightsMin);
+		UE_DRAW_PP(FilmGrainHighlightsMax);
+		UE_DRAW_PP(FilmGrainTexelSize);
+		UE_DRAW_PP(AmbientOcclusionIntensity);
+		UE_DRAW_PP(AmbientOcclusionStaticFraction);
+		UE_DRAW_PP(AmbientOcclusionRadius);
+		UE_DRAW_PP(AmbientOcclusionFadeDistance);
+		UE_DRAW_PP(AmbientOcclusionFadeRadius);
+		UE_DRAW_PP(AmbientOcclusionDistance_DEPRECATED);
+		UE_DRAW_PP(AmbientOcclusionPower);
+		UE_DRAW_PP(AmbientOcclusionBias);
+		UE_DRAW_PP(AmbientOcclusionQuality);
+		UE_DRAW_PP(AmbientOcclusionMipBlend);
+		UE_DRAW_PP(AmbientOcclusionMipScale);
+		UE_DRAW_PP(AmbientOcclusionMipThreshold);
+		UE_DRAW_PP(AmbientOcclusionTemporalBlendWeight);
+		UE_DRAW_PP(IndirectLightingColor);
+		UE_DRAW_PP(IndirectLightingIntensity);
+
+		UE_DRAW_PP(DepthOfFieldFocalDistance);
+
+		UE_DRAW_PP(DepthOfFieldFstop);
+		UE_DRAW_PP(DepthOfFieldMinFstop);
+		UE_DRAW_PP(DepthOfFieldSensorWidth);
+		UE_DRAW_PP(DepthOfFieldSqueezeFactor);
+		UE_DRAW_PP(DepthOfFieldDepthBlurRadius);
+		UE_DRAW_PP(DepthOfFieldUseHairDepth)
+		UE_DRAW_PP(DepthOfFieldDepthBlurAmount);
+		UE_DRAW_PP(DepthOfFieldFocalRegion);
+		UE_DRAW_PP(DepthOfFieldNearTransitionRegion);
+		UE_DRAW_PP(DepthOfFieldFarTransitionRegion);
+		UE_DRAW_PP(DepthOfFieldScale);
+		UE_DRAW_PP(DepthOfFieldNearBlurSize);
+		UE_DRAW_PP(DepthOfFieldFarBlurSize);
+		UE_DRAW_PP(DepthOfFieldOcclusion);
+		UE_DRAW_PP(DepthOfFieldSkyFocusDistance);
+		UE_DRAW_PP(DepthOfFieldVignetteSize);
+		UE_DRAW_PP(DepthOfFieldAspectRatioScalar);
+		UE_DRAW_PP(DepthOfFieldPetzvalBokeh);
+		UE_DRAW_PP(DepthOfFieldPetzvalBokehFalloff);
+		UE_DRAW_PP(DepthOfFieldPetzvalExclusionBoxExtents);
+		UE_DRAW_PP(DepthOfFieldPetzvalExclusionBoxRadius);
+		UE_DRAW_PP(DepthOfFieldBarrelRadius);
+		UE_DRAW_PP(DepthOfFieldBarrelLength);
+		if (PPS.bOverride_DepthOfFieldMatteBoxFlags)
+		{
+			Renderer.AddText(TEXT("DepthOfFieldMatteBoxFlags  : "));
+			for (uint32 Index = 0; Index < UE_ARRAY_COUNT(PPS.DepthOfFieldMatteBoxFlags); ++Index)
+			{
+				const FMatteBoxFlag& Flag = PPS.DepthOfFieldMatteBoxFlags[Index];
+				Renderer.AddText(TEXT(" %f;%f;%f"), Flag.Pitch, Flag.Roll, Flag.Length);
+			}
+			Renderer.AddText(TEXT("\n"));
+		}
+
+		UE_DRAW_PP(MotionBlurAmount);
+		UE_DRAW_PP(MotionBlurMax);
+		UE_DRAW_PP(MotionBlurPerObjectSize);
+		UE_DRAW_PP(ScreenSpaceReflectionQuality);
+		UE_DRAW_PP(ScreenSpaceReflectionIntensity);
+		UE_DRAW_PP(ScreenSpaceReflectionMaxRoughness);
+
+		UE_DRAW_PP(TranslucencyType);
+		UE_DRAW_PP(RayTracingTranslucencyMaxRoughness);
+		UE_DRAW_PP(RayTracingTranslucencyRefractionRays);
+		UE_DRAW_PP(RayTracingTranslucencySamplesPerPixel);
+		UE_DRAW_PP(RayTracingTranslucencyShadows);
+		UE_DRAW_PP(RayTracingTranslucencyRefraction);
+
+		UE_DRAW_PP(DynamicGlobalIlluminationMethod);
+		UE_DRAW_PP(LumenSurfaceCacheResolution);
+		UE_DRAW_PP(LumenSceneLightingQuality);
+		UE_DRAW_PP(LumenSceneDetail);
+		UE_DRAW_PP(LumenSceneViewDistance);
+		UE_DRAW_PP(LumenSceneLightingUpdateSpeed);
+		UE_DRAW_PP(LumenFinalGatherQuality);
+		UE_DRAW_PP(LumenFinalGatherLightingUpdateSpeed);
+		UE_DRAW_PP(LumenFinalGatherScreenTraces);
+		UE_DRAW_PP(LumenMaxTraceDistance);
+
+		UE_DRAW_PP(LumenDiffuseColorBoost);
+		UE_DRAW_PP(LumenSkylightLeaking);
+		UE_DRAW_PP(LumenSkylightLeakingTint);
+		UE_DRAW_PP(LumenFullSkylightLeakingDistance);
+
+		UE_DRAW_PP(LumenRayLightingMode);
+		UE_DRAW_PP(LumenReflectionsScreenTraces);
+		UE_DRAW_PP(LumenFrontLayerTranslucencyReflections);
+		UE_DRAW_PP(LumenMaxRoughnessToTraceReflections);
+		UE_DRAW_PP(LumenMaxReflectionBounces);
+		UE_DRAW_PP(LumenMaxRefractionBounces);
+		UE_DRAW_PP(ReflectionMethod);
+		UE_DRAW_PP(LumenReflectionQuality);
+		UE_DRAW_PP(RayTracingAO);
+		UE_DRAW_PP(RayTracingAOSamplesPerPixel);
+		UE_DRAW_PP(RayTracingAOIntensity);
+		UE_DRAW_PP(RayTracingAORadius);
+
+		UE_DRAW_PP(PathTracingMaxBounces);
+		UE_DRAW_PP(PathTracingSamplesPerPixel);
+		UE_DRAW_PP(PathTracingMaxPathIntensity);
+		UE_DRAW_PP(PathTracingEnableEmissiveMaterials);
+		UE_DRAW_PP(PathTracingEnableReferenceDOF);
+		UE_DRAW_PP(PathTracingEnableReferenceAtmosphere);
+		UE_DRAW_PP(PathTracingEnableDenoiser);
+		UE_DRAW_PP(PathTracingIncludeEmissive);
+		UE_DRAW_PP(PathTracingIncludeDiffuse);
+		UE_DRAW_PP(PathTracingIncludeIndirectDiffuse);
+		UE_DRAW_PP(PathTracingIncludeSpecular);
+		UE_DRAW_PP(PathTracingIncludeIndirectSpecular);
+		UE_DRAW_PP(PathTracingIncludeVolume);
+		UE_DRAW_PP(PathTracingIncludeIndirectVolume);
+
+		UE_DRAW_PP(DepthOfFieldBladeCount);
+
+		// This is no bOverride_AmbientCubemap so just see if it is set.
+		if (PPS.AmbientCubemap)
+		{
+			Renderer.AddText(TEXT("AmbientCubemap  : %s\n"), *PPS.AmbientCubemap->GetName());
+		}
+
+		UE_DRAW_PP(ColorGradingIntensity);
+		
+		UE_DRAW_PP_NAME(ColorGradingLUT)
+		UE_DRAW_PP_NAME(BloomDirtMask);
+		UE_DRAW_PP(BloomMethod);
+		UE_DRAW_PP_NAME(BloomConvolutionTexture)
+		UE_DRAW_PP_NAME(FilmGrainTexture)
+
+		UE_DRAW_PP(BloomConvolutionBufferScale);
+
+		UE_DRAW_PP_NAME(AutoExposureBiasCurve);
+		UE_DRAW_PP_NAME(AutoExposureMeterMask);
+		UE_DRAW_PP_NAME(LocalExposureHighlightContrastCurve);
+		UE_DRAW_PP_NAME(LocalExposureShadowContrastCurve);
+		UE_DRAW_PP_NAME(LensFlareBokehShape);
+
+		if (PPS.bOverride_LensFlareTints)
+		{
+			Renderer.AddText(TEXT("LensFlareTints  :"));
+			for (uint32 i = 0; i < 8; ++i)
+			{
+				Renderer.AddText(TEXT(" %s"), *ToDebugString(PPS.LensFlareTints[i]));
+			}
+			Renderer.AddText(TEXT("\n"));
+		}
+
+		if (PPS.bOverride_MobileHQGaussian)
+		{
+			Renderer.AddText(TEXT("MobileHQGaussian  : %s\n"), *ToDebugString(PPS.bMobileHQGaussian));
+		}
+
+		UE_DRAW_PP(AutoExposureMethod);
+		UE_DRAW_PP(AmbientOcclusionRadiusInWS);
+		UE_DRAW_PP(MotionBlurTargetFPS);
+		UE_DRAW_PP(AutoExposureApplyPhysicalCameraExposure);
+		UE_DRAW_PP(UserFlags);
+
+		if (PPS.WeightedBlendables.Array.Num() > 0)
+		{
+			Renderer.AddText(TEXT("WeightedBlendables  :"));
+			for (int32 Index = 0; Index < PPS.WeightedBlendables.Array.Num(); ++Index)
+			{
+				const FWeightedBlendable& Blendable = PPS.WeightedBlendables.Array[Index];
+				Renderer.AddText(TEXT(" [%d] %s (%d%%)"), Index, *GetNameSafe(Blendable.Object), (int32)(Blendable.Weight * 100.f));
+			}
+			Renderer.AddText(TEXT("\n"));
+		}
+	}
+
+#undef UE_DRAW_PP
+#undef UE_DRAW_PP_NAME
+}
+
+void FPostProcessSettingsDebugBlock::OnSerialize(FArchive& Ar)
+{
+	PostProcessSettings.Serialize(Ar);
+}
+
+}  // namespace UE::Cameras
+
+#endif  // UE_GAMEPLAY_CAMERAS_DEBUG
+
diff --git a/Engine/Plugins/Cameras/GameplayCameras/Source/GameplayCameras/Private/GameFramework/GameplayCameraComponentBase.cpp b/Engine/Plugins/Cameras/GameplayCameras/Source/GameplayCameras/Private/GameFramework/GameplayCameraComponentBase.cpp
--- a/Engine/Plugins/Cameras/GameplayCameras/Source/GameplayCameras/Private/GameFramework/GameplayCameraComponentBase.cpp
+++ b/Engine/Plugins/Cameras/GameplayCameras/Source/GameplayCameras/Private/GameFramework/GameplayCameraComponentBase.cpp
@@ -442,6 +442,9 @@ void UGameplayCameraComponentBase::UpdateOutputCameraComponent()
 			OutputCameraComponent->ProjectionMode = Result.CameraPose.GetProjectionMode();
 			OutputCameraComponent->OrthoWidth = Result.CameraPose.GetOrthographicWidth();
 
+			OutputCameraComponent->PostProcessSettings = Result.PostProcessSettings.Get();
+			OutputCameraComponent->PostProcessBlendWeight = 1.f;
+
 			bGotValidTransform = true;
 		}
 	}
diff --git a/Engine/Plugins/Cameras/GameplayCameras/Source/GameplayCameras/Public/Core/CameraSystemEvaluator.h b/Engine/Plugins/Cameras/GameplayCameras/Source/GameplayCameras/Public/Core/CameraSystemEvaluator.h
--- a/Engine/Plugins/Cameras/GameplayCameras/Source/GameplayCameras/Public/Core/CameraSystemEvaluator.h
+++ b/Engine/Plugins/Cameras/GameplayCameras/Source/GameplayCameras/Public/Core/CameraSystemEvaluator.h
@@ -85,6 +85,12 @@ struct FCameraSystemEvaluationResult
 	/** The result camera variable table. */
 	FCameraVariableTable VariableTable;
 
+	/** The result camera context data table. */
+	FCameraContextDataTable ContextDataTable;
+	
+	/** The result post-process settings. */
+	FPostProcessSettingsCollection PostProcessSettings;
+
 	/** Whether this evaluation was a camera cut. */
 	bool bIsCameraCut = false;
 
diff --git a/Engine/Plugins/Cameras/GameplayCameras/Source/GameplayCameras/Public/Core/PostProcessSettingsCollection.h b/Engine/Plugins/Cameras/GameplayCameras/Source/GameplayCameras/Public/Core/PostProcessSettingsCollection.h
--- a/Engine/Plugins/Cameras/GameplayCameras/Source/GameplayCameras/Public/Core/PostProcessSettingsCollection.h
+++ b/Engine/Plugins/Cameras/GameplayCameras/Source/GameplayCameras/Public/Core/PostProcessSettingsCollection.h
@@ -19,9 +19,6 @@ struct GAMEPLAYCAMERAS_API FPostProcessSettingsCollection
 	/** Gets the effective post-process settings. */
 	const FPostProcessSettings& Get() const { return PostProcessSettings; }
 
-	/** Returns whether this collection has any set post-process settings. */
-	bool HasAnyPostProcessSettings() const { return bHasAnySetting; }
-
 	/** Resets this collection to the default post-process settings. */
 	void Reset();
 
@@ -56,7 +53,6 @@ private:
 private:
 
 	FPostProcessSettings PostProcessSettings;
-	bool bHasAnySetting = false;
 };
 
 }  // namespace UE::Cameras
diff --git a/Engine/Plugins/Cameras/GameplayCameras/Source/GameplayCameras/Public/Debug/CameraNodeEvaluationResultDebugBlock.h b/Engine/Plugins/Cameras/GameplayCameras/Source/GameplayCameras/Public/Debug/CameraNodeEvaluationResultDebugBlock.h
--- a/Engine/Plugins/Cameras/GameplayCameras/Source/GameplayCameras/Public/Debug/CameraNodeEvaluationResultDebugBlock.h
+++ b/Engine/Plugins/Cameras/GameplayCameras/Source/GameplayCameras/Public/Debug/CameraNodeEvaluationResultDebugBlock.h
@@ -11,6 +11,7 @@ namespace UE::Cameras
 
 class FCameraPoseDebugBlock;
 class FContextDataTableDebugBlock;
+class FPostProcessSettingsDebugBlock;
 class FVariableTableDebugBlock;
 struct FCameraDebugBlockBuilder;
 struct FCameraNodeEvaluationResult;
@@ -52,6 +53,8 @@ public:
 	FVariableTableDebugBlock* GetVariableTableDebugBlock();
 	/** Gets the context data table debug block. */
 	FContextDataTableDebugBlock* GetContextDataTableDebugBlock();
+	/** Gets the pose stats debug block. */
+	FPostProcessSettingsDebugBlock* GetPostProcessSettingsDebugBlock();
 
 protected:
 
@@ -60,7 +63,6 @@ protected:
 
 private:
 
-	bool bHasAnyPostProcessSettings;
 	bool bIsCameraCut;
 	bool bIsValid;
 };
diff --git a/Engine/Plugins/Cameras/GameplayCameras/Source/GameplayCameras/Public/Debug/DebugTextRenderer.h b/Engine/Plugins/Cameras/GameplayCameras/Source/GameplayCameras/Public/Debug/DebugTextRenderer.h
--- a/Engine/Plugins/Cameras/GameplayCameras/Source/GameplayCameras/Public/Debug/DebugTextRenderer.h
+++ b/Engine/Plugins/Cameras/GameplayCameras/Source/GameplayCameras/Public/Debug/DebugTextRenderer.h
@@ -9,6 +9,7 @@
 #include "GameplayCameras.h"
 #include "Math/UnrealMath.h"
 #include "Misc/TVariant.h"
+#include <type_traits>
 
 class FCanvas;
 class UFont;
@@ -22,7 +23,21 @@ namespace UE::Cameras
 template<typename FieldType>
 FString ToDebugString(const FieldType& FieldValue)
 {
-	return LexToString(FieldValue);
+	if constexpr (std::is_enum_v<FieldType>)
+	{
+		const UEnum* EnumClass = StaticEnum<FieldType>();
+		return EnumClass->GetValueAsString(FieldValue);
+	}
+	else
+	{
+		return LexToString(FieldValue);
+	}
+}
+template<typename EnumType>
+FString ToDebugString(TEnumAsByte<EnumType> EnumValue)
+{
+	const UEnum* EnumClass = StaticEnum<EnumType>();
+	return EnumClass->GetValueAsString(EnumValue);
 }
 template<typename T>
 FString ToDebugString(const UE::Math::TVector<T>& FieldValue)
@@ -50,6 +65,11 @@ FString ToDebugString(const UE::Math::TTransform<T>& FieldValue)
 	return FieldValue.ToString();
 }
 template<>
+inline FString ToDebugString(const FLinearColor& FieldValue)
+{
+	return FieldValue.ToString();
+}
+template<>
 inline FString ToDebugString(const EAspectRatioAxisConstraint& FieldValue)
 {
 	switch (FieldValue)
diff --git a/Engine/Plugins/Cameras/GameplayCameras/Source/GameplayCameras/Public/Debug/PostProcessSettingsDebugBlock.h b/Engine/Plugins/Cameras/GameplayCameras/Source/GameplayCameras/Public/Debug/PostProcessSettingsDebugBlock.h
new file mode 100644
--- /dev/null
+++ b/Engine/Plugins/Cameras/GameplayCameras/Source/GameplayCameras/Public/Debug/PostProcessSettingsDebugBlock.h
@@ -0,0 +1,40 @@
+// Copyright Epic Games, Inc. All Rights Reserved.
+
+#pragma once
+
+#include "Core/PostProcessSettingsCollection.h"
+#include "Debug/CameraDebugBlock.h"
+
+#if UE_GAMEPLAY_CAMERAS_DEBUG
+
+namespace UE::Cameras
+{
+
+/**
+ * A debug block that displays information about post-process settings.
+ */
+class FPostProcessSettingsDebugBlock : public FCameraDebugBlock
+{
+	UE_DECLARE_CAMERA_DEBUG_BLOCK(GAMEPLAYCAMERAS_API, FPostProcessSettingsDebugBlock)
+
+public:
+
+	/** Creates a new post-process settings debug block. */
+	FPostProcessSettingsDebugBlock();
+	/** Creates a new post-process settings debug block. */
+	FPostProcessSettingsDebugBlock(const FPostProcessSettingsCollection& InPostProcessSettings);
+
+protected:
+
+	virtual void OnDebugDraw(const FCameraDebugBlockDrawParams& Params, FCameraDebugRenderer& Renderer) override;
+	virtual void OnSerialize(FArchive& Ar) override;
+
+private:
+
+	FPostProcessSettingsCollection PostProcessSettings;
+};
+
+}  // namespace UE::Cameras
+
+#endif  // UE_GAMEPLAY_CAMERAS_DEBUG
+
