Cameras: fix uninitialized projection mode

diff --git a/Engine/Plugins/Cameras/GameplayCameras/Source/GameplayCameras/Public/Core/CameraPose.h b/Engine/Plugins/Cameras/GameplayCameras/Source/GameplayCameras/Public/Core/CameraPose.h
--- a/Engine/Plugins/Cameras/GameplayCameras/Source/GameplayCameras/Public/Core/CameraPose.h
+++ b/Engine/Plugins/Cameras/GameplayCameras/Source/GameplayCameras/Public/Core/CameraPose.h
@@ -2,6 +2,7 @@
 
 #pragma once
 
+#include "Camera/CameraTypes.h"
 #include "CoreTypes.h"
 #include "Engine/EngineTypes.h"
 #include "EngineDefines.h"
@@ -13,7 +14,6 @@
 class FArchive;
 struct FPostProcessSettings;
 enum EAspectRatioAxisConstraint : int;
-namespace ECameraProjectionMode { enum Type : int; };
 
 #define UE_CAMERA_POSE_FOR_TRANSFORM_PROPERTIES()\
 	UE_CAMERA_POSE_FOR_PROPERTY(FVector, Location)\
@@ -347,7 +347,7 @@ private:
 
 	/** Projection mode - Perspective or orthographic */
 	UPROPERTY()
-	TEnumAsByte<ECameraProjectionMode::Type> ProjectionMode;
+	TEnumAsByte<ECameraProjectionMode::Type> ProjectionMode = ECameraProjectionMode::Perspective;
 
 private:
 	
