Fix for FBlueprintCameraPose::ProjectionMode not having a default value

diff --git a/Engine/Plugins/Cameras/GameplayCameras/Source/GameplayCameras/Public/GameFramework/BlueprintCameraPose.h b/Engine/Plugins/Cameras/GameplayCameras/Source/GameplayCameras/Public/GameFramework/BlueprintCameraPose.h
--- a/Engine/Plugins/Cameras/GameplayCameras/Source/GameplayCameras/Public/GameFramework/BlueprintCameraPose.h
+++ b/Engine/Plugins/Cameras/GameplayCameras/Source/GameplayCameras/Public/GameFramework/BlueprintCameraPose.h
@@ -129,7 +129,7 @@ public:
 
 	/** Projection mode - Perspective or orthographic */
 	UPROPERTY(BlueprintReadWrite, EditAnywhere, Category = "Camera")
-	TEnumAsByte<ECameraProjectionMode::Type> ProjectionMode;
+	TEnumAsByte<ECameraProjectionMode::Type> ProjectionMode = ECameraProjectionMode::Perspective;
 
 public:
 
