Cameras: correctly update pins on the camera rig prefab node when the reference changes

Now that the reference is in a structure, we need to change what property to check.

diff --git a/Engine/Plugins/Cameras/GameplayCameras/Source/GameplayCameras/Private/Nodes/Common/CameraRigCameraNode.cpp b/Engine/Plugins/Cameras/GameplayCameras/Source/GameplayCameras/Private/Nodes/Common/CameraRigCameraNode.cpp
--- a/Engine/Plugins/Cameras/GameplayCameras/Source/GameplayCameras/Private/Nodes/Common/CameraRigCameraNode.cpp
+++ b/Engine/Plugins/Cameras/GameplayCameras/Source/GameplayCameras/Private/Nodes/Common/CameraRigCameraNode.cpp
@@ -6,6 +6,7 @@
 #include "Build/CameraObjectBuildContext.h"
 #include "Core/CameraNodeEvaluator.h"
 #include "Core/CameraRigAsset.h"
+#include "Core/CameraRigAssetReference.h"
 #include "Helpers/CameraObjectReferenceParameterOverrideEvaluator.h"
 
 #include UE_INLINE_GENERATED_CPP_BY_NAME(CameraRigCameraNode)
@@ -124,7 +125,8 @@ void UCameraRigCameraNode::PostEditChangeProperty(struct FPropertyChangedEvent&
 {
 	Super::PostEditChangeProperty(PropertyChangedEvent);
 
-	if (PropertyChangedEvent.GetPropertyName() == GET_MEMBER_NAME_CHECKED(UCameraRigCameraNode, CameraRigReference))
+	if (PropertyChangedEvent.GetMemberPropertyName() == GET_MEMBER_NAME_CHECKED(UCameraRigCameraNode, CameraRigReference) &&
+			PropertyChangedEvent.GetPropertyName() == TEXT("CameraRig"))
 	{
 		OnCustomCameraNodeParametersChanged(this);
 	}
