Cameras: add debug info to attach-to-actor node, make it retain its transform when the attachment is lost

diff --git a/Engine/Plugins/Cameras/GameplayCameras/Source/GameplayCameras/Private/Nodes/Attach/AttachToActorCameraNode.cpp b/Engine/Plugins/Cameras/GameplayCameras/Source/GameplayCameras/Private/Nodes/Attach/AttachToActorCameraNode.cpp
--- a/Engine/Plugins/Cameras/GameplayCameras/Source/GameplayCameras/Private/Nodes/Attach/AttachToActorCameraNode.cpp
+++ b/Engine/Plugins/Cameras/GameplayCameras/Source/GameplayCameras/Private/Nodes/Attach/AttachToActorCameraNode.cpp
@@ -3,6 +3,9 @@
 #include "Nodes/Attach/AttachToActorCameraNode.h"
 
 #include "Core/CameraParameterReader.h"
+#include "Debug/CameraDebugBlock.h"
+#include "Debug/CameraDebugBlockBuilder.h"
+#include "Debug/CameraDebugRenderer.h"
 
 #include UE_INLINE_GENERATED_CPP_BY_NAME(AttachToActorCameraNode)
 
@@ -17,16 +20,31 @@ protected:
 
 	virtual void OnInitialize(const FCameraNodeEvaluatorInitializeParams& Params, FCameraNodeEvaluationResult& OutResult) override;
 	virtual void OnRun(const FCameraNodeEvaluationParams& Params, FCameraNodeEvaluationResult& OutResult) override;
+	virtual void OnSerialize(const FCameraNodeEvaluatorSerializeParams& Params, FArchive& Ar) override;
+
+#if UE_GAMEPLAY_CAMERAS_DEBUG
+	virtual void OnBuildDebugBlocks(const FCameraDebugBlockBuildParams& Params, FCameraDebugBlockBuilder& Builder) override;
+#endif  // UE_GAMEPLAY_CAMERAS_DEBUG
 
 private:
 
 	FCameraActorAttachmentInfoReader AttachmentReader;
 	TCameraParameterReader<bool> AttachToLocationReader;
 	TCameraParameterReader<bool> AttachToRotationReader;
+
+	TOptional<FTransform3d> LastAttachTransform;
+	bool bIsAttachValid = true;
 };
 
 UE_DEFINE_CAMERA_NODE_EVALUATOR(FAttachToActorCameraNodeEvaluator)
 
+UE_DECLARE_CAMERA_DEBUG_BLOCK_START(GAMEPLAYCAMERAS_API, FAttachToActorCameraDebugBlock)
+	UE_DECLARE_CAMERA_DEBUG_BLOCK_FIELD(FTransform3d, AttachTransform);
+	UE_DECLARE_CAMERA_DEBUG_BLOCK_FIELD(FString, AttachInfo);
+UE_DECLARE_CAMERA_DEBUG_BLOCK_END()
+
+UE_DEFINE_CAMERA_DEBUG_BLOCK_WITH_FIELDS(FAttachToActorCameraDebugBlock)
+
 void FAttachToActorCameraNodeEvaluator::OnInitialize(const FCameraNodeEvaluatorInitializeParams& Params, FCameraNodeEvaluationResult& OutResult)
 {
 	const UAttachToActorCameraNode* AttachNode = GetCameraNodeAs<UAttachToActorCameraNode>();
@@ -38,9 +56,26 @@ void FAttachToActorCameraNodeEvaluator::OnInitialize(const FCameraNodeEvaluatorI
 void FAttachToActorCameraNodeEvaluator::OnRun(const FCameraNodeEvaluationParams& Params, FCameraNodeEvaluationResult& OutResult)
 {
 	FTransform3d AttachTransform;
-	if (!AttachmentReader.GetAttachmentTransform(OutResult.ContextDataTable, AttachTransform))
+	const bool bWasAttachValid = bIsAttachValid;
+	bIsAttachValid = AttachmentReader.GetAttachmentTransform(OutResult.ContextDataTable, AttachTransform);
+	if (!bIsAttachValid)
 	{
-		return;
+		// If we just became invalid, log a warning.
+		if (bWasAttachValid)
+		{
+			UE_LOG(LogCameraSystem, Warning, 
+					TEXT("AttachToActorCameraNode: Couldn't resolve attachment! The camera will stay in place."));
+		}
+
+		// Stay in the last known place, if any, otherwise completely bail out.
+		if (LastAttachTransform.IsSet())
+		{
+			AttachTransform = LastAttachTransform.GetValue();
+		}
+		else
+		{
+			return;
+		}
 	}
 
 	const bool bAttachToLocation = AttachToLocationReader.Get(OutResult.VariableTable);
@@ -57,8 +92,40 @@ void FAttachToActorCameraNodeEvaluator::OnRun(const FCameraNodeEvaluationParams&
 		const FRotator3d AttachRotation = AttachTransform.GetRotation().Rotator();
 		OutResult.CameraPose.SetRotation(AttachRotation);
 	}
+
+	if (bIsAttachValid)
+	{
+		LastAttachTransform = AttachTransform;
+	}
 }
 
+void FAttachToActorCameraNodeEvaluator::OnSerialize(const FCameraNodeEvaluatorSerializeParams& Params, FArchive& Ar)
+{
+	Ar << LastAttachTransform;
+	Ar << bIsAttachValid;
+}
+
+#if UE_GAMEPLAY_CAMERAS_DEBUG
+
+void FAttachToActorCameraNodeEvaluator::OnBuildDebugBlocks(const FCameraDebugBlockBuildParams& Params, FCameraDebugBlockBuilder& Builder)
+{
+	FAttachToActorCameraDebugBlock& DebugBlock = Builder.AttachDebugBlock<FAttachToActorCameraDebugBlock>();
+	DebugBlock.AttachTransform = LastAttachTransform.Get(FTransform3d::Identity);
+	DebugBlock.AttachInfo = AttachmentReader.RenderAttachmentInfo();
+}
+
+void FAttachToActorCameraDebugBlock::OnDebugDraw(const FCameraDebugBlockDrawParams& Params, FCameraDebugRenderer& Renderer)
+{
+	Renderer.AddText(AttachInfo);
+
+	if (Renderer.IsExternalRendering())
+	{
+		Renderer.DrawPoint(AttachTransform.GetLocation(), 2.f, FColorList::NeonBlue, 2.f);
+	}
+}
+
+#endif  // UE_GAMEPLAY_CAMERAS_DEBUG
+
 }  // namespace UE::Cameras
 
 FCameraNodeEvaluatorPtr UAttachToActorCameraNode::OnBuildEvaluator(FCameraNodeEvaluatorBuilder& Builder) const
diff --git a/Engine/Plugins/Cameras/GameplayCameras/Source/GameplayCameras/Private/Nodes/Attach/CameraActorAttachmentInfo.cpp b/Engine/Plugins/Cameras/GameplayCameras/Source/GameplayCameras/Private/Nodes/Attach/CameraActorAttachmentInfo.cpp
--- a/Engine/Plugins/Cameras/GameplayCameras/Source/GameplayCameras/Private/Nodes/Attach/CameraActorAttachmentInfo.cpp
+++ b/Engine/Plugins/Cameras/GameplayCameras/Source/GameplayCameras/Private/Nodes/Attach/CameraActorAttachmentInfo.cpp
@@ -63,7 +63,7 @@ bool FCameraActorAttachmentInfoReader::GetAttachmentTransform(const FCameraConte
 
 	if (CachedSkeletalMeshComponent && !CachedBoneName.IsNone())
 	{
-		OutTransform = CachedSkeletalMeshComponent->GetSocketTransform(CachedAttachmentInfo.SocketName);
+		OutTransform = CachedSkeletalMeshComponent->GetBoneTransform(CachedBoneName);
 		return true;
 	}
 	else if (CachedAttachmentInfo.Actor)
@@ -150,5 +150,14 @@ bool FCameraActorAttachmentInfoArrayReader::GetAttachmentTransform(const FCamera
 	return true;
 }
 
+#if UE_GAMEPLAY_CAMERAS_DEBUG
+
+FString FCameraActorAttachmentInfoReader::RenderAttachmentInfo() const
+{
+	return FString::Printf(TEXT("Actor '%s' (Bone '%s')"), *GetNameSafe(CachedAttachmentInfo.Actor), *CachedBoneName.ToString());
+}
+
+#endif  // UE_GAMEPLAY_CAMERAS_DEBUG
+
 }  // namespace UE::Cameras
 
diff --git a/Engine/Plugins/Cameras/GameplayCameras/Source/GameplayCameras/Public/Nodes/Attach/CameraActorAttachmentInfo.h b/Engine/Plugins/Cameras/GameplayCameras/Source/GameplayCameras/Public/Nodes/Attach/CameraActorAttachmentInfo.h
--- a/Engine/Plugins/Cameras/GameplayCameras/Source/GameplayCameras/Public/Nodes/Attach/CameraActorAttachmentInfo.h
+++ b/Engine/Plugins/Cameras/GameplayCameras/Source/GameplayCameras/Public/Nodes/Attach/CameraActorAttachmentInfo.h
@@ -49,6 +49,10 @@ struct FCameraActorAttachmentInfoReader
 
 	bool GetAttachmentTransform(const FCameraContextDataTable& ContextDataTable, FTransform3d& OutTransform);
 
+#if UE_GAMEPLAY_CAMERAS_DEBUG
+	FString RenderAttachmentInfo() const;
+#endif  // UE_GAMEPLAY_CAMERAS_DEBUG
+
 private:
 
 	void CacheAttachmentInfo(const FCameraActorAttachmentInfo& InAttachmentInfo);
