Cameras: attachment and targeting info can now be set from Blueprint

diff --git a/Engine/Plugins/Cameras/GameplayCameras/Source/GameplayCameras/Public/Nodes/Attach/CameraActorAttachmentInfo.h b/Engine/Plugins/Cameras/GameplayCameras/Source/GameplayCameras/Public/Nodes/Attach/CameraActorAttachmentInfo.h
--- a/Engine/Plugins/Cameras/GameplayCameras/Source/GameplayCameras/Public/Nodes/Attach/CameraActorAttachmentInfo.h
+++ b/Engine/Plugins/Cameras/GameplayCameras/Source/GameplayCameras/Public/Nodes/Attach/CameraActorAttachmentInfo.h
@@ -18,19 +18,19 @@ struct GAMEPLAYCAMERAS_API FCameraActorAttachmentInfo
 	GENERATED_BODY()
 
 	/** The actor to attach to. */
-	UPROPERTY(EditAnywhere, Category="Attachment")
+	UPROPERTY(BlueprintReadWrite, EditAnywhere, Category="Attachment")
 	TObjectPtr<AActor> Actor;
 
 	/** An optional socket to attach to on the actor. */
-	UPROPERTY(EditAnywhere, Interp, Category="Attachment")
+	UPROPERTY(BlueprintReadWrite, EditAnywhere, Interp, Category="Attachment")
 	FName SocketName;
 
 	/** An optional bone to attach to on the actor. */
-	UPROPERTY(EditAnywhere, Interp, Category="Target")
+	UPROPERTY(BlueprintReadWrite, EditAnywhere, Interp, Category="Target")
 	FName BoneName;
 
 	/** The weight of this attachment. Unused if only one attachment is used. */
-	UPROPERTY(EditAnywhere, Interp, Category="Attachment")
+	UPROPERTY(BlueprintReadWrite, EditAnywhere, Interp, Category="Attachment")
 	float Weight = 1.f;
 
 	bool operator== (const FCameraActorAttachmentInfo& Other) const = default;
diff --git a/Engine/Plugins/Cameras/GameplayCameras/Source/GameplayCameras/Public/Nodes/Framing/CameraActorTargetInfo.h b/Engine/Plugins/Cameras/GameplayCameras/Source/GameplayCameras/Public/Nodes/Framing/CameraActorTargetInfo.h
--- a/Engine/Plugins/Cameras/GameplayCameras/Source/GameplayCameras/Public/Nodes/Framing/CameraActorTargetInfo.h
+++ b/Engine/Plugins/Cameras/GameplayCameras/Source/GameplayCameras/Public/Nodes/Framing/CameraActorTargetInfo.h
@@ -35,27 +35,27 @@ struct GAMEPLAYCAMERAS_API FCameraActorTargetInfo
 	GENERATED_BODY()
 
 	/** The actor to target. */
-	UPROPERTY(EditAnywhere, Interp, Category="Target")
+	UPROPERTY(BlueprintReadWrite, EditAnywhere, Interp, Category="Target")
 	TObjectPtr<AActor> Actor;
 
 	/** An optional socket to target on the actor. */
-	UPROPERTY(EditAnywhere, Interp, Category="Target")
+	UPROPERTY(BlueprintReadWrite, EditAnywhere, Interp, Category="Target")
 	FName SocketName;
 
 	/** An optional bone to target on the actor. */
-	UPROPERTY(EditAnywhere, Interp, Category="Target")
+	UPROPERTY(BlueprintReadWrite, EditAnywhere, Interp, Category="Target")
 	FName BoneName;
 
 	/** Specifies the shape of the target. */
-	UPROPERTY(EditAnywhere, Interp, Category="Target")
+	UPROPERTY(BlueprintReadWrite, EditAnywhere, Interp, Category="Target")
 	ECameraTargetShape TargetShape = ECameraTargetShape::Point;
 
 	/** Specifies the size of target. Only used if TargetShape is set to manual bounds. */
-	UPROPERTY(EditAnywhere, Interp, Category="Target")
+	UPROPERTY(BlueprintReadWrite, EditAnywhere, Interp, Category="Target")
 	float TargetSize = 10.f;
 
 	/** The weight of this target. Unused when only one target is used. */
-	UPROPERTY(EditAnywhere, Interp, Category="Target")
+	UPROPERTY(BlueprintReadWrite, EditAnywhere, Interp, Category="Target")
 	float Weight = 1.f;
 
 	bool IsValid() const { return Actor != nullptr; }
