Cameras: add world context to the Blueprint camera director

diff --git a/Engine/Plugins/Cameras/GameplayCameras/Source/GameplayCameras/Private/Directors/BlueprintCameraDirector.cpp b/Engine/Plugins/Cameras/GameplayCameras/Source/GameplayCameras/Private/Directors/BlueprintCameraDirector.cpp
--- a/Engine/Plugins/Cameras/GameplayCameras/Source/GameplayCameras/Private/Directors/BlueprintCameraDirector.cpp
+++ b/Engine/Plugins/Cameras/GameplayCameras/Source/GameplayCameras/Private/Directors/BlueprintCameraDirector.cpp
@@ -13,6 +13,7 @@
 #include "Core/CameraSystemEvaluator.h"
 #include "Core/RootCameraNode.h"
 #include "Engine/Blueprint.h"
+#include "Engine/World.h"
 #include "GameFramework/Actor.h"
 #include "GameFramework/ControllerGameplayCameraEvaluationComponent.h"
 #include "GameplayCameras.h"
@@ -399,6 +400,34 @@ void UBlueprintCameraDirectorEvaluator::SetInitialContextCameraPose(const FBluep
 	}
 }
 
+UWorld* UBlueprintCameraDirectorEvaluator::GetWorld() const
+{
+	if (UWorld* CachedWorld = WeakCachedWorld.Get())
+	{
+		return CachedWorld;
+	}
+
+	if (HasAllFlags(RF_ClassDefaultObject))
+	{
+		return nullptr;
+	}
+
+	UObject* Outer = GetOuter();
+	while (Outer)
+	{
+		UWorld* World = Outer->GetWorld();
+		if (World)
+		{
+			WeakCachedWorld = World;
+			return World;
+		}
+
+		Outer = Outer->GetOuter();
+	}
+
+	return nullptr;
+}
+
 void UBlueprintCameraDirectorEvaluator::NativeActivateCameraDirector(const UE::Cameras::FCameraDirectorActivateParams& Params)
 {
 	using namespace UE::Cameras;
diff --git a/Engine/Plugins/Cameras/GameplayCameras/Source/GameplayCameras/Public/Directors/BlueprintCameraDirector.h b/Engine/Plugins/Cameras/GameplayCameras/Source/GameplayCameras/Public/Directors/BlueprintCameraDirector.h
--- a/Engine/Plugins/Cameras/GameplayCameras/Source/GameplayCameras/Public/Directors/BlueprintCameraDirector.h
+++ b/Engine/Plugins/Cameras/GameplayCameras/Source/GameplayCameras/Public/Directors/BlueprintCameraDirector.h
@@ -196,6 +196,14 @@ public:
 			meta=(DeprecatedFunction, DeprecationMessage="Please use GetSharedEvaluationData"))
 	FBlueprintCameraVariableTable GetInitialContextVariableTable() const;
 
+public:
+
+	// UObject interface.
+	virtual UWorld* GetWorld() const override;
+#if WITH_EDITOR
+	virtual bool ImplementsGetWorld() const override { return true; }
+#endif  // WITH_EDITOR
+
 public:
 
 	using FCameraEvaluationContext = UE::Cameras::FCameraEvaluationContext;
@@ -220,6 +228,9 @@ private:
 
 	/** The current evaluation context. */
 	TSharedPtr<FCameraEvaluationContext> EvaluationContext;
+
+	/** Cached world. */
+	mutable TWeakObjectPtr<UWorld> WeakCachedWorld;
 };
 
 /**
