Cameras: replace all the Blueprint setter functions for camera parameters with a generic one that uses a CustomThunk implementation

This not only simplifies the code dramatically, it also makes it possible to support data parameters, especially arrays, without annoying "boxing" into instanced structs.

diff --git a/Engine/Plugins/Cameras/GameplayCameras/Source/GameplayCameras/Private/Core/CameraContextDataTable.cpp b/Engine/Plugins/Cameras/GameplayCameras/Source/GameplayCameras/Private/Core/CameraContextDataTable.cpp
--- a/Engine/Plugins/Cameras/GameplayCameras/Source/GameplayCameras/Private/Core/CameraContextDataTable.cpp
+++ b/Engine/Plugins/Cameras/GameplayCameras/Source/GameplayCameras/Private/Core/CameraContextDataTable.cpp
@@ -648,26 +648,26 @@ const uint8* FCameraContextDataTable::TryGetData(
 	return nullptr;
 }
 
-uint8* FCameraContextDataTable::TryGetMutableData(
+const FCameraContextDataTable::FEntryScriptArray* FCameraContextDataTable::TryGetArrayData(
 		FCameraContextDataID DataID,
 		ECameraContextDataType ExpectedDataType,
-		const UObject* ExpectedDataTypeObject)
+		const UObject* ExpectedDataTypeObject) const
 {
-	FEntry* Entry = FindEntry(DataID);
+	const FEntry* Entry = FindEntry(DataID);
 	if (Entry)
 	{
 		if (Entry->Type == ExpectedDataType && 
-				Entry->ContainerType == ECameraContextDataContainerType::None && 
+				Entry->ContainerType == ECameraContextDataContainerType::Array && 
 				Entry->TypeObject == ExpectedDataTypeObject)
 		{
-			return Memory + Entry->Offset;
+			return (FEntryScriptArray*)(Memory + Entry->Offset);
 		}
 	}
 
 	return nullptr;
 }
 
-const FCameraContextDataTable::FEntryScriptArray* FCameraContextDataTable::TryGetArrayData(
+const uint8* FCameraContextDataTable::TryGetRawDataPtr(
 		FCameraContextDataID DataID,
 		ECameraContextDataType ExpectedDataType,
 		const UObject* ExpectedDataTypeObject) const
@@ -675,11 +675,9 @@ const FCameraContextDataTable::FEntryScriptArray* FCameraContextDataTable::TryGe
 	const FEntry* Entry = FindEntry(DataID);
 	if (Entry)
 	{
-		if (Entry->Type == ExpectedDataType && 
-				Entry->ContainerType == ECameraContextDataContainerType::Array && 
-				Entry->TypeObject == ExpectedDataTypeObject)
+		if (Entry->Type == ExpectedDataType && Entry->TypeObject == ExpectedDataTypeObject)
 		{
-			return (FEntryScriptArray*)(Memory + Entry->Offset);
+			return Memory + Entry->Offset;
 		}
 	}
 
@@ -789,6 +787,29 @@ bool FCameraContextDataTable::TrySetArrayData(
 	return true;
 }
 
+uint8* FCameraContextDataTable::TryGetMutableRawDataPtr(
+		FCameraContextDataID DataID,
+		ECameraContextDataType ExpectedDataType,
+		const UObject* ExpectedDataTypeObject,
+		bool bMarkAsWrittenThisFrame)
+{
+	FEntry* Entry = FindEntry(DataID);
+	if (Entry)
+	{
+		if (Entry->Type == ExpectedDataType && Entry->TypeObject == ExpectedDataTypeObject)
+		{
+			Entry->Flags |= EEntryFlags::Written;
+			if (bMarkAsWrittenThisFrame)
+			{
+				Entry->Flags |= EEntryFlags::WrittenThisFrame;
+			}
+			return Memory + Entry->Offset;
+		}
+	}
+
+	return nullptr;
+}
+
 bool FCameraContextDataTable::SetDataValue(ECameraContextDataType DataType, ECameraContextDataContainerType DataContainerType, const UObject* DataTypeObject, uint8* DestDataPtr, const uint8* SrcDataPtr)
 {
 	if (DataContainerType == ECameraContextDataContainerType::None)
diff --git a/Engine/Plugins/Cameras/GameplayCameras/Source/GameplayCameras/Private/GameFramework/CameraRigParameterInterop.cpp b/Engine/Plugins/Cameras/GameplayCameras/Source/GameplayCameras/Private/GameFramework/CameraRigParameterInterop.cpp
--- a/Engine/Plugins/Cameras/GameplayCameras/Source/GameplayCameras/Private/GameFramework/CameraRigParameterInterop.cpp
+++ b/Engine/Plugins/Cameras/GameplayCameras/Source/GameplayCameras/Private/GameFramework/CameraRigParameterInterop.cpp
@@ -2,6 +2,8 @@
 
 #include "GameFramework/CameraRigParameterInterop.h"
 
+#include "Blueprint/BlueprintExceptionInfo.h"
+#include "Core/CameraObjectInterfaceParameterDefinition.h"
 #include "Core/CameraRigAsset.h"
 #include "Core/CameraNodeEvaluator.h"
 #include "Core/CameraVariableTable.h"
@@ -14,303 +16,168 @@
 
 #include UE_INLINE_GENERATED_CPP_BY_NAME(CameraRigParameterInterop)
 
-namespace UE::Cameras::Private
-{
-
-template<typename ValueType>
-void SetCameraRigBlendableParameter(FBlueprintCameraVariableTable VariableTable, FCameraVariableID PrivateVariableID, typename TCallTraits<ValueType>::ParamType Value)
-{
-	if (!VariableTable.IsValid())
-	{
-		FFrame::KismetExecutionMessage(TEXT("Invalid camera variable table was passed."), ELogVerbosity::Error);
-		return;
-	}
-	if (!PrivateVariableID.IsValid())
-	{
-		FFrame::KismetExecutionMessage(TEXT("No camera rig was passed."), ELogVerbosity::Error);
-		return;
-	}
-
-	VariableTable.GetVariableTable()->TrySetValue<ValueType>(PrivateVariableID, Value);
-}
-
-FCameraVariableID GetParameterPrivateVariableID(UCameraRigAsset* CameraRig, const FString& ParameterName, const UScriptStruct* BlendableStructType = nullptr)
-{
-	UCameraObjectInterfaceBlendableParameter* BlendableParameter = CameraRig->Interface.FindBlendableParameterByName(ParameterName);
-	if (!BlendableParameter)
-	{
-		const FText Text = LOCTEXT("NoSuchBlendableParameter", "No parameter '{0}' found on camera rig '{1}'. Setting this camera variable table value will most probably accomplish nothing.");
-		FFrame::KismetExecutionMessage(*FText::Format(Text, FText::FromString(ParameterName), FText::FromString(CameraRig->GetPathName())).ToString(), ELogVerbosity::Warning);
-		return FCameraVariableID();
-	}
-
-	if (!BlendableParameter->PrivateVariableID)
-	{
-		const FText Text = LOCTEXT("CameraRigNeedsBuilding", "Parameter '{0}' isn't built. Please build camera rig '{1}'.");
-		FFrame::KismetExecutionMessage(*FText::Format(Text, FText::FromString(ParameterName), FText::FromString(CameraRig->GetPathName())).ToString(), ELogVerbosity::Warning);
-		return FCameraVariableID();
-	}
-
-	if (BlendableParameter->BlendableStructType != BlendableStructType)
-	{
-		const FText Text = LOCTEXT("InvalidParameterType", "Parameter '{0}' has an incorrect type. Please build camera rig '{1}'.");
-		FFrame::KismetExecutionMessage(*FText::Format(Text, FText::FromString(ParameterName), FText::FromString(CameraRig->GetPathName())).ToString(), ELogVerbosity::Warning);
-		return FCameraVariableID();
-	}
-
-	return BlendableParameter->PrivateVariableID;
-}
-
-bool ValidateSetCameraRigDataParameter(FBlueprintCameraContextDataTable& ContextDataTable, FCameraContextDataID DataID)
-{
-	if (!ContextDataTable.IsValid())
-	{
-		FFrame::KismetExecutionMessage(TEXT("Invalid camera variable table was passed."), ELogVerbosity::Error);
-		return false;
-	}
-	if (!DataID.IsValid())
-	{
-		FFrame::KismetExecutionMessage(TEXT("No camera rig was passed."), ELogVerbosity::Error);
-		return false;
-	}
-	return true;
-}
-
-FCameraContextDataID GetParameterPrivateDataID(UCameraRigAsset* CameraRig, const FString& ParameterName)
-{
-	UCameraObjectInterfaceDataParameter* DataParameter = CameraRig->Interface.FindDataParameterByName(ParameterName);
-	if (!DataParameter)
-	{
-		const FText Text = LOCTEXT("NoSuchDataParameter", "No parameter '{0}' found on camera rig '{1}'. Setting this data will most probably accomplish nothing.");
-		FFrame::KismetExecutionMessage(*FText::Format(Text, FText::FromString(ParameterName), FText::FromString(CameraRig->GetPathName())).ToString(), ELogVerbosity::Warning);
-		return FCameraContextDataID();
-	}
-
-	if (!DataParameter->PrivateDataID)
-	{
-		const FText Text = LOCTEXT("CameraRigNeedsBuilding", "Parameter '{0}' isn't built. Please build camera rig '{1}'.");
-		FFrame::KismetExecutionMessage(*FText::Format(Text, FText::FromString(ParameterName), FText::FromString(CameraRig->GetPathName())).ToString(), ELogVerbosity::Warning);
-		return FCameraContextDataID();
-	}
-
-	return DataParameter->PrivateDataID;
-}
-
-}  // namespace UE::Cameras::Private
-
 UCameraRigParameterInterop::UCameraRigParameterInterop(const FObjectInitializer& ObjectInit)
 	: Super(ObjectInit)
 {
 }
 
-void UCameraRigParameterInterop::SetBooleanParameter(FBlueprintCameraNodeEvaluationResult& Result, UCameraRigAsset* CameraRig, const FString& ParameterName, bool ParameterValue)
-{
-	UE::Cameras::Private::SetCameraRigBlendableParameter<bool>(
-			Result.GetVariableTable(), 
-			UE::Cameras::Private::GetParameterPrivateVariableID(CameraRig, ParameterName), 
-			ParameterValue);
-}
-
-void UCameraRigParameterInterop::SetIntegerParameter(FBlueprintCameraNodeEvaluationResult& Result, UCameraRigAsset* CameraRig, const FString& ParameterName, int32 ParameterValue)
-{
-	UE::Cameras::Private::SetCameraRigBlendableParameter<int32>(
-			Result.GetVariableTable(), 
-			UE::Cameras::Private::GetParameterPrivateVariableID(CameraRig, ParameterName), 
-			ParameterValue);
-}
-
-void UCameraRigParameterInterop::SetFloatParameter(FBlueprintCameraNodeEvaluationResult& Result, UCameraRigAsset* CameraRig, const FString& ParameterName, double ParameterValue)
-{
-	UE::Cameras::Private::SetCameraRigBlendableParameter<float>(
-			Result.GetVariableTable(), 
-			UE::Cameras::Private::GetParameterPrivateVariableID(CameraRig, ParameterName),
-			(float)ParameterValue);
-}
-
-void UCameraRigParameterInterop::SetDoubleParameter(FBlueprintCameraNodeEvaluationResult& Result, UCameraRigAsset* CameraRig, const FString& ParameterName, double ParameterValue)
-{
-	UE::Cameras::Private::SetCameraRigBlendableParameter<double>(
-			Result.GetVariableTable(), 
-			UE::Cameras::Private::GetParameterPrivateVariableID(CameraRig, ParameterName),
-			ParameterValue);
-}
-
-void UCameraRigParameterInterop::SetVector2Parameter(FBlueprintCameraNodeEvaluationResult& Result, UCameraRigAsset* CameraRig, const FString& ParameterName, FVector2D ParameterValue)
-{
-	UE::Cameras::Private::SetCameraRigBlendableParameter<FVector2D>(
-			Result.GetVariableTable(), 
-			UE::Cameras::Private::GetParameterPrivateVariableID(CameraRig, ParameterName),
-			ParameterValue);
-}
-
-void UCameraRigParameterInterop::SetVector3Parameter(FBlueprintCameraNodeEvaluationResult& Result, UCameraRigAsset* CameraRig, const FString& ParameterName, FVector ParameterValue)
-{
-	UE::Cameras::Private::SetCameraRigBlendableParameter<FVector>(
-			Result.GetVariableTable(), 
-			UE::Cameras::Private::GetParameterPrivateVariableID(CameraRig, ParameterName),
-			ParameterValue);
-}
-
-void UCameraRigParameterInterop::SetVector4Parameter(FBlueprintCameraNodeEvaluationResult& Result, UCameraRigAsset* CameraRig, const FString& ParameterName, FVector4 ParameterValue)
-{
-	UE::Cameras::Private::SetCameraRigBlendableParameter<FVector4>(
-			Result.GetVariableTable(), 
-			UE::Cameras::Private::GetParameterPrivateVariableID(CameraRig, ParameterName),
-			ParameterValue);
-}
-
-void UCameraRigParameterInterop::SetRotatorParameter(FBlueprintCameraNodeEvaluationResult& Result, UCameraRigAsset* CameraRig, const FString& ParameterName, FRotator ParameterValue)
-{
-	UE::Cameras::Private::SetCameraRigBlendableParameter<FRotator>(
-			Result.GetVariableTable(), 
-			UE::Cameras::Private::GetParameterPrivateVariableID(CameraRig, ParameterName),
-			ParameterValue);
-}
-
-void UCameraRigParameterInterop::SetTransformParameter(FBlueprintCameraNodeEvaluationResult& Result, UCameraRigAsset* CameraRig, const FString& ParameterName, FTransform ParameterValue)
-{
-	UE::Cameras::Private::SetCameraRigBlendableParameter<FTransform>(
-			Result.GetVariableTable(), 
-			UE::Cameras::Private::GetParameterPrivateVariableID(CameraRig, ParameterName),
-			ParameterValue);
-}
-
-void UCameraRigParameterInterop::SetBlendableStructParameter(FBlueprintCameraNodeEvaluationResult& Result, UCameraRigAsset* CameraRig, const FString& ParameterName, const FInstancedStruct& ParameterValue)
+void UCameraRigParameterInterop::GetCameraParameter(const FBlueprintCameraNodeEvaluationResult& Result, UCameraRigAsset* CameraRig, FName ParameterName, int32& ReturnValue)
 {
-	FBlueprintCameraVariableTable VariableTable = Result.GetVariableTable();
-	if (!VariableTable.IsValid())
-	{
-		FFrame::KismetExecutionMessage(TEXT("Invalid camera variable table was passed."), ELogVerbosity::Error);
-		return;
-	}
-
-	FCameraVariableID PrivateVariableID = UE::Cameras::Private::GetParameterPrivateVariableID(CameraRig, ParameterName, ParameterValue.GetScriptStruct());
-	if (!PrivateVariableID.IsValid())
-	{
-		FFrame::KismetExecutionMessage(TEXT("No camera rig was passed."), ELogVerbosity::Error);
-		return;
-	}
-
-	VariableTable.GetVariableTable()->TrySetValue(PrivateVariableID, ECameraVariableType::BlendableStruct, ParameterValue.GetScriptStruct(), ParameterValue.GetMemory());
-}
-
-void UCameraRigParameterInterop::SetNameParameter(UPARAM(Ref) FBlueprintCameraNodeEvaluationResult& Result, UCameraRigAsset* CameraRig, const FString& ParameterName, FName ParameterValue)
-{
-	const FCameraContextDataID DataID(UE::Cameras::Private::GetParameterPrivateDataID(CameraRig, ParameterName));
-	FBlueprintCameraContextDataTable ContextDataTable(Result.GetContextDataTable());
-	if (UE::Cameras::Private::ValidateSetCameraRigDataParameter(ContextDataTable, DataID))
-	{
-		ContextDataTable.GetContextDataTable()->SetNameData(DataID, ParameterValue);
-	}
-}
-
-void UCameraRigParameterInterop::SetStringParameter(FBlueprintCameraNodeEvaluationResult& Result, UCameraRigAsset* CameraRig, const FString& ParameterName, FString ParameterValue)
-{
-	const FCameraContextDataID DataID(UE::Cameras::Private::GetParameterPrivateDataID(CameraRig, ParameterName));
-	FBlueprintCameraContextDataTable ContextDataTable(Result.GetContextDataTable());
-	if (UE::Cameras::Private::ValidateSetCameraRigDataParameter(ContextDataTable, DataID))
-	{
-		ContextDataTable.GetContextDataTable()->SetStringData(DataID, ParameterValue);
-	}
-}
-
-void UCameraRigParameterInterop::SetEnumParameter(FBlueprintCameraNodeEvaluationResult& Result, UCameraRigAsset* CameraRig, const FString& ParameterName, const UEnum* EnumType, uint8 ParameterValue)
-{
-	const FCameraContextDataID DataID(UE::Cameras::Private::GetParameterPrivateDataID(CameraRig, ParameterName));
-	FBlueprintCameraContextDataTable ContextDataTable(Result.GetContextDataTable());
-	if (UE::Cameras::Private::ValidateSetCameraRigDataParameter(ContextDataTable, DataID))
-	{
-		ContextDataTable.GetContextDataTable()->SetEnumData(DataID, EnumType, ParameterValue);
-	}
+	checkNoEntry();
 }
 
-void UCameraRigParameterInterop::SetStructParameter(FBlueprintCameraNodeEvaluationResult& Result, UCameraRigAsset* CameraRig, const FString& ParameterName, const FInstancedStruct& ParameterValue)
+void UCameraRigParameterInterop::SetCameraParameter(FBlueprintCameraNodeEvaluationResult& Result, UCameraRigAsset* CameraRig, FName ParameterName, const int32& NewValue)
 {
-	const FCameraContextDataID DataID(UE::Cameras::Private::GetParameterPrivateDataID(CameraRig, ParameterName));
-	FBlueprintCameraContextDataTable ContextDataTable(Result.GetContextDataTable());
-	if (UE::Cameras::Private::ValidateSetCameraRigDataParameter(ContextDataTable, DataID))
-	{
-		ContextDataTable.GetContextDataTable()->SetInstancedStructData(DataID, ParameterValue);
-	}
+	checkNoEntry();
 }
 
-void UCameraRigParameterInterop::SetObjectParameter(FBlueprintCameraNodeEvaluationResult& Result, UCameraRigAsset* CameraRig, const FString& ParameterName, UObject* ParameterValue)
+DEFINE_FUNCTION(UCameraRigParameterInterop::execGetCameraParameter)
 {
-	const FCameraContextDataID DataID(UE::Cameras::Private::GetParameterPrivateDataID(CameraRig, ParameterName));
-	FBlueprintCameraContextDataTable ContextDataTable(Result.GetContextDataTable());
-	if (UE::Cameras::Private::ValidateSetCameraRigDataParameter(ContextDataTable, DataID))
-	{
-		ContextDataTable.GetContextDataTable()->SetObjectData(DataID, ParameterValue);
-	}
-}
+	P_GET_STRUCT_REF(FBlueprintCameraNodeEvaluationResult, Result);
+	P_GET_OBJECT(UCameraRigAsset, CameraRig);
+	P_GET_STRUCT(FName, ParameterName);
 
-void UCameraRigParameterInterop::SetClassParameter(FBlueprintCameraNodeEvaluationResult& Result, UCameraRigAsset* CameraRig, const FString& ParameterName, UClass* ParameterValue)
-{
-	const FCameraContextDataID DataID(UE::Cameras::Private::GetParameterPrivateDataID(CameraRig, ParameterName));
-	FBlueprintCameraContextDataTable ContextDataTable(Result.GetContextDataTable());
-	if (UE::Cameras::Private::ValidateSetCameraRigDataParameter(ContextDataTable, DataID))
-	{
-		ContextDataTable.GetContextDataTable()->SetClassData(DataID, ParameterValue);
-	}
-}
+	// Read wildcard value input.
+	Stack.MostRecentPropertyAddress = nullptr;
+	Stack.MostRecentPropertyContainer = nullptr;
+	Stack.StepCompiledIn<FProperty>(nullptr);
 
-void UCameraRigParameterInterop::SetNameArrayParameter(UPARAM(Ref) FBlueprintCameraNodeEvaluationResult& Result, UCameraRigAsset* CameraRig, const FString& ParameterName, const TArray<FName>& ParameterValue)
-{
-	const FCameraContextDataID DataID(UE::Cameras::Private::GetParameterPrivateDataID(CameraRig, ParameterName));
-	FBlueprintCameraContextDataTable ContextDataTable(Result.GetContextDataTable());
-	if (UE::Cameras::Private::ValidateSetCameraRigDataParameter(ContextDataTable, DataID))
-	{
-		ContextDataTable.GetContextDataTable()->SetNameArrayData(DataID, ParameterValue);
-	}
-}
+	const FProperty* TargetProperty = Stack.MostRecentProperty;
+	void* TargetPtr = Stack.MostRecentPropertyAddress;
 
-void UCameraRigParameterInterop::SetStringArrayParameter(UPARAM(Ref) FBlueprintCameraNodeEvaluationResult& Result, UCameraRigAsset* CameraRig, const FString& ParameterName, const TArray<FString>& ParameterValue)
-{
-	const FCameraContextDataID DataID(UE::Cameras::Private::GetParameterPrivateDataID(CameraRig, ParameterName));
-	FBlueprintCameraContextDataTable ContextDataTable(Result.GetContextDataTable());
-	if (UE::Cameras::Private::ValidateSetCameraRigDataParameter(ContextDataTable, DataID))
-	{
-		ContextDataTable.GetContextDataTable()->SetStringArrayData(DataID, ParameterValue);
-	}
-}
+	P_FINISH;
 
-void UCameraRigParameterInterop::SetEnumArrayParameter(UPARAM(Ref) FBlueprintCameraNodeEvaluationResult& Result, UCameraRigAsset* CameraRig, const FString& ParameterName, const UEnum* EnumType, const TArray<uint8>& ParameterValue)
-{
-	const FCameraContextDataID DataID(UE::Cameras::Private::GetParameterPrivateDataID(CameraRig, ParameterName));
-	FBlueprintCameraContextDataTable ContextDataTable(Result.GetContextDataTable());
-	if (UE::Cameras::Private::ValidateSetCameraRigDataParameter(ContextDataTable, DataID))
+	if (TargetProperty == nullptr || TargetPtr == nullptr)
 	{
-		ContextDataTable.GetContextDataTable()->SetEnumArrayData(DataID, EnumType, ParameterValue);
+		FBlueprintExceptionInfo ExceptionInfo(
+			EBlueprintExceptionType::AbortExecution,
+			LOCTEXT("InvalidGetCameraParameterReturnValue", "Failed to resolve ReturnValue for GetCameraParameter")
+		);
+		FBlueprintCoreDelegates::ThrowScriptException(P_THIS, Stack, ExceptionInfo);
 	}
-}
-
-void UCameraRigParameterInterop::SetStructArrayParameter(UPARAM(Ref) FBlueprintCameraNodeEvaluationResult& Result, UCameraRigAsset* CameraRig, const FString& ParameterName, const TArray<FInstancedStruct>& ParameterValue)
-{
-	const FCameraContextDataID DataID(UE::Cameras::Private::GetParameterPrivateDataID(CameraRig, ParameterName));
-	FBlueprintCameraContextDataTable ContextDataTable(Result.GetContextDataTable());
-	if (UE::Cameras::Private::ValidateSetCameraRigDataParameter(ContextDataTable, DataID))
+	else
 	{
-		ContextDataTable.GetContextDataTable()->SetInstancedStructArrayData(DataID, ParameterValue);
-	}
-}
+		P_NATIVE_BEGIN
 
-void UCameraRigParameterInterop::SetObjectArrayParameter(UPARAM(Ref) FBlueprintCameraNodeEvaluationResult& Result, UCameraRigAsset* CameraRig, const FString& ParameterName, const TArray<UObject*>& ParameterValue)
-{
-	const FCameraContextDataID DataID(UE::Cameras::Private::GetParameterPrivateDataID(CameraRig, ParameterName));
-	FBlueprintCameraContextDataTable ContextDataTable(Result.GetContextDataTable());
-	if (UE::Cameras::Private::ValidateSetCameraRigDataParameter(ContextDataTable, DataID))
-	{
-		ContextDataTable.GetContextDataTable()->SetObjectArrayData(DataID, ParameterValue);
-	}
-}
+		using namespace UE::Cameras;
 
-void UCameraRigParameterInterop::SetClassArrayParameter(UPARAM(Ref) FBlueprintCameraNodeEvaluationResult& Result, UCameraRigAsset* CameraRig, const FString& ParameterName, const TArray<UClass*>& ParameterValue)
-{
-	const FCameraContextDataID DataID(UE::Cameras::Private::GetParameterPrivateDataID(CameraRig, ParameterName));
-	FBlueprintCameraContextDataTable ContextDataTable(Result.GetContextDataTable());
-	if (UE::Cameras::Private::ValidateSetCameraRigDataParameter(ContextDataTable, DataID))
-	{
-		ContextDataTable.GetContextDataTable()->SetClassArrayData(DataID, ParameterValue);
+		const FCameraObjectInterfaceParameterDefinition* ParameterDefinition = CameraRig->GetParameterDefinitions()
+			.FindByPredicate(
+				[ParameterName](const FCameraObjectInterfaceParameterDefinition& Item)
+				{
+					return Item.ParameterName == ParameterName;
+				});
+		if (!ParameterDefinition)
+		{
+			FBlueprintExceptionInfo ExceptionInfo(
+				EBlueprintExceptionType::NonFatalError,
+				FText::Format(LOCTEXT("ParameterDefinitionNotFound", "No such camera parameter: {0}"), FText::FromName(ParameterName))
+			);
+			FBlueprintCoreDelegates::ThrowScriptException(P_THIS, Stack, ExceptionInfo);
+		}
+		else
+		{
+			const FInstancedPropertyBag& DefaultParameters = CameraRig->GetDefaultParameters();
+			const UPropertyBag* PropertyBag = DefaultParameters.GetPropertyBagStruct();
+
+			if (ParameterDefinition->ParameterType == ECameraObjectInterfaceParameterType::Blendable)
+			{
+				if (const FCameraVariableTable* VariableTable = Result.GetVariableTable().GetVariableTable())
+				{
+					const uint8* RawValue = VariableTable->TryGetValue(ParameterDefinition->VariableID, ParameterDefinition->VariableType, ParameterDefinition->BlendableStructType);
+					if (RawValue)
+					{
+						TargetProperty->CopyCompleteValue(TargetPtr, RawValue);
+					}
+				}
+			}
+			else if (ParameterDefinition->ParameterType == ECameraObjectInterfaceParameterType::Data)
+			{
+				if (FCameraContextDataTable* ContextDataTable = Result.GetContextDataTable().GetContextDataTable())
+				{
+					const uint8* RawValue = ContextDataTable->TryGetRawDataPtr(ParameterDefinition->DataID, ParameterDefinition->DataType, ParameterDefinition->DataTypeObject);
+					if (RawValue)
+					{
+						TargetProperty->CopyCompleteValue(TargetPtr, RawValue);
+					}
+				}
+			}
+		}
+
+		P_NATIVE_END
+	}
+}
+
+DEFINE_FUNCTION(UCameraRigParameterInterop::execSetCameraParameter)
+{
+	P_GET_STRUCT_REF(FBlueprintCameraNodeEvaluationResult, Result);
+	P_GET_OBJECT(UCameraRigAsset, CameraRig);
+	P_GET_STRUCT(FName, ParameterName);
+
+	// Read wildcard value input.
+	Stack.MostRecentPropertyAddress = nullptr;
+	Stack.MostRecentPropertyContainer = nullptr;
+	Stack.StepCompiledIn<FProperty>(nullptr);
+
+	const FProperty* SourceProperty = Stack.MostRecentProperty;
+	const uint8* SourcePtr = Stack.MostRecentPropertyAddress;
+
+	P_FINISH;
+
+	if (SourceProperty == nullptr || SourcePtr == nullptr)
+	{
+		FBlueprintExceptionInfo ExceptionInfo(
+			EBlueprintExceptionType::AbortExecution,
+			LOCTEXT("InvalidSetCameraParameterNewValue", "Failed to resolve NewValue for SetCameraParameter")
+		);
+		FBlueprintCoreDelegates::ThrowScriptException(P_THIS, Stack, ExceptionInfo);
+	}
+	else
+	{
+		P_NATIVE_BEGIN
+
+		using namespace UE::Cameras;
+
+		const FCameraObjectInterfaceParameterDefinition* ParameterDefinition = CameraRig->GetParameterDefinitions()
+			.FindByPredicate(
+				[ParameterName](const FCameraObjectInterfaceParameterDefinition& Item)
+				{
+					return Item.ParameterName == ParameterName;
+				});
+		if (!ParameterDefinition)
+		{
+			FBlueprintExceptionInfo ExceptionInfo(
+				EBlueprintExceptionType::NonFatalError,
+				FText::Format(LOCTEXT("ParameterDefinitionNotFound", "No such camera parameter: {0}"), FText::FromName(ParameterName))
+			);
+			FBlueprintCoreDelegates::ThrowScriptException(P_THIS, Stack, ExceptionInfo);
+		}
+		else
+		{
+			const FInstancedPropertyBag& DefaultParameters = CameraRig->GetDefaultParameters();
+			const UPropertyBag* PropertyBag = DefaultParameters.GetPropertyBagStruct();
+
+			if (ParameterDefinition->ParameterType == ECameraObjectInterfaceParameterType::Blendable)
+			{
+				if (FCameraVariableTable* VariableTable = Result.GetVariableTable().GetVariableTable())
+				{
+					VariableTable->TrySetValue(ParameterDefinition->VariableID, ParameterDefinition->VariableType, ParameterDefinition->BlendableStructType, SourcePtr);
+				}
+			}
+			else if (ParameterDefinition->ParameterType == ECameraObjectInterfaceParameterType::Data)
+			{
+				if (FCameraContextDataTable* ContextDataTable = Result.GetContextDataTable().GetContextDataTable())
+				{
+					uint8* RawValue = ContextDataTable->TryGetMutableRawDataPtr(ParameterDefinition->DataID, ParameterDefinition->DataType, ParameterDefinition->DataTypeObject);
+					if (RawValue)
+					{
+						SourceProperty->CopyCompleteValue(RawValue, SourcePtr);
+					}
+				}
+			}
+		}
+
+		P_NATIVE_END
 	}
 }
 
diff --git a/Engine/Plugins/Cameras/GameplayCameras/Source/GameplayCameras/Public/Core/CameraContextDataTable.h b/Engine/Plugins/Cameras/GameplayCameras/Source/GameplayCameras/Public/Core/CameraContextDataTable.h
--- a/Engine/Plugins/Cameras/GameplayCameras/Source/GameplayCameras/Public/Core/CameraContextDataTable.h
+++ b/Engine/Plugins/Cameras/GameplayCameras/Source/GameplayCameras/Public/Core/CameraContextDataTable.h
@@ -159,12 +159,12 @@ public:
 			ECameraContextDataType ExpectedDataType,
 			const UObject* ExpectedDataTypeObject) const;
 
-	uint8* TryGetMutableData(
+	const FEntryScriptArray* TryGetArrayData(
 			FCameraContextDataID DataID,
 			ECameraContextDataType ExpectedDataType,
-			const UObject* ExpectedDataTypeObject);
+			const UObject* ExpectedDataTypeObject) const;
 
-	const FEntryScriptArray* TryGetArrayData(
+	const uint8* TryGetRawDataPtr(
 			FCameraContextDataID DataID,
 			ECameraContextDataType ExpectedDataType,
 			const UObject* ExpectedDataTypeObject) const;
@@ -196,6 +196,12 @@ public:
 			const uint8* InRawDataPtr,
 			bool bMarkAsWrittenThisFrame = true);
 
+	uint8* TryGetMutableRawDataPtr(
+			FCameraContextDataID DataID,
+			ECameraContextDataType ExpectedDataType,
+			const UObject* ExpectedDataTypeObject,
+			bool bMarkAsWrittenThisFrame = true);
+
 	bool IsValueWritten(FCameraContextDataID InID) const;
 	void UnsetValue(FCameraContextDataID InID);
 	void UnsetAllValues();
diff --git a/Engine/Plugins/Cameras/GameplayCameras/Source/GameplayCameras/Public/GameFramework/CameraRigParameterInterop.h b/Engine/Plugins/Cameras/GameplayCameras/Source/GameplayCameras/Public/GameFramework/CameraRigParameterInterop.h
--- a/Engine/Plugins/Cameras/GameplayCameras/Source/GameplayCameras/Public/GameFramework/CameraRigParameterInterop.h
+++ b/Engine/Plugins/Cameras/GameplayCameras/Source/GameplayCameras/Public/GameFramework/CameraRigParameterInterop.h
@@ -11,11 +11,7 @@
 #include "CameraRigParameterInterop.generated.h"
 
 class UCameraRigAsset;
-class UCameraVariableAsset;
-struct FBlueprintCameraContextDataTable;
 struct FBlueprintCameraNodeEvaluationResult;
-struct FBlueprintCameraVariableTable;
-struct FInstancedStruct;
 
 /**
  * Blueprint internal methods to set values on a camera rig's exposed parameters.
@@ -34,74 +30,15 @@ public:
 
 public:
 
-	UFUNCTION(BlueprintCallable, Category="Camera", meta=(BlueprintInternalUseOnly="true"))
-	static void SetBooleanParameter(UPARAM(Ref) FBlueprintCameraNodeEvaluationResult& Result, UCameraRigAsset* CameraRig, const FString& ParameterName, bool ParameterValue);
+	UFUNCTION(BlueprintCallable, CustomThunk, Category="Camera", meta=(BlueprintInternalUseOnly="true", CustomStructureParam="ReturnValue"))
+	static void GetCameraParameter(UPARAM(Ref) const FBlueprintCameraNodeEvaluationResult& Result, UCameraRigAsset* CameraRig, FName ParameterName, int32& ReturnValue);
 
-	UFUNCTION(BlueprintCallable, Category="Camera", meta=(BlueprintInternalUseOnly="true"))
-	static void SetIntegerParameter(UPARAM(Ref) FBlueprintCameraNodeEvaluationResult& Result, UCameraRigAsset* CameraRig, const FString& ParameterName, int32 ParameterValue);
+	UFUNCTION(BlueprintCallable, CustomThunk, Category="Camera", meta=(BlueprintInternalUseOnly="true", CustomStructureParam="NewValue"))
+	static void SetCameraParameter(UPARAM(Ref) FBlueprintCameraNodeEvaluationResult& Result, UCameraRigAsset* CameraRig, FName ParameterName, UPARAM(Ref) const int32& NewValue);
 
-	UFUNCTION(BlueprintCallable, Category="Camera", meta=(BlueprintInternalUseOnly="true"))
-	static void SetFloatParameter(UPARAM(Ref) FBlueprintCameraNodeEvaluationResult& Result, UCameraRigAsset* CameraRig, const FString& ParameterName, double ParameterValue);
+private:
 
-	UFUNCTION(BlueprintCallable, Category="Camera", meta=(BlueprintInternalUseOnly="true"))
-	static void SetDoubleParameter(UPARAM(Ref) FBlueprintCameraNodeEvaluationResult& Result, UCameraRigAsset* CameraRig, const FString& ParameterName, double ParameterValue);
-
-	UFUNCTION(BlueprintCallable, Category="Camera", meta=(BlueprintInternalUseOnly="true"))
-	static void SetVector2Parameter(UPARAM(Ref) FBlueprintCameraNodeEvaluationResult& Result, UCameraRigAsset* CameraRig, const FString& ParameterName, FVector2D ParameterValue);
-
-	UFUNCTION(BlueprintCallable, Category="Camera", meta=(BlueprintInternalUseOnly="true"))
-	static void SetVector3Parameter(UPARAM(Ref) FBlueprintCameraNodeEvaluationResult& Result, UCameraRigAsset* CameraRig, const FString& ParameterName, FVector ParameterValue);
-
-	UFUNCTION(BlueprintCallable, Category="Camera", meta=(BlueprintInternalUseOnly="true"))
-	static void SetVector4Parameter(UPARAM(Ref) FBlueprintCameraNodeEvaluationResult& Result, UCameraRigAsset* CameraRig, const FString& ParameterName, FVector4 ParameterValue);
-
-	UFUNCTION(BlueprintCallable, Category="Camera", meta=(BlueprintInternalUseOnly="true"))
-	static void SetRotatorParameter(UPARAM(Ref) FBlueprintCameraNodeEvaluationResult& Result, UCameraRigAsset* CameraRig, const FString& ParameterName, FRotator ParameterValue);
-
-	UFUNCTION(BlueprintCallable, Category="Camera", meta=(BlueprintInternalUseOnly="true"))
-	static void SetTransformParameter(UPARAM(Ref) FBlueprintCameraNodeEvaluationResult& Result, UCameraRigAsset* CameraRig, const FString& ParameterName, FTransform ParameterValue);
-
-	UFUNCTION(BlueprintCallable, Category="Camera", meta=(BlueprintInternalUseOnly="true"))
-	static void SetBlendableStructParameter(UPARAM(Ref) FBlueprintCameraNodeEvaluationResult& Result, UCameraRigAsset* CameraRig, const FString& ParameterName, const FInstancedStruct& ParameterValue);
-
-public:
-
-	UFUNCTION(BlueprintCallable, Category="Camera", meta=(BlueprintInternalUseOnly="true"))
-	static void SetNameParameter(UPARAM(Ref) FBlueprintCameraNodeEvaluationResult& Result, UCameraRigAsset* CameraRig, const FString& ParameterName, FName ParameterValue);
-
-	UFUNCTION(BlueprintCallable, Category="Camera", meta=(BlueprintInternalUseOnly="true"))
-	static void SetStringParameter(UPARAM(Ref) FBlueprintCameraNodeEvaluationResult& Result, UCameraRigAsset* CameraRig, const FString& ParameterName, FString ParameterValue);
-
-	UFUNCTION(BlueprintCallable, Category="Camera", meta=(BlueprintInternalUseOnly="true"))
-	static void SetEnumParameter(UPARAM(Ref) FBlueprintCameraNodeEvaluationResult& Result, UCameraRigAsset* CameraRig, const FString& ParameterName, const UEnum* EnumType, uint8 ParameterValue);
-
-	UFUNCTION(BlueprintCallable, Category="Camera", meta=(BlueprintInternalUseOnly="true"))
-	static void SetStructParameter(UPARAM(Ref) FBlueprintCameraNodeEvaluationResult& Result, UCameraRigAsset* CameraRig, const FString& ParameterName, const FInstancedStruct& ParameterValue);
-
-	UFUNCTION(BlueprintCallable, Category="Camera", meta=(BlueprintInternalUseOnly="true"))
-	static void SetObjectParameter(UPARAM(Ref) FBlueprintCameraNodeEvaluationResult& Result, UCameraRigAsset* CameraRig, const FString& ParameterName, UObject* ParameterValue);
-
-	UFUNCTION(BlueprintCallable, Category="Camera", meta=(BlueprintInternalUseOnly="true"))
-	static void SetClassParameter(UPARAM(Ref) FBlueprintCameraNodeEvaluationResult& Result, UCameraRigAsset* CameraRig, const FString& ParameterName, UClass* ParameterValue);
-
-public:
-
-	UFUNCTION(BlueprintCallable, Category="Camera", meta=(BlueprintInternalUseOnly="true"))
-	static void SetNameArrayParameter(UPARAM(Ref) FBlueprintCameraNodeEvaluationResult& Result, UCameraRigAsset* CameraRig, const FString& ParameterName, const TArray<FName>& ParameterValue);
-
-	UFUNCTION(BlueprintCallable, Category="Camera", meta=(BlueprintInternalUseOnly="true"))
-	static void SetStringArrayParameter(UPARAM(Ref) FBlueprintCameraNodeEvaluationResult& Result, UCameraRigAsset* CameraRig, const FString& ParameterName, const TArray<FString>& ParameterValue);
-
-	UFUNCTION(BlueprintCallable, Category="Camera", meta=(BlueprintInternalUseOnly="true"))
-	static void SetEnumArrayParameter(UPARAM(Ref) FBlueprintCameraNodeEvaluationResult& Result, UCameraRigAsset* CameraRig, const FString& ParameterName, const UEnum* EnumType, const TArray<uint8>& ParameterValue);
-
-	UFUNCTION(BlueprintCallable, Category="Camera", meta=(BlueprintInternalUseOnly="true"))
-	static void SetStructArrayParameter(UPARAM(Ref) FBlueprintCameraNodeEvaluationResult& Result, UCameraRigAsset* CameraRig, const FString& ParameterName, const TArray<FInstancedStruct>& ParameterValue);
-
-	UFUNCTION(BlueprintCallable, Category="Camera", meta=(BlueprintInternalUseOnly="true"))
-	static void SetObjectArrayParameter(UPARAM(Ref) FBlueprintCameraNodeEvaluationResult& Result, UCameraRigAsset* CameraRig, const FString& ParameterName, const TArray<UObject*>& ParameterValue);
-
-	UFUNCTION(BlueprintCallable, Category="Camera", meta=(BlueprintInternalUseOnly="true"))
-	static void SetClassArrayParameter(UPARAM(Ref) FBlueprintCameraNodeEvaluationResult& Result, UCameraRigAsset* CameraRig, const FString& ParameterName, const TArray<UClass*>& ParameterValue);
+	DECLARE_FUNCTION(execGetCameraParameter);
+	DECLARE_FUNCTION(execSetCameraParameter);
 };
 
diff --git a/Engine/Plugins/Cameras/GameplayCameras/Source/GameplayCamerasUncookedOnly/Private/BlueprintGraph/K2Node_CameraRigBase.cpp b/Engine/Plugins/Cameras/GameplayCameras/Source/GameplayCamerasUncookedOnly/Private/BlueprintGraph/K2Node_CameraRigBase.cpp
--- a/Engine/Plugins/Cameras/GameplayCameras/Source/GameplayCamerasUncookedOnly/Private/BlueprintGraph/K2Node_CameraRigBase.cpp
+++ b/Engine/Plugins/Cameras/GameplayCameras/Source/GameplayCamerasUncookedOnly/Private/BlueprintGraph/K2Node_CameraRigBase.cpp
@@ -238,120 +238,5 @@ FEdGraphPinType UK2Node_CameraRigBase::MakeDataParameterPinType(ECameraContextDa
 	return PinType;
 }
 
-FName UK2Node_CameraRigBase::GetBlendableParameterInteropSettingFunctionName(const UCameraObjectInterfaceBlendableParameter* BlendableParameter)
-{
-	if (ensure(BlendableParameter && BlendableParameter->PrivateVariableID))
-	{
-		return GetBlendableParameterInteropSettingFunctionName(BlendableParameter->ParameterType);
-	}
-	return NAME_None;
-}
-
-FName UK2Node_CameraRigBase::GetBlendableParameterInteropSettingFunctionName(ECameraVariableType CameraVariableType)
-{
-	// Figure out the sort of SetXxxParameter function we want to call for this parameter.
-	FName CallSetParameterFuncName;
-	switch (CameraVariableType)
-	{
-		case ECameraVariableType::Boolean:
-			CallSetParameterFuncName = GET_FUNCTION_NAME_CHECKED(UCameraRigParameterInterop, SetBooleanParameter);
-			break;
-		case ECameraVariableType::Integer32:
-			CallSetParameterFuncName = GET_FUNCTION_NAME_CHECKED(UCameraRigParameterInterop, SetIntegerParameter);
-			break;
-		case ECameraVariableType::Float:
-			CallSetParameterFuncName = GET_FUNCTION_NAME_CHECKED(UCameraRigParameterInterop, SetFloatParameter);
-			break;
-		case ECameraVariableType::Double:
-			CallSetParameterFuncName = GET_FUNCTION_NAME_CHECKED(UCameraRigParameterInterop, SetDoubleParameter);
-			break;
-		case ECameraVariableType::Vector2d:
-			CallSetParameterFuncName = GET_FUNCTION_NAME_CHECKED(UCameraRigParameterInterop, SetVector2Parameter);
-			break;
-		case ECameraVariableType::Vector3d:
-			CallSetParameterFuncName = GET_FUNCTION_NAME_CHECKED(UCameraRigParameterInterop, SetVector3Parameter);
-			break;
-		case ECameraVariableType::Vector4d:
-			CallSetParameterFuncName = GET_FUNCTION_NAME_CHECKED(UCameraRigParameterInterop, SetVector4Parameter);
-			break;
-		case ECameraVariableType::Rotator3d:
-			CallSetParameterFuncName = GET_FUNCTION_NAME_CHECKED(UCameraRigParameterInterop, SetRotatorParameter);
-			break;
-		case ECameraVariableType::Transform3d:
-			CallSetParameterFuncName = GET_FUNCTION_NAME_CHECKED(UCameraRigParameterInterop, SetTransformParameter);
-			break;
-		case ECameraVariableType::BlendableStruct:
-			CallSetParameterFuncName = GET_FUNCTION_NAME_CHECKED(UCameraRigParameterInterop, SetBlendableStructParameter);
-			break;
-	}
-
-	return CallSetParameterFuncName;
-}
-
-FName UK2Node_CameraRigBase::GetDataParameterInteropSettingFunctionName(const UCameraObjectInterfaceDataParameter* DataParameter)
-{
-	if (ensure(DataParameter && DataParameter->PrivateDataID.IsValid()))
-	{
-		return GetDataParameterInteropSettingFunctionName(DataParameter->DataType, DataParameter->DataContainerType, DataParameter->DataTypeObject);
-	}
-	return NAME_None;
-}
-
-FName UK2Node_CameraRigBase::GetDataParameterInteropSettingFunctionName(ECameraContextDataType CameraContextDataType, ECameraContextDataContainerType CameraContextDataContainerType, const UObject* CameraContextDataTypeObject)
-{
-	// Figure out the sort of SetXxxParameter function we want to call for this parameter.
-	FName CallSetParameterFuncName;
-	if (CameraContextDataContainerType == ECameraContextDataContainerType::None)
-	{
-		switch (CameraContextDataType)
-		{
-			case ECameraContextDataType::Name:
-				CallSetParameterFuncName = GET_FUNCTION_NAME_CHECKED(UCameraRigParameterInterop, SetNameParameter);
-				break;
-			case ECameraContextDataType::String:
-				CallSetParameterFuncName = GET_FUNCTION_NAME_CHECKED(UCameraRigParameterInterop, SetStringParameter);
-				break;
-			case ECameraContextDataType::Enum:
-				CallSetParameterFuncName = GET_FUNCTION_NAME_CHECKED(UCameraRigParameterInterop, SetEnumParameter);
-				break;
-			case ECameraContextDataType::Struct:
-				CallSetParameterFuncName = GET_FUNCTION_NAME_CHECKED(UCameraRigParameterInterop, SetStructParameter);
-				break;
-			case ECameraContextDataType::Object:
-				CallSetParameterFuncName = GET_FUNCTION_NAME_CHECKED(UCameraRigParameterInterop, SetObjectParameter);
-				break;
-			case ECameraContextDataType::Class:
-				CallSetParameterFuncName = GET_FUNCTION_NAME_CHECKED(UCameraRigParameterInterop, SetClassParameter);
-				break;
-		}
-	}
-	else
-	{
-		switch (CameraContextDataType)
-		{
-			case ECameraContextDataType::Name:
-				CallSetParameterFuncName = GET_FUNCTION_NAME_CHECKED(UCameraRigParameterInterop, SetNameArrayParameter);
-				break;
-			case ECameraContextDataType::String:
-				CallSetParameterFuncName = GET_FUNCTION_NAME_CHECKED(UCameraRigParameterInterop, SetStringArrayParameter);
-				break;
-			case ECameraContextDataType::Enum:
-				CallSetParameterFuncName = GET_FUNCTION_NAME_CHECKED(UCameraRigParameterInterop, SetEnumArrayParameter);
-				break;
-			case ECameraContextDataType::Struct:
-				CallSetParameterFuncName = GET_FUNCTION_NAME_CHECKED(UCameraRigParameterInterop, SetStructArrayParameter);
-				break;
-			case ECameraContextDataType::Object:
-				CallSetParameterFuncName = GET_FUNCTION_NAME_CHECKED(UCameraRigParameterInterop, SetObjectArrayParameter);
-				break;
-			case ECameraContextDataType::Class:
-				CallSetParameterFuncName = GET_FUNCTION_NAME_CHECKED(UCameraRigParameterInterop, SetClassArrayParameter);
-				break;
-		}
-	}
-
-	return CallSetParameterFuncName;
-}
-
 #undef LOCTEXT_NAMESPACE
 
diff --git a/Engine/Plugins/Cameras/GameplayCameras/Source/GameplayCamerasUncookedOnly/Private/BlueprintGraph/K2Node_CameraRigBase.h b/Engine/Plugins/Cameras/GameplayCameras/Source/GameplayCamerasUncookedOnly/Private/BlueprintGraph/K2Node_CameraRigBase.h
--- a/Engine/Plugins/Cameras/GameplayCameras/Source/GameplayCamerasUncookedOnly/Private/BlueprintGraph/K2Node_CameraRigBase.h
+++ b/Engine/Plugins/Cameras/GameplayCameras/Source/GameplayCamerasUncookedOnly/Private/BlueprintGraph/K2Node_CameraRigBase.h
@@ -49,11 +49,6 @@ public:
 	static FEdGraphPinType MakeDataParameterPinType(const UCameraObjectInterfaceDataParameter* DataParameter);
 	static FEdGraphPinType MakeDataParameterPinType(ECameraContextDataType CameraContextDataType, ECameraContextDataContainerType CameraContextDataContainerType, const UObject* CameraContextDataTypeObject);
 
-	static FName GetBlendableParameterInteropSettingFunctionName(const UCameraObjectInterfaceBlendableParameter* BlendableParameter);
-	static FName GetBlendableParameterInteropSettingFunctionName(ECameraVariableType CameraVariableType);
-	static FName GetDataParameterInteropSettingFunctionName(const UCameraObjectInterfaceDataParameter* DataParameter);
-	static FName GetDataParameterInteropSettingFunctionName(ECameraContextDataType CameraContextDataType, ECameraContextDataContainerType CameraContextDataContainerType, const UObject* CameraContextDataTypeObject);
-
 protected:
 
 	// UK2Node_CameraRigBase interface.
diff --git a/Engine/Plugins/Cameras/GameplayCameras/Source/GameplayCamerasUncookedOnly/Private/BlueprintGraph/K2Node_SetCameraRigParameter.cpp b/Engine/Plugins/Cameras/GameplayCameras/Source/GameplayCamerasUncookedOnly/Private/BlueprintGraph/K2Node_SetCameraRigParameter.cpp
--- a/Engine/Plugins/Cameras/GameplayCameras/Source/GameplayCamerasUncookedOnly/Private/BlueprintGraph/K2Node_SetCameraRigParameter.cpp
+++ b/Engine/Plugins/Cameras/GameplayCameras/Source/GameplayCamerasUncookedOnly/Private/BlueprintGraph/K2Node_SetCameraRigParameter.cpp
@@ -152,27 +152,6 @@ void UK2Node_SetCameraRigParameter::ExpandNode(FKismetCompilerContext& CompilerC
 		return;
 	}
 
-	FName CallSetParameterFuncName;
-	switch (CameraParameterType)
-	{
-		case EK2Node_CameraParameterType::Blendable:
-			CallSetParameterFuncName = UK2Node_CameraRigBase::GetBlendableParameterInteropSettingFunctionName(
-					BlendableCameraParameterType);
-			break;
-		case EK2Node_CameraParameterType::Data:
-			CallSetParameterFuncName = UK2Node_CameraRigBase::GetDataParameterInteropSettingFunctionName(
-					DataCameraParameterType, DataCameraParameterContainerType, DataCameraParameterTypeObject);
-			break;
-		default:
-			break;
-	}
-	if (CallSetParameterFuncName.IsNone())
-	{
-		CompilerContext.MessageLog.Error(*LOCTEXT("ErrorUnknownParameterType", "SetCameraRigParameter node @@ is trying to set parameter @@ but that parameter is missing or invalid.").ToString(), this, *CameraParameterName);
-		BreakAllNodeLinks();
-		return;
-	}
-
 	// Additional setup is needed for some types.
 	UEnum* ExtraEnumTypePinObject = nullptr;
 	bool bMakeExtraInstanceStructSetup = false;
@@ -191,7 +170,7 @@ void UK2Node_SetCameraRigParameter::ExpandNode(FKismetCompilerContext& CompilerC
 
 	// Make the SetXxxParameter function call node.
 	UK2Node_CallFunction* CallSetParameter = CompilerContext.SpawnIntermediateNode<UK2Node_CallFunction>(this, SourceGraph);
-	CallSetParameter->FunctionReference.SetExternalMember(CallSetParameterFuncName, UCameraRigParameterInterop::StaticClass());
+	CallSetParameter->FunctionReference.SetExternalMember(TEXT("SetCameraParameter"), UCameraRigParameterInterop::StaticClass());
 	CallSetParameter->AllocateDefaultPins();
 	UEdGraphPin* FirstExecPin = CallSetParameter->GetExecPin();
 
@@ -208,33 +187,15 @@ void UK2Node_SetCameraRigParameter::ExpandNode(FKismetCompilerContext& CompilerC
 	CallSetParameterNamePin->DefaultValue = CameraParameterName;
 
 	// Set or connect the parameter value argument.
-	UEdGraphPin* CallSetParameterValuePin = CallSetParameter->FindPinChecked(TEXT("ParameterValue"));
-	if (bMakeExtraInstanceStructSetup)
+	UEdGraphPin* CallSetParameterValuePin = CallSetParameter->FindPinChecked(TEXT("NewValue"));
+	CallSetParameterValuePin->PinType = CameraParameterValuePin->PinType;
+	CallSetParameterValuePin->DefaultValue = CameraParameterValuePin->DefaultValue;
+	CallSetParameterValuePin->DefaultTextValue = CameraParameterValuePin->DefaultTextValue;
+	CallSetParameterValuePin->AutogeneratedDefaultValue = CameraParameterValuePin->AutogeneratedDefaultValue;
+	CallSetParameterValuePin->DefaultObject = CameraParameterValuePin->DefaultObject;
+	if (CameraParameterValuePin->LinkedTo.Num() > 0)
 	{
-		// If we are setting a struct, first turn it into an FInstancedStruct. Insert this node in the
-		// execution chain, and use that node's Value pin as the pin to set the incoming struct value.
-		UK2Node_CallFunction* CallMakeInstancedStruct = CompilerContext.SpawnIntermediateNode<UK2Node_CallFunction>(this, SourceGraph);
-		CallMakeInstancedStruct->FunctionReference.SetExternalMember(GET_FUNCTION_NAME_CHECKED(UBlueprintInstancedStructLibrary, MakeInstancedStruct), UBlueprintInstancedStructLibrary::StaticClass());
-		CallMakeInstancedStruct->AllocateDefaultPins();
-
-		CallMakeInstancedStruct->GetReturnValuePin()->MakeLinkTo(CallSetParameterValuePin);
-		CallMakeInstancedStruct->GetThenPin()->MakeLinkTo(FirstExecPin);
-		FirstExecPin = CallMakeInstancedStruct->GetExecPin();
-
-		UEdGraphPin* CallMakeInstancedStructValuePin = CallMakeInstancedStruct->FindPinChecked(TEXT("Value"));
-		CompilerContext.MovePinLinksToIntermediate(*CameraParameterValuePin, *CallMakeInstancedStructValuePin);
-		CallMakeInstancedStruct->PinConnectionListChanged(CallMakeInstancedStructValuePin);
-	}
-	else
-	{
-		CallSetParameterValuePin->DefaultValue = CameraParameterValuePin->DefaultValue;
-		CallSetParameterValuePin->DefaultTextValue = CameraParameterValuePin->DefaultTextValue;
-		CallSetParameterValuePin->AutogeneratedDefaultValue = CameraParameterValuePin->AutogeneratedDefaultValue;
-		CallSetParameterValuePin->DefaultObject = CameraParameterValuePin->DefaultObject;
-		if (CameraParameterValuePin->LinkedTo.Num() > 0)
-		{
-			CompilerContext.MovePinLinksToIntermediate(*CameraParameterValuePin, *CallSetParameterValuePin);
-		}
+		CompilerContext.MovePinLinksToIntermediate(*CameraParameterValuePin, *CallSetParameterValuePin);
 	}
 
 	// Set extra type pin for enums.
diff --git a/Engine/Plugins/Cameras/GameplayCameras/Source/GameplayCamerasUncookedOnly/Private/BlueprintGraph/K2Node_SetCameraRigParameters.cpp b/Engine/Plugins/Cameras/GameplayCameras/Source/GameplayCamerasUncookedOnly/Private/BlueprintGraph/K2Node_SetCameraRigParameters.cpp
--- a/Engine/Plugins/Cameras/GameplayCameras/Source/GameplayCamerasUncookedOnly/Private/BlueprintGraph/K2Node_SetCameraRigParameters.cpp
+++ b/Engine/Plugins/Cameras/GameplayCameras/Source/GameplayCamerasUncookedOnly/Private/BlueprintGraph/K2Node_SetCameraRigParameters.cpp
@@ -158,18 +158,11 @@ void UK2Node_SetCameraRigParameters::ExpandNode(FKismetCompilerContext& Compiler
 			continue;
 		}
 
-		// Figure out the sort of SetXxxParameter function we want to call for this parameter.
-		FName CallSetParameterFuncName = GetBlendableParameterInteropSettingFunctionName(BlendableParameter);
-		if (CallSetParameterFuncName.IsNone())
-		{
-			CompilerContext.MessageLog.Error(*LOCTEXT("ErrorUnsupportedParameterType", "SetCameraRigParameters node @@ is trying to set parameter @@ but it has an unsupported type.").ToString(), this, *RigParameterPin->GetName());
-			continue;
-		}
-
 		// Make the SetXxxParameter function call node.
 		UK2Node_CallFunction* CallSetParameter = CompilerContext.SpawnIntermediateNode<UK2Node_CallFunction>(this, SourceGraph);
-		CallSetParameter->FunctionReference.SetExternalMember(CallSetParameterFuncName, UCameraRigParameterInterop::StaticClass());
+		CallSetParameter->FunctionReference.SetExternalMember(TEXT("SetCameraParameter"), UCameraRigParameterInterop::StaticClass());
 		CallSetParameter->AllocateDefaultPins();
+		CompilerContext.MessageLog.NotifyIntermediateObjectCreation(CallSetParameter, SourceGraph);
 
 		// Connect the camera evaluation result argument.
 		UEdGraphPin* CallSetParameterResultPin = CallSetParameter->FindPinChecked(TEXT("Result"));
@@ -184,7 +177,8 @@ void UK2Node_SetCameraRigParameters::ExpandNode(FKismetCompilerContext& Compiler
 		CallSetParameterNamePin->DefaultValue = BlendableParameter->InterfaceParameterName;
 
 		// Set or connect the parameter value argument.
-		UEdGraphPin* CallSetParameterValuePin = CallSetParameter->FindPinChecked(TEXT("ParameterValue"));
+		UEdGraphPin* CallSetParameterValuePin = CallSetParameter->FindPinChecked(TEXT("NewValue"));
+		CallSetParameterValuePin->PinType = RigParameterPin->PinType;
 		CallSetParameterValuePin->DefaultValue = RigParameterPin->DefaultValue;
 		CallSetParameterValuePin->DefaultTextValue = RigParameterPin->DefaultTextValue;
 		CallSetParameterValuePin->AutogeneratedDefaultValue = RigParameterPin->AutogeneratedDefaultValue;
@@ -226,18 +220,12 @@ void UK2Node_SetCameraRigParameters::ExpandNode(FKismetCompilerContext& Compiler
 			continue;
 		}
 
-		// Figure out the sort of SetXxxParameter function we want to call for this parameter.
-		FName CallSetParameterFuncName = GetDataParameterInteropSettingFunctionName(DataParameter);
-		if (CallSetParameterFuncName.IsNone())
-		{
-			CompilerContext.MessageLog.Error(*LOCTEXT("ErrorUnsupportedParameterType", "SetCameraRigParameters node @@ is trying to set parameter @@ but it has an unsupported type.").ToString(), this, *RigParameterPin->GetName());
-			continue;
-		}
-
 		// Make the SetXxxData function call node.
 		UK2Node_CallFunction* CallSetParameter = CompilerContext.SpawnIntermediateNode<UK2Node_CallFunction>(this, SourceGraph);
-		CallSetParameter->FunctionReference.SetExternalMember(CallSetParameterFuncName, UCameraRigParameterInterop::StaticClass());
+		CallSetParameter->FunctionReference.SetExternalMember(TEXT("SetCameraParameter"), UCameraRigParameterInterop::StaticClass());
 		CallSetParameter->AllocateDefaultPins();
+		CompilerContext.MessageLog.NotifyIntermediateObjectCreation(CallSetParameter, SourceGraph);
+
 		UEdGraphPin* CurExecPin = CallSetParameter->GetExecPin();
 
 		// Connect the camera evaluation result argument.
@@ -253,41 +241,15 @@ void UK2Node_SetCameraRigParameters::ExpandNode(FKismetCompilerContext& Compiler
 		CallSetParameterNamePin->DefaultValue = DataParameter->InterfaceParameterName;
 
 		// Set or connect the parameter value argument.
-		UEdGraphPin* CallSetParameterValuePin = CallSetParameter->FindPinChecked(TEXT("ParameterValue"));
-		if (DataParameter->DataType == ECameraContextDataType::Struct)
-		{
-			// If we are setting a struct, first turn it into an FInstancedStruct. Insert this node in the
-			// execution chain, and use that node's Value pin as the pin to set the incoming struct value.
-			UK2Node_CallFunction* CallMakeInstancedStruct = CompilerContext.SpawnIntermediateNode<UK2Node_CallFunction>(this, SourceGraph);
-			CallMakeInstancedStruct->FunctionReference.SetExternalMember(GET_FUNCTION_NAME_CHECKED(UBlueprintInstancedStructLibrary, MakeInstancedStruct), UBlueprintInstancedStructLibrary::StaticClass());
-			CallMakeInstancedStruct->AllocateDefaultPins();
-
-			CallMakeInstancedStruct->GetReturnValuePin()->MakeLinkTo(CallSetParameterValuePin);
-			CallMakeInstancedStruct->GetThenPin()->MakeLinkTo(CurExecPin);
-			CurExecPin = CallMakeInstancedStruct->GetExecPin();
-
-			UEdGraphPin* CallMakeInstancedStructValuePin = CallMakeInstancedStruct->FindPinChecked(TEXT("Value"));
-			CompilerContext.MovePinLinksToIntermediate(*RigParameterPin, *CallMakeInstancedStructValuePin);
-			CallMakeInstancedStruct->PinConnectionListChanged(CallMakeInstancedStructValuePin);
-		}
-		else
-		{
-			CallSetParameterValuePin->DefaultValue = RigParameterPin->DefaultValue;
-			CallSetParameterValuePin->DefaultTextValue = RigParameterPin->DefaultTextValue;
-			CallSetParameterValuePin->AutogeneratedDefaultValue = RigParameterPin->AutogeneratedDefaultValue;
-			CallSetParameterValuePin->DefaultObject = RigParameterPin->DefaultObject;
-			if (RigParameterPin->LinkedTo.Num() > 0)
-			{
-				CompilerContext.MovePinLinksToIntermediate(*RigParameterPin, *CallSetParameterValuePin);
-			}
-		}
-
-		// Set extra type pin for enums.
-		if (DataParameter->DataType == ECameraContextDataType::Enum)
+		UEdGraphPin* CallSetParameterValuePin = CallSetParameter->FindPinChecked(TEXT("NewValue"));
+		CallSetParameterValuePin->PinType = RigParameterPin->PinType;
+		CallSetParameterValuePin->DefaultValue = RigParameterPin->DefaultValue;
+		CallSetParameterValuePin->DefaultTextValue = RigParameterPin->DefaultTextValue;
+		CallSetParameterValuePin->AutogeneratedDefaultValue = RigParameterPin->AutogeneratedDefaultValue;
+		CallSetParameterValuePin->DefaultObject = RigParameterPin->DefaultObject;
+		if (RigParameterPin->LinkedTo.Num() > 0)
 		{
-			const UEnum* EnumType = CastChecked<const UEnum>(DataParameter->DataTypeObject);
-			UEdGraphPin* CallSetParameterEnumTypePin = CallSetParameter->FindPinChecked(TEXT("EnumType"));
-			CallSetParameterEnumTypePin->DefaultObject = const_cast<UEnum*>(EnumType);
+			CompilerContext.MovePinLinksToIntermediate(*RigParameterPin, *CallSetParameterValuePin);
 		}
 
 		// Chain the execution.
