Cameras: fix possible memory corruption when gathering custom parameters from a camera rig reference

The previous code was returning pointers to the value field of the TValueOrError<> structures returned by the instanced struct API. This was obviously bad, since those only contained *copies* of the struct data. The new code returns pointers to the actual value fields of the instanced struct.

diff --git a/Engine/Plugins/Cameras/GameplayCameras/Source/GameplayCameras/Private/Core/BaseCameraObjectReference.cpp b/Engine/Plugins/Cameras/GameplayCameras/Source/GameplayCameras/Private/Core/BaseCameraObjectReference.cpp
--- a/Engine/Plugins/Cameras/GameplayCameras/Source/GameplayCameras/Private/Core/BaseCameraObjectReference.cpp
+++ b/Engine/Plugins/Cameras/GameplayCameras/Source/GameplayCameras/Private/Core/BaseCameraObjectReference.cpp
@@ -164,9 +164,10 @@ void FBaseCameraObjectReference::GetCustomCameraNodeParameters(FCustomCameraNode
 		return;
 	}
 
+	uint8* ParametersMemory = Parameters.GetMutableValue().GetMemory();
 	const UPropertyBag* ParametersStruct = Parameters.GetPropertyBagStruct();
 	const FInstancedPropertyBag& DefaultParameters = CameraObject->GetDefaultParameters();
-	if (!ensure(ParametersStruct && ParametersStruct == DefaultParameters.GetPropertyBagStruct()))
+	if (!ensure(ParametersMemory && ParametersStruct && ParametersStruct == DefaultParameters.GetPropertyBagStruct()))
 	{
 		return;
 	}
@@ -174,7 +175,7 @@ void FBaseCameraObjectReference::GetCustomCameraNodeParameters(FCustomCameraNode
 	for (const FCameraObjectInterfaceParameterDefinition& Definition : CameraObject->GetParameterDefinitions())
 	{
 		const FPropertyBagPropertyDesc* PropertyDesc = ParametersStruct->FindPropertyDescByID(Definition.ParameterGuid);
-		if (!ensure(PropertyDesc))
+		if (!ensure(PropertyDesc && PropertyDesc->CachedProperty))
 		{
 			continue;
 		}
@@ -193,12 +194,10 @@ void FBaseCameraObjectReference::GetCustomCameraNodeParameters(FCustomCameraNode
 					if (ensure(PropertyDesc->ValueTypeObject == F##ValueName##CameraParameter::StaticStruct()))\
 					{\
 						using CameraParameterType = F##ValueName##CameraParameter;\
-						TValueOrError<CameraParameterType*, EPropertyBagResult> PropertyValue =\
-							Parameters.GetValueStruct<CameraParameterType>(*PropertyDesc);\
-						if (ensure(PropertyValue.HasValue() && !PropertyValue.HasError()))\
+						void* PropertyValue = PropertyDesc->CachedProperty->ContainerPtrToValuePtr<void>(ParametersMemory);\
+						if (ensure(PropertyValue))\
 						{\
-							CameraParameterType* CameraParameter = PropertyValue.GetValue();\
-							check(CameraParameter);\
+							CameraParameterType* CameraParameter = static_cast<CameraParameterType*>(PropertyValue);\
 							OutParameterInfos.AddBlendableParameter(\
 									Definition.ParameterName,\
 									Definition.VariableType,\
@@ -212,19 +211,15 @@ void FBaseCameraObjectReference::GetCustomCameraNodeParameters(FCustomCameraNode
 #undef UE_CAMERA_VARIABLE_FOR_TYPE
 				case ECameraVariableType::BlendableStruct:
 					{
-						TValueOrError<FStructView, EPropertyBagResult> PropertyValue =
-							Parameters.GetValueStruct(*PropertyDesc, Definition.BlendableStructType);
-						if (ensure(PropertyValue.HasValue() && !PropertyValue.HasError()))
+						void* PropertyValue = PropertyDesc->CachedProperty->ContainerPtrToValuePtr<void>(ParametersMemory);
+						if (ensure(PropertyValue))
 						{
-							FStructView& StructValue = PropertyValue.GetValue();
-							check(StructValue.IsValid());
-
 							FCameraObjectInterfaceParameterMetaData& MetaData = FindOrAddMetaData(Definition.ParameterGuid);
 							OutParameterInfos.AddBlendableParameter(
 									Definition.ParameterName,
 									Definition.VariableType,
 									Definition.BlendableStructType,
-									StructValue.GetMemory(),
+									reinterpret_cast<uint8*>(PropertyValue),
 									&MetaData.OverrideVariableID);
 						}
 					}
@@ -235,20 +230,26 @@ void FBaseCameraObjectReference::GetCustomCameraNodeParameters(FCustomCameraNode
 		{
 			FCameraObjectInterfaceParameterMetaData& MetaData = FindOrAddMetaData(Definition.ParameterGuid);
 
+			ECameraContextDataContainerType ContainerType = ECameraContextDataContainerType::None;
+			if (PropertyDesc->ContainerTypes.GetFirstContainerType() == EPropertyBagContainerType::Array)
+			{
+				ContainerType = ECameraContextDataContainerType::Array;
+			}
+
+			const uint8* DefaultValue = reinterpret_cast<uint8*>(PropertyDesc->CachedProperty->ContainerPtrToValuePtr<void>(ParametersMemory));
+
 			switch (Definition.DataType)
 			{
 				case ECameraContextDataType::Name:
 					if (ensure(PropertyDesc->ValueType == EPropertyBagPropertyType::Name))
 					{
-						const uint8* DefaultValue = reinterpret_cast<uint8*>(DefaultParameters.GetValueName(*PropertyDesc).TryGetValue());
-						OutParameterInfos.AddDataParameter(Definition.ParameterName, ECameraContextDataType::Name, ECameraContextDataContainerType::None, nullptr, DefaultValue, &MetaData.OverrideDataID);
+						OutParameterInfos.AddDataParameter(Definition.ParameterName, ECameraContextDataType::Name, ContainerType, nullptr, DefaultValue, &MetaData.OverrideDataID);
 					}
 					break;
 				case ECameraContextDataType::String:
 					if (ensure(PropertyDesc->ValueType == EPropertyBagPropertyType::String))
 					{
-						const uint8* DefaultValue = reinterpret_cast<uint8*>(DefaultParameters.GetValueString(*PropertyDesc).TryGetValue());
-						OutParameterInfos.AddDataParameter(Definition.ParameterName, ECameraContextDataType::String, ECameraContextDataContainerType::None, nullptr, DefaultValue, &MetaData.OverrideDataID);
+						OutParameterInfos.AddDataParameter(Definition.ParameterName, ECameraContextDataType::String, ContainerType, nullptr, DefaultValue, &MetaData.OverrideDataID);
 					}
 					break;
 				case ECameraContextDataType::Enum:
@@ -256,8 +257,7 @@ void FBaseCameraObjectReference::GetCustomCameraNodeParameters(FCustomCameraNode
 								PropertyDesc->ValueTypeObject == Definition.DataTypeObject))
 					{
 						const UEnum* EnumType = CastChecked<const UEnum>(Definition.DataTypeObject);
-						const uint8* DefaultValue = reinterpret_cast<uint8*>(DefaultParameters.GetValueEnum(*PropertyDesc, EnumType).TryGetValue());
-						OutParameterInfos.AddDataParameter(Definition.ParameterName, ECameraContextDataType::Enum, ECameraContextDataContainerType::None, EnumType, DefaultValue, &MetaData.OverrideDataID);
+						OutParameterInfos.AddDataParameter(Definition.ParameterName, ECameraContextDataType::Enum, ContainerType, EnumType, DefaultValue, &MetaData.OverrideDataID);
 					}
 					break;
 				case ECameraContextDataType::Struct:
@@ -265,22 +265,19 @@ void FBaseCameraObjectReference::GetCustomCameraNodeParameters(FCustomCameraNode
 								PropertyDesc->ValueTypeObject == Definition.DataTypeObject))
 					{
 						const UScriptStruct* DataType = CastChecked<const UScriptStruct>(Definition.DataTypeObject);
-						const uint8* DefaultValue = reinterpret_cast<uint8*>(DefaultParameters.GetValueStruct(*PropertyDesc, DataType).TryGetValue());
-						OutParameterInfos.AddDataParameter(Definition.ParameterName, ECameraContextDataType::Struct, ECameraContextDataContainerType::None, DataType, DefaultValue, &MetaData.OverrideDataID);
+						OutParameterInfos.AddDataParameter(Definition.ParameterName, ECameraContextDataType::Struct, ContainerType, DataType, DefaultValue, &MetaData.OverrideDataID);
 					}
 					break;
 				case ECameraContextDataType::Object:
 					if (ensure(PropertyDesc->ValueType == EPropertyBagPropertyType::Object))
 					{
-						const uint8* DefaultValue = reinterpret_cast<uint8*>(DefaultParameters.GetValueObject(*PropertyDesc).TryGetValue());
-						OutParameterInfos.AddDataParameter(Definition.ParameterName, ECameraContextDataType::Object, ECameraContextDataContainerType::None, Definition.DataTypeObject, DefaultValue, &MetaData.OverrideDataID);
+						OutParameterInfos.AddDataParameter(Definition.ParameterName, ECameraContextDataType::Object, ContainerType, Definition.DataTypeObject, DefaultValue, &MetaData.OverrideDataID);
 					}
 					break;
 				case ECameraContextDataType::Class:
 					if (ensure(PropertyDesc->ValueType == EPropertyBagPropertyType::Class))
 					{
-						const uint8* DefaultValue = reinterpret_cast<uint8*>(DefaultParameters.GetValueClass(*PropertyDesc).TryGetValue());
-						OutParameterInfos.AddDataParameter(Definition.ParameterName, ECameraContextDataType::Class, ECameraContextDataContainerType::None, Definition.DataTypeObject, DefaultValue, &MetaData.OverrideDataID);
+						OutParameterInfos.AddDataParameter(Definition.ParameterName, ECameraContextDataType::Class, ContainerType, Definition.DataTypeObject, DefaultValue, &MetaData.OverrideDataID);
 					}
 					break;
 			}
