Cameras: correctly update extra data for transient and persistent blend stack entries when an entry is hot-reloaded.

diff --git a/Engine/Plugins/Cameras/GameplayCameras/Source/GameplayCameras/Private/Core/BlendStackCameraNode.cpp b/Engine/Plugins/Cameras/GameplayCameras/Source/GameplayCameras/Private/Core/BlendStackCameraNode.cpp
--- a/Engine/Plugins/Cameras/GameplayCameras/Source/GameplayCameras/Private/Core/BlendStackCameraNode.cpp
+++ b/Engine/Plugins/Cameras/GameplayCameras/Source/GameplayCameras/Private/Core/BlendStackCameraNode.cpp
@@ -522,6 +522,8 @@ void FBlendStackCameraNodeEvaluator::OnPostBuildAsset(const FGameplayCameraAsset
 				Entry.bIsFrozen = true;
 				continue;
 			}
+
+			OnEntryReinitialized(Index);
 		}
 	}
 }
diff --git a/Engine/Plugins/Cameras/GameplayCameras/Source/GameplayCameras/Private/Core/PersistentBlendStackCameraNode.cpp b/Engine/Plugins/Cameras/GameplayCameras/Source/GameplayCameras/Private/Core/PersistentBlendStackCameraNode.cpp
--- a/Engine/Plugins/Cameras/GameplayCameras/Source/GameplayCameras/Private/Core/PersistentBlendStackCameraNode.cpp
+++ b/Engine/Plugins/Cameras/GameplayCameras/Source/GameplayCameras/Private/Core/PersistentBlendStackCameraNode.cpp
@@ -359,5 +359,26 @@ const UCameraRigTransition* FPersistentBlendStackCameraNodeEvaluator::FindExitTr
 			nullptr, nullptr);
 }
 
+#if WITH_EDITOR
+
+void FPersistentBlendStackCameraNodeEvaluator::OnEntryReinitialized(int32 EntryIndex)
+{
+	if (!ensure(EntryExtraInfos.IsValidIndex(EntryIndex)))
+	{
+		return;
+	}
+
+	FCameraRigEntryExtraInfo& ExtraInfo = EntryExtraInfos[EntryIndex];
+	{
+		// When hot-reloading camera rigs, the base class replaces the blend node with pop so
+		// let's update our own extra info accordingly.
+		ExtraInfo.bIsBlendFull = true;
+		ExtraInfo.bIsBlendFinished = true;
+		ExtraInfo.BlendStatus = EBlendStatus::None;
+	}
+}
+
+#endif
+
 }  // namespace UE::Cameras
 
diff --git a/Engine/Plugins/Cameras/GameplayCameras/Source/GameplayCameras/Private/Core/TransientBlendStackCameraNode.cpp b/Engine/Plugins/Cameras/GameplayCameras/Source/GameplayCameras/Private/Core/TransientBlendStackCameraNode.cpp
--- a/Engine/Plugins/Cameras/GameplayCameras/Source/GameplayCameras/Private/Core/TransientBlendStackCameraNode.cpp
+++ b/Engine/Plugins/Cameras/GameplayCameras/Source/GameplayCameras/Private/Core/TransientBlendStackCameraNode.cpp
@@ -649,5 +649,33 @@ const UCameraRigTransition* FTransientBlendStackCameraNodeEvaluator::FindTransit
 	return nullptr;
 }
 
+#if WITH_EDITOR
+
+void FTransientBlendStackCameraNodeEvaluator::OnEntryReinitialized(int32 EntryIndex)
+{
+	if (!ensure(EntryExtraInfos.IsValidIndex(EntryIndex)))
+	{
+		return;
+	}
+
+	FCameraRigEntry& Entry = Entries[EntryIndex];
+	FCameraRigEntryExtraInfo& ExtraInfo = EntryExtraInfos[EntryIndex];
+	{
+		// Update our extra info with the new data.
+		FCameraNodeEvaluationResult& UnblendedResult = ExtraInfo.UnblendedResult;
+
+		if (Entry.CameraRig)
+		{
+			const FCameraObjectAllocationInfo& AllocationInfo = Entry.CameraRig->AllocationInfo;
+			UnblendedResult.VariableTable.Initialize(AllocationInfo.VariableTableInfo);
+			UnblendedResult.ContextDataTable.Initialize(AllocationInfo.ContextDataTableInfo);
+		}
+
+		UnblendedResult.OverrideAll(Entry.Result, true);
+	}
+}
+
+#endif
+
 }  // namespace UE::Cameras
 
diff --git a/Engine/Plugins/Cameras/GameplayCameras/Source/GameplayCameras/Public/Core/BlendStackCameraNode.h b/Engine/Plugins/Cameras/GameplayCameras/Source/GameplayCameras/Public/Core/BlendStackCameraNode.h
--- a/Engine/Plugins/Cameras/GameplayCameras/Source/GameplayCameras/Public/Core/BlendStackCameraNode.h
+++ b/Engine/Plugins/Cameras/GameplayCameras/Source/GameplayCameras/Public/Core/BlendStackCameraNode.h
@@ -162,6 +162,12 @@ protected:
 	void RemoveListenedPackages(TSharedPtr<IGameplayCamerasLiveEditManager> LiveEditManager, FCameraRigEntry& Entry);
 #endif
 
+protected:
+
+#if WITH_EDITOR
+	virtual void OnEntryReinitialized(int32 EntryIndex) {}
+#endif
+
 protected:
 
 	struct FResolvedEntry
diff --git a/Engine/Plugins/Cameras/GameplayCameras/Source/GameplayCameras/Public/Core/PersistentBlendStackCameraNode.h b/Engine/Plugins/Cameras/GameplayCameras/Source/GameplayCameras/Public/Core/PersistentBlendStackCameraNode.h
--- a/Engine/Plugins/Cameras/GameplayCameras/Source/GameplayCameras/Public/Core/PersistentBlendStackCameraNode.h
+++ b/Engine/Plugins/Cameras/GameplayCameras/Source/GameplayCameras/Public/Core/PersistentBlendStackCameraNode.h
@@ -86,6 +86,11 @@ protected:
 	// FCameraNodeEvaluator interface.
 	virtual void OnRun(const FCameraNodeEvaluationParams& Params, FCameraNodeEvaluationResult& OutResult) override;
 
+#if WITH_EDITOR
+	// FBlendStackCameraNodeEvaluator interface.
+	virtual void OnEntryReinitialized(int32 EntryIndex) override;
+#endif
+
 private:
 
 	void InternalUpdate(TArrayView<FResolvedEntry> ResolvedEntries, const FCameraNodeEvaluationParams& Params, FCameraNodeEvaluationResult& OutResult);
diff --git a/Engine/Plugins/Cameras/GameplayCameras/Source/GameplayCameras/Public/Core/TransientBlendStackCameraNode.h b/Engine/Plugins/Cameras/GameplayCameras/Source/GameplayCameras/Public/Core/TransientBlendStackCameraNode.h
--- a/Engine/Plugins/Cameras/GameplayCameras/Source/GameplayCameras/Public/Core/TransientBlendStackCameraNode.h
+++ b/Engine/Plugins/Cameras/GameplayCameras/Source/GameplayCameras/Public/Core/TransientBlendStackCameraNode.h
@@ -73,6 +73,11 @@ protected:
 	virtual void OnRun(const FCameraNodeEvaluationParams& Params, FCameraNodeEvaluationResult& OutResult) override;
 	virtual void OnSerialize(const FCameraNodeEvaluatorSerializeParams& Params, FArchive& Ar) override;
 
+#if WITH_EDITOR
+	// FBlendStackCameraNodeEvaluator interface.
+	virtual void OnEntryReinitialized(int32 EntryIndex) override;
+#endif
+
 private:
 
 	// Update methods.
