Cameras: add project setting for whether camera rigs can run in the editor

Also fixup how the settings classes are configured, since some of them were showing up twice in the settings dialog.

diff --git a/Engine/Plugins/Cameras/GameplayCameras/Source/GameplayCameras/Private/GameFramework/GameplayCameraComponentBase.cpp b/Engine/Plugins/Cameras/GameplayCameras/Source/GameplayCameras/Private/GameFramework/GameplayCameraComponentBase.cpp
--- a/Engine/Plugins/Cameras/GameplayCameras/Source/GameplayCameras/Private/GameFramework/GameplayCameraComponentBase.cpp
+++ b/Engine/Plugins/Cameras/GameplayCameras/Source/GameplayCameras/Private/GameFramework/GameplayCameraComponentBase.cpp
@@ -15,6 +15,8 @@
 #include "GameFramework/GameplayCameraSystemActor.h"
 #include "GameFramework/GameplayCameraSystemHost.h"
 #include "GameplayCamerasDelegates.h"
+#include "IGameplayCamerasLiveEditManager.h"
+#include "IGameplayCamerasModule.h"
 #include "Kismet/GameplayStatics.h"
 #include "Misc/AssertionMacros.h"
 #include "PrimitiveDrawInterface.h"
@@ -490,7 +492,7 @@ void UGameplayCameraComponentBase::AutoManageEditorPreviewEvaluator()
 		return;
 	}
 	
-	if (bRunInEditor && !(EditorPreviewEvaluator && EvaluationContext))
+	if (CanRunInEditor() && !(EditorPreviewEvaluator && EvaluationContext))
 	{
 		// We want to run the camera logic in the editor but we haven't set things up for that.
 		// Let's create the preview evaluator and the evaluation context.
@@ -513,7 +515,7 @@ void UGameplayCameraComponentBase::AutoManageEditorPreviewEvaluator()
 			}
 		}
 	}
-	else if (!bRunInEditor && (EditorPreviewEvaluator || EvaluationContext))
+	else if (!CanRunInEditor() && (EditorPreviewEvaluator || EvaluationContext))
 	{
 		// We don't want to run the camera logic in the editor anymore. Let's tear things down.
 		EditorPreviewEvaluator = nullptr;
@@ -528,24 +530,30 @@ void UGameplayCameraComponentBase::OnEditorPreviewCameraRigIndexChanged()
 		return;
 	}
 
-	if (bRunInEditor && EditorPreviewEvaluator && EvaluationContext)
+	if (CanRunInEditor() && EditorPreviewEvaluator && EvaluationContext)
 	{
 		EvaluationContext->SetEditorPreviewCameraRigIndex(EditorPreviewCameraRigIndex);
 	}
 }
 
+bool UGameplayCameraComponentBase::CanRunInEditor() const
+{
+	IGameplayCamerasModule& Module = IGameplayCamerasModule::Get();
+	return bRunInEditor && Module.GetLiveEditManager()->CanRunInEditor();
+}
+
 bool UGameplayCameraComponentBase::HasActiveEditorPreview() const
 {
 	return (
 			bIsEditorWorld &&
-			bRunInEditor &&
+			CanRunInEditor() &&
 			EditorPreviewEvaluator &&
 			EvaluationContext);
 }
 
 void UGameplayCameraComponentBase::RecreateEditorPreviewEvaluationContext()
 {
-	if (bIsEditorWorld && bRunInEditor && EditorPreviewEvaluator && EvaluationContext)
+	if (bIsEditorWorld && CanRunInEditor() && EditorPreviewEvaluator && EvaluationContext)
 	{
 		EditorPreviewEvaluator->RemoveEvaluationContext(EvaluationContext.ToSharedRef());
 		EvaluationContext = nullptr;
@@ -586,7 +594,7 @@ void UGameplayCameraComponentBase::OnDrawVisualizationHUD(const FViewport* Viewp
 	using namespace UE::Cameras;
 
 	const bool bHasShowFlag = SceneView->Family->EngineShowFlags.GetSingleFlag(CustomShowFlag);
-	if (bHasShowFlag && bRunInEditor && EditorPreviewEvaluator && EvaluationContext)
+	if (bHasShowFlag && CanRunInEditor() && EditorPreviewEvaluator && EvaluationContext)
 	{
 		const AActor* OwnerActor = GetOwner();
 
diff --git a/Engine/Plugins/Cameras/GameplayCameras/Source/GameplayCameras/Private/GameplayCamerasSettings.cpp b/Engine/Plugins/Cameras/GameplayCameras/Source/GameplayCameras/Private/GameplayCamerasSettings.cpp
--- a/Engine/Plugins/Cameras/GameplayCameras/Source/GameplayCameras/Private/GameplayCamerasSettings.cpp
+++ b/Engine/Plugins/Cameras/GameplayCameras/Source/GameplayCameras/Private/GameplayCamerasSettings.cpp
@@ -4,8 +4,9 @@
 
 #include UE_INLINE_GENERATED_CPP_BY_NAME(GameplayCamerasSettings)
 
-FName UGameplayCamerasSettings::GetCategoryName() const
+UGameplayCamerasSettings::UGameplayCamerasSettings(const FObjectInitializer& ObjectInitializer)
+	: Super(ObjectInitializer)
 {
-	return TEXT("Plugins");
+	CategoryName = TEXT("Plugins");
 }
 
diff --git a/Engine/Plugins/Cameras/GameplayCameras/Source/GameplayCameras/Public/GameFramework/GameplayCameraComponentBase.h b/Engine/Plugins/Cameras/GameplayCameras/Source/GameplayCameras/Public/GameFramework/GameplayCameraComponentBase.h
--- a/Engine/Plugins/Cameras/GameplayCameras/Source/GameplayCameras/Public/GameFramework/GameplayCameraComponentBase.h
+++ b/Engine/Plugins/Cameras/GameplayCameras/Source/GameplayCameras/Public/GameFramework/GameplayCameraComponentBase.h
@@ -131,6 +131,7 @@ protected:
 
 	void UpdateCameraEvaluationContext(bool bForceApplyParameterOverrides);
 #if WITH_EDITOR
+	bool CanRunInEditor() const;
 	bool HasActiveEditorPreview() const;
 	void RecreateEditorPreviewEvaluationContext();
 	void UpdateEditorPreviewEvaluator(float DeltaTime);
diff --git a/Engine/Plugins/Cameras/GameplayCameras/Source/GameplayCameras/Public/GameplayCamerasSettings.h b/Engine/Plugins/Cameras/GameplayCameras/Source/GameplayCameras/Public/GameplayCamerasSettings.h
--- a/Engine/Plugins/Cameras/GameplayCameras/Source/GameplayCameras/Public/GameplayCamerasSettings.h
+++ b/Engine/Plugins/Cameras/GameplayCameras/Source/GameplayCameras/Public/GameplayCamerasSettings.h
@@ -17,6 +17,10 @@ class UGameplayCamerasSettings : public UDeveloperSettings
 {
 	GENERATED_BODY()
 
+public:
+
+	UGameplayCamerasSettings(const FObjectInitializer& ObjectInitializer);	
+
 public:
 
 	/**
@@ -60,10 +64,5 @@ public:
 	/** The distance below which any IK aiming operation is disabled. */
 	UPROPERTY(EditAnywhere, Config, Category="IK Aiming")
 	double DefaultIKAimingMinDistance = 100.0;  // 1m
-
-protected:
-
-	// UDeveloperSettings interface.
-	virtual FName GetCategoryName() const override;
 };
 
diff --git a/Engine/Plugins/Cameras/GameplayCameras/Source/GameplayCameras/Public/IGameplayCamerasLiveEditManager.h b/Engine/Plugins/Cameras/GameplayCameras/Source/GameplayCameras/Public/IGameplayCamerasLiveEditManager.h
--- a/Engine/Plugins/Cameras/GameplayCameras/Source/GameplayCameras/Public/IGameplayCamerasLiveEditManager.h
+++ b/Engine/Plugins/Cameras/GameplayCameras/Source/GameplayCameras/Public/IGameplayCamerasLiveEditManager.h
@@ -26,6 +26,9 @@ public:
 
 	virtual ~IGameplayCamerasLiveEditManager() {}
 
+	/** Whether cameras should be run in editor. */
+	virtual bool CanRunInEditor() const = 0;
+
 	/** Notify all listeners to reload cameras related to the given package. */
 	virtual void NotifyPostBuildAsset(const UPackage* InAssetPackage) const = 0;
 
diff --git a/Engine/Plugins/Cameras/GameplayCameras/Source/GameplayCamerasEditor/Private/GameplayCamerasEditorModule.cpp b/Engine/Plugins/Cameras/GameplayCameras/Source/GameplayCamerasEditor/Private/GameplayCamerasEditorModule.cpp
--- a/Engine/Plugins/Cameras/GameplayCameras/Source/GameplayCamerasEditor/Private/GameplayCamerasEditorModule.cpp
+++ b/Engine/Plugins/Cameras/GameplayCameras/Source/GameplayCamerasEditor/Private/GameplayCamerasEditorModule.cpp
@@ -97,7 +97,6 @@ public:
 
 		FCoreDelegates::OnEnginePreExit.AddRaw(this, &FGameplayCamerasEditorModule::OnPreExit);
 
-		RegisterSettings();
 		RegisterCameraDirectorEditors();
 		RegisterCoreDebugCategories();
 		RegisterRewindDebuggerFeatures();
@@ -126,7 +125,6 @@ public:
 		FGameplayCamerasDebuggerCommands::Unregister();
 		FObjectTreeGraphEditorCommands::Unregister();
 
-		UnregisterSettings();
 		UnregisterCameraDirectorEditors();
 		UnregisterCoreDebugCategories();
 		UnregisterRewindDebuggerFeatures();
@@ -273,30 +271,6 @@ private:
 		SGameplayCamerasDebugger::UnregisterTabSpawners();
 	}
 
-	void RegisterSettings()
-	{
-		ISettingsModule* SettingsModule = FModuleManager::GetModulePtr<ISettingsModule>("Settings");
-
-		if (SettingsModule != nullptr)
-		{
-			SettingsModule->RegisterSettings("Editor", "Plugins", "Gameplay Cameras",
-				LOCTEXT("GameplayCamerasEditorProjectSettingsName", "Gameplay Cameras"),
-				LOCTEXT("GameplayCamerasEditorProjectSettingsDescription", "Configure the gameplay cameras editors."),
-				GetMutableDefault<UGameplayCamerasEditorSettings>()
-			);
-		}
-	}
-
-	void UnregisterSettings()
-	{
-		ISettingsModule* SettingsModule = FModuleManager::GetModulePtr<ISettingsModule>("Settings");
-
-		if (SettingsModule != nullptr)
-		{
-			SettingsModule->UnregisterSettings("Editor", "Plugins", "Gameplay Cameras");
-		}
-	}
-
 	void RegisterCameraDirectorEditors()
 	{
 		using namespace UE::Cameras;
diff --git a/Engine/Plugins/Cameras/GameplayCameras/Source/GameplayCamerasEditor/Private/GameplayCamerasEditorSettings.cpp b/Engine/Plugins/Cameras/GameplayCameras/Source/GameplayCamerasEditor/Private/GameplayCamerasEditorSettings.cpp
--- a/Engine/Plugins/Cameras/GameplayCameras/Source/GameplayCamerasEditor/Private/GameplayCamerasEditorSettings.cpp
+++ b/Engine/Plugins/Cameras/GameplayCameras/Source/GameplayCamerasEditor/Private/GameplayCamerasEditorSettings.cpp
@@ -7,6 +7,8 @@
 UGameplayCamerasEditorSettings::UGameplayCamerasEditorSettings(const FObjectInitializer& ObjectInitializer)
 	: Super(ObjectInitializer)
 {
+	CategoryName = TEXT("Plugins");
+
 	CameraNodeTitleColor = FLinearColor(0.190525f, 0.583898f, 1.0f, 1.0f);           // Light blue
 	CameraAssetTitleColor = FLinearColor(1.0f, 0.170000f, 0.0f, 1.0f);				 // Orange
 	CameraRigAssetTitleColor = FLinearColor(1.0f, 0.170000f, 0.0f, 1.0f);            // Orange
@@ -16,8 +18,3 @@ UGameplayCamerasEditorSettings::UGameplayCamerasEditorSettings(const FObjectInit
 	CameraBlendNodeTitleColor = FLinearColor(0.6f, 0.0f, 1.0f, 1.0f);                // Purple
 }
 
-FName UGameplayCamerasEditorSettings::GetCategoryName() const
-{
-	return TEXT("Plugins");
-}
-
diff --git a/Engine/Plugins/Cameras/GameplayCameras/Source/GameplayCamerasEditor/Private/GameplayCamerasLiveEditManager.cpp b/Engine/Plugins/Cameras/GameplayCameras/Source/GameplayCamerasEditor/Private/GameplayCamerasLiveEditManager.cpp
--- a/Engine/Plugins/Cameras/GameplayCameras/Source/GameplayCamerasEditor/Private/GameplayCamerasLiveEditManager.cpp
+++ b/Engine/Plugins/Cameras/GameplayCameras/Source/GameplayCamerasEditor/Private/GameplayCamerasLiveEditManager.cpp
@@ -3,6 +3,7 @@
 #include "GameplayCamerasLiveEditManager.h"
 
 #include "IGameplayCamerasLiveEditListener.h"
+#include "GameplayCamerasEditorSettings.h"
 #include "Misc/CoreDelegates.h"
 #include "UObject/Package.h"
 #include "UObject/UObjectGlobals.h"
@@ -61,6 +62,12 @@ FGameplayCamerasLiveEditManager::~FGameplayCamerasLiveEditManager()
 	FCoreUObjectDelegates::GetPostGarbageCollect().RemoveAll(this);
 }
 
+bool FGameplayCamerasLiveEditManager::CanRunInEditor() const
+{
+	const UGameplayCamerasEditorSettings* Settings = GetDefault<UGameplayCamerasEditorSettings>();
+	return Settings->bEnableRunInEditor;
+}
+
 void FGameplayCamerasLiveEditManager::NotifyPostBuildAsset(const UPackage* InAssetPackage) const
 {
 	if (const FListenerArray* Listeners = PackageListenerMap.Find(InAssetPackage))
diff --git a/Engine/Plugins/Cameras/GameplayCameras/Source/GameplayCamerasEditor/Private/GameplayCamerasLiveEditManager.h b/Engine/Plugins/Cameras/GameplayCameras/Source/GameplayCamerasEditor/Private/GameplayCamerasLiveEditManager.h
--- a/Engine/Plugins/Cameras/GameplayCameras/Source/GameplayCamerasEditor/Private/GameplayCamerasLiveEditManager.h
+++ b/Engine/Plugins/Cameras/GameplayCameras/Source/GameplayCamerasEditor/Private/GameplayCamerasLiveEditManager.h
@@ -22,6 +22,7 @@ public:
 public:
 
 	// IGameplayCamerasLiveEditManager interface
+	virtual bool CanRunInEditor() const override;
 	virtual void NotifyPostBuildAsset(const UPackage* InAssetPackage) const override;
 	virtual void AddListener(const UPackage* InAssetPackage, IGameplayCamerasLiveEditListener* Listener) override;
 	virtual void RemoveListener(const UPackage* InAssetPackage, IGameplayCamerasLiveEditListener* Listener) override;
diff --git a/Engine/Plugins/Cameras/GameplayCameras/Source/GameplayCamerasEditor/Public/GameplayCamerasEditorSettings.h b/Engine/Plugins/Cameras/GameplayCameras/Source/GameplayCamerasEditor/Public/GameplayCamerasEditorSettings.h
--- a/Engine/Plugins/Cameras/GameplayCameras/Source/GameplayCamerasEditor/Public/GameplayCamerasEditorSettings.h
+++ b/Engine/Plugins/Cameras/GameplayCameras/Source/GameplayCamerasEditor/Public/GameplayCamerasEditorSettings.h
@@ -11,7 +11,7 @@
 
 #include "GameplayCamerasEditorSettings.generated.h"
 
-UCLASS(Config=GameplayCameras, DefaultConfig, MinimalAPI, meta=(DisplayName="Gameplay Cameras Editor"))
+UCLASS(Config=EditorPerProjectUserSettings, MinimalAPI, meta=(DisplayName="Gameplay Cameras Editor"))
 class UGameplayCamerasEditorSettings : public UDeveloperSettings
 {
 	GENERATED_BODY()
@@ -22,35 +22,44 @@ public:
 
 public:
 
+	/** Title color for common camera nodes. */
 	UPROPERTY(EditAnywhere, Config, Category=NodeTitleColors)
 	FLinearColor CameraNodeTitleColor;
 
+	/** Title color for a camera asset's root node. */
 	UPROPERTY(EditAnywhere, Config, Category=NodeTitleColors)
 	FLinearColor CameraAssetTitleColor;
 
+	/** Title color for a camera rig asset's root node. */
 	UPROPERTY(EditAnywhere, Config, Category=NodeTitleColors)
 	FLinearColor CameraRigAssetTitleColor;
 
+	/** Title color for a camera shake asset's root node. */
 	UPROPERTY(EditAnywhere, Config, Category=NodeTitleColors)
 	FLinearColor CameraShakeAssetTitleColor;
 
+	/** Title color for a camera transition node. */
 	UPROPERTY(EditAnywhere, Config, Category=NodeTitleColors)
 	FLinearColor CameraRigTransitionTitleColor;
 
+	/** Title color for a camera transition condition node. */
 	UPROPERTY(EditAnywhere, Config, Category=NodeTitleColors)
 	FLinearColor CameraRigTransitionConditionTitleColor;
 
+	/** Title color for a camera blend node. */
 	UPROPERTY(EditAnywhere, Config, Category=NodeTitleColors)
 	FLinearColor CameraBlendNodeTitleColor;
 
 public:
 
+	/** Camera asset mode to restore on open (director, shared transitions, etc.) */
 	UPROPERTY()
 	FName LastCameraAssetToolkitModeName;
 
-protected:
+public:
 
-	// UDeveloperSettings interface.
-	virtual FName GetCategoryName() const override;
+	/** Global enable/disable flag for running camera rigs in editor. */
+	UPROPERTY(EditAnywhere, Config, Category="Editor Preview")
+	bool bEnableRunInEditor = true;
 };
 
