Cameras: add camera pose location trail to debug drawing

diff --git a/Engine/Plugins/Cameras/GameplayCameras/Source/GameplayCameras/Private/Core/BlendStackCameraNode.cpp b/Engine/Plugins/Cameras/GameplayCameras/Source/GameplayCameras/Private/Core/BlendStackCameraNode.cpp
--- a/Engine/Plugins/Cameras/GameplayCameras/Source/GameplayCameras/Private/Core/BlendStackCameraNode.cpp
+++ b/Engine/Plugins/Cameras/GameplayCameras/Source/GameplayCameras/Private/Core/BlendStackCameraNode.cpp
@@ -15,6 +15,7 @@
 #include "Debug/CameraDebugRenderer.h"
 #include "Debug/CameraNodeEvaluationResultDebugBlock.h"
 #include "Debug/CameraPoseDebugBlock.h"
+#include "Debug/CameraPoseLocationTrailDebugBlock.h"
 #include "Debug/VariableTableDebugBlock.h"
 #include "HAL/IConsoleManager.h"
 #include "Helpers/CameraObjectInterfaceParameterOverrideHelper.h"
@@ -534,7 +535,11 @@ void FBlendStackCameraNodeEvaluator::OnPostBuildAsset(const FGameplayCameraAsset
 
 void FBlendStackCameraNodeEvaluator::OnBuildDebugBlocks(const FCameraDebugBlockBuildParams& Params, FCameraDebugBlockBuilder& Builder)
 {
-	Builder.AttachDebugBlock<FBlendStackSummaryCameraDebugBlock>(*this);
+	FBlendStackSummaryCameraDebugBlock& DebugBlock = Builder.AttachDebugBlock<FBlendStackSummaryCameraDebugBlock>(*this);
+	for (const FCameraRigEntry& Entry : Entries)
+	{
+		DebugBlock.AddChild(&Builder.BuildDebugBlock<FCameraPoseLocationTrailDebugBlock>(Entry.Result));
+	}
 }
 
 FBlendStackCameraDebugBlock* FBlendStackCameraNodeEvaluator::BuildDetailedDebugBlock(const FCameraDebugBlockBuildParams& Params, FCameraDebugBlockBuilder& Builder)
diff --git a/Engine/Plugins/Cameras/GameplayCameras/Source/GameplayCameras/Private/Core/PersistentBlendStackCameraNode.cpp b/Engine/Plugins/Cameras/GameplayCameras/Source/GameplayCameras/Private/Core/PersistentBlendStackCameraNode.cpp
--- a/Engine/Plugins/Cameras/GameplayCameras/Source/GameplayCameras/Private/Core/PersistentBlendStackCameraNode.cpp
+++ b/Engine/Plugins/Cameras/GameplayCameras/Source/GameplayCameras/Private/Core/PersistentBlendStackCameraNode.cpp
@@ -224,6 +224,7 @@ void FPersistentBlendStackCameraNodeEvaluator::InternalUpdate(TArrayView<FResolv
 
 			// Start with the input given to us.
 			{
+				CurResult.Reset();
 				CurResult.CameraPose = OutResult.CameraPose;
 				CurResult.VariableTable.OverrideAll(OutResult.VariableTable);
 				CurResult.CameraRigJoints.OverrideAll(OutResult.CameraRigJoints);
diff --git a/Engine/Plugins/Cameras/GameplayCameras/Source/GameplayCameras/Private/Debug/CameraPoseLocationTrailDebugBlock.cpp b/Engine/Plugins/Cameras/GameplayCameras/Source/GameplayCameras/Private/Debug/CameraPoseLocationTrailDebugBlock.cpp
new file mode 100644
--- /dev/null
+++ b/Engine/Plugins/Cameras/GameplayCameras/Source/GameplayCameras/Private/Debug/CameraPoseLocationTrailDebugBlock.cpp
@@ -0,0 +1,55 @@
+// Copyright Epic Games, Inc. All Rights Reserved.
+
+#include "Debug/CameraPoseLocationTrailDebugBlock.h"
+
+#include "Core/CameraNodeEvaluator.h"
+#include "Debug/CameraDebugRenderer.h"
+
+#if UE_GAMEPLAY_CAMERAS_DEBUG
+
+namespace UE::Cameras
+{
+
+UE_DEFINE_CAMERA_DEBUG_BLOCK(FCameraPoseLocationTrailDebugBlock)
+
+FCameraPoseLocationTrailDebugBlock::FCameraPoseLocationTrailDebugBlock()
+{
+}
+
+FCameraPoseLocationTrailDebugBlock::FCameraPoseLocationTrailDebugBlock(const FCameraNodeEvaluationResult& InResult)
+{
+	Trail = InResult.GetCameraPoseLocationTrail();
+}
+
+void FCameraPoseLocationTrailDebugBlock::OnDebugDraw(const FCameraDebugBlockDrawParams& Params, FCameraDebugRenderer& Renderer)
+{
+	if (Renderer.IsExternalRendering())
+	{
+		DrawCameraPoseLocationTrail(Renderer);
+	}
+}
+
+void FCameraPoseLocationTrailDebugBlock::DrawCameraPoseLocationTrail(FCameraDebugRenderer& Renderer) const
+{
+	const FLinearColor TrailColor(FColorList::LightBlue);
+	for (int32 Index = 1; Index < Trail.Num(); ++Index)
+	{
+		const FVector3d& PrevPoint(Trail[Index - 1]);
+		const FVector3d& NextPoint(Trail[Index]);
+
+		if (!PrevPoint.IsZero() && FVector3d::Distance(PrevPoint, NextPoint) > UE_SMALL_NUMBER)
+		{
+			Renderer.DrawLine(PrevPoint, NextPoint, TrailColor, 1.f);
+		}
+	}
+}
+
+void FCameraPoseLocationTrailDebugBlock::OnSerialize(FArchive& Ar)
+{
+	Ar << Trail;
+}
+
+}  // namespace UE::Cameras
+
+#endif  // UE_GAMEPLAY_CAMERAS_DEBUG
+
diff --git a/Engine/Plugins/Cameras/GameplayCameras/Source/GameplayCameras/Public/Debug/CameraPoseLocationTrailDebugBlock.h b/Engine/Plugins/Cameras/GameplayCameras/Source/GameplayCameras/Public/Debug/CameraPoseLocationTrailDebugBlock.h
new file mode 100644
--- /dev/null
+++ b/Engine/Plugins/Cameras/GameplayCameras/Source/GameplayCameras/Public/Debug/CameraPoseLocationTrailDebugBlock.h
@@ -0,0 +1,44 @@
+// Copyright Epic Games, Inc. All Rights Reserved.
+
+#pragma once
+
+#include "Debug/CameraDebugBlock.h"
+
+#if UE_GAMEPLAY_CAMERAS_DEBUG
+
+namespace UE::Cameras
+{
+
+struct FCameraNodeEvaluationResult;
+
+/**
+ * A debug block that draws a camera pose location trail.
+ */
+class FCameraPoseLocationTrailDebugBlock : public FCameraDebugBlock
+{
+	UE_DECLARE_CAMERA_DEBUG_BLOCK(GAMEPLAYCAMERAS_API, FCameraPoseLocationTrailDebugBlock)
+
+public:
+
+	FCameraPoseLocationTrailDebugBlock();
+	FCameraPoseLocationTrailDebugBlock(const FCameraNodeEvaluationResult& InResult);
+
+protected:
+
+	// FCameraDebugBlock interface.
+	virtual void OnDebugDraw(const FCameraDebugBlockDrawParams& Params, FCameraDebugRenderer& Renderer) override;
+	virtual void OnSerialize(FArchive& Ar) override;
+
+private:
+
+	void DrawCameraPoseLocationTrail(FCameraDebugRenderer& Renderer) const;
+
+private:
+
+	TArray<FVector3d> Trail;
+};
+
+}  // namespace UE::Cameras
+
+#endif  // UE_GAMEPLAY_CAMERAS_DEBUG
+
