Cameras: refactor the camera director evaluation API and improve support for nesting

- Camera directors can now return a list of requests for activating/deactivating camera rigs on any layer (before they could only activate camera rigs on the main layer)
- The list of request can either reference actual camera rigs, or indirectly reference them via camera rig proxies (before, each type of camera director had to support proxies themselves... this change makes proxies into a first-class supported API)
- This change makes the Blueprint director implementation simpler since a lot of it is now handled in the core.
- The Blueprint director can now add/remove/execute nested camera directors.
- Gameplay camera actors/components now have options for inserting or wrapping an existing director instead of just taking over.
- Improve debug display for nested camera directors, and fix a few problems with the nested directors API.

diff --git a/Engine/Plugins/Cameras/GameplayCameras/Source/GameplayCameras/Private/Core/CameraDirector.cpp b/Engine/Plugins/Cameras/GameplayCameras/Source/GameplayCameras/Private/Core/CameraDirector.cpp
--- a/Engine/Plugins/Cameras/GameplayCameras/Source/GameplayCameras/Private/Core/CameraDirector.cpp
+++ b/Engine/Plugins/Cameras/GameplayCameras/Source/GameplayCameras/Private/Core/CameraDirector.cpp
@@ -2,6 +2,8 @@
 
 #include "Core/CameraDirector.h"
 
+#include "Core/CameraRigProxyTable.h"
+
 #include UE_INLINE_GENERATED_CPP_BY_NAME(CameraDirector)
 
 FCameraDirectorEvaluatorPtr UCameraDirector::BuildEvaluator(FCameraDirectorEvaluatorBuilder& Builder) const
@@ -26,6 +28,11 @@ void UCameraDirector::GatherRigUsageInfo(FCameraDirectorRigUsageInfo& UsageInfo)
 
 void UCameraDirector::FactoryCreateAsset(const FCameraDirectorFactoryCreateParams& InParams)
 {
+	if (bHasCameraRigProxyTable && !CameraRigProxyTable)
+	{
+		CameraRigProxyTable = NewObject<UCameraRigProxyTable>(this);
+	}
+
 	OnFactoryCreateAsset(InParams);
 }
 
diff --git a/Engine/Plugins/Cameras/GameplayCameras/Source/GameplayCameras/Private/Core/CameraDirectorEvaluator.cpp b/Engine/Plugins/Cameras/GameplayCameras/Source/GameplayCameras/Private/Core/CameraDirectorEvaluator.cpp
--- a/Engine/Plugins/Cameras/GameplayCameras/Source/GameplayCameras/Private/Core/CameraDirectorEvaluator.cpp
+++ b/Engine/Plugins/Cameras/GameplayCameras/Source/GameplayCameras/Private/Core/CameraDirectorEvaluator.cpp
@@ -4,6 +4,7 @@
 
 #include "Core/CameraDirector.h"
 #include "Core/CameraEvaluationContext.h"
+#include "Core/CameraRigProxyTable.h"
 #include "Core/CameraRigTransition.h"
 
 namespace UE::Cameras
@@ -34,7 +35,7 @@ void FCameraDirectorEvaluator::Initialize(const FCameraDirectorInitializeParams&
 
 void FCameraDirectorEvaluator::Activate(const FCameraDirectorActivateParams& Params)
 {
-	WeakOwnerContext = Params.OwnerContext;
+	Evaluator = Params.Evaluator;
 
 	OnActivate(Params);
 }
@@ -43,24 +44,46 @@ void FCameraDirectorEvaluator::Deactivate(const FCameraDirectorDeactivateParams&
 {
 	OnDeactivate(Params);
 
-	WeakOwnerContext.Reset();
+	Evaluator = nullptr;
 }
 
 void FCameraDirectorEvaluator::Run(const FCameraDirectorEvaluationParams& Params, FCameraDirectorEvaluationResult& OutResult)
 {
 	OnRun(Params, OutResult);
 
-	if (NextActivationTransitionOverride)
+	// Set some overrides on the first main-layer activation we find.
+	if (NextActivationTransitionOverride || bNextActivationForce)
 	{
-		OutResult.TransitionOverride = NextActivationTransitionOverride;
+		for (FCameraRigActivationDeactivationRequest& Request : OutResult.Requests)
+		{
+			if (Request.RequestType == ECameraRigActivationDeactivationRequestType::Activate && Request.Layer == ECameraRigLayer::Main)
+			{
+				Request.TransitionOverride = NextActivationTransitionOverride;
+				Request.bForceActivateDeactivate |= bNextActivationForce;
+			}
+		}
 	}
+
 	NextActivationTransitionOverride = nullptr;
+	bNextActivationForce = false;
+}
+
+const UCameraRigAsset* FCameraDirectorEvaluator::FindCameraRigByProxy(const UCameraRigProxyAsset* InProxy) const
+{
+	if (!ensure(PrivateCameraDirector))
+	{
+		return nullptr;
+	}
 
-	if (bNextActivationForce)
+	UCameraRigProxyTable* ProxyTable = PrivateCameraDirector->CameraRigProxyTable;
+	if (!ensureMsgf(ProxyTable, TEXT("Cannot resolve camera rig proxy: no proxy table set on camera director.")))
 	{
-		OutResult.bForceActivateCameraRigs = true;
+		return nullptr;
 	}
-	bNextActivationForce = false;
+
+	FCameraRigProxyTableResolveParams ResolveParams;
+	ResolveParams.CameraRigProxy = InProxy;
+	return ProxyTable->ResolveProxy(ResolveParams);
 }
 
 void FCameraDirectorEvaluator::OnEndCameraSystemUpdate()
@@ -109,7 +132,7 @@ bool FCameraDirectorEvaluator::AddChildEvaluationContext(TSharedRef<FCameraEvalu
 			break;
 		case EChildContextManipulationResult::ChildContextSuccess:
 			// A sub-director of our director accepted the child context, so it already
-			// activate it and we don't need to do it ourselves.
+			// activated it and we don't need to do it ourselves.
 			bReturn = true;
 			break;
 	}
@@ -118,6 +141,8 @@ bool FCameraDirectorEvaluator::AddChildEvaluationContext(TSharedRef<FCameraEvalu
 		OwnerContext->RegisterChildContext(InContext);
 
 		FCameraEvaluationContextActivateParams ActivateParams;
+		ActivateParams.Evaluator = Evaluator;
+		ActivateParams.ParentContext = OwnerContext;
 		InContext->Activate(ActivateParams);
 	}
 	return bReturn;
diff --git a/Engine/Plugins/Cameras/GameplayCameras/Source/GameplayCameras/Private/Core/CameraEvaluationContext.cpp b/Engine/Plugins/Cameras/GameplayCameras/Source/GameplayCameras/Private/Core/CameraEvaluationContext.cpp
--- a/Engine/Plugins/Cameras/GameplayCameras/Source/GameplayCameras/Private/Core/CameraEvaluationContext.cpp
+++ b/Engine/Plugins/Cameras/GameplayCameras/Source/GameplayCameras/Private/Core/CameraEvaluationContext.cpp
@@ -41,7 +41,7 @@ protected:
 
 		if (PreviewIndex != INDEX_NONE)
 		{
-			OutResult.Add(Params.OwnerContext, CameraRigs[PreviewIndex]);
+			OutResult.Add(GetEvaluationContext(), CameraRigs[PreviewIndex]);
 		}
 	}
 
@@ -280,7 +280,6 @@ void FCameraEvaluationContext::Activate(const FCameraEvaluationContextActivatePa
 	{
 		FCameraDirectorActivateParams DirectorParams;
 		DirectorParams.Evaluator = Params.Evaluator;
-		DirectorParams.OwnerContext = SharedThis(this);
 		DirectorEvaluator->Activate(DirectorParams);
 	}
 
@@ -298,7 +297,6 @@ void FCameraEvaluationContext::Deactivate(const FCameraEvaluationContextDeactiva
 	if (DirectorEvaluator)
 	{
 		FCameraDirectorDeactivateParams DirectorParams;
-		DirectorParams.OwnerContext = SharedThis(this);
 		DirectorEvaluator->Deactivate(DirectorParams);
 	}
 
@@ -311,6 +309,28 @@ void FCameraEvaluationContext::Deactivate(const FCameraEvaluationContextDeactiva
 	bActivated = false;
 }
 
+bool FCameraEvaluationContext::AddChildContext(TSharedRef<FCameraEvaluationContext> ChildContext)
+{
+	if (DirectorEvaluator)
+	{
+		ensure(DirectorEvaluator->GetEvaluationContext() == SharedThis(this));
+		return DirectorEvaluator->AddChildEvaluationContext(ChildContext);
+	}
+
+	return false;
+}
+
+bool FCameraEvaluationContext::RemoveChildContext(TSharedRef<FCameraEvaluationContext> ChildContext)
+{
+	if (DirectorEvaluator)
+	{
+		ensure(DirectorEvaluator->GetEvaluationContext() == SharedThis(this));
+		return DirectorEvaluator->RemoveChildEvaluationContext(ChildContext);
+	}
+
+	return false;
+}
+
 bool FCameraEvaluationContext::RegisterChildContext(TSharedRef<FCameraEvaluationContext> ChildContext)
 {
 	if (!ensureMsgf(ChildContext->WeakParent == nullptr, TEXT("The given evaluation context already has a parent!")))
diff --git a/Engine/Plugins/Cameras/GameplayCameras/Source/GameplayCameras/Private/Core/CameraEvaluationContextStack.cpp b/Engine/Plugins/Cameras/GameplayCameras/Source/GameplayCameras/Private/Core/CameraEvaluationContextStack.cpp
--- a/Engine/Plugins/Cameras/GameplayCameras/Source/GameplayCameras/Private/Core/CameraEvaluationContextStack.cpp
+++ b/Engine/Plugins/Cameras/GameplayCameras/Source/GameplayCameras/Private/Core/CameraEvaluationContextStack.cpp
@@ -68,17 +68,39 @@ void FCameraEvaluationContextStack::PushContext(TSharedRef<FCameraEvaluationCont
 	Entries.Push(MoveTemp(NewEntry));
 }
 
-bool FCameraEvaluationContextStack::AddChildContext(TSharedRef<FCameraEvaluationContext> Context)
+bool FCameraEvaluationContextStack::AddChildContext(TSharedRef<FCameraEvaluationContext> Context, TSharedPtr<FCameraEvaluationContext> ParentContext)
 {
-	TSharedPtr<FCameraEvaluationContext> ActiveContext = GetActiveContext();
-	if (ensureMsgf(ActiveContext.IsValid(), TEXT("Can't add child context to the stack, no active context was found!")))
+	if (ParentContext == nullptr && !Entries.IsEmpty())
 	{
-		FCameraDirectorEvaluator* DirectorEvaluator = ActiveContext->GetDirectorEvaluator();
-		if (ensureMsgf(DirectorEvaluator, TEXT("Can't add child context, active context has no camera director evaluator!")))
-		{
-			return DirectorEvaluator->AddChildEvaluationContext(Context);
-		}
+		ParentContext = GetActiveContext();
+	}
+
+	// No parent context provided, and no active context found in the stack.
+	if (!ParentContext.IsValid())
+	{
+		return false;
+	}
+
+	// The context is already in the stack. The caller should remove it first.
+	const int32 ExistingIndex = Entries.IndexOfByPredicate(
+			[Context](const FContextEntry& Entry) { return Entry.WeakContext == Context; });
+	if (ExistingIndex != INDEX_NONE)
+	{
+		return false;
+	}
+
+	// Check invalid situation.
+	if (!ensureMsgf(Context != ParentContext, TEXT("Can't add a context as a child of itself")))
+	{
+		return false;
 	}
+
+	FCameraDirectorEvaluator* DirectorEvaluator = ParentContext->GetDirectorEvaluator();
+	if (ensureMsgf(DirectorEvaluator, TEXT("Can't add child context, active context has no camera director evaluator!")))
+	{
+		return DirectorEvaluator->AddChildEvaluationContext(Context);
+	}
+
 	return false;
 }
 
diff --git a/Engine/Plugins/Cameras/GameplayCameras/Source/GameplayCameras/Private/Core/CameraSystemEvaluator.cpp b/Engine/Plugins/Cameras/GameplayCameras/Source/GameplayCameras/Private/Core/CameraSystemEvaluator.cpp
--- a/Engine/Plugins/Cameras/GameplayCameras/Source/GameplayCameras/Private/Core/CameraSystemEvaluator.cpp
+++ b/Engine/Plugins/Cameras/GameplayCameras/Source/GameplayCameras/Private/Core/CameraSystemEvaluator.cpp
@@ -232,63 +232,11 @@ void FCameraSystemEvaluator::UpdateImpl(float DeltaTime, ECameraNodeEvaluationTy
 	FCameraDirectorEvaluator* ActiveDirectorEvaluator = ActiveContext->GetDirectorEvaluator();
 	if (ActiveDirectorEvaluator)
 	{
-		FCameraDirectorEvaluationParams DirectorParams;
-		DirectorParams.DeltaTime = DeltaTime;
-		DirectorParams.OwnerContext = ActiveContext;
-
-		FCameraDirectorEvaluationResult DirectorResult;
-
-		ActiveDirectorEvaluator->Run(DirectorParams, DirectorResult);
-
-		if (DirectorResult.ActiveCameraRigs.Num() == 1)
-		{
-			// Only one camera rig to activate... let's do that.
-			const FActiveCameraRigInfo& ActiveCameraRig = DirectorResult.ActiveCameraRigs[0];
-
-			FActivateCameraRigParams CameraRigParams;
-			CameraRigParams.EvaluationContext = ActiveCameraRig.EvaluationContext;
-			CameraRigParams.CameraRig = ActiveCameraRig.CameraRig;
-			CameraRigParams.TransitionOverride = DirectorResult.TransitionOverride;
-			CameraRigParams.bForceActivate = DirectorResult.bForceActivateCameraRigs;
-			RootEvaluator->ActivateCameraRig(CameraRigParams);
-		}
-		else if (DirectorResult.ActiveCameraRigs.Num() > 1)
-		{
-			// We have a combination of camera rigs to activate. Let's dynamically generate a new camera rig
-			// asset that combines them.
-#if WITH_EDITOR
-			const UGameplayCamerasSettings* Settings = GetDefault<UGameplayCamerasSettings>();
-			if (DirectorResult.ActiveCameraRigs.Num() > Settings->CombinedCameraRigNumThreshold)
-			{
-				UE_LOG(LogCameraSystem, Warning, 
-						TEXT("Activating %d camera rigs combined! Is the camera director doing this on purpose? "
-							"If so, raise the CombinedCameraRigNumThreshold setting to remove this warning."),
-						DirectorResult.ActiveCameraRigs.Num());
-			}
-#endif
-
-			// All combined camera rigs must belong to the same evaluation context.
-			TArray<const UCameraRigAsset*> Combination;
-			TSharedPtr<const FCameraEvaluationContext> CommonContext = DirectorResult.ActiveCameraRigs[0].EvaluationContext;
-			for (const FActiveCameraRigInfo& ActiveCameraRig : DirectorResult.ActiveCameraRigs)
-			{
-				Combination.Add(ActiveCameraRig.CameraRig);
-				ensureMsgf(ActiveCameraRig.EvaluationContext == CommonContext,
-						TEXT("All combined camera rigs must be activated from the same evaluation context."));
-			}
-			const UCameraRigAsset* CombinedCameraRig = CameraRigCombinationRegistry->FindOrCreateCombination(Combination);
-
-			FActivateCameraRigParams CameraRigParams;
-			CameraRigParams.EvaluationContext = CommonContext;
-			CameraRigParams.CameraRig = CombinedCameraRig;
-			CameraRigParams.TransitionOverride = DirectorResult.TransitionOverride;
-			CameraRigParams.bForceActivate = DirectorResult.bForceActivateCameraRigs;
-			RootEvaluator->ActivateCameraRig(CameraRigParams);
-		}
+		UpdateCameraDirector(DeltaTime, ActiveDirectorEvaluator);
 	}
 
+	// Run the camera node tree.
 	{
-		// Setup the params/result for running the root camera node.
 		FCameraNodeEvaluationParams NodeParams;
 		NodeParams.Evaluator = this;
 		NodeParams.DeltaTime = DeltaTime;
@@ -296,7 +244,6 @@ void FCameraSystemEvaluator::UpdateImpl(float DeltaTime, ECameraNodeEvaluationTy
 
 		RootNodeResult.Reset();
 
-		// Run the root camera node.
 		RootEvaluator->Run(NodeParams, RootNodeResult);
 
 		RootNodeResult.bIsValid = true;
@@ -317,6 +264,130 @@ void FCameraSystemEvaluator::UpdateImpl(float DeltaTime, ECameraNodeEvaluationTy
 	ContextStack.OnEndCameraSystemUpdate();
 }
 
+void FCameraSystemEvaluator::UpdateCameraDirector(float DeltaTime, FCameraDirectorEvaluator* CameraDirectorEvaluator)
+{
+	check(CameraDirectorEvaluator);
+
+	FCameraDirectorEvaluationResult DirectorResult;
+	{
+		FCameraDirectorEvaluationParams DirectorParams;
+		DirectorParams.DeltaTime = DeltaTime;
+
+		CameraDirectorEvaluator->Run(DirectorParams, DirectorResult);
+	}
+
+	TArray<FCameraRigActivationDeactivationRequest, TInlineAllocator<2>> MainLayerActivations;
+	TArray<FCameraRigActivationDeactivationRequest, TInlineAllocator<2>> MainLayerDeactivations;
+
+	for (FCameraRigActivationDeactivationRequest& Request : DirectorResult.Requests)
+	{
+		if (!ensure(Request.EvaluationContext))
+		{
+			continue;
+		}
+
+		// Resolve camera rig proxies if needed.
+		if (Request.CameraRig == nullptr && Request.CameraRigProxy)
+		{
+			Request.CameraRig = CameraDirectorEvaluator->FindCameraRigByProxy(Request.CameraRigProxy);
+		}
+		if (!ensure(Request.CameraRig))
+		{
+			continue;
+		}
+
+		// Put the main layer requests aside while we handle the other requests.
+		if (Request.Layer == ECameraRigLayer::Main)
+		{
+			if (Request.RequestType == ECameraRigActivationDeactivationRequestType::Activate)
+			{
+				MainLayerActivations.Add(Request);
+			}
+			else if (Request.RequestType == ECameraRigActivationDeactivationRequestType::Deactivate)
+			{
+				MainLayerDeactivations.Add(Request);
+			}
+		}
+		else
+		{
+			RootEvaluator->ExecuteCameraDirectorRequest(Request);
+		}
+	}
+
+	if (MainLayerActivations.Num() == 1)
+	{
+		RootEvaluator->ExecuteCameraDirectorRequest(MainLayerActivations[0]);
+	}
+	else if (MainLayerActivations.Num() > 1)
+	{
+		FCameraRigActivationDeactivationRequest CombinedRequest;
+		CombinedRequest.RequestType = ECameraRigActivationDeactivationRequestType::Activate;
+		GetCombinedCameraRigRequest(MainLayerActivations, CombinedRequest);
+
+		RootEvaluator->ExecuteCameraDirectorRequest(CombinedRequest);
+	}
+
+	if (MainLayerDeactivations.Num() == 1)
+	{
+		RootEvaluator->ExecuteCameraDirectorRequest(MainLayerDeactivations[0]);
+	}
+	else if (MainLayerDeactivations.Num() > 1)
+	{
+		FCameraRigActivationDeactivationRequest CombinedRequest;
+		CombinedRequest.RequestType = ECameraRigActivationDeactivationRequestType::Deactivate;
+		GetCombinedCameraRigRequest(MainLayerDeactivations, CombinedRequest);
+
+		RootEvaluator->ExecuteCameraDirectorRequest(CombinedRequest);
+	}
+}
+
+void FCameraSystemEvaluator::GetCombinedCameraRigRequest(TConstArrayView<FCameraRigActivationDeactivationRequest> Requests, FCameraRigActivationDeactivationRequest& OutCombinedRequest)
+{
+	// We have a combination of camera rigs to activate. Let's dynamically generate a new camera rig
+	// asset that combines them.
+#if WITH_EDITOR
+	const UGameplayCamerasSettings* Settings = GetDefault<UGameplayCamerasSettings>();
+	if (Requests.Num() > Settings->CombinedCameraRigNumThreshold)
+	{
+		UE_LOG(LogCameraSystem, Warning, 
+				TEXT("Activating %d camera rigs combined! Is the camera director doing this on purpose? "
+					"If so, raise the CombinedCameraRigNumThreshold setting to remove this warning."),
+				Requests.Num());
+	}
+#endif
+
+	// All combined camera rigs must belong to the same evaluation context, and we can't have more
+	// than one transition override.
+	TArray<const UCameraRigAsset*> Combination;
+	TSharedPtr<const FCameraEvaluationContext> CommonContext = Requests[0].EvaluationContext;
+	const UCameraRigTransition* FirstTransitionOverride = nullptr;
+	bool bAnyForceActivationDeactivation = false;
+	for (const FCameraRigActivationDeactivationRequest& Request : Requests)
+	{
+		Combination.Add(Request.CameraRig);
+		ensureMsgf(Request.EvaluationContext == CommonContext,
+				TEXT("All combined camera rigs must be activated from the same evaluation context."));
+
+		if (Request.TransitionOverride)
+		{
+			if (ensureMsgf(FirstTransitionOverride == nullptr || FirstTransitionOverride == Request.TransitionOverride,
+					TEXT("Only one transition override can be specified when activating/deactivating multiple main-layer rigs.")))
+			{
+				FirstTransitionOverride = Request.TransitionOverride;
+			}
+		}
+
+		bAnyForceActivationDeactivation |= Request.bForceActivateDeactivate;
+	}
+
+	const UCameraRigAsset* CombinedCameraRig = CameraRigCombinationRegistry->FindOrCreateCombination(Combination);
+
+	OutCombinedRequest.EvaluationContext = CommonContext;
+	OutCombinedRequest.CameraRig = CombinedCameraRig;
+	OutCombinedRequest.TransitionOverride = FirstTransitionOverride;
+	OutCombinedRequest.bForceActivateDeactivate = bAnyForceActivationDeactivation;
+}
+
 void FCameraSystemEvaluator::PreUpdateServices(float DeltaTime, ECameraEvaluationServiceFlags ExtraFlags)
 {
 	FCameraEvaluationServiceUpdateParams ServiceUpdateParams;
diff --git a/Engine/Plugins/Cameras/GameplayCameras/Source/GameplayCameras/Private/Core/RootCameraNode.cpp b/Engine/Plugins/Cameras/GameplayCameras/Source/GameplayCameras/Private/Core/RootCameraNode.cpp
--- a/Engine/Plugins/Cameras/GameplayCameras/Source/GameplayCameras/Private/Core/RootCameraNode.cpp
+++ b/Engine/Plugins/Cameras/GameplayCameras/Source/GameplayCameras/Private/Core/RootCameraNode.cpp
@@ -2,6 +2,7 @@
 
 #include "Core/RootCameraNode.h"
 
+#include "Core/CameraDirectorEvaluator.h"
 #include "Core/CameraEvaluationService.h"
 #include "Core/CameraSystemEvaluator.h"
 #include "Core/RootCameraNodeCameraRigEvent.h"
@@ -26,6 +27,31 @@ void FRootCameraNodeEvaluator::DeactivateCameraRig(const FDeactivateCameraRigPar
 	OnDeactivateCameraRig(Params);
 }
 
+void FRootCameraNodeEvaluator::ExecuteCameraDirectorRequest(const FCameraRigActivationDeactivationRequest& Request)
+{
+	if (Request.RequestType == ECameraRigActivationDeactivationRequestType::Activate)
+	{
+		FActivateCameraRigParams CameraRigParams;
+		CameraRigParams.EvaluationContext = Request.EvaluationContext;
+		CameraRigParams.CameraRig = Request.CameraRig;
+		CameraRigParams.Layer = Request.Layer;
+		CameraRigParams.OrderKey = Request.OrderKey;
+		CameraRigParams.TransitionOverride = Request.TransitionOverride;
+		CameraRigParams.bForceActivate = Request.bForceActivateDeactivate;
+		ActivateCameraRig(CameraRigParams);
+	}
+	else if (Request.RequestType == ECameraRigActivationDeactivationRequestType::Deactivate)
+	{
+		FDeactivateCameraRigParams CameraRigParams;
+		CameraRigParams.EvaluationContext = Request.EvaluationContext;
+		CameraRigParams.CameraRig = Request.CameraRig;
+		CameraRigParams.Layer = Request.Layer;
+		CameraRigParams.TransitionOverride = Request.TransitionOverride;
+		CameraRigParams.bDeactiveImmediately = Request.bForceActivateDeactivate;
+		DeactivateCameraRig(CameraRigParams);
+	}
+}
+
 void FRootCameraNodeEvaluator::GetActiveCameraRigInfo(FCameraRigEvaluationInfo& OutCameraRigInfo) const
 {
 	OnGetActiveCameraRigInfo(OutCameraRigInfo);
diff --git a/Engine/Plugins/Cameras/GameplayCameras/Source/GameplayCameras/Private/Debug/CameraDirectorTreeDebugBlock.cpp b/Engine/Plugins/Cameras/GameplayCameras/Source/GameplayCameras/Private/Debug/CameraDirectorTreeDebugBlock.cpp
--- a/Engine/Plugins/Cameras/GameplayCameras/Source/GameplayCameras/Private/Debug/CameraDirectorTreeDebugBlock.cpp
+++ b/Engine/Plugins/Cameras/GameplayCameras/Source/GameplayCameras/Private/Debug/CameraDirectorTreeDebugBlock.cpp
@@ -65,11 +65,14 @@ void FCameraDirectorTreeDebugBlock::InitializeEntry(TSharedPtr<FCameraEvaluation
 		const FCameraDirectorEvaluator* DirectorEvaluator = Context->GetDirectorEvaluator();
 		const UCameraDirector* CameraDirector = DirectorEvaluator->GetCameraDirector();
 
+		TArrayView<const TSharedPtr<FCameraEvaluationContext>> ChildrenContexts = Context->GetChildrenContexts();
+
 		EntryDebugInfo.ContextClassName = ContextTypeName;
 		EntryDebugInfo.OwnerName = *GetPathNameSafe(ContextOwner);
 		EntryDebugInfo.OwnerClassName = ContextOwner ? ContextOwner->GetClass()->GetFName() : NAME_None;
 		EntryDebugInfo.CameraAssetName = GetNameSafe(Context->GetCameraAsset());
 		EntryDebugInfo.CameraDirectorClassName = GetFNameSafe(CameraDirector);
+		EntryDebugInfo.NumChildrenContexts = ChildrenContexts.Num();
 		EntryDebugInfo.InitialContextTransform = Context->GetInitialResult().CameraPose.GetTransform();
 		EntryDebugInfo.bIsValid = true;
 
@@ -77,12 +80,12 @@ void FCameraDirectorTreeDebugBlock::InitializeEntry(TSharedPtr<FCameraEvaluation
 		AddChild(&Builder.BuildDebugBlock<FCameraPoseDebugBlock>(InitialResult.CameraPose)
 				.WithShowUnchangedCVar(TEXT("GameplayCameras.Debug.ContextInitialResult.ShowUnchanged")));
 
-		TArrayView<const TSharedPtr<FCameraEvaluationContext>> ChildrenContexts = Context->GetChildrenContexts();
 		if (ChildrenContexts.Num() > 0)
 		{
 			FCameraDirectorTreeDebugBlock& ChildBlock = Builder.StartChildDebugBlock<FCameraDirectorTreeDebugBlock>();
 			{
 				ChildBlock.Initialize(ChildrenContexts, Builder);
+				ChildBlock.bIsTreeRoot = false;
 			}
 			Builder.EndChildDebugBlock();
 		}
@@ -100,18 +103,32 @@ void FCameraDirectorTreeDebugBlock::OnDebugDraw(const FCameraDebugBlockDrawParam
 {
 	const FCameraDebugColors& Colors = FCameraDebugColors::Get();
 
+	// Our children view looks like this:
+	//
+	//		[0] Initial result pose block
+	//		[1] Child director block
+	//		[2] Child director block
+	//		[3] Initial result pose block
+	//		[4] Initial result pose block
+	//
 	TArrayView<FCameraDebugBlock*> ChildrenView(GetChildren());
 
-	const int32 MinNum = FMath::Min(ChildrenView.Num(), CameraDirectors.Num());
+	// For the root debug info, separate the active director from the inactive ones.
+	// The inactive ones are at the beginning (bottom) of the stack.
+	if (bIsTreeRoot)
+	{
+		Renderer.SetTextColor(Colors.Notice);
+		Renderer.AddText("Inactive Directors\n");
+		Renderer.SetTextColor(Colors.Default);
+		Renderer.AddIndent();
+	}
 
-	Renderer.SetTextColor(Colors.Notice);
-	Renderer.AddText("Inactive Directors\n");
-	Renderer.SetTextColor(Colors.Default);
-	Renderer.AddIndent();
+	int32 ChildContextIndex = 0;
 
-	for (int32 Index = 0; Index < MinNum; ++Index)
+	for (int32 Index = 0; Index < ChildrenView.Num(); ++Index)
 	{
-		if (Index == CameraDirectors.Num() - 1)
+		// If we reached the top of the stack, display the active director separately.
+		if (bIsTreeRoot && ChildContextIndex == CameraDirectors.Num() - 1)
 		{
 			Renderer.RemoveIndent();
 
@@ -121,9 +138,9 @@ void FCameraDirectorTreeDebugBlock::OnDebugDraw(const FCameraDebugBlockDrawParam
 			Renderer.AddIndent();
 		}
 
-		Renderer.AddText(TEXT("{cam_passive}[%d]{cam_default} "), Index + 1);
+		Renderer.AddText(TEXT("{cam_passive}[%d]{cam_default} "), ChildContextIndex + 1);
 
-		const FDirectorDebugInfo& EntryDebugInfo(CameraDirectors[Index]);
+		const FDirectorDebugInfo& EntryDebugInfo(CameraDirectors[ChildContextIndex]);
 		if (EntryDebugInfo.bIsValid)
 		{
 			Renderer.AddText(TEXT("{cam_passive}[%s]{cam_default}"), 
@@ -151,14 +168,33 @@ void FCameraDirectorTreeDebugBlock::OnDebugDraw(const FCameraDebugBlockDrawParam
 			Renderer.AddText(TEXT("{cam_error}Invalid context!{cam_default}\n"));
 		}
 
+		// Initial result pose block.
 		Renderer.AddIndent();
 		ChildrenView[Index]->DebugDraw(Params, Renderer);
 		Renderer.RemoveIndent();
 
+		// Children blocks.
+		if (EntryDebugInfo.NumChildrenContexts > 0)
+		{
+			Renderer.AddIndent();
+			for (int32 ChildIndex = 0; ChildIndex < EntryDebugInfo.NumChildrenContexts; ++ChildIndex)
+			{
+				ChildrenView[Index + 1 + ChildIndex]->DebugDraw(Params, Renderer);
+			}
+			Renderer.RemoveIndent();
+
+			Index += EntryDebugInfo.NumChildrenContexts;
+		}
+
+		++ChildContextIndex;
+
 		Renderer.NewLine();
 	}
 
-	Renderer.RemoveIndent();
+	if (bIsTreeRoot)
+	{
+		Renderer.RemoveIndent();
+	}
 	Renderer.SetTextColor(Colors.Default);
 
 	Renderer.SkipAllBlocks();
@@ -167,6 +203,7 @@ void FCameraDirectorTreeDebugBlock::OnDebugDraw(const FCameraDebugBlockDrawParam
 void FCameraDirectorTreeDebugBlock::OnSerialize(FArchive& Ar)
 {
 	 Ar << CameraDirectors;
+	 Ar << bIsTreeRoot;
 }
 
 FArchive& operator<< (FArchive& Ar, FCameraDirectorTreeDebugBlock::FDirectorDebugInfo& DirectorDebugInfo)
diff --git a/Engine/Plugins/Cameras/GameplayCameras/Source/GameplayCameras/Private/Directors/BlueprintCameraDirector.cpp b/Engine/Plugins/Cameras/GameplayCameras/Source/GameplayCameras/Private/Directors/BlueprintCameraDirector.cpp
--- a/Engine/Plugins/Cameras/GameplayCameras/Source/GameplayCameras/Private/Directors/BlueprintCameraDirector.cpp
+++ b/Engine/Plugins/Cameras/GameplayCameras/Source/GameplayCameras/Private/Directors/BlueprintCameraDirector.cpp
@@ -9,7 +9,6 @@
 #include "Core/CameraEvaluationContext.h"
 #include "Core/CameraRigAsset.h"
 #include "Core/CameraRigProxyAsset.h"
-#include "Core/CameraRigProxyTable.h"
 #include "Core/CameraSystemEvaluator.h"
 #include "Core/RootCameraNode.h"
 #include "Engine/Blueprint.h"
@@ -26,15 +25,6 @@
 namespace UE::Cameras
 {
 
-void FBlueprintCameraDirectorEvaluationResult::Reset()
-{
-	ActiveCameraRigProxies.Reset();
-	ActiveCameraRigs.Reset();
-	ActivePersistentCameraRigs.Reset();
-	InactivePersistentCameraRigs.Reset();
-	bForceActivateCameraRigs = false;
-}
-
 class FBlueprintCameraDirectorEvaluator : public FCameraDirectorEvaluator
 {
 	UE_DECLARE_CAMERA_DIRECTOR_EVALUATOR(GAMEPLAYCAMERAS_API, FBlueprintCameraDirectorEvaluator)
@@ -45,20 +35,10 @@ protected:
 	virtual void OnActivate(const FCameraDirectorActivateParams& Params) override;
 	virtual void OnDeactivate(const FCameraDirectorDeactivateParams& Params) override;
 	virtual void OnRun(const FCameraDirectorEvaluationParams& Params, FCameraDirectorEvaluationResult& OutResult) override;
+	virtual void OnAddChildEvaluationContext(const FChildContextManulationParams& Params, FChildContextManulationResult& Result) override;
+	virtual void OnRemoveChildEvaluationContext(const FChildContextManulationParams& Params, FChildContextManulationResult& Result) override;
 	virtual void OnAddReferencedObjects(FReferenceCollector& Collector) override;
 
-private:
-
-	void ActivateTransientCameraRigs(
-			const FCameraDirectorEvaluationParams& Params, 
-			const FBlueprintCameraDirectorEvaluationResult& BlueprintResult, 
-			FCameraDirectorEvaluationResult& OutResult);
-	void ActivateDeactivePersistentCameraRigs(
-			TSharedPtr<FCameraEvaluationContext> EvaluationContext,
-			const FBlueprintCameraDirectorEvaluationResult& BlueprintResult);
-
-	const UCameraRigAsset* FindCameraRigByProxy(const UCameraRigProxyAsset* InProxy);
-
 private:
 
 	FCameraSystemEvaluator* OwningEvaluator = nullptr;
@@ -86,6 +66,7 @@ void FBlueprintCameraDirectorEvaluator::OnInitialize(const FCameraDirectorInitia
 	{
 		UObject* Outer = Params.OwnerContext->GetOwner();
 		EvaluatorBlueprint = NewObject<UBlueprintCameraDirectorEvaluator>(Outer, Blueprint->CameraDirectorEvaluatorClass);
+		EvaluatorBlueprint->NativeInitializeCameraDirector(Params);
 	}
 	else
 	{
@@ -100,9 +81,6 @@ void FBlueprintCameraDirectorEvaluator::OnActivate(const FCameraDirectorActivate
 	if (EvaluatorBlueprint)
 	{
 		EvaluatorBlueprint->NativeActivateCameraDirector(Params);
-
-		const FBlueprintCameraDirectorEvaluationResult& BlueprintResult = EvaluatorBlueprint->GetEvaluationResult();
-		ActivateDeactivePersistentCameraRigs(Params.OwnerContext, BlueprintResult);
 	}
 	else
 	{
@@ -115,9 +93,6 @@ void FBlueprintCameraDirectorEvaluator::OnDeactivate(const FCameraDirectorDeacti
 	if (EvaluatorBlueprint)
 	{
 		EvaluatorBlueprint->NativeDeactivateCameraDirector(Params);
-
-		const FBlueprintCameraDirectorEvaluationResult& BlueprintResult = EvaluatorBlueprint->GetEvaluationResult();
-		ActivateDeactivePersistentCameraRigs(Params.OwnerContext, BlueprintResult);
 	}
 
 	OwningEvaluator = nullptr;
@@ -129,172 +104,148 @@ void FBlueprintCameraDirectorEvaluator::OnRun(const FCameraDirectorEvaluationPar
 	{
 		EvaluatorBlueprint->NativeRunCameraDirector(Params);
 
-		const FBlueprintCameraDirectorEvaluationResult& BlueprintResult = EvaluatorBlueprint->GetEvaluationResult();
-		ActivateTransientCameraRigs(Params, BlueprintResult, OutResult);
-		ActivateDeactivePersistentCameraRigs(Params.OwnerContext, BlueprintResult);
+		const FCameraDirectorEvaluationResult& BlueprintResult = EvaluatorBlueprint->GetEvaluationResult();
+		OutResult = BlueprintResult;
 	}
 }
 
-void FBlueprintCameraDirectorEvaluator::ActivateTransientCameraRigs(
-		const FCameraDirectorEvaluationParams& Params, 
-		const FBlueprintCameraDirectorEvaluationResult& BlueprintResult, 
-		FCameraDirectorEvaluationResult& OutResult)
+void FBlueprintCameraDirectorEvaluator::OnAddChildEvaluationContext(const FChildContextManulationParams& Params, FChildContextManulationResult& Result)
 {
-	TArray<const UCameraRigAsset*, TInlineAllocator<2>> CameraRigs;
-	const UCameraAsset* CameraAsset = Params.OwnerContext->GetCameraAsset();
-
-	// Gather camera rigs.
-	for (const UCameraRigAsset* ActiveCameraRig : BlueprintResult.ActiveCameraRigs)
+	if (EvaluatorBlueprint)
 	{
-		if (ActiveCameraRig)
+		UObject* ChildContextOwner = Params.ChildContext->GetOwner();
+		const bool bAdded = EvaluatorBlueprint->NativeAddChildEvaluationContext(ChildContextOwner);
+		if (bAdded)
 		{
-			CameraRigs.Add(ActiveCameraRig);
-		}
-		else
-		{
-			UE_LOG(
-					LogCameraSystem, 
-					Error, 
-					TEXT("Null camera rig specified in camera director '%s'."),
-					*EvaluatorBlueprint->GetClass()->GetPathName());
+			Result.Result = EChildContextManipulationResult::Success;
 		}
 	}
+}
 
-	// Resolve camera rig proxies.
-	for (const UCameraRigProxyAsset* ActiveCameraRigProxy : BlueprintResult.ActiveCameraRigProxies)
+void FBlueprintCameraDirectorEvaluator::OnRemoveChildEvaluationContext(const FChildContextManulationParams& Params, FChildContextManulationResult& Result)
+{
+	if (EvaluatorBlueprint)
 	{
-		const UCameraRigAsset* ActiveCameraRig = FindCameraRigByProxy(ActiveCameraRigProxy);
-		if (ActiveCameraRig)
-		{
-			CameraRigs.Add(ActiveCameraRig);
-		}
-		else
+		UObject* ChildContextOwner = Params.ChildContext->GetOwner();
+		const bool bRemoved = EvaluatorBlueprint->NativeRemoveChildEvaluationContext(ChildContextOwner);
+		
+		if (bRemoved)
 		{
-			UE_LOG(
-					LogCameraSystem, 
-					Error, 
-					TEXT("No camera rig found mapped to proxy '%s' in camera '%s'."),
-					*ActiveCameraRigProxy->GetPathName(), *CameraAsset->GetPathName());
+			Result.Result = EChildContextManipulationResult::Success;
 		}
 	}
-
-	// The BP interface doesn't specify the evaluation context for the chosen camera rigs: we always automatically
-	// make them run in our own owner context.
-	for (const UCameraRigAsset* ActiveCameraRig : CameraRigs)
-	{
-		OutResult.Add(Params.OwnerContext, ActiveCameraRig);
-	}
-
-	OutResult.bForceActivateCameraRigs = BlueprintResult.bForceActivateCameraRigs;
 }
 
-void FBlueprintCameraDirectorEvaluator::ActivateDeactivePersistentCameraRigs(
-		TSharedPtr<FCameraEvaluationContext> EvaluationContext,
-		const FBlueprintCameraDirectorEvaluationResult& BlueprintResult)
+void FBlueprintCameraDirectorEvaluator::OnAddReferencedObjects(FReferenceCollector& Collector)
 {
-	if (BlueprintResult.InactivePersistentCameraRigs.IsEmpty() && BlueprintResult.ActivePersistentCameraRigs.IsEmpty())
-	{
-		return;
-	}
-
-	if (!ensure(OwningEvaluator))
-	{
-		return;
-	}
-
-	FRootCameraNodeEvaluator* RootNodeEvaluator = OwningEvaluator->GetRootNodeEvaluator();
+	Collector.AddReferencedObject(EvaluatorBlueprint);
+}
 
-	APlayerController* PlayerController = EvaluationContext->GetPlayerController();
-	TSharedPtr<const FCameraEvaluationContext> ControllerEvaluationContext;
-	if (ensure(PlayerController))
-	{
-		ControllerEvaluationContext = UControllerGameplayCameraEvaluationComponent::FindOrAddEvaluationContext(PlayerController);
-	}
+}  // namespace UE::Cameras
 
-	for (const FBlueprintPersistentCameraRigInfo& CameraRigInfo : BlueprintResult.InactivePersistentCameraRigs)
-	{
-		FDeactivateCameraRigParams DeactivateParams;
-		DeactivateParams.EvaluationContext = ControllerEvaluationContext;
-		DeactivateParams.CameraRig = CameraRigInfo.CameraRig;
-		DeactivateParams.Layer = CameraRigInfo.Layer;
-		RootNodeEvaluator->DeactivateCameraRig(DeactivateParams);
-	}
+bool UBlueprintCameraDirectorEvaluator::RunChildCameraDirector(float DeltaTime, FName ChildSlotName)
+{
+	using namespace UE::Cameras;
 
-	for (const FBlueprintPersistentCameraRigInfo& CameraRigInfo : BlueprintResult.ActivePersistentCameraRigs)
-	{
-		FActivateCameraRigParams ActivateParams;
-		ActivateParams.EvaluationContext = ControllerEvaluationContext;
-		ActivateParams.CameraRig = CameraRigInfo.CameraRig;
-		ActivateParams.Layer = CameraRigInfo.Layer;
-		RootNodeEvaluator->ActivateCameraRig(ActivateParams);
-	}
-}
+	const int32 ChildIndex = ChildrenContextSlotNames.Find(ChildSlotName);
+	TArrayView<const TSharedPtr<FCameraEvaluationContext>> ChildrenContexts = EvaluationContext->GetChildrenContexts();
 
-const UCameraRigAsset* FBlueprintCameraDirectorEvaluator::FindCameraRigByProxy(const UCameraRigProxyAsset* InProxy)
-{
-	const UBlueprintCameraDirector* Blueprint = GetCameraDirectorAs<UBlueprintCameraDirector>();
-	if (!ensure(Blueprint))
+	if (ChildIndex >= 0 && ensure(ChildrenContexts.IsValidIndex(ChildIndex)))
 	{
-		return nullptr;
-	}
+		TSharedPtr<const FCameraEvaluationContext> ChildContext = ChildrenContexts[ChildIndex];
 
-	UCameraRigProxyTable* ProxyTable = Blueprint->CameraRigProxyTable;
-	if (!ensureMsgf(ProxyTable, TEXT("No proxy table set on Blueprint director '%s'."), *Blueprint->GetPathName()))
-	{
-		return nullptr;
+		if (FCameraDirectorEvaluator* ChildDirectorEvaluator = ChildContext->GetDirectorEvaluator())
+		{
+			FCameraDirectorEvaluationParams ChildParams;
+			ChildParams.DeltaTime = DeltaTime;
+			
+			FCameraDirectorEvaluationResult ChildResult;
+			
+			ChildDirectorEvaluator->Run(ChildParams, ChildResult);
+
+			if (!ChildResult.Requests.IsEmpty())
+			{
+				EvaluationResult = ChildResult;
+				return true;
+			}
+		}
 	}
 
-	FCameraRigProxyTableResolveParams ResolveParams;
-	ResolveParams.CameraRigProxy = InProxy;
-	return Blueprint->CameraRigProxyTable->ResolveProxy(ResolveParams);
+	return false;
 }
 
-void FBlueprintCameraDirectorEvaluator::OnAddReferencedObjects(FReferenceCollector& Collector)
-{
-	Collector.AddReferencedObject(EvaluatorBlueprint);
-}
-
-}  // namespace UE::Cameras
-
 void UBlueprintCameraDirectorEvaluator::ActivatePersistentBaseCameraRig(UCameraRigAsset* CameraRigPrefab)
 {
-	EvaluationResult.ActivePersistentCameraRigs.Add({ CameraRigPrefab, ECameraRigLayer::Base });
+	using namespace UE::Cameras;
+	FCameraRigActivationDeactivationRequest Request(EvaluationContext, CameraRigPrefab);
+	Request.RequestType = ECameraRigActivationDeactivationRequestType::Activate;
+	Request.Layer = ECameraRigLayer::Base;
+	EvaluationResult.Requests.Add(Request);
 }
 
 void UBlueprintCameraDirectorEvaluator::ActivatePersistentGlobalCameraRig(UCameraRigAsset* CameraRigPrefab)
 {
-	EvaluationResult.ActivePersistentCameraRigs.Add({ CameraRigPrefab, ECameraRigLayer::Global });
+	using namespace UE::Cameras;
+	FCameraRigActivationDeactivationRequest Request(EvaluationContext, CameraRigPrefab);
+	Request.RequestType = ECameraRigActivationDeactivationRequestType::Activate;
+	Request.Layer = ECameraRigLayer::Global;
+	EvaluationResult.Requests.Add(Request);
 }
 
 void UBlueprintCameraDirectorEvaluator::ActivatePersistentVisualCameraRig(UCameraRigAsset* CameraRigPrefab)
 {
-	EvaluationResult.ActivePersistentCameraRigs.Add({ CameraRigPrefab, ECameraRigLayer::Visual });
+	using namespace UE::Cameras;
+	FCameraRigActivationDeactivationRequest Request(EvaluationContext, CameraRigPrefab);
+	Request.RequestType = ECameraRigActivationDeactivationRequestType::Activate;
+	Request.Layer = ECameraRigLayer::Visual;
+	EvaluationResult.Requests.Add(Request);
 }
 
 void UBlueprintCameraDirectorEvaluator::DeactivatePersistentBaseCameraRig(UCameraRigAsset* CameraRigPrefab)
 {
-	EvaluationResult.InactivePersistentCameraRigs.Add({ CameraRigPrefab, ECameraRigLayer::Base });
+	using namespace UE::Cameras;
+	FCameraRigActivationDeactivationRequest Request(EvaluationContext, CameraRigPrefab);
+	Request.RequestType = ECameraRigActivationDeactivationRequestType::Deactivate;
+	Request.Layer = ECameraRigLayer::Base;
+	EvaluationResult.Requests.Add(Request);
 }
 
 void UBlueprintCameraDirectorEvaluator::DeactivatePersistentGlobalCameraRig(UCameraRigAsset* CameraRigPrefab)
 {
-	EvaluationResult.InactivePersistentCameraRigs.Add({ CameraRigPrefab, ECameraRigLayer::Global });
+	using namespace UE::Cameras;
+	FCameraRigActivationDeactivationRequest Request(EvaluationContext, CameraRigPrefab);
+	Request.RequestType = ECameraRigActivationDeactivationRequestType::Deactivate;
+	Request.Layer = ECameraRigLayer::Global;
+	EvaluationResult.Requests.Add(Request);
 }
 
 void UBlueprintCameraDirectorEvaluator::DeactivatePersistentVisualCameraRig(UCameraRigAsset* CameraRigPrefab)
 {
-	EvaluationResult.InactivePersistentCameraRigs.Add({ CameraRigPrefab, ECameraRigLayer::Visual });
+	using namespace UE::Cameras;
+	FCameraRigActivationDeactivationRequest Request(EvaluationContext, CameraRigPrefab);
+	Request.RequestType = ECameraRigActivationDeactivationRequestType::Deactivate;
+	Request.Layer = ECameraRigLayer::Visual;
+	EvaluationResult.Requests.Add(Request);
 }
 
 void UBlueprintCameraDirectorEvaluator::ActivateCameraRig(UCameraRigAsset* CameraRig, bool bForceNewInstance)
 {
-	EvaluationResult.ActiveCameraRigs.Add(CameraRig);
-	EvaluationResult.bForceActivateCameraRigs |= bForceNewInstance;
+	using namespace UE::Cameras;
+	FCameraRigActivationDeactivationRequest Request(EvaluationContext, CameraRig);
+	Request.RequestType = ECameraRigActivationDeactivationRequestType::Activate;
+	Request.Layer = ECameraRigLayer::Main;
+	Request.bForceActivateDeactivate = bForceNewInstance;
+	EvaluationResult.Requests.Add(Request);
 }
 
-void UBlueprintCameraDirectorEvaluator::ActivateCameraRigViaProxy(UCameraRigProxyAsset* CameraRigProxy)
+void UBlueprintCameraDirectorEvaluator::ActivateCameraRigViaProxy(UCameraRigProxyAsset* CameraRigProxy, bool bForceNewInstance)
 {
-	EvaluationResult.ActiveCameraRigProxies.Add(CameraRigProxy);
+	using namespace UE::Cameras;
+	FCameraRigActivationDeactivationRequest Request(EvaluationContext, CameraRigProxy);
+	Request.RequestType = ECameraRigActivationDeactivationRequestType::Activate;
+	Request.Layer = ECameraRigLayer::Main;
+	Request.bForceActivateDeactivate = bForceNewInstance;
+	EvaluationResult.Requests.Add(Request);
 }
 
 AActor* UBlueprintCameraDirectorEvaluator::FindEvaluationContextOwnerActor(TSubclassOf<AActor> ActorClass) const
@@ -428,21 +379,26 @@ UWorld* UBlueprintCameraDirectorEvaluator::GetWorld() const
 	return nullptr;
 }
 
-void UBlueprintCameraDirectorEvaluator::NativeActivateCameraDirector(const UE::Cameras::FCameraDirectorActivateParams& Params)
+void UBlueprintCameraDirectorEvaluator::NativeInitializeCameraDirector(const UE::Cameras::FCameraDirectorInitializeParams& Params)
 {
 	using namespace UE::Cameras;
 
 	EvaluationContext = Params.OwnerContext;
+}
+
+void UBlueprintCameraDirectorEvaluator::NativeActivateCameraDirector(const UE::Cameras::FCameraDirectorActivateParams& Params)
+{
+	using namespace UE::Cameras;
 
 	EvaluationResult.Reset();
 	{
-		FBlueprintCameraDirectorActivateParams BlueprintParams;
-		if (Params.OwnerContext)
+		UObject* EvaluationContextOwner = nullptr;
+		if (EvaluationContext)
 		{
-			BlueprintParams.EvaluationContextOwner = Params.OwnerContext->GetOwner();
+			EvaluationContextOwner = EvaluationContext->GetOwner();
 		}
 
-		ActivateCameraDirector(BlueprintParams);
+		ActivateCameraDirector(EvaluationContextOwner);
 	}
 }
 
@@ -450,31 +406,79 @@ void UBlueprintCameraDirectorEvaluator::NativeDeactivateCameraDirector(const UE:
 {
 	EvaluationResult.Reset();
 	{
-		FBlueprintCameraDirectorDeactivateParams BlueprintParams;
-		if (Params.OwnerContext)
+		UObject* EvaluationContextOwner = nullptr;
+		if (EvaluationContext)
 		{
-			BlueprintParams.EvaluationContextOwner = Params.OwnerContext->GetOwner();
+			EvaluationContextOwner = EvaluationContext->GetOwner();
 		}
 
-		DeactivateCameraDirector(BlueprintParams);
+		DeactivateCameraDirector(EvaluationContextOwner);
 	}
-
-	EvaluationContext = nullptr;
 }
 
 void UBlueprintCameraDirectorEvaluator::NativeRunCameraDirector(const UE::Cameras::FCameraDirectorEvaluationParams& Params)
 {
 	EvaluationResult.Reset();
 	{
-		FBlueprintCameraDirectorEvaluationParams BlueprintParams;
-		BlueprintParams.DeltaTime = Params.DeltaTime;
-		if (Params.OwnerContext)
+		UObject* EvaluationContextOwner = nullptr;
+		if (EvaluationContext)
+		{
+			EvaluationContextOwner = EvaluationContext->GetOwner();
+		}
+
+		RunCameraDirector(Params.DeltaTime, EvaluationContextOwner);
+	}
+}
+
+bool UBlueprintCameraDirectorEvaluator::NativeAddChildEvaluationContext(UObject* ChildEvaluationContextOwner)
+{
+	const FName ChildSlotName = AddChildEvaluationContext(ChildEvaluationContextOwner);
+	
+	if (!ChildSlotName.IsNone())
+	{
+		const int32 NewChildIndex = EvaluationContext->GetChildrenContexts().Num();
+
+		ensure(!ChildrenContextSlotNames.IsValidIndex(NewChildIndex) || ChildrenContextSlotNames[NewChildIndex].IsNone());
+		if (ChildrenContextSlotNames.Num() <= NewChildIndex)
 		{
-			BlueprintParams.EvaluationContextOwner = Params.OwnerContext->GetOwner();
+			ChildrenContextSlotNames.SetNum(NewChildIndex + 1);
 		}
 
-		RunCameraDirector(BlueprintParams);
+		ChildrenContextSlotNames[NewChildIndex] = ChildSlotName;
+
+		return true;
 	}
+
+	return false;
+}
+
+bool UBlueprintCameraDirectorEvaluator::NativeRemoveChildEvaluationContext(UObject* ChildEvaluationContextOwner)
+{
+	const int32 ChildIndex = EvaluationContext->GetChildrenContexts().IndexOfByPredicate(
+			[ChildEvaluationContextOwner](TSharedPtr<FCameraEvaluationContext> Item)
+			{
+				return Item->GetOwner() == ChildEvaluationContextOwner;
+			});
+	if (ChildIndex < 0)
+	{
+		return false;
+	}
+
+	if (!ensure(ChildrenContextSlotNames.IsValidIndex(ChildIndex) && !ChildrenContextSlotNames[ChildIndex].IsNone()))
+	{
+		return false;
+	}
+
+	const FName ChildSlotName = ChildrenContextSlotNames[ChildIndex];
+	const bool bRemoved = RemoveChildEvaluationContext(ChildEvaluationContextOwner, ChildSlotName);
+	if (bRemoved)
+	{
+		ChildrenContextSlotNames[ChildIndex] = NAME_None;
+
+		return true;
+	}
+
+	return false;
 }
 
 FCameraDirectorEvaluatorPtr UBlueprintCameraDirector::OnBuildEvaluator(FCameraDirectorEvaluatorBuilder& Builder) const
@@ -520,17 +524,5 @@ void UBlueprintCameraDirector::OnGatherRigUsageInfo(FCameraDirectorRigUsageInfo&
 #endif  // WITH_EDITORONLY_DATA
 }
 
-#if WITH_EDITOR
-
-void UBlueprintCameraDirector::OnFactoryCreateAsset(const FCameraDirectorFactoryCreateParams& InParams)
-{
-	if (!CameraRigProxyTable)
-	{
-		CameraRigProxyTable = NewObject<UCameraRigProxyTable>(this);
-	}
-}
-
-#endif  // WITH_EDITOR
-
 #undef LOCTEXT_NAMESPACE
 
diff --git a/Engine/Plugins/Cameras/GameplayCameras/Source/GameplayCameras/Private/Directors/PriorityQueueCameraDirector.cpp b/Engine/Plugins/Cameras/GameplayCameras/Source/GameplayCameras/Private/Directors/PriorityQueueCameraDirector.cpp
--- a/Engine/Plugins/Cameras/GameplayCameras/Source/GameplayCameras/Private/Directors/PriorityQueueCameraDirector.cpp
+++ b/Engine/Plugins/Cameras/GameplayCameras/Source/GameplayCameras/Private/Directors/PriorityQueueCameraDirector.cpp
@@ -71,10 +71,9 @@ void FPriorityQueueCameraDirectorEvaluator::OnRun(const FCameraDirectorEvaluatio
 			if (FCameraDirectorEvaluator* DirectorEvaluator = HighestPriorityEntry.ChildContext->GetDirectorEvaluator())
 			{
 				FCameraDirectorEvaluationParams ChildParams(Params);
-				ChildParams.OwnerContext = HighestPriorityEntry.ChildContext;
 				DirectorEvaluator->Run(ChildParams, OutResult);
 
-				if (OutResult.ActiveCameraRigs.Num() > 0)
+				if (OutResult.Requests.Num() > 0)
 				{
 					return;
 				}
diff --git a/Engine/Plugins/Cameras/GameplayCameras/Source/GameplayCameras/Private/Directors/SingleCameraDirector.cpp b/Engine/Plugins/Cameras/GameplayCameras/Source/GameplayCameras/Private/Directors/SingleCameraDirector.cpp
--- a/Engine/Plugins/Cameras/GameplayCameras/Source/GameplayCameras/Private/Directors/SingleCameraDirector.cpp
+++ b/Engine/Plugins/Cameras/GameplayCameras/Source/GameplayCameras/Private/Directors/SingleCameraDirector.cpp
@@ -22,7 +22,7 @@ protected:
 		const USingleCameraDirector* SingleDirector = GetCameraDirectorAs<USingleCameraDirector>();
 		if (SingleDirector->CameraRig)
 		{
-			OutResult.Add(Params.OwnerContext, SingleDirector->CameraRig);
+			OutResult.Add(GetEvaluationContext(), SingleDirector->CameraRig);
 		}
 	}
 };
@@ -34,6 +34,7 @@ UE_DEFINE_CAMERA_DIRECTOR_EVALUATOR(FSingleCameraDirectorEvaluator)
 USingleCameraDirector::USingleCameraDirector(const FObjectInitializer& ObjectInit)
 	: Super(ObjectInit)
 {
+	bHasCameraRigProxyTable = false;
 }
 
 FCameraDirectorEvaluatorPtr USingleCameraDirector::OnBuildEvaluator(FCameraDirectorEvaluatorBuilder& Builder) const
diff --git a/Engine/Plugins/Cameras/GameplayCameras/Source/GameplayCameras/Private/Directors/StateTreeCameraDirector.cpp b/Engine/Plugins/Cameras/GameplayCameras/Source/GameplayCameras/Private/Directors/StateTreeCameraDirector.cpp
--- a/Engine/Plugins/Cameras/GameplayCameras/Source/GameplayCameras/Private/Directors/StateTreeCameraDirector.cpp
+++ b/Engine/Plugins/Cameras/GameplayCameras/Source/GameplayCameras/Private/Directors/StateTreeCameraDirector.cpp
@@ -6,8 +6,6 @@
 #include "Core/CameraAsset.h"
 #include "Core/CameraEvaluationContext.h"
 #include "Core/CameraRigAsset.h"
-#include "Core/CameraRigProxyAsset.h"
-#include "Core/CameraRigProxyTable.h"
 #include "Directors/CameraDirectorStateTreeSchema.h"
 #include "GameplayCameras.h"
 #include "Helpers/OutgoingReferenceFinder.h"
@@ -40,7 +38,6 @@ protected:
 private:
 
 	bool SetContextRequirements(TSharedPtr<const FCameraEvaluationContext> OwnerContext, FStateTreeExecutionContext& StateTreeContext);
-	const UCameraRigAsset* FindCameraRigByProxy(const UCameraRigProxyAsset* InProxy);
 
 private:
 
@@ -65,7 +62,7 @@ void FStateTreeCameraDirectorEvaluator::OnActivate(const FCameraDirectorActivate
 		return;
 	}
 
-	UObject* ContextOwner = Params.OwnerContext->GetOwner();
+	UObject* ContextOwner = GetEvaluationContext()->GetOwner();
 	if (!ContextOwner)
 	{
 		UE_LOG(LogCameraSystem, Error,
@@ -87,7 +84,7 @@ void FStateTreeCameraDirectorEvaluator::OnActivate(const FCameraDirectorActivate
 
 	// TODO: validate schema.
 	
-	if (!SetContextRequirements(Params.OwnerContext, StateTreeContext))
+	if (!SetContextRequirements(GetEvaluationContext(), StateTreeContext))
 	{
 		UE_LOG(LogCameraSystem, Error,
 			TEXT("Can't activate camera director '%s': failed to setup external data views for StateTree asset '%s'."),
@@ -104,7 +101,7 @@ void FStateTreeCameraDirectorEvaluator::OnDeactivate(const FCameraDirectorDeacti
 	const FStateTreeReference& StateTreeReference = StateTreeDirector->StateTreeReference;
 	const UStateTree* StateTree = StateTreeReference.GetStateTree();
 
-	UObject* ContextOwner = Params.OwnerContext->GetOwner();
+	UObject* ContextOwner = GetEvaluationContext()->GetOwner();
 	if (!ContextOwner)
 	{
 		UE_LOG(LogCameraSystem, Error,
@@ -123,7 +120,7 @@ void FStateTreeCameraDirectorEvaluator::OnDeactivate(const FCameraDirectorDeacti
 	
 	FStateTreeExecutionContext StateTreeContext(*ContextOwner, *StateTree, StateTreeInstanceData);
 
-	if (SetContextRequirements(Params.OwnerContext, StateTreeContext))
+	if (SetContextRequirements(GetEvaluationContext(), StateTreeContext))
 	{
 		StateTreeContext.Stop();
 	}
@@ -135,7 +132,7 @@ void FStateTreeCameraDirectorEvaluator::OnRun(const FCameraDirectorEvaluationPar
 	const FStateTreeReference& StateTreeReference = StateTreeDirector->StateTreeReference;
 	const UStateTree* StateTree = StateTreeReference.GetStateTree();
 
-	UObject* ContextOwner = Params.OwnerContext->GetOwner();
+	UObject* ContextOwner = GetEvaluationContext()->GetOwner();
 
 	if (!StateTree || !ContextOwner)
 	{
@@ -145,19 +142,15 @@ void FStateTreeCameraDirectorEvaluator::OnRun(const FCameraDirectorEvaluationPar
 	
 	FStateTreeExecutionContext StateTreeContext(*ContextOwner, *StateTree, StateTreeInstanceData);
 
-	if (SetContextRequirements(Params.OwnerContext, StateTreeContext))
+	if (SetContextRequirements(GetEvaluationContext(), StateTreeContext))
 	{
 		StateTreeContext.Tick(Params.DeltaTime);
 
-		TArray<const UCameraRigAsset*, TInlineAllocator<2>> CameraRigs;
-		const UCameraAsset* CameraAsset = Params.OwnerContext->GetCameraAsset();
-
-		// Gather camera rigs.
 		for (const UCameraRigAsset* ActiveCameraRig : EvaluationData.ActiveCameraRigs)
 		{
 			if (ActiveCameraRig)
 			{
-				CameraRigs.Add(ActiveCameraRig);
+				OutResult.Add(GetEvaluationContext(), ActiveCameraRig);
 			}
 			else
 			{
@@ -166,26 +159,19 @@ void FStateTreeCameraDirectorEvaluator::OnRun(const FCameraDirectorEvaluationPar
 			}
 		}
 
-		// Resolve camera rig proxies.
 		for (const UCameraRigProxyAsset* ActiveCameraRigProxy : EvaluationData.ActiveCameraRigProxies)
 		{
-			const UCameraRigAsset* ActiveCameraRig = FindCameraRigByProxy(ActiveCameraRigProxy);
-			if (ActiveCameraRig)
+			if (ActiveCameraRigProxy)
 			{
-				CameraRigs.Add(ActiveCameraRig);
+				OutResult.Add(GetEvaluationContext(), ActiveCameraRigProxy);
 			}
 			else
 			{
-				UE_LOG(LogCameraSystem, Error, TEXT("No camera rig found mapped to proxy '%s' in camera '%s'."),
-						*ActiveCameraRigProxy->GetPathName(), *CameraAsset->GetPathName());
+				UE_LOG(LogCameraSystem, Error, TEXT("Null camera rig proxy specified in camera director '%s'."),
+						*StateTree->GetPathName());
 			}
 		}
 
-		// Set all collected camera rigs as our active rigs this frame.
-		for (const UCameraRigAsset* CameraRig : CameraRigs)
-		{
-			OutResult.Add(Params.OwnerContext, CameraRig);
-		}
 	}
 }
 
@@ -216,25 +202,6 @@ bool FStateTreeCameraDirectorEvaluator::SetContextRequirements(TSharedPtr<const
 	return true;
 }
 
-const UCameraRigAsset* FStateTreeCameraDirectorEvaluator::FindCameraRigByProxy(const UCameraRigProxyAsset* InProxy)
-{
-	const UStateTreeCameraDirector* Director = GetCameraDirectorAs<UStateTreeCameraDirector>();
-	if (!ensure(Director))
-	{
-		return nullptr;
-	}
-
-	const UCameraRigProxyTable* ProxyTable = Director->CameraRigProxyTable;
-	if (!ensureMsgf(ProxyTable, TEXT("No proxy table set on StateTree director '%s'."), *Director->GetPathName()))
-	{
-		return nullptr;
-	}
-
-	FCameraRigProxyTableResolveParams ResolveParams;
-	ResolveParams.CameraRigProxy = InProxy;
-	return Director->CameraRigProxyTable->ResolveProxy(ResolveParams);
-}
-
 void FStateTreeCameraDirectorEvaluator::OnAddReferencedObjects(FReferenceCollector& Collector)
 {
 	StateTreeInstanceData.AddStructReferencedObjects(Collector);
@@ -280,17 +247,5 @@ void UStateTreeCameraDirector::OnGatherRigUsageInfo(FCameraDirectorRigUsageInfo&
 	ReferenceFinder.GetReferencesOfClass<UCameraRigAsset>(UsageInfo.CameraRigs);
 }
 
-#if WITH_EDITOR
-
-void UStateTreeCameraDirector::OnFactoryCreateAsset(const FCameraDirectorFactoryCreateParams& InParams)
-{
-	if (!CameraRigProxyTable)
-	{
-		CameraRigProxyTable = NewObject<UCameraRigProxyTable>(this);
-	}
-}
-
-#endif  // WITH_EDITOR
-
 #undef LOCTEXT_NAMESPACE
 
diff --git a/Engine/Plugins/Cameras/GameplayCameras/Source/GameplayCameras/Private/GameFramework/GameplayCameraComponentBase.cpp b/Engine/Plugins/Cameras/GameplayCameras/Source/GameplayCameras/Private/GameFramework/GameplayCameraComponentBase.cpp
--- a/Engine/Plugins/Cameras/GameplayCameras/Source/GameplayCameras/Private/GameFramework/GameplayCameraComponentBase.cpp
+++ b/Engine/Plugins/Cameras/GameplayCameras/Source/GameplayCameras/Private/GameFramework/GameplayCameraComponentBase.cpp
@@ -102,14 +102,14 @@ APlayerController* UGameplayCameraComponentBase::GetPlayerController() const
 	return nullptr;
 }
 
-void UGameplayCameraComponentBase::ActivateCameraForPlayerIndex(int32 PlayerIndex)
+void UGameplayCameraComponentBase::ActivateCameraForPlayerIndex(int32 PlayerIndex, EGameplayCameraComponentActivationMode ActivationMode)
 {
-	ActivateCameraEvaluationContext(PlayerIndex);
+	ActivateCameraEvaluationContext(PlayerIndex, ActivationMode);
 }
 
-void UGameplayCameraComponentBase::ActivateCameraForPlayerController(APlayerController* PlayerController)
+void UGameplayCameraComponentBase::ActivateCameraForPlayerController(APlayerController* PlayerController, EGameplayCameraComponentActivationMode ActivationMode)
 {
-	ActivateCameraEvaluationContext(PlayerController);
+	ActivateCameraEvaluationContext(PlayerController, ActivationMode);
 }
 
 void UGameplayCameraComponentBase::DeactivateCamera()
@@ -117,12 +117,12 @@ void UGameplayCameraComponentBase::DeactivateCamera()
 	DeactivateCameraEvaluationContext();
 }
 
-void UGameplayCameraComponentBase::ActivateCameraEvaluationContext(int32 PlayerIndex)
+void UGameplayCameraComponentBase::ActivateCameraEvaluationContext(int32 PlayerIndex, EGameplayCameraComponentActivationMode ActivationMode)
 {
 	DeactivateCameraEvaluationContext();
 
 	APlayerController* PlayerController = UGameplayStatics::GetPlayerController(this, PlayerIndex);
-	ActivateCameraEvaluationContext(PlayerController);
+	ActivateCameraEvaluationContext(PlayerController, ActivationMode);
 }
 
 void UGameplayCameraComponentBase::DeactivateCameraEvaluationContext()
@@ -136,8 +136,15 @@ void UGameplayCameraComponentBase::DeactivateCameraEvaluationContext()
 
 	if (EvaluationContext.IsValid())
 	{
-		TSharedPtr<FCameraSystemEvaluator> Evaluator = CameraSystemHost->GetCameraSystemEvaluator();
-		Evaluator->RemoveEvaluationContext(EvaluationContext.ToSharedRef());
+		if (TSharedPtr<FCameraEvaluationContext> ParentContext = EvaluationContext->GetParentContext())
+		{
+			ParentContext->RemoveChildContext(EvaluationContext.ToSharedRef());
+		}
+		else
+		{
+			TSharedPtr<FCameraSystemEvaluator> Evaluator = CameraSystemHost->GetCameraSystemEvaluator();
+			Evaluator->RemoveEvaluationContext(EvaluationContext.ToSharedRef());
+		}
 	}
 
 	if (OutputCameraComponent)
@@ -149,7 +156,7 @@ void UGameplayCameraComponentBase::DeactivateCameraEvaluationContext()
 	// running camera rigs blend out.
 }
 
-void UGameplayCameraComponentBase::ActivateCameraEvaluationContext(APlayerController* PlayerController)
+void UGameplayCameraComponentBase::ActivateCameraEvaluationContext(APlayerController* PlayerController, EGameplayCameraComponentActivationMode ActivationMode)
 {
 	using namespace UE::Cameras;
 
@@ -172,6 +179,7 @@ void UGameplayCameraComponentBase::ActivateCameraEvaluationContext(APlayerContro
 
 	if (!OnValidateCameraEvaluationContextActivation())
 	{
+		// We expect the sub-class to print a useful error or warning message.
 		return;
 	}
 
@@ -195,7 +203,40 @@ void UGameplayCameraComponentBase::ActivateCameraEvaluationContext(APlayerContro
 	}
 
 	TSharedPtr<FCameraSystemEvaluator> CameraSystemEvaluator = CameraSystemHost->GetCameraSystemEvaluator();
-	CameraSystemEvaluator->PushEvaluationContext(EvaluationContext.ToSharedRef());
+	FCameraEvaluationContextStack& ContextStack = CameraSystemEvaluator->GetEvaluationContextStack();
+
+	switch (ActivationMode)
+	{
+		case EGameplayCameraComponentActivationMode::Push:
+			ContextStack.PushContext(EvaluationContext.ToSharedRef());
+			break;
+		case EGameplayCameraComponentActivationMode::PushAndInsert:
+			{
+				TSharedPtr<FCameraEvaluationContext> PreviousActiveContext = ContextStack.GetActiveContext();
+				ContextStack.PushContext(EvaluationContext.ToSharedRef());
+				if (PreviousActiveContext)
+				{
+					ContextStack.RemoveContext(PreviousActiveContext.ToSharedRef());
+					EvaluationContext->AddChildContext(PreviousActiveContext.ToSharedRef());
+				}
+			}
+			break;
+		case EGameplayCameraComponentActivationMode::InsertOrPush:
+			{
+				if (TSharedPtr<FCameraEvaluationContext> ActiveContext = ContextStack.GetActiveContext())
+				{
+					ActiveContext->AddChildContext(EvaluationContext.ToSharedRef());
+				}
+				else
+				{
+					ContextStack.PushContext(EvaluationContext.ToSharedRef());
+				}
+			}
+			break;
+		default:
+			ensure(false);
+			break;
+	}
 
 	// Make sure the component is active so it receives tick updates to maintain the evaluation context.
 	Activate();
diff --git a/Engine/Plugins/Cameras/GameplayCameras/Source/GameplayCameras/Public/Core/CameraDirector.h b/Engine/Plugins/Cameras/GameplayCameras/Source/GameplayCameras/Public/Core/CameraDirector.h
--- a/Engine/Plugins/Cameras/GameplayCameras/Source/GameplayCameras/Public/Core/CameraDirector.h
+++ b/Engine/Plugins/Cameras/GameplayCameras/Source/GameplayCameras/Public/Core/CameraDirector.h
@@ -8,6 +8,8 @@
 
 #include "CameraDirector.generated.h"
 
+class UCameraRigProxyTable;
+
 namespace UE::Cameras { class FCameraBuildLog; }
 
 #if WITH_EDITOR
@@ -71,5 +73,20 @@ protected:
 	/** Called by the asset factories to setup new data before the editor opens. */
 	virtual void OnFactoryCreateAsset(const FCameraDirectorFactoryCreateParams& InParams) {}
 #endif
+
+public:
+
+	/** 
+	 * The table that maps camera rig proxies (used in the evaluator Blueprint graph)
+	 * to actual camera rigs.
+	 */
+	UPROPERTY(EditAnywhere, Instanced, Category="Evaluation")
+	TObjectPtr<UCameraRigProxyTable> CameraRigProxyTable;
+
+protected:
+
+	/** Whether this type of director supports having a camera rig proxy table. */
+	UPROPERTY(Transient)
+	bool bHasCameraRigProxyTable = true;
 };
 
diff --git a/Engine/Plugins/Cameras/GameplayCameras/Source/GameplayCameras/Public/Core/CameraDirectorEvaluator.h b/Engine/Plugins/Cameras/GameplayCameras/Source/GameplayCameras/Public/Core/CameraDirectorEvaluator.h
--- a/Engine/Plugins/Cameras/GameplayCameras/Source/GameplayCameras/Public/Core/CameraDirectorEvaluator.h
+++ b/Engine/Plugins/Cameras/GameplayCameras/Source/GameplayCameras/Public/Core/CameraDirectorEvaluator.h
@@ -2,6 +2,7 @@
 
 #pragma once
 
+#include "Core/CameraRigInstanceID.h"
 #include "Core/CameraObjectRtti.h"
 #include "CoreTypes.h"
 #include "GameplayCameras.h"
@@ -10,6 +11,7 @@
 
 class UCameraDirector;
 class UCameraRigAsset;
+class UCameraRigProxyAsset;
 class UCameraRigTransition;
 
 namespace UE::Cameras
@@ -39,9 +41,6 @@ struct FCameraDirectorActivateParams
 {
 	/** The camera system that will run the camera director. */
 	FCameraSystemEvaluator* Evaluator = nullptr;
-
-	/** The evaluation context that owns the camera director. */
-	TSharedPtr<FCameraEvaluationContext> OwnerContext;
 };
 
 /**
@@ -49,8 +48,6 @@ struct FCameraDirectorActivateParams
  */
 struct FCameraDirectorDeactivateParams
 {
-	/** The evaluation context that owns the camera director. */
-	TSharedPtr<FCameraEvaluationContext> OwnerContext;
 };
 
 /**
@@ -60,20 +57,53 @@ struct FCameraDirectorEvaluationParams
 {
 	/** Time interval for the update. */
 	float DeltaTime = 0.f;
+};
 
-	/** The context in which this director runs. */
-	TSharedPtr<FCameraEvaluationContext> OwnerContext;
+/** The type of request for activating or deactivating a camera rig. */
+enum class ECameraRigActivationDeactivationRequestType
+{
+	Activate,
+	Deactivate
 };
 
 /**
- * Structure describing a camera rig that should be active.
+ * Structure for requesting that a camera rig be activated or deactivated.
  */
-struct FActiveCameraRigInfo
+struct FCameraRigActivationDeactivationRequest
 {
 	/** The evaluation context to run the specified camera rig. */
 	TSharedPtr<const FCameraEvaluationContext> EvaluationContext;
 	/** The camera rig that should be running. */
 	TObjectPtr<const UCameraRigAsset> CameraRig;
+	/** The camera proxy that determines the rig that should be running (if CameraRig is null). */
+	TObjectPtr<const UCameraRigProxyAsset> CameraRigProxy;
+	/** The type of the request. */
+	ECameraRigActivationDeactivationRequestType RequestType = ECameraRigActivationDeactivationRequestType::Activate;
+	/** The layer on which to activate or deactivate this camera rig. */
+	ECameraRigLayer Layer = ECameraRigLayer::Main;
+	/** A transition to use for this activation, if possible, instead of looking up one according to usual rules. */
+	TObjectPtr<const UCameraRigTransition> TransitionOverride;
+	/** An order key for activating camera rigs on ordered layers. */
+	int32 OrderKey = 0;
+	/** Whether to force creating a new instance, even if an identical one is already active. */
+	bool bForceActivateDeactivate = false;
+
+public:
+
+	FCameraRigActivationDeactivationRequest()
+	{}
+	FCameraRigActivationDeactivationRequest(
+			TSharedPtr<const FCameraEvaluationContext> InContext, 
+			TObjectPtr<const UCameraRigAsset> InCameraRig)
+		: EvaluationContext(InContext)
+		, CameraRig(InCameraRig)
+	{}
+	FCameraRigActivationDeactivationRequest(
+			TSharedPtr<const FCameraEvaluationContext> InContext, 
+			TObjectPtr<const UCameraRigProxyAsset> InCameraRigProxy)
+		: EvaluationContext(InContext)
+		, CameraRigProxy(InCameraRigProxy)
+	{}
 };
 
 /**
@@ -81,21 +111,31 @@ struct FActiveCameraRigInfo
  */
 struct FCameraDirectorEvaluationResult
 {
-	using FActiveCameraRigInfos = TArray<FActiveCameraRigInfo, TInlineAllocator<2>>;
+	using FRequests = TArray<FCameraRigActivationDeactivationRequest, TInlineAllocator<4>>;
 
-	/** The camera rig(s) that the director says should be active this frame in the main layer. */
-	FActiveCameraRigInfos ActiveCameraRigs;
+	/** Requests for activating/deactivating camera rigs this frame. */
+	FRequests Requests;
 
-	/** A transition to use for the next camera rig activation, instead of looking up one according to usual rules. */
-	TObjectPtr<const UCameraRigTransition> TransitionOverride;
+	/** Adds a simple activation request for the main layer. */
+	void Add(
+			TSharedPtr<const FCameraEvaluationContext> InEvaluationContext,
+			TObjectPtr<const UCameraRigAsset> InCameraRig)
+	{
+		Requests.Add({ InEvaluationContext, InCameraRig });
+	}
 
-	/** Whether to force creating a new instance of the camera rigs, even if they are already active in the main layer. */
-	bool bForceActivateCameraRigs = false;
+	/** Adds a simple activation request for the main layer. */
+	void Add(
+			TSharedPtr<const FCameraEvaluationContext> InEvaluationContext,
+			TObjectPtr<const UCameraRigProxyAsset> InCameraRigProxy)
+	{
+		Requests.Add({ InEvaluationContext, InCameraRigProxy });
+	}
 
-	/** Adds a given camera rig with the given evaluation context to the activation list. */
-	void Add(TSharedPtr<const FCameraEvaluationContext> InContext, TObjectPtr<const UCameraRigAsset> InCameraRig)
+	/** Reset this result. */
+	void Reset()
 	{
-		ActiveCameraRigs.Add(FActiveCameraRigInfo{ InContext, InCameraRig });
+		Requests.Reset();
 	}
 };
 
@@ -170,6 +210,9 @@ public:
 		return Cast<CameraDirectorType>(PrivateCameraDirector);
 	}
 
+	/** Gets the owning evaluation context. */
+	TSharedPtr<FCameraEvaluationContext> GetEvaluationContext() const { return WeakOwnerContext.Pin(); }
+
 public:
 
 	/** Runs the camera director to determine what camera rig(s) should be active this frame. */
@@ -183,33 +226,46 @@ public:
 
 public:
 
+	/** Add a child context to this camera director. */
 	GAMEPLAYCAMERAS_API bool AddChildEvaluationContext(TSharedRef<FCameraEvaluationContext> InContext);
+
+	/** Remove the given child context from this camera director. */
 	GAMEPLAYCAMERAS_API bool RemoveChildEvaluationContext(TSharedRef<FCameraEvaluationContext> InContext);
 
+	/** Garbage collection pass. */
 	GAMEPLAYCAMERAS_API void AddReferencedObjects(FReferenceCollector& Collector);
 
 public:
 
 	// Internal API.
+
 	void SetPrivateCameraDirector(const UCameraDirector* InCameraDirector);
 
+	const UCameraRigAsset* FindCameraRigByProxy(const UCameraRigProxyAsset* InProxy) const;
+
 	void OnEndCameraSystemUpdate();
 
 protected:
 
+	/** Result for adding/removing children contexts. */
 	enum class EChildContextManipulationResult
 	{
+		/** The add or removal failed. */
 		Failure,
+		/** The add or removal was successful. */
 		Success,
+		/** The add or removal was successfully handled by a child director. */
 		ChildContextSuccess
 	};
 
+	/** Parameter struct for adding/removing children contexts. */
 	struct FChildContextManulationParams
 	{
 		TSharedPtr<FCameraEvaluationContext> ParentContext;
 		TSharedPtr<FCameraEvaluationContext> ChildContext;
 	};
 
+	/** Result struct for adding/removing children contexts. */
 	struct FChildContextManulationResult
 	{
 		EChildContextManipulationResult Result = EChildContextManipulationResult::Failure;
@@ -229,13 +285,21 @@ protected:
 	/** Runs the camera director to determine what camera rig(s) should be active this frame. */
 	virtual void OnRun(const FCameraDirectorEvaluationParams& Params, FCameraDirectorEvaluationResult& OutResult) {}
 
+	/** Add a child context to this camera director. */
 	virtual void OnAddChildEvaluationContext(const FChildContextManulationParams& Params, FChildContextManulationResult& Result) {}
+
+	/** Remove the given child context from this camera director. */
 	virtual void OnRemoveChildEvaluationContext(const FChildContextManulationParams& Params, FChildContextManulationResult& Result) {}
 
+	/** Garbage collection pass. */
 	virtual void OnAddReferencedObjects(FReferenceCollector& Collector) {}
 
 private:
 
+	/** The camera system this evaluator is running inside of. */
+	FCameraSystemEvaluator* Evaluator = nullptr;
+
+	/** The evaluation context that owns this evaluator. */
 	TWeakPtr<FCameraEvaluationContext> WeakOwnerContext;
 
 	/** The camera director this evaluator is running. */
diff --git a/Engine/Plugins/Cameras/GameplayCameras/Source/GameplayCameras/Public/Core/CameraEvaluationContext.h b/Engine/Plugins/Cameras/GameplayCameras/Source/GameplayCameras/Public/Core/CameraEvaluationContext.h
--- a/Engine/Plugins/Cameras/GameplayCameras/Source/GameplayCameras/Public/Core/CameraEvaluationContext.h
+++ b/Engine/Plugins/Cameras/GameplayCameras/Source/GameplayCameras/Public/Core/CameraEvaluationContext.h
@@ -45,7 +45,11 @@ struct FCameraEvaluationContextInitializeParams
  */
 struct FCameraEvaluationContextActivateParams
 {
+	/** The system the context will evaluate inside of. */
 	FCameraSystemEvaluator* Evaluator = nullptr;
+
+	/** The parent context, if any. */
+	TSharedPtr<FCameraEvaluationContext> ParentContext;
 };
 
 /**
@@ -103,6 +107,12 @@ public:
 	/** Gets the camera director evaluator that runs the camera rigs of this context. */
 	FCameraDirectorEvaluator* GetDirectorEvaluator() const { return DirectorEvaluator; }
 
+	/** Gets the parent evaluation context, if any. */
+	TSharedPtr<const FCameraEvaluationContext> GetParentContext() const { return WeakParent.Pin(); }
+
+	/** Gets the parent evaluation context, if any. */
+	TSharedPtr<FCameraEvaluationContext> GetParentContext() { return WeakParent.Pin(); }
+
 	/** Gets the children evaluation contexts running inside this context. */
 	TArrayView<const TSharedPtr<FCameraEvaluationContext>> GetChildrenContexts() const { return ChildrenContexts; }
 
@@ -136,6 +146,14 @@ public:
 	/** Whether this context is active. */
 	bool IsActive() const { return bActivated; }
 
+public:
+
+	/** Adds a child context. */
+	bool AddChildContext(TSharedRef<FCameraEvaluationContext> ChildContext);
+
+	/** Removes a child context. */
+	bool RemoveChildContext(TSharedRef<FCameraEvaluationContext> ChildContext);
+
 public:
 
 	// Internal API.
diff --git a/Engine/Plugins/Cameras/GameplayCameras/Source/GameplayCameras/Public/Core/CameraEvaluationContextStack.h b/Engine/Plugins/Cameras/GameplayCameras/Source/GameplayCameras/Public/Core/CameraEvaluationContextStack.h
--- a/Engine/Plugins/Cameras/GameplayCameras/Source/GameplayCameras/Public/Core/CameraEvaluationContextStack.h
+++ b/Engine/Plugins/Cameras/GameplayCameras/Source/GameplayCameras/Public/Core/CameraEvaluationContextStack.h
@@ -37,12 +37,14 @@ public:
 	void PushContext(TSharedRef<FCameraEvaluationContext> Context);
 
 	/**
-	 * Tries to add a context inside the active context. This will query the active context's
+	 * Tries to add a context inside the given context. This will query the active context's
 	 * director in order to find an "available spot" for the child context.
 	 * 
+	 * @param Context  The new context to activate.
+	 * @param ParentContext  The parent context to activate inside of, or the active context (if any) if null.
 	 * @return Whether the child context was acccepted.
 	 */
-	bool AddChildContext(TSharedRef<FCameraEvaluationContext> Context);
+	bool AddChildContext(TSharedRef<FCameraEvaluationContext> Context, TSharedPtr<FCameraEvaluationContext> ParentContext = nullptr);
 
 	/** Remove an existing context from the stack. */
 	bool RemoveContext(TSharedRef<FCameraEvaluationContext> Context);
diff --git a/Engine/Plugins/Cameras/GameplayCameras/Source/GameplayCameras/Public/Core/CameraSystemEvaluator.h b/Engine/Plugins/Cameras/GameplayCameras/Source/GameplayCameras/Public/Core/CameraSystemEvaluator.h
--- a/Engine/Plugins/Cameras/GameplayCameras/Source/GameplayCameras/Public/Core/CameraSystemEvaluator.h
+++ b/Engine/Plugins/Cameras/GameplayCameras/Source/GameplayCameras/Public/Core/CameraSystemEvaluator.h
@@ -25,12 +25,14 @@ struct FMinimalViewInfo;
 namespace UE::Cameras
 {
 
+class FCameraDirectorEvaluator;
 class FCameraEvaluationContext;
 class FCameraEvaluationService;
 class FCameraRigCombinationRegistry;
 class FRootCameraNodeEvaluator;
 enum class ECameraEvaluationServiceFlags;
 enum class ECameraNodeEvaluationType;
+struct FCameraRigActivationDeactivationRequest;
 struct FRootCameraNodeCameraRigEvent;
 
 #if UE_GAMEPLAY_CAMERAS_DEBUG
@@ -247,6 +249,9 @@ private:
 
 	void UpdateImpl(float DeltaTime, ECameraNodeEvaluationType EvaluationType);
 
+	void UpdateCameraDirector(float DeltaTime, FCameraDirectorEvaluator* CameraDirectorEvaluator);
+	void GetCombinedCameraRigRequest(TConstArrayView<FCameraRigActivationDeactivationRequest> Requests, FCameraRigActivationDeactivationRequest& OutCombinedRequest);
+
 	void PreUpdateServices(float DeltaTime, ECameraEvaluationServiceFlags ExtraFlags);
 	void PostUpdateServices(float DeltaTime, ECameraEvaluationServiceFlags ExtraFlags);
 
diff --git a/Engine/Plugins/Cameras/GameplayCameras/Source/GameplayCameras/Public/Core/RootCameraNode.h b/Engine/Plugins/Cameras/GameplayCameras/Source/GameplayCameras/Public/Core/RootCameraNode.h
--- a/Engine/Plugins/Cameras/GameplayCameras/Source/GameplayCameras/Public/Core/RootCameraNode.h
+++ b/Engine/Plugins/Cameras/GameplayCameras/Source/GameplayCameras/Public/Core/RootCameraNode.h
@@ -27,6 +27,7 @@ namespace UE::Cameras
 
 class FCameraEvaluationContext;
 class FCameraSystemEvaluator;
+struct FCameraRigActivationDeactivationRequest;
 struct FRootCameraNodeCameraRigEvent;
 
 /**
@@ -122,6 +123,11 @@ public:
 	 */
 	void DeactivateCameraRig(const FDeactivateCameraRigParams& Params);
 
+	/**
+	 * Execute a request to activate or deactivate a camera rig.
+	 */
+	void ExecuteCameraDirectorRequest(const FCameraRigActivationDeactivationRequest& Request);
+
 	/** Gets information about the active camera rig in the main layer. */
 	void GetActiveCameraRigInfo(FCameraRigEvaluationInfo& OutCameraRigInfo) const;
 
diff --git a/Engine/Plugins/Cameras/GameplayCameras/Source/GameplayCameras/Public/Debug/CameraDirectorTreeDebugBlock.h b/Engine/Plugins/Cameras/GameplayCameras/Source/GameplayCameras/Public/Debug/CameraDirectorTreeDebugBlock.h
--- a/Engine/Plugins/Cameras/GameplayCameras/Source/GameplayCameras/Public/Debug/CameraDirectorTreeDebugBlock.h
+++ b/Engine/Plugins/Cameras/GameplayCameras/Source/GameplayCameras/Public/Debug/CameraDirectorTreeDebugBlock.h
@@ -48,10 +48,13 @@ private:
 		FString OwnerName;
 		FString CameraAssetName;
 		FName CameraDirectorClassName;
+		int32 NumChildrenContexts = 0;
 		FTransform3d InitialContextTransform = FTransform3d::Identity;
 		bool bIsValid = false;
 	};
+
 	TArray<FDirectorDebugInfo> CameraDirectors;
+	bool bIsTreeRoot = true;
 
 	friend FArchive& operator<< (FArchive&, FDirectorDebugInfo&);
 };
diff --git a/Engine/Plugins/Cameras/GameplayCameras/Source/GameplayCameras/Public/Directors/BlueprintCameraDirector.h b/Engine/Plugins/Cameras/GameplayCameras/Source/GameplayCameras/Public/Directors/BlueprintCameraDirector.h
--- a/Engine/Plugins/Cameras/GameplayCameras/Source/GameplayCameras/Public/Directors/BlueprintCameraDirector.h
+++ b/Engine/Plugins/Cameras/GameplayCameras/Source/GameplayCameras/Public/Directors/BlueprintCameraDirector.h
@@ -14,88 +14,8 @@
 
 class UCameraRigAsset;
 class UCameraRigProxyAsset;
-class UCameraRigProxyTable;
 enum class ECameraRigLayer : uint8;
 
-namespace UE::Cameras
-{
-
-/** Information about a persitent camera rig to be activated or deactivated. */
-struct FBlueprintPersistentCameraRigInfo
-{
-	UCameraRigAsset* CameraRig;
-	ECameraRigLayer Layer;
-};
-
-/**
- * The evaluation result for the Blueprint camera director evaluator.
- */
-struct FBlueprintCameraDirectorEvaluationResult
-{
-	/** The list of camera rigs that should be active this frame. */
-	TArray<UCameraRigProxyAsset*> ActiveCameraRigProxies;
-
-	/** The list of camera rigs that should be active this frame. */
-	TArray<UCameraRigAsset*> ActiveCameraRigs;
-
-	/** The list of persistent camera rigs to activate. */
-	TArray<FBlueprintPersistentCameraRigInfo> ActivePersistentCameraRigs;
-
-	/** The list of persistent camera rigs to deactivate. */
-	TArray<FBlueprintPersistentCameraRigInfo> InactivePersistentCameraRigs;
-
-	/** Whether to force a new instance of the main layer camera rigs. */
-	bool bForceActivateCameraRigs = false;
-
-	/** Reset this result for a new evaluation. */
-	void Reset();
-};
-
-}  // namespace UE::Cameras
-
-/**
- * Parameter struct for activating the Blueprint camera director evaluator.
- */
-USTRUCT(BlueprintType)
-struct FBlueprintCameraDirectorActivateParams
-{
-	GENERATED_BODY()
-
-	/** The owner (if any) of the evaluation context we are running inside of. */
-	UPROPERTY(BlueprintReadWrite, Category="Evaluation")
-	TObjectPtr<UObject> EvaluationContextOwner;
-};
-
-/**
- * Parameter struct for deactivating the Blueprint camera director evaluator.
- */
-USTRUCT(BlueprintType)
-struct FBlueprintCameraDirectorDeactivateParams
-{
-	GENERATED_BODY()
-
-	/** The owner (if any) of the evaluation context we were running inside of. */
-	UPROPERTY(BlueprintReadWrite, Category="Evaluation")
-	TObjectPtr<UObject> EvaluationContextOwner;
-};
-
-/**
- * Parameter struct for running the Blueprint camera director evaluator.
- */
-USTRUCT(BlueprintType)
-struct FBlueprintCameraDirectorEvaluationParams
-{
-	GENERATED_BODY()
-
-	/** The elapsed time since the last evaluation. */
-	UPROPERTY(BlueprintReadWrite, Category="Evaluation")
-	float DeltaTime = 0.f;
-
-	/** The owner (if any) of the evaluation context we are running inside of. */
-	UPROPERTY(BlueprintReadWrite, Category="Evaluation")
-	TObjectPtr<UObject> EvaluationContextOwner;
-};
-
 /**
  * Base class for a Blueprint camera director evaluator.
  */
@@ -111,21 +31,32 @@ public:
 	 * camera director gets activated.
 	 */
 	UFUNCTION(BlueprintCallable, BlueprintImplementableEvent, Category="Camera Director|Activation")
-	void ActivateCameraDirector(const FBlueprintCameraDirectorActivateParams& Params);
+	void ActivateCameraDirector(UObject* EvaluationContextOwner);
 
 	/**
 	 * Override this method in Blueprint to execute custom logic when this
 	 * camera director gets deactivated.
 	 */
 	UFUNCTION(BlueprintCallable, BlueprintImplementableEvent, Category="Camera Director|Activation")
-	void DeactivateCameraDirector(const FBlueprintCameraDirectorDeactivateParams& Params);
+	void DeactivateCameraDirector(UObject* EvaluationContextOwner);
 	
 	/**
 	 * Override this method in Blueprint to execute the custom logic that determines
 	 * what camera rig(s) should be active every frame.
 	 */
 	UFUNCTION(BlueprintCallable, BlueprintImplementableEvent, Category="Evaluation")
-	void RunCameraDirector(const FBlueprintCameraDirectorEvaluationParams& Params);
+	void RunCameraDirector(float DeltaTime, UObject* EvaluationContextOwner);
+
+public:
+
+	UFUNCTION(BlueprintCallable, BlueprintImplementableEvent, Category="Evaluation")
+	FName AddChildEvaluationContext(UObject* ChildEvaluationContextOwner);
+
+	UFUNCTION(BlueprintCallable, BlueprintImplementableEvent, Category="Evaluation")
+	bool RemoveChildEvaluationContext(UObject* ChildEvaluationContextOwner, FName ChildSlotName);
+
+	UFUNCTION(BlueprintCallable, Category="Evaluation")
+	bool RunChildCameraDirector(float DeltaTime, FName ChildSlotName);
 
 public:
 
@@ -164,7 +95,7 @@ public:
 	 * via the proxy table of the Blueprint camera director.
 	 */
 	UFUNCTION(BlueprintCallable, Category="Camera Director|Evaluation")
-	void ActivateCameraRigViaProxy(UCameraRigProxyAsset* CameraRigProxy);
+	void ActivateCameraRigViaProxy(UCameraRigProxyAsset* CameraRigProxy, bool bForceNewInstance = false);
 
 public:
 
@@ -207,7 +138,15 @@ public:
 public:
 
 	using FCameraEvaluationContext = UE::Cameras::FCameraEvaluationContext;
-	using FBlueprintCameraDirectorEvaluationResult = UE::Cameras::FBlueprintCameraDirectorEvaluationResult;
+	using FCameraDirectorEvaluationResult = UE::Cameras::FCameraDirectorEvaluationResult;
+
+	// Internal API.
+
+	const FCameraDirectorEvaluationResult& GetEvaluationResult() const { return EvaluationResult; }
+
+public:
+
+	void NativeInitializeCameraDirector(const UE::Cameras::FCameraDirectorInitializeParams& Params);
 
 	/** Native wrapper for ActivateCameraDirector. */
 	void NativeActivateCameraDirector(const UE::Cameras::FCameraDirectorActivateParams& Params);
@@ -218,17 +157,21 @@ public:
 	/** Native wrapper for RunCameraDirector. */
 	void NativeRunCameraDirector(const UE::Cameras::FCameraDirectorEvaluationParams& Params);
 
-	/** Get the last result for this camera director. */
-	const FBlueprintCameraDirectorEvaluationResult& GetEvaluationResult() const { return EvaluationResult; }
+	bool NativeAddChildEvaluationContext(UObject* ChildEvaluationContextOwner);
+
+	bool NativeRemoveChildEvaluationContext(UObject* ChildEvaluationContextOwner);
 
 private:
 
-	/** The current camera director evaluation result. */
-	FBlueprintCameraDirectorEvaluationResult EvaluationResult;
+	/** The current evaluation result. */
+	FCameraDirectorEvaluationResult EvaluationResult;
 
 	/** The current evaluation context. */
 	TSharedPtr<FCameraEvaluationContext> EvaluationContext;
 
+	/** Currently registered children contexts. */
+	TArray<FName> ChildrenContextSlotNames;
+
 	/** Cached world. */
 	mutable TWeakObjectPtr<UWorld> WeakCachedWorld;
 };
@@ -247,21 +190,11 @@ public:
 	UPROPERTY(EditAnywhere, Category="Evaluation")
 	TSubclassOf<UBlueprintCameraDirectorEvaluator> CameraDirectorEvaluatorClass;
 
-	/** 
-	 * The table that maps camera rig proxies (used in the evaluator Blueprint graph)
-	 * to actual camera rigs.
-	 */
-	UPROPERTY(EditAnywhere, Instanced, Category="Evaluation")
-	TObjectPtr<UCameraRigProxyTable> CameraRigProxyTable;
-
 protected:
 
 	// UCameraDirector interface.
 	virtual FCameraDirectorEvaluatorPtr OnBuildEvaluator(FCameraDirectorEvaluatorBuilder& Builder) const override;
 	virtual void OnBuildCameraDirector(UE::Cameras::FCameraBuildLog& BuildLog) override;
 	virtual void OnGatherRigUsageInfo(FCameraDirectorRigUsageInfo& UsageInfo) override;
-#if WITH_EDITOR
-	virtual void OnFactoryCreateAsset(const FCameraDirectorFactoryCreateParams& InParams) override;
-#endif
 };
 
diff --git a/Engine/Plugins/Cameras/GameplayCameras/Source/GameplayCameras/Public/Directors/StateTreeCameraDirector.h b/Engine/Plugins/Cameras/GameplayCameras/Source/GameplayCameras/Public/Directors/StateTreeCameraDirector.h
--- a/Engine/Plugins/Cameras/GameplayCameras/Source/GameplayCameras/Public/Directors/StateTreeCameraDirector.h
+++ b/Engine/Plugins/Cameras/GameplayCameras/Source/GameplayCameras/Public/Directors/StateTreeCameraDirector.h
@@ -7,9 +7,6 @@
 
 #include "StateTreeCameraDirector.generated.h"
 
-class UCameraRigProxyAsset;
-class UCameraRigProxyTable;
-
 /**
  * A camera director that runs a StateTree to specify which camera rigs should be active
  * any given frame.
@@ -29,9 +26,6 @@ protected:
 	virtual FCameraDirectorEvaluatorPtr OnBuildEvaluator(FCameraDirectorEvaluatorBuilder& Builder) const override;
 	virtual void OnBuildCameraDirector(UE::Cameras::FCameraBuildLog& BuildLog) override;
 	virtual void OnGatherRigUsageInfo(FCameraDirectorRigUsageInfo& UsageInfo) override;
-#if WITH_EDITOR
-	virtual void OnFactoryCreateAsset(const FCameraDirectorFactoryCreateParams& InParams) override;
-#endif
 
 public:
 
@@ -39,12 +33,5 @@ public:
 	UPROPERTY(EditAnywhere, Category="StateTree",
 			meta=(Schema="/Script/GameplayCameras.CameraDirectorStateTreeSchema"))
 	FStateTreeReference StateTreeReference;
-
-	/** 
-	 * The table that maps camera rig proxies (used in the evaluator State Tree's tasks)
-	 * to actual camera rigs.
-	 */
-	UPROPERTY(EditAnywhere, Instanced, Category="Evaluation")
-	TObjectPtr<UCameraRigProxyTable> CameraRigProxyTable;
 };
 
diff --git a/Engine/Plugins/Cameras/GameplayCameras/Source/GameplayCameras/Public/GameFramework/GameplayCameraComponentBase.h b/Engine/Plugins/Cameras/GameplayCameras/Source/GameplayCameras/Public/GameFramework/GameplayCameraComponentBase.h
--- a/Engine/Plugins/Cameras/GameplayCameras/Source/GameplayCameras/Public/GameFramework/GameplayCameraComponentBase.h
+++ b/Engine/Plugins/Cameras/GameplayCameras/Source/GameplayCameras/Public/GameFramework/GameplayCameraComponentBase.h
@@ -29,6 +29,20 @@ class FGameplayCameraComponentEvaluationContext;
 
 }  // namespace UE::Cameras
 
+/**
+ * Defines how to activate a gameplay camera component.
+ */
+UENUM()
+enum class EGameplayCameraComponentActivationMode : uint8
+{
+	/** Push the camera director over any existing ones. */
+	Push,
+	/** Push the camera director and try to insert any active one as a child. */
+	PushAndInsert,
+	/** Inserts the camera director as a child of the active one, or push it if there is no active one. */
+	InsertOrPush
+};
+
 /**
  * A component that can run a camera asset inside its own camera evaluation context.
  */
@@ -62,11 +76,15 @@ public:
 
 	/** Activates the camera for the given player. */
 	UFUNCTION(BlueprintCallable, Category=Camera)
-	GAMEPLAYCAMERAS_API void ActivateCameraForPlayerIndex(int32 PlayerIndex);
+	GAMEPLAYCAMERAS_API void ActivateCameraForPlayerIndex(
+			int32 PlayerIndex, 
+			EGameplayCameraComponentActivationMode ActivationMode = EGameplayCameraComponentActivationMode::Push);
 
 	/** Activates the camera for the given player. */
 	UFUNCTION(BlueprintCallable, Category=Camera)
-	GAMEPLAYCAMERAS_API void ActivateCameraForPlayerController(APlayerController* PlayerController);
+	GAMEPLAYCAMERAS_API void ActivateCameraForPlayerController(
+			APlayerController* PlayerController,
+			EGameplayCameraComponentActivationMode ActivationMode = EGameplayCameraComponentActivationMode::Push);
 
 	/** Deactivates the camera for the last player it was activated for. */
 	UFUNCTION(BlueprintCallable, Category=Camera)
@@ -139,8 +157,8 @@ protected:
 
 private:
 
-	void ActivateCameraEvaluationContext(int32 PlayerIndex);
-	void ActivateCameraEvaluationContext(APlayerController* PlayerController);
+	void ActivateCameraEvaluationContext(int32 PlayerIndex, EGameplayCameraComponentActivationMode ActivationMode);
+	void ActivateCameraEvaluationContext(APlayerController* PlayerController, EGameplayCameraComponentActivationMode ActivationMode);
 	void TryCreateCameraEvaluationContext(APlayerController* PlayerController);
 	void UpdateOutputCameraComponent();
 	void DeactivateCameraEvaluationContext();
