Cameras: Blueprint API changes

- Completely remove specific wrappers for the variable and context data tables.
- The BP wrapper for the C++ result struct is now more explicitly a "reference" to "camera data". All methods to get/set the camera pose or parameter values are called directly on this BP structure.
- The BP camera data structure can be created, which creates a whole new result under the hood. This allows setting up such a structure and then overriding or interpolating it with other instances.

diff --git a/Engine/Plugins/Cameras/GameplayCameras/Config/BaseGameplayCameras.ini b/Engine/Plugins/Cameras/GameplayCameras/Config/BaseGameplayCameras.ini
--- a/Engine/Plugins/Cameras/GameplayCameras/Config/BaseGameplayCameras.ini
+++ b/Engine/Plugins/Cameras/GameplayCameras/Config/BaseGameplayCameras.ini
@@ -21,3 +21,7 @@
 +PropertyRedirects=(OldName="/Script/GameplayCameras.CameraObjectInterface.InterfaceParameters", NewName="BlendableParameters")
 +EnumRedirects=(OldName="/Script/GameplayCameras.ECameraRigInterfaceParameterType",NewName="ECameraObjectInterfaceCameraParameterType")
 
++StructRedirects=(OldName="/Script/GameplayCameras.BlueprintCameraNodeEvaluationResult",NewName="BlueprintCameraEvaluationDataRef")
++PropertyRedirects=(OldName="GetCameraPose.Data",NewName="CameraData")
++PropertyRedirects=(OldName="SetCameraPose.Data",NewName="CameraData")
+
diff --git a/Engine/Plugins/Cameras/GameplayCameras/Source/GameplayCameras/Private/Directors/BlueprintCameraDirector.cpp b/Engine/Plugins/Cameras/GameplayCameras/Source/GameplayCameras/Private/Directors/BlueprintCameraDirector.cpp
--- a/Engine/Plugins/Cameras/GameplayCameras/Source/GameplayCameras/Private/Directors/BlueprintCameraDirector.cpp
+++ b/Engine/Plugins/Cameras/GameplayCameras/Source/GameplayCameras/Private/Directors/BlueprintCameraDirector.cpp
@@ -274,80 +274,33 @@ AActor* UBlueprintCameraDirectorEvaluator::FindEvaluationContextOwnerActor(TSubc
 	}
 }
 
-FBlueprintCameraNodeEvaluationResult UBlueprintCameraDirectorEvaluator::GetInitialContextResult() const
+FBlueprintCameraEvaluationDataRef UBlueprintCameraDirectorEvaluator::GetInitialContextResult() const
 {
 	if (EvaluationContext)
 	{
-		return FBlueprintCameraNodeEvaluationResult(&EvaluationContext->GetInitialResult());
+		return FBlueprintCameraEvaluationDataRef::MakeExternalRef(&EvaluationContext->GetInitialResult());
 	}
 	else
 	{
 		FFrame::KismetExecutionMessage(
 				TEXT("Can't access evaluation context's initial result outside of RunCameraDirector"), 
 				ELogVerbosity::Error);
-		return FBlueprintCameraNodeEvaluationResult();
+		return FBlueprintCameraEvaluationDataRef();
 	}
 }
 
-FBlueprintCameraNodeEvaluationResult UBlueprintCameraDirectorEvaluator::GetConditionalContextResult(ECameraEvaluationDataCondition Condition) const
+FBlueprintCameraEvaluationDataRef UBlueprintCameraDirectorEvaluator::GetConditionalContextResult(ECameraEvaluationDataCondition Condition) const
 {
 	if (EvaluationContext)
 	{
-		return FBlueprintCameraNodeEvaluationResult(&EvaluationContext->GetOrAddConditionalResult(Condition));
-	}
-	else
-	{
-		FFrame::KismetExecutionMessage(
-				TEXT("Can't access evaluation context's initial result outside of RunCameraDirector"), 
-				ELogVerbosity::Error);
-		return FBlueprintCameraNodeEvaluationResult();
-	}
-}
-
-FBlueprintCameraPose UBlueprintCameraDirectorEvaluator::GetInitialContextCameraPose() const
-{
-	if (EvaluationContext)
-	{
-		return FBlueprintCameraPose::FromCameraPose(EvaluationContext->GetInitialResult().CameraPose);
-	}
-	else
-	{
-		FFrame::KismetExecutionMessage(
-				TEXT("Can't access evaluation context's initial result outside of RunCameraDirector"), 
-				ELogVerbosity::Error);
-		return FBlueprintCameraPose();
-	}
-}
-
-FBlueprintCameraVariableTable UBlueprintCameraDirectorEvaluator::GetInitialContextVariableTable() const
-{
-	using namespace UE::Cameras;
-
-	if (EvaluationContext)
-	{
-		FCameraVariableTable& VariableTable = EvaluationContext->GetInitialResult().VariableTable;
-		return FBlueprintCameraVariableTable(&VariableTable);
-	}
-	else
-	{
-		FFrame::KismetExecutionMessage(
-				TEXT("Can't access evaluation context's initial result outside of RunCameraDirector"), 
-				ELogVerbosity::Error);
-		return FBlueprintCameraVariableTable();
-	}
-}
-
-void UBlueprintCameraDirectorEvaluator::SetInitialContextCameraPose(const FBlueprintCameraPose& InCameraPose)
-{
-	if (EvaluationContext)
-	{
-		InCameraPose.ApplyTo(EvaluationContext->GetInitialResult().CameraPose);
+		return FBlueprintCameraEvaluationDataRef::MakeExternalRef(&EvaluationContext->GetOrAddConditionalResult(Condition));
 	}
 	else
 	{
 		FFrame::KismetExecutionMessage(
 				TEXT("Can't access evaluation context's initial result outside of RunCameraDirector"), 
 				ELogVerbosity::Error);
+		return FBlueprintCameraEvaluationDataRef();
 	}
 }
 
diff --git a/Engine/Plugins/Cameras/GameplayCameras/Source/GameplayCameras/Private/GameFramework/BlueprintCameraContextDataTable.cpp b/Engine/Plugins/Cameras/GameplayCameras/Source/GameplayCameras/Private/GameFramework/BlueprintCameraContextDataTable.cpp
deleted file mode 100644
diff --git a/Engine/Plugins/Cameras/GameplayCameras/Source/GameplayCameras/Private/GameFramework/BlueprintCameraEvaluationDataRef.cpp b/Engine/Plugins/Cameras/GameplayCameras/Source/GameplayCameras/Private/GameFramework/BlueprintCameraEvaluationDataRef.cpp
new file mode 100644
--- /dev/null
+++ b/Engine/Plugins/Cameras/GameplayCameras/Source/GameplayCameras/Private/GameFramework/BlueprintCameraEvaluationDataRef.cpp
@@ -0,0 +1,303 @@
+// Copyright Epic Games, Inc. All Rights Reserved.
+
+#include "GameFramework/BlueprintCameraEvaluationDataRef.h"
+
+#include "Core/CameraRigAsset.h"
+#include "Core/CameraVariableAssets.h"
+#include "Core/CameraVariableTable.h"
+#include "Helpers/CameraObjectInterfaceParameterOverrideHelper.h"
+
+FBlueprintCameraEvaluationDataRef FBlueprintCameraEvaluationDataRef::MakeExternalRef(FCameraNodeEvaluationResult* InResult)
+{
+	FBlueprintCameraEvaluationDataRef Ref;
+	Ref.Result = InResult;
+	return Ref;
+}
+
+FBlueprintCameraEvaluationDataRef FBlueprintCameraEvaluationDataRef::MakeOwningRef()
+{
+	FBlueprintCameraEvaluationDataRef Ref;
+	Ref.SharedResult = MakeShared<FCameraNodeEvaluationResult>();
+	Ref.Result = Ref.SharedResult.Get();
+	return Ref;
+}
+
+FBlueprintCameraEvaluationDataRef::FBlueprintCameraEvaluationDataRef()
+{
+}
+
+FBlueprintCameraEvaluationDataRef UBlueprintCameraEvaluationDataFunctionLibrary::MakeCameraEvaluationData()
+{
+	return FBlueprintCameraEvaluationDataRef::MakeOwningRef();
+}
+
+FBlueprintCameraPose UBlueprintCameraEvaluationDataFunctionLibrary::GetCameraPose(const FBlueprintCameraEvaluationDataRef& CameraData)
+{
+	using namespace UE::Cameras;
+
+	if (const FCameraNodeEvaluationResult* Result = CameraData.GetResult())
+	{
+		return FBlueprintCameraPose::FromCameraPose(Result->CameraPose);
+	}
+
+	return FBlueprintCameraPose();
+}
+
+void UBlueprintCameraEvaluationDataFunctionLibrary::SetCameraPose(const FBlueprintCameraEvaluationDataRef& CameraData, const FBlueprintCameraPose& CameraPose)
+{
+	using namespace UE::Cameras;
+
+	if (FCameraNodeEvaluationResult* Result = CameraData.GetResult())
+	{
+		CameraPose.ApplyTo(Result->CameraPose);
+	}
+}
+
+void UBlueprintCameraEvaluationDataFunctionLibrary::BlendCameraEvaluationData(const FBlueprintCameraEvaluationDataRef& FromCameraData, const FBlueprintCameraEvaluationDataRef& ToCameraData, float Factor)
+{
+	using namespace UE::Cameras;
+
+	FCameraNodeEvaluationResult* FromResult = FromCameraData.GetResult();
+	const FCameraNodeEvaluationResult* ToResult = ToCameraData.GetResult();
+	if (FromResult && ToResult)
+	{
+		FromResult->LerpAll(*ToResult, Factor, true);
+	}
+}
+
+void UBlueprintCameraEvaluationDataFunctionLibrary::SetDefaultCameraRigParameters(const FBlueprintCameraEvaluationDataRef& CameraData, const UCameraRigAsset* CameraRig)
+{
+	using namespace UE::Cameras;
+
+	if (FCameraNodeEvaluationResult* Result = CameraData.GetResult())
+	{
+		FCameraObjectInterfaceParameterOverrideHelper::ApplyDefaultParameters(CameraRig, Result->VariableTable, Result->ContextDataTable);
+	}
+}
+
+#define UE_PRIVATE_BLUEPRINT_CAMERA_VARIABLE_TABLE_VALIDATE(ErrorResult)\
+	if (!CameraData.IsValid())\
+	{\
+		FFrame::KismetExecutionMessage(TEXT("No camera variable table has been set"), ELogVerbosity::Error);\
+		return ErrorResult;\
+	}
+
+#define UE_PRIVATE_BLUEPRINT_CAMERA_VARIABLE_TABLE_VALIDATE_PARAM(ErrorResult)\
+	if (!Variable)\
+	{\
+		FFrame::KismetExecutionMessage(TEXT("No camera variable asset was given"), ELogVerbosity::Error);\
+		return ErrorResult;\
+	}\
+
+#define UE_PRIVATE_BLUEPRINT_CAMERA_VARIABLE_TABLE_GET_VARIABLE(VariableType)\
+	static VariableType ErrorResult {};\
+	UE_PRIVATE_BLUEPRINT_CAMERA_VARIABLE_TABLE_VALIDATE(ErrorResult)\
+	UE_PRIVATE_BLUEPRINT_CAMERA_VARIABLE_TABLE_VALIDATE_PARAM(ErrorResult)\
+	return CameraData.GetResult()->VariableTable.GetValue<VariableType>(Variable->GetVariableID(), Variable->GetDefaultValue());
+
+#define UE_PRIVATE_BLUEPRINT_CAMERA_VARIABLE_TABLE_SET_VARIABLE(VariableType)\
+	UE_PRIVATE_BLUEPRINT_CAMERA_VARIABLE_TABLE_VALIDATE()\
+	UE_PRIVATE_BLUEPRINT_CAMERA_VARIABLE_TABLE_VALIDATE_PARAM()\
+	CameraData.GetResult()->VariableTable.SetValue(Variable, Value, true);
+
+bool UBlueprintCameraVariableTableFunctionLibrary::GetBooleanCameraVariable(const FBlueprintCameraEvaluationDataRef& CameraData, UBooleanCameraVariable* Variable)
+{
+	UE_PRIVATE_BLUEPRINT_CAMERA_VARIABLE_TABLE_GET_VARIABLE(bool);
+}
+
+int32 UBlueprintCameraVariableTableFunctionLibrary::GetInteger32CameraVariable(const FBlueprintCameraEvaluationDataRef& CameraData, UInteger32CameraVariable* Variable)
+{
+	UE_PRIVATE_BLUEPRINT_CAMERA_VARIABLE_TABLE_GET_VARIABLE(int32);
+}
+
+float UBlueprintCameraVariableTableFunctionLibrary::GetFloatCameraVariable(const FBlueprintCameraEvaluationDataRef& CameraData, UFloatCameraVariable* Variable)
+{
+	UE_PRIVATE_BLUEPRINT_CAMERA_VARIABLE_TABLE_GET_VARIABLE(float);
+}
+
+double UBlueprintCameraVariableTableFunctionLibrary::GetDoubleCameraVariable(const FBlueprintCameraEvaluationDataRef& CameraData, UDoubleCameraVariable* Variable)
+{
+	UE_PRIVATE_BLUEPRINT_CAMERA_VARIABLE_TABLE_GET_VARIABLE(double);
+}
+
+FVector2D UBlueprintCameraVariableTableFunctionLibrary::GetVector2CameraVariable(const FBlueprintCameraEvaluationDataRef& CameraData, UVector2dCameraVariable* Variable)
+{
+	UE_PRIVATE_BLUEPRINT_CAMERA_VARIABLE_TABLE_GET_VARIABLE(FVector2d);
+}
+
+FVector UBlueprintCameraVariableTableFunctionLibrary::GetVector3CameraVariable(const FBlueprintCameraEvaluationDataRef& CameraData, UVector3dCameraVariable* Variable)
+{
+	UE_PRIVATE_BLUEPRINT_CAMERA_VARIABLE_TABLE_GET_VARIABLE(FVector3d);
+}
+
+FVector4 UBlueprintCameraVariableTableFunctionLibrary::GetVector4CameraVariable(const FBlueprintCameraEvaluationDataRef& CameraData, UVector4dCameraVariable* Variable)
+{
+	UE_PRIVATE_BLUEPRINT_CAMERA_VARIABLE_TABLE_GET_VARIABLE(FVector4d);
+}
+
+FRotator UBlueprintCameraVariableTableFunctionLibrary::GetRotatorCameraVariable(const FBlueprintCameraEvaluationDataRef& CameraData, URotator3dCameraVariable* Variable)
+{
+	UE_PRIVATE_BLUEPRINT_CAMERA_VARIABLE_TABLE_GET_VARIABLE(FRotator3d);
+}
+
+FTransform UBlueprintCameraVariableTableFunctionLibrary::GetTransformCameraVariable(const FBlueprintCameraEvaluationDataRef& CameraData, UTransform3dCameraVariable* Variable)
+{
+	UE_PRIVATE_BLUEPRINT_CAMERA_VARIABLE_TABLE_GET_VARIABLE(FTransform3d);
+}
+
+void UBlueprintCameraVariableTableFunctionLibrary::SetBooleanCameraVariable(const FBlueprintCameraEvaluationDataRef& CameraData, UBooleanCameraVariable* Variable, bool Value)
+{
+	UE_PRIVATE_BLUEPRINT_CAMERA_VARIABLE_TABLE_SET_VARIABLE(bool);
+}
+
+void UBlueprintCameraVariableTableFunctionLibrary::SetInteger32CameraVariable(const FBlueprintCameraEvaluationDataRef& CameraData, UInteger32CameraVariable* Variable, int32 Value)
+{
+	UE_PRIVATE_BLUEPRINT_CAMERA_VARIABLE_TABLE_SET_VARIABLE(int32);
+}
+
+void UBlueprintCameraVariableTableFunctionLibrary::SetFloatCameraVariable(const FBlueprintCameraEvaluationDataRef& CameraData, UFloatCameraVariable* Variable, float Value)
+{
+	UE_PRIVATE_BLUEPRINT_CAMERA_VARIABLE_TABLE_SET_VARIABLE(float);
+}
+
+void UBlueprintCameraVariableTableFunctionLibrary::SetDoubleCameraVariable(const FBlueprintCameraEvaluationDataRef& CameraData, UDoubleCameraVariable* Variable, double Value)
+{
+	UE_PRIVATE_BLUEPRINT_CAMERA_VARIABLE_TABLE_SET_VARIABLE(double);
+}
+
+void UBlueprintCameraVariableTableFunctionLibrary::SetVector2CameraVariable(const FBlueprintCameraEvaluationDataRef& CameraData, UVector2dCameraVariable* Variable, const FVector2D& Value)
+{
+	UE_PRIVATE_BLUEPRINT_CAMERA_VARIABLE_TABLE_SET_VARIABLE(FVector2d);
+}
+
+void UBlueprintCameraVariableTableFunctionLibrary::SetVector3CameraVariable(const FBlueprintCameraEvaluationDataRef& CameraData, UVector3dCameraVariable* Variable, const FVector& Value)
+{
+	UE_PRIVATE_BLUEPRINT_CAMERA_VARIABLE_TABLE_SET_VARIABLE(FVector3d);
+}
+
+void UBlueprintCameraVariableTableFunctionLibrary::SetVector4CameraVariable(const FBlueprintCameraEvaluationDataRef& CameraData, UVector4dCameraVariable* Variable, const FVector4& Value)
+{
+	UE_PRIVATE_BLUEPRINT_CAMERA_VARIABLE_TABLE_SET_VARIABLE(FVector4d);
+}
+
+void UBlueprintCameraVariableTableFunctionLibrary::SetRotatorCameraVariable(const FBlueprintCameraEvaluationDataRef& CameraData, URotator3dCameraVariable* Variable, const FRotator& Value)
+{
+	UE_PRIVATE_BLUEPRINT_CAMERA_VARIABLE_TABLE_SET_VARIABLE(FRotator3d);
+}
+
+void UBlueprintCameraVariableTableFunctionLibrary::SetTransformCameraVariable(const FBlueprintCameraEvaluationDataRef& CameraData, UTransform3dCameraVariable* Variable, const FTransform& Value)
+{
+	UE_PRIVATE_BLUEPRINT_CAMERA_VARIABLE_TABLE_SET_VARIABLE(FTransform3d);
+}
+
+#undef UE_PRIVATE_BLUEPRINT_CAMERA_VARIABLE_TABLE_VALIDATE
+#undef UE_PRIVATE_BLUEPRINT_CAMERA_VARIABLE_TABLE_VALIDATE_PARAM
+#undef UE_PRIVATE_BLUEPRINT_CAMERA_VARIABLE_TABLE_GET_VARIABLE
+#undef UE_PRIVATE_BLUEPRINT_CAMERA_VARIABLE_TABLE_SET_VARIABLE
+
+#define UE_PRIVATE_BLUEPRINT_CAMERA_CONTEXT_DATA_TABLE_VALIDATE(ErrorResult)\
+	using namespace UE::Cameras;\
+	if (!DataID.IsValid())\
+	{\
+		FFrame::KismetExecutionMessage(TEXT("Invalid camera context data ID"), ELogVerbosity::Error);\
+		return ErrorResult;\
+	}\
+	if (!CameraData.IsValid())\
+	{\
+		FFrame::KismetExecutionMessage(TEXT("No camera context data table has been set"), ELogVerbosity::Error);\
+		return ErrorResult;\
+	}\
+
+
+FName UBlueprintCameraContextDataTableFunctionLibrary::GetNameData(const FBlueprintCameraEvaluationDataRef& CameraData, FCameraContextDataID DataID)
+{
+	UE_PRIVATE_BLUEPRINT_CAMERA_CONTEXT_DATA_TABLE_VALIDATE(NAME_None);
+	const FCameraContextDataTable& ActualTable = CameraData.GetResult()->ContextDataTable;
+	return ActualTable.GetNameData(DataID);
+}
+
+FString UBlueprintCameraContextDataTableFunctionLibrary::GetStringData(const FBlueprintCameraEvaluationDataRef& CameraData, FCameraContextDataID DataID)
+{
+	UE_PRIVATE_BLUEPRINT_CAMERA_CONTEXT_DATA_TABLE_VALIDATE(FString());
+	const FCameraContextDataTable& ActualTable = CameraData.GetResult()->ContextDataTable;
+	return ActualTable.GetStringData(DataID);
+}
+
+uint8 UBlueprintCameraContextDataTableFunctionLibrary::GetEnumData(const FBlueprintCameraEvaluationDataRef& CameraData, FCameraContextDataID DataID, const UEnum* EnumType)
+{
+	UE_PRIVATE_BLUEPRINT_CAMERA_CONTEXT_DATA_TABLE_VALIDATE(0);
+	const FCameraContextDataTable& ActualTable = CameraData.GetResult()->ContextDataTable;
+	return ActualTable.GetEnumData(DataID, EnumType);
+}
+
+FInstancedStruct UBlueprintCameraContextDataTableFunctionLibrary::GetStructData(const FBlueprintCameraEvaluationDataRef& CameraData, FCameraContextDataID DataID, const UScriptStruct* DataStructType)
+{
+	UE_PRIVATE_BLUEPRINT_CAMERA_CONTEXT_DATA_TABLE_VALIDATE(FInstancedStruct());
+	const FCameraContextDataTable& ActualTable = CameraData.GetResult()->ContextDataTable;
+	return ActualTable.GetInstancedStructData(DataID, DataStructType);
+}
+
+UObject* UBlueprintCameraContextDataTableFunctionLibrary::GetObjectData(const FBlueprintCameraEvaluationDataRef& CameraData, FCameraContextDataID DataID)
+{
+	UE_PRIVATE_BLUEPRINT_CAMERA_CONTEXT_DATA_TABLE_VALIDATE(nullptr);
+	const FCameraContextDataTable& ActualTable = CameraData.GetResult()->ContextDataTable;
+	return ActualTable.GetObjectData(DataID);
+}
+
+UClass* UBlueprintCameraContextDataTableFunctionLibrary::GetClassData(const FBlueprintCameraEvaluationDataRef& CameraData, FCameraContextDataID DataID)
+{
+	UE_PRIVATE_BLUEPRINT_CAMERA_CONTEXT_DATA_TABLE_VALIDATE(nullptr);
+	const FCameraContextDataTable& ActualTable = CameraData.GetResult()->ContextDataTable;
+	return ActualTable.GetClassData(DataID);
+}
+
+bool UBlueprintCameraContextDataTableFunctionLibrary::SetNameData(const FBlueprintCameraEvaluationDataRef& CameraData, FCameraContextDataID DataID, const FName& Data)
+{
+	UE_PRIVATE_BLUEPRINT_CAMERA_CONTEXT_DATA_TABLE_VALIDATE(false);
+	FCameraContextDataTable& ActualTable = CameraData.GetResult()->ContextDataTable;
+	ActualTable.SetNameData(DataID, Data);
+	return true;
+}
+
+bool UBlueprintCameraContextDataTableFunctionLibrary::SetStringData(const FBlueprintCameraEvaluationDataRef& CameraData, FCameraContextDataID DataID, const FString& Data)
+{
+	UE_PRIVATE_BLUEPRINT_CAMERA_CONTEXT_DATA_TABLE_VALIDATE(false);
+	FCameraContextDataTable& ActualTable = CameraData.GetResult()->ContextDataTable;
+	ActualTable.SetStringData(DataID, Data);
+	return true;
+}
+
+bool UBlueprintCameraContextDataTableFunctionLibrary::SetEnumData(const FBlueprintCameraEvaluationDataRef& CameraData, FCameraContextDataID DataID, const UEnum* EnumType, uint8 Data)
+{
+	UE_PRIVATE_BLUEPRINT_CAMERA_CONTEXT_DATA_TABLE_VALIDATE(false);
+	FCameraContextDataTable& ActualTable = CameraData.GetResult()->ContextDataTable;
+	ActualTable.SetEnumData(DataID, EnumType, Data);
+	return true;
+}
+
+bool UBlueprintCameraContextDataTableFunctionLibrary::SetStructData(const FBlueprintCameraEvaluationDataRef& CameraData, FCameraContextDataID DataID, const FInstancedStruct& Data)
+{
+	UE_PRIVATE_BLUEPRINT_CAMERA_CONTEXT_DATA_TABLE_VALIDATE(false);
+	FCameraContextDataTable& ActualTable = CameraData.GetResult()->ContextDataTable;
+	ActualTable.SetInstancedStructData(DataID, Data);
+	return true;
+}
+
+bool UBlueprintCameraContextDataTableFunctionLibrary::SetObjectData(const FBlueprintCameraEvaluationDataRef& CameraData, FCameraContextDataID DataID, UObject* Data)
+{
+	UE_PRIVATE_BLUEPRINT_CAMERA_CONTEXT_DATA_TABLE_VALIDATE(false);
+	FCameraContextDataTable& ActualTable = CameraData.GetResult()->ContextDataTable;
+	ActualTable.SetObjectData(DataID, Data);
+	return true;
+}
+
+bool UBlueprintCameraContextDataTableFunctionLibrary::SetClassData(const FBlueprintCameraEvaluationDataRef& CameraData, FCameraContextDataID DataID, UClass* Data)
+{
+	UE_PRIVATE_BLUEPRINT_CAMERA_CONTEXT_DATA_TABLE_VALIDATE(false);
+	FCameraContextDataTable& ActualTable = CameraData.GetResult()->ContextDataTable;
+	ActualTable.SetClassData(DataID, Data);
+	return true;
+}
+
+#undef UE_PRIVATE_BLUEPRINT_CAMERA_CONTEXT_DATA_TABLE_VALIDATE
+
diff --git a/Engine/Plugins/Cameras/GameplayCameras/Source/GameplayCameras/Private/GameFramework/BlueprintCameraNodeEvaluationResult.cpp b/Engine/Plugins/Cameras/GameplayCameras/Source/GameplayCameras/Private/GameFramework/BlueprintCameraNodeEvaluationResult.cpp
deleted file mode 100644
diff --git a/Engine/Plugins/Cameras/GameplayCameras/Source/GameplayCameras/Private/GameFramework/BlueprintCameraVariableTable.cpp b/Engine/Plugins/Cameras/GameplayCameras/Source/GameplayCameras/Private/GameFramework/BlueprintCameraVariableTable.cpp
deleted file mode 100644
diff --git a/Engine/Plugins/Cameras/GameplayCameras/Source/GameplayCameras/Private/GameFramework/CameraRigParameterInterop.cpp b/Engine/Plugins/Cameras/GameplayCameras/Source/GameplayCameras/Private/GameFramework/CameraRigParameterInterop.cpp
--- a/Engine/Plugins/Cameras/GameplayCameras/Source/GameplayCameras/Private/GameFramework/CameraRigParameterInterop.cpp
+++ b/Engine/Plugins/Cameras/GameplayCameras/Source/GameplayCameras/Private/GameFramework/CameraRigParameterInterop.cpp
@@ -7,10 +7,7 @@
 #include "Core/CameraRigAsset.h"
 #include "Core/CameraNodeEvaluator.h"
 #include "Core/CameraVariableTable.h"
-#include "GameFramework/BlueprintCameraContextDataTable.h"
-#include "GameFramework/BlueprintCameraNodeEvaluationResult.h"
-#include "GameFramework/BlueprintCameraVariableTable.h"
-#include "Templates/UnrealTypeTraits.h"
+#include "GameFramework/BlueprintCameraEvaluationDataRef.h"
 
 #define LOCTEXT_NAMESPACE "CameraRigParameterInterop"
 
@@ -21,19 +18,19 @@ UCameraRigParameterInterop::UCameraRigParameterInterop(const FObjectInitializer&
 {
 }
 
-void UCameraRigParameterInterop::GetCameraParameter(const FBlueprintCameraNodeEvaluationResult& Result, UCameraRigAsset* CameraRig, FName ParameterName, int32& ReturnValue)
+void UCameraRigParameterInterop::GetCameraParameter(const FBlueprintCameraEvaluationDataRef& CameraData, UCameraRigAsset* CameraRig, FName ParameterName, int32& ReturnValue)
 {
 	checkNoEntry();
 }
 
-void UCameraRigParameterInterop::SetCameraParameter(FBlueprintCameraNodeEvaluationResult& Result, UCameraRigAsset* CameraRig, FName ParameterName, const int32& NewValue)
+void UCameraRigParameterInterop::SetCameraParameter(FBlueprintCameraEvaluationDataRef& CameraData, UCameraRigAsset* CameraRig, FName ParameterName, const int32& NewValue)
 {
 	checkNoEntry();
 }
 
 DEFINE_FUNCTION(UCameraRigParameterInterop::execGetCameraParameter)
 {
-	P_GET_STRUCT_REF(FBlueprintCameraNodeEvaluationResult, Result);
+	P_GET_STRUCT_REF(FBlueprintCameraEvaluationDataRef, CameraData);
 	P_GET_OBJECT(UCameraRigAsset, CameraRig);
 	P_GET_STRUCT(FName, ParameterName);
 
@@ -75,6 +72,14 @@ DEFINE_FUNCTION(UCameraRigParameterInterop::execGetCameraParameter)
 			);
 			FBlueprintCoreDelegates::ThrowScriptException(P_THIS, Stack, ExceptionInfo);
 		}
+		else if (!CameraData.IsValid())
+		{
+			FBlueprintExceptionInfo ExceptionInfo(
+				EBlueprintExceptionType::NonFatalError,
+				LOCTEXT("InvalidCameraData", "CameraData is an invalid reference")
+			);
+			FBlueprintCoreDelegates::ThrowScriptException(P_THIS, Stack, ExceptionInfo);
+		}
 		else
 		{
 			const FInstancedPropertyBag& DefaultParameters = CameraRig->GetDefaultParameters();
@@ -82,24 +87,20 @@ DEFINE_FUNCTION(UCameraRigParameterInterop::execGetCameraParameter)
 
 			if (ParameterDefinition->ParameterType == ECameraObjectInterfaceParameterType::Blendable)
 			{
-				if (const FCameraVariableTable* VariableTable = Result.GetVariableTable().GetVariableTable())
+				const FCameraVariableTable& VariableTable = CameraData.GetResult()->VariableTable;
+				const uint8* RawValue = VariableTable.TryGetValue(ParameterDefinition->VariableID, ParameterDefinition->VariableType, ParameterDefinition->BlendableStructType);
+				if (RawValue)
 				{
-					const uint8* RawValue = VariableTable->TryGetValue(ParameterDefinition->VariableID, ParameterDefinition->VariableType, ParameterDefinition->BlendableStructType);
-					if (RawValue)
-					{
-						TargetProperty->CopyCompleteValue(TargetPtr, RawValue);
-					}
+					TargetProperty->CopyCompleteValue(TargetPtr, RawValue);
 				}
 			}
 			else if (ParameterDefinition->ParameterType == ECameraObjectInterfaceParameterType::Data)
 			{
-				if (FCameraContextDataTable* ContextDataTable = Result.GetContextDataTable().GetContextDataTable())
+				const FCameraContextDataTable& ContextDataTable = CameraData.GetResult()->ContextDataTable;
+				const uint8* RawValue = ContextDataTable.TryGetRawDataPtr(ParameterDefinition->DataID, ParameterDefinition->DataType, ParameterDefinition->DataTypeObject);
+				if (RawValue)
 				{
-					const uint8* RawValue = ContextDataTable->TryGetRawDataPtr(ParameterDefinition->DataID, ParameterDefinition->DataType, ParameterDefinition->DataTypeObject);
-					if (RawValue)
-					{
-						TargetProperty->CopyCompleteValue(TargetPtr, RawValue);
-					}
+					TargetProperty->CopyCompleteValue(TargetPtr, RawValue);
 				}
 			}
 		}
@@ -110,7 +111,7 @@ DEFINE_FUNCTION(UCameraRigParameterInterop::execGetCameraParameter)
 
 DEFINE_FUNCTION(UCameraRigParameterInterop::execSetCameraParameter)
 {
-	P_GET_STRUCT_REF(FBlueprintCameraNodeEvaluationResult, Result);
+	P_GET_STRUCT_REF(FBlueprintCameraEvaluationDataRef, CameraData);
 	P_GET_OBJECT(UCameraRigAsset, CameraRig);
 	P_GET_STRUCT(FName, ParameterName);
 
@@ -152,6 +153,14 @@ DEFINE_FUNCTION(UCameraRigParameterInterop::execSetCameraParameter)
 			);
 			FBlueprintCoreDelegates::ThrowScriptException(P_THIS, Stack, ExceptionInfo);
 		}
+		else if (!CameraData.IsValid())
+		{
+			FBlueprintExceptionInfo ExceptionInfo(
+				EBlueprintExceptionType::NonFatalError,
+				LOCTEXT("InvalidCameraData", "CameraData is an invalid reference")
+			);
+			FBlueprintCoreDelegates::ThrowScriptException(P_THIS, Stack, ExceptionInfo);
+		}
 		else
 		{
 			const FInstancedPropertyBag& DefaultParameters = CameraRig->GetDefaultParameters();
@@ -159,20 +168,16 @@ DEFINE_FUNCTION(UCameraRigParameterInterop::execSetCameraParameter)
 
 			if (ParameterDefinition->ParameterType == ECameraObjectInterfaceParameterType::Blendable)
 			{
-				if (FCameraVariableTable* VariableTable = Result.GetVariableTable().GetVariableTable())
-				{
-					VariableTable->TrySetValue(ParameterDefinition->VariableID, ParameterDefinition->VariableType, ParameterDefinition->BlendableStructType, SourcePtr);
-				}
+				FCameraVariableTable& VariableTable = CameraData.GetResult()->VariableTable;
+				VariableTable.TrySetValue(ParameterDefinition->VariableID, ParameterDefinition->VariableType, ParameterDefinition->BlendableStructType, SourcePtr);
 			}
 			else if (ParameterDefinition->ParameterType == ECameraObjectInterfaceParameterType::Data)
 			{
-				if (FCameraContextDataTable* ContextDataTable = Result.GetContextDataTable().GetContextDataTable())
+				FCameraContextDataTable& ContextDataTable = CameraData.GetResult()->ContextDataTable;
+				uint8* RawValue = ContextDataTable.TryGetMutableRawDataPtr(ParameterDefinition->DataID, ParameterDefinition->DataType, ParameterDefinition->DataTypeObject);
+				if (RawValue)
 				{
-					uint8* RawValue = ContextDataTable->TryGetMutableRawDataPtr(ParameterDefinition->DataID, ParameterDefinition->DataType, ParameterDefinition->DataTypeObject);
-					if (RawValue)
-					{
-						SourceProperty->CopyCompleteValue(RawValue, SourcePtr);
-					}
+					SourceProperty->CopyCompleteValue(RawValue, SourcePtr);
 				}
 			}
 		}
diff --git a/Engine/Plugins/Cameras/GameplayCameras/Source/GameplayCameras/Private/GameFramework/GameplayCameraComponentBase.cpp b/Engine/Plugins/Cameras/GameplayCameras/Source/GameplayCameras/Private/GameFramework/GameplayCameraComponentBase.cpp
--- a/Engine/Plugins/Cameras/GameplayCameras/Source/GameplayCameras/Private/GameFramework/GameplayCameraComponentBase.cpp
+++ b/Engine/Plugins/Cameras/GameplayCameras/Source/GameplayCameras/Private/GameFramework/GameplayCameraComponentBase.cpp
@@ -275,44 +275,18 @@ void UGameplayCameraComponentBase::TryCreateCameraEvaluationContext(APlayerContr
 		return ErrorResult;\
 	}
 
-FBlueprintCameraNodeEvaluationResult UGameplayCameraComponentBase::GetInitialResult() const
+FBlueprintCameraEvaluationDataRef UGameplayCameraComponentBase::GetInitialResult() const
 {
-	UE_PRIVATE_GAMEPLAY_CAMERA_COMPONENT_VALIDATE_EVALUATION_CONTEXT("Can't get shared camera data", FBlueprintCameraNodeEvaluationResult());
+	UE_PRIVATE_GAMEPLAY_CAMERA_COMPONENT_VALIDATE_EVALUATION_CONTEXT("Can't get shared camera data", FBlueprintCameraEvaluationDataRef());
 
-	return FBlueprintCameraNodeEvaluationResult(&EvaluationContext->GetInitialResult());
+	return FBlueprintCameraEvaluationDataRef::MakeExternalRef(&EvaluationContext->GetInitialResult());
 }
 
-FBlueprintCameraNodeEvaluationResult UGameplayCameraComponentBase::GetConditionalResult(ECameraEvaluationDataCondition Condition) const
+FBlueprintCameraEvaluationDataRef UGameplayCameraComponentBase::GetConditionalResult(ECameraEvaluationDataCondition Condition) const
 {
-	UE_PRIVATE_GAMEPLAY_CAMERA_COMPONENT_VALIDATE_EVALUATION_CONTEXT("Can't get conditional camera data", FBlueprintCameraNodeEvaluationResult());
+	UE_PRIVATE_GAMEPLAY_CAMERA_COMPONENT_VALIDATE_EVALUATION_CONTEXT("Can't get conditional camera data", FBlueprintCameraEvaluationDataRef());
 
-	return FBlueprintCameraNodeEvaluationResult(&EvaluationContext->GetOrAddConditionalResult(Condition));
-}
-
-FBlueprintCameraPose UGameplayCameraComponentBase::GetInitialPose() const
-{
-	UE_PRIVATE_GAMEPLAY_CAMERA_COMPONENT_VALIDATE_EVALUATION_CONTEXT("Can't get initial camera pose", FBlueprintCameraPose());
-
-	return FBlueprintCameraPose::FromCameraPose(EvaluationContext->GetInitialResult().CameraPose);
-}
-
-bool UGameplayCameraComponentBase::SetInitialPose(const FBlueprintCameraPose& CameraPose)
-{
-	UE_PRIVATE_GAMEPLAY_CAMERA_COMPONENT_VALIDATE_EVALUATION_CONTEXT("Can't set initial camera pose", false);
-
-	FCameraPose InitialPose = EvaluationContext->GetInitialResult().CameraPose;
-	CameraPose.ApplyTo(InitialPose);
-	return true;
-}
-
-FBlueprintCameraVariableTable UGameplayCameraComponentBase::GetInitialVariableTable() const
-{
-	using namespace UE::Cameras;
-
-	UE_PRIVATE_GAMEPLAY_CAMERA_COMPONENT_VALIDATE_EVALUATION_CONTEXT("Can't get initial camera variable table", FBlueprintCameraVariableTable());
-
-	FCameraVariableTable& VariableTable = EvaluationContext->GetInitialResult().VariableTable;
-	return FBlueprintCameraVariableTable(&VariableTable);
+	return FBlueprintCameraEvaluationDataRef::MakeExternalRef(&EvaluationContext->GetOrAddConditionalResult(Condition));
 }
 
 #undef UE_PRIVATE_GAMEPLAY_CAMERA_COMPONENT_VALIDATE_EVALUATION_CONTEXT
diff --git a/Engine/Plugins/Cameras/GameplayCameras/Source/GameplayCameras/Private/Helpers/CameraObjectInterfaceParameterOverrideHelper.cpp b/Engine/Plugins/Cameras/GameplayCameras/Source/GameplayCameras/Private/Helpers/CameraObjectInterfaceParameterOverrideHelper.cpp
--- a/Engine/Plugins/Cameras/GameplayCameras/Source/GameplayCameras/Private/Helpers/CameraObjectInterfaceParameterOverrideHelper.cpp
+++ b/Engine/Plugins/Cameras/GameplayCameras/Source/GameplayCameras/Private/Helpers/CameraObjectInterfaceParameterOverrideHelper.cpp
@@ -453,6 +453,16 @@ void FCameraObjectInterfaceParameterOverrideHelper::ApplyParameterOverride(
 }
 
 void FCameraObjectInterfaceParameterOverrideHelper::ApplyDefaultBlendableParameters(const UBaseCameraObject* CameraObject, FCameraVariableTable& OutVariableTable)
+{
+	ApplyDefaultParametersImpl(CameraObject, &OutVariableTable, nullptr);
+}
+
+void FCameraObjectInterfaceParameterOverrideHelper::ApplyDefaultParameters(const UBaseCameraObject* CameraObject, FCameraVariableTable& OutVariableTable, FCameraContextDataTable& OutContextDataTable)
+{
+	ApplyDefaultParametersImpl(CameraObject, &OutVariableTable, &OutContextDataTable);
+}
+
+void FCameraObjectInterfaceParameterOverrideHelper::ApplyDefaultParametersImpl(const UBaseCameraObject* CameraObject, FCameraVariableTable* OutVariableTable, FCameraContextDataTable* OutContextDataTable)
 {
 	if (!ensure(CameraObject))
 	{
@@ -464,23 +474,42 @@ void FCameraObjectInterfaceParameterOverrideHelper::ApplyDefaultBlendableParamet
 
 	for (const FCameraObjectInterfaceParameterDefinition& Definition : CameraObject->GetParameterDefinitions())
 	{
-		if (Definition.ParameterType != ECameraObjectInterfaceParameterType::Blendable)
+		if (Definition.ParameterType == ECameraObjectInterfaceParameterType::Blendable)
 		{
-			continue;
+			if (!OutVariableTable || !Definition.VariableID.IsValid())
+			{
+				continue;
+			}
+
+			const FPropertyBagPropertyDesc* PropertyDesc = DefaultParameters.FindPropertyDescByID(Definition.ParameterGuid);
+			if (!ensure(PropertyDesc && PropertyDesc->CachedProperty))
+			{
+				continue;
+			}
+
+			const void* RawValuePtr = PropertyDesc->CachedProperty->ContainerPtrToValuePtr<void>(RawDefaultParametersContainer);
+			OutVariableTable->TrySetValue(Definition.VariableID, Definition.VariableType, Definition.BlendableStructType, (const uint8*)RawValuePtr);
 		}
-		if (!Definition.VariableID.IsValid())
+		else if (Definition.ParameterType == ECameraObjectInterfaceParameterType::Data)
 		{
-			continue;
-		}
+			if (!OutContextDataTable || !Definition.DataID.IsValid())
+			{
+				continue;
+			}
 
-		const FPropertyBagPropertyDesc* PropertyDesc = DefaultParameters.FindPropertyDescByID(Definition.ParameterGuid);
-		if (!ensure(PropertyDesc && PropertyDesc->CachedProperty))
-		{
-			continue;
-		}
+			const FPropertyBagPropertyDesc* PropertyDesc = DefaultParameters.FindPropertyDescByID(Definition.ParameterGuid);
+			if (!ensure(PropertyDesc && PropertyDesc->CachedProperty))
+			{
+				continue;
+			}
 
-		const void* RawValuePtr = PropertyDesc->CachedProperty->ContainerPtrToValuePtr<void>(RawDefaultParametersContainer);
-		OutVariableTable.SetValue(Definition.VariableID, Definition.VariableType, Definition.BlendableStructType, (const uint8*)RawValuePtr);
+			uint8* RawDestPtr = OutContextDataTable->TryGetMutableRawDataPtr(Definition.DataID, Definition.DataType, Definition.DataTypeObject);
+			if (RawDestPtr)
+			{
+				const void* RawValuePtr = PropertyDesc->CachedProperty->ContainerPtrToValuePtr<void>(RawDefaultParametersContainer);
+				PropertyDesc->CachedProperty->CopyCompleteValue(RawDestPtr, RawValuePtr);
+			}
+		}
 	}
 }
 
diff --git a/Engine/Plugins/Cameras/GameplayCameras/Source/GameplayCameras/Private/Helpers/CameraObjectInterfaceParameterOverrideHelper.h b/Engine/Plugins/Cameras/GameplayCameras/Source/GameplayCameras/Private/Helpers/CameraObjectInterfaceParameterOverrideHelper.h
--- a/Engine/Plugins/Cameras/GameplayCameras/Source/GameplayCameras/Private/Helpers/CameraObjectInterfaceParameterOverrideHelper.h
+++ b/Engine/Plugins/Cameras/GameplayCameras/Source/GameplayCameras/Private/Helpers/CameraObjectInterfaceParameterOverrideHelper.h
@@ -25,6 +25,13 @@ public:
 	/** Sets default values of blendable interface parameters in the given variable table. */
 	static void ApplyDefaultBlendableParameters(const UBaseCameraObject* CameraObject, FCameraVariableTable& OutVariableTable);
 
+	/** Sets default values of interface parameters in the given variable and context data tables. */
+	static void ApplyDefaultParameters(const UBaseCameraObject* CameraObject, FCameraVariableTable& OutVariableTable, FCameraContextDataTable& OutContextDataTable);
+
+private:
+
+	static void ApplyDefaultParametersImpl(const UBaseCameraObject* CameraObject, FCameraVariableTable* OutVariableTable, FCameraContextDataTable* OutContextDataTable);
+
 public:
 
 	/** 
diff --git a/Engine/Plugins/Cameras/GameplayCameras/Source/GameplayCameras/Private/Nodes/Utility/BlueprintCameraNode.cpp b/Engine/Plugins/Cameras/GameplayCameras/Source/GameplayCameras/Private/Nodes/Utility/BlueprintCameraNode.cpp
--- a/Engine/Plugins/Cameras/GameplayCameras/Source/GameplayCameras/Private/Nodes/Utility/BlueprintCameraNode.cpp
+++ b/Engine/Plugins/Cameras/GameplayCameras/Source/GameplayCameras/Private/Nodes/Utility/BlueprintCameraNode.cpp
@@ -8,6 +8,7 @@
 #include "Core/CameraContextDataTableFwd.h"
 #include "Core/CameraEvaluationContext.h"
 #include "Core/CameraNodeEvaluator.h"
+#include "Core/CameraRigAsset.h"
 #include "Core/CameraSystemEvaluator.h"
 #include "Core/CameraVariableAssets.h"
 #include "Core/CameraVariableTable.h"
@@ -17,6 +18,7 @@
 #include "Engine/World.h"
 #include "GameFramework/Actor.h"
 #include "GameplayCameras.h"
+#include "Helpers/CameraObjectInterfaceParameterOverrideHelper.h"
 #include "Misc/AssertionMacros.h"
 #include "UObject/Object.h"
 #include "UObject/ObjectMacros.h"
@@ -86,7 +88,7 @@ void FBlueprintCameraNodeEvaluator::OnInitialize(const FCameraNodeEvaluatorIniti
 
 		ApplyParameterOverrides(OutResult.VariableTable, OutResult.ContextDataTable);
 
-		EvaluatorBlueprint->NativeInitializeCameraNode(Params, OutResult);
+		EvaluatorBlueprint->NativeInitializeCameraNode(BlueprintNode, Params, OutResult);
 	}
 	else
 	{
@@ -175,17 +177,18 @@ void FBlueprintCameraDebugBlock::OnDebugDraw(const FCameraDebugBlockDrawParams&
 
 }  // namespace UE::Cameras
 
-void UBlueprintCameraNodeEvaluator::NativeInitializeCameraNode(const FCameraNodeEvaluatorInitializeParams& Params, FCameraNodeEvaluationResult& OutResult)
+void UBlueprintCameraNodeEvaluator::NativeInitializeCameraNode(const UBlueprintCameraNode* InBlueprintNode, const FCameraNodeEvaluatorInitializeParams& Params, FCameraNodeEvaluationResult& OutResult)
 {
 	using namespace UE::Cameras;
 
+	ensure(BlueprintNode == nullptr);
+	BlueprintNode = InBlueprintNode;
+
 	SetupExecution(Params.EvaluationContext, OutResult);
 	{
 		bIsFirstFrame = true;
 
 		InitializeCameraNode();
-
-		CameraPose.ApplyTo(OutResult.CameraPose);
 	}
 	TeardownExecution();
 }
@@ -199,8 +202,6 @@ void UBlueprintCameraNodeEvaluator::NativeRunCameraNode(const FCameraNodeEvaluat
 		bIsFirstFrame = Params.bIsFirstFrame;
 
 		TickCameraNode(Params.DeltaTime);
-
-		CameraPose.ApplyTo(OutResult.CameraPose);
 	}
 	TeardownExecution();
 }
@@ -208,19 +209,18 @@ void UBlueprintCameraNodeEvaluator::NativeRunCameraNode(const FCameraNodeEvaluat
 void UBlueprintCameraNodeEvaluator::SetupExecution(TSharedPtr<const FCameraEvaluationContext> EvaluationContext, FCameraNodeEvaluationResult& OutResult)
 {
 	EvaluationContextOwner = EvaluationContext->GetOwner();
-	CameraPose = FBlueprintCameraPose::FromCameraPose(OutResult.CameraPose);
+
+	ensure(!CameraData.IsValid());
+	CameraData = FBlueprintCameraEvaluationDataRef::MakeExternalRef(&OutResult);
 
 	ensure(!CurrentContext.IsValid());
 	CurrentContext = EvaluationContext;
-
-	ensure(CurrentResult == nullptr);
-	CurrentResult = &OutResult;
 }
 
 void UBlueprintCameraNodeEvaluator::TeardownExecution()
 {
+	CameraData = FBlueprintCameraEvaluationDataRef();
 	CurrentContext = nullptr;
-	CurrentResult = nullptr;
 }
 
 AActor* UBlueprintCameraNodeEvaluator::FindEvaluationContextOwnerActor(TSubclassOf<AActor> ActorClass) const
@@ -249,6 +249,27 @@ AActor* UBlueprintCameraNodeEvaluator::FindEvaluationContextOwnerActor(TSubclass
 	}
 }
 
+FBlueprintCameraPose UBlueprintCameraNodeEvaluator::GetCurrentCameraPose() const
+{
+	return UBlueprintCameraEvaluationDataFunctionLibrary::GetCameraPose(CameraData);
+}
+
+void UBlueprintCameraNodeEvaluator::SetCurrentCameraPose(const FBlueprintCameraPose& CameraPose)
+{
+	UBlueprintCameraEvaluationDataFunctionLibrary::SetCameraPose(CameraData, CameraPose);
+}
+
+void UBlueprintCameraNodeEvaluator::SetDefaultOwningCameraRigParameters(FBlueprintCameraEvaluationDataRef TargetCameraData) const
+{
+	using namespace UE::Cameras;
+
+	if (FCameraNodeEvaluationResult* Result = TargetCameraData.GetResult())
+	{
+		const UCameraRigAsset* OwningCameraRig = BlueprintNode->GetTypedOuter<UCameraRigAsset>();
+		FCameraObjectInterfaceParameterOverrideHelper::ApplyDefaultParameters(OwningCameraRig, Result->VariableTable, Result->ContextDataTable);
+	}
+}
+
 UWorld* UBlueprintCameraNodeEvaluator::GetWorld() const
 {
 	if (UWorld* CachedWorld = WeakCachedWorld.Get())
diff --git a/Engine/Plugins/Cameras/GameplayCameras/Source/GameplayCameras/Public/Directors/BlueprintCameraDirector.h b/Engine/Plugins/Cameras/GameplayCameras/Source/GameplayCameras/Public/Directors/BlueprintCameraDirector.h
--- a/Engine/Plugins/Cameras/GameplayCameras/Source/GameplayCameras/Public/Directors/BlueprintCameraDirector.h
+++ b/Engine/Plugins/Cameras/GameplayCameras/Source/GameplayCameras/Public/Directors/BlueprintCameraDirector.h
@@ -5,9 +5,7 @@
 #include "Core/CameraDirector.h"
 #include "Core/CameraDirectorEvaluator.h"
 #include "Core/CameraEvaluationContext.h"
-#include "GameFramework/BlueprintCameraNodeEvaluationResult.h"
-#include "GameFramework/BlueprintCameraPose.h"
-#include "GameFramework/BlueprintCameraVariableTable.h"
+#include "GameFramework/BlueprintCameraEvaluationDataRef.h"
 #include "Templates/SubclassOf.h"
 
 #include "BlueprintCameraDirector.generated.h"
@@ -109,23 +107,11 @@ public:
 
 	/** Gets the shared evaluation context data. */
 	UFUNCTION(BlueprintPure, Category="Camera Director|Evaluation", meta=(DisplayName="Get Shared Camera Data"))
-	FBlueprintCameraNodeEvaluationResult GetInitialContextResult() const;
+	FBlueprintCameraEvaluationDataRef GetInitialContextResult() const;
 
 	/** Gets the evaluation context data for a sub-set of camera rigs. */	
 	UFUNCTION(BlueprintPure, Category="Camera Director|Evaluation", meta=(DisplayName="Get Conditional Camera Data"))
-	FBlueprintCameraNodeEvaluationResult GetConditionalContextResult(ECameraEvaluationDataCondition Condition) const;
-
-	UFUNCTION(BlueprintPure, Category="Camera Director|Evaluation",
-			meta=(DeprecatedFunction, DeprecationMessage="Please use GetSharedEvaluationData"))
-	FBlueprintCameraPose GetInitialContextCameraPose() const;
-
-	UFUNCTION(BlueprintCallable, Category="Camera Director|Evaluation",
-			meta=(DeprecatedFunction, DeprecationMessage="Please use GetSharedEvaluationData"))
-	void SetInitialContextCameraPose(const FBlueprintCameraPose& InCameraPose);
-
-	UFUNCTION(BlueprintPure, Category="Camera Director|Evaluation", 
-			meta=(DeprecatedFunction, DeprecationMessage="Please use GetSharedEvaluationData"))
-	FBlueprintCameraVariableTable GetInitialContextVariableTable() const;
+	FBlueprintCameraEvaluationDataRef GetConditionalContextResult(ECameraEvaluationDataCondition Condition) const;
 
 public:
 
diff --git a/Engine/Plugins/Cameras/GameplayCameras/Source/GameplayCameras/Public/GameFramework/BlueprintCameraContextDataTable.h b/Engine/Plugins/Cameras/GameplayCameras/Source/GameplayCameras/Public/GameFramework/BlueprintCameraContextDataTable.h
deleted file mode 100644
diff --git a/Engine/Plugins/Cameras/GameplayCameras/Source/GameplayCameras/Public/GameFramework/BlueprintCameraEvaluationDataRef.h b/Engine/Plugins/Cameras/GameplayCameras/Source/GameplayCameras/Public/GameFramework/BlueprintCameraEvaluationDataRef.h
new file mode 100644
--- /dev/null
+++ b/Engine/Plugins/Cameras/GameplayCameras/Source/GameplayCameras/Public/GameFramework/BlueprintCameraEvaluationDataRef.h
@@ -0,0 +1,234 @@
+// Copyright Epic Games, Inc. All Rights Reserved.
+
+#pragma once
+
+#include "Core/CameraNodeEvaluator.h"
+#include "Core/CameraVariableAssets.h"
+#include "GameFramework/BlueprintCameraPose.h"
+#include "Kismet/BlueprintFunctionLibrary.h"
+
+#include "BlueprintCameraEvaluationDataRef.generated.h"
+
+class UCameraRigAsset;
+
+namespace UE::Cameras
+{
+
+struct FCameraNodeEvaluationResult;
+
+}  // namespace UE::Cameras
+
+/**
+ * Blueprint wrapper for camera evaluation data.
+ */
+USTRUCT(BlueprintType, DisplayName="Camera Evaluation Data Ref", 
+		meta=(HasNativeMake="BlueprintCameraEvaluationDataFunctionLibrary.MakeCameraEvaluationData"))
+struct GAMEPLAYCAMERAS_API FBlueprintCameraEvaluationDataRef
+{
+	GENERATED_BODY()
+
+public:
+
+	using FCameraNodeEvaluationResult = UE::Cameras::FCameraNodeEvaluationResult;
+
+	FBlueprintCameraEvaluationDataRef();
+
+	bool IsValid() const { return Result != nullptr; }
+
+	FCameraNodeEvaluationResult* GetResult() const { return Result; }
+
+public:
+
+	static FBlueprintCameraEvaluationDataRef MakeExternalRef(FCameraNodeEvaluationResult* InResult);
+	static FBlueprintCameraEvaluationDataRef MakeOwningRef();
+
+private:
+
+	/** The underlying camera evaluation result. */
+	FCameraNodeEvaluationResult* Result = nullptr;
+	
+	/** If the underlying camera evaluation result is owned by this data ref, a shared pointer to keep it alive. */
+	TSharedPtr<FCameraNodeEvaluationResult> SharedResult;
+};
+
+/**
+ * Blueprint function library for camera evaluation data references.
+ */
+UCLASS()
+class UBlueprintCameraEvaluationDataFunctionLibrary : public UBlueprintFunctionLibrary
+{
+	GENERATED_BODY()
+
+public:
+
+	/** Creates a camera evaluation data. */
+	UFUNCTION(BlueprintPure, Category="Camera", meta=(ReturnDisplayName="Camera Evaluation Data"))
+	static FBlueprintCameraEvaluationDataRef MakeCameraEvaluationData();
+
+public:
+
+	/** Gets the camera pose. */
+	UFUNCTION(BlueprintPure, Category="Camera")
+	static FBlueprintCameraPose GetCameraPose(const FBlueprintCameraEvaluationDataRef& CameraData);
+
+	/** Sets the camera pose. */
+	UFUNCTION(BlueprintCallable, Category="Camera")
+	static void SetCameraPose(const FBlueprintCameraEvaluationDataRef& CameraData, const FBlueprintCameraPose& CameraPose);
+
+public:
+
+	/** Interpolates one camera data towards another. */
+	UFUNCTION(BlueprintCallable, Category="Camera")
+	static void BlendCameraEvaluationData(const FBlueprintCameraEvaluationDataRef& FromCameraData, const FBlueprintCameraEvaluationDataRef& ToCameraData, float Factor);
+
+	/** Sets the default values for all parameters in the given camera rig. */
+	UFUNCTION(BlueprintCallable, Category=Camera)
+	static void SetDefaultCameraRigParameters(const FBlueprintCameraEvaluationDataRef& CameraData, const UCameraRigAsset* CameraRig);
+};
+
+/**
+ * Blueprint function library for camera variable tables.
+ */
+UCLASS()
+class UBlueprintCameraVariableTableFunctionLibrary : public UBlueprintFunctionLibrary
+{
+	GENERATED_BODY()
+
+public:
+
+	/** Gets a camera variable's value from the given table. */
+	UFUNCTION(BlueprintCallable, DisplayName="Get Boolean Variable", Category=Camera, meta=(CompactNodeTitle="Get Boolean Variable"))
+	static bool GetBooleanCameraVariable(const FBlueprintCameraEvaluationDataRef& CameraData, UBooleanCameraVariable* Variable);
+
+	/** Gets a camera variable's value from the given table. */
+	UFUNCTION(BlueprintCallable, DisplayName="Get Integer Variable", Category=Camera, meta=(CompactNodeTitle="Get Integer Variable"))
+	static int32 GetInteger32CameraVariable(const FBlueprintCameraEvaluationDataRef& CameraData, UInteger32CameraVariable* Variable);
+
+	/** Gets a camera variable's value from the given table. */
+	UFUNCTION(BlueprintCallable, DisplayName="Get Float Variable", Category=Camera, meta=(CompactNodeTitle="Get Float Variable"))
+	static float GetFloatCameraVariable(const FBlueprintCameraEvaluationDataRef& CameraData, UFloatCameraVariable* Variable);
+
+	/** Gets a camera variable's value from the given table. */
+	UFUNCTION(BlueprintCallable, DisplayName="Get Double Variable", Category=Camera, meta=(CompactNodeTitle="Get Double Variable"))
+	static double GetDoubleCameraVariable(const FBlueprintCameraEvaluationDataRef& CameraData, UDoubleCameraVariable* Variable);
+
+	/** Gets a camera variable's value from the given table. */
+	UFUNCTION(BlueprintCallable, DisplayName="Get Vector 2D Variable", Category=Camera, meta=(CompactNodeTitle="Get Vector 2D Variable"))
+	static FVector2D GetVector2CameraVariable(const FBlueprintCameraEvaluationDataRef& CameraData, UVector2dCameraVariable* Variable);
+
+	/** Gets a camera variable's value from the given table. */
+	UFUNCTION(BlueprintCallable, DisplayName="Get Vector Variable", Category=Camera, meta=(CompactNodeTitle="Get Vector Variable"))
+	static FVector GetVector3CameraVariable(const FBlueprintCameraEvaluationDataRef& CameraData, UVector3dCameraVariable* Variable);
+
+	/** Gets a camera variable's value from the given table. */
+	UFUNCTION(BlueprintCallable, DisplayName="Get Vector 4 Variable", Category=Camera, meta=(CompactNodeTitle="Get Vector 4 Variable"))
+	static FVector4 GetVector4CameraVariable(const FBlueprintCameraEvaluationDataRef& CameraData, UVector4dCameraVariable* Variable);
+
+	/** Gets a camera variable's value from the given table. */
+	UFUNCTION(BlueprintCallable, DisplayName="Get Rotator Variable", Category=Camera, meta=(CompactNodeTitle="Get Rotator Variable"))
+	static FRotator GetRotatorCameraVariable(const FBlueprintCameraEvaluationDataRef& CameraData, URotator3dCameraVariable* Variable);
+
+	/** Gets a camera variable's value from the given table. */
+	UFUNCTION(BlueprintCallable, DisplayName="Get Transform Variable", Category=Camera, meta=(CompactNodeTitle="Get Transform Variable"))
+	static FTransform GetTransformCameraVariable(const FBlueprintCameraEvaluationDataRef& CameraData, UTransform3dCameraVariable* Variable);
+
+public:
+
+	/** Sets a camera variable's value in the given table. */
+	UFUNCTION(BlueprintCallable, DisplayName="Set Boolean Variable", Category=Camera)
+	static void SetBooleanCameraVariable(const FBlueprintCameraEvaluationDataRef& CameraData, UBooleanCameraVariable* Variable, bool Value);
+
+	/** Sets a camera variable's value in the given table. */
+	UFUNCTION(BlueprintCallable, DisplayName="Set Integer Variable", Category=Camera)
+	static void SetInteger32CameraVariable(const FBlueprintCameraEvaluationDataRef& CameraData, UInteger32CameraVariable* Variable, int32 Value);
+
+	/** Sets a camera variable's value in the given table. */
+	UFUNCTION(BlueprintCallable, DisplayName="Set Float Variable", Category=Camera)
+	static void SetFloatCameraVariable(const FBlueprintCameraEvaluationDataRef& CameraData, UFloatCameraVariable* Variable, float Value);
+
+	/** Sets a camera variable's value in the given table. */
+	UFUNCTION(BlueprintCallable, DisplayName="Set Double Variable", Category=Camera)
+	static void SetDoubleCameraVariable(const FBlueprintCameraEvaluationDataRef& CameraData, UDoubleCameraVariable* Variable, double Value);
+
+	/** Sets a camera variable's value in the given table. */
+	UFUNCTION(BlueprintCallable, DisplayName="Set Vector 2D Variable", Category=Camera)
+	static void SetVector2CameraVariable(const FBlueprintCameraEvaluationDataRef& CameraData, UVector2dCameraVariable* Variable, const FVector2D& Value);
+
+	/** Sets a camera variable's value in the given table. */
+	UFUNCTION(BlueprintCallable, DisplayName="Set Vector Variable", Category=Camera)
+	static void SetVector3CameraVariable(const FBlueprintCameraEvaluationDataRef& CameraData, UVector3dCameraVariable* Variable, const FVector& Value);
+
+	/** Sets a camera variable's value in the given table. */
+	UFUNCTION(BlueprintCallable, DisplayName="Set Vector 4 Variable", Category=Camera)
+	static void SetVector4CameraVariable(const FBlueprintCameraEvaluationDataRef& CameraData, UVector4dCameraVariable* Variable, const FVector4& Value);
+
+	/** Sets a camera variable's value in the given table. */
+	UFUNCTION(BlueprintCallable, DisplayName="Set Rotator Variable", Category=Camera)
+	static void SetRotatorCameraVariable(const FBlueprintCameraEvaluationDataRef& CameraData, URotator3dCameraVariable* Variable, const FRotator& Value);
+
+	/** Sets a camera variable's value in the given table. */
+	UFUNCTION(BlueprintCallable, DisplayName="Set Transform Variable", Category=Camera)
+	static void SetTransformCameraVariable(const FBlueprintCameraEvaluationDataRef& CameraData, UTransform3dCameraVariable* Variable, const FTransform& Value);
+};
+
+/**
+ * Utility Blueprint functions for camera context data tables.
+ */
+UCLASS()
+class UBlueprintCameraContextDataTableFunctionLibrary : public UBlueprintFunctionLibrary
+{
+	GENERATED_BODY()
+
+public:
+
+	/** Gets a value from the given camera context data table. */
+	UFUNCTION(BlueprintCallable, Category=Camera)
+	static FName GetNameData(const FBlueprintCameraEvaluationDataRef& CameraData, FCameraContextDataID DataID);
+
+	/** Gets a value from the given camera context data table. */
+	UFUNCTION(BlueprintCallable, Category=Camera)
+	static FString GetStringData(const FBlueprintCameraEvaluationDataRef& CameraData, FCameraContextDataID DataID);
+
+	/** Gets a value from the given camera context data table. */
+	UFUNCTION(BlueprintCallable, Category=Camera, meta=(DeterminesOutputType="EnumType"))
+	static uint8 GetEnumData(const FBlueprintCameraEvaluationDataRef& CameraData, FCameraContextDataID DataID, const UEnum* EnumType);
+
+	/** Gets a value from the given camera context data table. */
+	UFUNCTION(BlueprintCallable, Category=Camera)
+	static FInstancedStruct GetStructData(const FBlueprintCameraEvaluationDataRef& CameraData, FCameraContextDataID DataID, const UScriptStruct* DataStructType);
+
+	/** Gets a value from the given camera context data table. */
+	UFUNCTION(BlueprintCallable, Category=Camera)
+	static UObject* GetObjectData(const FBlueprintCameraEvaluationDataRef& CameraData, FCameraContextDataID DataID);
+
+	/** Gets a value from the given camera context data table. */
+	UFUNCTION(BlueprintCallable, Category=Camera)
+	static UClass* GetClassData(const FBlueprintCameraEvaluationDataRef& CameraData, FCameraContextDataID DataID);
+
+public:
+
+	/** Sets a value in the given camera context data table. */
+	UFUNCTION(BlueprintCallable, Category=Camera)
+	static bool SetNameData(const FBlueprintCameraEvaluationDataRef& CameraData, FCameraContextDataID DataID, const FName& Data);
+
+	/** Sets a value in the given camera context data table. */
+	UFUNCTION(BlueprintCallable, Category=Camera)
+	static bool SetStringData(const FBlueprintCameraEvaluationDataRef& CameraData, FCameraContextDataID DataID, const FString& Data);
+
+	/** Sets a value in the given camera context data table. */
+	UFUNCTION(BlueprintCallable, Category=Camera)
+	static bool SetEnumData(const FBlueprintCameraEvaluationDataRef& CameraData, FCameraContextDataID DataID, const UEnum* EnumType, uint8 Data);
+
+	/** Sets a value in the given camera context data table. */
+	UFUNCTION(BlueprintCallable, Category=Camera)
+	static bool SetStructData(const FBlueprintCameraEvaluationDataRef& CameraData, FCameraContextDataID DataID, const FInstancedStruct& Data);
+
+	/** Sets a value in the given camera context data table. */
+	UFUNCTION(BlueprintCallable, Category=Camera)
+	static bool SetObjectData(const FBlueprintCameraEvaluationDataRef& CameraData, FCameraContextDataID DataID, UObject* Data);
+
+	/** Sets a value in the given camera context data table. */
+	UFUNCTION(BlueprintCallable, Category=Camera)
+	static bool SetClassData(const FBlueprintCameraEvaluationDataRef& CameraData, FCameraContextDataID DataID, UClass* Data);
+};
+
diff --git a/Engine/Plugins/Cameras/GameplayCameras/Source/GameplayCameras/Public/GameFramework/BlueprintCameraNodeEvaluationResult.h b/Engine/Plugins/Cameras/GameplayCameras/Source/GameplayCameras/Public/GameFramework/BlueprintCameraNodeEvaluationResult.h
deleted file mode 100644
diff --git a/Engine/Plugins/Cameras/GameplayCameras/Source/GameplayCameras/Public/GameFramework/BlueprintCameraVariableTable.h b/Engine/Plugins/Cameras/GameplayCameras/Source/GameplayCameras/Public/GameFramework/BlueprintCameraVariableTable.h
deleted file mode 100644
diff --git a/Engine/Plugins/Cameras/GameplayCameras/Source/GameplayCameras/Public/GameFramework/CameraRigParameterInterop.h b/Engine/Plugins/Cameras/GameplayCameras/Source/GameplayCameras/Public/GameFramework/CameraRigParameterInterop.h
--- a/Engine/Plugins/Cameras/GameplayCameras/Source/GameplayCameras/Public/GameFramework/CameraRigParameterInterop.h
+++ b/Engine/Plugins/Cameras/GameplayCameras/Source/GameplayCameras/Public/GameFramework/CameraRigParameterInterop.h
@@ -11,7 +11,7 @@
 #include "CameraRigParameterInterop.generated.h"
 
 class UCameraRigAsset;
-struct FBlueprintCameraNodeEvaluationResult;
+struct FBlueprintCameraEvaluationDataRef;
 
 /**
  * Blueprint internal methods to set values on a camera rig's exposed parameters.
@@ -31,10 +31,10 @@ public:
 public:
 
 	UFUNCTION(BlueprintCallable, CustomThunk, Category="Camera", meta=(BlueprintInternalUseOnly="true", CustomStructureParam="ReturnValue"))
-	static void GetCameraParameter(UPARAM(Ref) const FBlueprintCameraNodeEvaluationResult& Result, UCameraRigAsset* CameraRig, FName ParameterName, int32& ReturnValue);
+	static void GetCameraParameter(UPARAM(Ref) const FBlueprintCameraEvaluationDataRef& CameraData, UCameraRigAsset* CameraRig, FName ParameterName, int32& ReturnValue);
 
 	UFUNCTION(BlueprintCallable, CustomThunk, Category="Camera", meta=(BlueprintInternalUseOnly="true", CustomStructureParam="NewValue"))
-	static void SetCameraParameter(UPARAM(Ref) FBlueprintCameraNodeEvaluationResult& Result, UCameraRigAsset* CameraRig, FName ParameterName, UPARAM(Ref) const int32& NewValue);
+	static void SetCameraParameter(UPARAM(Ref) FBlueprintCameraEvaluationDataRef& CameraData, UCameraRigAsset* CameraRig, FName ParameterName, UPARAM(Ref) const int32& NewValue);
 
 private:
 
diff --git a/Engine/Plugins/Cameras/GameplayCameras/Source/GameplayCameras/Public/GameFramework/GameplayCameraComponentBase.h b/Engine/Plugins/Cameras/GameplayCameras/Source/GameplayCameras/Public/GameFramework/GameplayCameraComponentBase.h
--- a/Engine/Plugins/Cameras/GameplayCameras/Source/GameplayCameras/Public/GameFramework/GameplayCameraComponentBase.h
+++ b/Engine/Plugins/Cameras/GameplayCameras/Source/GameplayCameras/Public/GameFramework/GameplayCameraComponentBase.h
@@ -5,9 +5,7 @@
 #include "Components/SceneComponent.h"
 #include "Core/CameraAssetReference.h"
 #include "Core/CameraEvaluationContext.h"
-#include "GameFramework/BlueprintCameraNodeEvaluationResult.h"
-#include "GameFramework/BlueprintCameraPose.h"
-#include "GameFramework/BlueprintCameraVariableTable.h"
+#include "GameFramework/BlueprintCameraEvaluationDataRef.h"
 #include "UObject/ObjectMacros.h"
 
 #include "GameplayCameraComponentBase.generated.h"
@@ -92,23 +90,11 @@ public:
 
 	/** Gets the shared camera evaluation data for this component's evaluation context. */
 	UFUNCTION(BlueprintPure, Category=Camera, meta=(DisplayName="Get Shared Camera Data"))
-	GAMEPLAYCAMERAS_API FBlueprintCameraNodeEvaluationResult GetInitialResult() const;
-	
+	GAMEPLAYCAMERAS_API FBlueprintCameraEvaluationDataRef GetInitialResult() const;
+
 	/** Gets the camera evaluation data for a given sub-set of camera rigs in this component's evaluation context. */
 	UFUNCTION(BlueprintPure, Category=Camera, meta=(DisplayName="Get Conditional Camera Data"))
-	GAMEPLAYCAMERAS_API FBlueprintCameraNodeEvaluationResult GetConditionalResult(ECameraEvaluationDataCondition Condition) const;
-
-	/** Gets the initial camera pose for this component's camera evaluation context. */
-	UFUNCTION(BlueprintPure, Category=Camera, meta=(DeprecatedFunction, DeprecationMessage="Please use GetSharedCameraData"))
-	GAMEPLAYCAMERAS_API FBlueprintCameraPose GetInitialPose() const;
-
-	/** Sets the initial camera pose for this component's camera evaluation context. */
-	UFUNCTION(BlueprintCallable, Category=Camera, meta=(DeprecatedFunction, DeprecationMessage="Please use GetSharedCameraData"))
-	GAMEPLAYCAMERAS_API bool SetInitialPose(const FBlueprintCameraPose& CameraPose);
-
-	/** Gets the initial camera variable table for this component's camera evaluation context. */
-	UFUNCTION(BlueprintPure, Category=Camera, meta=(DeprecatedFunction, DeprecationMessage="Please use GetSharedCameraData"))
-	GAMEPLAYCAMERAS_API FBlueprintCameraVariableTable GetInitialVariableTable() const;
+	GAMEPLAYCAMERAS_API FBlueprintCameraEvaluationDataRef GetConditionalResult(ECameraEvaluationDataCondition Condition) const;
 
 public:
 
diff --git a/Engine/Plugins/Cameras/GameplayCameras/Source/GameplayCameras/Public/Nodes/Utility/BlueprintCameraNode.h b/Engine/Plugins/Cameras/GameplayCameras/Source/GameplayCameras/Public/Nodes/Utility/BlueprintCameraNode.h
--- a/Engine/Plugins/Cameras/GameplayCameras/Source/GameplayCameras/Public/Nodes/Utility/BlueprintCameraNode.h
+++ b/Engine/Plugins/Cameras/GameplayCameras/Source/GameplayCameras/Public/Nodes/Utility/BlueprintCameraNode.h
@@ -6,8 +6,7 @@
 #include "Core/CameraNode.h"
 #include "Core/ICustomCameraNodeParameterProvider.h"
 #include "Core/ObjectTreeGraphObject.h"
-#include "GameFramework/BlueprintCameraPose.h"
-#include "GameFramework/BlueprintCameraVariableTable.h"
+#include "GameFramework/BlueprintCameraEvaluationDataRef.h"
 
 #include "BlueprintCameraNode.generated.h"
 
@@ -39,6 +38,8 @@ public:
 	UFUNCTION(BlueprintImplementableEvent, Category="Evaluation")
 	void TickCameraNode(float DeltaTime);
 
+public:
+
 	/**
 	 * A utility function that tries to find if an actor owns the evaluation context.
 	 * Handles the situation where the evaluation context is an actor component (like a
@@ -47,6 +48,18 @@ public:
 	UFUNCTION(BlueprintPure, Category="Evaluation", meta=(DeterminesOutputType="ActorClass"))
 	AActor* FindEvaluationContextOwnerActor(TSubclassOf<AActor> ActorClass) const;
 
+	/** A utility function to get the current camera pose from this node's camera data. */
+	UFUNCTION(BlueprintPure, Category="Evaluation")
+	FBlueprintCameraPose GetCurrentCameraPose() const;
+
+	/** A utility function to set the current camera pose on this node's camera data. */
+	UFUNCTION(BlueprintCallable, Category="Evaluation")
+	void SetCurrentCameraPose(const FBlueprintCameraPose& CameraPose);
+
+	/** Assigns the default parameter values of the owning camera rig to the given camera evaluation data. */
+	UFUNCTION(BlueprintCallable, Category="Evaluation")
+	void SetDefaultOwningCameraRigParameters(FBlueprintCameraEvaluationDataRef TargetCameraData) const;
+
 public:
 
 	using FCameraNodeEvaluatorInitializeParams = UE::Cameras::FCameraNodeEvaluatorInitializeParams;
@@ -55,7 +68,7 @@ public:
 	using FCameraEvaluationContext = UE::Cameras::FCameraEvaluationContext;
 
 	/** Initialize this camera node. */
-	void NativeInitializeCameraNode(const FCameraNodeEvaluatorInitializeParams& Params, FCameraNodeEvaluationResult& OutResult);
+	void NativeInitializeCameraNode(const UBlueprintCameraNode* InBlueprintNode, const FCameraNodeEvaluatorInitializeParams& Params, FCameraNodeEvaluationResult& OutResult);
 
 	/** Runs this camera node. */
 	void NativeRunCameraNode(const FCameraNodeEvaluationParams& Params, FCameraNodeEvaluationResult& OutResult);
@@ -83,21 +96,17 @@ protected:
 	UPROPERTY(BlueprintReadOnly, Category="Evaluation")
 	TObjectPtr<UObject> EvaluationContextOwner;
 
-	/** The input/output camera pose for this frame. */
-	UPROPERTY(BlueprintReadWrite, Category="Evaluation")
-	FBlueprintCameraPose CameraPose;
-
-	/** The input/output camera variable table for this frame. */
-	UPROPERTY(BlueprintReadWrite, Category="Evaluation")
-	FBlueprintCameraVariableTable VariableTable;
+	/** The input/output camera data for this frame. */
+	UPROPERTY(BlueprintReadOnly, Category="Evaluation")
+	FBlueprintCameraEvaluationDataRef CameraData;
 
 private:
 
 	TSharedPtr<const UE::Cameras::FCameraEvaluationContext> CurrentContext;
 
-	FCameraNodeEvaluationResult* CurrentResult = nullptr;
-
 	mutable TWeakObjectPtr<UWorld> WeakCachedWorld;
+
+	const UBlueprintCameraNode* BlueprintNode = nullptr;
 };
 
 /**
diff --git a/Engine/Plugins/Cameras/GameplayCameras/Source/GameplayCamerasUncookedOnly/Private/BlueprintGraph/K2Node_CameraRigBase.cpp b/Engine/Plugins/Cameras/GameplayCameras/Source/GameplayCamerasUncookedOnly/Private/BlueprintGraph/K2Node_CameraRigBase.cpp
--- a/Engine/Plugins/Cameras/GameplayCameras/Source/GameplayCamerasUncookedOnly/Private/BlueprintGraph/K2Node_CameraRigBase.cpp
+++ b/Engine/Plugins/Cameras/GameplayCameras/Source/GameplayCamerasUncookedOnly/Private/BlueprintGraph/K2Node_CameraRigBase.cpp
@@ -9,12 +9,11 @@
 #include "EdGraphSchema_K2.h"
 #include "Editor.h"
 #include "EditorCategoryUtils.h"
-#include "GameFramework/BlueprintCameraNodeEvaluationResult.h"
+#include "GameFramework/BlueprintCameraEvaluationDataRef.h"
 #include "GameFramework/CameraRigParameterInterop.h"
 #include "GameplayCamerasDelegates.h"
 #include "Kismet2/BlueprintEditorUtils.h"
 #include "KismetCompiler.h"
-#include "Nodes/Framing/CameraFramingZone.h"
 #include "Subsystems/AssetEditorSubsystem.h"
 
 #define LOCTEXT_NAMESPACE "K2Node_CameraRigBase"
@@ -47,7 +46,7 @@ void UK2Node_CameraRigBase::AllocateDefaultPins()
 	CreatePin(EGPD_Output, UEdGraphSchema_K2::PC_Exec, UEdGraphSchema_K2::PN_Then);
 
 	// Add evalation result pin.
-	CreatePin(EGPD_Input, UEdGraphSchema_K2::PC_Struct, FBlueprintCameraNodeEvaluationResult::StaticStruct(), CameraNodeEvaluationResultPinName);
+	CreatePin(EGPD_Input, UEdGraphSchema_K2::PC_Struct, FBlueprintCameraEvaluationDataRef::StaticStruct(), CameraNodeEvaluationResultPinName);
 
 	Super::AllocateDefaultPins();
 }
diff --git a/Engine/Plugins/Cameras/GameplayCameras/Source/GameplayCamerasUncookedOnly/Private/BlueprintGraph/K2Node_SetCameraRigParameter.cpp b/Engine/Plugins/Cameras/GameplayCameras/Source/GameplayCamerasUncookedOnly/Private/BlueprintGraph/K2Node_SetCameraRigParameter.cpp
--- a/Engine/Plugins/Cameras/GameplayCameras/Source/GameplayCamerasUncookedOnly/Private/BlueprintGraph/K2Node_SetCameraRigParameter.cpp
+++ b/Engine/Plugins/Cameras/GameplayCameras/Source/GameplayCamerasUncookedOnly/Private/BlueprintGraph/K2Node_SetCameraRigParameter.cpp
@@ -175,7 +175,7 @@ void UK2Node_SetCameraRigParameter::ExpandNode(FKismetCompilerContext& CompilerC
 	UEdGraphPin* FirstExecPin = CallSetParameter->GetExecPin();
 
 	// Connect the camera evaluation result argument.
-	UEdGraphPin* CallSetParameterResultPin = CallSetParameter->FindPinChecked(TEXT("Result"));
+	UEdGraphPin* CallSetParameterResultPin = CallSetParameter->FindPinChecked(TEXT("CameraData"));
 	CompilerContext.CopyPinLinksToIntermediate(*CameraNodeEvaluationResultPin, *CallSetParameterResultPin);
 
 	// Set the camera rig argument.
diff --git a/Engine/Plugins/Cameras/GameplayCameras/Source/GameplayCamerasUncookedOnly/Private/BlueprintGraph/K2Node_SetCameraRigParameters.cpp b/Engine/Plugins/Cameras/GameplayCameras/Source/GameplayCamerasUncookedOnly/Private/BlueprintGraph/K2Node_SetCameraRigParameters.cpp
--- a/Engine/Plugins/Cameras/GameplayCameras/Source/GameplayCamerasUncookedOnly/Private/BlueprintGraph/K2Node_SetCameraRigParameters.cpp
+++ b/Engine/Plugins/Cameras/GameplayCameras/Source/GameplayCamerasUncookedOnly/Private/BlueprintGraph/K2Node_SetCameraRigParameters.cpp
@@ -165,7 +165,7 @@ void UK2Node_SetCameraRigParameters::ExpandNode(FKismetCompilerContext& Compiler
 		CompilerContext.MessageLog.NotifyIntermediateObjectCreation(CallSetParameter, SourceGraph);
 
 		// Connect the camera evaluation result argument.
-		UEdGraphPin* CallSetParameterResultPin = CallSetParameter->FindPinChecked(TEXT("Result"));
+		UEdGraphPin* CallSetParameterResultPin = CallSetParameter->FindPinChecked(TEXT("CameraData"));
 		CompilerContext.CopyPinLinksToIntermediate(*CameraNodeEvaluationResultPin, *CallSetParameterResultPin);
 
 		// Set the camera rig argument.
@@ -229,7 +229,7 @@ void UK2Node_SetCameraRigParameters::ExpandNode(FKismetCompilerContext& Compiler
 		UEdGraphPin* CurExecPin = CallSetParameter->GetExecPin();
 
 		// Connect the camera evaluation result argument.
-		UEdGraphPin* CallSetParameterResultPin = CallSetParameter->FindPinChecked(TEXT("Result"));
+		UEdGraphPin* CallSetParameterResultPin = CallSetParameter->FindPinChecked(TEXT("CameraData"));
 		CompilerContext.CopyPinLinksToIntermediate(*CameraNodeEvaluationResultPin, *CallSetParameterResultPin);
 
 		// Set the camera rig argument.
