Cameras: add spline-based offset and spline-based FOV nodes

diff --git a/Engine/Plugins/Cameras/GameplayCameras/Source/GameplayCameras/Private/Math/CameraNodeSpaceMath.cpp b/Engine/Plugins/Cameras/GameplayCameras/Source/GameplayCameras/Private/Math/CameraNodeSpaceMath.cpp
--- a/Engine/Plugins/Cameras/GameplayCameras/Source/GameplayCameras/Private/Math/CameraNodeSpaceMath.cpp
+++ b/Engine/Plugins/Cameras/GameplayCameras/Source/GameplayCameras/Private/Math/CameraNodeSpaceMath.cpp
@@ -4,12 +4,14 @@
 
 #include "Core/BuiltInCameraVariables.h"
 #include "Core/CameraEvaluationContext.h"
+#include "Core/CameraEvaluationContext.h"
 #include "Core/CameraNodeEvaluator.h"
 #include "Core/CameraPose.h"
 #include "Core/CameraRigJoints.h"
 #include "Core/CameraSystemEvaluator.h"
 #include "GameFramework/Pawn.h"
 #include "GameFramework/PlayerController.h"
+#include "Math/Axis.h"
 
 namespace UE::Cameras
 {
@@ -276,5 +278,84 @@ bool FCameraNodeSpaceMath::OffsetCameraNodeSpacePosition(const FCameraNodeSpaceP
 	return bGotWorldOffset;
 }
 
+bool FCameraNodeSpaceMath::OffsetCameraNodeSpaceTransform(const FCameraNodeSpaceParams& Params, const FTransform3d& InTransform, const FVector3d& InLocationOffset, const FRotator3d& InRotationOffset, ECameraNodeSpace InSpace, FTransform3d& OutTransform)
+{
+	FTransform3d LocalSpace;
+	bool bApplyLocalSpaceTransform = true;
+	bool bIsValid = false;
+
+	switch (InSpace)
+	{
+		case ECameraNodeSpace::CameraPose:
+		default:
+			{
+				OutTransform = FTransform3d(InRotationOffset, InLocationOffset) * InTransform;
+				bApplyLocalSpaceTransform = false;
+				bIsValid = true;
+			}
+			break;
+		case ECameraNodeSpace::ActiveContext:
+			if (TSharedPtr<const FCameraEvaluationContext> ActiveContext = Params.GetActiveContext())
+			{
+				const FCameraNodeEvaluationResult& InitialResult = ActiveContext->GetInitialResult();
+				LocalSpace = InitialResult.CameraPose.GetTransform();
+				bIsValid = true;
+			}
+			break;
+		case ECameraNodeSpace::OwningContext:
+			if (TSharedPtr<const FCameraEvaluationContext> OwningContext = Params.GetOwningContext())
+			{ 
+				const FCameraNodeEvaluationResult& InitialResult = OwningContext->GetInitialResult();
+				LocalSpace = InitialResult.CameraPose.GetTransform();
+				bIsValid = true;
+			}
+			break;
+		case ECameraNodeSpace::Pivot:
+			{
+				FTransform PivotTransform;
+				if (Params.FindPivotTransform(PivotTransform))
+				{
+					LocalSpace = PivotTransform;
+					bIsValid = true;
+				}
+				else if (TSharedPtr<const FCameraEvaluationContext> ActiveContext = Params.GetActiveContext())
+				{
+					const FCameraNodeEvaluationResult& InitialResult = ActiveContext->GetInitialResult();
+					LocalSpace = InitialResult.CameraPose.GetTransform();
+					bIsValid = true;
+				}
+			}
+			break;
+		case ECameraNodeSpace::World:
+			{
+				OutTransform = InTransform;
+				OutTransform.SetTranslation(InLocationOffset + InTransform.GetTranslation());
+				OutTransform.SetRotation(InRotationOffset.Quaternion() * InTransform.GetRotation());
+				bApplyLocalSpaceTransform = false;
+				bIsValid = true;
+			}
+			break;
+	}
+
+	if (bIsValid && bApplyLocalSpaceTransform)
+	{
+		const FVector3d WorldTranslationOffset = LocalSpace.TransformVector(InLocationOffset);
+
+		const FVector3d ContextForward = LocalSpace.GetUnitAxis(EAxis::X);
+		const FVector3d ContextRight = LocalSpace.GetUnitAxis(EAxis::Y);
+		const FVector3d ContextUp = LocalSpace.GetUnitAxis(EAxis::Z);
+		const FQuat WorldRotationOffset = 
+			FQuat(ContextUp, FMath::DegreesToRadians(InRotationOffset.Yaw)) * 
+			FQuat(ContextRight, -FMath::DegreesToRadians(InRotationOffset.Pitch)) *
+			FQuat(ContextForward, -FMath::DegreesToRadians(InRotationOffset.Roll));
+
+		OutTransform = InTransform;
+		OutTransform.SetTranslation(WorldTranslationOffset + InTransform.GetTranslation());
+		OutTransform.SetRotation(WorldRotationOffset * InTransform.GetRotation());
+	}
+
+	return bIsValid;
+}
+
 }  // namespace UE::Cameras
 
diff --git a/Engine/Plugins/Cameras/GameplayCameras/Source/GameplayCameras/Private/Nodes/Common/OffsetCameraNode.cpp b/Engine/Plugins/Cameras/GameplayCameras/Source/GameplayCameras/Private/Nodes/Common/OffsetCameraNode.cpp
--- a/Engine/Plugins/Cameras/GameplayCameras/Source/GameplayCameras/Private/Nodes/Common/OffsetCameraNode.cpp
+++ b/Engine/Plugins/Cameras/GameplayCameras/Source/GameplayCameras/Private/Nodes/Common/OffsetCameraNode.cpp
@@ -2,10 +2,8 @@
 
 #include "Nodes/Common/OffsetCameraNode.h"
 
-#include "Core/CameraEvaluationContext.h"
 #include "Core/CameraParameterReader.h"
-#include "GameplayCameras.h"
-#include "Math/Axis.h"
+#include "Math/CameraNodeSpaceMath.h"
 
 #include UE_INLINE_GENERATED_CPP_BY_NAME(OffsetCameraNode)
 
@@ -42,55 +40,18 @@ void FOffsetCameraNodeEvaluator::OnRun(const FCameraNodeEvaluationParams& Params
 	const FRotator3d RotationOffset = RotationReader.Get(OutResult.VariableTable);
 
 	const UOffsetCameraNode* OffsetNode = GetCameraNodeAs<UOffsetCameraNode>();
-	switch (OffsetNode->OffsetSpace)
+
+	FTransform3d OutTransform;
+	const bool bSuccess = FCameraNodeSpaceMath::OffsetCameraNodeSpaceTransform(
+			FCameraNodeSpaceParams(Params, OutResult),
+			OutResult.CameraPose.GetTransform(),
+			TranslationOffset,
+			RotationOffset,
+			OffsetNode->OffsetSpace,
+			OutTransform);
+	if (bSuccess)
 	{
-		case ECameraNodeSpace::CameraPose:
-		default:
-			{
-				FTransform3d Transform = OutResult.CameraPose.GetTransform();
-				Transform = FTransform3d(RotationOffset, TranslationOffset) * Transform;
-				OutResult.CameraPose.SetTransform(Transform);
-			}
-			break;
-		case ECameraNodeSpace::OwningContext:
-			if (Params.EvaluationContext)
-			{ 
-				// The offsets are meant to be treated as context-local. Let's get the context transform
-				// and apply the offsets using that transform's axes.
-				const FCameraNodeEvaluationResult& InitialResult = Params.EvaluationContext->GetInitialResult();
-				const FTransform3d ContextTransform = InitialResult.CameraPose.GetTransform();
-
-				const FVector3d WorldTranslationOffset = ContextTransform.TransformVector(TranslationOffset);
-
-				const FVector3d ContextForward = ContextTransform.GetUnitAxis(EAxis::X);
-				const FVector3d ContextRight = ContextTransform.GetUnitAxis(EAxis::Y);
-				const FVector3d ContextUp = ContextTransform.GetUnitAxis(EAxis::Z);
-				const FQuat WorldRotationOffset = 
-					FQuat(ContextUp, FMath::DegreesToRadians(RotationOffset.Yaw)) * 
-					FQuat(ContextRight, -FMath::DegreesToRadians(RotationOffset.Pitch)) *
-					FQuat(ContextForward, -FMath::DegreesToRadians(RotationOffset.Roll));
-
-				FTransform3d Transform = OutResult.CameraPose.GetTransform();
-				Transform.SetTranslation(Transform.GetTranslation() + WorldTranslationOffset);
-				Transform.SetRotation(WorldRotationOffset * Transform.GetRotation());
-				OutResult.CameraPose.SetTransform(Transform);
-			}
-			else
-			{
-				UE_LOG(LogCameraSystem, Error, 
-						TEXT("OffsetCameraNode: cannot offset in context space when there is "
-							 "no current context set."));
-				return;
-			}
-			break;
-		case ECameraNodeSpace::World:
-			{
-				FTransform3d Transform = OutResult.CameraPose.GetTransform();
-				Transform.SetTranslation(Transform.GetTranslation() + TranslationOffset);
-				Transform.SetRotation(RotationOffset.Quaternion() * Transform.GetRotation());
-				OutResult.CameraPose.SetTransform(Transform);
-			}
-			break;
+		OutResult.CameraPose.SetTransform(OutTransform);
 	}
 }
 
diff --git a/Engine/Plugins/Cameras/GameplayCameras/Source/GameplayCameras/Private/Nodes/Common/SplineFieldOfViewCameraNode.cpp b/Engine/Plugins/Cameras/GameplayCameras/Source/GameplayCameras/Private/Nodes/Common/SplineFieldOfViewCameraNode.cpp
new file mode 100644
--- /dev/null
+++ b/Engine/Plugins/Cameras/GameplayCameras/Source/GameplayCameras/Private/Nodes/Common/SplineFieldOfViewCameraNode.cpp
@@ -0,0 +1,108 @@
+// Copyright Epic Games, Inc. All Rights Reserved.
+
+#include "Nodes/Common/SplineFieldOfViewCameraNode.h"
+
+#include "Core/CameraEvaluationContext.h"
+#include "Core/CameraParameterReader.h"
+#include "GameplayCameras.h"
+#include "IGameplayCamerasLiveEditListener.h"
+#include "IGameplayCamerasLiveEditManager.h"
+#include "IGameplayCamerasModule.h"
+
+#include UE_INLINE_GENERATED_CPP_BY_NAME(SplineFieldOfViewCameraNode)
+
+namespace UE::Cameras
+{
+
+class FSplineFieldOfViewCameraNodeEvaluator 
+	: public FCameraNodeEvaluator
+#if WITH_EDITOR
+    , public IGameplayCamerasLiveEditListener
+#endif  // WITH_EDITOR
+{
+	UE_DECLARE_CAMERA_NODE_EVALUATOR(GAMEPLAYCAMERAS_API, FSplineFieldOfViewCameraNodeEvaluator)
+
+public:
+
+	~FSplineFieldOfViewCameraNodeEvaluator();
+
+protected:
+
+	virtual void OnInitialize(const FCameraNodeEvaluatorInitializeParams& Params, FCameraNodeEvaluationResult& OutResult) override;
+	virtual void OnRun(const FCameraNodeEvaluationParams& Params, FCameraNodeEvaluationResult& OutResult) override;
+
+private:
+
+	void RebuildCurve();
+
+private:
+
+	TCameraParameterReader<float> SplineInputReader;
+
+	FCompressedRichCurve FieldOfViewSpline;
+
+	bool bHasAnyValues = false;
+};
+
+UE_DEFINE_CAMERA_NODE_EVALUATOR(FSplineFieldOfViewCameraNodeEvaluator)
+
+void FSplineFieldOfViewCameraNodeEvaluator::OnInitialize(const FCameraNodeEvaluatorInitializeParams& Params, FCameraNodeEvaluationResult& OutResult)
+{
+	const USplineFieldOfViewCameraNode* FieldOfViewNode = GetCameraNodeAs<USplineFieldOfViewCameraNode>();
+	SplineInputReader.Initialize(FieldOfViewNode->SplineInput);
+
+	RebuildCurve();
+
+#if WITH_EDITOR
+	IGameplayCamerasModule& GameplayCamerasModule = IGameplayCamerasModule::Get();
+	if (TSharedPtr<IGameplayCamerasLiveEditManager> LiveEditManager = GameplayCamerasModule.GetLiveEditManager())
+	{
+		LiveEditManager->AddListener(GetCameraNode(), this);
+	}
+#endif  // WITH_EDITOR
+}
+
+FSplineFieldOfViewCameraNodeEvaluator::~FSplineFieldOfViewCameraNodeEvaluator()
+{
+#if WITH_EDITOR
+	IGameplayCamerasModule& GameplayCamerasModule = IGameplayCamerasModule::Get();
+	if (TSharedPtr<IGameplayCamerasLiveEditManager> LiveEditManager = GameplayCamerasModule.GetLiveEditManager())
+	{
+		LiveEditManager->RemoveListener(this);
+	}
+#endif  // WITH_EDITOR
+}
+
+void FSplineFieldOfViewCameraNodeEvaluator::RebuildCurve()
+{
+	const USplineFieldOfViewCameraNode* FieldOfViewNode = GetCameraNodeAs<USplineFieldOfViewCameraNode>();
+
+	bHasAnyValues = FieldOfViewNode->FieldOfViewSpline.HasAnyData();
+
+	FieldOfViewNode->FieldOfViewSpline.Curve.CompressCurve(FieldOfViewSpline);
+}
+
+void FSplineFieldOfViewCameraNodeEvaluator::OnRun(const FCameraNodeEvaluationParams& Params, FCameraNodeEvaluationResult& OutResult)
+{
+	if (bHasAnyValues)
+	{
+		const float SplineInput = SplineInputReader.Get(OutResult.VariableTable);
+		const float FieldOfView = FieldOfViewSpline.Eval(SplineInput);
+		OutResult.CameraPose.SetFieldOfView(FieldOfView);
+		OutResult.CameraPose.SetFocalLength(-1);
+	}
+}
+
+}  // namespace UE::Cameras
+
+USplineFieldOfViewCameraNode::USplineFieldOfViewCameraNode(const FObjectInitializer& ObjectInit)
+	: Super(ObjectInit)
+{
+}
+
+FCameraNodeEvaluatorPtr USplineFieldOfViewCameraNode::OnBuildEvaluator(FCameraNodeEvaluatorBuilder& Builder) const
+{
+	using namespace UE::Cameras;
+	return Builder.BuildEvaluator<FSplineFieldOfViewCameraNodeEvaluator>();
+}
+
diff --git a/Engine/Plugins/Cameras/GameplayCameras/Source/GameplayCameras/Private/Nodes/Common/SplineOffsetCameraNode.cpp b/Engine/Plugins/Cameras/GameplayCameras/Source/GameplayCameras/Private/Nodes/Common/SplineOffsetCameraNode.cpp
new file mode 100644
--- /dev/null
+++ b/Engine/Plugins/Cameras/GameplayCameras/Source/GameplayCameras/Private/Nodes/Common/SplineOffsetCameraNode.cpp
@@ -0,0 +1,154 @@
+// Copyright Epic Games, Inc. All Rights Reserved.
+
+#include "Nodes/Common/SplineOffsetCameraNode.h"
+
+#include "Core/CameraParameterReader.h"
+#include "IGameplayCamerasModule.h"
+#include "IGameplayCamerasLiveEditListener.h"
+#include "IGameplayCamerasLiveEditManager.h"
+#include "Math/CameraNodeSpaceMath.h"
+
+#include UE_INLINE_GENERATED_CPP_BY_NAME(SplineOffsetCameraNode)
+
+namespace UE::Cameras
+{
+
+class FSplineOffsetCameraNodeEvaluator
+	: public FCameraNodeEvaluator
+#if WITH_EDITOR
+    , public IGameplayCamerasLiveEditListener
+#endif  // WITH_EDITOR
+{
+	UE_DECLARE_CAMERA_NODE_EVALUATOR(GAMEPLAYCAMERAS_API, FSplineOffsetCameraNodeEvaluator)
+
+public:
+	
+	~FSplineOffsetCameraNodeEvaluator();
+
+protected:
+
+	virtual void OnInitialize(const FCameraNodeEvaluatorInitializeParams& Params, FCameraNodeEvaluationResult& OutResult) override;
+	virtual void OnRun(const FCameraNodeEvaluationParams& Params, FCameraNodeEvaluationResult& OutResult) override;
+
+#if WITH_EDITOR
+	virtual void OnPostEditChangeProperty(const UCameraNode* InCameraNode, const FPropertyChangedEvent& PropertyChangedEvent) override;
+#endif  // WITH_EDITOR
+
+private:
+
+	void RebuildCurves();
+
+private:
+
+	TCameraParameterReader<float> SplineInputReader;
+
+	FCompressedRichCurve LocationOffsetSpline[3];
+	FCompressedRichCurve RotationOffsetSpline[3];
+
+	bool bHasAnyLocationOffset = false;
+	bool bHasAnyRotationOffset = false;
+};
+
+UE_DEFINE_CAMERA_NODE_EVALUATOR(FSplineOffsetCameraNodeEvaluator)
+
+void FSplineOffsetCameraNodeEvaluator::OnInitialize(const FCameraNodeEvaluatorInitializeParams& Params, FCameraNodeEvaluationResult& OutResult)
+{
+	const USplineOffsetCameraNode* SplineOffsetNode = GetCameraNodeAs<USplineOffsetCameraNode>();
+	SplineInputReader.Initialize(SplineOffsetNode->SplineInput);
+
+	RebuildCurves();
+
+#if WITH_EDITOR
+	IGameplayCamerasModule& GameplayCamerasModule = IGameplayCamerasModule::Get();
+	if (TSharedPtr<IGameplayCamerasLiveEditManager> LiveEditManager = GameplayCamerasModule.GetLiveEditManager())
+	{
+		LiveEditManager->AddListener(GetCameraNode(), this);
+	}
+#endif  // WITH_EDITOR
+}
+
+FSplineOffsetCameraNodeEvaluator::~FSplineOffsetCameraNodeEvaluator()
+{
+#if WITH_EDITOR
+	IGameplayCamerasModule& GameplayCamerasModule = IGameplayCamerasModule::Get();
+	if (TSharedPtr<IGameplayCamerasLiveEditManager> LiveEditManager = GameplayCamerasModule.GetLiveEditManager())
+	{
+		LiveEditManager->RemoveListener(this);
+	}
+#endif  // WITH_EDITOR
+}
+
+void FSplineOffsetCameraNodeEvaluator::RebuildCurves()
+{
+	const USplineOffsetCameraNode* SplineOffsetNode = GetCameraNodeAs<USplineOffsetCameraNode>();
+
+	bHasAnyLocationOffset = SplineOffsetNode->TranslationOffsetSpline.HasAnyData();
+	bHasAnyRotationOffset = SplineOffsetNode->RotationOffsetSpline.HasAnyData();
+
+	SplineOffsetNode->TranslationOffsetSpline.Curves[0].CompressCurve(LocationOffsetSpline[0]);
+	SplineOffsetNode->TranslationOffsetSpline.Curves[1].CompressCurve(LocationOffsetSpline[1]);
+	SplineOffsetNode->TranslationOffsetSpline.Curves[2].CompressCurve(LocationOffsetSpline[2]);
+
+	SplineOffsetNode->RotationOffsetSpline.Curves[0].CompressCurve(RotationOffsetSpline[0]);
+	SplineOffsetNode->RotationOffsetSpline.Curves[1].CompressCurve(RotationOffsetSpline[1]);
+	SplineOffsetNode->RotationOffsetSpline.Curves[2].CompressCurve(RotationOffsetSpline[2]);
+}
+
+void FSplineOffsetCameraNodeEvaluator::OnRun(const FCameraNodeEvaluationParams& Params, FCameraNodeEvaluationResult& OutResult)
+{
+	const float SplineInput = SplineInputReader.Get(OutResult.VariableTable);
+
+	const FVector3d TranslationOffset = bHasAnyLocationOffset ?
+		FVector3d(
+			LocationOffsetSpline[0].Eval(SplineInput),
+			LocationOffsetSpline[1].Eval(SplineInput),
+			LocationOffsetSpline[2].Eval(SplineInput))
+		:
+		FVector3d::ZeroVector;
+
+	const FRotator3d RotationOffset = bHasAnyRotationOffset ?
+		FRotator3d(
+			RotationOffsetSpline[0].Eval(SplineInput),
+			RotationOffsetSpline[1].Eval(SplineInput),
+			RotationOffsetSpline[2].Eval(SplineInput))
+		:
+		FRotator3d::ZeroRotator;
+
+	const USplineOffsetCameraNode* OffsetNode = GetCameraNodeAs<USplineOffsetCameraNode>();
+
+	FTransform3d OutTransform;
+	const bool bSuccess = FCameraNodeSpaceMath::OffsetCameraNodeSpaceTransform(
+			FCameraNodeSpaceParams(Params, OutResult),
+			OutResult.CameraPose.GetTransform(),
+			TranslationOffset,
+			RotationOffset,
+			OffsetNode->OffsetSpace,
+			OutTransform);
+	if (bSuccess)
+	{
+		OutResult.CameraPose.SetTransform(OutTransform);
+	}
+}
+
+#if WITH_EDITOR
+
+void FSplineOffsetCameraNodeEvaluator::OnPostEditChangeProperty(const UCameraNode* InCameraNode, const FPropertyChangedEvent& PropertyChangedEvent)
+{
+	const FName PropertyName = PropertyChangedEvent.GetMemberPropertyName();
+	if (PropertyName == GET_MEMBER_NAME_CHECKED(USplineOffsetCameraNode, TranslationOffsetSpline) ||
+			PropertyName == GET_MEMBER_NAME_CHECKED(USplineOffsetCameraNode, RotationOffsetSpline))
+	{
+		RebuildCurves();
+	}
+}
+
+#endif  // WITH_EDITOR
+
+}  // namespace UE::Cameras
+
+FCameraNodeEvaluatorPtr USplineOffsetCameraNode::OnBuildEvaluator(FCameraNodeEvaluatorBuilder& Builder) const
+{
+	using namespace UE::Cameras;
+	return Builder.BuildEvaluator<FSplineOffsetCameraNodeEvaluator>();
+}
+
diff --git a/Engine/Plugins/Cameras/GameplayCameras/Source/GameplayCameras/Public/Math/CameraNodeSpaceMath.h b/Engine/Plugins/Cameras/GameplayCameras/Source/GameplayCameras/Public/Math/CameraNodeSpaceMath.h
--- a/Engine/Plugins/Cameras/GameplayCameras/Source/GameplayCameras/Public/Math/CameraNodeSpaceMath.h
+++ b/Engine/Plugins/Cameras/GameplayCameras/Source/GameplayCameras/Public/Math/CameraNodeSpaceMath.h
@@ -56,6 +56,9 @@ public:
 
 	/** Offsets a given world position in the specified camera node space. */
 	static bool OffsetCameraNodeSpacePosition(const FCameraNodeSpaceParams& Params, const FVector3d& InPosition, const FVector3d& InOffset, ECameraNodeSpace InSpace, FVector3d& OutPosition);
+
+	/** Offsets a given world transform in the specified camera node space. */
+	static bool OffsetCameraNodeSpaceTransform(const FCameraNodeSpaceParams& Params, const FTransform3d& InTransform, const FVector3d& InLocationOffset, const FRotator3d& InRotationOffset, ECameraNodeSpace InSpace, FTransform3d& OutTransform);
 };
 
 }  // namespace UE::Cameras
diff --git a/Engine/Plugins/Cameras/GameplayCameras/Source/GameplayCameras/Public/Nodes/Common/SplineFieldOfViewCameraNode.h b/Engine/Plugins/Cameras/GameplayCameras/Source/GameplayCameras/Public/Nodes/Common/SplineFieldOfViewCameraNode.h
new file mode 100644
--- /dev/null
+++ b/Engine/Plugins/Cameras/GameplayCameras/Source/GameplayCameras/Public/Nodes/Common/SplineFieldOfViewCameraNode.h
@@ -0,0 +1,39 @@
+// Copyright Epic Games, Inc. All Rights Reserved.
+
+#pragma once
+
+#include "Core/CameraNode.h"
+#include "Core/CameraParameters.h"
+#include "Curves/CameraSingleCurve.h"
+#include "Nodes/CameraNodeTypes.h"
+
+#include "SplineFieldOfViewCameraNode.generated.h"
+
+/**
+ * A camera node that sets the field of view of the camera by evaluating a curve.
+ */
+UCLASS(MinimalAPI, meta=(CameraNodeCategories="Common,Lens"))
+class USplineFieldOfViewCameraNode : public UCameraNode
+{
+	GENERATED_BODY()
+
+public:
+
+	/** The input to pass to the field of view spline. */
+	UPROPERTY(EditAnywhere, Category=Common)
+	FFloatCameraParameter SplineInput;
+
+	/** The field of view, in degrees. */
+	UPROPERTY(EditAnywhere, Category=Common)
+	FCameraSingleCurve FieldOfViewSpline;
+
+public:
+
+	USplineFieldOfViewCameraNode(const FObjectInitializer& ObjectInit);
+
+protected:
+
+	// UCameraNode interface.
+	virtual FCameraNodeEvaluatorPtr OnBuildEvaluator(FCameraNodeEvaluatorBuilder& Builder) const override;
+};
+
diff --git a/Engine/Plugins/Cameras/GameplayCameras/Source/GameplayCameras/Public/Nodes/Common/SplineOffsetCameraNode.h b/Engine/Plugins/Cameras/GameplayCameras/Source/GameplayCameras/Public/Nodes/Common/SplineOffsetCameraNode.h
new file mode 100644
--- /dev/null
+++ b/Engine/Plugins/Cameras/GameplayCameras/Source/GameplayCameras/Public/Nodes/Common/SplineOffsetCameraNode.h
@@ -0,0 +1,44 @@
+// Copyright Epic Games, Inc. All Rights Reserved.
+
+#pragma once
+
+#include "Core/CameraNode.h"
+#include "Core/CameraParameters.h"
+#include "Curves/CameraRotatorCurve.h"
+#include "Curves/CameraVectorCurve.h"
+#include "Nodes/CameraNodeTypes.h"
+
+#include "SplineOffsetCameraNode.generated.h"
+
+/**
+ * A camera node that offsets the location and rotation of the camera by evaluating curves.
+ */
+UCLASS(MinimalAPI, meta=(CameraNodeCategories="Common,Transform"))
+class USplineOffsetCameraNode : public UCameraNode
+{
+	GENERATED_BODY()
+
+protected:
+
+	// UCameraNode interface.
+	virtual FCameraNodeEvaluatorPtr OnBuildEvaluator(FCameraNodeEvaluatorBuilder& Builder) const override;
+
+public:
+
+	/** The input to pass to the offset splines. */
+	UPROPERTY(EditAnywhere, Category=Common)
+	FFloatCameraParameter SplineInput;
+
+	/** The spline that defines the translation offset to apply to the camera. */
+	UPROPERTY(EditAnywhere, Category=Common)
+	FCameraVectorCurve TranslationOffsetSpline;
+
+	/** The rotation offset to apply to the camera. */
+	UPROPERTY(EditAnywhere, Category=Common)
+	FCameraRotatorCurve RotationOffsetSpline;
+
+	/** The space in which to apply the offset. */
+	UPROPERTY(EditAnywhere, Category=Common)
+	ECameraNodeSpace OffsetSpace = ECameraNodeSpace::CameraPose;
+};
+
