Cameras: update BP action database when compiling a camera asset

This updates the list of BP graph actions to show any new camera parameter setter nodes.

diff --git a/Engine/Plugins/Cameras/GameplayCameras/Source/GameplayCamerasEditor/Private/Toolkits/CameraAssetEditorToolkit.cpp b/Engine/Plugins/Cameras/GameplayCameras/Source/GameplayCamerasEditor/Private/Toolkits/CameraAssetEditorToolkit.cpp
--- a/Engine/Plugins/Cameras/GameplayCameras/Source/GameplayCamerasEditor/Private/Toolkits/CameraAssetEditorToolkit.cpp
+++ b/Engine/Plugins/Cameras/GameplayCameras/Source/GameplayCamerasEditor/Private/Toolkits/CameraAssetEditorToolkit.cpp
@@ -4,6 +4,7 @@
 
 #include "AssetRegistry/AssetRegistryModule.h"
 #include "AssetTools/CameraAssetEditor.h"
+#include "BlueprintActionDatabase.h"
 #include "Build/CameraAssetBuilder.h"
 #include "Build/CameraBuildLog.h"
 #include "Commands/CameraAssetEditorCommands.h"
@@ -360,6 +361,11 @@ void FCameraAssetEditorToolkit::OnBuild()
 			LiveEditManager->NotifyPostBuildAsset(BuiltPackage);
 		}
 	}
+
+	if (FBlueprintActionDatabase* Database = FBlueprintActionDatabase::TryGet())
+	{
+		Database->RefreshAssetActions(CameraAsset);
+	}
 }
 
 void FCameraAssetEditorToolkit::OnFindInCamera()
diff --git a/Engine/Plugins/Cameras/GameplayCameras/Source/GameplayCamerasEditor/Private/Toolkits/CameraRigAssetEditorToolkit.cpp b/Engine/Plugins/Cameras/GameplayCameras/Source/GameplayCamerasEditor/Private/Toolkits/CameraRigAssetEditorToolkit.cpp
--- a/Engine/Plugins/Cameras/GameplayCameras/Source/GameplayCamerasEditor/Private/Toolkits/CameraRigAssetEditorToolkit.cpp
+++ b/Engine/Plugins/Cameras/GameplayCameras/Source/GameplayCamerasEditor/Private/Toolkits/CameraRigAssetEditorToolkit.cpp
@@ -2,6 +2,7 @@
 
 #include "Toolkits/CameraRigAssetEditorToolkit.h"
 
+#include "BlueprintActionDatabase.h"
 #include "Build/CameraBuildLog.h"
 #include "Build/CameraRigAssetBuilder.h"
 #include "Commands/CameraRigAssetEditorCommands.h"
@@ -337,6 +338,11 @@ void FCameraRigAssetEditorToolkit::OnBuild()
 	{
 		LiveEditManager->NotifyPostBuildAsset(BuiltPackage);
 	}
+
+	if (FBlueprintActionDatabase* Database = FBlueprintActionDatabase::TryGet())
+	{
+		Database->RefreshAssetActions(CameraRigAsset);
+	}
 }
 
 void FCameraRigAssetEditorToolkit::OnFindInCameraRig()
