Cameras: add camera variable setters and adjust how parameter blending is handled

diff --git a/Engine/Plugins/Cameras/GameplayCameras/Source/GameplayCameras/Private/Core/BlendStackCameraNode.cpp b/Engine/Plugins/Cameras/GameplayCameras/Source/GameplayCameras/Private/Core/BlendStackCameraNode.cpp
--- a/Engine/Plugins/Cameras/GameplayCameras/Source/GameplayCameras/Private/Core/BlendStackCameraNode.cpp
+++ b/Engine/Plugins/Cameras/GameplayCameras/Source/GameplayCameras/Private/Core/BlendStackCameraNode.cpp
@@ -94,16 +94,15 @@ bool FBlendStackCameraNodeEvaluator::InitializeEntry(
 	}
 
 	// Allocate variable table and context data table.
+	NewEntry.ContextResult.VariableTable.Initialize(CameraRig->AllocationInfo.VariableTableInfo);
+	NewEntry.ContextResult.ContextDataTable.Initialize(CameraRig->AllocationInfo.ContextDataTableInfo);
 	NewEntry.Result.VariableTable.Initialize(CameraRig->AllocationInfo.VariableTableInfo);
 	NewEntry.Result.ContextDataTable.Initialize(CameraRig->AllocationInfo.ContextDataTableInfo);
 
-	// Set the default values, so that pre-blending from default values works.
-	FCameraObjectInterfaceParameterOverrideHelper::ApplyDefaultBlendableParameters(CameraRig, NewEntry.Result.VariableTable);
-
 	// Set all the data from the context.
 	const FCameraNodeEvaluationResult& ContextResult = EvaluationContext->GetInitialResult();
-	NewEntry.Result.VariableTable.OverrideAll(ContextResult.VariableTable, true);
-	NewEntry.Result.ContextDataTable.OverrideAll(ContextResult.ContextDataTable);
+	NewEntry.ContextResult.VariableTable.OverrideAll(ContextResult.VariableTable, true);
+	NewEntry.ContextResult.ContextDataTable.OverrideAll(ContextResult.ContextDataTable);
 
 	// Add some conditional result if necessary.
 	if (bSetActiveResult && EvaluationContext)
@@ -111,8 +110,8 @@ bool FBlendStackCameraNodeEvaluator::InitializeEntry(
 		const FCameraNodeEvaluationResult* ActiveOnlyResult = EvaluationContext->GetConditionalResult(ECameraEvaluationDataCondition::ActiveCameraRig);
 		if (ActiveOnlyResult)
 		{
-			NewEntry.Result.VariableTable.OverrideAll(ActiveOnlyResult->VariableTable, true);
-			NewEntry.Result.ContextDataTable.OverrideAll(ActiveOnlyResult->ContextDataTable);
+			NewEntry.ContextResult.VariableTable.OverrideAll(ActiveOnlyResult->VariableTable, true);
+			NewEntry.ContextResult.ContextDataTable.OverrideAll(ActiveOnlyResult->ContextDataTable);
 		}
 	}
 
@@ -121,7 +120,13 @@ bool FBlendStackCameraNodeEvaluator::InitializeEntry(
 	InitParams.Evaluator = OwningEvaluator;
 	InitParams.EvaluationContext = EvaluationContext;
 	InitParams.LastActiveCameraRigInfo = GetActiveCameraRigEvaluationInfo();
-	RootEvaluator->Initialize(InitParams, NewEntry.Result);
+	RootEvaluator->Initialize(InitParams, NewEntry.ContextResult);  // Initializing with the context result here.
+
+	// Set default values for unset entries in the variable table, so that pre-blending from default 
+	// values works.
+	FCameraObjectInterfaceParameterOverrideHelper::ApplyDefaultBlendableParameters(CameraRig, NewEntry.ContextResult.VariableTable);
+
+	NewEntry.Result.OverrideAll(NewEntry.ContextResult, true);
 
 	// Wrap up!
 	NewEntry.EntryID = FBlendStackEntryID(NextEntryID++);
@@ -217,6 +222,9 @@ void FBlendStackCameraNodeEvaluator::OnInitialize(const FCameraNodeEvaluatorInit
 
 void FBlendStackCameraNodeEvaluator::ResolveEntries(TArray<FResolvedEntry>& OutResolvedEntries)
 {
+	constexpr ECameraVariableTableFilter VariableTableFilter = ECameraVariableTableFilter::ChangedOnly;
+	constexpr ECameraContextDataTableFilter ContextDataTableFilter = ECameraContextDataTableFilter::ChangedOnly;
+
 	// Build up these structures so we don't re-resolve evaluation context weak-pointers
 	// multiple times in this function..
 	for (int32 Index = 0; Index < Entries.Num(); ++Index)
@@ -281,6 +289,21 @@ void FBlendStackCameraNodeEvaluator::ResolveEntries(TArray<FResolvedEntry>& OutR
 			// Reset this entry's flags for this frame.
 			FCameraNodeEvaluationResult& CurResult = Entry.Result;
 			CurResult.ResetFrameFlags();
+
+			// Bring the entry's context result up to date with any changes.
+			Entry.ContextResult.CameraPose.OverrideChanged(ContextResult.CameraPose);
+			Entry.ContextResult.VariableTable.Override(ContextResult.VariableTable, VariableTableFilter);
+			Entry.ContextResult.ContextDataTable.Override(ContextResult.ContextDataTable, ContextDataTableFilter);
+			if (ResolvedEntry.bIsActiveEntry)
+			{
+				if (const FCameraNodeEvaluationResult* ActiveOnlyResult = ResolvedEntry.Context->GetConditionalResult(ECameraEvaluationDataCondition::ActiveCameraRig))
+				{
+					Entry.ContextResult.VariableTable.Override(ActiveOnlyResult->VariableTable, VariableTableFilter);
+					Entry.ContextResult.ContextDataTable.Override(ActiveOnlyResult->ContextDataTable, ContextDataTableFilter);
+				}
+			}
+			Entry.ContextResult.bIsCameraCut = ContextResult.bIsCameraCut;
+			Entry.ContextResult.bIsValid = ContextResult.bIsValid;
 		}
 		// else: frozen entries may have null contexts or invalid initial results
 		//       because we're not going to update them anyway. We will however blend
@@ -481,6 +504,7 @@ void FBlendStackCameraNodeEvaluator::OnSerialize(const FCameraNodeEvaluatorSeria
 
 	for (FCameraRigEntry& Entry : Entries)
 	{
+		Entry.ContextResult.Serialize(Ar);
 		Entry.Result.Serialize(Ar);
 		Ar << Entry.bIsFirstFrame;
 		Ar << Entry.bIsFrozen;
diff --git a/Engine/Plugins/Cameras/GameplayCameras/Source/GameplayCameras/Private/Core/CameraVariableSetter.cpp b/Engine/Plugins/Cameras/GameplayCameras/Source/GameplayCameras/Private/Core/CameraVariableSetter.cpp
new file mode 100644
--- /dev/null
+++ b/Engine/Plugins/Cameras/GameplayCameras/Source/GameplayCameras/Private/Core/CameraVariableSetter.cpp
@@ -0,0 +1,111 @@
+// Copyright Epic Games, Inc. All Rights Reserved.
+
+#include "Core/CameraVariableSetter.h"
+
+#include "Math/Interpolation.h"
+#include "Math/UnrealMath.h"
+
+namespace UE::Cameras
+{
+
+void FCameraVariableSetter::Update(float DeltaTime)
+{
+	UpdateState(DeltaTime);
+}
+
+void FCameraVariableSetter::Stop(bool bImmediately)
+{
+	EState PrevState = State;
+	if (PrevState != EState::Inactive)
+	{
+		if (!bImmediately)
+		{
+			State = EState::BlendOut;
+			if (PrevState == EState::Full)
+			{
+				CurrentTime = 0.f;
+			}
+			else if (PrevState == EState::BlendIn)
+			{
+				// TODO: this assumes a symmetrical blend curve, which may not always be true.
+				CurrentTime = FMath::Clamp(BlendTime - CurrentTime, 0.f, BlendTime);
+			}
+			// else: was blending out already.
+		}
+		else
+		{
+			State = EState::Inactive;
+		}
+	}
+	// else: we were already stopped.
+}
+
+void FCameraVariableSetter::UpdateState(float DeltaTime)
+{
+	if (State == EState::BlendIn || State == EState::BlendOut)
+	{
+		const float NewTime = CurrentTime + DeltaTime;
+		if (State == EState::BlendIn)
+		{
+			if (NewTime < BlendTime)
+			{
+				CurrentTime = NewTime;
+			}
+			else
+			{
+				State = EState::Full;
+				CurrentTime = BlendTime;
+			}
+		}
+		else if (State == EState::BlendOut)
+		{
+			if (NewTime < BlendTime)
+			{
+				CurrentTime = NewTime;
+			}
+			else
+			{
+				State = EState::Inactive;
+				CurrentTime = BlendTime;
+			}
+		}
+	}
+}
+
+float FCameraVariableSetter::GetBlendFactor()
+{
+	float BlendPercent = 0.f;
+	switch (State)
+	{
+		case EState::BlendIn:
+			BlendPercent = (BlendTime > 0.f) ? (CurrentTime / BlendTime) : 1.f;
+			break;
+		case EState::Full:
+			BlendPercent = 1.f;
+			break;
+		case EState::BlendOut:
+			BlendPercent = 1.f - ((BlendTime > 0.f) ? (CurrentTime / BlendTime) : 1.f);
+			break;
+		case EState::Inactive:
+		default:
+			break;
+	}
+
+	BlendPercent = FMath::Clamp(BlendPercent, 0.f, 1.f);
+	switch (BlendType)
+	{
+		case ECameraVariableSetterBlendType::None:
+			return BlendPercent >= 1.f ? 1.f : 0.f;
+		case ECameraVariableSetterBlendType::Linear:
+			return BlendPercent;
+		case ECameraVariableSetterBlendType::SmoothStep:
+			return SmoothStep(BlendPercent);
+		case ECameraVariableSetterBlendType::SmootherStep:
+			return SmootherStep(BlendPercent);
+		default:
+			return 1.f;
+	}
+}
+
+}  // namespace UE::Cameras
+
diff --git a/Engine/Plugins/Cameras/GameplayCameras/Source/GameplayCameras/Private/Core/DefaultRootCameraNode.cpp b/Engine/Plugins/Cameras/GameplayCameras/Source/GameplayCameras/Private/Core/DefaultRootCameraNode.cpp
--- a/Engine/Plugins/Cameras/GameplayCameras/Source/GameplayCameras/Private/Core/DefaultRootCameraNode.cpp
+++ b/Engine/Plugins/Cameras/GameplayCameras/Source/GameplayCameras/Private/Core/DefaultRootCameraNode.cpp
@@ -76,6 +76,8 @@ FCameraNodeEvaluatorChildrenView FDefaultRootCameraNodeEvaluator::OnGetChildren(
 
 void FDefaultRootCameraNodeEvaluator::OnRun(const FCameraNodeEvaluationParams& Params, FCameraNodeEvaluationResult& OutResult)
 {
+	UpdateCameraVariableSetters(Params.DeltaTime);
+
 	BaseLayer->Run(Params, OutResult);
 	MainLayer->Run(Params, OutResult);
 	GlobalLayer->Run(Params, OutResult);
diff --git a/Engine/Plugins/Cameras/GameplayCameras/Source/GameplayCameras/Private/Core/PersistentBlendStackCameraNode.cpp b/Engine/Plugins/Cameras/GameplayCameras/Source/GameplayCameras/Private/Core/PersistentBlendStackCameraNode.cpp
--- a/Engine/Plugins/Cameras/GameplayCameras/Source/GameplayCameras/Private/Core/PersistentBlendStackCameraNode.cpp
+++ b/Engine/Plugins/Cameras/GameplayCameras/Source/GameplayCameras/Private/Core/PersistentBlendStackCameraNode.cpp
@@ -204,9 +204,6 @@ void FPersistentBlendStackCameraNodeEvaluator::OnRun(const FCameraNodeEvaluation
 
 void FPersistentBlendStackCameraNodeEvaluator::InternalUpdate(TArrayView<FResolvedEntry> ResolvedEntries, const FCameraNodeEvaluationParams& Params, FCameraNodeEvaluationResult& OutResult)
 {
-	constexpr ECameraVariableTableFilter VariableTableFilter = ECameraVariableTableFilter::ChangedOnly;
-	constexpr ECameraContextDataTableFilter ContextDataTableFilter = ECameraContextDataTableFilter::ChangedOnly;
-
 	TArray<int32, TInlineAllocator<4>> EntriesToRemove;
 
 	for (FResolvedEntry& ResolvedEntry : ResolvedEntries)
@@ -222,21 +219,21 @@ void FPersistentBlendStackCameraNodeEvaluator::InternalUpdate(TArrayView<FResolv
 
 			FCameraNodeEvaluationResult& CurResult(Entry.Result);
 
-			// Start with the input given to us.
+			// Update the current entry's result.
 			{
 				CurResult.Reset();
+
+				// Start with the input given to us.
 				CurResult.CameraPose = OutResult.CameraPose;
 				CurResult.VariableTable.OverrideAll(OutResult.VariableTable);
 				CurResult.CameraRigJoints.OverrideAll(OutResult.CameraRigJoints);
 				CurResult.PostProcessSettings.OverrideAll(OutResult.PostProcessSettings);
 
 				// Override it with whatever the evaluation context has set on its result.
-				// Evaluation contexts may have private variables we need to pass along, such as when rig parameter
-				// overrides have been set on them, so include private variables in the filter.
-				const FCameraNodeEvaluationResult& ContextResult(ResolvedEntry.Context->GetInitialResult());
+				const FCameraNodeEvaluationResult& ContextResult(Entry.ContextResult);
 				CurResult.CameraPose.OverrideChanged(ContextResult.CameraPose);
-				CurResult.VariableTable.Override(ContextResult.VariableTable, VariableTableFilter);
-				CurResult.ContextDataTable.Override(ContextResult.ContextDataTable, ContextDataTableFilter);
+				CurResult.VariableTable.OverrideAll(ContextResult.VariableTable, true);
+				CurResult.ContextDataTable.OverrideAll(ContextResult.ContextDataTable);
 
 				// Setup flags.
 				CurResult.bIsCameraCut = OutResult.bIsCameraCut || ContextResult.bIsCameraCut || Entry.bForceCameraCut;
diff --git a/Engine/Plugins/Cameras/GameplayCameras/Source/GameplayCameras/Private/Core/RootCameraNode.cpp b/Engine/Plugins/Cameras/GameplayCameras/Source/GameplayCameras/Private/Core/RootCameraNode.cpp
--- a/Engine/Plugins/Cameras/GameplayCameras/Source/GameplayCameras/Private/Core/RootCameraNode.cpp
+++ b/Engine/Plugins/Cameras/GameplayCameras/Source/GameplayCameras/Private/Core/RootCameraNode.cpp
@@ -100,5 +100,45 @@ void FRootCameraNodeEvaluator::BroadcastCameraRigEvent(const FRootCameraNodeCame
 	OnCameraRigEventDelegate.Broadcast(InEvent);
 }
 
+void FRootCameraNodeEvaluator::StopCameraVariableSetter(const FCameraVariableSetterHandle& InHandle, bool bImmediately)
+{
+	if (VariableSetters.IsValidIndex(InHandle.Value))
+	{
+		FVariableSetterEntry& Entry(VariableSetters[InHandle.Value]);
+		if (Entry.ThisHandle.SerialNumber == InHandle.SerialNumber)
+		{
+			Entry.Setter->Stop(bImmediately);
+		}
+	}
+}
+
+void FRootCameraNodeEvaluator::UpdateCameraVariableSetters(float DeltaTime)
+{
+	for (auto It = VariableSetters.CreateIterator(); It; ++It)
+	{
+		FVariableSetterEntry& Entry(*It);
+		if (!Entry.Setter.IsValid())
+		{
+			It.RemoveCurrent();
+			continue;
+		}
+
+		Entry.Setter->Update(DeltaTime);
+
+		if (!Entry.Setter->IsActive())
+		{
+			It.RemoveCurrent();
+		}
+	}
+}
+
+void FRootCameraNodeEvaluator::ApplyCameraVariableSetters(FCameraVariableTable& OutVariableTable)
+{
+	for (FVariableSetterEntry& Entry : VariableSetters)
+	{
+		Entry.Setter->Apply(OutVariableTable);
+	}
+}
+
 }  // namespace UE::Cameras
 
diff --git a/Engine/Plugins/Cameras/GameplayCameras/Source/GameplayCameras/Private/Core/TransientBlendStackCameraNode.cpp b/Engine/Plugins/Cameras/GameplayCameras/Source/GameplayCameras/Private/Core/TransientBlendStackCameraNode.cpp
--- a/Engine/Plugins/Cameras/GameplayCameras/Source/GameplayCameras/Private/Core/TransientBlendStackCameraNode.cpp
+++ b/Engine/Plugins/Cameras/GameplayCameras/Source/GameplayCameras/Private/Core/TransientBlendStackCameraNode.cpp
@@ -11,6 +11,8 @@
 #include "Core/CameraRigAsset.h"
 #include "Core/CameraRigCombinationRegistry.h"
 #include "Core/CameraRigTransition.h"
+#include "Core/CameraSystemEvaluator.h"
+#include "Core/RootCameraNode.h"
 #include "Helpers/CameraRigTransitionFinder.h"
 #include "Nodes/Blends/PopBlendCameraNode.h"
 
@@ -111,21 +113,6 @@ FBlendStackEntryID FTransientBlendStackCameraNodeEvaluator::PushNewEntry(const F
 		return FBlendStackEntryID();
 	}
 	
-	// Take a snapshot of the initial result for having base values during pre-blending.
-	FCameraRigEntryExtraInfo NewExtraInfo;
-	{
-		FCameraNodeEvaluationResult& UnblendedResult = NewExtraInfo.UnblendedResult;
-
-		if (Params.CameraRig)
-		{
-			const FCameraObjectAllocationInfo& AllocationInfo = Params.CameraRig->AllocationInfo;
-			UnblendedResult.VariableTable.Initialize(AllocationInfo.VariableTableInfo);
-			UnblendedResult.ContextDataTable.Initialize(AllocationInfo.ContextDataTableInfo);
-		}
-
-		UnblendedResult.OverrideAll(NewEntry.Result, true);
-	}
-
 #if WITH_EDITOR
 	// Listen to changes to the packages inside which this camera rig is defined. We will hot-reload the
 	// camera node evaluators for this camera rig when we detect changes.
@@ -137,7 +124,7 @@ FBlendStackEntryID FTransientBlendStackCameraNodeEvaluator::PushNewEntry(const F
 	// Important: we need to move the new entry here because copying evaluator storage
 	// is disabled.
 	Entries.Add(MoveTemp(NewEntry));
-	EntryExtraInfos.Add(MoveTemp(NewExtraInfo));
+	EntryExtraInfos.Emplace();
 
 	if (OnCameraRigEventDelegate.IsBound())
 	{
@@ -245,6 +232,9 @@ void FTransientBlendStackCameraNodeEvaluator::InternalPreBlendPrepare(TArrayView
 	constexpr ECameraVariableTableFilter VariableTableFilter = ECameraVariableTableFilter::ChangedOnly;
 	constexpr ECameraContextDataTableFilter ContextDataTableFilter = ECameraContextDataTableFilter::ChangedOnly;
 
+	FRootCameraNodeEvaluator* RootNodeEvaluator = Params.Evaluator->GetRootNodeEvaluator();
+	check(RootNodeEvaluator);
+
 	for (FResolvedEntry& ResolvedEntry : ResolvedEntries)
 	{
 		FCameraRigEntry& Entry(ResolvedEntry.Entry);
@@ -259,35 +249,25 @@ void FTransientBlendStackCameraNodeEvaluator::InternalPreBlendPrepare(TArrayView
 		CurParams.EvaluationContext = ResolvedEntry.Context;
 		CurParams.bIsFirstFrame = Entry.bIsFirstFrame;
 
-		FCameraNodeEvaluationResult& UnblendedResult(EntryExtraInfo.UnblendedResult);
-
-		// Reset frame flags to match what happend to the Entry.Result in ResolveEntries.
-		UnblendedResult.ResetFrameFlags();
+		FCameraNodeEvaluationResult& CurResult(Entry.Result);
 
 		// Start with the input given to us.
-		UnblendedResult.VariableTable.OverrideAll(OutResult.VariableTable);
+		CurResult.VariableTable.OverrideAll(OutResult.VariableTable);
 
 		// Override it with whatever the evaluation context has set on its result this frame.
-		// Evaluation contexts may have private variables we need to pass along, such as when rig parameter
-		// overrides have been set on them, so include private variables in the filter.
-		const FCameraNodeEvaluationResult& ContextResult(ResolvedEntry.Context->GetInitialResult());
-		UnblendedResult.VariableTable.Override(ContextResult.VariableTable, VariableTableFilter);
-		UnblendedResult.ContextDataTable.Override(ContextResult.ContextDataTable, ContextDataTableFilter);
-		if (ResolvedEntry.bIsActiveEntry)
-		{
-			if (const FCameraNodeEvaluationResult* ActiveOnlyResult = ResolvedEntry.Context->GetConditionalResult(ECameraEvaluationDataCondition::ActiveCameraRig))
-			{
-				UnblendedResult.VariableTable.Override(ActiveOnlyResult->VariableTable, VariableTableFilter);
-				UnblendedResult.ContextDataTable.Override(ActiveOnlyResult->ContextDataTable, ContextDataTableFilter);
-			}
-		}
+		const FCameraNodeEvaluationResult& ContextResult(Entry.ContextResult);
+		CurResult.VariableTable.OverrideAll(ContextResult.VariableTable, true);
+		CurResult.ContextDataTable.OverrideAll(ContextResult.ContextDataTable);
+
+		// Override it with variable setters.
+		RootNodeEvaluator->ApplyCameraVariableSetters(CurResult.VariableTable);
 
 		// Gather input parameters if needed (and remember if it was indeed needed).
 		if (!EntryExtraInfo.bInputRunThisFrame)
 		{
 			bool bHasPreBlendedParameters = false;
-			FCameraBlendedParameterUpdateParams InputParams(CurParams, UnblendedResult.CameraPose);
-			FCameraBlendedParameterUpdateResult InputResult(UnblendedResult.VariableTable);
+			FCameraBlendedParameterUpdateParams InputParams(CurParams, CurResult.CameraPose);
+			FCameraBlendedParameterUpdateResult InputResult(CurResult.VariableTable);
 
 			Entry.EvaluatorHierarchy.ForEachEvaluator(ECameraNodeEvaluatorFlags::NeedsParameterUpdate,
 					[&bHasPreBlendedParameters, &InputParams, &InputResult](FCameraNodeEvaluator* ParameterEvaluator)
@@ -308,29 +288,11 @@ void FTransientBlendStackCameraNodeEvaluator::InternalPreBlendPrepare(TArrayView
 			FBlendCameraNodeEvaluator* EntryBlendEvaluator = Entry.RootEvaluator->GetBlendEvaluator();
 			if (EntryBlendEvaluator)
 			{
-				EntryBlendEvaluator->Run(CurParams, UnblendedResult);
+				EntryBlendEvaluator->Run(CurParams, CurResult);
 			}
 
 			EntryExtraInfo.bBlendRunThisFrame = true;
 		}
-
-		// Assign the new unblended result back to the entry result. The goal is that:
-		//
-		// 1) We only write values from the context result when they have been written to.
-		//    That is, we don't always set all values every frame. When "left alone", a camera rig
-		//    should retain whatever values it was last set with. For instance, if a camera rig 
-		//    received special values when it was the active camera rig, we want to keep those 
-		//    values when it starts blending out, instead of immediately overwriting them with
-		//    the "common" values (i.e. the values from the initial result).
-		//
-		// 2) However, we do want to write all the values into the PreBlendVariableTable every
-		//    frame, so that we can have them blend together. The blended values are then reassigned
-		//    to Entry.Result in InternalPreBlendExecute, which means we lose each camera rig's
-		//    "intended" values. Getting them from the evaluation context each frame would violate
-		//    previous item (1) so that's why we use the UnblendedResult as a snapshot.
-		//
-		FCameraNodeEvaluationResult& CurResult(Entry.Result);
-		CurResult.OverrideAll(UnblendedResult, true);
 	}
 }
 
@@ -353,7 +315,7 @@ void FTransientBlendStackCameraNodeEvaluator::InternalPreBlendExecute(TArrayView
 		FCameraRigEntry& Entry(ResolvedEntry.Entry);
 		FCameraRigEntryExtraInfo& EntryExtraInfo(EntryExtraInfos[ResolvedEntry.EntryIndex]);
 
-		FCameraNodeEvaluationResult& UnblendedResult(EntryExtraInfo.UnblendedResult);
+		FCameraNodeEvaluationResult& CurResult(Entry.Result);
 
 		if (!Entry.bIsFrozen)
 		{
@@ -361,7 +323,7 @@ void FTransientBlendStackCameraNodeEvaluator::InternalPreBlendExecute(TArrayView
 			CurParams.EvaluationContext = ResolvedEntry.Context;
 			CurParams.bIsFirstFrame = Entry.bIsFirstFrame;
 
-			FCameraNodePreBlendParams PreBlendParams(CurParams, UnblendedResult.CameraPose, UnblendedResult.VariableTable);
+			FCameraNodePreBlendParams PreBlendParams(CurParams, CurResult.CameraPose, CurResult.VariableTable);
 			PreBlendParams.VariableTableFilter = VariableTableFilter;
 
 			FCameraNodePreBlendResult PreBlendResult(PreBlendVariableTable);
@@ -374,14 +336,14 @@ void FTransientBlendStackCameraNodeEvaluator::InternalPreBlendExecute(TArrayView
 			}
 			else
 			{
-				PreBlendVariableTable.Override(UnblendedResult.VariableTable, PreBlendParams.VariableTableFilter);
+				PreBlendVariableTable.Override(CurResult.VariableTable, PreBlendParams.VariableTableFilter);
 				EntryExtraInfo.bIsPreBlendFull = true;
 			}
 		}
 		else
 		{
 			// Frozen entries still contribute to the blend using their last evaluated values.
-			PreBlendVariableTable.Override(UnblendedResult.VariableTable, VariableTableFilter);
+			PreBlendVariableTable.Override(CurResult.VariableTable, VariableTableFilter);
 		}
 	}
 
@@ -391,14 +353,10 @@ void FTransientBlendStackCameraNodeEvaluator::InternalPreBlendExecute(TArrayView
 	for (FResolvedEntry& ResolvedEntry : ResolvedEntries)
 	{
 		FCameraRigEntry& Entry(ResolvedEntry.Entry);
-		FCameraRigEntryExtraInfo& EntryExtraInfo(EntryExtraInfos[ResolvedEntry.EntryIndex]);
-
 		if (!Entry.bIsFrozen)
 		{
-			FCameraNodeEvaluationResult& UnblendedResult(EntryExtraInfo.UnblendedResult);
 			FCameraNodeEvaluationResult& CurResult(Entry.Result);
-
-			Entry.Result.VariableTable.Override(PreBlendVariableTable, ECameraVariableTableFilter::KnownOnly);
+			CurResult.VariableTable.Override(PreBlendVariableTable, ECameraVariableTableFilter::KnownOnly);
 		}
 	}
 }
@@ -428,17 +386,9 @@ void FTransientBlendStackCameraNodeEvaluator::InternalUpdate(TArrayView<FResolve
 		CurResult.PostProcessSettings.OverrideAll(OutResult.PostProcessSettings);
 
 		// Override it with whatever the evaluation context has set on its result.
-		const FCameraNodeEvaluationResult& ContextResult(ResolvedEntry.Context->GetInitialResult());
+		const FCameraNodeEvaluationResult& ContextResult(Entry.ContextResult);
 		CurResult.CameraPose.OverrideChanged(ContextResult.CameraPose);
 		CurResult.bIsCameraCut = OutResult.bIsCameraCut || ContextResult.bIsCameraCut || Entry.bForceCameraCut;
-		if (ResolvedEntry.bIsActiveEntry)
-		{
-			if (const FCameraNodeEvaluationResult* ActiveOnlyResult = ResolvedEntry.Context->GetConditionalResult(ECameraEvaluationDataCondition::ActiveCameraRig))
-			{
-				CurResult.CameraPose.OverrideChanged(ActiveOnlyResult->CameraPose);
-				CurResult.bIsCameraCut |= ActiveOnlyResult->bIsCameraCut;
-			}
-		}
 		
 		CurResult.bIsValid = true;
 
@@ -649,33 +599,5 @@ const UCameraRigTransition* FTransientBlendStackCameraNodeEvaluator::FindTransit
 	return nullptr;
 }
 
-#if WITH_EDITOR
-
-void FTransientBlendStackCameraNodeEvaluator::OnEntryReinitialized(int32 EntryIndex)
-{
-	if (!ensure(EntryExtraInfos.IsValidIndex(EntryIndex)))
-	{
-		return;
-	}
-
-	FCameraRigEntry& Entry = Entries[EntryIndex];
-	FCameraRigEntryExtraInfo& ExtraInfo = EntryExtraInfos[EntryIndex];
-	{
-		// Update our extra info with the new data.
-		FCameraNodeEvaluationResult& UnblendedResult = ExtraInfo.UnblendedResult;
-
-		if (Entry.CameraRig)
-		{
-			const FCameraObjectAllocationInfo& AllocationInfo = Entry.CameraRig->AllocationInfo;
-			UnblendedResult.VariableTable.Initialize(AllocationInfo.VariableTableInfo);
-			UnblendedResult.ContextDataTable.Initialize(AllocationInfo.ContextDataTableInfo);
-		}
-
-		UnblendedResult.OverrideAll(Entry.Result, true);
-	}
-}
-
-#endif
-
 }  // namespace UE::Cameras
 
diff --git a/Engine/Plugins/Cameras/GameplayCameras/Source/GameplayCameras/Private/Helpers/CameraObjectInterfaceParameterOverrideHelper.cpp b/Engine/Plugins/Cameras/GameplayCameras/Source/GameplayCameras/Private/Helpers/CameraObjectInterfaceParameterOverrideHelper.cpp
--- a/Engine/Plugins/Cameras/GameplayCameras/Source/GameplayCameras/Private/Helpers/CameraObjectInterfaceParameterOverrideHelper.cpp
+++ b/Engine/Plugins/Cameras/GameplayCameras/Source/GameplayCameras/Private/Helpers/CameraObjectInterfaceParameterOverrideHelper.cpp
@@ -480,6 +480,11 @@ void FCameraObjectInterfaceParameterOverrideHelper::ApplyDefaultParametersImpl(c
 			{
 				continue;
 			}
+			
+			if (OutVariableTable->IsValueWritten(Definition.VariableID))
+			{
+				continue;
+			}
 
 			const FPropertyBagPropertyDesc* PropertyDesc = DefaultParameters.FindPropertyDescByID(Definition.ParameterGuid);
 			if (!ensure(PropertyDesc && PropertyDesc->CachedProperty))
@@ -497,6 +502,11 @@ void FCameraObjectInterfaceParameterOverrideHelper::ApplyDefaultParametersImpl(c
 				continue;
 			}
 
+			if (OutContextDataTable->IsValueWritten(Definition.DataID))
+			{
+				continue;
+			}
+
 			const FPropertyBagPropertyDesc* PropertyDesc = DefaultParameters.FindPropertyDescByID(Definition.ParameterGuid);
 			if (!ensure(PropertyDesc && PropertyDesc->CachedProperty))
 			{
diff --git a/Engine/Plugins/Cameras/GameplayCameras/Source/GameplayCameras/Private/Nodes/Common/CameraRigCameraNode.cpp b/Engine/Plugins/Cameras/GameplayCameras/Source/GameplayCameras/Private/Nodes/Common/CameraRigCameraNode.cpp
--- a/Engine/Plugins/Cameras/GameplayCameras/Source/GameplayCameras/Private/Nodes/Common/CameraRigCameraNode.cpp
+++ b/Engine/Plugins/Cameras/GameplayCameras/Source/GameplayCameras/Private/Nodes/Common/CameraRigCameraNode.cpp
@@ -8,6 +8,7 @@
 #include "Core/CameraRigAsset.h"
 #include "Core/CameraRigAssetReference.h"
 #include "Helpers/CameraObjectReferenceParameterOverrideEvaluator.h"
+#include "Helpers/CameraObjectInterfaceParameterOverrideHelper.h"
 
 #include UE_INLINE_GENERATED_CPP_BY_NAME(CameraRigCameraNode)
 
@@ -44,6 +45,14 @@ void FCameraRigCameraNodeEvaluator::OnInitialize(const FCameraNodeEvaluatorIniti
 {
 	// Apply overrides right away.
 	ApplyParameterOverrides(OutResult.VariableTable, OutResult.ContextDataTable, false);
+
+	const UCameraRigCameraNode* PrefabNode = GetCameraNodeAs<UCameraRigCameraNode>();
+	if (const UCameraRigAsset* CameraRig = PrefabNode->CameraRigReference.GetCameraRig())
+	{
+		// Set default values for unset entries in the variable table, so that pre-blending from default 
+		// values works.
+		FCameraObjectInterfaceParameterOverrideHelper::ApplyDefaultBlendableParameters(CameraRig, OutResult.VariableTable);
+	}
 }
 
 void FCameraRigCameraNodeEvaluator::OnUpdateParameters(const FCameraBlendedParameterUpdateParams& Params, FCameraBlendedParameterUpdateResult& OutResult)
diff --git a/Engine/Plugins/Cameras/GameplayCameras/Source/GameplayCameras/Public/Core/BlendStackCameraNode.h b/Engine/Plugins/Cameras/GameplayCameras/Source/GameplayCameras/Public/Core/BlendStackCameraNode.h
--- a/Engine/Plugins/Cameras/GameplayCameras/Source/GameplayCameras/Public/Core/BlendStackCameraNode.h
+++ b/Engine/Plugins/Cameras/GameplayCameras/Source/GameplayCameras/Public/Core/BlendStackCameraNode.h
@@ -203,7 +203,9 @@ protected:
 		FBlendStackRootCameraNodeEvaluator* RootEvaluator = nullptr;
 		/** The evaluator tree. */
 		FCameraNodeEvaluatorHierarchy EvaluatorHierarchy;
-		/** Result for this node tree. */
+		/** Context result for this node tree. */
+		FCameraNodeEvaluationResult ContextResult;
+		/** Evaluation result for this node tree. */
 		FCameraNodeEvaluationResult Result;
 		/** Whether this is the first frame this entry runs. */
 		bool bIsFirstFrame = false;
diff --git a/Engine/Plugins/Cameras/GameplayCameras/Source/GameplayCameras/Public/Core/CameraVariableSetter.h b/Engine/Plugins/Cameras/GameplayCameras/Source/GameplayCameras/Public/Core/CameraVariableSetter.h
new file mode 100644
--- /dev/null
+++ b/Engine/Plugins/Cameras/GameplayCameras/Source/GameplayCameras/Public/Core/CameraVariableSetter.h
@@ -0,0 +1,175 @@
+// Copyright Epic Games, Inc. All Rights Reserved.
+
+#pragma once
+
+#include "Core/CameraVariableTable.h"
+#include "Misc/InlineValue.h"
+#include "Templates/UnrealTypeTraits.h"
+#include <type_traits>
+
+#include "CameraVariableSetter.generated.h"
+
+namespace UE::Cameras
+{
+	class FRootCameraNodeEvaluator;
+}
+
+UENUM()
+enum class ECameraVariableSetterBlendType
+{
+	None,
+	Linear,
+	SmoothStep,
+	SmootherStep
+};
+
+USTRUCT(BlueprintType)
+struct FCameraVariableSetterHandle
+{
+	GENERATED_BODY()
+
+public:
+
+	FCameraVariableSetterHandle() : Value(INVALID), SerialNumber(0) {}
+
+	bool IsValid() const { return Value != INVALID; }
+
+	explicit operator bool() const { return IsValid(); }
+
+public:
+
+	friend bool operator==(FCameraVariableSetterHandle A, FCameraVariableSetterHandle B)
+	{
+		return A.Value == B.Value && A.SerialNumber == B.SerialNumber;
+	}
+
+	friend bool operator!=(FCameraVariableSetterHandle A, FCameraVariableSetterHandle B)
+	{
+		return !(A == B);
+	}
+
+	friend uint32 GetTypeHash(FCameraVariableSetterHandle In)
+	{
+		return HashCombineFast(In.Value, In.SerialNumber);
+	}
+
+	friend FArchive& operator<< (FArchive& Ar, FCameraVariableSetterHandle& In)
+	{
+		Ar << In.Value;
+		Ar << In.SerialNumber;
+		return Ar;
+	}
+
+private:
+
+	FCameraVariableSetterHandle(uint32 InValue, uint32 InSerialNumber) 
+		: Value(InValue)
+		, SerialNumber(InSerialNumber)
+	{}
+
+	static const uint32 INVALID = uint32(-1);
+
+	UPROPERTY()
+	uint32 Value;
+
+	UPROPERTY()
+	uint32 SerialNumber;
+
+	friend class UE::Cameras::FRootCameraNodeEvaluator;
+};
+
+namespace UE::Cameras
+{
+
+struct FCameraVariableSetter
+{
+	ECameraVariableSetterBlendType BlendType = ECameraVariableSetterBlendType::Linear;
+	float BlendTime = 1.f;
+
+	virtual ~FCameraVariableSetter() {}
+
+	virtual void Apply(FCameraVariableTable& VariableTable) {}
+
+	void Update(float DeltaTime);
+	void Stop(bool bImmediately);
+	bool IsActive() const { return State != EState::Inactive; }
+	
+protected:
+
+	FCameraVariableSetter() {}
+	FCameraVariableSetter(FCameraVariableID InVariableID) : VariableID(InVariableID) {}
+
+	void UpdateState(float DeltaTime);
+	float GetBlendFactor();
+
+protected:
+
+	enum class EState { Inactive, BlendIn, Full, BlendOut };
+
+	float CurrentTime = 0.f;
+	EState State = EState::BlendIn;
+	FCameraVariableID VariableID;
+};
+
+using FCameraVariableSetterPtr = TInlineValue<FCameraVariableSetter, 64>;
+
+namespace Internal
+{
+
+template<typename ValueType>
+struct TVariableSetterValueInterpolator
+{
+	static ValueType Lerp(typename TCallTraits<ValueType>::ParamType FromValue, typename TCallTraits<ValueType>::ParamType ToValue, float BlendFactor)
+	{
+		return FMath::Lerp(FromValue, ToValue, BlendFactor);
+	}
+};
+
+template<>
+struct TVariableSetterValueInterpolator<bool>
+{
+	static bool Lerp(bool FromValue, bool ToValue, float BlendFactor)
+	{
+		return (BlendFactor >= 0.5f) ? ToValue : FromValue;
+	}
+};
+
+template<typename T>
+struct TVariableSetterValueInterpolator<UE::Math::TTransform<T>>
+{
+	static UE::Math::TTransform<T> Lerp(const UE::Math::TTransform<T>& FromValue, const UE::Math::TTransform<T>& ToValue, float BlendFactor)
+	{
+		UE::Math::TTransform<T> CurrentValue(FromValue);
+		CurrentValue.BlendWith(ToValue, BlendFactor);
+		return CurrentValue;
+	}
+};
+
+}  // namespace Internal
+
+template<typename ValueType>
+struct TCameraVariableSetter : public FCameraVariableSetter
+{
+	ValueType ToValue;
+
+	TCameraVariableSetter() = default;
+	TCameraVariableSetter(FCameraVariableID InVariableID, TCallTraits<ValueType>::ParamType InToValue)
+		: FCameraVariableSetter(InVariableID)
+		, ToValue(InToValue)
+	{}
+
+	virtual void Apply(FCameraVariableTable& VariableTable) override
+	{
+		ValueType FromValue;
+		const bool bIsValid = VariableTable.TryGetValue<ValueType>(VariableID, FromValue);
+		if (bIsValid)
+		{
+			const float Factor = GetBlendFactor();
+			const ValueType CurrentValue = Internal::TVariableSetterValueInterpolator<ValueType>::Lerp(FromValue, ToValue, Factor);
+			VariableTable.SetValue<ValueType>(VariableID, CurrentValue);
+		}
+	}
+};
+
+}  // namespace UE::Cameras
+
diff --git a/Engine/Plugins/Cameras/GameplayCameras/Source/GameplayCameras/Public/Core/RootCameraNode.h b/Engine/Plugins/Cameras/GameplayCameras/Source/GameplayCameras/Public/Core/RootCameraNode.h
--- a/Engine/Plugins/Cameras/GameplayCameras/Source/GameplayCameras/Public/Core/RootCameraNode.h
+++ b/Engine/Plugins/Cameras/GameplayCameras/Source/GameplayCameras/Public/Core/RootCameraNode.h
@@ -2,10 +2,12 @@
 
 #pragma once
 
+#include "Containers/SparseArray.h"
 #include "Core/CameraNode.h"
 #include "Core/CameraNodeEvaluator.h"
 #include "Core/CameraRigEvaluationInfo.h"
 #include "Core/CameraRigInstanceID.h"
+#include "Core/CameraVariableSetter.h"
 
 #include "RootCameraNode.generated.h"
 
@@ -158,6 +160,15 @@ public:
 	/** Gets the delegate for camera rig events. */
 	FOnRootCameraNodeCameraRigEvent& OnCameraRigEvent() { return OnCameraRigEventDelegate; }
 
+public:
+
+	template<typename ValueType>
+	FCameraVariableSetterHandle AddCameraVariableSetter(const TCameraVariableSetter<ValueType>& InSetter);
+
+	void StopCameraVariableSetter(const FCameraVariableSetterHandle& InHandle, bool bImmediately = false);
+
+	void ApplyCameraVariableSetters(FCameraVariableTable& OutVariableTable);
+
 protected:
 
 	// FCameraNodeEvaluator interface.
@@ -190,6 +201,8 @@ protected:
 
 	void BroadcastCameraRigEvent(const FRootCameraNodeCameraRigEvent& InEvent) const;
 
+	void UpdateCameraVariableSetters(float DeltaTime);
+
 private:
 
 	/** The camera system that owns this root node. */
@@ -197,7 +210,26 @@ private:
 
 	/** The delegate to notify when an event happens. */
 	FOnRootCameraNodeCameraRigEvent OnCameraRigEventDelegate;
+
+	struct FVariableSetterEntry
+	{
+		FCameraVariableSetterPtr Setter;
+		FCameraVariableSetterHandle ThisHandle;
+	};
+
+	using FCameraVariableSetters = TSparseArray<FVariableSetterEntry>;
+	FCameraVariableSetters VariableSetters;
+	uint32 NextVariableSetterSerial = 0;
 };
 
+template<typename ValueType>
+FCameraVariableSetterHandle FRootCameraNodeEvaluator::AddCameraVariableSetter(const TCameraVariableSetter<ValueType>& InSetter)
+{
+	FSparseArrayAllocationInfo NewAllocation = VariableSetters.AddUninitialized();
+	FCameraVariableSetterHandle NewHandle(NewAllocation.Index, NextVariableSetterSerial++);
+	new(NewAllocation.Pointer) FVariableSetterEntry{ InSetter, NewHandle };
+	return NewHandle;
+}
+
 }  // namespace UE::Cameras
 
diff --git a/Engine/Plugins/Cameras/GameplayCameras/Source/GameplayCameras/Public/Core/TransientBlendStackCameraNode.h b/Engine/Plugins/Cameras/GameplayCameras/Source/GameplayCameras/Public/Core/TransientBlendStackCameraNode.h
--- a/Engine/Plugins/Cameras/GameplayCameras/Source/GameplayCameras/Public/Core/TransientBlendStackCameraNode.h
+++ b/Engine/Plugins/Cameras/GameplayCameras/Source/GameplayCameras/Public/Core/TransientBlendStackCameraNode.h
@@ -73,11 +73,6 @@ protected:
 	virtual void OnRun(const FCameraNodeEvaluationParams& Params, FCameraNodeEvaluationResult& OutResult) override;
 	virtual void OnSerialize(const FCameraNodeEvaluatorSerializeParams& Params, FArchive& Ar) override;
 
-#if WITH_EDITOR
-	// FBlendStackCameraNodeEvaluator interface.
-	virtual void OnEntryReinitialized(int32 EntryIndex) override;
-#endif
-
 private:
 
 	// Update methods.
@@ -97,9 +92,6 @@ private:
 	/** Extra blending-related info for each camera rig entry. */
 	struct FCameraRigEntryExtraInfo
 	{
-		/** Result without pre-blending. */
-		FCameraNodeEvaluationResult UnblendedResult;
-
 		/** Whether input slots were run. */
 		bool bInputRunThisFrame = false;
 		/** Whether the blend node was run. */
