Cameras: add location offset multiplier parameter to spline orbit node

diff --git a/Engine/Plugins/Cameras/GameplayCameras/Source/GameplayCameras/Private/Nodes/Common/SplineOrbitCameraNode.cpp b/Engine/Plugins/Cameras/GameplayCameras/Source/GameplayCameras/Private/Nodes/Common/SplineOrbitCameraNode.cpp
--- a/Engine/Plugins/Cameras/GameplayCameras/Source/GameplayCameras/Private/Nodes/Common/SplineOrbitCameraNode.cpp
+++ b/Engine/Plugins/Cameras/GameplayCameras/Source/GameplayCameras/Private/Nodes/Common/SplineOrbitCameraNode.cpp
@@ -4,6 +4,7 @@
 
 #include "Core/CameraEvaluationContext.h"
 #include "Core/CameraOperation.h"
+#include "Core/CameraParameterReader.h"
 #include "Core/CameraRigJoints.h"
 #include "Debug/CameraDebugBlock.h"
 #include "Debug/CameraDebugBlockBuilder.h"
@@ -14,9 +15,7 @@
 #include "IGameplayCamerasModule.h"
 #include "IGameplayCamerasLiveEditListener.h"
 #include "IGameplayCamerasLiveEditManager.h"
-#include "Math/Axis.h"
 #include "Math/CameraNodeSpaceMath.h"
-#include "Math/InterpCurve.h"
 #include "Math/Ray.h"
 #include "Nodes/Input/Input2DCameraNode.h"
 
@@ -79,6 +78,8 @@ private:
 	FCompressedRichCurve TargetOffsetSpline[3];
 	FCompressedRichCurve RotationOffsetSpline[3];
 
+	TCameraParameterReader<float> LocationOffsetMultiplierReader;
+
 	bool bHasAnyTargetOffset = false;
 	bool bHasAnyRotationOffset = false;
 
@@ -126,6 +127,9 @@ void FSplineOrbitCameraNodeEvaluator::OnInitialize(const FCameraNodeEvaluatorIni
 
 	RebuildCurves();
 
+	const USplineOrbitCameraNode* SplineOrbitNode = GetCameraNodeAs<USplineOrbitCameraNode>();
+	LocationOffsetMultiplierReader.Initialize(SplineOrbitNode->LocationOffsetMultiplier);
+
 #if WITH_EDITOR
 	IGameplayCamerasModule& GameplayCamerasModule = IGameplayCamerasModule::Get();
 	if (TSharedPtr<IGameplayCamerasLiveEditManager> LiveEditManager = GameplayCamerasModule.GetLiveEditManager())
@@ -194,7 +198,8 @@ void FSplineOrbitCameraNodeEvaluator::OnRun(const FCameraNodeEvaluationParams& P
 			LocationOffsetSpline[0].Eval(OrbitRotation.Pitch),
 			LocationOffsetSpline[1].Eval(OrbitRotation.Pitch),
 			LocationOffsetSpline[2].Eval(OrbitRotation.Pitch));
-	FTransform3d OrbitTransform(FTransform3d(LocationOffset) * OrbitPivot);
+	const float LocationOffsetMultiplier = LocationOffsetMultiplierReader.Get(OutResult.VariableTable);
+	FTransform3d OrbitTransform(FTransform3d(LocationOffset * LocationOffsetMultiplier) * OrbitPivot);
 
 #if UE_GAMEPLAY_CAMERAS_DEBUG
 	DebugYawPitch = FVector2d(OrbitRotation.Yaw, OrbitRotation.Pitch);
diff --git a/Engine/Plugins/Cameras/GameplayCameras/Source/GameplayCameras/Public/Nodes/Common/SplineOrbitCameraNode.h b/Engine/Plugins/Cameras/GameplayCameras/Source/GameplayCameras/Public/Nodes/Common/SplineOrbitCameraNode.h
--- a/Engine/Plugins/Cameras/GameplayCameras/Source/GameplayCameras/Public/Nodes/Common/SplineOrbitCameraNode.h
+++ b/Engine/Plugins/Cameras/GameplayCameras/Source/GameplayCameras/Public/Nodes/Common/SplineOrbitCameraNode.h
@@ -84,6 +84,9 @@ public:
 	UPROPERTY(EditAnywhere, Category="Spline Orbit", meta=(CurveEditorXLabelFormat="{0}deg"))
 	FCameraRotatorCurve RotationOffsetSpline;
 
+	UPROPERTY(EditAnywhere, Category="Spline Orbit")
+	FFloatCameraParameter LocationOffsetMultiplier = 1.f;
+
 	/**
 	 * The space in which the control points' TargetOffset is applied.
 	 */
