Cameras: fix android build

diff --git a/Engine/Plugins/Cameras/GameplayCameras/Source/GameplayCameras/Public/Core/CameraVariableSetter.h b/Engine/Plugins/Cameras/GameplayCameras/Source/GameplayCameras/Public/Core/CameraVariableSetter.h
--- a/Engine/Plugins/Cameras/GameplayCameras/Source/GameplayCameras/Public/Core/CameraVariableSetter.h
+++ b/Engine/Plugins/Cameras/GameplayCameras/Source/GameplayCameras/Public/Core/CameraVariableSetter.h
@@ -153,7 +153,7 @@ struct TCameraVariableSetter : public FCameraVariableSetter
 	ValueType ToValue;
 
 	TCameraVariableSetter() = default;
-	TCameraVariableSetter(FCameraVariableID InVariableID, TCallTraits<ValueType>::ParamType InToValue)
+	TCameraVariableSetter(FCameraVariableID InVariableID, typename TCallTraits<ValueType>::ParamType InToValue)
 		: FCameraVariableSetter(InVariableID)
 		, ToValue(InToValue)
 	{}
