Cameras: hot-reload camera rigs if sub-rigs have been rebuilt

This handles the case of building a prefab rig that is used by running camera rigs.

diff --git a/Engine/Plugins/Cameras/GameplayCameras/Source/GameplayCameras/Private/Core/CameraRigAsset.cpp b/Engine/Plugins/Cameras/GameplayCameras/Source/GameplayCameras/Private/Core/CameraRigAsset.cpp
--- a/Engine/Plugins/Cameras/GameplayCameras/Source/GameplayCameras/Private/Core/CameraRigAsset.cpp
+++ b/Engine/Plugins/Cameras/GameplayCameras/Source/GameplayCameras/Private/Core/CameraRigAsset.cpp
@@ -6,6 +6,7 @@
 #include "Build/CameraRigAssetBuilder.h"
 #include "Core/CameraNode.h"
 #include "Core/CameraVariableAssets.h"
+#include "Core/IAssetReferenceCameraNode.h"
 #include "UObject/AssetRegistryTagsContext.h"
 #include "UObject/ObjectSaveContext.h"
 
@@ -149,6 +150,11 @@ void UCameraRigAsset::GatherPackages(FCameraRigPackages& OutPackages) const
 		const UPackage* CurrentPackage = CurrentNode->GetOutermost();
 		OutPackages.AddUnique(CurrentPackage);
 
+		if (IAssetReferenceCameraNode* AssetReferencer = Cast<IAssetReferenceCameraNode>(CurrentNode))
+		{
+			AssetReferencer->GatherPackages(OutPackages);
+		}
+
 		FCameraNodeChildrenView CurrentChildren = CurrentNode->GetChildren();
 		for (UCameraNode* CurrentChild : ReverseIterate(CurrentChildren))
 		{
diff --git a/Engine/Plugins/Cameras/GameplayCameras/Source/GameplayCameras/Private/Nodes/Common/CameraRigCameraNode.cpp b/Engine/Plugins/Cameras/GameplayCameras/Source/GameplayCameras/Private/Nodes/Common/CameraRigCameraNode.cpp
--- a/Engine/Plugins/Cameras/GameplayCameras/Source/GameplayCameras/Private/Nodes/Common/CameraRigCameraNode.cpp
+++ b/Engine/Plugins/Cameras/GameplayCameras/Source/GameplayCameras/Private/Nodes/Common/CameraRigCameraNode.cpp
@@ -117,6 +117,16 @@ void UCameraRigCameraNode::OnBuild(FCameraObjectBuildContext& BuildContext)
 	BuildContext.AllocationInfo.Append(CameraRig->AllocationInfo);
 }
 
+void UCameraRigCameraNode::GatherPackages(FCameraRigPackages& OutPackages) const
+{
+	if (const UCameraRigAsset* CameraRig = CameraRigReference.GetCameraRig())
+	{
+#if WITH_EDITOR
+		CameraRig->GatherPackages(OutPackages);
+#endif  // WITH_EDITOR
+	}
+}
+
 void UCameraRigCameraNode::GetCustomCameraNodeParameters(FCustomCameraNodeParameterInfos& OutParameterInfos)
 {
 	CameraRigReference.GetCustomCameraNodeParameters(OutParameterInfos);
diff --git a/Engine/Plugins/Cameras/GameplayCameras/Source/GameplayCameras/Public/Core/IAssetReferenceCameraNode.h b/Engine/Plugins/Cameras/GameplayCameras/Source/GameplayCameras/Public/Core/IAssetReferenceCameraNode.h
new file mode 100644
--- /dev/null
+++ b/Engine/Plugins/Cameras/GameplayCameras/Source/GameplayCameras/Public/Core/IAssetReferenceCameraNode.h
@@ -0,0 +1,27 @@
+// Copyright Epic Games, Inc. All Rights Reserved.
+
+#pragma once
+
+#include "Core/CameraRigAsset.h"
+#include "Containers/Array.h"
+#include "UObject/Interface.h"
+
+#include "IAssetReferenceCameraNode.generated.h"
+
+class UCameraNode;
+
+UINTERFACE(MinimalAPI)
+class UAssetReferenceCameraNode : public UInterface
+{
+	GENERATED_BODY()
+};
+
+class IAssetReferenceCameraNode
+{
+	GENERATED_BODY()
+
+public:
+
+	virtual void GatherPackages(FCameraRigPackages& OutPackages) const = 0;
+};
+
diff --git a/Engine/Plugins/Cameras/GameplayCameras/Source/GameplayCameras/Public/Nodes/Common/CameraRigCameraNode.h b/Engine/Plugins/Cameras/GameplayCameras/Source/GameplayCameras/Public/Nodes/Common/CameraRigCameraNode.h
--- a/Engine/Plugins/Cameras/GameplayCameras/Source/GameplayCameras/Public/Nodes/Common/CameraRigCameraNode.h
+++ b/Engine/Plugins/Cameras/GameplayCameras/Source/GameplayCameras/Public/Nodes/Common/CameraRigCameraNode.h
@@ -5,6 +5,7 @@
 #include "Core/CameraNode.h"
 #include "Core/CameraRigAsset.h"
 #include "Core/CameraRigAssetReference.h"
+#include "Core/IAssetReferenceCameraNode.h"
 #include "Core/ICustomCameraNodeParameterProvider.h"
 
 #include "CameraRigCameraNode.generated.h"
@@ -15,6 +16,7 @@
 UCLASS(MinimalAPI, meta=(DisplayName="Camera Rig Prefab", CameraNodeCategories="Common,Utility"))
 class UCameraRigCameraNode 
 	: public UCameraNode
+	, public IAssetReferenceCameraNode
 	, public ICustomCameraNodeParameterProvider
 {
 	GENERATED_BODY()
@@ -26,6 +28,9 @@ protected:
 	virtual void OnBuild(FCameraObjectBuildContext& BuildContext) override;
 	virtual FCameraNodeEvaluatorPtr OnBuildEvaluator(FCameraNodeEvaluatorBuilder& Builder) const override;
 
+	// IAssetReferenceCameraNode interface.
+	virtual void GatherPackages(FCameraRigPackages& OutPackages) const override;
+
 	// ICustomCameraNodeParameterProvider interface.
 	virtual void GetCustomCameraNodeParameters(FCustomCameraNodeParameterInfos& OutParameterInfos) override;
 
