Cameras: remove extra Blueprint methods on the camera pose

They were duplicating what the break/set members nodes are already doing in vanilla Blueprint.

diff --git a/Engine/Plugins/Cameras/GameplayCameras/Source/GameplayCameras/Private/GameFramework/BlueprintCameraPose.cpp b/Engine/Plugins/Cameras/GameplayCameras/Source/GameplayCameras/Private/GameFramework/BlueprintCameraPose.cpp
--- a/Engine/Plugins/Cameras/GameplayCameras/Source/GameplayCameras/Private/GameFramework/BlueprintCameraPose.cpp
+++ b/Engine/Plugins/Cameras/GameplayCameras/Source/GameplayCameras/Private/GameFramework/BlueprintCameraPose.cpp
@@ -25,41 +25,6 @@ UE_CAMERA_POSE_FOR_ALL_PROPERTIES()
 #undef UE_CAMERA_POSE_FOR_PROPERTY
 }
 
-FBlueprintCameraPose UBlueprintCameraPoseFunctionLibrary::SetLocation(const FBlueprintCameraPose& CameraPose, const FVector& Location)
-{
-	FBlueprintCameraPose Result(CameraPose);
-	Result.Location = Location;
-	return Result;
-}
-
-FBlueprintCameraPose UBlueprintCameraPoseFunctionLibrary::SetRotation(const FBlueprintCameraPose& CameraPose, const FRotator& Rotation)
-{
-	FBlueprintCameraPose Result(CameraPose);
-	Result.Rotation = Rotation;
-	return Result;
-}
-
-FBlueprintCameraPose UBlueprintCameraPoseFunctionLibrary::SetTargetDistance(const FBlueprintCameraPose& CameraPose, double TargetDistance)
-{
-	FBlueprintCameraPose Result(CameraPose);
-	Result.TargetDistance = TargetDistance;
-	return Result;
-}
-
-FBlueprintCameraPose UBlueprintCameraPoseFunctionLibrary::SetFieldOfView(const FBlueprintCameraPose& CameraPose, float FieldOfView)
-{
-	FBlueprintCameraPose Result(CameraPose);
-	Result.FieldOfView = FieldOfView;
-	return Result;
-}
-
-FBlueprintCameraPose UBlueprintCameraPoseFunctionLibrary::SetFocalLength(const FBlueprintCameraPose& CameraPose, float FocalLength)
-{
-	FBlueprintCameraPose Result(CameraPose);
-	Result.FocalLength = FocalLength;
-	return Result;
-}
-
 FTransform UBlueprintCameraPoseFunctionLibrary::GetTransform(const FBlueprintCameraPose& CameraPose)
 {
 	FTransform Transform;
diff --git a/Engine/Plugins/Cameras/GameplayCameras/Source/GameplayCameras/Public/GameFramework/BlueprintCameraPose.h b/Engine/Plugins/Cameras/GameplayCameras/Source/GameplayCameras/Public/GameFramework/BlueprintCameraPose.h
--- a/Engine/Plugins/Cameras/GameplayCameras/Source/GameplayCameras/Public/GameFramework/BlueprintCameraPose.h
+++ b/Engine/Plugins/Cameras/GameplayCameras/Source/GameplayCameras/Public/GameFramework/BlueprintCameraPose.h
@@ -148,50 +148,6 @@ class UBlueprintCameraPoseFunctionLibrary : public UBlueprintFunctionLibrary
 {
 	GENERATED_BODY()
 
-public:
-
-	/** Gets the location of the camera pose. */
-	UFUNCTION(BlueprintPure, Category="Camera")
-	static inline FVector GetLocation(const FBlueprintCameraPose& CameraPose) { return CameraPose.Location; }
-
-	/** Gets the rotation of the camera pose. */
-	UFUNCTION(BlueprintPure, Category="Camera")
-	static inline FRotator GetRotation(const FBlueprintCameraPose& CameraPose) { return CameraPose.Rotation; }
-
-	/** Gets the target distance of the camera pose. */
-	UFUNCTION(BlueprintPure, Category="Camera")
-	static inline double GetTargetDistance(const FBlueprintCameraPose& CameraPose) { return CameraPose.TargetDistance; }
-
-	/** Gets the field of view of the camera pose. */
-	UFUNCTION(BlueprintPure, Category="Camera")
-	static inline double GetFieldOfView(const FBlueprintCameraPose& CameraPose) { return CameraPose.FieldOfView; }
-
-	/** Gets the focal length of the camera pose. */
-	UFUNCTION(BlueprintPure, Category="Camera")
-	static inline double GetFocalLength(const FBlueprintCameraPose& CameraPose) { return CameraPose.FocalLength; }
-	
-public:
-
-	/** Creates a copy of the given camera pose with the given location. */
-	UFUNCTION(BlueprintCallable, Category="Camera")
-	static FBlueprintCameraPose SetLocation(const FBlueprintCameraPose& CameraPose, const FVector& Location);
-
-	/** Creates a copy of the given camera pose with the given rotation. */
-	UFUNCTION(BlueprintCallable, Category="Camera")
-	static FBlueprintCameraPose SetRotation(const FBlueprintCameraPose& CameraPose, const FRotator& Rotation);
-
-	/** Creates a copy of the given camera pose with the given target distance. */
-	UFUNCTION(BlueprintCallable, Category="Camera")
-	static FBlueprintCameraPose SetTargetDistance(const FBlueprintCameraPose& CameraPose, double TargetDistance);
-
-	/** Creates a copy of the given camera pose with the given field of view. */
-	UFUNCTION(BlueprintCallable, Category="Camera")
-	static FBlueprintCameraPose SetFieldOfView(const FBlueprintCameraPose& CameraPose, float FieldOfView);
-
-	/** Creates a copy of the given camera pose with the given focal length. */
-	UFUNCTION(BlueprintCallable, Category="Camera")
-	static FBlueprintCameraPose SetFocalLength(const FBlueprintCameraPose& CameraPose, float FocalLength);
-
 public:
 
 	/** Gets the transform matrix of the camera pose. */
