Cameras: fix variable reference readers not always using the desired default value

diff --git a/Engine/Plugins/Cameras/GameplayCameras/Source/GameplayCameras/Private/Nodes/Collision/CollisionPushCameraNode.cpp b/Engine/Plugins/Cameras/GameplayCameras/Source/GameplayCameras/Private/Nodes/Collision/CollisionPushCameraNode.cpp
--- a/Engine/Plugins/Cameras/GameplayCameras/Source/GameplayCameras/Private/Nodes/Collision/CollisionPushCameraNode.cpp
+++ b/Engine/Plugins/Cameras/GameplayCameras/Source/GameplayCameras/Private/Nodes/Collision/CollisionPushCameraNode.cpp
@@ -102,7 +102,7 @@ void FCollisionPushCameraNodeEvaluator::OnInitialize(const FCameraNodeEvaluatorI
 {
 	const UCollisionPushCameraNode* CollisionPushNode = GetCameraNodeAs<UCollisionPushCameraNode>();
 
-	EnableCollisionReader.Initialize(CollisionPushNode->EnableCollision);
+	EnableCollisionReader.Initialize(CollisionPushNode->EnableCollision, true);
 	CustomSafePositionReader.Initialize(CollisionPushNode->CustomSafePosition);
 
 	CollisionSphereRadiusReader.Initialize(CollisionPushNode->CollisionSphereRadius);
diff --git a/Engine/Plugins/Cameras/GameplayCameras/Source/GameplayCameras/Public/Core/CameraParameterReader.h b/Engine/Plugins/Cameras/GameplayCameras/Source/GameplayCameras/Public/Core/CameraParameterReader.h
--- a/Engine/Plugins/Cameras/GameplayCameras/Source/GameplayCameras/Public/Core/CameraParameterReader.h
+++ b/Engine/Plugins/Cameras/GameplayCameras/Source/GameplayCameras/Public/Core/CameraParameterReader.h
@@ -44,7 +44,7 @@ public:
 	/**
 	 * Gets the actual value for the parameter.
 	 */
-	typename TCallTraits<ValueType>::ParamType Get(const FCameraVariableTable& VariableTable) const
+	const ValueType& Get(const FCameraVariableTable& VariableTable) const
 	{
 		checkf(DefaultValuePtr, TEXT("Parameter reader has no value pointer!"));
 		if (!VariableID.IsValid())
diff --git a/Engine/Plugins/Cameras/GameplayCameras/Source/GameplayCameras/Public/Core/CameraVariableReferenceReader.h b/Engine/Plugins/Cameras/GameplayCameras/Source/GameplayCameras/Public/Core/CameraVariableReferenceReader.h
--- a/Engine/Plugins/Cameras/GameplayCameras/Source/GameplayCameras/Public/Core/CameraVariableReferenceReader.h
+++ b/Engine/Plugins/Cameras/GameplayCameras/Source/GameplayCameras/Public/Core/CameraVariableReferenceReader.h
@@ -30,6 +30,15 @@ public:
 	 */
 	template<typename VariableReferenceType>
 	void Initialize(const VariableReferenceType& Reference)
+	{
+		Initialize(Reference, ValueType());
+	}
+
+	/**
+	 * Initializes the reader around the given variable reference.
+	 */
+	template<typename VariableReferenceType>
+	void Initialize(const VariableReferenceType& Reference, TCallTraits<ValueType>::ParamType DefaultValueIfNoReference)
 	{
 		static_assert(
 				std::is_same<ValueType, typename VariableReferenceType::VariableAssetType::ValueType>(),
@@ -42,7 +51,7 @@ public:
 		}
 		else
 		{
-			DefaultValue = ValueType();
+			DefaultValue = DefaultValueIfNoReference;
 			VariableID = Reference.VariableID;
 		}
 	}
@@ -50,7 +59,7 @@ public:
 	/**
 	 * Gets the actual value for the referenced variable.
 	 */
-	typename TCallTraits<ValueType>::ParamType Get(const FCameraVariableTable& VariableTable) const
+	const ValueType& Get(const FCameraVariableTable& VariableTable) const
 	{
 		if (!VariableID.IsValid())
 		{
