Cameras: fix android build

diff --git a/Engine/Plugins/Cameras/GameplayCameras/Source/GameplayCameras/Public/Core/CameraVariableReferenceReader.h b/Engine/Plugins/Cameras/GameplayCameras/Source/GameplayCameras/Public/Core/CameraVariableReferenceReader.h
--- a/Engine/Plugins/Cameras/GameplayCameras/Source/GameplayCameras/Public/Core/CameraVariableReferenceReader.h
+++ b/Engine/Plugins/Cameras/GameplayCameras/Source/GameplayCameras/Public/Core/CameraVariableReferenceReader.h
@@ -38,7 +38,7 @@ public:
 	 * Initializes the reader around the given variable reference.
 	 */
 	template<typename VariableReferenceType>
-	void Initialize(const VariableReferenceType& Reference, TCallTraits<ValueType>::ParamType DefaultValueIfNoReference)
+	void Initialize(const VariableReferenceType& Reference, typename TCallTraits<ValueType>::ParamType DefaultValueIfNoReference)
 	{
 		static_assert(
 				std::is_same<ValueType, typename VariableReferenceType::VariableAssetType::ValueType>(),
