UE: Fix short path reference

diff --git a/Engine/Plugins/Cameras/GameplayCameras/Source/GameplayCameras/Public/GameFramework/BlueprintCameraEvaluationDataRef.h b/Engine/Plugins/Cameras/GameplayCameras/Source/GameplayCameras/Public/GameFramework/BlueprintCameraEvaluationDataRef.h
--- a/Engine/Plugins/Cameras/GameplayCameras/Source/GameplayCameras/Public/GameFramework/BlueprintCameraEvaluationDataRef.h
+++ b/Engine/Plugins/Cameras/GameplayCameras/Source/GameplayCameras/Public/GameFramework/BlueprintCameraEvaluationDataRef.h
@@ -22,7 +22,7 @@ struct FCameraNodeEvaluationResult;
  * Blueprint wrapper for camera evaluation data.
  */
 USTRUCT(BlueprintType, DisplayName="Camera Evaluation Data Ref", 
-		meta=(HasNativeMake="BlueprintCameraEvaluationDataFunctionLibrary.MakeCameraEvaluationData"))
+		meta=(HasNativeMake="/Script/GameplayCameras.BlueprintCameraEvaluationDataFunctionLibrary.MakeCameraEvaluationData"))
 struct GAMEPLAYCAMERAS_API FBlueprintCameraEvaluationDataRef
 {
 	GENERATED_BODY()
