Cameras: remove the use of an auto-spawned camera system actor, replace it with camera actors running their own camera system

The previous approach (where we would automatically spawn a camera system actor under the hood and set it as the view target) turned out to have some downsides:

- The actor OwnerNoSee/OwnerOnlySee settings didn't work, since we set the view target actor to something else
- Sequencer camera cut tracks didn't work, since they call SetViewTarget
- In some cases it forced users to go through different Blueprint functions than the usual SetViewTarget

So this change removes the auto-spawning code, and any dependencies we have on the camera system actor. That actor still exists, but as a standalone/independent actor.

The gameplay camera and gameplay camera rig actors now check for a player-controller-level camera system host (such as one created by the player camera manager), indicating that someone is taking care of running the camera system. If they find it, they just register with it. If they don't find it, they instantiate their own camera system and update themselves.

The downside is that there can be no support for base/global/visual rigs, and no support for advanced blending (since blending is taken care of by the default camera manager), but the upside is that things "just work" at first.

diff --git a/Engine/Plugins/Cameras/GameplayCameras/Config/BaseGameplayCameras.ini b/Engine/Plugins/Cameras/GameplayCameras/Config/BaseGameplayCameras.ini
--- a/Engine/Plugins/Cameras/GameplayCameras/Config/BaseGameplayCameras.ini
+++ b/Engine/Plugins/Cameras/GameplayCameras/Config/BaseGameplayCameras.ini
@@ -25,3 +25,5 @@
 +PropertyRedirects=(OldName="GetCameraPose.Data",NewName="CameraData")
 +PropertyRedirects=(OldName="SetCameraPose.Data",NewName="CameraData")
 
++FunctionRedirects=(OldName="GameplayCameraSystemActor.AutoManageActiveViewTarget", NewName="SetAutoSpawnedCameraSystemActorViewTarget")
+
diff --git a/Engine/Plugins/Cameras/GameplayCameras/Source/GameplayCameras/Private/Core/BlendStackCameraNode.cpp b/Engine/Plugins/Cameras/GameplayCameras/Source/GameplayCameras/Private/Core/BlendStackCameraNode.cpp
--- a/Engine/Plugins/Cameras/GameplayCameras/Source/GameplayCameras/Private/Core/BlendStackCameraNode.cpp
+++ b/Engine/Plugins/Cameras/GameplayCameras/Source/GameplayCameras/Private/Core/BlendStackCameraNode.cpp
@@ -199,6 +199,18 @@ FCameraRigEvaluationInfo FBlendStackCameraNodeEvaluator::GetCameraRigEvaluationI
 	return FCameraRigEvaluationInfo();
 }
 
+bool FBlendStackCameraNodeEvaluator::HasAnyRunningCameraRig(TSharedPtr<const FCameraEvaluationContext> InContext) const
+{
+	for (const FCameraRigEntry& Entry : Entries)
+	{
+		if (Entry.EvaluationContext == InContext)
+		{
+			return true;
+		}
+	}
+	return false;
+}
+
 FCameraNodeEvaluatorChildrenView FBlendStackCameraNodeEvaluator::OnGetChildren()
 {
 	FCameraNodeEvaluatorChildrenView View;
diff --git a/Engine/Plugins/Cameras/GameplayCameras/Source/GameplayCameras/Private/Core/DefaultRootCameraNode.cpp b/Engine/Plugins/Cameras/GameplayCameras/Source/GameplayCameras/Private/Core/DefaultRootCameraNode.cpp
--- a/Engine/Plugins/Cameras/GameplayCameras/Source/GameplayCameras/Private/Core/DefaultRootCameraNode.cpp
+++ b/Engine/Plugins/Cameras/GameplayCameras/Source/GameplayCameras/Private/Core/DefaultRootCameraNode.cpp
@@ -169,11 +169,24 @@ void FDefaultRootCameraNodeEvaluator::OnDeactivateCameraRig(const FDeactivateCam
 	}
 }
 
+void FDefaultRootCameraNodeEvaluator::OnDeactivateAllCameraRigs(TSharedPtr<const FCameraEvaluationContext> InContext, bool bImmediately)
+{
+	BaseLayer->RemoveAll(InContext, bImmediately);
+	MainLayer->FreezeAll(InContext);
+	GlobalLayer->RemoveAll(InContext, bImmediately);
+	VisualLayer->RemoveAll(InContext, bImmediately);
+}
+
 void FDefaultRootCameraNodeEvaluator::OnGetActiveCameraRigInfo(FCameraRigEvaluationInfo& OutCameraRigInfo) const
 {
 	OutCameraRigInfo = MainLayer->GetActiveCameraRigEvaluationInfo();
 }
 
+bool FDefaultRootCameraNodeEvaluator::OnHasAnyRunningCameraRig(TSharedPtr<const FCameraEvaluationContext> InContext) const
+{
+	return MainLayer->HasAnyRunningCameraRig(InContext);
+}
+
 void FDefaultRootCameraNodeEvaluator::OnGetCameraRigInfo(const FCameraRigInstanceID InstanceID, FCameraRigEvaluationInfo& OutCameraRigInfo) const
 {
 	FPersistentBlendStackCameraNodeEvaluator* TargetLayer = nullptr;
diff --git a/Engine/Plugins/Cameras/GameplayCameras/Source/GameplayCameras/Private/Core/PersistentBlendStackCameraNode.cpp b/Engine/Plugins/Cameras/GameplayCameras/Source/GameplayCameras/Private/Core/PersistentBlendStackCameraNode.cpp
--- a/Engine/Plugins/Cameras/GameplayCameras/Source/GameplayCameras/Private/Core/PersistentBlendStackCameraNode.cpp
+++ b/Engine/Plugins/Cameras/GameplayCameras/Source/GameplayCameras/Private/Core/PersistentBlendStackCameraNode.cpp
@@ -118,71 +118,93 @@ void FPersistentBlendStackCameraNodeEvaluator::Remove(const FBlendStackCameraRem
 		}
 	}
 
-	// If we need to remove the camera rigs immediately, simply pop out their entries.
-	if (Params.bRemoveImmediately)
+	for (int32 Index : EntriesToRemove)
 	{
-		for (int32 Index : EntriesToRemove)
+		RemoveEntry(Index, Params.TransitionOverride, Params.bRemoveImmediately);
+	}
+}
+
+void FPersistentBlendStackCameraNodeEvaluator::RemoveAll(TSharedPtr<const FCameraEvaluationContext> InContext, bool bImmediately)
+{
+	TArray<int32, TInlineAllocator<4>> EntriesToRemove;
+	for (int32 Index = Entries.Num() - 1; Index >= 0; --Index)
+	{
+		FCameraRigEntry& Entry(Entries[Index]);
+		if (Entry.EvaluationContext == InContext)
 		{
-			PopEntry(Index);
-			EntryExtraInfos.RemoveAt(Index);
+			EntriesToRemove.Add(Index);
 		}
 	}
-	// Else, we need to start blending out these entries.
+
+	for (int32 Index : EntriesToRemove)
+	{
+		RemoveEntry(Index, nullptr, bImmediately);
+	}
+}
+
+void FPersistentBlendStackCameraNodeEvaluator::RemoveEntry(int32 EntryIndex, const UCameraRigTransition* TransitionOverride, bool bImmediately)
+{
+	ensure(Entries.IsValidIndex(EntryIndex));
+
+	// If we need to remove the camera rig immediately, simply pop out its entries.
+	// Else, we need to start blending out that entry.
+	if (bImmediately)
+	{
+		PopEntry(EntryIndex);
+		EntryExtraInfos.RemoveAt(EntryIndex);
+	}
 	else
 	{
-		for (int32 Index : EntriesToRemove)
+		FCameraRigEntry& Entry(Entries[EntryIndex]);
+		FCameraRigEntryExtraInfo& EntryExtraInfo(EntryExtraInfos[EntryIndex]);
+		const UCameraRigTransition* Transition = FindExitTransition(Entry, TransitionOverride);
+		if (Transition && Transition->Blend)
 		{
-			FCameraRigEntry& Entry(Entries[Index]);
-			FCameraRigEntryExtraInfo& EntryExtraInfo(EntryExtraInfos[Index]);
-			const UCameraRigTransition* Transition = FindExitTransition(Params, Entry);
-			if (Transition && Transition->Blend)
+			// Swap the blend-in evaluator on this entry with a blend-out one.
+			if (EntryExtraInfo.BlendStatus != EBlendStatus::BlendOut)
 			{
-				// Swap the blend-in evaluator on this entry with a blend-out one.
-				if (EntryExtraInfo.BlendStatus != EBlendStatus::BlendOut)
+				FCameraNodeEvaluatorBuilder BlendOutBuilder(Entry.EvaluatorStorage);
+				FCameraNodeEvaluatorBuildParams BlendOutBuildParams(BlendOutBuilder);
+				FBlendCameraNodeEvaluator* BlendOutEvaluator = BlendOutBuildParams.BuildEvaluatorAs<FBlendCameraNodeEvaluator>(Transition->Blend);
+
+				FCameraNodeEvaluatorInitializeParams BlendOutInitParams;
+				BlendOutInitParams.Evaluator = OwningEvaluator;
+				BlendOutInitParams.EvaluationContext = Entry.EvaluationContext.Pin();
+				BlendOutEvaluator->Initialize(BlendOutInitParams, Entry.Result);
+
+				// Reverse this blend so it plays as a blend-out. Also, see if we are going to 
+				// interrupt an ongoing blend-in... if so, give a chance for the blend-out to
+				// start at an "equivalent spot".
+				if (!BlendOutEvaluator->SetReversed(true))
 				{
-					FCameraNodeEvaluatorBuilder BlendOutBuilder(Entry.EvaluatorStorage);
-					FCameraNodeEvaluatorBuildParams BlendOutBuildParams(BlendOutBuilder);
-					FBlendCameraNodeEvaluator* BlendOutEvaluator = BlendOutBuildParams.BuildEvaluatorAs<FBlendCameraNodeEvaluator>(Transition->Blend);
-
-					FCameraNodeEvaluatorInitializeParams BlendOutInitParams;
-					BlendOutInitParams.Evaluator = OwningEvaluator;
-					BlendOutInitParams.EvaluationContext = Entry.EvaluationContext.Pin();
-					BlendOutEvaluator->Initialize(BlendOutInitParams, Entry.Result);
-
-					// Reverse this blend so it plays as a blend-out. Also, see if we are going to 
-					// interrupt an ongoing blend-in... if so, give a chance for the blend-out to
-					// start at an "equivalent spot".
-					if (!BlendOutEvaluator->SetReversed(true))
-					{
-						BlendOutEvaluator = Entry.EvaluatorStorage.BuildEvaluator<FReverseBlendCameraNodeEvaluator>(BlendOutEvaluator);
-					}
-					if (EntryExtraInfo.BlendStatus == EBlendStatus::BlendIn)
-					{
-						FBlendCameraNodeEvaluator* OngoingBlend = Entry.RootEvaluator->GetBlendEvaluator();
+					BlendOutEvaluator = Entry.EvaluatorStorage.BuildEvaluator<FReverseBlendCameraNodeEvaluator>(BlendOutEvaluator);
+				}
+				if (EntryExtraInfo.BlendStatus == EBlendStatus::BlendIn)
+				{
+					FBlendCameraNodeEvaluator* OngoingBlend = Entry.RootEvaluator->GetBlendEvaluator();
 
-						FCameraNodeBlendInterruptionParams InterruptionParams;
-						InterruptionParams.InterruptedBlend = OngoingBlend;
-						if (!BlendOutEvaluator->InitializeFromInterruption(InterruptionParams))
-						{
-							BlendOutEvaluator = Entry.EvaluatorStorage.BuildEvaluator<FInterruptedBlendCameraNodeEvaluator>(BlendOutEvaluator, OngoingBlend);
-						}
+					FCameraNodeBlendInterruptionParams InterruptionParams;
+					InterruptionParams.InterruptedBlend = OngoingBlend;
+					if (!BlendOutEvaluator->InitializeFromInterruption(InterruptionParams))
+					{
+						BlendOutEvaluator = Entry.EvaluatorStorage.BuildEvaluator<FInterruptedBlendCameraNodeEvaluator>(BlendOutEvaluator, OngoingBlend);
 					}
-					// Note: neither the reverse or interrupted blends need initialization, but
-					// technically we're missing calling it on them.
-					Entry.RootEvaluator->SetBlendEvaluator(BlendOutEvaluator);
-
-					EntryExtraInfo.BlendStatus = EBlendStatus::BlendOut;
-					EntryExtraInfo.bIsBlendFinished = false;
-					EntryExtraInfo.bIsBlendFull = false;
 				}
-				// else: we were already blending out, so let this continue.
-			}
-			else
-			{
-				// No transition found... just cut.
-				PopEntry(Index);
-				EntryExtraInfos.RemoveAt(Index);
+				// Note: neither the reverse or interrupted blends need initialization, but
+				// technically we're missing calling it on them.
+				Entry.RootEvaluator->SetBlendEvaluator(BlendOutEvaluator);
+
+				EntryExtraInfo.BlendStatus = EBlendStatus::BlendOut;
+				EntryExtraInfo.bIsBlendFinished = false;
+				EntryExtraInfo.bIsBlendFull = false;
 			}
+			// else: we were already blending out, so let this continue.
+		}
+		else
+		{
+			// No transition found... just cut.
+			PopEntry(EntryIndex);
+			EntryExtraInfos.RemoveAt(EntryIndex);
 		}
 	}
 }
@@ -342,12 +364,12 @@ const UCameraRigTransition* FPersistentBlendStackCameraNodeEvaluator::FindEnterT
 			Params.CameraRig, nullptr);
 }
 
-const UCameraRigTransition* FPersistentBlendStackCameraNodeEvaluator::FindExitTransition(const FBlendStackCameraRemoveParams& Params, const FCameraRigEntry& Entry) const
+const UCameraRigTransition* FPersistentBlendStackCameraNodeEvaluator::FindExitTransition(const FCameraRigEntry& Entry, const UCameraRigTransition* TransitionOverride) const
 {
 	// If we are forced to use a specific transition, our search is over.
-	if (Params.TransitionOverride)
+	if (TransitionOverride)
 	{
-		return Params.TransitionOverride;
+		return TransitionOverride;
 	}
 
 	// Find a transition that works for blending the given camera rig out.
diff --git a/Engine/Plugins/Cameras/GameplayCameras/Source/GameplayCameras/Private/Core/RootCameraNode.cpp b/Engine/Plugins/Cameras/GameplayCameras/Source/GameplayCameras/Private/Core/RootCameraNode.cpp
--- a/Engine/Plugins/Cameras/GameplayCameras/Source/GameplayCameras/Private/Core/RootCameraNode.cpp
+++ b/Engine/Plugins/Cameras/GameplayCameras/Source/GameplayCameras/Private/Core/RootCameraNode.cpp
@@ -27,6 +27,11 @@ void FRootCameraNodeEvaluator::DeactivateCameraRig(const FDeactivateCameraRigPar
 	OnDeactivateCameraRig(Params);
 }
 
+void FRootCameraNodeEvaluator::DeactivateAllCameraRigs(TSharedPtr<const FCameraEvaluationContext> InContext, bool bImmediately)
+{
+	OnDeactivateAllCameraRigs(InContext, bImmediately);
+}
+
 void FRootCameraNodeEvaluator::ExecuteCameraDirectorRequest(const FCameraRigActivationDeactivationRequest& Request)
 {
 	if (Request.RequestType == ECameraRigActivationDeactivationRequestType::Activate)
@@ -69,6 +74,11 @@ bool FRootCameraNodeEvaluator::HasAnyActiveCameraRig() const
 	return RigInfo.RootEvaluator != nullptr;
 }
 
+bool FRootCameraNodeEvaluator::HasAnyRunningCameraRig(TSharedPtr<const FCameraEvaluationContext> InContext) const
+{
+	return OnHasAnyRunningCameraRig(InContext);
+}
+
 const FCameraVariableTable* FRootCameraNodeEvaluator::GetBlendedParameters() const
 {
 	return OnGetBlendedParameters();
diff --git a/Engine/Plugins/Cameras/GameplayCameras/Source/GameplayCameras/Private/Core/TransientBlendStackCameraNode.cpp b/Engine/Plugins/Cameras/GameplayCameras/Source/GameplayCameras/Private/Core/TransientBlendStackCameraNode.cpp
--- a/Engine/Plugins/Cameras/GameplayCameras/Source/GameplayCameras/Private/Core/TransientBlendStackCameraNode.cpp
+++ b/Engine/Plugins/Cameras/GameplayCameras/Source/GameplayCameras/Private/Core/TransientBlendStackCameraNode.cpp
@@ -174,7 +174,6 @@ void FTransientBlendStackCameraNodeEvaluator::Freeze(const FBlendStackCameraFree
 	}
 	else
 	{
-
 		// Freeze any entries matching the given context and rig asset.
 		for (FCameraRigEntry& Entry : Entries)
 		{
@@ -188,7 +187,7 @@ void FTransientBlendStackCameraNodeEvaluator::Freeze(const FBlendStackCameraFree
 	}
 }
 
-void FTransientBlendStackCameraNodeEvaluator::FreezeAll(TSharedPtr<FCameraEvaluationContext> EvaluationContext)
+void FTransientBlendStackCameraNodeEvaluator::FreezeAll(TSharedPtr<const FCameraEvaluationContext> EvaluationContext)
 {
 	for (FCameraRigEntry& Entry : Entries)
 	{
diff --git a/Engine/Plugins/Cameras/GameplayCameras/Source/GameplayCameras/Private/GameFramework/ActivateCameraRigFunctions.cpp b/Engine/Plugins/Cameras/GameplayCameras/Source/GameplayCameras/Private/GameFramework/ActivateCameraRigFunctions.cpp
--- a/Engine/Plugins/Cameras/GameplayCameras/Source/GameplayCameras/Private/GameFramework/ActivateCameraRigFunctions.cpp
+++ b/Engine/Plugins/Cameras/GameplayCameras/Source/GameplayCameras/Private/GameFramework/ActivateCameraRigFunctions.cpp
@@ -2,8 +2,11 @@
 
 #include "GameFramework/ActivateCameraRigFunctions.h"
 
+#include "Core/CameraRigAsset.h"
 #include "Core/RootCameraNode.h"
 #include "GameFramework/ControllerGameplayCameraEvaluationComponent.h"
+#include "GameFramework/IGameplayCameraSystemHost.h"
+#include "GameFramework/PlayerController.h"
 #include "GameplayCameras.h"
 
 #include UE_INLINE_GENERATED_CPP_BY_NAME(ActivateCameraRigFunctions)
@@ -27,16 +30,51 @@ void UActivateCameraRigFunctions::ActivateCameraRigImpl(UObject* WorldContextObj
 {
 	using namespace UE::Cameras;
 
+	if (!PlayerController)
+	{
+		FFrame::KismetExecutionMessage(
+				TEXT("No player controller was given to activate a camera rig!"),
+				ELogVerbosity::Error);
+		return;
+	}
+
 	if (!CameraRig)
 	{
-		UE_LOG(LogCameraSystem, Error, TEXT("No camera rig was given to activate!"));
+		FFrame::KismetExecutionMessage(
+				TEXT("No camera rig was given to activate!"),
+				ELogVerbosity::Error);
 		return;
 	}
 
-	UControllerGameplayCameraEvaluationComponent* CameraEvaluationComponent = UControllerGameplayCameraEvaluationComponent::FindOrAddComponent(PlayerController);
-	if (ensure(CameraEvaluationComponent))
+	// Look for a camera system either under the player controller (such as with AGameplayCamerasPlayerCameraManager),
+	// or under the current view target (such as with a UGameplayCameraComponentBase whose actor is the view target).
+	UControllerGameplayCameraEvaluationComponent* CameraEvaluationComponent = nullptr;
+	if (IGameplayCameraSystemHost* FoundHost = IGameplayCameraSystemHost::FindActiveHost(PlayerController))
+	{
+		AActor* HostOwningActor = FoundHost->GetAsObject()->GetTypedOuter<AActor>();
+		check(HostOwningActor);
+
+		bool bComponentCreated = false;
+		CameraEvaluationComponent = UControllerGameplayCameraEvaluationComponent::FindOrAddComponent(HostOwningActor, &bComponentCreated);
+		if (bComponentCreated)
+		{
+			CameraEvaluationComponent->Initialize(FoundHost->GetAsScriptInterface());
+		}
+	}
+
+	if (CameraEvaluationComponent)
 	{
 		CameraEvaluationComponent->ActivateCameraRig(CameraRig, EvaluationLayer);
 	}
+	else
+	{
+		UE_LOG(LogCameraSystem, Error,
+				TEXT("Can't activate camera rig '%s' on layer '%s' because no camera system was found! "
+					 "Neither the player controller ('%s') or the current view target ('%s') have one."),
+				*GetNameSafe(CameraRig),
+				*UEnum::GetValueAsString(EvaluationLayer),
+				*GetNameSafe(PlayerController),
+				*GetNameSafe(PlayerController->GetViewTarget()));
+	}
 }
 
diff --git a/Engine/Plugins/Cameras/GameplayCameras/Source/GameplayCameras/Private/GameFramework/ControllerGameplayCameraEvaluationComponent.cpp b/Engine/Plugins/Cameras/GameplayCameras/Source/GameplayCameras/Private/GameFramework/ControllerGameplayCameraEvaluationComponent.cpp
--- a/Engine/Plugins/Cameras/GameplayCameras/Source/GameplayCameras/Private/GameFramework/ControllerGameplayCameraEvaluationComponent.cpp
+++ b/Engine/Plugins/Cameras/GameplayCameras/Source/GameplayCameras/Private/GameFramework/ControllerGameplayCameraEvaluationComponent.cpp
@@ -6,7 +6,7 @@
 #include "Core/CameraRigAsset.h"
 #include "Core/CameraSystemEvaluator.h"
 #include "Core/RootCameraNode.h"
-#include "GameFramework/GameplayCameraSystemHost.h"
+#include "GameFramework/IGameplayCameraSystemHost.h"
 #include "GameFramework/PlayerController.h"
 
 #include UE_INLINE_GENERATED_CPP_BY_NAME(ControllerGameplayCameraEvaluationComponent)
@@ -17,6 +17,12 @@ UControllerGameplayCameraEvaluationComponent::UControllerGameplayCameraEvaluatio
 	bAutoActivate = true;
 }
 
+void UControllerGameplayCameraEvaluationComponent::Initialize(TScriptInterface<IGameplayCameraSystemHost> InCameraSystemHost)
+{
+	ensure(InCameraSystemHost);
+	CameraSystemHost = InCameraSystemHost;
+}
+
 void UControllerGameplayCameraEvaluationComponent::ActivateCameraRig(UCameraRigAsset* CameraRig, ECameraRigLayer EvaluationLayer)
 {
 	FCameraRigInfo NewCameraRigInfo;
@@ -50,8 +56,7 @@ void UControllerGameplayCameraEvaluationComponent::ActivateCameraRigs()
 {
 	using namespace UE::Cameras;
 
-	EnsureCameraSystemHost();
-	if (!CameraSystemHost)
+	if (!ensure(CameraSystemHost))
 	{
 		return;
 	}
@@ -86,6 +91,7 @@ void UControllerGameplayCameraEvaluationComponent::EnsureEvaluationContext()
 
 	if (!EvaluationContext.IsValid())
 	{
+		// TODO: we won't find a player controller this way if we are under a view target.
 		APlayerController* PlayerController = GetOwner<APlayerController>();
 
 		FCameraEvaluationContextInitializeParams InitParams;
@@ -96,48 +102,24 @@ void UControllerGameplayCameraEvaluationComponent::EnsureEvaluationContext()
 	}
 }
 
-void UControllerGameplayCameraEvaluationComponent::EnsureCameraSystemHost()
-{
-	if (!CameraSystemHost)
-	{
-		APlayerController* PlayerController = GetOwner<APlayerController>();
-
-		CameraSystemHost = UGameplayCameraSystemHost::FindOrCreateHost(PlayerController);
-	}
-}
-
-UControllerGameplayCameraEvaluationComponent* UControllerGameplayCameraEvaluationComponent::FindComponent(APlayerController* PlayerController)
+UControllerGameplayCameraEvaluationComponent* UControllerGameplayCameraEvaluationComponent::FindComponent(AActor* OwnerActor)
 {
-	return PlayerController->FindComponentByClass<UControllerGameplayCameraEvaluationComponent>();
+	return OwnerActor->FindComponentByClass<UControllerGameplayCameraEvaluationComponent>();
 }
 
-UControllerGameplayCameraEvaluationComponent* UControllerGameplayCameraEvaluationComponent::FindOrAddComponent(APlayerController* PlayerController)
+UControllerGameplayCameraEvaluationComponent* UControllerGameplayCameraEvaluationComponent::FindOrAddComponent(AActor* OwnerActor, bool* bOutCreated)
 {
-	UControllerGameplayCameraEvaluationComponent* ControllerComponent = FindComponent(PlayerController);
+	UControllerGameplayCameraEvaluationComponent* ControllerComponent = FindComponent(OwnerActor);
 	if (!ControllerComponent)
 	{
 		ControllerComponent = NewObject<UControllerGameplayCameraEvaluationComponent>(
-				PlayerController, TEXT("ControllerGameplayCameraEvaluationComponent"), RF_Transient);
+				OwnerActor, TEXT("ControllerGameplayCameraEvaluationComponent"), RF_Transient);
 		ControllerComponent->RegisterComponent();
+		if (bOutCreated)
+		{
+			*bOutCreated = true;
+		}
 	}
 	return ControllerComponent;
 }
 
-TSharedPtr<UE::Cameras::FCameraEvaluationContext> UControllerGameplayCameraEvaluationComponent::FindEvaluationContext(APlayerController* PlayerController)
-{
-	if (UControllerGameplayCameraEvaluationComponent* ControllerComponent = FindComponent(PlayerController))
-	{
-		ControllerComponent->EnsureEvaluationContext();
-		return ControllerComponent->EvaluationContext;
-	}
-	return nullptr;
-}
-
-TSharedRef<UE::Cameras::FCameraEvaluationContext> UControllerGameplayCameraEvaluationComponent::FindOrAddEvaluationContext(APlayerController* PlayerController)
-{
-	UControllerGameplayCameraEvaluationComponent* ControllerComponent = FindOrAddComponent(PlayerController);
-	check(ControllerComponent);
-	ControllerComponent->EnsureEvaluationContext();
-	return ControllerComponent->EvaluationContext.ToSharedRef();
-}
-
diff --git a/Engine/Plugins/Cameras/GameplayCameras/Source/GameplayCameras/Private/GameFramework/GameplayCameraActor.cpp b/Engine/Plugins/Cameras/GameplayCameras/Source/GameplayCameras/Private/GameFramework/GameplayCameraActor.cpp
--- a/Engine/Plugins/Cameras/GameplayCameras/Source/GameplayCameras/Private/GameFramework/GameplayCameraActor.cpp
+++ b/Engine/Plugins/Cameras/GameplayCameras/Source/GameplayCameras/Private/GameFramework/GameplayCameraActor.cpp
@@ -23,19 +23,9 @@ USceneComponent* AGameplayCameraActor::GetDefaultAttachComponent() const
 	return CameraComponent;
 }
 
-void AGameplayCameraActor::Tick(float DeltaTime)
+UGameplayCameraComponentBase* AGameplayCameraActor::GetCameraComponentBase() const
 {
-	Super::Tick(DeltaTime);
-}
-
-bool AGameplayCameraActor::ShouldTickIfViewportsOnly() const
-{
-	return true;
-}
-
-void AGameplayCameraActor::DisplayDebug(UCanvas* Canvas, const FDebugDisplayInfo& DebugDisplay, float& YL, float& YPos)
-{
-	Super::DisplayDebug(Canvas, DebugDisplay, YL, YPos);
+	return CameraComponent;
 }
 
 #undef LOCTEXT_NAMESPACE
diff --git a/Engine/Plugins/Cameras/GameplayCameras/Source/GameplayCameras/Private/GameFramework/GameplayCameraActorBase.cpp b/Engine/Plugins/Cameras/GameplayCameras/Source/GameplayCameras/Private/GameFramework/GameplayCameraActorBase.cpp
new file mode 100644
--- /dev/null
+++ b/Engine/Plugins/Cameras/GameplayCameras/Source/GameplayCameras/Private/GameFramework/GameplayCameraActorBase.cpp
@@ -0,0 +1,33 @@
+// Copyright Epic Games, Inc. All Rights Reserved.
+
+#include "GameFramework/GameplayCameraActorBase.h"
+
+#include "Core/CameraSystemEvaluator.h"
+#include "GameFramework/GameplayCameraComponentBase.h"
+
+#include UE_INLINE_GENERATED_CPP_BY_NAME(GameplayCameraActorBase)
+
+#define LOCTEXT_NAMESPACE "GameplayCameraActorBase"
+
+AGameplayCameraActorBase::AGameplayCameraActorBase(const FObjectInitializer& ObjectInit)
+	: Super(ObjectInit)
+{
+}
+
+void AGameplayCameraActorBase::CalcCamera(float DeltaTime, FMinimalViewInfo& OutResult)
+{
+	using namespace UE::Cameras;
+
+	UGameplayCameraComponentBase* CameraComponentBase = GetCameraComponentBase();
+	if (ensure(CameraComponentBase))
+	{
+		CameraComponentBase->CalcCamera(DeltaTime, OutResult);
+	}
+	else
+	{
+		Super::CalcCamera(DeltaTime, OutResult);
+	}
+}
+
+#undef LOCTEXT_NAMESPACE
+
diff --git a/Engine/Plugins/Cameras/GameplayCameras/Source/GameplayCameras/Private/GameFramework/GameplayCameraComponent.cpp b/Engine/Plugins/Cameras/GameplayCameras/Source/GameplayCameras/Private/GameFramework/GameplayCameraComponent.cpp
--- a/Engine/Plugins/Cameras/GameplayCameras/Source/GameplayCameras/Private/GameFramework/GameplayCameraComponent.cpp
+++ b/Engine/Plugins/Cameras/GameplayCameras/Source/GameplayCameras/Private/GameFramework/GameplayCameraComponent.cpp
@@ -92,7 +92,7 @@ void UGameplayCameraComponent::OnCameraAssetBuilt(const UCameraAsset* InCameraAs
 
 	// If our camera asset was just built, it may have some new parameters. We need to rebuild
 	// our variable table and context data table, and re-apply overrides.
-	if (HasActiveEditorPreview())
+	if (HasCameraEvaluationContext())
 	{
 		const FCameraAssetAllocationInfo& AllocationInfo = InCameraAsset->GetAllocationInfo();
 		FCameraNodeEvaluationResult& InitialResult = GetEvaluationContext()->GetInitialResult();
@@ -114,12 +114,12 @@ void UGameplayCameraComponent::PostEditChangeProperty( struct FPropertyChangedEv
 	const FName MemberPropertyName = PropertyChangedEvent.GetMemberPropertyName();
 	if (MemberPropertyName == GET_MEMBER_NAME_CHECKED(UGameplayCameraComponent, CameraReference))
 	{
-		if (HasActiveEditorPreview())
+		if (HasCameraEvaluationContext())
 		{
 			if (PropertyChangedEvent.GetPropertyName() == TEXT("CameraAsset"))
 			{
 				// The camera asset has changed! Recreate the context.
-				RecreateEditorPreviewEvaluationContext();
+				RecreateCameraEvaluationContext();
 			}
 			else
 			{
diff --git a/Engine/Plugins/Cameras/GameplayCameras/Source/GameplayCameras/Private/GameFramework/GameplayCameraComponentBase.cpp b/Engine/Plugins/Cameras/GameplayCameras/Source/GameplayCameras/Private/GameFramework/GameplayCameraComponentBase.cpp
--- a/Engine/Plugins/Cameras/GameplayCameras/Source/GameplayCameras/Private/GameFramework/GameplayCameraComponentBase.cpp
+++ b/Engine/Plugins/Cameras/GameplayCameras/Source/GameplayCameras/Private/GameFramework/GameplayCameraComponentBase.cpp
@@ -2,8 +2,6 @@
 
 #include "GameFramework/GameplayCameraComponentBase.h"
 
-#include "Build/CameraAssetBuilder.h"
-#include "Build/CameraBuildLog.h"
 #include "CineCameraComponent.h"
 #include "Core/CameraAsset.h"
 #include "Core/CameraSystemEvaluator.h"
@@ -12,8 +10,6 @@
 #include "Engine/Canvas.h"
 #include "Engine/EngineTypes.h"
 #include "Engine/World.h"
-#include "GameFramework/GameplayCameraSystemActor.h"
-#include "GameFramework/GameplayCameraSystemHost.h"
 #include "GameplayCamerasDelegates.h"
 #include "IGameplayCamerasLiveEditManager.h"
 #include "IGameplayCamerasModule.h"
@@ -35,6 +31,7 @@
 UGameplayCameraComponentBase::UGameplayCameraComponentBase(const FObjectInitializer& ObjectInit)
 	: Super(ObjectInit)
 {
+	bAutoActivate = true;
 	bTickInEditor = true;
 	bWantsOnUpdateTransform = true;
 
@@ -46,21 +43,11 @@ UGameplayCameraComponentBase::UGameplayCameraComponentBase(const FObjectInitiali
 
 void UGameplayCameraComponentBase::BeginDestroy()
 {
-	Super::BeginDestroy();
-
-#if WITH_EDITOR
-
-	if (EditorPreviewEvaluator)
-	{
-		EditorPreviewEvaluator.Reset();
-	}
+	DestroyCameraSystem();
 
-	if (EvaluationContext)
-	{
-		EvaluationContext.Reset();
-	}
+	EvaluationContext = nullptr;
 
-#endif  // WITH_EDITOR
+	Super::BeginDestroy();
 }
 
 void UGameplayCameraComponentBase::AddReferencedObjects(UObject* InThis, FReferenceCollector& Collector)
@@ -68,19 +55,11 @@ void UGameplayCameraComponentBase::AddReferencedObjects(UObject* InThis, FRefere
 	Super::AddReferencedObjects(InThis, Collector);
 
 	UGameplayCameraComponentBase* This = CastChecked<UGameplayCameraComponentBase>(InThis);
+	This->OnAddReferencedObjects(Collector);
 	if (This->EvaluationContext.IsValid())
 	{
 		This->EvaluationContext->AddReferencedObjects(Collector);
 	}
-
-#if WITH_EDITOR
-
-	if (This->EditorPreviewEvaluator)
-	{
-		This->EditorPreviewEvaluator->AddReferencedObjects(Collector);
-	}
-
-#endif  // WITH_EDITOR
 }
 
 TSharedPtr<const UE::Cameras::FCameraEvaluationContext> UGameplayCameraComponentBase::GetEvaluationContext() const
@@ -93,87 +72,174 @@ TSharedPtr<UE::Cameras::FCameraEvaluationContext> UGameplayCameraComponentBase::
 	return EvaluationContext;
 }
 
-APlayerController* UGameplayCameraComponentBase::GetPlayerController() const
+void UGameplayCameraComponentBase::ActivateCameraForPlayerIndex(
+		int32 PlayerIndex, bool bSetAsViewTarget, EGameplayCameraComponentActivationMode ActivationMode)
 {
-	if (CameraSystemHost)
-	{
-		return CameraSystemHost->GetPlayerController();
-	}
-	return nullptr;
+	APlayerController* PlayerController = UGameplayStatics::GetPlayerController(this, PlayerIndex);
+	ActivateCameraForPlayerController(PlayerController, bSetAsViewTarget, ActivationMode);
 }
 
-void UGameplayCameraComponentBase::ActivateCameraForPlayerIndex(int32 PlayerIndex, EGameplayCameraComponentActivationMode ActivationMode)
+void UGameplayCameraComponentBase::ActivateCameraForPlayerController(
+		APlayerController* PlayerController, bool bSetAsViewTarget, EGameplayCameraComponentActivationMode ActivationMode)
 {
-	ActivateCameraEvaluationContext(PlayerIndex, ActivationMode);
-}
+	using namespace UE::Cameras;
 
-void UGameplayCameraComponentBase::ActivateCameraForPlayerController(APlayerController* PlayerController, EGameplayCameraComponentActivationMode ActivationMode)
-{
-	ActivateCameraEvaluationContext(PlayerController, ActivationMode);
+	// Make sure we are activated, since we need to tick and udpate our evaluation context and, possibly,
+	// our private camera system.
+	Super::Activate(false);
+
+	// Deactivate any existing evaluation context immediately first, since we might be re-activating with
+	// a different player controller, a different insertion point in the director tree, etc.
+	DeactivateCameraEvaluationContext(true);
+
+	// See if we are going to run the camera system on our own, or if we can activate our camera inside
+	// the player camera manager.
+	IGameplayCameraSystemHost* PlayerControllerHost = nullptr;
+	if (bSetAsViewTarget && PlayerController)
+	{
+		PlayerControllerHost = Cast<IGameplayCameraSystemHost>(PlayerController->PlayerCameraManager);
+	}
+	if (!PlayerControllerHost && ActivationMode != EGameplayCameraComponentActivationMode::Push)
+	{
+		UE_LOG(LogCameraSystem, Warning,
+				TEXT("Gameplay camera component '%s' cannot activate with mode '%s' because no camera system "
+					 "was found on the given player controller, or not player controller was specified. "),
+				*GetNameSafe(this),
+				*UEnum::GetValueAsString(ActivationMode));
+	}
+
+	if (PlayerControllerHost)
+	{
+		TeardownCameraSystemHost();
+		ActivateCameraEvaluationContext(PlayerController, PlayerControllerHost, ActivationMode);
+	}
+	else
+	{
+		EnsureCameraSystemHost();
+		ActivateCameraEvaluationContext(PlayerController, this, EGameplayCameraComponentActivationMode::Push);
+
+		if (bSetAsViewTarget && PlayerController)
+		{
+			AActor* OwnerActor = GetOwner();
+			PlayerController->SetViewTarget(OwnerActor);
+		}
+	}
 }
 
-void UGameplayCameraComponentBase::DeactivateCamera()
+bool UGameplayCameraComponentBase::CalcCamera(float DeltaTime, FMinimalViewInfo& OutResult)
 {
-	DeactivateCameraEvaluationContext();
+	using namespace UE::Cameras;
+
+	if (CameraSystemEvaluator)
+	{
+		CameraSystemEvaluator->GetEvaluatedCameraView(OutResult);
+		return true;
+	}
+	return false;
 }
 
-void UGameplayCameraComponentBase::ActivateCameraEvaluationContext(int32 PlayerIndex, EGameplayCameraComponentActivationMode ActivationMode)
+void UGameplayCameraComponentBase::DeactivateCamera(bool bImmediately)
 {
-	DeactivateCameraEvaluationContext();
-
-	APlayerController* PlayerController = UGameplayStatics::GetPlayerController(this, PlayerIndex);
-	ActivateCameraEvaluationContext(PlayerController, ActivationMode);
+	DeactivateCameraEvaluationContext(bImmediately);
 }
 
-void UGameplayCameraComponentBase::DeactivateCameraEvaluationContext()
+void UGameplayCameraComponentBase::DeactivateCameraEvaluationContext(bool bImmediately)
 {
 	using namespace UE::Cameras;
 
-	if (!CameraSystemHost)
+	// Remove the evaluation context for the director tree.
+	if (EvaluationContext && CameraSystemEvaluator)
 	{
-		return;
-	}
+		UE_LOG(LogCameraSystem, Log, TEXT("Deactivating gameplay camera '%s'."), *GetNameSafe(this));
 
-	if (EvaluationContext.IsValid())
-	{
 		if (TSharedPtr<FCameraEvaluationContext> ParentContext = EvaluationContext->GetParentContext())
 		{
 			ParentContext->RemoveChildContext(EvaluationContext.ToSharedRef());
 		}
 		else
 		{
-			TSharedPtr<FCameraSystemEvaluator> Evaluator = CameraSystemHost->GetCameraSystemEvaluator();
-			Evaluator->RemoveEvaluationContext(EvaluationContext.ToSharedRef());
+			CameraSystemEvaluator->RemoveEvaluationContext(EvaluationContext.ToSharedRef());
+		}
+
+		if (bImmediately)
+		{
+			// We are deactivating immediately (i.e. without letting our camera rigs blend out), so make 
+			// sure everything is frozen or disabled before we delete our evaluation context.
+			FRootCameraNodeEvaluator* RootNodeEvaluator = CameraSystemEvaluator->GetRootNodeEvaluator();
+			RootNodeEvaluator->DeactivateAllCameraRigs(EvaluationContext, true);
 		}
 	}
 
-	if (OutputCameraComponent)
+	if (EvaluationContext && CameraSystemEvaluator && !bImmediately)
 	{
-		OutputCameraComponent->SetRelativeTransform(FTransform());
+		// Don't deactivate the component right away: we still need to update our evaluation context 
+		// while any running camera rigs blend out.
+		bIsDeactivating = true;
 	}
+	else
+	{
+		if (OutputCameraComponent)
+		{
+			OutputCameraComponent->SetRelativeTransform(FTransform());
+		}
 
-	// Don't deactivate the component: we still need to update our evaluation context while any
-	// running camera rigs blend out.
+		EvaluationContext = nullptr;
+	}
 }
 
-void UGameplayCameraComponentBase::ActivateCameraEvaluationContext(APlayerController* PlayerController, EGameplayCameraComponentActivationMode ActivationMode)
+bool UGameplayCameraComponentBase::CanRunCameraSystem() const
+{
+#if WITH_EDITOR
+	IGameplayCamerasModule& Module = IGameplayCamerasModule::Get();
+	const bool bCanRunInEditor = bRunInEditor && Module.GetLiveEditManager()->CanRunInEditor();
+	return (!bIsEditorWorld || bCanRunInEditor);
+#else
+	return true;
+#endif  // WITH_EDITOR
+}
+
+void UGameplayCameraComponentBase::EnsureCameraSystemHost()
 {
 	using namespace UE::Cameras;
 
-	if (!PlayerController)
+	if (!HasCameraSystem() && CanRunCameraSystem())
+	{
+		AActor* OwnerActor = GetOwner();
+		UE_LOG(LogCameraSystem, Log, 
+				TEXT("Creating camera system host for gameplay camera '%s'."),
+				*GetNameSafe(OwnerActor));
+
+		FCameraSystemEvaluatorCreateParams Params;
+		Params.Owner = this;
+#if WITH_EDITOR
+		if (bIsEditorWorld)
+		{
+			Params.Role = ECameraSystemEvaluatorRole::EditorPreview;
+		}
+#endif  // WITH_EDITOR
+		InitializeCameraSystem(Params);
+	}
+}
+
+void UGameplayCameraComponentBase::TeardownCameraSystemHost()
+{
+	if (HasCameraSystem())
+	{
+		DestroyCameraSystem();
+	}
+}
+
+void UGameplayCameraComponentBase::ActivateCameraEvaluationContext(APlayerController* PlayerController, IGameplayCameraSystemHost* Host, EGameplayCameraComponentActivationMode ActivationMode)
+{
+	using namespace UE::Cameras;
+
+	if (!ensure(Host))
 	{
-		FFrame::KismetExecutionMessage(
-				TEXT("Can't activate gameplay camera component: invalid player controller!"),
-				ELogVerbosity::Error);
 		return;
 	}
 
-	CameraSystemHost = UGameplayCameraSystemHost::FindOrCreateHost(PlayerController);
-	if (!CameraSystemHost)
+	if (!CanRunCameraSystem())
 	{
-		FFrame::KismetExecutionMessage(
-				TEXT("Can't activate gameplay camera component: no camera system host found!"),
-				ELogVerbosity::Error);
 		return;
 	}
 
@@ -183,8 +249,6 @@ void UGameplayCameraComponentBase::ActivateCameraEvaluationContext(APlayerContro
 		return;
 	}
 
-	AGameplayCameraSystemActor::AutoManageActiveViewTarget(PlayerController);
-
 	// Make sure the evaluation context has been created. However, this can fail, such as when
 	// we don't have a valid camera asset specified.
 	TryCreateCameraEvaluationContext(PlayerController);
@@ -202,8 +266,13 @@ void UGameplayCameraComponentBase::ActivateCameraEvaluationContext(APlayerContro
 		return;
 	}
 
-	TSharedPtr<FCameraSystemEvaluator> CameraSystemEvaluator = CameraSystemHost->GetCameraSystemEvaluator();
-	FCameraEvaluationContextStack& ContextStack = CameraSystemEvaluator->GetEvaluationContextStack();
+	UE_LOG(LogCameraSystem, Log, 
+			TEXT("Activating gameplay camera '%s' with mode '%s'."),
+			*GetNameSafe(this), 
+			*UEnum::GetValueAsString(ActivationMode));
+
+	TSharedPtr<FCameraSystemEvaluator> HostEvaluator = Host->GetCameraSystemEvaluator();
+	FCameraEvaluationContextStack& ContextStack = HostEvaluator->GetEvaluationContextStack();
 
 	switch (ActivationMode)
 	{
@@ -238,14 +307,16 @@ void UGameplayCameraComponentBase::ActivateCameraEvaluationContext(APlayerContro
 			break;
 	}
 
-	// Make sure the component is active so it receives tick updates to maintain the evaluation context.
-	Activate();
+	// Cancel any ongoing deactivation.
+	bIsDeactivating = false;
 }
 
 void UGameplayCameraComponentBase::TryCreateCameraEvaluationContext(APlayerController* PlayerController)
 {
 	using namespace UE::Cameras;
 
+	ensure(CanRunCameraSystem());
+
 	if (!EvaluationContext.IsValid())
 	{
 		if (const UCameraAsset* CameraAsset = GetCameraAsset())
@@ -314,33 +385,23 @@ void UGameplayCameraComponentBase::BeginPlay()
 
 	Super::BeginPlay();
 
-#if WITH_EDITOR
-
-	UCameraAsset* CameraAsset = GetCameraAsset();
-	if (CameraAsset)
+	if (IsActive())
 	{
-		UWorld* World = GetWorld();
-		if (World && World->WorldType == EWorldType::PIE)
+		if (AutoActivateForPlayer != EAutoReceiveInput::Disabled && GetNetMode() != NM_DedicatedServer)
 		{
-			// Auto-build the camera asset on begin play to make sure we've got the latest user edits.
-			FCameraBuildLog BuildLog;
-			FCameraAssetBuilder Builder(BuildLog);
-			Builder.BuildCamera(CameraAsset);
+			const int32 PlayerIndex = AutoActivateForPlayer.GetIntValue() - 1;
+			ActivateCameraForPlayerIndex(PlayerIndex);
+		}
+		else
+		{
+			ActivateCameraForPlayerController(nullptr, false, EGameplayCameraComponentActivationMode::Push);
 		}
-	}
-
-#endif  // WITH_EDITOR
-
-	if (IsActive() && AutoActivateForPlayer != EAutoReceiveInput::Disabled && GetNetMode() != NM_DedicatedServer)
-	{
-		const int32 PlayerIndex = AutoActivateForPlayer.GetIntValue() - 1;
-		ActivateCameraForPlayerIndex(PlayerIndex);
 	}
 }
 
 void UGameplayCameraComponentBase::EndPlay(const EEndPlayReason::Type EndPlayReason)
 {
-	DeactivateCameraEvaluationContext();
+	DeactivateCameraEvaluationContext(true);
 
 	Super::EndPlay(EndPlayReason);
 }
@@ -358,6 +419,30 @@ void UGameplayCameraComponentBase::OnUnregister()
 	Super::OnUnregister();
 }
 
+void UGameplayCameraComponentBase::Activate(bool bReset)
+{
+	// When auto-activing, this method gets called during OnRegister, before we have started playing.
+	// In this case, we don't activate the camera right away, we wait until BeginPlay.
+	const bool bDoActivate = (bReset || ShouldActivate()) && HasBegunPlay();
+
+	Super::Activate(bReset);
+
+	if (bDoActivate)
+	{
+		DeactivateCameraEvaluationContext(true);
+
+		EnsureCameraSystemHost();
+		ActivateCameraEvaluationContext(nullptr, this, EGameplayCameraComponentActivationMode::Push);
+	}
+}
+
+void UGameplayCameraComponentBase::Deactivate()
+{
+	DeactivateCameraEvaluationContext(true);
+
+	Super::Deactivate();
+}
+
 void UGameplayCameraComponentBase::TickComponent(float DeltaTime, ELevelTick TickType, FActorComponentTickFunction *ThisTickFunction)
 {
 	using namespace UE::Cameras;
@@ -374,14 +459,43 @@ void UGameplayCameraComponentBase::TickComponent(float DeltaTime, ELevelTick Tic
 	if (EvaluationContext)
 	{
 		UpdateCameraEvaluationContext(false);
+	}
 
-#if WITH_EDITOR
-		
-		UpdateEditorPreviewEvaluator(DeltaTime);
+	UpdateCameraSystem(DeltaTime);
+	UpdateOutputCameraComponent();
 
-#endif  // WITH_EDITOR
+	CheckPendingDeactivation();
+}
+
+void UGameplayCameraComponentBase::CheckPendingDeactivation()
+{
+	using namespace UE::Cameras;
+
+	if (!bIsDeactivating)
+	{
+		return;
+	}
+
+	// If we were deactivating, we wait until we don't have any running camera rig anymore,
+	// at which point we can tear down all our evaluation apparatus.
+	bool bDoneDeactivating = true;
 
-		UpdateOutputCameraComponent();
+	if (CameraSystemEvaluator && EvaluationContext)
+	{
+		FRootCameraNodeEvaluator* RootNodeEvaluator = CameraSystemEvaluator->GetRootNodeEvaluator();
+		bDoneDeactivating = (RootNodeEvaluator->HasAnyRunningCameraRig(EvaluationContext) == false);
+	}
+
+	if (bDoneDeactivating)
+	{
+		TeardownCameraSystemHost();
+		EvaluationContext = nullptr;
+
+		// Only call the base class method here: we just want to finish deactivating ourselves
+		// by stopping ticking.
+		Super::Deactivate();
+
+		bIsDeactivating = false;
 	}
 }
 
@@ -405,10 +519,18 @@ void UGameplayCameraComponentBase::UpdateCameraEvaluationContext(bool bForceAppl
 	OnUpdateCameraEvaluationContext(bForceApplyParameterOverrides);
 
 #if WITH_EDITOR
+
 	EvaluationContext->UpdateForEditorPreview();
+
 #endif  // WITH_EDITOR
 }
 
+void UGameplayCameraComponentBase::RecreateCameraEvaluationContext()
+{
+	EvaluationContext = nullptr;
+	TryCreateCameraEvaluationContext(nullptr);
+}
+
 void UGameplayCameraComponentBase::UpdateOutputCameraComponent()
 {
 	using namespace UE::Cameras;
@@ -418,19 +540,6 @@ void UGameplayCameraComponentBase::UpdateOutputCameraComponent()
 		return;
 	}
 
-	TSharedPtr<FCameraSystemEvaluator> CameraSystemEvaluator;
-
-	if (CameraSystemHost)
-	{
-		CameraSystemEvaluator = CameraSystemHost->GetCameraSystemEvaluator();
-	}
-#if WITH_EDITOR
-	else if (EditorPreviewEvaluator)
-	{
-		CameraSystemEvaluator = EditorPreviewEvaluator;
-	}
-#endif  // WITH_EDITOR
-
 	bool bGotValidTransform = false;
 	if (CameraSystemEvaluator)
 	{
@@ -472,11 +581,6 @@ void UGameplayCameraComponentBase::UpdateOutputCameraComponent()
 	}
 }
 
-void UGameplayCameraComponentBase::OnComponentDestroyed(bool bDestroyingHierarchy)
-{
-	Super::OnComponentDestroyed(bDestroyingHierarchy);
-}
-
 void UGameplayCameraComponentBase::OnUpdateTransform(EUpdateTransformFlags UpdateTransformFlags, ETeleportType Teleport)
 {
 	Super::OnUpdateTransform(UpdateTransformFlags, Teleport);
@@ -507,76 +611,35 @@ void UGameplayCameraComponentBase::AutoManageEditorPreviewEvaluator()
 		return;
 	}
 	
-	if (CanRunInEditor() && !(EditorPreviewEvaluator && EvaluationContext))
+	const bool bCanRun = CanRunCameraSystem();
+	if (bCanRun && !(CameraSystemEvaluator && EvaluationContext))
 	{
 		// We want to run the camera logic in the editor but we haven't set things up for that.
 		// Let's create the preview evaluator and the evaluation context.
-		if (!EditorPreviewEvaluator)
-		{
-			EditorPreviewEvaluator = MakeShared<FCameraSystemEvaluator>();
+		EnsureCameraSystemHost();
 
-			FCameraSystemEvaluatorCreateParams CreateParams;
-			CreateParams.Owner = this;
-			CreateParams.Role = ECameraSystemEvaluatorRole::EditorPreview;
-			EditorPreviewEvaluator->Initialize(CreateParams);
-		}
-		if (!EvaluationContext)
+		ActivateCameraEvaluationContext(nullptr, this, EGameplayCameraComponentActivationMode::Push);
+		if (ensure(EvaluationContext))
 		{
-			TryCreateCameraEvaluationContext(nullptr);
-			if (EvaluationContext.IsValid())
-			{
-				EditorPreviewEvaluator->PushEvaluationContext(EvaluationContext.ToSharedRef());
-				EvaluationContext->SetEditorPreviewCameraRigIndex(EditorPreviewCameraRigIndex);
-			}
+			EvaluationContext->SetEditorPreviewCameraRigIndex(EditorPreviewCameraRigIndex);
 		}
 	}
-	else if (!CanRunInEditor() && (EditorPreviewEvaluator || EvaluationContext))
+	else if (!bCanRun && (CameraSystemEvaluator || EvaluationContext))
 	{
 		// We don't want to run the camera logic in the editor anymore. Let's tear things down.
-		EditorPreviewEvaluator = nullptr;
+		DeactivateCameraEvaluationContext(true);
+		TeardownCameraSystemHost();
 		EvaluationContext = nullptr;
 	}
 }
 
 void UGameplayCameraComponentBase::OnEditorPreviewCameraRigIndexChanged()
 {
-	if (!bIsEditorWorld)
-	{
-		return;
-	}
-
-	if (CanRunInEditor() && EditorPreviewEvaluator && EvaluationContext)
+	if (bIsEditorWorld)
 	{
-		EvaluationContext->SetEditorPreviewCameraRigIndex(EditorPreviewCameraRigIndex);
-	}
-}
-
-bool UGameplayCameraComponentBase::CanRunInEditor() const
-{
-	IGameplayCamerasModule& Module = IGameplayCamerasModule::Get();
-	return bRunInEditor && Module.GetLiveEditManager()->CanRunInEditor();
-}
-
-bool UGameplayCameraComponentBase::HasActiveEditorPreview() const
-{
-	return (
-			bIsEditorWorld &&
-			CanRunInEditor() &&
-			EditorPreviewEvaluator &&
-			EvaluationContext);
-}
-
-void UGameplayCameraComponentBase::RecreateEditorPreviewEvaluationContext()
-{
-	if (bIsEditorWorld && CanRunInEditor() && EditorPreviewEvaluator && EvaluationContext)
-	{
-		EditorPreviewEvaluator->RemoveEvaluationContext(EvaluationContext.ToSharedRef());
-		EvaluationContext = nullptr;
-
-		TryCreateCameraEvaluationContext(nullptr);
-		if (EvaluationContext.IsValid())
+		const bool bCanRun = CanRunCameraSystem();
+		if (bCanRun && CameraSystemEvaluator && EvaluationContext)
 		{
-			EditorPreviewEvaluator->PushEvaluationContext(EvaluationContext.ToSharedRef());
 			EvaluationContext->SetEditorPreviewCameraRigIndex(EditorPreviewCameraRigIndex);
 		}
 	}
@@ -592,24 +655,13 @@ bool UGameplayCameraComponentBase::GetEditorPreviewInfo(float DeltaTime, FMinima
 	return false;
 }
 
-void UGameplayCameraComponentBase::UpdateEditorPreviewEvaluator(float DeltaTime)
-{
-	using namespace UE::Cameras;
-
-	if (EditorPreviewEvaluator)
-	{
-		FCameraSystemEvaluationParams Params;
-		Params.DeltaTime = DeltaTime;
-		EditorPreviewEvaluator->Update(Params);
-	}
-}
-
 void UGameplayCameraComponentBase::OnDrawVisualizationHUD(const FViewport* Viewport, const FSceneView* SceneView, FCanvas* Canvas) const
 {
 	using namespace UE::Cameras;
 
+	const bool bCanRun = CanRunCameraSystem();
 	const bool bHasShowFlag = SceneView->Family->EngineShowFlags.GetSingleFlag(CustomShowFlag);
-	if (bHasShowFlag && CanRunInEditor() && EditorPreviewEvaluator && EvaluationContext)
+	if (bCanRun && bHasShowFlag && CameraSystemEvaluator && EvaluationContext)
 	{
 		const AActor* OwnerActor = GetOwner();
 
@@ -622,7 +674,7 @@ void UGameplayCameraComponentBase::OnDrawVisualizationHUD(const FViewport* Viewp
 		Params.bIsLockedToCamera = bIsLockedToCamera;
 		Params.bDrawWorldDebug = false;
 
-		EditorPreviewEvaluator->DrawEditorPreview(Params);
+		CameraSystemEvaluator->DrawEditorPreview(Params);
 	}
 }
 
@@ -675,4 +727,3 @@ void FGameplayCameraComponentEvaluationContext::UpdateForEditorPreview()
 
 #undef LOCTEXT_NAMESPACE
 
-
diff --git a/Engine/Plugins/Cameras/GameplayCameras/Source/GameplayCameras/Private/GameFramework/GameplayCameraRigActor.cpp b/Engine/Plugins/Cameras/GameplayCameras/Source/GameplayCameras/Private/GameFramework/GameplayCameraRigActor.cpp
--- a/Engine/Plugins/Cameras/GameplayCameras/Source/GameplayCameras/Private/GameFramework/GameplayCameraRigActor.cpp
+++ b/Engine/Plugins/Cameras/GameplayCameras/Source/GameplayCameras/Private/GameFramework/GameplayCameraRigActor.cpp
@@ -23,19 +23,9 @@ USceneComponent* AGameplayCameraRigActor::GetDefaultAttachComponent() const
 	return CameraRigComponent;
 }
 
-void AGameplayCameraRigActor::Tick(float DeltaTime)
+UGameplayCameraComponentBase* AGameplayCameraRigActor::GetCameraComponentBase() const
 {
-	Super::Tick(DeltaTime);
-}
-
-bool AGameplayCameraRigActor::ShouldTickIfViewportsOnly() const
-{
-	return true;
-}
-
-void AGameplayCameraRigActor::DisplayDebug(UCanvas* Canvas, const FDebugDisplayInfo& DebugDisplay, float& YL, float& YPos)
-{
-	Super::DisplayDebug(Canvas, DebugDisplay, YL, YPos);
+	return CameraRigComponent;
 }
 
 #undef LOCTEXT_NAMESPACE
diff --git a/Engine/Plugins/Cameras/GameplayCameras/Source/GameplayCameras/Private/GameFramework/GameplayCameraRigComponent.cpp b/Engine/Plugins/Cameras/GameplayCameras/Source/GameplayCameras/Private/GameFramework/GameplayCameraRigComponent.cpp
--- a/Engine/Plugins/Cameras/GameplayCameras/Source/GameplayCameras/Private/GameFramework/GameplayCameraRigComponent.cpp
+++ b/Engine/Plugins/Cameras/GameplayCameras/Source/GameplayCameras/Private/GameFramework/GameplayCameraRigComponent.cpp
@@ -92,7 +92,7 @@ void UGameplayCameraRigComponent::OnCameraRigAssetBuilt(const UCameraRigAsset* I
 
 	// If our camera rig asset was just built, it may have some new parameters. We need to rebuild
 	// our variable table and context data table, and re-apply overrides.
-	if (HasActiveEditorPreview())
+	if (HasCameraEvaluationContext())
 	{
 		const FCameraObjectAllocationInfo& AllocationInfo = InCameraRigAsset->AllocationInfo;
 		FCameraNodeEvaluationResult& InitialResult = GetEvaluationContext()->GetInitialResult();
@@ -114,13 +114,13 @@ void UGameplayCameraRigComponent::PostEditChangeProperty( struct FPropertyChange
 	const FName MemberPropertyName = PropertyChangedEvent.GetMemberPropertyName();
 	if (MemberPropertyName == GET_MEMBER_NAME_CHECKED(UGameplayCameraRigComponent, CameraRigReference))
 	{
-		if (HasActiveEditorPreview())
+		if (HasCameraEvaluationContext())
 		{
 			if (PropertyChangedEvent.GetPropertyName() == TEXT("CameraRig"))
 			{
 				// The camera rig asset has changed! Recreate the context.
 				GeneratedCameraAsset = nullptr;
-				RecreateEditorPreviewEvaluationContext();
+				RecreateCameraEvaluationContext();
 			}
 			else
 			{
diff --git a/Engine/Plugins/Cameras/GameplayCameras/Source/GameplayCameras/Private/GameFramework/GameplayCameraSystemActor.cpp b/Engine/Plugins/Cameras/GameplayCameras/Source/GameplayCameras/Private/GameFramework/GameplayCameraSystemActor.cpp
--- a/Engine/Plugins/Cameras/GameplayCameras/Source/GameplayCameras/Private/GameFramework/GameplayCameraSystemActor.cpp
+++ b/Engine/Plugins/Cameras/GameplayCameras/Source/GameplayCameras/Private/GameFramework/GameplayCameraSystemActor.cpp
@@ -2,13 +2,7 @@
 
 #include "GameFramework/GameplayCameraSystemActor.h"
 
-#include "Engine/EngineTypes.h"
-#include "Engine/World.h"
 #include "GameFramework/GameplayCameraSystemComponent.h"
-#include "GameFramework/GameplayCameraSystemHost.h"
-#include "GameFramework/GameplayCamerasPlayerCameraManager.h"
-#include "GameFramework/PlayerController.h"
-#include "GameplayCamerasSettings.h"
 
 #include UE_INLINE_GENERATED_CPP_BY_NAME(GameplayCameraSystemActor)
 
@@ -26,86 +20,5 @@ void AGameplayCameraSystemActor::CalcCamera(float DeltaTime, struct FMinimalView
 	CameraSystemComponent->GetCameraView(DeltaTime, OutResult);
 }
 
-AGameplayCameraSystemActor* AGameplayCameraSystemActor::GetAutoSpawnedCameraSystemActor(APlayerController* PlayerController, bool bSpawnIfMissing)
-{
-	bool bDidSpawn;
-	return GetAutoSpawnedCameraSystemActor(PlayerController, bSpawnIfMissing, &bDidSpawn);
-}
-
-AGameplayCameraSystemActor* AGameplayCameraSystemActor::GetAutoSpawnedCameraSystemActor(APlayerController* PlayerController, bool bSpawnIfMissing, bool* bOutSpawned)
-{
-	static const TCHAR* AutoSpawnedActorName = TEXT("AutoSpawnedGameplayCameraSystemActor");
-
-	UGameplayCameraSystemHost* Host = UGameplayCameraSystemHost::FindHost(PlayerController);
-	if (!Host)
-	{
-		if (bSpawnIfMissing)
-		{
-			Host = UGameplayCameraSystemHost::FindOrCreateHost(PlayerController);
-		}
-		else
-		{
-			FFrame::KismetExecutionMessage(
-					TEXT("Can't auto-manage active view target: no camera system host found!"),
-					ELogVerbosity::Error);
-			return nullptr;
-		}
-	}
-
-	*bOutSpawned = false;
-	AGameplayCameraSystemActor* SpawnedActor = FindObject<AGameplayCameraSystemActor>(PlayerController, AutoSpawnedActorName);
-	if (!SpawnedActor)
-	{
-		if (bSpawnIfMissing)
-		{
-			FActorSpawnParameters SpawnParams;
-			SpawnParams.Name = AutoSpawnedActorName;
-
-			UWorld* World = PlayerController->GetWorld();
-			SpawnedActor = World->SpawnActor<AGameplayCameraSystemActor>(SpawnParams);
-
-			SpawnedActor->Rename(nullptr, PlayerController);
-
-			*bOutSpawned = true;
-		}
-		else
-		{
-			return nullptr;
-		}
-	}
-
-	return SpawnedActor;
-}
-
-void AGameplayCameraSystemActor::AutoManageActiveViewTarget(APlayerController* PlayerController)
-{
-	const UGameplayCamerasSettings* Settings = GetDefault<UGameplayCamerasSettings>();
-	if (!Settings->bAutoSpawnCameraSystemActor)
-	{
-		return;
-	}
-
-	if (AGameplayCamerasPlayerCameraManager* PlayerCameraManager = Cast<AGameplayCamerasPlayerCameraManager>(PlayerController->PlayerCameraManager))
-	{
-		return;
-	}
-
-	bool bDidSpawn = false;
-	AGameplayCameraSystemActor* SpawnedActor = GetAutoSpawnedCameraSystemActor(PlayerController, true, &bDidSpawn);
-	if (SpawnedActor)
-	{
-		UGameplayCameraSystemComponent* CameraSystemComponent = SpawnedActor->CameraSystemComponent;
-		if (ensure(CameraSystemComponent))
-		{
-			if (bDidSpawn)
-			{
-				CameraSystemComponent->bSetPlayerControllerRotation = Settings->bAutoSpawnCameraSystemActorSetsControlRotation;
-			}
-
-			CameraSystemComponent->ActivateCameraSystemForPlayerController(PlayerController);
-		}
-	}
-}
-
 #undef LOCTEXT_NAMESPACE
 
diff --git a/Engine/Plugins/Cameras/GameplayCameras/Source/GameplayCameras/Private/GameFramework/GameplayCameraSystemComponent.cpp b/Engine/Plugins/Cameras/GameplayCameras/Source/GameplayCameras/Private/GameFramework/GameplayCameraSystemComponent.cpp
--- a/Engine/Plugins/Cameras/GameplayCameras/Source/GameplayCameras/Private/GameFramework/GameplayCameraSystemComponent.cpp
+++ b/Engine/Plugins/Cameras/GameplayCameras/Source/GameplayCameras/Private/GameFramework/GameplayCameraSystemComponent.cpp
@@ -11,7 +11,6 @@
 #include "Engine/Texture2D.h"
 #include "Engine/World.h"
 #include "GameFramework/Controller.h"
-#include "GameFramework/GameplayCameraSystemHost.h"
 #include "Kismet/GameplayStatics.h"
 #include "Services/CameraModifierService.h"
 #include "UObject/ConstructorHelpers.h"
@@ -26,35 +25,23 @@ UGameplayCameraSystemComponent::UGameplayCameraSystemComponent(const FObjectInit
 {
 }
 
-TSharedPtr<UE::Cameras::FCameraSystemEvaluator> UGameplayCameraSystemComponent::GetCameraSystemEvaluator(bool bEnsureIfNull)
-{
-	UGameplayCameraSystemHost* HostPtr = CameraSystemHost.Get();
-	ensureMsgf(HostPtr || !bEnsureIfNull, TEXT("Accessing camera system evaluator when we haven't found or created a host for one."));
-	if (HostPtr)
-	{
-		return HostPtr->GetCameraSystemEvaluator();
-	}
-	return nullptr;
-}
-
 void UGameplayCameraSystemComponent::GetCameraView(float DeltaTime, FMinimalViewInfo& DesiredView)
 {
 	using namespace UE::Cameras;
 
-	TSharedPtr<FCameraSystemEvaluator> Evaluator = GetCameraSystemEvaluator();
-	if (Evaluator.IsValid())
+	if (CameraSystemEvaluator.IsValid())
 	{
 		FCameraSystemEvaluationParams UpdateParams;
 		UpdateParams.DeltaTime = DeltaTime;
-		Evaluator->Update(UpdateParams);
+		CameraSystemEvaluator->Update(UpdateParams);
 
-		Evaluator->GetEvaluatedCameraView(DesiredView);
+		CameraSystemEvaluator->GetEvaluatedCameraView(DesiredView);
 
 		if (bSetPlayerControllerRotation)
 		{
 			if (APlayerController* PlayerController = WeakPlayerController.Get())
 			{
-				PlayerController->SetControlRotation(Evaluator->GetEvaluatedResult().CameraPose.GetRotation());
+				PlayerController->SetControlRotation(CameraSystemEvaluator->GetEvaluatedResult().CameraPose.GetRotation());
 			}
 		}
 	}
@@ -116,6 +103,8 @@ void UGameplayCameraSystemComponent::ActivateCameraSystemForPlayerIndex(int32 Pl
 
 void UGameplayCameraSystemComponent::ActivateCameraSystemForPlayerController(APlayerController* PlayerController)
 {
+	using namespace UE::Cameras;
+
 	if (!PlayerController)
 	{
 		FFrame::KismetExecutionMessage(
@@ -132,6 +121,8 @@ void UGameplayCameraSystemComponent::ActivateCameraSystemForPlayerController(APl
 		}
 	}
 
+	EnsureCameraSystemInitialized();
+
 	AActor* OwningActor = GetOwner();
 	if (!OwningActor)
 	{
@@ -141,16 +132,6 @@ void UGameplayCameraSystemComponent::ActivateCameraSystemForPlayerController(APl
 		return;
 	}
 
-	if (!CameraSystemHost)
-	{
-		CameraSystemHost = UGameplayCameraSystemHost::FindOrCreateHost(PlayerController);
-		if (!CameraSystemHost)
-		{
-			FFrame::KismetExecutionMessage(TEXT("can't create camera system host!"), ELogVerbosity::Error);
-			return;
-		}
-	}
-
 	PlayerController->SetViewTarget(OwningActor);
 	WeakPlayerController = PlayerController;
 
@@ -172,7 +153,7 @@ bool UGameplayCameraSystemComponent::IsCameraSystemActiveForPlayController(APlay
 		return false;
 	}
 	
-	if (!CameraSystemHost)
+	if (!HasCameraSystem())
 	{
 		return false;
 	}
@@ -219,9 +200,9 @@ FCameraRigInstanceID UGameplayCameraSystemComponent::StartGlobalCameraModifierRi
 {
 	using namespace UE::Cameras;
 
-	if (TSharedPtr<FCameraSystemEvaluator> Evaluator = GetCameraSystemEvaluator())
+	if (CameraSystemEvaluator)
 	{
-		TSharedPtr<FCameraModifierService> CameraModifierService = Evaluator->FindEvaluationService<FCameraModifierService>();
+		TSharedPtr<FCameraModifierService> CameraModifierService = CameraSystemEvaluator->FindEvaluationService<FCameraModifierService>();
 		return CameraModifierService->StartCameraModifierRig(CameraRig, ECameraRigLayer::Global, OrderKey);
 	}
 
@@ -232,9 +213,9 @@ FCameraRigInstanceID UGameplayCameraSystemComponent::StartVisualCameraModifierRi
 {
 	using namespace UE::Cameras;
 
-	if (TSharedPtr<FCameraSystemEvaluator> Evaluator = GetCameraSystemEvaluator())
+	if (CameraSystemEvaluator)
 	{
-		TSharedPtr<FCameraModifierService> CameraModifierService = Evaluator->FindEvaluationService<FCameraModifierService>();
+		TSharedPtr<FCameraModifierService> CameraModifierService = CameraSystemEvaluator->FindEvaluationService<FCameraModifierService>();
 		return CameraModifierService->StartCameraModifierRig(CameraRig, ECameraRigLayer::Visual);
 	}
 
@@ -245,9 +226,9 @@ void UGameplayCameraSystemComponent::StopCameraModifierRig(FCameraRigInstanceID
 {
 	using namespace UE::Cameras;
 
-	if (TSharedPtr<FCameraSystemEvaluator> Evaluator = GetCameraSystemEvaluator())
+	if (CameraSystemEvaluator)
 	{
-		TSharedPtr<FCameraModifierService> CameraModifierService = Evaluator->FindEvaluationService<FCameraModifierService>();
+		TSharedPtr<FCameraModifierService> CameraModifierService = CameraSystemEvaluator->FindEvaluationService<FCameraModifierService>();
 		CameraModifierService->StopCameraModifierRig(InstanceID, bImmediately);
 	}
 }
diff --git a/Engine/Plugins/Cameras/GameplayCameras/Source/GameplayCameras/Private/GameFramework/GameplayCameraSystemHost.cpp b/Engine/Plugins/Cameras/GameplayCameras/Source/GameplayCameras/Private/GameFramework/GameplayCameraSystemHost.cpp
deleted file mode 100644
diff --git a/Engine/Plugins/Cameras/GameplayCameras/Source/GameplayCameras/Private/GameFramework/GameplayCamerasPlayerCameraManager.cpp b/Engine/Plugins/Cameras/GameplayCameras/Source/GameplayCameras/Private/GameFramework/GameplayCamerasPlayerCameraManager.cpp
--- a/Engine/Plugins/Cameras/GameplayCameras/Source/GameplayCameras/Private/GameFramework/GameplayCamerasPlayerCameraManager.cpp
+++ b/Engine/Plugins/Cameras/GameplayCameras/Source/GameplayCameras/Private/GameFramework/GameplayCamerasPlayerCameraManager.cpp
@@ -13,11 +13,7 @@
 #include "Engine/Canvas.h"
 #include "Engine/Engine.h"
 #include "GameFramework/ActorCameraEvaluationContext.h"
-#include "GameFramework/GameplayCameraComponent.h"
-#include "GameFramework/GameplayCameraRigComponent.h"
-#include "GameFramework/GameplayCameraSystemActor.h"
-#include "GameFramework/GameplayCameraSystemComponent.h"
-#include "GameFramework/GameplayCameraSystemHost.h"
+#include "GameFramework/GameplayCameraComponentBase.h"
 #include "GameFramework/PlayerController.h"
 #include "Services/CameraModifierService.h"
 #include "UObject/Package.h"
@@ -125,6 +121,21 @@ AGameplayCamerasPlayerCameraManager::AGameplayCamerasPlayerCameraManager(const F
 {
 }
 
+void AGameplayCamerasPlayerCameraManager::BeginDestroy()
+{
+	TeardownCameraSystemHost();
+
+	Super::BeginDestroy();
+}
+
+void AGameplayCamerasPlayerCameraManager::AddReferencedObjects(UObject* InThis, FReferenceCollector& Collector)
+{
+	Super::AddReferencedObjects(InThis, Collector);
+
+	AGameplayCamerasPlayerCameraManager* This = CastChecked<AGameplayCamerasPlayerCameraManager>(InThis);
+	This->IGameplayCameraSystemHost::OnAddReferencedObjects(Collector);
+}
+
 void AGameplayCamerasPlayerCameraManager::StealPlayerController(APlayerController* PlayerController)
 {
 	using namespace UE::Cameras;
@@ -140,28 +151,7 @@ void AGameplayCamerasPlayerCameraManager::StealPlayerController(APlayerControlle
 	PlayerController->PlayerCameraManager = this;
 	InitializeFor(PlayerController);
 
-	if (WeakOriginalAutoCameraSystemActor != OriginalViewTarget)
-	{
-		// We had some arbitrary view target before. Let's keep it.
-		SetViewTarget(OriginalViewTarget);
-	}
-	else
-	{
-		// The original view target was our camera system, so nothing changes. We will restore 
-		// that actor as the view target on the original camera manager when we're done.
-		//
-		// Let's set the view target to whatever owns the active context, or fall back to the
-		// player controller if the active context owner isn't an actor. We need to avoid using
-		// SetViewTarget otherwise that would create another evaluation context. Instead we
-		// call AssignViewTarget directly.
-		TSharedPtr<FCameraEvaluationContext> ActiveContext = CameraSystemEvaluator->GetEvaluationContextStack().GetActiveContext();
-		if (ActiveContext && ActiveContext->GetOwner())
-		{
-			AActor* ActiveViewTarget = ActiveContext->GetOwner()->GetTypedOuter<AActor>();
-			AssignViewTarget(ActiveViewTarget, ViewTarget);
-		}
-		ViewTarget.CheckViewTarget(PlayerController);
-	}
+	SetViewTarget(OriginalViewTarget);
 }
 
 void AGameplayCamerasPlayerCameraManager::ReleasePlayerController()
@@ -172,17 +162,12 @@ void AGameplayCamerasPlayerCameraManager::ReleasePlayerController()
 	}
 
 	PCOwner->PlayerCameraManager = OriginalCameraManager;
-	if (AGameplayCameraSystemActor* OriginalAutoCameraSystemActor = WeakOriginalAutoCameraSystemActor.Get())
-	{
-		OriginalCameraManager->SetViewTarget(OriginalAutoCameraSystemActor);
-	}
 
 	ViewTarget.Target = nullptr;
 
 	OriginalCameraManager = nullptr;
-	WeakOriginalAutoCameraSystemActor = nullptr;
 
-	ReleaseCameraSystemHost();
+	TeardownCameraSystemHost();
 
 	PCOwner = nullptr;
 }
@@ -226,7 +211,8 @@ void AGameplayCamerasPlayerCameraManager::StopCameraModifierRig(FCameraRigInstan
 
 void AGameplayCamerasPlayerCameraManager::InitializeFor(APlayerController* PlayerController)
 {
-	AcquireCameraSystemHost(PlayerController);
+	EnsureCameraSystemHost();
+
 	Super::InitializeFor(PlayerController);
 }
 
@@ -264,14 +250,10 @@ void AGameplayCamerasPlayerCameraManager::SetViewTarget(AActor* NewViewTarget, F
 		return;
 	}
 
-	if (UGameplayCameraComponent* GameplayCameraComponent = NewViewTarget->FindComponentByClass<UGameplayCameraComponent>())
+	if (UGameplayCameraComponentBase* GameplayCameraComponent = NewViewTarget->FindComponentByClass<UGameplayCameraComponentBase>())
 	{
 		GameplayCameraComponent->ActivateCameraForPlayerController(PCOwner);
 	}
-	else if (UGameplayCameraRigComponent* GameplayCameraRigComponent = NewViewTarget->FindComponentByClass<UGameplayCameraRigComponent>())
-	{
-		GameplayCameraRigComponent->ActivateCameraForPlayerController(PCOwner);
-	}
 	else if (UCameraComponent* CameraComponent = NewViewTarget->FindComponentByClass<UCameraComponent>())
 	{
 		TSharedRef<FActorCameraEvaluationContext> NewContext = MakeShared<FActorCameraEvaluationContext>(CameraComponent);
@@ -303,37 +285,30 @@ void AGameplayCamerasPlayerCameraManager::SetViewTarget(AActor* NewViewTarget, F
 	}
 }
 
-void AGameplayCamerasPlayerCameraManager::AcquireCameraSystemHost(APlayerController* PlayerController)
+void AGameplayCamerasPlayerCameraManager::EnsureCameraSystemHost()
 {
 	using namespace UE::Cameras;
 
-	// Early-out if we already have acquired the host.
-	if (CameraSystemHost)
+	if (!HasCameraSystem())
 	{
-		return;
-	}
-
-	// Find or create a camera system host and grab its camera system.
-	CameraSystemHost = UGameplayCameraSystemHost::FindOrCreateHost(PlayerController);
-	ensure(CameraSystemHost);
-	CameraSystemEvaluator = CameraSystemHost->GetCameraSystemEvaluator();
+		InitializeCameraSystem();
 
-	ViewTargetContextReferencerService = MakeShared<FViewTargetContextReferencerService>();
-	CameraSystemEvaluator->RegisterEvaluationService(ViewTargetContextReferencerService.ToSharedRef());
-
-	// See if there was an auto-spawned camera system actor. If so, remember to set it back.
-	WeakOriginalAutoCameraSystemActor = AGameplayCameraSystemActor::GetAutoSpawnedCameraSystemActor(PlayerController, false);
+		ViewTargetContextReferencerService = MakeShared<FViewTargetContextReferencerService>();
+		CameraSystemEvaluator->RegisterEvaluationService(ViewTargetContextReferencerService.ToSharedRef());
+	}
 }
 
-void AGameplayCamerasPlayerCameraManager::ReleaseCameraSystemHost()
+void AGameplayCamerasPlayerCameraManager::TeardownCameraSystemHost()
 {
-	if (ViewTargetContextReferencerService)
+	if (HasCameraSystem())
 	{
-		CameraSystemEvaluator->UnregisterEvaluationService(ViewTargetContextReferencerService.ToSharedRef());
-	}
+		if (ensure(ViewTargetContextReferencerService))
+		{
+			CameraSystemEvaluator->UnregisterEvaluationService(ViewTargetContextReferencerService.ToSharedRef());
+		}
 
-	CameraSystemHost = nullptr;
-	CameraSystemEvaluator = nullptr;
+		DestroyCameraSystem();
+	}
 }
 
 void AGameplayCamerasPlayerCameraManager::ProcessViewRotation(float DeltaTime, FRotator& OutViewRotation, FRotator& OutDeltaRot)
@@ -364,18 +339,6 @@ void AGameplayCamerasPlayerCameraManager::DoUpdateCamera(float DeltaTime)
 	}
 }
 
-void AGameplayCamerasPlayerCameraManager::BeginPlay()
-{
-	Super::BeginPlay();
-}
-
-void AGameplayCamerasPlayerCameraManager::EndPlay(const EEndPlayReason::Type EndPlayReason)
-{
-	// TODO: destroy camera system actor?
-
-	Super::EndPlay(EndPlayReason);
-}
-
 void AGameplayCamerasPlayerCameraManager::DisplayDebug(UCanvas* Canvas, const FDebugDisplayInfo& DebugDisplay, float& YL, float& YPos)
 {
 	int Indentation = 1;
diff --git a/Engine/Plugins/Cameras/GameplayCameras/Source/GameplayCameras/Private/GameFramework/GameplayControlRotationComponent.cpp b/Engine/Plugins/Cameras/GameplayCameras/Source/GameplayCameras/Private/GameFramework/GameplayControlRotationComponent.cpp
--- a/Engine/Plugins/Cameras/GameplayCameras/Source/GameplayCameras/Private/GameFramework/GameplayControlRotationComponent.cpp
+++ b/Engine/Plugins/Cameras/GameplayCameras/Source/GameplayCameras/Private/GameFramework/GameplayControlRotationComponent.cpp
@@ -9,9 +9,7 @@
 #include "Engine/World.h"
 #include "GameFramework/Controller.h"
 #include "GameFramework/GameplayCameraComponent.h"
-#include "GameFramework/GameplayCameraSystemActor.h"
-#include "GameFramework/GameplayCameraSystemComponent.h"
-#include "GameFramework/GameplayCameraSystemHost.h"
+#include "GameFramework/IGameplayCameraSystemHost.h"
 #include "GameFramework/Pawn.h"
 #include "GameplayCameras.h"
 #include "Kismet/GameplayStatics.h"
@@ -98,22 +96,17 @@ void UGameplayControlRotationComponent::InitializeControlRotationService(APlayer
 		return;
 	}
 
-	UGameplayCameraSystemHost* FoundHost = UGameplayCameraSystemHost::FindHost(InPlayerController);
+	IGameplayCameraSystemHost* FoundHost = IGameplayCameraSystemHost::FindActiveHost(InPlayerController);
 	if (!FoundHost)
 	{
-		UE_LOG(LogCameraSystem, Error, TEXT("Can't find camera system host on the player controller."));
+		UE_LOG(LogCameraSystem, Error, 
+				TEXT("Can't find camera system host on the player controller. "
+				 	 "UGameplayControlRotationComponent requires using AGameplayCamerasPlayerCameraManager, or similar, as a camera manager."));
 		return;
 	}
 
 	PlayerController = InPlayerController;
-	CameraSystemHost = FoundHost;
-
-	// Make sure that if there's an auto-spawned camera system actor, it doesn't set the control rotation.
-	AGameplayCameraSystemActor* AutoSpawnedCameraSystemActor = AGameplayCameraSystemActor::GetAutoSpawnedCameraSystemActor(PlayerController);
-	if (AutoSpawnedCameraSystemActor)
-	{
-		AutoSpawnedCameraSystemActor->GetCameraSystemComponent()->bSetPlayerControllerRotation = false;
-	}
+	CameraSystemHost = FoundHost->GetAsScriptInterface();
 
 	// Create the evaluation service, with a copy of our parameters.
 	FPlayerControlRotationParams ServiceParams;
diff --git a/Engine/Plugins/Cameras/GameplayCameras/Source/GameplayCameras/Private/GameFramework/IGameplayCameraSystemHost.cpp b/Engine/Plugins/Cameras/GameplayCameras/Source/GameplayCameras/Private/GameFramework/IGameplayCameraSystemHost.cpp
new file mode 100644
--- /dev/null
+++ b/Engine/Plugins/Cameras/GameplayCameras/Source/GameplayCameras/Private/GameFramework/IGameplayCameraSystemHost.cpp
@@ -0,0 +1,138 @@
+// Copyright Epic Games, Inc. All Rights Reserved.
+
+#include "GameFramework/IGameplayCameraSystemHost.h"
+
+#include "Core/CameraEvaluationContext.h"
+#include "Core/CameraEvaluationContextStack.h"
+#include "Core/CameraSystemEvaluator.h"
+#include "Debug/DebugDrawService.h"
+#include "Engine/Canvas.h"
+#include "Engine/World.h"
+#include "GameFramework/PlayerController.h"
+
+void IGameplayCameraSystemHost::InitializeCameraSystem()
+{
+	using namespace UE::Cameras;
+
+	FCameraSystemEvaluatorCreateParams Params;
+	Params.Owner = GetAsObject();
+	InitializeCameraSystem(Params);
+}
+
+void IGameplayCameraSystemHost::InitializeCameraSystem(const UE::Cameras::FCameraSystemEvaluatorCreateParams& Params)
+{
+	ensure(!CameraSystemEvaluator.IsValid());
+	ensure(Params.Owner != nullptr && Params.Owner == GetAsObject());
+
+	CameraSystemEvaluator = MakeShared<FCameraSystemEvaluator>();
+	CameraSystemEvaluator->Initialize(Params);
+
+#if UE_GAMEPLAY_CAMERAS_DEBUG
+	ensure(!DebugDrawDelegateHandle.IsValid());
+	UWorld* World = Params.Owner->GetWorld();
+	if (World && World->IsGameWorld())
+	{
+		DebugDrawDelegateHandle = UDebugDrawService::Register(
+				TEXT("Game"), FDebugDrawDelegate::CreateRaw(this, &IGameplayCameraSystemHost::DebugDraw));
+	}
+#endif  // UE_GAMEPLAY_CAMERAS_DEBUG
+}
+
+TScriptInterface<IGameplayCameraSystemHost> IGameplayCameraSystemHost::GetAsScriptInterface()
+{
+	TScriptInterface<IGameplayCameraSystemHost> Result(GetAsObject());
+	checkSlow(Result.GetInterface() != nullptr);
+	return Result;
+}
+
+IGameplayCameraSystemHost* IGameplayCameraSystemHost::FindActiveHost(APlayerController* PlayerController)
+{
+	if (PlayerController->PlayerCameraManager)
+	{
+		if (IGameplayCameraSystemHost* CameraManagerHost = Cast<IGameplayCameraSystemHost>(PlayerController->PlayerCameraManager))
+		{
+			return CameraManagerHost;
+		}
+	}
+	else if (AActor* ViewTarget = PlayerController->GetViewTarget())
+	{
+		if (IGameplayCameraSystemHost* ViewTargetHost = ViewTarget->FindComponentByInterface<IGameplayCameraSystemHost>())
+		{
+			return ViewTargetHost;
+		}
+	}
+	return nullptr;
+}
+
+void IGameplayCameraSystemHost::EnsureCameraSystemInitialized()
+{
+	if (!CameraSystemEvaluator.IsValid())
+	{
+		InitializeCameraSystem();
+	}
+}
+
+void IGameplayCameraSystemHost::DestroyCameraSystem()
+{
+#if UE_GAMEPLAY_CAMERAS_DEBUG
+	if (DebugDrawDelegateHandle.IsValid())
+	{
+		UDebugDrawService::Unregister(DebugDrawDelegateHandle);
+		DebugDrawDelegateHandle.Reset();
+	}
+#endif  // UE_GAMEPLAY_CAMERAS_DEBUG
+
+	CameraSystemEvaluator.Reset();
+}
+
+void IGameplayCameraSystemHost::OnAddReferencedObjects(FReferenceCollector& Collector)
+{
+	if (CameraSystemEvaluator.IsValid())
+	{
+		CameraSystemEvaluator->AddReferencedObjects(Collector);
+	}
+}
+
+void IGameplayCameraSystemHost::UpdateCameraSystem(float DeltaTime)
+{
+	using namespace UE::Cameras;
+
+	if (CameraSystemEvaluator.IsValid())
+	{
+		FCameraSystemEvaluationParams Params;
+		Params.DeltaTime = DeltaTime;
+		CameraSystemEvaluator->Update(Params);
+	}
+}
+
+TSharedPtr<UE::Cameras::FCameraSystemEvaluator> IGameplayCameraSystemHost::GetCameraSystemEvaluator()
+{
+	return CameraSystemEvaluator;
+}
+
+#if UE_GAMEPLAY_CAMERAS_DEBUG
+
+void IGameplayCameraSystemHost::DebugDraw(UCanvas* Canvas, APlayerController* PlayerController)
+{
+	using namespace UE::Cameras;
+
+	if (CameraSystemEvaluator.IsValid())
+	{
+		// We're looking from the outside if we are not the view target, or if we don't have a player
+		// anymore (which happens in spectator mode like with the debug camera).
+		bool bIsDebugCameraEnabled = false;
+		if (TSharedPtr<FCameraEvaluationContext> ActiveContext = CameraSystemEvaluator->GetEvaluationContextStack().GetActiveContext())
+		{
+			const APlayerController* ActivePlayerController = ActiveContext->GetPlayerController();
+			bIsDebugCameraEnabled = !ActivePlayerController || !ActivePlayerController->Player;
+		}
+
+		FCameraSystemDebugUpdateParams DebugUpdateParams;
+		DebugUpdateParams.CanvasObject = Canvas;
+		DebugUpdateParams.bIsDebugCameraEnabled = bIsDebugCameraEnabled;
+		CameraSystemEvaluator->DebugUpdate(DebugUpdateParams);
+	}
+}
+
+#endif  // UE_GAMEPLAY_CAMERAS_DEBUG
+
diff --git a/Engine/Plugins/Cameras/GameplayCameras/Source/GameplayCameras/Private/Math/CameraPoseMath.cpp b/Engine/Plugins/Cameras/GameplayCameras/Source/GameplayCameras/Private/Math/CameraPoseMath.cpp
--- a/Engine/Plugins/Cameras/GameplayCameras/Source/GameplayCameras/Private/Math/CameraPoseMath.cpp
+++ b/Engine/Plugins/Cameras/GameplayCameras/Source/GameplayCameras/Private/Math/CameraPoseMath.cpp
@@ -69,8 +69,7 @@ double FCameraPoseMath::GetEffectiveAspectRatio(const FCameraPose& CameraPose, T
 	else
 	{
 		FIntPoint ViewportSize = EvaluationContext->GetViewportSize();
-		if (ensureMsgf(ViewportSize.X > 0 && ViewportSize.Y > 0,
-				TEXT("Can't get viewport aspect ratio! Using sensor aspect ratio instead.")))
+		if (ViewportSize.X > 0 && ViewportSize.Y > 0)
 		{
 			return (double)ViewportSize.X / (double)ViewportSize.Y;
 		}
diff --git a/Engine/Plugins/Cameras/GameplayCameras/Source/GameplayCameras/Public/Core/BlendStackCameraNode.h b/Engine/Plugins/Cameras/GameplayCameras/Source/GameplayCameras/Public/Core/BlendStackCameraNode.h
--- a/Engine/Plugins/Cameras/GameplayCameras/Source/GameplayCameras/Public/Core/BlendStackCameraNode.h
+++ b/Engine/Plugins/Cameras/GameplayCameras/Source/GameplayCameras/Public/Core/BlendStackCameraNode.h
@@ -110,6 +110,9 @@ public:
 	/** Returns information about a given camera rig, if any. */
 	FCameraRigEvaluationInfo GetCameraRigEvaluationInfo(FBlendStackEntryID EntryID) const;
 
+	/** Returns whether the stack contains any running camera rig with the given context. */
+	bool HasAnyRunningCameraRig(TSharedPtr<const FCameraEvaluationContext> InContext) const;
+
 #if UE_GAMEPLAY_CAMERAS_DEBUG
 	FBlendStackCameraDebugBlock* BuildDetailedDebugBlock(const FCameraDebugBlockBuildParams& Params, FCameraDebugBlockBuilder& Builder);
 #endif  // UE_GAMEPLAY_CAMERAS_DEBUG
diff --git a/Engine/Plugins/Cameras/GameplayCameras/Source/GameplayCameras/Public/Core/DefaultRootCameraNode.h b/Engine/Plugins/Cameras/GameplayCameras/Source/GameplayCameras/Public/Core/DefaultRootCameraNode.h
--- a/Engine/Plugins/Cameras/GameplayCameras/Source/GameplayCameras/Public/Core/DefaultRootCameraNode.h
+++ b/Engine/Plugins/Cameras/GameplayCameras/Source/GameplayCameras/Public/Core/DefaultRootCameraNode.h
@@ -67,7 +67,9 @@ protected:
 	// FRootCameraNodeEvaluator interface.
 	virtual FCameraRigInstanceID OnActivateCameraRig(const FActivateCameraRigParams& Params) override;
 	virtual void OnDeactivateCameraRig(const FDeactivateCameraRigParams& Params) override;
+	virtual void OnDeactivateAllCameraRigs(TSharedPtr<const FCameraEvaluationContext> InContext, bool bImmediately) override;
 	virtual void OnGetActiveCameraRigInfo(FCameraRigEvaluationInfo& OutCameraRigInfo) const override;
+	virtual bool OnHasAnyRunningCameraRig(TSharedPtr<const FCameraEvaluationContext> InContext) const override;
 	virtual void OnGetCameraRigInfo(const FCameraRigInstanceID InstanceID, FCameraRigEvaluationInfo& OutCameraRigInfo) const override;
 	virtual const FCameraVariableTable* OnGetBlendedParameters() const override;
 	virtual void OnBuildSingleCameraRigHierarchy(const FSingleCameraRigHierarchyBuildParams& Params, FCameraNodeEvaluatorHierarchy& OutHierarchy) override;
diff --git a/Engine/Plugins/Cameras/GameplayCameras/Source/GameplayCameras/Public/Core/PersistentBlendStackCameraNode.h b/Engine/Plugins/Cameras/GameplayCameras/Source/GameplayCameras/Public/Core/PersistentBlendStackCameraNode.h
--- a/Engine/Plugins/Cameras/GameplayCameras/Source/GameplayCameras/Public/Core/PersistentBlendStackCameraNode.h
+++ b/Engine/Plugins/Cameras/GameplayCameras/Source/GameplayCameras/Public/Core/PersistentBlendStackCameraNode.h
@@ -81,6 +81,9 @@ public:
 	/** Remove an existing camera rig from the blend stack. */
 	void Remove(const FBlendStackCameraRemoveParams& Params);
 
+	/** Remove all existing camera rigs with the given context from the blend stack. */
+	void RemoveAll(TSharedPtr<const FCameraEvaluationContext> InContext, bool bImmediately);
+
 protected:
 
 	// FCameraNodeEvaluator interface.
@@ -95,8 +98,10 @@ private:
 
 	void InternalUpdate(TArrayView<FResolvedEntry> ResolvedEntries, const FCameraNodeEvaluationParams& Params, FCameraNodeEvaluationResult& OutResult);
 
+	void RemoveEntry(int32 EntryIndex, const UCameraRigTransition* TransitionOverride, bool bImmediately);
+
 	const UCameraRigTransition* FindEnterTransition(const FBlendStackCameraInsertParams& Params) const;
-	const UCameraRigTransition* FindExitTransition(const FBlendStackCameraRemoveParams& Params, const FCameraRigEntry& Entry) const;
+	const UCameraRigTransition* FindExitTransition(const FCameraRigEntry& Entry, const UCameraRigTransition* TransitionOverride) const;
 
 private:
 
diff --git a/Engine/Plugins/Cameras/GameplayCameras/Source/GameplayCameras/Public/Core/RootCameraNode.h b/Engine/Plugins/Cameras/GameplayCameras/Source/GameplayCameras/Public/Core/RootCameraNode.h
--- a/Engine/Plugins/Cameras/GameplayCameras/Source/GameplayCameras/Public/Core/RootCameraNode.h
+++ b/Engine/Plugins/Cameras/GameplayCameras/Source/GameplayCameras/Public/Core/RootCameraNode.h
@@ -125,9 +125,10 @@ public:
 	 */
 	void DeactivateCameraRig(const FDeactivateCameraRigParams& Params);
 
-	/**
-	 * Execute a request to activate or deactivate a camera rig.
-	 */
+	/** Deactivates all camera rigs with the given evaluation context. */
+	void DeactivateAllCameraRigs(TSharedPtr<const FCameraEvaluationContext> InContext, bool bImmediately);
+
+	/** Execute a request to activate or deactivate a camera rig. */
 	void ExecuteCameraDirectorRequest(const FCameraRigActivationDeactivationRequest& Request);
 
 	/** Gets information about the active camera rig in the main layer. */
@@ -136,9 +137,12 @@ public:
 	/** Gets information about a specified camera rig. */
 	void GetCameraRigInfo(const FCameraRigInstanceID InstanceID, FCameraRigEvaluationInfo& OutCameraRigInfo) const;
 
-	/** Gets whether any camera rig is running. */
+	/** Gets whether any camera rig is running in the main layer. */
 	bool HasAnyActiveCameraRig() const;
 
+	/** Gets whether any camera rig is running with the given context in the main layer. */
+	bool HasAnyRunningCameraRig(TSharedPtr<const FCameraEvaluationContext> InContext) const;
+
 	/** Gets the variable table containing the blended camera object interfaces parameters. */
 	const FCameraVariableTable* GetBlendedParameters() const;
 
@@ -182,9 +186,15 @@ protected:
 	/** Deactivates a camera rig. */
 	virtual void OnDeactivateCameraRig(const FDeactivateCameraRigParams& Params) {}
 
+	/** Deactivates all camera rigs with the given evaluation context. */
+	virtual void OnDeactivateAllCameraRigs(TSharedPtr<const FCameraEvaluationContext> InContext, bool bImmediately) {}
+
 	/** Gets information about the active camera rig in the main layer. */
 	virtual void OnGetActiveCameraRigInfo(FCameraRigEvaluationInfo& OutCameraRigInfo) const {}
 
+	/** Gets whether any camera rig is running with the given context in the main layer. */
+	virtual bool OnHasAnyRunningCameraRig(TSharedPtr<const FCameraEvaluationContext> InContext) const { return false; }
+
 	/** Gets information about a specified camera rig. */
 	virtual void OnGetCameraRigInfo(const FCameraRigInstanceID InstanceID, FCameraRigEvaluationInfo& OutCameraRigInfo) const {}
 
diff --git a/Engine/Plugins/Cameras/GameplayCameras/Source/GameplayCameras/Public/Core/TransientBlendStackCameraNode.h b/Engine/Plugins/Cameras/GameplayCameras/Source/GameplayCameras/Public/Core/TransientBlendStackCameraNode.h
--- a/Engine/Plugins/Cameras/GameplayCameras/Source/GameplayCameras/Public/Core/TransientBlendStackCameraNode.h
+++ b/Engine/Plugins/Cameras/GameplayCameras/Source/GameplayCameras/Public/Core/TransientBlendStackCameraNode.h
@@ -62,7 +62,7 @@ public:
 	void Freeze(const FBlendStackCameraFreezeParams& Params);
 
 	/** Freeze all camera rigs that belong to a given evaluation context. */
-	void FreezeAll(TSharedPtr<FCameraEvaluationContext> EvaluationContext);
+	void FreezeAll(TSharedPtr<const FCameraEvaluationContext> EvaluationContext);
 
 	/** Gets the variable table containing the blended camera object interfaces parameters. */
 	const FCameraVariableTable& GetBlendedParameters() const { return PreBlendVariableTable; }
diff --git a/Engine/Plugins/Cameras/GameplayCameras/Source/GameplayCameras/Public/GameFramework/ControllerGameplayCameraEvaluationComponent.h b/Engine/Plugins/Cameras/GameplayCameras/Source/GameplayCameras/Public/GameFramework/ControllerGameplayCameraEvaluationComponent.h
--- a/Engine/Plugins/Cameras/GameplayCameras/Source/GameplayCameras/Public/GameFramework/ControllerGameplayCameraEvaluationComponent.h
+++ b/Engine/Plugins/Cameras/GameplayCameras/Source/GameplayCameras/Public/GameFramework/ControllerGameplayCameraEvaluationComponent.h
@@ -8,8 +8,8 @@
 
 #include "ControllerGameplayCameraEvaluationComponent.generated.h"
 
+class IGameplayCameraSystemHost;
 class UCameraRigAsset;
-class UGameplayCameraSystemHost;
 enum class ECameraRigLayer : uint8;
 
 namespace UE::Cameras
@@ -30,16 +30,15 @@ public:
 
 	UControllerGameplayCameraEvaluationComponent(const FObjectInitializer& ObjectInitializer);
 
+	void Initialize(TScriptInterface<IGameplayCameraSystemHost> InCameraSystemHost);
+
 	/** Activates a new camera rig. */
 	void ActivateCameraRig(UCameraRigAsset* CameraRig, ECameraRigLayer EvaluationLayer);
 
 public:
 
-	static UControllerGameplayCameraEvaluationComponent* FindComponent(APlayerController* PlayerController);
-	static UControllerGameplayCameraEvaluationComponent* FindOrAddComponent(APlayerController* PlayerController);
-
-	static TSharedPtr<UE::Cameras::FCameraEvaluationContext> FindEvaluationContext(APlayerController* PlayerController);
-	static TSharedRef<UE::Cameras::FCameraEvaluationContext> FindOrAddEvaluationContext(APlayerController* PlayerController);
+	static UControllerGameplayCameraEvaluationComponent* FindComponent(AActor* OwnerActor);
+	static UControllerGameplayCameraEvaluationComponent* FindOrAddComponent(AActor* OwnerActor, bool* bOutCreated = nullptr);
 
 public:
 
@@ -51,7 +50,6 @@ private:
 
 	void ActivateCameraRigs();
 	void EnsureEvaluationContext();
-	void EnsureCameraSystemHost();
 
 private:
 
@@ -67,6 +65,6 @@ private:
 	TSharedPtr<UE::Cameras::FCameraEvaluationContext> EvaluationContext;
 
 	UPROPERTY()
-	TObjectPtr<UGameplayCameraSystemHost> CameraSystemHost;
+	TScriptInterface<IGameplayCameraSystemHost> CameraSystemHost;
 };
 
diff --git a/Engine/Plugins/Cameras/GameplayCameras/Source/GameplayCameras/Public/GameFramework/GameplayCameraActor.h b/Engine/Plugins/Cameras/GameplayCameras/Source/GameplayCameras/Public/GameFramework/GameplayCameraActor.h
--- a/Engine/Plugins/Cameras/GameplayCameras/Source/GameplayCameras/Public/GameFramework/GameplayCameraActor.h
+++ b/Engine/Plugins/Cameras/GameplayCameras/Source/GameplayCameras/Public/GameFramework/GameplayCameraActor.h
@@ -2,10 +2,7 @@
 
 #pragma once
 
-#include "CoreMinimal.h"
-#include "Engine/EngineTypes.h"
-#include "GameFramework/Actor.h"
-#include "UObject/ObjectMacros.h"
+#include "GameFramework/GameplayCameraActorBase.h"
 
 #include "GameplayCameraActor.generated.h"
 
@@ -15,7 +12,7 @@ class UGameplayCameraComponent;
  * An actor that can run a camera asset.
  */
 UCLASS(BlueprintType, MinimalAPI, ClassGroup=Camera, HideCategories=(Input, Rendering))
-class AGameplayCameraActor : public AActor
+class AGameplayCameraActor : public AGameplayCameraActorBase
 {
 	GENERATED_BODY()
 
@@ -33,9 +30,11 @@ public:
 
 	// AActor interface.
 	virtual USceneComponent* GetDefaultAttachComponent() const override;
-	virtual void Tick(float DeltaTime) override;
-	virtual bool ShouldTickIfViewportsOnly() const override;
-	virtual void DisplayDebug(UCanvas* Canvas, const FDebugDisplayInfo& DebugDisplay, float& YL, float& YPos) override;
+
+protected:
+
+	// AGameplayCameraActorBase interface.
+	virtual UGameplayCameraComponentBase* GetCameraComponentBase() const override;
 
 private:
 
diff --git a/Engine/Plugins/Cameras/GameplayCameras/Source/GameplayCameras/Public/GameFramework/GameplayCameraActorBase.h b/Engine/Plugins/Cameras/GameplayCameras/Source/GameplayCameras/Public/GameFramework/GameplayCameraActorBase.h
new file mode 100644
--- /dev/null
+++ b/Engine/Plugins/Cameras/GameplayCameras/Source/GameplayCameras/Public/GameFramework/GameplayCameraActorBase.h
@@ -0,0 +1,38 @@
+// Copyright Epic Games, Inc. All Rights Reserved.
+
+#pragma once
+
+#include "GameFramework/Actor.h"
+#include "UObject/ObjectMacros.h"
+
+#include "GameplayCameraActorBase.generated.h"
+
+class UGameplayCameraComponentBase;
+class UGameplayCameraSystemHost;
+
+/**
+ * A base class for actors that can run a camera object.
+ *
+ * When the actor becomes the view target, it is able to start updating itself by instantiating
+ * a private instance of the camera system. It does this if no camera system was found attached
+ * under the player controller.
+ */
+UCLASS(BlueprintType, MinimalAPI, ClassGroup=Camera, HideCategories=(Input, Rendering))
+class AGameplayCameraActorBase : public AActor
+{
+	GENERATED_BODY()
+
+public:
+
+	AGameplayCameraActorBase(const FObjectInitializer& ObjectInit);
+
+public:
+
+	// AActor interface.
+	virtual void CalcCamera(float DeltaTime, FMinimalViewInfo& OutResult) override;
+
+protected:
+
+	virtual UGameplayCameraComponentBase* GetCameraComponentBase() const { return nullptr; }
+};
+
diff --git a/Engine/Plugins/Cameras/GameplayCameras/Source/GameplayCameras/Public/GameFramework/GameplayCameraComponentBase.h b/Engine/Plugins/Cameras/GameplayCameras/Source/GameplayCameras/Public/GameFramework/GameplayCameraComponentBase.h
--- a/Engine/Plugins/Cameras/GameplayCameras/Source/GameplayCameras/Public/GameFramework/GameplayCameraComponentBase.h
+++ b/Engine/Plugins/Cameras/GameplayCameras/Source/GameplayCameras/Public/GameFramework/GameplayCameraComponentBase.h
@@ -6,6 +6,7 @@
 #include "Core/CameraAssetReference.h"
 #include "Core/CameraEvaluationContext.h"
 #include "GameFramework/BlueprintCameraEvaluationDataRef.h"
+#include "GameFramework/IGameplayCameraSystemHost.h"
 #include "UObject/ObjectMacros.h"
 
 #include "GameplayCameraComponentBase.generated.h"
@@ -17,7 +18,6 @@ class FViewport;
 class UCameraAsset;
 class UCanvas;
 class UCineCameraComponent;
-class UGameplayCameraSystemHost;
 
 namespace UE::Cameras
 {
@@ -48,7 +48,9 @@ UCLASS(Blueprintable, MinimalAPI, Abstract,
 		ClassGroup=Camera, 
 		HideCategories=(Mobility, Rendering, LOD), 
 		meta=(BlueprintSpawnableComponent))
-class UGameplayCameraComponentBase : public USceneComponent
+class UGameplayCameraComponentBase 
+	: public USceneComponent
+	, public IGameplayCameraSystemHost
 {
 	GENERATED_BODY()
 
@@ -63,30 +65,53 @@ public:
 	/** Get the camera evaluation context used by this component. */
 	GAMEPLAYCAMERAS_API TSharedPtr<UE::Cameras::FCameraEvaluationContext> GetEvaluationContext();
 
-	/** Get the player controller this component is currently activated for (if any). */
-	GAMEPLAYCAMERAS_API APlayerController* GetPlayerController() const;
-
 public:
 
 	/** Gets the child camera component used as the "output" for the gameplay/procedural camera. */
 	UFUNCTION(BlueprintGetter, Category=Camera)
 	UCineCameraComponent* GetOutputCameraComponent() const { return OutputCameraComponent; }
 
-	/** Activates the camera for the given player. */
+	/** 
+	 * Activates the camera for the given player.
+	 *
+	 * @param PlayerIndex        The player to activate the camera for.
+	 * @param bSetAsViewTarget   Whether to set this component's actor as the view target for the player.
+	 * @param ActivationMode     How to activate this camera into the player's camera system. Only valid 
+	 *                           and used when the player camera manager is running the camera system.
+	 *                           Must be 'Push' otherwise, when this component runs as a standalone camera
+	 *                           system.
+	 */
 	UFUNCTION(BlueprintCallable, Category=Camera)
 	GAMEPLAYCAMERAS_API void ActivateCameraForPlayerIndex(
 			int32 PlayerIndex, 
+			bool bSetAsViewTarget = true,
 			EGameplayCameraComponentActivationMode ActivationMode = EGameplayCameraComponentActivationMode::Push);
 
-	/** Activates the camera for the given player. */
+	/** 
+	 * Activates the camera for the given player.
+	 *
+	 * @param PlayerController   The player to activate the camera for.
+	 * @param bSetAsViewTarget   Whether to set this component's actor as the view target for the player.
+	 * @param ActivationMode     How to activate this camera into the player's camera system. Only valid 
+	 *                           and used when the player camera manager is running the camera system.
+	 *                           Must be 'Push' otherwise, when this component runs as a standalone camera
+	 *                           system.
+	 */
 	UFUNCTION(BlueprintCallable, Category=Camera)
 	GAMEPLAYCAMERAS_API void ActivateCameraForPlayerController(
 			APlayerController* PlayerController,
+			bool bSetAsViewTarget = true,
 			EGameplayCameraComponentActivationMode ActivationMode = EGameplayCameraComponentActivationMode::Push);
 
-	/** Deactivates the camera for the last player it was activated for. */
+	/** 
+	 * Deactivates the camera.
+	 *
+	 * @param bImmediately       Whether to let this component's camera rigs gracefully blend out before
+	 *                           deactivating. If true, any running camera rigs will be frozen or forcibly
+	 *                           removed from the camera system.
+	 */
 	UFUNCTION(BlueprintCallable, Category=Camera)
-	GAMEPLAYCAMERAS_API void DeactivateCamera();
+	GAMEPLAYCAMERAS_API void DeactivateCamera(bool bImmediately = false);
 
 	/** Gets the shared camera evaluation data for this component's evaluation context. */
 	UFUNCTION(BlueprintPure, Category=Camera, meta=(DisplayName="Get Shared Camera Data"))
@@ -103,8 +128,9 @@ public:
 	virtual void BeginPlay() override;
 	virtual void EndPlay(const EEndPlayReason::Type EndPlayReason) override;
 	virtual void OnUnregister() override;
+	virtual void Activate(bool bReset = false) override;
+	virtual void Deactivate() override;
 	virtual void TickComponent(float DeltaTime, ELevelTick TickType, FActorComponentTickFunction *ThisTickFunction) override;
-	virtual void OnComponentDestroyed(bool bDestroyingHierarchy) override;
 #if WITH_EDITOR
 	virtual bool GetEditorPreviewInfo(float DeltaTime, FMinimalViewInfo& ViewOut) override;
 #endif 
@@ -120,11 +146,16 @@ public:
 
 	static void AddReferencedObjects(UObject* InThis, FReferenceCollector& Collector);
 
+	// IGameplayCameraSystemHost interface.
+	virtual UObject* GetAsObject() override { return this; }
+
 public:
 
+	bool CalcCamera(float DeltaTime, FMinimalViewInfo& OutResult);
+
 #if WITH_EDITOR
 	GAMEPLAYCAMERAS_API void OnDrawVisualizationHUD(const FViewport* Viewport, const FSceneView* SceneView, FCanvas* Canvas) const;
-#endif
+#endif  // WITH_EDITOR
 
 protected:
 
@@ -134,20 +165,21 @@ protected:
 	virtual void OnUpdateCameraEvaluationContext(bool bForceApplyParameterOverrides) {}
 
 	void UpdateCameraEvaluationContext(bool bForceApplyParameterOverrides);
-#if WITH_EDITOR
-	bool CanRunInEditor() const;
-	bool HasActiveEditorPreview() const;
-	void RecreateEditorPreviewEvaluationContext();
-	void UpdateEditorPreviewEvaluator(float DeltaTime);
-#endif  // WITH_EDITOR
+	bool HasCameraEvaluationContext() const { return EvaluationContext.IsValid(); }
+
+	void RecreateCameraEvaluationContext();
 
 private:
 
-	void ActivateCameraEvaluationContext(int32 PlayerIndex, EGameplayCameraComponentActivationMode ActivationMode);
-	void ActivateCameraEvaluationContext(APlayerController* PlayerController, EGameplayCameraComponentActivationMode ActivationMode);
+	bool CanRunCameraSystem() const;
+	void EnsureCameraSystemHost();
+	void TeardownCameraSystemHost();
+
+	void ActivateCameraEvaluationContext(APlayerController* PlayerController, IGameplayCameraSystemHost* Host, EGameplayCameraComponentActivationMode ActivationMode);
 	void TryCreateCameraEvaluationContext(APlayerController* PlayerController);
 	void UpdateOutputCameraComponent();
-	void DeactivateCameraEvaluationContext();
+	void DeactivateCameraEvaluationContext(bool bImmediately);
+	void CheckPendingDeactivation();
 
 #if WITH_EDITOR
 	void AutoManageEditorPreviewEvaluator();
@@ -158,16 +190,18 @@ public:
 
 	/**
 	 * If AutoActivate is set, auto-activates this component's camera for the given player.
-	 * This is equivalent to calling ActivateCamera on BeginPlay.
+	 * This is equivalent to calling ActivateCameraForPlayerIndex on BeginPlay.
 	 */
 	UPROPERTY(BlueprintReadWrite, EditAnywhere, Category=Activation, meta=(EditCondition="bAutoActivate"))
 	TEnumAsByte<EAutoReceiveInput::Type> AutoActivateForPlayer;
 
 #if WITH_EDITORONLY_DATA
 
+	/** Whether to run this camera in editor. */
 	UPROPERTY(EditAnywhere, Category=Camera)
 	bool bRunInEditor = true;
 
+	/** The camera rig to run in the editor. */
 	UPROPERTY(EditAnywhere, Category=Camera, meta=(EditCondition="bRunInEditor"))
 	int32 EditorPreviewCameraRigIndex = 0;
 
@@ -178,9 +212,6 @@ private:
 	UPROPERTY(Transient)
 	TObjectPtr<UCineCameraComponent> OutputCameraComponent;
 
-	UPROPERTY(Transient)
-	TObjectPtr<UGameplayCameraSystemHost> CameraSystemHost;
-
 private:
 
 	using FGameplayCameraComponentEvaluationContext = UE::Cameras::FGameplayCameraComponentEvaluationContext;
@@ -189,10 +220,10 @@ private:
 
 	bool bIsCameraCutNextFrame = false;
 
+	bool bIsDeactivating = false;
+
 #if WITH_EDITOR
 	
-	TSharedPtr<UE::Cameras::FCameraSystemEvaluator> EditorPreviewEvaluator;
-
 	bool bIsEditorWorld = false;
 
 	int32 CustomShowFlag = INDEX_NONE;
diff --git a/Engine/Plugins/Cameras/GameplayCameras/Source/GameplayCameras/Public/GameFramework/GameplayCameraRigActor.h b/Engine/Plugins/Cameras/GameplayCameras/Source/GameplayCameras/Public/GameFramework/GameplayCameraRigActor.h
--- a/Engine/Plugins/Cameras/GameplayCameras/Source/GameplayCameras/Public/GameFramework/GameplayCameraRigActor.h
+++ b/Engine/Plugins/Cameras/GameplayCameras/Source/GameplayCameras/Public/GameFramework/GameplayCameraRigActor.h
@@ -2,10 +2,7 @@
 
 #pragma once
 
-#include "CoreMinimal.h"
-#include "Engine/EngineTypes.h"
-#include "GameFramework/Actor.h"
-#include "UObject/ObjectMacros.h"
+#include "GameFramework/GameplayCameraActorBase.h"
 
 #include "GameplayCameraRigActor.generated.h"
 
@@ -15,7 +12,7 @@ class UGameplayCameraRigComponent;
  * An actor that can run a camera asset.
  */
 UCLASS(BlueprintType, MinimalAPI, ClassGroup=Camera, HideCategories=(Input, Rendering))
-class AGameplayCameraRigActor : public AActor
+class AGameplayCameraRigActor : public AGameplayCameraActorBase
 {
 	GENERATED_BODY()
 
@@ -33,9 +30,11 @@ public:
 
 	// AActor interface.
 	virtual USceneComponent* GetDefaultAttachComponent() const override;
-	virtual void Tick(float DeltaTime) override;
-	virtual bool ShouldTickIfViewportsOnly() const override;
-	virtual void DisplayDebug(UCanvas* Canvas, const FDebugDisplayInfo& DebugDisplay, float& YL, float& YPos) override;
+
+protected:
+
+	// AGameplayCameraActorBase interface.
+	virtual UGameplayCameraComponentBase* GetCameraComponentBase() const override;
 
 private:
 
diff --git a/Engine/Plugins/Cameras/GameplayCameras/Source/GameplayCameras/Public/GameFramework/GameplayCameraSystemActor.h b/Engine/Plugins/Cameras/GameplayCameras/Source/GameplayCameras/Public/GameFramework/GameplayCameraSystemActor.h
--- a/Engine/Plugins/Cameras/GameplayCameras/Source/GameplayCameras/Public/GameFramework/GameplayCameraSystemActor.h
+++ b/Engine/Plugins/Cameras/GameplayCameras/Source/GameplayCameras/Public/GameFramework/GameplayCameraSystemActor.h
@@ -32,24 +32,6 @@ public:
 	UFUNCTION(BlueprintGetter, Category=Camera)
 	UGameplayCameraSystemComponent* GetCameraSystemComponent() const { return CameraSystemComponent; }
 
-public:
-
-	/**
-	 * Get the automatically spawned camera system actor, if it exists.
-	 * Returns null if bAutoSpawnCameraSystemActor is disabled in the GameplayCameras 
-	 * settings, or if no actor has been spawned yet.
-	 */
-	UFUNCTION(BlueprintCallable, Category=Camera)
-	static AGameplayCameraSystemActor* GetAutoSpawnedCameraSystemActor(APlayerController* PlayerController, bool bSpawnIfMissing = false);
-
-	/**
-	 * Automatically sets a camera system actor as the view-target, spawning a unique
-	 * instance if needed. Doesn't do anything if bAutoSpawnCameraSystemActor is disabled
-	 * in the GameplayCameras settings.
-	 */
-	UFUNCTION(BlueprintCallable, Category=Camera)
-	static void AutoManageActiveViewTarget(APlayerController* PlayerController);
-
 public:
 
 	// AActor interface.
diff --git a/Engine/Plugins/Cameras/GameplayCameras/Source/GameplayCameras/Public/GameFramework/GameplayCameraSystemComponent.h b/Engine/Plugins/Cameras/GameplayCameras/Source/GameplayCameras/Public/GameFramework/GameplayCameraSystemComponent.h
--- a/Engine/Plugins/Cameras/GameplayCameras/Source/GameplayCameras/Public/GameFramework/GameplayCameraSystemComponent.h
+++ b/Engine/Plugins/Cameras/GameplayCameras/Source/GameplayCameras/Public/GameFramework/GameplayCameraSystemComponent.h
@@ -4,13 +4,13 @@
 
 #include "UObject/ObjectMacros.h"
 #include "Components/SceneComponent.h"
+#include "GameFramework/IGameplayCameraSystemHost.h"
 
 #include "GameplayCameraSystemComponent.generated.h"
 
 class APlayerController;
 class UCameraRigAsset;
 class UCanvas;
-class UGameplayCameraSystemHost;
 struct FCameraRigInstanceID;
 struct FMinimalViewInfo;
 
@@ -23,7 +23,9 @@ namespace UE::Cameras
  * A component that hosts a camera system.
  */
 UCLASS(BlueprintType, MinimalAPI, ClassGroup=Camera, HideCategories=(Mobility, Rendering, LOD))
-class UGameplayCameraSystemComponent : public USceneComponent
+class UGameplayCameraSystemComponent 
+	: public USceneComponent
+	, public IGameplayCameraSystemHost
 {
 	GENERATED_BODY()
 
@@ -33,9 +35,6 @@ public:
 
 	UGameplayCameraSystemComponent(const FObjectInitializer& ObjectInit);
 
-	/** Gets the camera system evaluator. */
-	GAMEPLAYCAMERAS_API TSharedPtr<FCameraSystemEvaluator> GetCameraSystemEvaluator(bool bEnsureIfNull = true);
-
 	/** Updates the camera system and returns the computed view. */
 	GAMEPLAYCAMERAS_API void GetCameraView(float DeltaTime, FMinimalViewInfo& DesiredView);
 
@@ -68,11 +67,14 @@ public:
 
 public:
 
-	// UActorComponent interface
+	// UActorComponent interface.
 	virtual void OnRegister() override;
 	virtual void BeginPlay() override;
 	virtual void EndPlay(const EEndPlayReason::Type EndPlayReason) override;
 
+	// IGameplayCameraSystemHost interface.
+	virtual UObject* GetAsObject() { return this; }
+
 private:
 
 #if WITH_EDITOR
@@ -96,9 +98,6 @@ public:
 	bool bSetPlayerControllerRotation = false;
 
 private:
-	
-	UPROPERTY(Transient)
-	TObjectPtr<UGameplayCameraSystemHost> CameraSystemHost;
 
 	UPROPERTY(Transient)
 	TWeakObjectPtr<APlayerController> WeakPlayerController;
diff --git a/Engine/Plugins/Cameras/GameplayCameras/Source/GameplayCameras/Public/GameFramework/GameplayCameraSystemHost.h b/Engine/Plugins/Cameras/GameplayCameras/Source/GameplayCameras/Public/GameFramework/GameplayCameraSystemHost.h
deleted file mode 100644
diff --git a/Engine/Plugins/Cameras/GameplayCameras/Source/GameplayCameras/Public/GameFramework/GameplayCamerasPlayerCameraManager.h b/Engine/Plugins/Cameras/GameplayCameras/Source/GameplayCameras/Public/GameFramework/GameplayCamerasPlayerCameraManager.h
--- a/Engine/Plugins/Cameras/GameplayCameras/Source/GameplayCameras/Public/GameFramework/GameplayCamerasPlayerCameraManager.h
+++ b/Engine/Plugins/Cameras/GameplayCameras/Source/GameplayCameras/Public/GameFramework/GameplayCamerasPlayerCameraManager.h
@@ -4,12 +4,11 @@
 
 #include "Camera/PlayerCameraManager.h"
 #include "Core/CameraEvaluationContext.h"
+#include "GameFramework/IGameplayCameraSystemHost.h"
 #include "Nodes/Blends/SimpleBlendCameraNode.h"
 
 #include "GameplayCamerasPlayerCameraManager.generated.h"
 
-class AGameplayCameraSystemActor;
-class UGameplayCameraSystemHost;
 struct FCameraRigInstanceID;
 
 namespace UE::Cameras
@@ -38,7 +37,9 @@ class FViewTargetContextReferencerService;
  * two actors.
  */
 UCLASS(notplaceable, MinimalAPI)
-class AGameplayCamerasPlayerCameraManager : public APlayerCameraManager
+class AGameplayCamerasPlayerCameraManager 
+	: public APlayerCameraManager
+	, public IGameplayCameraSystemHost
 {
 	GENERATED_BODY()
 
@@ -46,9 +47,6 @@ public:
 
 	AGameplayCamerasPlayerCameraManager(const FObjectInitializer& ObjectInitializer);
 
-	/** Gets the camera system host object. */
-	UGameplayCameraSystemHost* GetCameraSystemHost() const { return CameraSystemHost; }
-
 public:
 
 	/** Replace the camera manager currently set on the provided controller with this camera manager. */
@@ -78,10 +76,15 @@ public:
 	virtual void ProcessViewRotation(float DeltaTime, FRotator& OutViewRotation, FRotator& OutDeltaRot) override;
 
 	// AActor interface.
-	virtual void BeginPlay() override;
-	virtual void EndPlay(const EEndPlayReason::Type EndPlayReason) override;
 	virtual void DisplayDebug(UCanvas* Canvas, const FDebugDisplayInfo& DebugDisplay, float& YL, float& YPos) override;
 
+	// UObject interface.
+	virtual void BeginDestroy() override;
+	static void AddReferencedObjects(UObject* InThis, FReferenceCollector& Collector);
+
+	// IGameplayCameraSystemHost interface.
+	virtual UObject* GetAsObject() override { return this; }
+
 protected:
 
 	// APlayerCameraManager interface.
@@ -89,22 +92,14 @@ protected:
 
 private:
 
-	void AcquireCameraSystemHost(APlayerController* PlayerController);
-	void ReleaseCameraSystemHost();
+	void EnsureCameraSystemHost();
+	void TeardownCameraSystemHost();
 
 private:
 
-	UPROPERTY(Transient)
-	TObjectPtr<UGameplayCameraSystemHost> CameraSystemHost;
-
 	UPROPERTY(Transient)
 	TObjectPtr<APlayerCameraManager> OriginalCameraManager;
 
-	UPROPERTY(Transient)
-	TWeakObjectPtr<AGameplayCameraSystemActor> WeakOriginalAutoCameraSystemActor;
-
-	TSharedPtr<UE::Cameras::FCameraSystemEvaluator> CameraSystemEvaluator;
-
 	TSharedPtr<UE::Cameras::FViewTargetContextReferencerService> ViewTargetContextReferencerService;
 
 	FMinimalViewInfo LastFrameDesiredView;
diff --git a/Engine/Plugins/Cameras/GameplayCameras/Source/GameplayCameras/Public/GameFramework/GameplayControlRotationComponent.h b/Engine/Plugins/Cameras/GameplayCameras/Source/GameplayCameras/Public/GameFramework/GameplayControlRotationComponent.h
--- a/Engine/Plugins/Cameras/GameplayCameras/Source/GameplayCameras/Public/GameFramework/GameplayControlRotationComponent.h
+++ b/Engine/Plugins/Cameras/GameplayCameras/Source/GameplayCameras/Public/GameFramework/GameplayControlRotationComponent.h
@@ -6,8 +6,8 @@
 
 #include "GameplayControlRotationComponent.generated.h"
 
+class IGameplayCameraSystemHost;
 class APlayerController;
-class UGameplayCameraSystemHost;
 class UCanvas;
 class UInputAction;
 
@@ -96,7 +96,7 @@ private:
 	TObjectPtr<APlayerController> PlayerController;
 
 	UPROPERTY()
-	TObjectPtr<UGameplayCameraSystemHost> CameraSystemHost;
+	TScriptInterface<IGameplayCameraSystemHost> CameraSystemHost;
 
 	TSharedPtr<UE::Cameras::FPlayerControlRotationEvaluationService> ControlRotationService;
 };
diff --git a/Engine/Plugins/Cameras/GameplayCameras/Source/GameplayCameras/Public/GameFramework/IGameplayCameraSystemHost.h b/Engine/Plugins/Cameras/GameplayCameras/Source/GameplayCameras/Public/GameFramework/IGameplayCameraSystemHost.h
new file mode 100644
--- /dev/null
+++ b/Engine/Plugins/Cameras/GameplayCameras/Source/GameplayCameras/Public/GameFramework/IGameplayCameraSystemHost.h
@@ -0,0 +1,90 @@
+// Copyright Epic Games, Inc. All Rights Reserved.
+
+#pragma once
+
+#include "GameplayCameras.h"
+#include "Templates/SharedPointerFwd.h"
+#include "UObject/Interface.h"
+
+#include "IGameplayCameraSystemHost.generated.h"
+
+class UCanvas;
+
+namespace UE::Cameras
+{
+	class FCameraSystemEvaluator;
+	struct FCameraSystemEvaluatorCreateParams;
+}
+
+/**
+ * An interface for objects that host a camera system evaluator.
+ */
+UINTERFACE(MinimalAPI)
+class UGameplayCameraSystemHost : public UInterface
+{
+	GENERATED_BODY()
+};
+
+/**
+ * An interface for objects that host a camera system evaluator.
+ */
+class IGameplayCameraSystemHost
+{
+	GENERATED_BODY()
+
+public:
+
+	using FCameraSystemEvaluator = UE::Cameras::FCameraSystemEvaluator;
+
+	/** Gets the camera system evaluator. */
+	GAMEPLAYCAMERAS_API TSharedPtr<FCameraSystemEvaluator> GetCameraSystemEvaluator();
+
+	/** Returns whether a camera system evaluation has been created in this host. */
+	bool HasCameraSystem() const { return CameraSystemEvaluator.IsValid(); }
+
+	/** Should be implemented by the underlying class to return itself as a UObject. */
+	virtual UObject* GetAsObject() { return nullptr; }
+
+	/** Returns this object as a script interface. */
+	TScriptInterface<IGameplayCameraSystemHost> GetAsScriptInterface();
+
+public:
+
+	/** Finds a valid host on the player controller's camera manager, or its view target. */
+	static IGameplayCameraSystemHost* FindActiveHost(APlayerController* PlayerController);
+
+protected:
+
+	/** Creates a new camera system. Asserts if there is already one. */
+	GAMEPLAYCAMERAS_API void InitializeCameraSystem();
+	/** Creates a new camera system. Asserts if there is already one. */
+	GAMEPLAYCAMERAS_API void InitializeCameraSystem(const UE::Cameras::FCameraSystemEvaluatorCreateParams& Params);
+	/** Ensures that the camera system is created. */
+	GAMEPLAYCAMERAS_API void EnsureCameraSystemInitialized();
+	/** Destroys the camera system. */
+	GAMEPLAYCAMERAS_API void DestroyCameraSystem();
+
+	/** Should be called by the underlying object for garbage collection. */
+	GAMEPLAYCAMERAS_API void OnAddReferencedObjects(FReferenceCollector& Collector);
+
+	/** Updates the camera system, if it exists. */
+	GAMEPLAYCAMERAS_API void UpdateCameraSystem(float DeltaTime);
+
+protected:
+
+	/** The camera system evaluator. */
+	TSharedPtr<FCameraSystemEvaluator> CameraSystemEvaluator;
+
+private:
+
+#if UE_GAMEPLAY_CAMERAS_DEBUG
+	void DebugDraw(UCanvas* Canvas, APlayerController* PlayerController);
+#endif  // UE_GAMEPLAY_CAMERAS_DEBUG
+
+private:
+
+#if UE_GAMEPLAY_CAMERAS_DEBUG
+	FDelegateHandle DebugDrawDelegateHandle;
+#endif  // UE_GAMEPLAY_CAMERAS_DEBUG
+};
+
diff --git a/Engine/Plugins/Cameras/GameplayCameras/Source/GameplayCameras/Public/GameplayCamerasSettings.h b/Engine/Plugins/Cameras/GameplayCameras/Source/GameplayCameras/Public/GameplayCamerasSettings.h
--- a/Engine/Plugins/Cameras/GameplayCameras/Source/GameplayCameras/Public/GameplayCamerasSettings.h
+++ b/Engine/Plugins/Cameras/GameplayCameras/Source/GameplayCameras/Public/GameplayCamerasSettings.h
@@ -23,24 +23,6 @@ public:
 
 public:
 
-	/**
-	 * Automatically spawn a camera system actor when any gameplay camera activates and no camera system
-	 * is found on the player controller's camera manager, or as a view target. This camera system actor
-	 * will be spawned and set as the view target automatically.
-	 */
-	UPROPERTY(EditAnywhere, Config, Category="General")
-	bool bAutoSpawnCameraSystemActor = true;
-
-	/**
-	 * Whether the automatically spawned camera system actor should set the control rotation on the
-	 * associated player controller. This is useful if camera rigs are managing their own rotation, e.g.
-	 * by specifying input slots on boom arms instead of using the existing control rotation.
-	 * Do not mix handling control rotation via camera nodes, and handling control rotation by calling
-	 * methods like AddYawInput/AddPitchInput/AddRollInput, as this can lead to a feedback loop.
-	 */
-	UPROPERTY(EditAnywhere, Config, Category="General", meta=(EditCondition="bAutoSpawnCameraSystemActor"))
-	bool bAutoSpawnCameraSystemActorSetsControlRotation = false;
-
 	/**
 	 * The number of camera rigs combined in one frame past which the camera system emits a warning.
 	 */
diff --git a/Engine/Plugins/Cameras/GameplayCameras/Source/GameplayCamerasEditor/Private/Trace/CameraSystemRewindDebuggerExtension.cpp b/Engine/Plugins/Cameras/GameplayCameras/Source/GameplayCamerasEditor/Private/Trace/CameraSystemRewindDebuggerExtension.cpp
--- a/Engine/Plugins/Cameras/GameplayCameras/Source/GameplayCamerasEditor/Private/Trace/CameraSystemRewindDebuggerExtension.cpp
+++ b/Engine/Plugins/Cameras/GameplayCameras/Source/GameplayCamerasEditor/Private/Trace/CameraSystemRewindDebuggerExtension.cpp
@@ -10,7 +10,6 @@
 #include "Debug/RootCameraDebugBlock.h"
 #include "Engine/Canvas.h"
 #include "Engine/World.h"
-#include "GameFramework/GameplayCameraSystemActor.h"
 #include "Trace/CameraSystemTraceProvider.h"
 
 #if UE_GAMEPLAY_CAMERAS_TRACE
