Cameras: move camera variable setters to a service, add customizable blend in/out times and types

diff --git a/Engine/Plugins/Cameras/GameplayCameras/Source/GameplayCameras/Private/Core/CameraSystemEvaluator.cpp b/Engine/Plugins/Cameras/GameplayCameras/Source/GameplayCameras/Private/Core/CameraSystemEvaluator.cpp
--- a/Engine/Plugins/Cameras/GameplayCameras/Source/GameplayCameras/Private/Core/CameraSystemEvaluator.cpp
+++ b/Engine/Plugins/Cameras/GameplayCameras/Source/GameplayCameras/Private/Core/CameraSystemEvaluator.cpp
@@ -16,6 +16,7 @@
 #include "GameplayCamerasSettings.h"
 #include "Math/ColorList.h"
 #include "Services/CameraModifierService.h"
+#include "Services/CameraParameterSetterService.h"
 #include "Services/CameraShakeService.h"
 #include "Services/OrientationInitializationService.h"
 #include "UObject/Package.h"
@@ -89,18 +90,19 @@ void FCameraSystemEvaluator::Initialize(const FCameraSystemEvaluatorCreateParams
 	BuildParams.RootCameraNode = RootNode;
 	RootEvaluator = static_cast<FRootCameraNodeEvaluator*>(RootEvaluatorStorage.BuildEvaluatorTree(BuildParams));
 
+	RegisterEvaluationService(MakeShared<FCameraModifierService>());
+	RegisterEvaluationService(MakeShared<FCameraParameterSetterService>());
+	RegisterEvaluationService(MakeShared<FCameraShakeService>());
+	RegisterEvaluationService(MakeShared<FOrientationInitializationService>());
+
+	CameraRigCombinationRegistry = MakeShared<FCameraRigCombinationRegistry>();
+
 	if (ensure(RootEvaluator))
 	{
 		FCameraNodeEvaluatorInitializeParams InitParams;
 		InitParams.Evaluator = this;
 		RootEvaluator->Initialize(InitParams, RootNodeResult);
 	}
-
-	RegisterEvaluationService(MakeShared<FCameraModifierService>());
-	RegisterEvaluationService(MakeShared<FCameraShakeService>());
-	RegisterEvaluationService(MakeShared<FOrientationInitializationService>());
-
-	CameraRigCombinationRegistry = MakeShared<FCameraRigCombinationRegistry>();
 }
 
 FCameraSystemEvaluator::~FCameraSystemEvaluator()
diff --git a/Engine/Plugins/Cameras/GameplayCameras/Source/GameplayCameras/Private/Core/CameraVariableSetter.cpp b/Engine/Plugins/Cameras/GameplayCameras/Source/GameplayCameras/Private/Core/CameraVariableSetter.cpp
--- a/Engine/Plugins/Cameras/GameplayCameras/Source/GameplayCameras/Private/Core/CameraVariableSetter.cpp
+++ b/Engine/Plugins/Cameras/GameplayCameras/Source/GameplayCameras/Private/Core/CameraVariableSetter.cpp
@@ -28,7 +28,8 @@ void FCameraVariableSetter::Stop(bool bImmediately)
 			else if (PrevState == EState::BlendIn)
 			{
 				// TODO: this assumes a symmetrical blend curve, which may not always be true.
-				CurrentTime = FMath::Clamp(BlendTime - CurrentTime, 0.f, BlendTime);
+				const float BlendInPercent = (BlendInTime > 0.f) ? FMath::Clamp(CurrentTime / BlendInTime, 0.f, 1.f) : 1.f;
+				CurrentTime = (1.f - BlendInPercent) * BlendOutTime;
 			}
 			// else: was blending out already.
 		}
@@ -47,26 +48,26 @@ void FCameraVariableSetter::UpdateState(float DeltaTime)
 		const float NewTime = CurrentTime + DeltaTime;
 		if (State == EState::BlendIn)
 		{
-			if (NewTime < BlendTime)
+			if (NewTime < BlendInTime)
 			{
 				CurrentTime = NewTime;
 			}
 			else
 			{
 				State = EState::Full;
-				CurrentTime = BlendTime;
+				CurrentTime = BlendInTime;
 			}
 		}
 		else if (State == EState::BlendOut)
 		{
-			if (NewTime < BlendTime)
+			if (NewTime < BlendOutTime)
 			{
 				CurrentTime = NewTime;
 			}
 			else
 			{
 				State = EState::Inactive;
-				CurrentTime = BlendTime;
+				CurrentTime = BlendOutTime;
 			}
 		}
 	}
@@ -78,13 +79,13 @@ float FCameraVariableSetter::GetBlendFactor()
 	switch (State)
 	{
 		case EState::BlendIn:
-			BlendPercent = (BlendTime > 0.f) ? (CurrentTime / BlendTime) : 1.f;
+			BlendPercent = (BlendInTime > 0.f) ? (CurrentTime / BlendInTime) : 1.f;
 			break;
 		case EState::Full:
 			BlendPercent = 1.f;
 			break;
 		case EState::BlendOut:
-			BlendPercent = 1.f - ((BlendTime > 0.f) ? (CurrentTime / BlendTime) : 1.f);
+			BlendPercent = 1.f - ((BlendOutTime > 0.f) ? (CurrentTime / BlendOutTime) : 1.f);
 			break;
 		case EState::Inactive:
 		default:
diff --git a/Engine/Plugins/Cameras/GameplayCameras/Source/GameplayCameras/Private/Core/DefaultRootCameraNode.cpp b/Engine/Plugins/Cameras/GameplayCameras/Source/GameplayCameras/Private/Core/DefaultRootCameraNode.cpp
--- a/Engine/Plugins/Cameras/GameplayCameras/Source/GameplayCameras/Private/Core/DefaultRootCameraNode.cpp
+++ b/Engine/Plugins/Cameras/GameplayCameras/Source/GameplayCameras/Private/Core/DefaultRootCameraNode.cpp
@@ -76,8 +76,6 @@ FCameraNodeEvaluatorChildrenView FDefaultRootCameraNodeEvaluator::OnGetChildren(
 
 void FDefaultRootCameraNodeEvaluator::OnRun(const FCameraNodeEvaluationParams& Params, FCameraNodeEvaluationResult& OutResult)
 {
-	UpdateCameraVariableSetters(Params.DeltaTime);
-
 	BaseLayer->Run(Params, OutResult);
 	MainLayer->Run(Params, OutResult);
 	GlobalLayer->Run(Params, OutResult);
diff --git a/Engine/Plugins/Cameras/GameplayCameras/Source/GameplayCameras/Private/Core/RootCameraNode.cpp b/Engine/Plugins/Cameras/GameplayCameras/Source/GameplayCameras/Private/Core/RootCameraNode.cpp
--- a/Engine/Plugins/Cameras/GameplayCameras/Source/GameplayCameras/Private/Core/RootCameraNode.cpp
+++ b/Engine/Plugins/Cameras/GameplayCameras/Source/GameplayCameras/Private/Core/RootCameraNode.cpp
@@ -110,45 +110,5 @@ void FRootCameraNodeEvaluator::BroadcastCameraRigEvent(const FRootCameraNodeCame
 	OnCameraRigEventDelegate.Broadcast(InEvent);
 }
 
-void FRootCameraNodeEvaluator::StopCameraVariableSetter(const FCameraVariableSetterHandle& InHandle, bool bImmediately)
-{
-	if (VariableSetters.IsValidIndex(InHandle.Value))
-	{
-		FVariableSetterEntry& Entry(VariableSetters[InHandle.Value]);
-		if (Entry.ThisHandle.SerialNumber == InHandle.SerialNumber)
-		{
-			Entry.Setter->Stop(bImmediately);
-		}
-	}
-}
-
-void FRootCameraNodeEvaluator::UpdateCameraVariableSetters(float DeltaTime)
-{
-	for (auto It = VariableSetters.CreateIterator(); It; ++It)
-	{
-		FVariableSetterEntry& Entry(*It);
-		if (!Entry.Setter.IsValid())
-		{
-			It.RemoveCurrent();
-			continue;
-		}
-
-		Entry.Setter->Update(DeltaTime);
-
-		if (!Entry.Setter->IsActive())
-		{
-			It.RemoveCurrent();
-		}
-	}
-}
-
-void FRootCameraNodeEvaluator::ApplyCameraVariableSetters(FCameraVariableTable& OutVariableTable)
-{
-	for (FVariableSetterEntry& Entry : VariableSetters)
-	{
-		Entry.Setter->Apply(OutVariableTable);
-	}
-}
-
 }  // namespace UE::Cameras
 
diff --git a/Engine/Plugins/Cameras/GameplayCameras/Source/GameplayCameras/Private/Core/TransientBlendStackCameraNode.cpp b/Engine/Plugins/Cameras/GameplayCameras/Source/GameplayCameras/Private/Core/TransientBlendStackCameraNode.cpp
--- a/Engine/Plugins/Cameras/GameplayCameras/Source/GameplayCameras/Private/Core/TransientBlendStackCameraNode.cpp
+++ b/Engine/Plugins/Cameras/GameplayCameras/Source/GameplayCameras/Private/Core/TransientBlendStackCameraNode.cpp
@@ -12,9 +12,9 @@
 #include "Core/CameraRigCombinationRegistry.h"
 #include "Core/CameraRigTransition.h"
 #include "Core/CameraSystemEvaluator.h"
-#include "Core/RootCameraNode.h"
 #include "Helpers/CameraRigTransitionFinder.h"
 #include "Nodes/Blends/PopBlendCameraNode.h"
+#include "Services/CameraParameterSetterService.h"
 
 namespace UE::Cameras
 {
@@ -198,6 +198,13 @@ void FTransientBlendStackCameraNodeEvaluator::FreezeAll(TSharedPtr<const FCamera
 	}
 }
 
+void FTransientBlendStackCameraNodeEvaluator::OnInitialize(const FCameraNodeEvaluatorInitializeParams& Params, FCameraNodeEvaluationResult& OutResult)
+{
+	Super::OnInitialize(Params, OutResult);
+
+	ParameterSetterService = Params.Evaluator->FindEvaluationService<FCameraParameterSetterService>();
+}
+
 void FTransientBlendStackCameraNodeEvaluator::OnRun(const FCameraNodeEvaluationParams& Params, FCameraNodeEvaluationResult& OutResult)
 {
 	ensure(Entries.Num() == EntryExtraInfos.Num());
@@ -231,9 +238,6 @@ void FTransientBlendStackCameraNodeEvaluator::InternalPreBlendPrepare(TArrayView
 	constexpr ECameraVariableTableFilter VariableTableFilter = ECameraVariableTableFilter::ChangedOnly;
 	constexpr ECameraContextDataTableFilter ContextDataTableFilter = ECameraContextDataTableFilter::ChangedOnly;
 
-	FRootCameraNodeEvaluator* RootNodeEvaluator = Params.Evaluator->GetRootNodeEvaluator();
-	check(RootNodeEvaluator);
-
 	for (FResolvedEntry& ResolvedEntry : ResolvedEntries)
 	{
 		FCameraRigEntry& Entry(ResolvedEntry.Entry);
@@ -259,7 +263,10 @@ void FTransientBlendStackCameraNodeEvaluator::InternalPreBlendPrepare(TArrayView
 		CurResult.ContextDataTable.OverrideAll(ContextResult.ContextDataTable);
 
 		// Override it with variable setters.
-		RootNodeEvaluator->ApplyCameraVariableSetters(CurResult.VariableTable);
+		if (ParameterSetterService)
+		{
+			ParameterSetterService->ApplyCameraVariableSetters(CurResult.VariableTable);
+		}
 
 		// Gather input parameters if needed (and remember if it was indeed needed).
 		if (!EntryExtraInfo.bInputRunThisFrame)
diff --git a/Engine/Plugins/Cameras/GameplayCameras/Source/GameplayCameras/Private/Services/CameraParameterSetterService.cpp b/Engine/Plugins/Cameras/GameplayCameras/Source/GameplayCameras/Private/Services/CameraParameterSetterService.cpp
new file mode 100644
--- /dev/null
+++ b/Engine/Plugins/Cameras/GameplayCameras/Source/GameplayCameras/Private/Services/CameraParameterSetterService.cpp
@@ -0,0 +1,61 @@
+// Copyright Epic Games, Inc. All Rights Reserved.
+
+#include "Services/CameraParameterSetterService.h"
+
+namespace UE::Cameras
+{
+
+UE_DEFINE_CAMERA_EVALUATION_SERVICE(FCameraParameterSetterService)
+
+FCameraParameterSetterService::FCameraParameterSetterService()
+{
+	SetEvaluationServiceFlags(ECameraEvaluationServiceFlags::NeedsPreUpdate);
+}
+
+void FCameraParameterSetterService::OnPreUpdate(const FCameraEvaluationServiceUpdateParams& Params, FCameraEvaluationServiceUpdateResult& OutResult)
+{
+	UpdateCameraVariableSetters(Params.DeltaTime);
+}
+
+void FCameraParameterSetterService::StopCameraVariableSetter(const FCameraVariableSetterHandle& InHandle, bool bImmediately)
+{
+	if (VariableSetters.IsValidIndex(InHandle.Value))
+	{
+		FVariableSetterEntry& Entry(VariableSetters[InHandle.Value]);
+		if (Entry.ThisHandle.SerialNumber == InHandle.SerialNumber)
+		{
+			Entry.Setter->Stop(bImmediately);
+		}
+	}
+}
+
+void FCameraParameterSetterService::UpdateCameraVariableSetters(float DeltaTime)
+{
+	for (auto It = VariableSetters.CreateIterator(); It; ++It)
+	{
+		FVariableSetterEntry& Entry(*It);
+		if (!Entry.Setter.IsValid())
+		{
+			It.RemoveCurrent();
+			continue;
+		}
+
+		Entry.Setter->Update(DeltaTime);
+
+		if (!Entry.Setter->IsActive())
+		{
+			It.RemoveCurrent();
+		}
+	}
+}
+
+void FCameraParameterSetterService::ApplyCameraVariableSetters(FCameraVariableTable& OutVariableTable)
+{
+	for (FVariableSetterEntry& Entry : VariableSetters)
+	{
+		Entry.Setter->Apply(OutVariableTable);
+	}
+}
+
+}  // namespace UE::Cameras
+
diff --git a/Engine/Plugins/Cameras/GameplayCameras/Source/GameplayCameras/Public/Core/CameraVariableSetter.h b/Engine/Plugins/Cameras/GameplayCameras/Source/GameplayCameras/Public/Core/CameraVariableSetter.h
--- a/Engine/Plugins/Cameras/GameplayCameras/Source/GameplayCameras/Public/Core/CameraVariableSetter.h
+++ b/Engine/Plugins/Cameras/GameplayCameras/Source/GameplayCameras/Public/Core/CameraVariableSetter.h
@@ -11,11 +11,14 @@
 
 namespace UE::Cameras
 {
-	class FRootCameraNodeEvaluator;
+	class FCameraParameterSetterService;
 }
 
-UENUM()
-enum class ECameraVariableSetterBlendType
+/**
+ * The blend type for a camera variable setter.
+ */
+UENUM(BlueprintType)
+enum class ECameraVariableSetterBlendType : uint8
 {
 	None,
 	Linear,
@@ -23,6 +26,9 @@ enum class ECameraVariableSetterBlendType
 	SmootherStep
 };
 
+/**
+ * A handle to an ongoing camera variable setter.
+ */
 USTRUCT(BlueprintType)
 struct FCameraVariableSetterHandle
 {
@@ -75,16 +81,23 @@ private:
 	UPROPERTY()
 	uint32 SerialNumber;
 
-	friend class UE::Cameras::FRootCameraNodeEvaluator;
+	friend class UE::Cameras::FCameraParameterSetterService;
 };
 
 namespace UE::Cameras
 {
 
+/**
+ * A base structure for overriding a camera variable value until told to stop.
+ * Overriding supports blending in and out.
+ *
+ * This base structure doesn't apply the values, see TCameraVariableSetter for that.
+ */
 struct FCameraVariableSetter
 {
 	ECameraVariableSetterBlendType BlendType = ECameraVariableSetterBlendType::Linear;
-	float BlendTime = 1.f;
+	float BlendInTime = 1.f;
+	float BlendOutTime = 1.f;
 
 	virtual ~FCameraVariableSetter() {}
 
@@ -147,6 +160,9 @@ struct TVariableSetterValueInterpolator<UE::Math::TTransform<T>>
 
 }  // namespace Internal
 
+/**
+ * An actual camera variable setter for a given type of values.
+ */
 template<typename ValueType>
 struct TCameraVariableSetter : public FCameraVariableSetter
 {
diff --git a/Engine/Plugins/Cameras/GameplayCameras/Source/GameplayCameras/Public/Core/RootCameraNode.h b/Engine/Plugins/Cameras/GameplayCameras/Source/GameplayCameras/Public/Core/RootCameraNode.h
--- a/Engine/Plugins/Cameras/GameplayCameras/Source/GameplayCameras/Public/Core/RootCameraNode.h
+++ b/Engine/Plugins/Cameras/GameplayCameras/Source/GameplayCameras/Public/Core/RootCameraNode.h
@@ -7,7 +7,6 @@
 #include "Core/CameraNodeEvaluator.h"
 #include "Core/CameraRigEvaluationInfo.h"
 #include "Core/CameraRigInstanceID.h"
-#include "Core/CameraVariableSetter.h"
 
 #include "RootCameraNode.generated.h"
 
@@ -164,15 +163,6 @@ public:
 	/** Gets the delegate for camera rig events. */
 	FOnRootCameraNodeCameraRigEvent& OnCameraRigEvent() { return OnCameraRigEventDelegate; }
 
-public:
-
-	template<typename ValueType>
-	FCameraVariableSetterHandle AddCameraVariableSetter(const TCameraVariableSetter<ValueType>& InSetter);
-
-	void StopCameraVariableSetter(const FCameraVariableSetterHandle& InHandle, bool bImmediately = false);
-
-	void ApplyCameraVariableSetters(FCameraVariableTable& OutVariableTable);
-
 protected:
 
 	// FCameraNodeEvaluator interface.
@@ -211,8 +201,6 @@ protected:
 
 	void BroadcastCameraRigEvent(const FRootCameraNodeCameraRigEvent& InEvent) const;
 
-	void UpdateCameraVariableSetters(float DeltaTime);
-
 private:
 
 	/** The camera system that owns this root node. */
@@ -220,26 +208,7 @@ private:
 
 	/** The delegate to notify when an event happens. */
 	FOnRootCameraNodeCameraRigEvent OnCameraRigEventDelegate;
-
-	struct FVariableSetterEntry
-	{
-		FCameraVariableSetterPtr Setter;
-		FCameraVariableSetterHandle ThisHandle;
-	};
-
-	using FCameraVariableSetters = TSparseArray<FVariableSetterEntry>;
-	FCameraVariableSetters VariableSetters;
-	uint32 NextVariableSetterSerial = 0;
 };
 
-template<typename ValueType>
-FCameraVariableSetterHandle FRootCameraNodeEvaluator::AddCameraVariableSetter(const TCameraVariableSetter<ValueType>& InSetter)
-{
-	FSparseArrayAllocationInfo NewAllocation = VariableSetters.AddUninitialized();
-	FCameraVariableSetterHandle NewHandle(NewAllocation.Index, NextVariableSetterSerial++);
-	new(NewAllocation.Pointer) FVariableSetterEntry{ InSetter, NewHandle };
-	return NewHandle;
-}
-
 }  // namespace UE::Cameras
 
diff --git a/Engine/Plugins/Cameras/GameplayCameras/Source/GameplayCameras/Public/Core/TransientBlendStackCameraNode.h b/Engine/Plugins/Cameras/GameplayCameras/Source/GameplayCameras/Public/Core/TransientBlendStackCameraNode.h
--- a/Engine/Plugins/Cameras/GameplayCameras/Source/GameplayCameras/Public/Core/TransientBlendStackCameraNode.h
+++ b/Engine/Plugins/Cameras/GameplayCameras/Source/GameplayCameras/Public/Core/TransientBlendStackCameraNode.h
@@ -4,10 +4,13 @@
 
 #include "Core/BlendStackCameraNode.h"
 #include "Core/BlendStackEntryID.h"
+#include "Templates/SharedPointerFwd.h"
 
 namespace UE::Cameras
 {
 
+class FCameraParameterSetterService;
+
 /**
  * Parameter structure for pushing a camera rig onto a transient blend stack.
  */
@@ -70,6 +73,7 @@ public:
 protected:
 
 	// FCameraNodeEvaluator interface.
+	virtual void OnInitialize(const FCameraNodeEvaluatorInitializeParams& Params, FCameraNodeEvaluationResult& OutResult) override;
 	virtual void OnRun(const FCameraNodeEvaluationParams& Params, FCameraNodeEvaluationResult& OutResult) override;
 	virtual void OnSerialize(const FCameraNodeEvaluatorSerializeParams& Params, FArchive& Ar) override;
 
@@ -106,6 +110,9 @@ private:
 
 	/** Variable table for pre-blending. */
 	FCameraVariableTable PreBlendVariableTable;
+
+	/** Cached pointer to the parameter setter service. */
+	TSharedPtr<FCameraParameterSetterService> ParameterSetterService;
 };
 
 }  // namespace UE::Cameras
diff --git a/Engine/Plugins/Cameras/GameplayCameras/Source/GameplayCameras/Public/Services/CameraParameterSetterService.h b/Engine/Plugins/Cameras/GameplayCameras/Source/GameplayCameras/Public/Services/CameraParameterSetterService.h
new file mode 100644
--- /dev/null
+++ b/Engine/Plugins/Cameras/GameplayCameras/Source/GameplayCameras/Public/Services/CameraParameterSetterService.h
@@ -0,0 +1,64 @@
+// Copyright Epic Games, Inc. All Rights Reserved.
+
+#pragma once
+
+#include "Core/CameraEvaluationService.h"
+#include "Core/CameraVariableSetter.h"
+
+namespace UE::Cameras
+{
+
+/**
+ * A camera system service that handles running camera parameter setters.
+ */
+class FCameraParameterSetterService : public FCameraEvaluationService
+{
+	UE_DECLARE_CAMERA_EVALUATION_SERVICE(GAMEPLAYCAMERAS_API, FCameraParameterSetterService)
+
+public:
+
+	FCameraParameterSetterService();
+
+	template<typename ValueType>
+	FCameraVariableSetterHandle AddCameraVariableSetter(const TCameraVariableSetter<ValueType>& InSetter);
+
+	void StopCameraVariableSetter(const FCameraVariableSetterHandle& InHandle, bool bImmediately = false);
+
+	void ApplyCameraVariableSetters(FCameraVariableTable& OutVariableTable);
+
+
+protected:
+
+	// FCameraEvaluationService interface.
+	virtual void OnPreUpdate(const FCameraEvaluationServiceUpdateParams& Params, FCameraEvaluationServiceUpdateResult& OutResult) override;
+
+private:
+
+	void UpdateCameraVariableSetters(float DeltaTime);
+
+private:
+
+	struct FVariableSetterEntry
+	{
+		FCameraVariableSetterPtr Setter;
+		FCameraVariableSetterHandle ThisHandle;
+	};
+
+	FCameraSystemEvaluator* Evaluator = nullptr;
+
+	using FCameraVariableSetters = TSparseArray<FVariableSetterEntry>;
+	FCameraVariableSetters VariableSetters;
+	uint32 NextVariableSetterSerial = 0;
+};
+
+template<typename ValueType>
+FCameraVariableSetterHandle FCameraParameterSetterService::AddCameraVariableSetter(const TCameraVariableSetter<ValueType>& InSetter)
+{
+	FSparseArrayAllocationInfo NewAllocation = VariableSetters.AddUninitialized();
+	FCameraVariableSetterHandle NewHandle(NewAllocation.Index, NextVariableSetterSerial++);
+	new(NewAllocation.Pointer) FVariableSetterEntry{ InSetter, NewHandle };
+	return NewHandle;
+}
+
+}  // namespace UE::Cameras
+
