Cameras: fix editor crash when opening old assets with stray graph objects

diff --git a/Engine/Plugins/Cameras/GameplayCameras/Source/GameplayCamerasEditor/Private/Toolkits/CameraRigAssetEditorToolkit.cpp b/Engine/Plugins/Cameras/GameplayCameras/Source/GameplayCamerasEditor/Private/Toolkits/CameraRigAssetEditorToolkit.cpp
--- a/Engine/Plugins/Cameras/GameplayCameras/Source/GameplayCamerasEditor/Private/Toolkits/CameraRigAssetEditorToolkit.cpp
+++ b/Engine/Plugins/Cameras/GameplayCameras/Source/GameplayCamerasEditor/Private/Toolkits/CameraRigAssetEditorToolkit.cpp
@@ -82,7 +82,7 @@ void FCameraRigAssetEditorToolkit::SetCameraRigAsset(UCameraRigAsset* InCameraRi
 	BuildButtonToolkit->SetTarget(InCameraRig);
 	InterfaceParametersToolkit->SetCameraObject(InCameraRig);
 
-	if (InCameraRig)
+	if (InCameraRig && bIsInitialized)
 	{
 		InCameraRig->EventHandlers.Register(EventHandler, this);
 	}
@@ -144,10 +144,7 @@ TSharedRef<SDockTab> FCameraRigAssetEditorToolkit::SpawnTab_Curves(const FSpawnT
 {
 	if (!CurveEditorToolkit->IsInitialized())
 	{
-		TArray<UObject*> CameraRigObjects;
-		UPackage* CameraRigPackage = Impl->GetCameraRigAsset()->GetPackage();
-		GetObjectsWithPackage(CameraRigPackage, CameraRigObjects);
-		CurveEditorToolkit->Initialize(CameraRigObjects);
+		CurveEditorToolkit->Initialize();
 	}
 
 	TSharedPtr<SDockTab> CurvesTab = SNew(SDockTab)
@@ -299,7 +296,19 @@ void FCameraRigAssetEditorToolkit::PostInitAssetEditor()
 
 	RegenerateMenusAndToolbars();
 
+	if (UCameraRigAsset* CameraRig = Impl->GetCameraRigAsset())
+	{
+		TArray<UObject*> CameraRigObjects;
+		UPackage* CameraRigPackage = CameraRig->GetPackage();
+		GetObjectsWithPackage(CameraRigPackage, CameraRigObjects);
+		CurveEditorToolkit->AddCurveOwners(CameraRigObjects);
+
+		CameraRig->EventHandlers.Register(EventHandler, this);
+	}
+
 	FRichCurveDetailsCustomization::OnInvokeCurveEditor().AddSP(this, &FCameraRigAssetEditorToolkit::OnInvokeCurveEditor);
+
+	bIsInitialized = true;
 }
 
 void FCameraRigAssetEditorToolkit::PostRegenerateMenusAndToolbars()
diff --git a/Engine/Plugins/Cameras/GameplayCameras/Source/GameplayCamerasEditor/Private/Toolkits/CurveEditorToolkit.cpp b/Engine/Plugins/Cameras/GameplayCameras/Source/GameplayCamerasEditor/Private/Toolkits/CurveEditorToolkit.cpp
--- a/Engine/Plugins/Cameras/GameplayCameras/Source/GameplayCamerasEditor/Private/Toolkits/CurveEditorToolkit.cpp
+++ b/Engine/Plugins/Cameras/GameplayCameras/Source/GameplayCamerasEditor/Private/Toolkits/CurveEditorToolkit.cpp
@@ -105,6 +105,17 @@ void FCurveEditorToolkit::AddCurveOwner(UObject* InCurveOwner)
 	}
 }
 
+void FCurveEditorToolkit::AddCurveOwners(TArrayView<UObject*> InCurveOwners)
+{
+	if (IsInitialized())
+	{
+		for (UObject* CurveOwner : InCurveOwners)
+		{
+			AddCurves(CurveOwner);
+		}
+	}
+}
+
 void FCurveEditorToolkit::RemoveCurveOwner(UObject* InCurveOwner)
 {
 	if (!IsInitialized())
diff --git a/Engine/Plugins/Cameras/GameplayCameras/Source/GameplayCamerasEditor/Public/Toolkits/CameraRigAssetEditorToolkit.h b/Engine/Plugins/Cameras/GameplayCameras/Source/GameplayCamerasEditor/Public/Toolkits/CameraRigAssetEditorToolkit.h
--- a/Engine/Plugins/Cameras/GameplayCameras/Source/GameplayCamerasEditor/Public/Toolkits/CameraRigAssetEditorToolkit.h
+++ b/Engine/Plugins/Cameras/GameplayCameras/Source/GameplayCamerasEditor/Public/Toolkits/CameraRigAssetEditorToolkit.h
@@ -107,6 +107,9 @@ private:
 
 	/** Event handler */
 	TCameraEventHandler<ICameraRigAssetEventHandler> EventHandler;
+
+	/** Whether the toolkit has finished initializing */
+	bool bIsInitialized = false;
 };
 
 }  // namespace UE::Cameras
diff --git a/Engine/Plugins/Cameras/GameplayCameras/Source/GameplayCamerasEditor/Public/Toolkits/CurveEditorToolkit.h b/Engine/Plugins/Cameras/GameplayCameras/Source/GameplayCamerasEditor/Public/Toolkits/CurveEditorToolkit.h
--- a/Engine/Plugins/Cameras/GameplayCameras/Source/GameplayCamerasEditor/Public/Toolkits/CurveEditorToolkit.h
+++ b/Engine/Plugins/Cameras/GameplayCameras/Source/GameplayCamerasEditor/Public/Toolkits/CurveEditorToolkit.h
@@ -48,6 +48,11 @@ public:
 	 * curve properties, nothing happens.
 	 */
 	void AddCurveOwner(UObject* InCurveOwner);
+	/**
+	 * Adds curves for any curve property on the given objects. If the object doesn't have any recognizable
+	 * curve properties, nothing happens.
+	 */
+	void AddCurveOwners(TArrayView<UObject*> InCurveOwners);
 	/**
 	 * Removes any existing curves belonging to the given object.
 	 */
