Cameras: add parameter setter components

These components can be added on actors like box/sphere/capsule triggers. They blend camera parameter overrides in and out for a given camera rig when their parent actor is overlapped.

diff --git a/Engine/Plugins/Cameras/GameplayCameras/Source/GameplayCameras/Private/GameFramework/GameplayCameraParameterSetterComponent.cpp b/Engine/Plugins/Cameras/GameplayCameras/Source/GameplayCameras/Private/GameFramework/GameplayCameraParameterSetterComponent.cpp
new file mode 100644
--- /dev/null
+++ b/Engine/Plugins/Cameras/GameplayCameras/Source/GameplayCameras/Private/GameFramework/GameplayCameraParameterSetterComponent.cpp
@@ -0,0 +1,167 @@
+// Copyright Epic Games, Inc. All Rights Reserved.
+
+#include "GameFramework/GameplayCameraParameterSetterComponent.h"
+
+#include "Core/CameraRigAsset.h"
+#include "Core/CameraSystemEvaluator.h"
+#include "Core/CameraVariableSetter.h"
+#include "Engine/EngineTypes.h"
+#include "Engine/World.h"
+#include "GameFramework/Actor.h"
+#include "GameFramework/IGameplayCameraSystemHost.h"
+#include "Kismet/GameplayStatics.h"
+#include "Services/CameraParameterSetterService.h"
+
+UGameplayCameraParameterSetterComponent::UGameplayCameraParameterSetterComponent(const FObjectInitializer& ObjInit)
+	: Super(ObjInit)
+{
+}
+
+void UGameplayCameraParameterSetterComponent::BeginPlay()
+{
+	Super::BeginPlay();
+
+	if (AActor* OwnerActor = GetOwner())
+	{
+		OwnerActor->OnActorBeginOverlap.AddDynamic(this, &UGameplayCameraParameterSetterComponent::OnActorBeginOverlap);
+		OwnerActor->OnActorEndOverlap.AddDynamic(this, &UGameplayCameraParameterSetterComponent::OnActorEndOverlap);
+	}
+}
+
+void UGameplayCameraParameterSetterComponent::EndPlay(const EEndPlayReason::Type EndPlayReason)
+{
+	if (AActor* OwnerActor = GetOwner())
+	{
+		OwnerActor->OnActorBeginOverlap.RemoveAll(this);
+		OwnerActor->OnActorEndOverlap.RemoveAll(this);
+	}
+
+	Super::EndPlay(EndPlayReason);
+}
+
+void UGameplayCameraParameterSetterComponent::OnActorBeginOverlap(AActor* OverlappedActor, AActor* OtherActor)
+{
+	StartParameterSetters();
+}
+
+void UGameplayCameraParameterSetterComponent::OnActorEndOverlap(AActor* OverlappedActor, AActor* OtherActor)
+{
+	StopParameterSetters(false);
+}
+
+TSharedPtr<UE::Cameras::FCameraParameterSetterService> UGameplayCameraParameterSetterComponent::GetParameterSetterService()
+{
+	using namespace UE::Cameras;
+
+	// For now we only support one local player.
+	APlayerController* PlayerController = UGameplayStatics::GetPlayerController(this, 0);
+	if (!PlayerController)
+	{
+		UE_LOG(LogCameraSystem, Error, TEXT("Can't set camera parameters: no player controller found!"));
+		return nullptr;
+	}
+
+	IGameplayCameraSystemHost* CameraSystemHost = IGameplayCameraSystemHost::FindActiveHost(PlayerController);
+	if (!CameraSystemHost)
+	{
+		UE_LOG(LogCameraSystem, Error, TEXT("Can't set camera parameters: no camera system found under the player controller!"));
+		return nullptr;
+	}
+
+	TSharedPtr<FCameraSystemEvaluator> SystemEvaluator = CameraSystemHost->GetCameraSystemEvaluator();
+	if (!SystemEvaluator)
+	{
+		UE_LOG(LogCameraSystem, Error, TEXT("Can't set camera parameters: no camera system is active!"));
+		return nullptr;
+	}
+
+	return SystemEvaluator->FindEvaluationService<FCameraParameterSetterService>();
+}
+
+void UGameplayCameraParameterSetterComponent::StartParameterSetters()
+{
+	using namespace UE::Cameras;
+
+	if (!CameraRigReference.GetCameraRig())
+	{
+		return;
+	}
+
+	TSharedPtr<FCameraParameterSetterService> ParameterSetterService = GetParameterSetterService();
+	if (!ParameterSetterService)
+	{
+		return;
+	}
+
+	const UCameraRigAsset* CameraRig = CameraRigReference.GetCameraRig();
+	const FInstancedPropertyBag& ParameterValues = CameraRigReference.GetParameters();
+	const uint8* ParameterValuesPtr = ParameterValues.GetValue().GetMemory();
+
+	for (const FCameraObjectInterfaceParameterDefinition& ParameterDefinition : CameraRig->GetParameterDefinitions())
+	{
+		if (ParameterDefinition.ParameterType != ECameraObjectInterfaceParameterType::Blendable)
+		{
+			continue;
+		}
+
+		if (!CameraRigReference.IsParameterOverridden(ParameterDefinition.ParameterGuid))
+		{
+			continue;
+		}
+
+		const FPropertyBagPropertyDesc* PropertyDesc = ParameterValues.FindPropertyDescByID(ParameterDefinition.ParameterGuid);
+		if (!PropertyDesc || !PropertyDesc->CachedProperty)
+		{
+			continue;
+		}
+
+		FCameraVariableSetterHandle SetterHandle;
+		const void* RawValue = PropertyDesc->CachedProperty->ContainerPtrToValuePtr<void>(ParameterValuesPtr);
+
+		switch (ParameterDefinition.VariableType)
+		{
+#define UE_CAMERA_VARIABLE_FOR_TYPE(VariableType, VariableName)\
+			case ECameraVariableType::VariableName:\
+				{\
+					TCameraVariableSetter<VariableType> Setter(\
+							ParameterDefinition.VariableID, *reinterpret_cast<const VariableType*>(RawValue));\
+					InitializeParameterSetter(Setter);\
+					SetterHandle = ParameterSetterService->AddCameraVariableSetter(Setter);\
+				}\
+				break;
+			UE_CAMERA_VARIABLE_FOR_ALL_TYPES()
+#undef UE_CAMERA_VARIABLE_FOR_TYPE
+			default:
+				break;
+		}
+
+		if (SetterHandle.IsValid())
+		{
+			SetterHandles.Add(SetterHandle);
+		}
+	}
+}
+
+void UGameplayCameraParameterSetterComponent::InitializeParameterSetter(UE::Cameras::FCameraVariableSetter& VariableSetter)
+{
+	VariableSetter.BlendInTime = BlendInTime;
+	VariableSetter.BlendOutTime = BlendOutTime;
+	VariableSetter.BlendType = BlendType;
+}
+
+void UGameplayCameraParameterSetterComponent::StopParameterSetters(bool bImmediately)
+{
+	using namespace UE::Cameras;
+
+	TSharedPtr<FCameraParameterSetterService> ParameterSetterService = GetParameterSetterService();
+	if (!ParameterSetterService)
+	{
+		return;
+	}
+
+	for (FCameraVariableSetterHandle Handle : SetterHandles)
+	{
+		ParameterSetterService->StopCameraVariableSetter(Handle, bImmediately);
+	}
+}
+
diff --git a/Engine/Plugins/Cameras/GameplayCameras/Source/GameplayCameras/Public/GameFramework/GameplayCameraParameterSetterComponent.h b/Engine/Plugins/Cameras/GameplayCameras/Source/GameplayCameras/Public/GameFramework/GameplayCameraParameterSetterComponent.h
new file mode 100644
--- /dev/null
+++ b/Engine/Plugins/Cameras/GameplayCameras/Source/GameplayCameras/Public/GameFramework/GameplayCameraParameterSetterComponent.h
@@ -0,0 +1,82 @@
+// Copyright Epic Games, Inc. All Rights Reserved.
+
+#pragma once
+
+#include "Components/ActorComponent.h"
+#include "Core/CameraRigAssetReference.h"
+#include "Core/CameraVariableSetter.h"
+
+#include "GameplayCameraParameterSetterComponent.generated.h"
+
+namespace UE::Cameras
+{
+	class FCameraParameterSetterService;
+	struct FCameraVariableSetter;
+}
+
+UCLASS(BlueprintType, MinimalAPI, ClassGroup=Camera, meta=(BlueprintSpawnableComponent))
+class UGameplayCameraParameterSetterComponent : public UActorComponent
+{
+	GENERATED_BODY()
+
+public:
+
+	UGameplayCameraParameterSetterComponent(const FObjectInitializer& ObjInit);
+
+public:
+
+	/** The camera asset whose parameters to override. */
+	UPROPERTY(BlueprintReadWrite, EditAnywhere, Category=Camera)
+	FCameraRigAssetReference CameraRigReference;
+
+	/** The blend-in time. */
+	UPROPERTY(BlueprintReadWrite, EditAnywhere, Category=Camera)
+	float BlendInTime = 1.f;
+
+	/** The blend-out time. */
+	UPROPERTY(BlueprintReadWrite, EditAnywhere, Category=Camera)
+	float BlendOutTime = 1.f;
+
+	/** The blend type. */
+	UPROPERTY(BlueprintReadWrite, EditAnywhere, Category=Camera)
+	ECameraVariableSetterBlendType BlendType = ECameraVariableSetterBlendType::Linear;
+
+public:
+
+	/** Start setting parameters. */
+	UFUNCTION(BlueprintCallable, Category=Camera)
+	void StartParameterSetters();
+
+	/**
+	 * Stop setting parameters.
+	 *
+	 * @param bImmediately  If false, blend out the parameter overrides. If true, stop overrides immediately.
+	 */
+	UFUNCTION(BlueprintCallable, Category=Camera)
+	void StopParameterSetters(bool bImmediately = false);
+
+public:
+
+	// UActorComponent interface.
+	virtual void BeginPlay() override;
+	virtual void EndPlay(const EEndPlayReason::Type EndPlayReason) override;
+
+private:
+
+	UFUNCTION()
+	void OnActorBeginOverlap(AActor* OverlappedActor, AActor* OtherActor);
+
+	UFUNCTION()
+	void OnActorEndOverlap(AActor* OverlappedActor, AActor* OtherActor);
+
+private:
+
+	TSharedPtr<UE::Cameras::FCameraParameterSetterService> GetParameterSetterService();
+
+	void InitializeParameterSetter(UE::Cameras::FCameraVariableSetter& VariableSetter);
+
+private:
+
+	TArray<FCameraVariableSetterHandle> SetterHandles;
+};
+
