Cameras: when a camera rig component's asset gets built, correct rebuild its generated camera asset to handle added/removed parameters

diff --git a/Engine/Plugins/Cameras/GameplayCameras/Source/GameplayCameras/Private/GameFramework/GameplayCameraRigComponent.cpp b/Engine/Plugins/Cameras/GameplayCameras/Source/GameplayCameras/Private/GameFramework/GameplayCameraRigComponent.cpp
--- a/Engine/Plugins/Cameras/GameplayCameras/Source/GameplayCameras/Private/GameFramework/GameplayCameraRigComponent.cpp
+++ b/Engine/Plugins/Cameras/GameplayCameras/Source/GameplayCameras/Private/GameFramework/GameplayCameraRigComponent.cpp
@@ -85,13 +85,18 @@ void UGameplayCameraRigComponent::OnCameraRigAssetBuilt(const UCameraRigAsset* I
 {
 	using namespace UE::Cameras;
 
-	if (InCameraRigAsset != CameraRigReference.GetCameraRig())
+	if (InCameraRigAsset != CameraRigReference.GetCameraRig() || bIsBuildingGeneratedCameraAsset)
 	{
 		return;
 	}
 
 	// If our camera rig asset was just built, it may have some new parameters. We need to rebuild
 	// our variable table and context data table, and re-apply overrides.
+	if (GeneratedCameraAsset)
+	{
+		TGuardValue<bool> ReentrancyGuard(bIsBuildingGeneratedCameraAsset, true);
+		GeneratedCameraAsset->BuildCamera();
+	}
 	if (HasCameraEvaluationContext())
 	{
 		const FCameraObjectAllocationInfo& AllocationInfo = InCameraRigAsset->AllocationInfo;
diff --git a/Engine/Plugins/Cameras/GameplayCameras/Source/GameplayCameras/Public/GameFramework/GameplayCameraRigComponent.h b/Engine/Plugins/Cameras/GameplayCameras/Source/GameplayCameras/Public/GameFramework/GameplayCameraRigComponent.h
--- a/Engine/Plugins/Cameras/GameplayCameras/Source/GameplayCameras/Public/GameFramework/GameplayCameraRigComponent.h
+++ b/Engine/Plugins/Cameras/GameplayCameras/Source/GameplayCameras/Public/GameFramework/GameplayCameraRigComponent.h
@@ -54,5 +54,9 @@ private:
 
 	UPROPERTY()
 	TObjectPtr<UCameraAsset> GeneratedCameraAsset;
+
+#if WITH_EDITOR
+	bool bIsBuildingGeneratedCameraAsset = false;
+#endif
 };
 
