Cameras: when deactivating editor preview, return a camera component's output camera to its default transform

diff --git a/Engine/Plugins/Cameras/GameplayCameras/Source/GameplayCameras/Private/GameFramework/GameplayCameraComponentBase.cpp b/Engine/Plugins/Cameras/GameplayCameras/Source/GameplayCameras/Private/GameFramework/GameplayCameraComponentBase.cpp
--- a/Engine/Plugins/Cameras/GameplayCameras/Source/GameplayCameras/Private/GameFramework/GameplayCameraComponentBase.cpp
+++ b/Engine/Plugins/Cameras/GameplayCameras/Source/GameplayCameras/Private/GameFramework/GameplayCameraComponentBase.cpp
@@ -623,6 +623,8 @@ void UGameplayCameraComponentBase::AutoManageEditorPreviewEvaluator()
 		{
 			EvaluationContext->SetEditorPreviewCameraRigIndex(EditorPreviewCameraRigIndex);
 		}
+
+		// OutputCameraComponent will be updated on the next tick.
 	}
 	else if (!bCanRun && (CameraSystemEvaluator || EvaluationContext))
 	{
@@ -630,6 +632,11 @@ void UGameplayCameraComponentBase::AutoManageEditorPreviewEvaluator()
 		DeactivateCameraEvaluationContext(true);
 		TeardownCameraSystemHost();
 		EvaluationContext = nullptr;
+
+		if (OutputCameraComponent)
+		{
+			OutputCameraComponent->SetRelativeTransform(FTransform());
+		}
 	}
 }
 
