Cameras: fix update flag test

diff --git a/Engine/Plugins/Cameras/GameplayCameras/Source/GameplayCameras/Private/Tests/CameraSystemUpdateTests.cpp b/Engine/Plugins/Cameras/GameplayCameras/Source/GameplayCameras/Private/Tests/CameraSystemUpdateTests.cpp
--- a/Engine/Plugins/Cameras/GameplayCameras/Source/GameplayCameras/Private/Tests/CameraSystemUpdateTests.cpp
+++ b/Engine/Plugins/Cameras/GameplayCameras/Source/GameplayCameras/Private/Tests/CameraSystemUpdateTests.cpp
@@ -33,6 +33,8 @@ bool FCameraSystemFrameFlagsTest::RunTest(const FString& Parameters)
 	Evaluator.Initialize();
 	Evaluator.PushEvaluationContext(EvaluationContext);
 
+	EvaluationContext->GetInitialResult().bIsValid = true;
+
 	FCameraSystemEvaluationParams Params;
 	Params.DeltaTime = 0.3f;
 	Evaluator.Update(Params);
