Cameras: FindEvaluationContextOwnerActor shouldn't return a valid pointer if it can't be cast to the desired actor class

diff --git a/Engine/Plugins/Cameras/GameplayCameras/Source/GameplayCameras/Private/Nodes/Utility/BlueprintCameraNode.cpp b/Engine/Plugins/Cameras/GameplayCameras/Source/GameplayCameras/Private/Nodes/Utility/BlueprintCameraNode.cpp
--- a/Engine/Plugins/Cameras/GameplayCameras/Source/GameplayCameras/Private/Nodes/Utility/BlueprintCameraNode.cpp
+++ b/Engine/Plugins/Cameras/GameplayCameras/Source/GameplayCameras/Private/Nodes/Utility/BlueprintCameraNode.cpp
@@ -227,18 +227,22 @@ AActor* UBlueprintCameraNodeEvaluator::FindEvaluationContextOwnerActor(TSubclass
 {
 	if (CurrentContext)
 	{
+		AActor* OwnerActor = nullptr;
+
 		if (UActorComponent* ContextOwnerAsComponent = Cast<UActorComponent>(CurrentContext->GetOwner()))
 		{
-			return ContextOwnerAsComponent->GetOwner();
+			OwnerActor = ContextOwnerAsComponent->GetOwner();
 		}
 		else if (AActor* ContextOwnerAsActor = Cast<AActor>(CurrentContext->GetOwner()))
 		{
-			return ContextOwnerAsActor;
+			OwnerActor = ContextOwnerAsActor;
 		}
-		else
+
+		if (OwnerActor && OwnerActor->IsA(ActorClass))
 		{
-			return nullptr;
+			return OwnerActor;
 		}
+		return nullptr;
 	}
 	else
 	{
