Cameras: Blueprint camera director needs more checks when deactivating due to GC'ed owner

diff --git a/Engine/Plugins/Cameras/GameplayCameras/Source/GameplayCameras/Private/Directors/BlueprintCameraDirector.cpp b/Engine/Plugins/Cameras/GameplayCameras/Source/GameplayCameras/Private/Directors/BlueprintCameraDirector.cpp
--- a/Engine/Plugins/Cameras/GameplayCameras/Source/GameplayCameras/Private/Directors/BlueprintCameraDirector.cpp
+++ b/Engine/Plugins/Cameras/GameplayCameras/Source/GameplayCameras/Private/Directors/BlueprintCameraDirector.cpp
@@ -90,7 +90,10 @@ void FBlueprintCameraDirectorEvaluator::OnActivate(const FCameraDirectorActivate
 
 void FBlueprintCameraDirectorEvaluator::OnDeactivate(const FCameraDirectorDeactivateParams& Params)
 {
-	if (EvaluatorBlueprint)
+	// We need to check a few more things here in case we're being deactivated while the owner object is getting GC'ed.
+	UObject* ContextOwner = GetEvaluationContext()->GetOwner();
+	const bool bIsValid = (EvaluatorBlueprint && ContextOwner && !ContextOwner->HasAnyFlags(RF_BeginDestroyed | RF_FinishDestroyed));
+	if (bIsValid)
 	{
 		EvaluatorBlueprint->NativeDeactivateCameraDirector(Params);
 	}
