Cameras: AttachToPlayerPawn should only log once if now player pawn can be found

diff --git a/Engine/Plugins/Cameras/GameplayCameras/Source/GameplayCameras/Private/Nodes/Attach/AttachToPlayerPawnCameraNode.cpp b/Engine/Plugins/Cameras/GameplayCameras/Source/GameplayCameras/Private/Nodes/Attach/AttachToPlayerPawnCameraNode.cpp
--- a/Engine/Plugins/Cameras/GameplayCameras/Source/GameplayCameras/Private/Nodes/Attach/AttachToPlayerPawnCameraNode.cpp
+++ b/Engine/Plugins/Cameras/GameplayCameras/Source/GameplayCameras/Private/Nodes/Attach/AttachToPlayerPawnCameraNode.cpp
@@ -27,6 +27,8 @@ private:
 
 	TCameraParameterReader<bool> AttachToLocationReader;
 	TCameraParameterReader<bool> AttachToRotationReader;
+
+	bool bHasValidPlayerController = true;
 };
 
 UE_DEFINE_CAMERA_NODE_EVALUATOR(FAttachToPlayerPawnCameraNodeEvaluator)
@@ -36,6 +38,14 @@ void FAttachToPlayerPawnCameraNodeEvaluator::OnInitialize(const FCameraNodeEvalu
 	const UAttachToPlayerPawnCameraNode* AttachNode = GetCameraNodeAs<UAttachToPlayerPawnCameraNode>();
 	AttachToLocationReader.Initialize(AttachNode->AttachToLocation);
 	AttachToRotationReader.Initialize(AttachNode->AttachToRotation);
+
+	APlayerController* PlayerController = Params.EvaluationContext->GetPlayerController();
+	if (!PlayerController)
+	{
+		UE_LOG(LogCameraSystem, Error, 
+				TEXT("Can't run AttatchToPlayerPawn camera node because no player controller was found on the context."));
+		bHasValidPlayerController = false;
+	}
 }
 
 void FAttachToPlayerPawnCameraNodeEvaluator::OnRun(const FCameraNodeEvaluationParams& Params, FCameraNodeEvaluationResult& OutResult)
@@ -48,9 +58,15 @@ void FAttachToPlayerPawnCameraNodeEvaluator::OnRun(const FCameraNodeEvaluationPa
 	}
 #endif
 
+	if (!bHasValidPlayerController)
+	{
+		return;
+	}
+
 	APlayerController* PlayerController = Params.EvaluationContext->GetPlayerController();
 	if (!ensure(PlayerController))
 	{
+		bHasValidPlayerController = false;
 		return;
 	}
 
