Cameras: add "family shortcut" bar to various editors

This shortcut bar is similar to the one found on animation editors, although it works a bit differently. This solves the problem of not being able to quickly navigate from a camera asset to any of its camera rigs, and back.

diff --git a/Engine/Plugins/Cameras/GameplayCameras/Source/GameplayCameras/Private/Core/CameraAsset.cpp b/Engine/Plugins/Cameras/GameplayCameras/Source/GameplayCameras/Private/Core/CameraAsset.cpp
--- a/Engine/Plugins/Cameras/GameplayCameras/Source/GameplayCameras/Private/Core/CameraAsset.cpp
+++ b/Engine/Plugins/Cameras/GameplayCameras/Source/GameplayCameras/Private/Core/CameraAsset.cpp
@@ -6,6 +6,7 @@
 #include "Build/CameraBuildLog.h"
 #include "Core/CameraDirector.h"
 #include "Core/CameraRigAsset.h"
+#include "UObject/AssetRegistryTagsContext.h"
 #include "UObject/ObjectRedirector.h"
 #include "UObject/ObjectSaveContext.h"
 
@@ -107,8 +108,38 @@ void UCameraAsset::PostLoad()
 #endif  // WITH_EDITOR
 }
 
+void UCameraAsset::GetAssetRegistryTags(FAssetRegistryTagsContext Context) const
+{
+	if (CameraDirector)
+	{
+		FCameraDirectorRigUsageInfo UsageInfo;
+		CameraDirector->GatherRigUsageInfo(UsageInfo);
+
+		TStringBuilder<256> CameraRigListBuilder;
+		for (const UCameraRigAsset* CameraRig : UsageInfo.CameraRigs)
+		{
+			const UPackage* CameraRigPackage = CameraRig->GetPackage();
+			CameraRigListBuilder << TEXT(";");
+			CameraRigListBuilder << CameraRigPackage->GetName();
+		}
+
+		FAssetRegistryTag CameraRigsTag;
+		CameraRigsTag.Type = FAssetRegistryTag::ETagType::TT_Hidden;
+		CameraRigsTag.Name = TEXT("UsedCameraRigs");
+		CameraRigsTag.Value = CameraRigListBuilder.ToString();
+		Context.AddTag(CameraRigsTag);
+	}
+
+	Super::GetAssetRegistryTags(Context);
+}
+
 #if WITH_EDITOR
 
+void UCameraAsset::GetAssetRegistryTagMetadata(TMap<FName, FAssetRegistryTagMetadata>& OutMetadata) const
+{
+	Super::GetAssetRegistryTagMetadata(OutMetadata);
+}
+
 void UCameraAsset::PostEditChangeProperty(FPropertyChangedEvent& PropertyChangedEvent)
 {
 	using namespace UE::Cameras;
diff --git a/Engine/Plugins/Cameras/GameplayCameras/Source/GameplayCameras/Public/Core/CameraAsset.h b/Engine/Plugins/Cameras/GameplayCameras/Source/GameplayCameras/Public/Core/CameraAsset.h
--- a/Engine/Plugins/Cameras/GameplayCameras/Source/GameplayCameras/Public/Core/CameraAsset.h
+++ b/Engine/Plugins/Cameras/GameplayCameras/Source/GameplayCameras/Public/Core/CameraAsset.h
@@ -155,7 +155,9 @@ protected:
 	// UObject interface.
 	virtual void PreSave(FObjectPreSaveContext ObjectSaveContext) override;
 	virtual void PostLoad() override;
+	virtual void GetAssetRegistryTags(FAssetRegistryTagsContext Context) const override;
 #if WITH_EDITOR
+	virtual void GetAssetRegistryTagMetadata(TMap<FName, FAssetRegistryTagMetadata>& OutMetadata) const override;
 	virtual void PostEditChangeProperty(FPropertyChangedEvent& PropertyChangedEvent) override;
 #endif
 
diff --git a/Engine/Plugins/Cameras/GameplayCameras/Source/GameplayCamerasEditor/GameplayCamerasEditor.Build.cs b/Engine/Plugins/Cameras/GameplayCameras/Source/GameplayCamerasEditor/GameplayCamerasEditor.Build.cs
--- a/Engine/Plugins/Cameras/GameplayCameras/Source/GameplayCamerasEditor/GameplayCamerasEditor.Build.cs
+++ b/Engine/Plugins/Cameras/GameplayCameras/Source/GameplayCamerasEditor/GameplayCamerasEditor.Build.cs
@@ -24,6 +24,7 @@ public class GameplayCamerasEditor : ModuleRules
 				"AssetTools",
 				"BlueprintGraph",
 				"CinematicCamera",
+				"ContentBrowser",
 				"Core",
 				"CoreUObject",
 				"CurveEditor",
diff --git a/Engine/Plugins/Cameras/GameplayCameras/Source/GameplayCamerasEditor/Private/Family/CameraAssetFamily.cpp b/Engine/Plugins/Cameras/GameplayCameras/Source/GameplayCamerasEditor/Private/Family/CameraAssetFamily.cpp
new file mode 100644
--- /dev/null
+++ b/Engine/Plugins/Cameras/GameplayCameras/Source/GameplayCamerasEditor/Private/Family/CameraAssetFamily.cpp
@@ -0,0 +1,94 @@
+// Copyright Epic Games, Inc. All Rights Reserved.
+
+#include "Family/CameraAssetFamily.h"
+
+#include "Core/CameraAsset.h"
+#include "Core/CameraDirector.h"
+#include "Core/CameraRigAsset.h"
+#include "Core/CameraRigProxyAsset.h"
+#include "Family/GameplayCamerasFamilyConstants.h"
+
+#define LOCTEXT_NAMESPACE "CameraAssetFamily"
+
+namespace UE::Cameras
+{
+
+FCameraAssetFamily::FCameraAssetFamily(UCameraAsset* InRootAsset)
+	: RootAsset(InRootAsset)
+{
+	ensure(InRootAsset);
+}
+
+UObject* FCameraAssetFamily::GetRootAsset() const
+{
+	return RootAsset;
+}
+
+void FCameraAssetFamily::GetAssetTypes(TArray<UClass*>& OutAssetTypes) const
+{
+	OutAssetTypes.Add(UCameraAsset::StaticClass());
+	OutAssetTypes.Add(UCameraRigAsset::StaticClass());
+	OutAssetTypes.Add(UCameraRigProxyAsset::StaticClass());
+}
+
+void FCameraAssetFamily::FindAssetsOfType(UClass* InAssetType, TArray<FAssetData>& OutAssets) const
+{
+	if (!RootAsset)
+	{
+		return;
+	}
+
+	if (InAssetType == UCameraAsset::StaticClass())
+	{
+		OutAssets.Add(RootAsset);
+		return;
+	}
+
+	UCameraDirector* CameraDirector = RootAsset->GetCameraDirector();
+	if (!CameraDirector)
+	{
+		return;
+	}
+
+	FCameraDirectorRigUsageInfo UsageInfo;
+	CameraDirector->GatherRigUsageInfo(UsageInfo);
+
+	if (InAssetType == UCameraRigAsset::StaticClass())
+	{
+		for (UCameraRigAsset* CameraRig : UsageInfo.CameraRigs)
+		{
+			OutAssets.Add(FAssetData(CameraRig));
+		}
+	}
+	else if (InAssetType == UCameraRigProxyAsset::StaticClass())
+	{
+	}
+}
+
+FText FCameraAssetFamily::GetAssetTypeTooltip(UClass* InAssetType) const
+{
+	if (InAssetType == UCameraRigAsset::StaticClass())
+	{
+		return LOCTEXT("CameraRigAssetTypeTooltip", "Open camera rigs referenced by this asset.");
+	}
+	else if (InAssetType == UCameraRigProxyAsset::StaticClass())
+	{
+		return LOCTEXT("CameraRigAssetTypeTooltip", "Open camera rig proxies referenced by this asset.");
+	}
+	return FText();
+}
+
+const FSlateBrush* FCameraAssetFamily::GetAssetIcon(UClass* InAssetType) const
+{
+	return FGameplayCamerasFamilyConstants::GetAssetIcon(InAssetType);
+}
+
+FSlateColor FCameraAssetFamily::GetAssetTint(UClass* InAssetType) const
+{
+	return FGameplayCamerasFamilyConstants::GetAssetTint(InAssetType);
+}
+
+}  // namespace UE::Cameras
+
+#undef LOCTEXT_NAMESPACE
+
diff --git a/Engine/Plugins/Cameras/GameplayCameras/Source/GameplayCamerasEditor/Private/Family/CameraAssetFamily.h b/Engine/Plugins/Cameras/GameplayCameras/Source/GameplayCamerasEditor/Private/Family/CameraAssetFamily.h
new file mode 100644
--- /dev/null
+++ b/Engine/Plugins/Cameras/GameplayCameras/Source/GameplayCamerasEditor/Private/Family/CameraAssetFamily.h
@@ -0,0 +1,32 @@
+// Copyright Epic Games, Inc. All Rights Reserved.
+
+#pragma once
+
+#include "IGameplayCamerasFamily.h"
+
+class UCameraAsset;
+
+namespace UE::Cameras
+{
+
+class FCameraAssetFamily : public IGameplayCamerasFamily
+{
+public:
+
+	FCameraAssetFamily(UCameraAsset* InRootAsset);
+
+	// IGameplayCamerasFamily interface.
+	virtual UObject* GetRootAsset() const override;
+	virtual void GetAssetTypes(TArray<UClass*>& OutAssetTypes) const override;
+	virtual void FindAssetsOfType(UClass* InAssetType, TArray<FAssetData>& OutAssets) const override;
+	virtual FText GetAssetTypeTooltip(UClass* InAssetType) const override;
+	virtual const FSlateBrush* GetAssetIcon(UClass* InAssetType) const override;
+	virtual FSlateColor GetAssetTint(UClass* InAssetType) const override;
+
+private:
+
+	UCameraAsset* RootAsset;
+};
+
+}  // namespace UE::Cameras
+
diff --git a/Engine/Plugins/Cameras/GameplayCameras/Source/GameplayCamerasEditor/Private/Family/CameraRigAssetFamily.cpp b/Engine/Plugins/Cameras/GameplayCameras/Source/GameplayCamerasEditor/Private/Family/CameraRigAssetFamily.cpp
new file mode 100644
--- /dev/null
+++ b/Engine/Plugins/Cameras/GameplayCameras/Source/GameplayCamerasEditor/Private/Family/CameraRigAssetFamily.cpp
@@ -0,0 +1,118 @@
+// Copyright Epic Games, Inc. All Rights Reserved.
+
+#include "Family/CameraRigAssetFamily.h"
+
+#include "AssetRegistry/AssetRegistryModule.h"
+#include "Core/CameraAsset.h"
+#include "Core/CameraDirector.h"
+#include "Core/CameraRigAsset.h"
+#include "Family/GameplayCamerasFamilyConstants.h"
+#include "UObject/ReferencerFinder.h"
+
+#define LOCTEXT_NAMESPACE "CameraRigAssetFamily"
+
+namespace UE::Cameras
+{
+
+FCameraRigAssetFamily::FCameraRigAssetFamily(UCameraRigAsset* InRootAsset)
+	: RootAsset(InRootAsset)
+{
+	ensure(InRootAsset);
+}
+
+UObject* FCameraRigAssetFamily::GetRootAsset() const
+{
+	return RootAsset;
+}
+
+void FCameraRigAssetFamily::GetAssetTypes(TArray<UClass*>& OutAssetTypes) const
+{
+	OutAssetTypes.Add(UCameraAsset::StaticClass());
+	OutAssetTypes.Add(UCameraRigAsset::StaticClass());
+}
+
+void FCameraRigAssetFamily::FindAssetsOfType(UClass* InAssetType, TArray<FAssetData>& OutAssets) const
+{
+	if (!RootAsset)
+	{
+		return;
+	}
+
+	if (InAssetType == UCameraRigAsset::StaticClass())
+	{
+		OutAssets.Add(RootAsset);
+	}
+	else if (InAssetType == UCameraAsset::StaticClass())
+	{
+		const FName RootPackageName = RootAsset->GetPackage()->GetFName();
+
+		TArray<FAssetData> AllCameraAssets;
+		IAssetRegistry& AssetRegistry = FAssetRegistryModule::GetRegistry();
+		AssetRegistry.GetAssetsByClass(UCameraAsset::StaticClass()->GetClassPathName(), AllCameraAssets);
+		for (const FAssetData& CameraAsset : AllCameraAssets)
+		{
+			// By default, use the asset tags to know which camera asset uses which camera rigs.
+			// If the asset is loaded, it might have been modified and not saved yet, so use the in-memory
+			// object instead.
+			// If the asset doesn't have tags, it hasn't been saved since tags were added, so load it
+			// and also use it directly in memory.
+			bool bUseObjectDirectly = (CameraAsset.IsAssetLoaded());
+			if (!bUseObjectDirectly)
+			{
+				FString UsedCameraRigsTag = CameraAsset.GetTagValueRef<FString>(TEXT("UsedCameraRigs"));
+				if (!UsedCameraRigsTag.IsEmpty())
+				{
+					TArray<FString> UsedCameraRigs;
+					UsedCameraRigsTag.ParseIntoArray(UsedCameraRigs, TEXT(";"));
+					if (UsedCameraRigs.Contains(RootPackageName))
+					{
+						OutAssets.Add(CameraAsset);
+					}
+				}
+				else
+				{
+					bUseObjectDirectly = true;
+				}
+			}
+			if (bUseObjectDirectly)
+			{
+				if (UCameraAsset* LoadedCameraAsset = Cast<UCameraAsset>(CameraAsset.GetAsset()))
+				{
+					if (UCameraDirector* LoadedCameraDirector = LoadedCameraAsset->GetCameraDirector())
+					{
+						FCameraDirectorRigUsageInfo UsageInfo;
+						LoadedCameraDirector->GatherRigUsageInfo(UsageInfo);
+						if (UsageInfo.CameraRigs.Contains(RootAsset))
+						{
+							OutAssets.Add(CameraAsset);
+						}
+					}
+				}
+			}
+		}
+	}
+}
+
+FText FCameraRigAssetFamily::GetAssetTypeTooltip(UClass* InAssetType) const
+{
+	if (InAssetType == UCameraAsset::StaticClass())
+	{
+		return LOCTEXT("CameraAssetTypeTooltip", "Open camera assets referencing this asset.");
+	}
+	return FText();
+}
+
+const FSlateBrush* FCameraRigAssetFamily::GetAssetIcon(UClass* InAssetType) const
+{
+	return FGameplayCamerasFamilyConstants::GetAssetIcon(InAssetType);
+}
+
+FSlateColor FCameraRigAssetFamily::GetAssetTint(UClass* InAssetType) const
+{
+	return FGameplayCamerasFamilyConstants::GetAssetTint(InAssetType);
+}
+
+}  // namespace UE::Cameras
+
+#undef LOCTEXT_NAMESPACE
+
diff --git a/Engine/Plugins/Cameras/GameplayCameras/Source/GameplayCamerasEditor/Private/Family/CameraRigAssetFamily.h b/Engine/Plugins/Cameras/GameplayCameras/Source/GameplayCamerasEditor/Private/Family/CameraRigAssetFamily.h
new file mode 100644
--- /dev/null
+++ b/Engine/Plugins/Cameras/GameplayCameras/Source/GameplayCamerasEditor/Private/Family/CameraRigAssetFamily.h
@@ -0,0 +1,32 @@
+// Copyright Epic Games, Inc. All Rights Reserved.
+
+#pragma once
+
+#include "IGameplayCamerasFamily.h"
+
+class UCameraRigAsset;
+
+namespace UE::Cameras
+{
+
+class FCameraRigAssetFamily : public IGameplayCamerasFamily
+{
+public:
+
+	FCameraRigAssetFamily(UCameraRigAsset* InRootAsset);
+
+	// IGameplayCamerasFamily interface.
+	virtual UObject* GetRootAsset() const override;
+	virtual void GetAssetTypes(TArray<UClass*>& OutAssetTypes) const override;
+	virtual void FindAssetsOfType(UClass* InAssetType, TArray<FAssetData>& OutAssets) const override;
+	virtual FText GetAssetTypeTooltip(UClass* InAssetType) const override;
+	virtual const FSlateBrush* GetAssetIcon(UClass* InAssetType) const override;
+	virtual FSlateColor GetAssetTint(UClass* InAssetType) const override;
+
+private:
+
+	UCameraRigAsset* RootAsset;
+};
+
+}  // namespace UE::Cameras
+
diff --git a/Engine/Plugins/Cameras/GameplayCameras/Source/GameplayCamerasEditor/Private/Family/CameraRigProxyAssetFamily.cpp b/Engine/Plugins/Cameras/GameplayCameras/Source/GameplayCamerasEditor/Private/Family/CameraRigProxyAssetFamily.cpp
new file mode 100644
--- /dev/null
+++ b/Engine/Plugins/Cameras/GameplayCameras/Source/GameplayCamerasEditor/Private/Family/CameraRigProxyAssetFamily.cpp
@@ -0,0 +1,60 @@
+// Copyright Epic Games, Inc. All Rights Reserved.
+
+#include "Family/CameraRigProxyAssetFamily.h"
+
+#include "Core/CameraAsset.h"
+#include "Core/CameraRigProxyAsset.h"
+#include "Family/GameplayCamerasFamilyConstants.h"
+
+#define LOCTEXT_NAMESPACE "CameraRigProxyAssetFamily"
+
+namespace UE::Cameras
+{
+
+FCameraRigProxyAssetFamily::FCameraRigProxyAssetFamily(UCameraRigProxyAsset* InRootAsset)
+	: RootAsset(InRootAsset)
+{
+	ensure(InRootAsset);
+}
+
+UObject* FCameraRigProxyAssetFamily::GetRootAsset() const
+{
+	return RootAsset;
+}
+
+void FCameraRigProxyAssetFamily::GetAssetTypes(TArray<UClass*>& OutAssetTypes) const
+{
+	OutAssetTypes.Add(UCameraAsset::StaticClass());
+	OutAssetTypes.Add(UCameraRigProxyAsset::StaticClass());
+}
+
+void FCameraRigProxyAssetFamily::FindAssetsOfType(UClass* InAssetType, TArray<FAssetData>& OutAssets) const
+{
+	if (InAssetType == UCameraAsset::StaticClass())
+	{
+	}
+}
+
+FText FCameraRigProxyAssetFamily::GetAssetTypeTooltip(UClass* InAssetType) const
+{
+	if (InAssetType == UCameraAsset::StaticClass())
+	{
+		return LOCTEXT("CameraRigAssetTypeTooltip", "Open camera assets referencing this asset.");
+	}
+	return FText();
+}
+
+const FSlateBrush* FCameraRigProxyAssetFamily::GetAssetIcon(UClass* InAssetType) const
+{
+	return FGameplayCamerasFamilyConstants::GetAssetIcon(InAssetType);
+}
+
+FSlateColor FCameraRigProxyAssetFamily::GetAssetTint(UClass* InAssetType) const
+{
+	return FGameplayCamerasFamilyConstants::GetAssetTint(InAssetType);
+}
+
+}  // namespace UE::Cameras
+
+#undef LOCTEXT_NAMESPACE
+
diff --git a/Engine/Plugins/Cameras/GameplayCameras/Source/GameplayCamerasEditor/Private/Family/CameraRigProxyAssetFamily.h b/Engine/Plugins/Cameras/GameplayCameras/Source/GameplayCamerasEditor/Private/Family/CameraRigProxyAssetFamily.h
new file mode 100644
--- /dev/null
+++ b/Engine/Plugins/Cameras/GameplayCameras/Source/GameplayCamerasEditor/Private/Family/CameraRigProxyAssetFamily.h
@@ -0,0 +1,32 @@
+// Copyright Epic Games, Inc. All Rights Reserved.
+
+#pragma once
+
+#include "IGameplayCamerasFamily.h"
+
+class UCameraRigProxyAsset;
+
+namespace UE::Cameras
+{
+
+class FCameraRigProxyAssetFamily : public IGameplayCamerasFamily
+{
+public:
+
+	FCameraRigProxyAssetFamily(UCameraRigProxyAsset* InRootAsset);
+
+	// IGameplayCamerasFamily interface.
+	virtual UObject* GetRootAsset() const override;
+	virtual void GetAssetTypes(TArray<UClass*>& OutAssetTypes) const override;
+	virtual void FindAssetsOfType(UClass* InAssetType, TArray<FAssetData>& OutAssets) const override;
+	virtual FText GetAssetTypeTooltip(UClass* InAssetType) const override;
+	virtual const FSlateBrush* GetAssetIcon(UClass* InAssetType) const override;
+	virtual FSlateColor GetAssetTint(UClass* InAssetType) const override;
+
+private:
+
+	UCameraRigProxyAsset* RootAsset;
+};
+
+}  // namespace UE::Cameras
+
diff --git a/Engine/Plugins/Cameras/GameplayCameras/Source/GameplayCamerasEditor/Private/Family/GameplayCamerasFamilyConstants.cpp b/Engine/Plugins/Cameras/GameplayCameras/Source/GameplayCamerasEditor/Private/Family/GameplayCamerasFamilyConstants.cpp
new file mode 100644
--- /dev/null
+++ b/Engine/Plugins/Cameras/GameplayCameras/Source/GameplayCamerasEditor/Private/Family/GameplayCamerasFamilyConstants.cpp
@@ -0,0 +1,54 @@
+// Copyright Epic Games, Inc. All Rights Reserved.
+
+#include "Family/GameplayCamerasFamilyConstants.h"
+
+#include "Core/CameraAsset.h"
+#include "Core/CameraRigAsset.h"
+#include "Core/CameraRigProxyAsset.h"
+#include "Styles/GameplayCamerasEditorStyle.h"
+
+namespace UE::Cameras
+{
+
+const FSlateBrush* FGameplayCamerasFamilyConstants::GetAssetIcon(UClass* InAssetType)
+{
+	if (InAssetType == UCameraAsset::StaticClass())
+	{
+		return FGameplayCamerasEditorStyle::Get()->GetBrush("Family.CameraAsset");
+	}
+	else if (InAssetType == UCameraRigAsset::StaticClass())
+	{
+		return FGameplayCamerasEditorStyle::Get()->GetBrush("Family.CameraRigAsset");
+	}
+	else if (InAssetType == UCameraRigProxyAsset::StaticClass())
+	{
+		return FGameplayCamerasEditorStyle::Get()->GetBrush("Family.CameraRigProxyAsset");
+	}
+	return nullptr;
+}
+
+FSlateColor FGameplayCamerasFamilyConstants::GetAssetTint(UClass* InAssetType)
+{
+	static const FColor Teal(23, 126, 137);
+	static const FColor MidnightGreen(8, 76, 97);
+	static const FColor Poppy(219, 58, 52);
+	static const FColor Sunglow(255, 200, 87);
+
+	if (InAssetType == UCameraAsset::StaticClass())
+	{
+		return FSlateColor(Sunglow);
+	}
+	else if (InAssetType == UCameraRigAsset::StaticClass())
+	{
+		return FSlateColor(Teal);
+	}
+	else if (InAssetType == UCameraRigProxyAsset::StaticClass())
+	{
+		return FSlateColor(MidnightGreen);
+	}
+	return FSlateColor();
+}
+
+}  // namespace UE::Cameras
+
+
diff --git a/Engine/Plugins/Cameras/GameplayCameras/Source/GameplayCamerasEditor/Private/Family/GameplayCamerasFamilyConstants.h b/Engine/Plugins/Cameras/GameplayCameras/Source/GameplayCamerasEditor/Private/Family/GameplayCamerasFamilyConstants.h
new file mode 100644
--- /dev/null
+++ b/Engine/Plugins/Cameras/GameplayCameras/Source/GameplayCamerasEditor/Private/Family/GameplayCamerasFamilyConstants.h
@@ -0,0 +1,23 @@
+// Copyright Epic Games, Inc. All Rights Reserved.
+
+#pragma once
+
+#include "Styling/SlateColor.h"
+
+class UClass;
+struct FSlateBrush;
+
+namespace UE::Cameras
+{
+
+class FGameplayCamerasFamilyConstants
+{
+public:
+
+	static const FSlateBrush* GetAssetIcon(UClass* InAssetType);
+
+	static FSlateColor GetAssetTint(UClass* InAssetType);
+};
+
+}  // namespace UE::Cameras
+
diff --git a/Engine/Plugins/Cameras/GameplayCameras/Source/GameplayCamerasEditor/Private/IGameplayCamerasFamily.cpp b/Engine/Plugins/Cameras/GameplayCameras/Source/GameplayCamerasEditor/Private/IGameplayCamerasFamily.cpp
new file mode 100644
--- /dev/null
+++ b/Engine/Plugins/Cameras/GameplayCameras/Source/GameplayCamerasEditor/Private/IGameplayCamerasFamily.cpp
@@ -0,0 +1,40 @@
+// Copyright Epic Games, Inc. All Rights Reserved.
+
+#include "IGameplayCamerasFamily.h"
+
+#include "Core/CameraAsset.h"
+#include "Core/CameraRigAsset.h"
+#include "Core/CameraRigProxyAsset.h"
+#include "Family/CameraAssetFamily.h"
+#include "Family/CameraRigAssetFamily.h"
+#include "Family/CameraRigProxyAssetFamily.h"
+
+namespace UE::Cameras
+{
+
+TSharedPtr<IGameplayCamerasFamily> IGameplayCamerasFamily::CreateFamily(UObject* InAsset)
+{
+	if (!ensure(InAsset))
+	{
+		return nullptr;
+	}
+
+	UClass* AssetType = InAsset->GetClass();
+	if (AssetType == UCameraAsset::StaticClass())
+	{
+		return MakeShared<FCameraAssetFamily>(CastChecked<UCameraAsset>(InAsset));
+	}
+	else if (AssetType == UCameraRigAsset::StaticClass())
+	{
+		return MakeShared<FCameraRigAssetFamily>(CastChecked<UCameraRigAsset>(InAsset));
+	}
+	else if (AssetType == UCameraRigProxyAsset::StaticClass())
+	{
+		return MakeShared<FCameraRigProxyAssetFamily>(CastChecked<UCameraRigProxyAsset>(InAsset));
+	}
+
+	return nullptr;
+}
+
+}  // namespace UE::Cameras
+
diff --git a/Engine/Plugins/Cameras/GameplayCameras/Source/GameplayCamerasEditor/Private/Styles/GameplayCamerasEditorStyle.cpp b/Engine/Plugins/Cameras/GameplayCameras/Source/GameplayCamerasEditor/Private/Styles/GameplayCamerasEditorStyle.cpp
--- a/Engine/Plugins/Cameras/GameplayCameras/Source/GameplayCamerasEditor/Private/Styles/GameplayCamerasEditorStyle.cpp
+++ b/Engine/Plugins/Cameras/GameplayCameras/Source/GameplayCamerasEditor/Private/Styles/GameplayCamerasEditorStyle.cpp
@@ -150,6 +150,11 @@ FGameplayCamerasEditorStyle::FGameplayCamerasEditorStyle()
 	Set("Graph.CameraRigParameterPin.Connected", new IMAGE_BRUSH("Graph/ObjectTreeGraphNode_DiamondPin_Connected", Icon16x16));
 	Set("Graph.CameraRigParameterPin.Disconnected", new IMAGE_BRUSH("Graph/ObjectTreeGraphNode_DiamondPin_Disconnected", Icon16x16));
 
+	// Family icons.
+	Set("Family.CameraAsset", new IMAGE_BRUSH_SVG("Icons/ContentBrowser-CameraKit", Icon20x20));
+	Set("Family.CameraRigAsset", new IMAGE_BRUSH_SVG("Icons/ContentBrowser-CameraRig", Icon20x20));
+	Set("Family.CameraRigProxyAsset", new IMAGE_BRUSH_SVG("Icons/ContentBrowser-CameraRigProxy", Icon20x20));
+
 	FSlateStyleRegistry::RegisterSlateStyle(*this);
 }
 
diff --git a/Engine/Plugins/Cameras/GameplayCameras/Source/GameplayCamerasEditor/Private/Toolkits/CameraAssetEditorToolkit.cpp b/Engine/Plugins/Cameras/GameplayCameras/Source/GameplayCamerasEditor/Private/Toolkits/CameraAssetEditorToolkit.cpp
--- a/Engine/Plugins/Cameras/GameplayCameras/Source/GameplayCamerasEditor/Private/Toolkits/CameraAssetEditorToolkit.cpp
+++ b/Engine/Plugins/Cameras/GameplayCameras/Source/GameplayCamerasEditor/Private/Toolkits/CameraAssetEditorToolkit.cpp
@@ -22,6 +22,7 @@
 #include "Helpers/ObjectReferenceFinder.h"
 #include "IAssetTools.h"
 #include "IGameplayCamerasEditorModule.h"
+#include "IGameplayCamerasFamily.h"
 #include "IGameplayCamerasLiveEditManager.h"
 #include "IGameplayCamerasModule.h"
 #include "PropertyEditorModule.h"
@@ -35,6 +36,7 @@
 #include "Toolkits/StandardToolkitLayout.h"
 #include "UObject/Package.h"
 #include "Widgets/Docking/SDockTab.h"
+#include "Widgets/SCameraFamilyShortcutBar.h"
 
 #include UE_INLINE_GENERATED_CPP_BY_NAME(CameraAssetEditorToolkit)
 
@@ -193,6 +195,19 @@ void FCameraAssetEditorToolkit::RegisterToolbar()
 		ModesSection.AddEntry(FToolMenuEntry::InitToolBarButton(Commands.ShowCameraDirector));
 		ModesSection.AddEntry(FToolMenuEntry::InitToolBarButton(Commands.ShowSharedTransitions));
 	}
+
+	TSharedRef<IGameplayCamerasFamily> Family = IGameplayCamerasFamily::CreateFamily(CameraAsset).ToSharedRef();
+	TSharedRef<FExtender> ToolbarExtender = MakeShared<FExtender>();
+	AddToolbarExtender(ToolbarExtender);
+	ToolbarExtender->AddToolBarExtension(
+			"Asset",
+			EExtensionHook::After,
+			GetToolkitCommands(),
+			FToolBarExtensionDelegate::CreateLambda([this, Family](FToolBarBuilder& Builder)
+				{
+					AddToolbarWidget(SNew(SCameraFamilyShortcutBar, SharedThis(this), Family));
+				})
+			);
 }
 
 void FCameraAssetEditorToolkit::InitToolMenuContext(FToolMenuContext& MenuContext)
diff --git a/Engine/Plugins/Cameras/GameplayCameras/Source/GameplayCamerasEditor/Private/Toolkits/CameraRigAssetEditorToolkit.cpp b/Engine/Plugins/Cameras/GameplayCameras/Source/GameplayCamerasEditor/Private/Toolkits/CameraRigAssetEditorToolkit.cpp
--- a/Engine/Plugins/Cameras/GameplayCameras/Source/GameplayCamerasEditor/Private/Toolkits/CameraRigAssetEditorToolkit.cpp
+++ b/Engine/Plugins/Cameras/GameplayCameras/Source/GameplayCamerasEditor/Private/Toolkits/CameraRigAssetEditorToolkit.cpp
@@ -16,6 +16,7 @@
 #include "Framework/Docking/TabManager.h"
 #include "Helpers/AssetTypeMenuOverlayHelper.h"
 #include "IGameplayCamerasEditorModule.h"
+#include "IGameplayCamerasFamily.h"
 #include "IGameplayCamerasLiveEditManager.h"
 #include "IGameplayCamerasModule.h"
 #include "Modules/ModuleManager.h"
@@ -28,6 +29,7 @@
 #include "Toolkits/CurveEditorToolkit.h"
 #include "Toolkits/StandardToolkitLayout.h"
 #include "Widgets/Docking/SDockTab.h"
+#include "Widgets/SCameraFamilyShortcutBar.h"
 
 #include UE_INLINE_GENERATED_CPP_BY_NAME(CameraRigAssetEditorToolkit)
 
@@ -258,6 +260,19 @@ void FCameraRigAssetEditorToolkit::RegisterToolbar()
 
 		Impl->BuildToolbarMenu(ToolbarMenu);
 	}
+	
+	TSharedRef<IGameplayCamerasFamily> Family = IGameplayCamerasFamily::CreateFamily(Impl->GetCameraRigAsset()).ToSharedRef();
+	TSharedRef<FExtender> ToolbarExtender = MakeShared<FExtender>();
+	AddToolbarExtender(ToolbarExtender);
+	ToolbarExtender->AddToolBarExtension(
+			"Asset",
+			EExtensionHook::After,
+			GetToolkitCommands(),
+			FToolBarExtensionDelegate::CreateLambda([this, Family](FToolBarBuilder& Builder)
+				{
+					AddToolbarWidget(SNew(SCameraFamilyShortcutBar, SharedThis(this), Family));
+				})
+			);
 }
 
 void FCameraRigAssetEditorToolkit::InitToolMenuContext(FToolMenuContext& MenuContext)
diff --git a/Engine/Plugins/Cameras/GameplayCameras/Source/GameplayCamerasEditor/Private/Widgets/SCameraFamilyAssetShortcut.cpp b/Engine/Plugins/Cameras/GameplayCameras/Source/GameplayCamerasEditor/Private/Widgets/SCameraFamilyAssetShortcut.cpp
new file mode 100644
--- /dev/null
+++ b/Engine/Plugins/Cameras/GameplayCameras/Source/GameplayCamerasEditor/Private/Widgets/SCameraFamilyAssetShortcut.cpp
@@ -0,0 +1,376 @@
+// Copyright Epic Games, Inc. All Rights Reserved.
+
+#include "Widgets/SCameraFamilyAssetShortcut.h"
+
+#include "AssetRegistry/AssetRegistryModule.h"
+#include "ContentBrowserModule.h"
+#include "Editor.h"
+#include "Framework/Application/SlateApplication.h"
+#include "Framework/MultiBox/MultiBoxBuilder.h"
+#include "Framework/MultiBox/MultiBoxDefs.h"
+#include "GameplayCamerasDelegates.h"
+#include "IAssetFamily.h"
+#include "IAssetTools.h"
+#include "IContentBrowserSingleton.h"
+#include "IGameplayCamerasFamily.h"
+#include "IGameplayCamerasEditorModule.h"
+#include "Modules/ModuleManager.h"
+#include "Styling/AppStyle.h"
+#include "Styling/SlateTypes.h"
+#include "Styling/ToolBarStyle.h"
+#include "Subsystems/AssetEditorSubsystem.h"
+#include "Textures/SlateIcon.h"
+#include "Tools/BaseAssetToolkit.h"
+#include "Widgets/Images/SImage.h"
+#include "Widgets/Input/SCheckBox.h"
+#include "Widgets/Input/SComboButton.h"
+#include "Widgets/Layout/SBox.h"
+#include "Widgets/Layout/SSeparator.h"
+#include "Widgets/SBoxPanel.h"
+#include "Widgets/SOverlay.h"
+
+#define LOCTEXT_NAMESPACE "SCameraFamilyShortcutBar"
+
+namespace UE::Cameras
+{
+
+void SCameraFamilyAssetShortcut::Construct(const FArguments& InArgs, const TSharedRef<FBaseAssetToolkit>& InToolkit, const TSharedRef<IGameplayCamerasFamily>& InFamily, UClass* InAssetType)
+{
+	WeakToolkit = InToolkit;
+
+	Family = InFamily;
+	FamilyAssetType = InAssetType;
+
+	bRefreshAssetDatas = true;
+
+	FAssetRegistryModule& AssetRegistryModule = FModuleManager::LoadModuleChecked<FAssetRegistryModule>(TEXT("AssetRegistry"));
+	AssetRegistryModule.Get().OnFilesLoaded().AddSP(this, &SCameraFamilyAssetShortcut::HandleFilesLoaded);
+	AssetRegistryModule.Get().OnAssetAdded().AddSP(this, &SCameraFamilyAssetShortcut::HandleAssetAdded);
+	AssetRegistryModule.Get().OnAssetRemoved().AddSP(this, &SCameraFamilyAssetShortcut::HandleAssetRemoved);
+	AssetRegistryModule.Get().OnAssetRenamed().AddSP(this, &SCameraFamilyAssetShortcut::HandleAssetRenamed);
+
+	FGameplayCamerasDelegates::OnCameraAssetBuilt().AddSP(this, &SCameraFamilyAssetShortcut::HandleCameraAssetBuilt);
+	FGameplayCamerasDelegates::OnCameraRigAssetBuilt().AddSP(this, &SCameraFamilyAssetShortcut::HandleCameraRigAssetBuilt);
+
+	const FToolBarStyle& ToolBarStyle = FAppStyle::Get().GetWidgetStyle<FToolBarStyle>("ToolBar");
+
+	ChildSlot
+	[
+		SNew(SHorizontalBox)
+
+		// Fat button, when there's only one asset of this type.
+		+SHorizontalBox::Slot()
+		.VAlign(VAlign_Center)
+		[
+			SAssignNew(SoloCheckBox, SCheckBox)
+			.Style(FAppStyle::Get(), "SegmentedCombo.ButtonOnly")
+			.OnCheckStateChanged(this, &SCameraFamilyAssetShortcut::HandleButtonClick)
+			.IsChecked(this, &SCameraFamilyAssetShortcut::GetCheckState)
+			.IsEnabled(this, &SCameraFamilyAssetShortcut::IsSoloButtonEnabled)
+			.Visibility(this, &SCameraFamilyAssetShortcut::GetSoloButtonVisibility)
+			.ToolTipText(this, &SCameraFamilyAssetShortcut::GetButtonTooltip)
+			.Padding(0.0f)
+			[
+				SNew(SOverlay)
+
+				+ SOverlay::Slot()
+				.VAlign(VAlign_Center)
+				.HAlign(HAlign_Center)
+				.Padding(FMargin(28.f, 4.f))
+				[
+					SNew(SImage)
+					.ColorAndOpacity(this, &SCameraFamilyAssetShortcut::GetAssetTint)
+					.Image(this, &SCameraFamilyAssetShortcut::GetAssetIcon)
+				]
+			]
+		]
+
+		// Button/dropdown pair, when there are more than one asset of this type.
+		+SHorizontalBox::Slot()
+		.VAlign(VAlign_Center)
+		[
+			SAssignNew(ComboCheckBox, SCheckBox)
+			.Style(FAppStyle::Get(), "SegmentedCombo.Left")
+			.OnCheckStateChanged(this, &SCameraFamilyAssetShortcut::HandleButtonClick)
+			.IsChecked(this, &SCameraFamilyAssetShortcut::GetCheckState)
+			.Visibility(this, &SCameraFamilyAssetShortcut::GetComboButtonVisibility)
+			.ToolTipText(this, &SCameraFamilyAssetShortcut::GetButtonTooltip)
+			.Padding(0.0f)
+			[
+				SNew(SOverlay)
+
+				+ SOverlay::Slot()
+				.VAlign(VAlign_Center)
+				.HAlign(HAlign_Center)
+				.Padding(FMargin(16.f, 4.f))
+				[
+					SNew(SImage)
+					.ColorAndOpacity(this, &SCameraFamilyAssetShortcut::GetAssetTint)
+					.Image(this, &SCameraFamilyAssetShortcut::GetAssetIcon)
+				]
+			]
+		]
+		+SHorizontalBox::Slot()
+		.VAlign(VAlign_Center)
+		.AutoWidth()
+		[
+			SNew(SSeparator)
+			.Visibility(this, &SCameraFamilyAssetShortcut::GetComboDropdownVisibility)
+			.Thickness(1.0f)
+			.Orientation(EOrientation::Orient_Vertical)
+		]
+		+SHorizontalBox::Slot()
+		.VAlign(VAlign_Center)
+		.AutoWidth()
+		[
+			SAssignNew(ComboDropdown, SComboButton)
+			.Visibility(this, &SCameraFamilyAssetShortcut::GetComboDropdownVisibility)
+			.ContentPadding(FMargin(7.f, 0.f))
+			.ForegroundColor(FSlateColor::UseForeground())
+			.ComboButtonStyle(&FAppStyle::Get(), "SegmentedCombo.Right")
+			.OnGetMenuContent(this, &SCameraFamilyAssetShortcut::HandleGetDropdownMenuContent)
+		]
+	];
+
+	EnableToolTipForceField(true);
+}
+
+SCameraFamilyAssetShortcut::~SCameraFamilyAssetShortcut()
+{
+	if (FModuleManager::Get().IsModuleLoaded(TEXT("AssetRegistry")))
+	{
+		IAssetRegistry* AssetRegistry = FModuleManager::GetModuleChecked<FAssetRegistryModule>(TEXT("AssetRegistry")).TryGet();
+		if (AssetRegistry)
+		{
+			AssetRegistry->OnFilesLoaded().RemoveAll(this);
+			AssetRegistry->OnAssetAdded().RemoveAll(this);
+			AssetRegistry->OnAssetRemoved().RemoveAll(this);
+			AssetRegistry->OnAssetRenamed().RemoveAll(this);
+		}
+	}
+
+	FGameplayCamerasDelegates::OnCameraAssetBuilt().RemoveAll(this);
+	FGameplayCamerasDelegates::OnCameraRigAssetBuilt().RemoveAll(this);
+}
+
+void SCameraFamilyAssetShortcut::Tick(const FGeometry& AllottedGeometry, const double InCurrentTime, const float InDeltaTime)
+{
+	SCompoundWidget::Tick(AllottedGeometry, InCurrentTime, InDeltaTime);
+
+	if (bRefreshAssetDatas)
+	{
+		bRefreshAssetDatas = false;
+
+		AssetDatas.Reset();
+		Family->FindAssetsOfType(FamilyAssetType, AssetDatas);
+	}
+}
+
+const FSlateBrush* SCameraFamilyAssetShortcut::GetAssetIcon() const 
+{
+	return Family->GetAssetIcon(FamilyAssetType);
+}
+
+FSlateColor SCameraFamilyAssetShortcut::GetAssetTint() const
+{
+	if (GetCheckState() == ECheckBoxState::Checked)
+	{
+		return FSlateColor::UseForeground();
+	}
+	return Family->GetAssetTint(FamilyAssetType);
+}
+
+ECheckBoxState SCameraFamilyAssetShortcut::GetCheckState() const
+{
+	if (TSharedPtr<FBaseAssetToolkit> Toolkit = WeakToolkit.Pin())
+	{
+		const TArray<UObject*>* Objects = Toolkit->GetObjectsCurrentlyBeingEdited();
+		if (Objects != nullptr && AssetDatas.Num() == 1)
+		{
+			const FAssetData& AssetData = AssetDatas[0];
+			for (UObject* Object : *Objects)
+			{
+				if (Object->GetPathName().Compare(AssetData.GetObjectPathString(), ESearchCase::IgnoreCase) == 0)
+				{
+					return ECheckBoxState::Checked;
+				}
+			}
+		}
+	}
+	return ECheckBoxState::Unchecked;
+}
+
+FText SCameraFamilyAssetShortcut::GetButtonTooltip() const
+{
+	if (AssetDatas.Num() > 1)
+	{
+		return FText::Format(
+				LOCTEXT("ShortcutComboTooltipFmt", "{0}\nAsset Types: {1}"),
+				Family->GetAssetTypeTooltip(FamilyAssetType),
+				FamilyAssetType->GetDisplayNameText());
+	}
+	else if (AssetDatas.Num() == 1)
+	{
+		return FText::Format(
+				LOCTEXT("ShortcutSoloTooltipFmt", "Open {0}\nAsset Type: {1}"),
+				FText::FromString(AssetDatas[0].GetFullName()),
+				FamilyAssetType->GetDisplayNameText());
+	}
+	else
+	{
+		return FText::Format(
+				LOCTEXT("ShortcutNoMatchTooltipFmt", "No related assets found\nAsset Type: {0}"),
+				FamilyAssetType->GetDisplayNameText());
+	}
+}
+
+void SCameraFamilyAssetShortcut::HandleButtonClick(ECheckBoxState InState)
+{
+	if (AssetDatas.Num() == 1)
+	{
+		const FAssetData& AssetData(AssetDatas[0]);
+		if (UObject* AssetObject = AssetData.GetAsset())
+		{
+			TArray<UObject*> Assets;
+			Assets.Add(AssetObject);
+			GEditor->GetEditorSubsystem<UAssetEditorSubsystem>()->OpenEditorForAssets(Assets);
+		}
+		else
+		{
+			UE_LOG(LogCameraSystemEditor, Error, TEXT("Asset cannot be opened: %s"), *AssetData.GetObjectPathString());
+		}
+	}
+	else if (AssetDatas.Num() > 1)
+	{
+		ComboDropdown->SetIsOpen(true);
+	}
+}
+
+EVisibility SCameraFamilyAssetShortcut::GetSoloButtonVisibility() const
+{
+	return AssetDatas.Num() <= 1 ? EVisibility::Visible : EVisibility::Collapsed;
+}
+
+bool SCameraFamilyAssetShortcut::IsSoloButtonEnabled() const
+{
+	return AssetDatas.Num() > 0;
+}
+
+EVisibility SCameraFamilyAssetShortcut::GetComboButtonVisibility() const
+{
+	return AssetDatas.Num() > 1 ? EVisibility::Visible : EVisibility::Collapsed;
+}
+
+EVisibility SCameraFamilyAssetShortcut::GetComboDropdownVisibility() const
+{
+	return AssetDatas.Num() > 1 ? EVisibility::Visible : EVisibility::Collapsed;
+}
+
+TSharedRef<SWidget> SCameraFamilyAssetShortcut::HandleGetDropdownMenuContent()
+{
+	const bool bInShouldCloseWindowAfterMenuSelection = true;
+	FMenuBuilder MenuBuilder(bInShouldCloseWindowAfterMenuSelection, nullptr);
+
+	if (AssetDatas.Num() > 1)
+	{
+		MenuBuilder.BeginSection("AssetSelection", LOCTEXT("AssetSelectionSection", "Select Asset"));
+		{
+			FAssetPickerConfig AssetPickerConfig;
+			AssetPickerConfig.Filter.ClassPaths.Add(FamilyAssetType->GetClassPathName());
+			for (const FAssetData& AssetData : AssetDatas)
+			{
+				AssetPickerConfig.Filter.SoftObjectPaths.Add(AssetData.GetSoftObjectPath());
+			}
+
+			AssetPickerConfig.SelectionMode = ESelectionMode::SingleToggle;
+			AssetPickerConfig.OnAssetSelected = FOnAssetSelected::CreateSP(this, &SCameraFamilyAssetShortcut::HandleOpenSecondaryAsset);
+			AssetPickerConfig.bAllowNullSelection = false;
+			AssetPickerConfig.ThumbnailLabel = EThumbnailLabel::ClassName;
+			AssetPickerConfig.InitialAssetViewType = EAssetViewType::List;
+
+			IContentBrowserSingleton& ContentBrowserSingleton = IContentBrowserSingleton::Get();
+
+			MenuBuilder.AddWidget(
+				SNew(SBox)
+				.WidthOverride(300.f)
+				.HeightOverride(600.f)
+				[
+					ContentBrowserSingleton.CreateAssetPicker(AssetPickerConfig)
+				],
+				FText(), true);
+		}
+		MenuBuilder.EndSection();
+	}
+	
+	return MenuBuilder.MakeWidget();
+}
+
+void SCameraFamilyAssetShortcut::HandleOpenSecondaryAsset(const FAssetData& InAssetData)
+{
+	if (InAssetData.IsValid())
+	{
+		FSlateApplication::Get().DismissAllMenus();
+
+		TArray<UObject*> Assets;
+		Assets.Add(InAssetData.GetAsset());
+		GEditor->GetEditorSubsystem<UAssetEditorSubsystem>()->OpenEditorForAssets(Assets);
+	}
+}
+
+bool SCameraFamilyAssetShortcut::DoesFamilySupport(const FAssetData& InAssetData)
+{
+	TArray<UClass*> AssetTypes;
+	Family->GetAssetTypes(AssetTypes);
+	return AssetTypes.ContainsByPredicate([&InAssetData](UClass* AssetType)
+				{
+					return InAssetData.AssetClassPath == AssetType->GetClassPathName();
+				});
+}
+
+void SCameraFamilyAssetShortcut::HandleFilesLoaded()
+{
+	bRefreshAssetDatas = true;
+}
+
+void SCameraFamilyAssetShortcut::HandleAssetRemoved(const FAssetData& InAssetData)
+{
+	if (DoesFamilySupport(InAssetData))
+	{
+		bRefreshAssetDatas = true;
+	}
+}
+
+void SCameraFamilyAssetShortcut::HandleAssetRenamed(const FAssetData& InAssetData, const FString& InOldObjectPath)
+{
+	if (DoesFamilySupport(InAssetData))
+	{
+		bRefreshAssetDatas = true;
+	}
+}
+
+void SCameraFamilyAssetShortcut::HandleAssetAdded(const FAssetData& InAssetData)
+{
+	const IAssetRegistry& AssetRegistry = FModuleManager::LoadModuleChecked<FAssetRegistryModule>("AssetRegistry").Get();
+	if (!AssetRegistry.IsLoadingAssets())
+	{
+		if (DoesFamilySupport(InAssetData))
+		{
+			bRefreshAssetDatas = true;
+		}
+	}
+}
+
+void SCameraFamilyAssetShortcut::HandleCameraAssetBuilt(const UCameraAsset* InCameraAsset)
+{
+	bRefreshAssetDatas = true;
+}
+
+void SCameraFamilyAssetShortcut::HandleCameraRigAssetBuilt(const UCameraRigAsset* InCameraAsset)
+{
+	bRefreshAssetDatas = true;
+}
+
+} // namespace UE::Cameras
+
+#undef LOCTEXT_NAMESPACE
+
diff --git a/Engine/Plugins/Cameras/GameplayCameras/Source/GameplayCamerasEditor/Private/Widgets/SCameraFamilyAssetShortcut.h b/Engine/Plugins/Cameras/GameplayCameras/Source/GameplayCamerasEditor/Private/Widgets/SCameraFamilyAssetShortcut.h
new file mode 100644
--- /dev/null
+++ b/Engine/Plugins/Cameras/GameplayCameras/Source/GameplayCamerasEditor/Private/Widgets/SCameraFamilyAssetShortcut.h
@@ -0,0 +1,77 @@
+// Copyright Epic Games, Inc. All Rights Reserved.
+
+#pragma once
+
+#include "AssetRegistry/AssetData.h"
+#include "Widgets/SCompoundWidget.h"
+
+class FAssetThumbnail;
+class FBaseAssetToolkit;
+class SCheckBox;
+class SComboButton;
+class UCameraAsset;
+class UCameraRigAsset;
+
+namespace UE::Cameras
+{
+
+class IGameplayCamerasFamily;
+
+class SCameraFamilyAssetShortcut : public SCompoundWidget
+{
+public:
+
+	SLATE_BEGIN_ARGS(SCameraFamilyAssetShortcut)
+	{}
+	SLATE_END_ARGS()
+
+	void Construct(const FArguments& InArgs, const TSharedRef<FBaseAssetToolkit>& InToolkit, const TSharedRef<IGameplayCamerasFamily>& InFamily, UClass* InAssetType);
+	~SCameraFamilyAssetShortcut();
+
+protected:
+
+	virtual void Tick(const FGeometry& AllottedGeometry, const double InCurrentTime, const float InDeltaTime) override;
+
+public:
+
+	const FSlateBrush* GetAssetIcon() const;
+	FSlateColor GetAssetTint() const;
+	ECheckBoxState GetCheckState() const;
+	FText GetButtonTooltip() const;
+	void HandleButtonClick(ECheckBoxState InState);
+
+	EVisibility GetSoloButtonVisibility() const;
+	bool IsSoloButtonEnabled() const;
+
+	EVisibility GetComboButtonVisibility() const;
+	EVisibility GetComboDropdownVisibility() const;
+	TSharedRef<SWidget> HandleGetDropdownMenuContent();
+	void HandleOpenSecondaryAsset(const FAssetData& InAssetData);
+
+	bool DoesFamilySupport(const FAssetData& InAssetData);
+
+	void HandleFilesLoaded();
+	void HandleAssetRemoved(const FAssetData& InAssetData);
+	void HandleAssetRenamed(const FAssetData& InAssetData, const FString& InOldObjectPath);
+	void HandleAssetAdded(const FAssetData& InAssetData);
+
+	void HandleCameraAssetBuilt(const UCameraAsset* InCameraAsset);
+	void HandleCameraRigAssetBuilt(const UCameraRigAsset* InCameraAsset);
+
+private:
+
+	TWeakPtr<FBaseAssetToolkit> WeakToolkit;
+	
+	TSharedPtr<IGameplayCamerasFamily> Family;
+	UClass* FamilyAssetType;
+
+	TSharedPtr<SCheckBox> SoloCheckBox;
+	TSharedPtr<SCheckBox> ComboCheckBox;
+	TSharedPtr<SComboButton> ComboDropdown;
+
+	TArray<FAssetData> AssetDatas;
+	bool bRefreshAssetDatas = false;
+};
+
+}  // namespace UE::Cameras
+
diff --git a/Engine/Plugins/Cameras/GameplayCameras/Source/GameplayCamerasEditor/Private/Widgets/SCameraFamilyShortcutBar.cpp b/Engine/Plugins/Cameras/GameplayCameras/Source/GameplayCamerasEditor/Private/Widgets/SCameraFamilyShortcutBar.cpp
new file mode 100644
--- /dev/null
+++ b/Engine/Plugins/Cameras/GameplayCameras/Source/GameplayCamerasEditor/Private/Widgets/SCameraFamilyShortcutBar.cpp
@@ -0,0 +1,57 @@
+// Copyright Epic Games, Inc. All Rights Reserved.
+
+#include "Widgets/SCameraFamilyShortcutBar.h"
+
+#include "AssetToolsModule.h"
+#include "ContentBrowserModule.h"
+#include "IAssetTools.h"
+#include "IGameplayCamerasFamily.h"
+#include "Modules/ModuleManager.h"
+#include "Subsystems/AssetEditorSubsystem.h"
+#include "Widgets/SBoxPanel.h"
+#include "Widgets/SCameraFamilyAssetShortcut.h"
+
+#define LOCTEXT_NAMESPACE "SCameraFamilyShortcutBar"
+
+namespace UE::Cameras
+{
+
+void SCameraFamilyShortcutBar::Construct(const FArguments& InArgs, const TSharedRef<FBaseAssetToolkit>& InToolkit, const TSharedRef<IGameplayCamerasFamily>& InFamily)
+{
+	WeakToolkit = InToolkit;
+
+	Family = InFamily;
+
+	HorizontalBox = SNew(SHorizontalBox);
+	{
+		BuildShortcuts();
+	}
+
+	ChildSlot
+	[
+		HorizontalBox.ToSharedRef()
+	];
+}
+
+void SCameraFamilyShortcutBar::BuildShortcuts()
+{
+	TArray<UClass*> AssetTypes;
+	Family->GetAssetTypes(AssetTypes);
+
+	TSharedRef<FBaseAssetToolkit> Toolkit = WeakToolkit.Pin().ToSharedRef();
+
+	for (UClass* Class : AssetTypes)
+	{
+		HorizontalBox->AddSlot()
+		.AutoWidth()
+		.Padding(0.0f, 4.0f, 16.0f, 4.0f)
+		[
+			SNew(SCameraFamilyAssetShortcut, Toolkit, Family.ToSharedRef(), Class)
+		];
+	}
+}
+
+}  // namespace UE::Cameras
+
+#undef LOCTEXT_NAMESPACE
+
diff --git a/Engine/Plugins/Cameras/GameplayCameras/Source/GameplayCamerasEditor/Public/IGameplayCamerasFamily.h b/Engine/Plugins/Cameras/GameplayCameras/Source/GameplayCamerasEditor/Public/IGameplayCamerasFamily.h
new file mode 100644
--- /dev/null
+++ b/Engine/Plugins/Cameras/GameplayCameras/Source/GameplayCamerasEditor/Public/IGameplayCamerasFamily.h
@@ -0,0 +1,38 @@
+// Copyright Epic Games, Inc. All Rights Reserved.
+
+#pragma once
+
+#include "AssetRegistry/AssetData.h"
+#include "Styling/SlateColor.h"
+#include "Templates/SharedPointer.h"
+
+struct FSlateBrush;
+
+namespace UE::Cameras
+{
+
+class IGameplayCamerasFamily : public TSharedFromThis<IGameplayCamerasFamily>
+{
+public:
+
+	virtual ~IGameplayCamerasFamily() {}
+
+	virtual UObject* GetRootAsset() const = 0;
+
+	virtual void GetAssetTypes(TArray<UClass*>& OutAssetTypes) const = 0;
+
+	virtual void FindAssetsOfType(UClass* InAssetType, TArray<FAssetData>& OutAssets) const = 0;
+
+	virtual FText GetAssetTypeTooltip(UClass* InAssetType) const = 0;
+
+	virtual const FSlateBrush* GetAssetIcon(UClass* InAssetType) const = 0;
+
+	virtual FSlateColor GetAssetTint(UClass* InAssetType) const = 0;
+
+public:
+
+	static TSharedPtr<IGameplayCamerasFamily> CreateFamily(UObject* InAsset);
+};
+
+}  // namespace UE::Cameras
+
diff --git a/Engine/Plugins/Cameras/GameplayCameras/Source/GameplayCamerasEditor/Public/Widgets/SCameraFamilyShortcutBar.h b/Engine/Plugins/Cameras/GameplayCameras/Source/GameplayCamerasEditor/Public/Widgets/SCameraFamilyShortcutBar.h
new file mode 100644
--- /dev/null
+++ b/Engine/Plugins/Cameras/GameplayCameras/Source/GameplayCamerasEditor/Public/Widgets/SCameraFamilyShortcutBar.h
@@ -0,0 +1,40 @@
+// Copyright Epic Games, Inc. All Rights Reserved.
+
+#pragma once
+
+#include "Widgets/DeclarativeSyntaxSupport.h"
+#include "Widgets/SCompoundWidget.h"
+
+class FBaseAssetToolkit;
+class SHorizontalBox;
+
+namespace UE::Cameras
+{
+
+class IGameplayCamerasFamily;
+
+class SCameraFamilyShortcutBar : public SCompoundWidget
+{
+public:
+
+	SLATE_BEGIN_ARGS(SCameraFamilyShortcutBar)
+	{}
+	SLATE_END_ARGS()
+
+	void Construct(const FArguments& InArgs, const TSharedRef<FBaseAssetToolkit>& InToolkit, const TSharedRef<IGameplayCamerasFamily>& InFamily);
+
+private:
+
+	void BuildShortcuts();
+
+private:
+	
+	TWeakPtr<FBaseAssetToolkit> WeakToolkit;
+
+	TSharedPtr<IGameplayCamerasFamily> Family;
+
+	TSharedPtr<SHorizontalBox> HorizontalBox;
+};
+
+}  // namespace UE::Cameras
+
