Cameras: new "initialize with ideal framing" option on screen-space framing nodes

diff --git a/Engine/Plugins/Cameras/GameplayCameras/Source/GameplayCameras/Private/Nodes/Framing/BaseFramingCameraNode.cpp b/Engine/Plugins/Cameras/GameplayCameras/Source/GameplayCameras/Private/Nodes/Framing/BaseFramingCameraNode.cpp
--- a/Engine/Plugins/Cameras/GameplayCameras/Source/GameplayCameras/Private/Nodes/Framing/BaseFramingCameraNode.cpp
+++ b/Engine/Plugins/Cameras/GameplayCameras/Source/GameplayCameras/Private/Nodes/Framing/BaseFramingCameraNode.cpp
@@ -65,6 +65,7 @@ void FBaseFramingCameraNodeEvaluator::OnInitialize(const FCameraNodeEvaluatorIni
 	Readers.TargetInfos.Initialize(BaseFramingNode->TargetInfos, BaseFramingNode->TargetInfosDataID);
 
 	Readers.IdealFramingLocation.Initialize(BaseFramingNode->IdealFramingLocation);
+	Readers.InitializeWithIdealFraming.Initialize(BaseFramingNode->InitializeWithIdealFraming);
 	Readers.SetTargetDistance.Initialize(BaseFramingNode->SetTargetDistance);
 
 	Readers.ReframeDampingFactor.Initialize(BaseFramingNode->ReframeDampingFactor);
@@ -77,17 +78,37 @@ void FBaseFramingCameraNodeEvaluator::OnInitialize(const FCameraNodeEvaluatorIni
 	Readers.SoftZone.Initialize(BaseFramingNode->SoftZone);
 
 	ScreenTargetHistory.History.Reserve(GFramingNumTargetMovementSamples);
+
+	TArray<FCameraActorComputedTargetInfo> AcquiredTargetInfos;
+	const bool bAcquireSuccess = AcquireTargetInfo(Params.EvaluationContext, OutResult, AcquiredTargetInfos);
+	WorldTargets.TargetInfos = AcquiredTargetInfos;
+}
+
+TOptional<FVector3d> FBaseFramingCameraNodeEvaluator::GetInitialDesiredWorldTarget(const FCameraNodeEvaluationParams& Params, const FCameraNodeEvaluationResult& OutResult)
+{
+	if (Params.bIsFirstFrame && 
+			Readers.InitializeWithIdealFraming.Get(OutResult.VariableTable) && 
+			!WorldTargets.TargetInfos.IsEmpty())
+	{
+		FVector3d ApproximatedWorldTarget(EForceInit::ForceInitToZero);
+		for (const FCameraActorComputedTargetInfo& TargetInfo : WorldTargets.TargetInfos)
+		{
+			ApproximatedWorldTarget += TargetInfo.Transform.GetLocation() * TargetInfo.NormalizedWeight;
+		}
+		return ApproximatedWorldTarget;
+	}
+	return TOptional<FVector3d>();
 }
 
 void FBaseFramingCameraNodeEvaluator::UpdateFramingState(const FCameraNodeEvaluationParams& Params, const FCameraNodeEvaluationResult& OutResult, const FTransform3d& LastFraming)
 {
 	TArray<FCameraActorComputedTargetInfo> AcquiredTargetInfos;
-	const bool bAcquireSuccess = AcquireTargetInfo(Params, OutResult, AcquiredTargetInfos);
+	const bool bAcquireSuccess = AcquireTargetInfo(Params.EvaluationContext, OutResult, AcquiredTargetInfos);
 	WorldTargets.TargetInfos = AcquiredTargetInfos;
 	if (bAcquireSuccess && AcquiredTargetInfos.Num() > 0)
 	{
 		ComputeCurrentState(Params, OutResult, LastFraming);
-		ComputeDesiredState(Params.DeltaTime);
+		ComputeDesiredState(Params, OutResult);
 	}
 }
 
@@ -101,7 +122,7 @@ void FBaseFramingCameraNodeEvaluator::EndFramingUpdate(const FCameraNodeEvaluati
 	}
 }
 
-bool FBaseFramingCameraNodeEvaluator::AcquireTargetInfo(const FCameraNodeEvaluationParams& Params, const FCameraNodeEvaluationResult& InResult, TArray<FCameraActorComputedTargetInfo>& OutInfos)
+bool FBaseFramingCameraNodeEvaluator::AcquireTargetInfo(TSharedPtr<const FCameraEvaluationContext> EvaluationContext, const FCameraNodeEvaluationResult& InResult, TArray<FCameraActorComputedTargetInfo>& OutInfos)
 {
 	const UBaseFramingCameraNode* FramingNode = GetCameraNodeAs<UBaseFramingCameraNode>();
 	if (FramingNode->TargetLocation.IsValid())
@@ -123,7 +144,7 @@ bool FBaseFramingCameraNodeEvaluator::AcquireTargetInfo(const FCameraNodeEvaluat
 			return true;
 		}
 	}
-	else if (APlayerController* PlayerController = Params.EvaluationContext->GetPlayerController())
+	else if (APlayerController* PlayerController = EvaluationContext->GetPlayerController())
 	{
 		APawn* Pawn = PlayerController->GetPawn();
 		FCameraActorComputedTargetInfo OutInfo;
@@ -482,7 +503,7 @@ FFramingZone FBaseFramingCameraNodeEvaluator::ComputeScreenTargetBounds(const FC
 	return FFramingZone::FromPoints(ScreenBoxCorners);
 }
 
-void FBaseFramingCameraNodeEvaluator::ComputeDesiredState(float DeltaTime)
+void FBaseFramingCameraNodeEvaluator::ComputeDesiredState(const FCameraNodeEvaluationParams& Params, const FCameraNodeEvaluationResult& OutResult)
 {
 	// If we  don't have any reframing to do, bail out.
 	FVector2d IdealToTarget(State.ScreenTarget - State.IdealTarget);
@@ -496,6 +517,16 @@ void FBaseFramingCameraNodeEvaluator::ComputeDesiredState(float DeltaTime)
 		return;
 	}
 
+	// We may need to jump directly to the ideal framing without interpolating.
+	bool bJumpToIdealFraming = false;
+	if (Params.bIsFirstFrame && Readers.InitializeWithIdealFraming.Get(OutResult.VariableTable))
+	{
+		Desired.ScreenTarget = State.IdealTarget;
+		Desired.FramingCorrection = Desired.ScreenTarget - State.ScreenTarget;
+		Desired.bHasCorrection = true;
+		return;
+	}
+
 	if (State.TargetFramingState == ETargetFramingState::InHardZone)
 	{
 		// Bring the target immediately to edge of the soft zone, in the direction of the 
@@ -519,7 +550,7 @@ void FBaseFramingCameraNodeEvaluator::ComputeDesiredState(float DeltaTime)
 	State.ReframeDamper.SetW0(EffectiveDampingFactor);
 
 	// Move the target towards the ideal framing using damping.
-	const double NewDistanceToGo = State.ReframeDamper.Update(DistanceToGo, DeltaTime);
+	const double NewDistanceToGo = State.ReframeDamper.Update(DistanceToGo, Params.DeltaTime);
 
 	// Compute where we want the target this frame.
 	const FVector2d InvReframeDir(IdealToTarget / DistanceToGo);
diff --git a/Engine/Plugins/Cameras/GameplayCameras/Source/GameplayCameras/Private/Nodes/Framing/DollyFramingCameraNode.cpp b/Engine/Plugins/Cameras/GameplayCameras/Source/GameplayCameras/Private/Nodes/Framing/DollyFramingCameraNode.cpp
--- a/Engine/Plugins/Cameras/GameplayCameras/Source/GameplayCameras/Private/Nodes/Framing/DollyFramingCameraNode.cpp
+++ b/Engine/Plugins/Cameras/GameplayCameras/Source/GameplayCameras/Private/Nodes/Framing/DollyFramingCameraNode.cpp
@@ -75,6 +75,19 @@ void FDollyFramingCameraNodeEvaluator::OnInitialize(const FCameraNodeEvaluatorIn
 
 void FDollyFramingCameraNodeEvaluator::OnRun(const FCameraNodeEvaluationParams& Params, FCameraNodeEvaluationResult& OutResult)
 {
+	// If this is the first frame, we may want to help frame the targets dead-on.
+	if (Params.bIsFirstFrame)
+	{
+		TOptional<FVector3d> ApproximatedWorldTarget = GetInitialDesiredWorldTarget(Params, OutResult);
+		if (ApproximatedWorldTarget.IsSet())
+		{
+			const FTransform3d InitialPose = OutResult.CameraPose.GetTransform();
+			const FVector3d InitialLocalTarget = InitialPose.InverseTransformPositionNoScale(ApproximatedWorldTarget.GetValue());
+			DollyPosition.X = InitialLocalTarget.Y;
+			DollyPosition.Y = InitialLocalTarget.Z;
+		}
+	}
+
 	// Let the base class figure out all the screen-space framing stuff.
 	const FTransform3d LastShotTransform = BuildDollyShotTransform(OutResult.CameraPose);
 	UpdateFramingState(Params, OutResult, LastShotTransform);
diff --git a/Engine/Plugins/Cameras/GameplayCameras/Source/GameplayCameras/Private/Nodes/Framing/PanningFramingCameraNode.cpp b/Engine/Plugins/Cameras/GameplayCameras/Source/GameplayCameras/Private/Nodes/Framing/PanningFramingCameraNode.cpp
--- a/Engine/Plugins/Cameras/GameplayCameras/Source/GameplayCameras/Private/Nodes/Framing/PanningFramingCameraNode.cpp
+++ b/Engine/Plugins/Cameras/GameplayCameras/Source/GameplayCameras/Private/Nodes/Framing/PanningFramingCameraNode.cpp
@@ -69,6 +69,21 @@ void FPanningFramingCameraNodeEvaluator::OnInitialize(const FCameraNodeEvaluator
 
 void FPanningFramingCameraNodeEvaluator::OnRun(const FCameraNodeEvaluationParams& Params, FCameraNodeEvaluationResult& OutResult)
 {
+	// If this is the first frame, we may want to help frame the targets dead-on.
+	if (Params.bIsFirstFrame)
+	{
+		TOptional<FVector3d> ApproximatedWorldTarget = GetInitialDesiredWorldTarget(Params, OutResult);
+		if (ApproximatedWorldTarget.IsSet())
+		{
+			const FVector3d InitialDesiredAim = ApproximatedWorldTarget.GetValue() - OutResult.CameraPose.GetLocation();
+			if (!InitialDesiredAim.IsNearlyZero())
+			{
+				const FVector3d InitialPoseAim = OutResult.CameraPose.GetAimDir();
+				PanningRotation = InitialDesiredAim.Rotation() - InitialPoseAim.Rotation();
+			}
+		}
+	}
+
 	// Let the base class figure out all the screen-space framing stuff.
 	const FTransform3d LastShotTransform = BuildPanningShotTransform(OutResult.CameraPose);
 	UpdateFramingState(Params, OutResult, LastShotTransform);
diff --git a/Engine/Plugins/Cameras/GameplayCameras/Source/GameplayCameras/Public/Nodes/Framing/BaseFramingCameraNode.h b/Engine/Plugins/Cameras/GameplayCameras/Source/GameplayCameras/Public/Nodes/Framing/BaseFramingCameraNode.h
--- a/Engine/Plugins/Cameras/GameplayCameras/Source/GameplayCameras/Public/Nodes/Framing/BaseFramingCameraNode.h
+++ b/Engine/Plugins/Cameras/GameplayCameras/Source/GameplayCameras/Public/Nodes/Framing/BaseFramingCameraNode.h
@@ -49,6 +49,12 @@ public:
 	UPROPERTY(EditAnywhere, Category="Target")
 	FBooleanCameraParameter SetTargetDistance;
 
+	/**
+	 * Whether to frame the target with the ideal framing immediately on the first frame.
+	 */
+	UPROPERTY(EditAnywhere, Category="Target")
+	FBooleanCameraParameter InitializeWithIdealFraming = true;
+
 	/** The ideal horizontal and vertical screen-space position of the target. */
 	UPROPERTY(EditAnywhere, Category="Framing Target")
 	FVector2dCameraParameter IdealFramingLocation;
@@ -168,6 +174,8 @@ protected:
 	struct FState;
 	struct FDesired;
 
+	/** The first frame aiming direction we need for proper initialization. */
+	TOptional<FVector3d> GetInitialDesiredWorldTarget(const FCameraNodeEvaluationParams& Params, const FCameraNodeEvaluationResult& OutResult);
 	/** Updates the framing state for the current tick, see State member field. */
 	void UpdateFramingState(const FCameraNodeEvaluationParams& Params, const FCameraNodeEvaluationResult& OutResult, const FTransform3d& LastFraming);
 	/** Wraps-up the update with optional operations. */
@@ -182,9 +190,9 @@ private:
 	FVector2d GetHardReframeCoords() const;
 
 	void ComputeCurrentState(const FCameraNodeEvaluationParams& Params, const FCameraNodeEvaluationResult& OutResult, const FTransform3d& LastFraming);
-	void ComputeDesiredState(float DeltaTime);
+	void ComputeDesiredState(const FCameraNodeEvaluationParams& Params, const FCameraNodeEvaluationResult& OutResult);
 
-	bool AcquireTargetInfo(const FCameraNodeEvaluationParams& Params, const FCameraNodeEvaluationResult& InResult, TArray<FCameraActorComputedTargetInfo>& OutInfos);
+	bool AcquireTargetInfo(TSharedPtr<const FCameraEvaluationContext> EvaluationContext, const FCameraNodeEvaluationResult& InResult, TArray<FCameraActorComputedTargetInfo>& OutInfos);
 	bool ComputeFinalTargetInfo(const FCameraNodeEvaluationParams& Params, const FCameraPose& CameraPose, FVector3d& OutWorldTarget, FVector2d& OutScreenTarget, FFramingZone& OutScreenBounds);
 	FVector2d ComputeAnticipatedScreenTarget(float DeltaTime, const FVector2d& InPreviousAnticipatedScreenTarget, const FVector2d& InScreenTarget);
 	FFramingZone ComputeEffectiveDeadZone();
@@ -219,6 +227,7 @@ protected:
 		FCameraActorTargetInfoArrayReader TargetInfos;
 
 		TCameraParameterReader<FVector2d> IdealFramingLocation;
+		TCameraParameterReader<bool> InitializeWithIdealFraming;
 		TCameraParameterReader<bool> SetTargetDistance;
 
 		TCameraParameterReader<float> ReframeDampingFactor;
