Cameras: don't re-activate a camera component on BeginPlay if it was already activated, re-add the "auto-build" feature in PIE

Also expose the "auto-build" feature as a project setting that can be turned off.

diff --git a/Engine/Plugins/Cameras/GameplayCameras/Source/GameplayCameras/Private/GameFramework/GameplayCameraComponentBase.cpp b/Engine/Plugins/Cameras/GameplayCameras/Source/GameplayCameras/Private/GameFramework/GameplayCameraComponentBase.cpp
--- a/Engine/Plugins/Cameras/GameplayCameras/Source/GameplayCameras/Private/GameFramework/GameplayCameraComponentBase.cpp
+++ b/Engine/Plugins/Cameras/GameplayCameras/Source/GameplayCameras/Private/GameFramework/GameplayCameraComponentBase.cpp
@@ -2,6 +2,8 @@
 
 #include "GameFramework/GameplayCameraComponentBase.h"
 
+#include "Build/CameraAssetBuilder.h"
+#include "Build/CameraBuildLog.h"
 #include "CineCameraComponent.h"
 #include "Core/CameraAsset.h"
 #include "Core/CameraSystemEvaluator.h"
@@ -11,6 +13,7 @@
 #include "Engine/EngineTypes.h"
 #include "Engine/World.h"
 #include "GameplayCamerasDelegates.h"
+#include "GameplayCamerasSettings.h"
 #include "IGameplayCamerasLiveEditManager.h"
 #include "IGameplayCamerasModule.h"
 #include "Kismet/GameplayStatics.h"
@@ -319,8 +322,25 @@ void UGameplayCameraComponentBase::TryCreateCameraEvaluationContext(APlayerContr
 
 	if (!EvaluationContext.IsValid())
 	{
-		if (const UCameraAsset* CameraAsset = GetCameraAsset())
+		if (UCameraAsset* CameraAsset = GetCameraAsset())
 		{
+#if WITH_EDITOR
+			if (bIsFirstActivation)
+			{
+				UWorld* World = GetWorld();
+				const UGameplayCamerasSettings* Settings = GetDefault<UGameplayCamerasSettings>();
+				if (Settings->bAutoBuildInPIE && World && World->WorldType == EWorldType::PIE)
+				{
+					// Auto-build the camera asset on begin play to make sure we've got the latest user edits.
+					FCameraBuildLog BuildLog;
+					FCameraAssetBuilder Builder(BuildLog);
+					Builder.BuildCamera(CameraAsset);
+				}
+
+				bIsFirstActivation = false;
+			}
+#endif
+
 			EvaluationContext = MakeShared<FGameplayCameraComponentEvaluationContext>();
 
 			FCameraEvaluationContextInitializeParams InitParams;
@@ -385,7 +405,13 @@ void UGameplayCameraComponentBase::BeginPlay()
 
 	Super::BeginPlay();
 
-	if (IsActive())
+	// If we have been activated in OnRegister() (which happens when bAutoActivate is true), our code 
+	// inside Activate() has postponed setting up the camera system evaluation until now, so let's
+	// do it.
+	// However, it can happen that some BP construction script already called ActivateCameraForXyz()
+	// before we got to start play (e.g. from a parent actor) and so in this case, let's skip
+	// re-activating for nothing.
+	if (IsActive() && !EvaluationContext)
 	{
 		if (AutoActivateForPlayer != EAutoReceiveInput::Disabled && GetNetMode() != NM_DedicatedServer)
 		{
diff --git a/Engine/Plugins/Cameras/GameplayCameras/Source/GameplayCameras/Public/GameFramework/GameplayCameraComponentBase.h b/Engine/Plugins/Cameras/GameplayCameras/Source/GameplayCameras/Public/GameFramework/GameplayCameraComponentBase.h
--- a/Engine/Plugins/Cameras/GameplayCameras/Source/GameplayCameras/Public/GameFramework/GameplayCameraComponentBase.h
+++ b/Engine/Plugins/Cameras/GameplayCameras/Source/GameplayCameras/Public/GameFramework/GameplayCameraComponentBase.h
@@ -226,6 +226,8 @@ private:
 	
 	bool bIsEditorWorld = false;
 
+	bool bIsFirstActivation = true;
+
 	int32 CustomShowFlag = INDEX_NONE;
 
 #endif  // WITH_EDITOR
diff --git a/Engine/Plugins/Cameras/GameplayCameras/Source/GameplayCameras/Public/GameplayCamerasSettings.h b/Engine/Plugins/Cameras/GameplayCameras/Source/GameplayCameras/Public/GameplayCamerasSettings.h
--- a/Engine/Plugins/Cameras/GameplayCameras/Source/GameplayCameras/Public/GameplayCamerasSettings.h
+++ b/Engine/Plugins/Cameras/GameplayCameras/Source/GameplayCameras/Public/GameplayCamerasSettings.h
@@ -23,6 +23,13 @@ public:
 
 public:
 
+	/**
+	 * Build camera assets before using them in PIE, so that they automatically run with the
+	 * latest user edits.
+	 */
+	UPROPERTY(EditAnywhere, Config, Category="General")
+	bool bAutoBuildInPIE = true;
+
 	/**
 	 * The number of camera rigs combined in one frame past which the camera system emits a warning.
 	 */
