Cameras: change the camera rig proxy table to be a struct

This simplifies object lifetimes and ownership, and makes the code more straightforward.
Remove the proxy table entry details-view customization since we haven't needed a custom camera rig picker in a while.
Make camera directors' GatherRigUsageInfo to be const, and add proxies to it.

diff --git a/Engine/Plugins/Cameras/GameplayCameras/Source/GameplayCameras/Private/Core/CameraAsset.cpp b/Engine/Plugins/Cameras/GameplayCameras/Source/GameplayCameras/Private/Core/CameraAsset.cpp
--- a/Engine/Plugins/Cameras/GameplayCameras/Source/GameplayCameras/Private/Core/CameraAsset.cpp
+++ b/Engine/Plugins/Cameras/GameplayCameras/Source/GameplayCameras/Private/Core/CameraAsset.cpp
@@ -6,6 +6,7 @@
 #include "Build/CameraBuildLog.h"
 #include "Core/CameraDirector.h"
 #include "Core/CameraRigAsset.h"
+#include "Core/CameraRigProxyAsset.h"
 #include "UObject/AssetRegistryTagsContext.h"
 #include "UObject/ObjectRedirector.h"
 #include "UObject/ObjectSaveContext.h"
@@ -112,22 +113,7 @@ void UCameraAsset::GetAssetRegistryTags(FAssetRegistryTagsContext Context) const
 {
 	if (CameraDirector)
 	{
-		FCameraDirectorRigUsageInfo UsageInfo;
-		CameraDirector->GatherRigUsageInfo(UsageInfo);
-
-		TStringBuilder<256> CameraRigListBuilder;
-		for (const UCameraRigAsset* CameraRig : UsageInfo.CameraRigs)
-		{
-			const UPackage* CameraRigPackage = CameraRig->GetPackage();
-			CameraRigListBuilder << TEXT(";");
-			CameraRigListBuilder << CameraRigPackage->GetName();
-		}
-
-		FAssetRegistryTag CameraRigsTag;
-		CameraRigsTag.Type = FAssetRegistryTag::ETagType::TT_Hidden;
-		CameraRigsTag.Name = TEXT("UsedCameraRigs");
-		CameraRigsTag.Value = CameraRigListBuilder.ToString();
-		Context.AddTag(CameraRigsTag);
+		CameraDirector->ExtendAssetRegistryTags(Context);
 	}
 
 	Super::GetAssetRegistryTags(Context);
@@ -135,11 +121,6 @@ void UCameraAsset::GetAssetRegistryTags(FAssetRegistryTagsContext Context) const
 
 #if WITH_EDITOR
 
-void UCameraAsset::GetAssetRegistryTagMetadata(TMap<FName, FAssetRegistryTagMetadata>& OutMetadata) const
-{
-	Super::GetAssetRegistryTagMetadata(OutMetadata);
-}
-
 void UCameraAsset::PostEditChangeProperty(FPropertyChangedEvent& PropertyChangedEvent)
 {
 	using namespace UE::Cameras;
diff --git a/Engine/Plugins/Cameras/GameplayCameras/Source/GameplayCameras/Private/Core/CameraDirector.cpp b/Engine/Plugins/Cameras/GameplayCameras/Source/GameplayCameras/Private/Core/CameraDirector.cpp
--- a/Engine/Plugins/Cameras/GameplayCameras/Source/GameplayCameras/Private/Core/CameraDirector.cpp
+++ b/Engine/Plugins/Cameras/GameplayCameras/Source/GameplayCameras/Private/Core/CameraDirector.cpp
@@ -2,7 +2,8 @@
 
 #include "Core/CameraDirector.h"
 
-#include "Core/CameraRigProxyTable.h"
+#include "Core/CameraRigAsset.h"
+#include "Core/CameraRigProxyAsset.h"
 
 #include UE_INLINE_GENERATED_CPP_BY_NAME(CameraDirector)
 
@@ -19,20 +20,63 @@ void UCameraDirector::BuildCameraDirector(UE::Cameras::FCameraBuildLog& BuildLog
 	OnBuildCameraDirector(BuildLog);
 }
 
-void UCameraDirector::GatherRigUsageInfo(FCameraDirectorRigUsageInfo& UsageInfo)
+void UCameraDirector::GatherRigUsageInfo(FCameraDirectorRigUsageInfo& UsageInfo) const
 {
 	OnGatherRigUsageInfo(UsageInfo);
 }
 
-#if WITH_EDITOR
-
-void UCameraDirector::FactoryCreateAsset(const FCameraDirectorFactoryCreateParams& InParams)
+void UCameraDirector::ExtendAssetRegistryTags(FAssetRegistryTagsContext Context) const
 {
-	if (bHasCameraRigProxyTable && !CameraRigProxyTable)
+	FCameraDirectorRigUsageInfo UsageInfo;
+	GatherRigUsageInfo(UsageInfo);
+
 	{
-		CameraRigProxyTable = NewObject<UCameraRigProxyTable>(this);
+		FAssetRegistryTag NumCameraRigsTag;
+		NumCameraRigsTag.Name = TEXT("NumUsedCameraRigs");
+		NumCameraRigsTag.Value = LexToString(UsageInfo.CameraRigs.Num());
+		Context.AddTag(NumCameraRigsTag);
+
+		TStringBuilder<256> CameraRigListBuilder;
+		for (const UCameraRigAsset* CameraRig : UsageInfo.CameraRigs)
+		{
+			const UPackage* CameraRigPackage = CameraRig->GetPackage();
+			CameraRigListBuilder << CameraRigPackage->GetName();
+			CameraRigListBuilder << TEXT("\n");
+		}
+
+		FAssetRegistryTag CameraRigsTag;
+		CameraRigsTag.Name = TEXT("UsedCameraRigs");
+		CameraRigsTag.Value = CameraRigListBuilder.ToString();
+		Context.AddTag(CameraRigsTag);
 	}
 
+	{
+		FAssetRegistryTag NumCameraRigProxiesTag;
+		NumCameraRigProxiesTag.Name = TEXT("NumUsedCameraRigProxies");
+		NumCameraRigProxiesTag.Value = LexToString(UsageInfo.CameraRigs.Num());
+		Context.AddTag(NumCameraRigProxiesTag);
+
+		TStringBuilder<256> CameraRigProxyListBuilder;
+		for (const UCameraRigProxyAsset* CameraRigProxy : UsageInfo.CameraRigProxies)
+		{
+			const UPackage* CameraRigPackage = CameraRigProxy->GetPackage();
+			CameraRigProxyListBuilder << CameraRigPackage->GetName();
+			CameraRigProxyListBuilder << TEXT("\n");
+		}
+
+		FAssetRegistryTag CameraRigProxiesTag;
+		CameraRigProxiesTag.Name = TEXT("UsedCameraRigProxies");
+		CameraRigProxiesTag.Value = CameraRigProxyListBuilder.ToString();
+		Context.AddTag(CameraRigProxiesTag);
+	}
+
+	OnExtendAssetRegistryTags(Context);
+}
+
+#if WITH_EDITOR
+
+void UCameraDirector::FactoryCreateAsset(const FCameraDirectorFactoryCreateParams& InParams)
+{
 	OnFactoryCreateAsset(InParams);
 }
 
diff --git a/Engine/Plugins/Cameras/GameplayCameras/Source/GameplayCameras/Private/Core/CameraDirectorEvaluator.cpp b/Engine/Plugins/Cameras/GameplayCameras/Source/GameplayCameras/Private/Core/CameraDirectorEvaluator.cpp
--- a/Engine/Plugins/Cameras/GameplayCameras/Source/GameplayCameras/Private/Core/CameraDirectorEvaluator.cpp
+++ b/Engine/Plugins/Cameras/GameplayCameras/Source/GameplayCameras/Private/Core/CameraDirectorEvaluator.cpp
@@ -75,15 +75,11 @@ const UCameraRigAsset* FCameraDirectorEvaluator::FindCameraRigByProxy(const UCam
 		return nullptr;
 	}
 
-	UCameraRigProxyTable* ProxyTable = PrivateCameraDirector->CameraRigProxyTable;
-	if (!ensureMsgf(ProxyTable, TEXT("Cannot resolve camera rig proxy: no proxy table set on camera director.")))
-	{
-		return nullptr;
-	}
+	const FCameraRigProxyTable& ProxyTable = PrivateCameraDirector->CameraRigProxyTable;
 
 	FCameraRigProxyTableResolveParams ResolveParams;
 	ResolveParams.CameraRigProxy = InProxy;
-	return ProxyTable->ResolveProxy(ResolveParams);
+	return ProxyTable.ResolveProxy(ResolveParams);
 }
 
 void FCameraDirectorEvaluator::OnEndCameraSystemUpdate()
diff --git a/Engine/Plugins/Cameras/GameplayCameras/Source/GameplayCameras/Private/Core/CameraRigProxyAsset.cpp b/Engine/Plugins/Cameras/GameplayCameras/Source/GameplayCameras/Private/Core/CameraRigProxyAsset.cpp
--- a/Engine/Plugins/Cameras/GameplayCameras/Source/GameplayCameras/Private/Core/CameraRigProxyAsset.cpp
+++ b/Engine/Plugins/Cameras/GameplayCameras/Source/GameplayCameras/Private/Core/CameraRigProxyAsset.cpp
@@ -4,6 +4,8 @@
 
 #include "Misc/Guid.h"
 
+#include UE_INLINE_GENERATED_CPP_BY_NAME(CameraRigProxyAsset)
+
 UCameraRigProxyAsset::UCameraRigProxyAsset(const FObjectInitializer& ObjectInit)
 	: Super(ObjectInit)
 {
diff --git a/Engine/Plugins/Cameras/GameplayCameras/Source/GameplayCameras/Private/Core/CameraRigProxyTable.cpp b/Engine/Plugins/Cameras/GameplayCameras/Source/GameplayCameras/Private/Core/CameraRigProxyTable.cpp
--- a/Engine/Plugins/Cameras/GameplayCameras/Source/GameplayCameras/Private/Core/CameraRigProxyTable.cpp
+++ b/Engine/Plugins/Cameras/GameplayCameras/Source/GameplayCameras/Private/Core/CameraRigProxyTable.cpp
@@ -2,12 +2,9 @@
 
 #include "Core/CameraRigProxyTable.h"
 
-UCameraRigProxyTable::UCameraRigProxyTable(const FObjectInitializer& ObjectInit)
-	: Super(ObjectInit)
-{
-}
+#include UE_INLINE_GENERATED_CPP_BY_NAME(CameraRigProxyTable)
 
-UCameraRigAsset* UCameraRigProxyTable::ResolveProxy(const FCameraRigProxyTableResolveParams& InParams) const
+UCameraRigAsset* FCameraRigProxyTable::ResolveProxy(const FCameraRigProxyTableResolveParams& InParams) const
 {
 	ensure(InParams.CameraRigProxy);
 
diff --git a/Engine/Plugins/Cameras/GameplayCameras/Source/GameplayCameras/Private/Directors/BlueprintCameraDirector.cpp b/Engine/Plugins/Cameras/GameplayCameras/Source/GameplayCameras/Private/Directors/BlueprintCameraDirector.cpp
--- a/Engine/Plugins/Cameras/GameplayCameras/Source/GameplayCameras/Private/Directors/BlueprintCameraDirector.cpp
+++ b/Engine/Plugins/Cameras/GameplayCameras/Source/GameplayCameras/Private/Directors/BlueprintCameraDirector.cpp
@@ -456,7 +456,7 @@ void UBlueprintCameraDirector::OnBuildCameraDirector(UE::Cameras::FCameraBuildLo
 	}
 }
 
-void UBlueprintCameraDirector::OnGatherRigUsageInfo(FCameraDirectorRigUsageInfo& UsageInfo)
+void UBlueprintCameraDirector::OnGatherRigUsageInfo(FCameraDirectorRigUsageInfo& UsageInfo) const
 {
 	using namespace UE::Cameras;
 
@@ -476,6 +476,7 @@ void UBlueprintCameraDirector::OnGatherRigUsageInfo(FCameraDirectorRigUsageInfo&
 	FOutgoingReferenceFinder ReferenceFinder(EvaluatorBlueprint, UCameraRigAsset::StaticClass());
 	ReferenceFinder.CollectReferences();
 	ReferenceFinder.GetReferencesOfClass<UCameraRigAsset>(UsageInfo.CameraRigs);
+	ReferenceFinder.GetReferencesOfClass<UCameraRigProxyAsset>(UsageInfo.CameraRigProxies);
 
 #endif  // WITH_EDITORONLY_DATA
 }
diff --git a/Engine/Plugins/Cameras/GameplayCameras/Source/GameplayCameras/Private/Directors/SingleCameraDirector.cpp b/Engine/Plugins/Cameras/GameplayCameras/Source/GameplayCameras/Private/Directors/SingleCameraDirector.cpp
--- a/Engine/Plugins/Cameras/GameplayCameras/Source/GameplayCameras/Private/Directors/SingleCameraDirector.cpp
+++ b/Engine/Plugins/Cameras/GameplayCameras/Source/GameplayCameras/Private/Directors/SingleCameraDirector.cpp
@@ -34,7 +34,6 @@ UE_DEFINE_CAMERA_DIRECTOR_EVALUATOR(FSingleCameraDirectorEvaluator)
 USingleCameraDirector::USingleCameraDirector(const FObjectInitializer& ObjectInit)
 	: Super(ObjectInit)
 {
-	bHasCameraRigProxyTable = false;
 }
 
 FCameraDirectorEvaluatorPtr USingleCameraDirector::OnBuildEvaluator(FCameraDirectorEvaluatorBuilder& Builder) const
@@ -51,7 +50,7 @@ void USingleCameraDirector::OnBuildCameraDirector(UE::Cameras::FCameraBuildLog&
 	}
 }
 
-void USingleCameraDirector::OnGatherRigUsageInfo(FCameraDirectorRigUsageInfo& UsageInfo)
+void USingleCameraDirector::OnGatherRigUsageInfo(FCameraDirectorRigUsageInfo& UsageInfo) const
 {
 	if (CameraRig)
 	{
diff --git a/Engine/Plugins/Cameras/GameplayCameras/Source/GameplayCameras/Private/Directors/StateTreeCameraDirector.cpp b/Engine/Plugins/Cameras/GameplayCameras/Source/GameplayCameras/Private/Directors/StateTreeCameraDirector.cpp
--- a/Engine/Plugins/Cameras/GameplayCameras/Source/GameplayCameras/Private/Directors/StateTreeCameraDirector.cpp
+++ b/Engine/Plugins/Cameras/GameplayCameras/Source/GameplayCameras/Private/Directors/StateTreeCameraDirector.cpp
@@ -6,6 +6,7 @@
 #include "Core/CameraAsset.h"
 #include "Core/CameraEvaluationContext.h"
 #include "Core/CameraRigAsset.h"
+#include "Core/CameraRigProxyAsset.h"
 #include "Directors/CameraDirectorStateTreeSchema.h"
 #include "GameplayCameras.h"
 #include "Helpers/OutgoingReferenceFinder.h"
@@ -232,19 +233,20 @@ void UStateTreeCameraDirector::OnBuildCameraDirector(UE::Cameras::FCameraBuildLo
 	}
 }
 
-void UStateTreeCameraDirector::OnGatherRigUsageInfo(FCameraDirectorRigUsageInfo& UsageInfo)
+void UStateTreeCameraDirector::OnGatherRigUsageInfo(FCameraDirectorRigUsageInfo& UsageInfo) const
 {
 	using namespace UE::Cameras;
 
-	UStateTree* StateTree = StateTreeReference.GetMutableStateTree();
+	const UStateTree* StateTree = StateTreeReference.GetStateTree();
 	if (!StateTree)
 	{
 		return;
 	}
 
-	FOutgoingReferenceFinder ReferenceFinder(StateTree, UCameraRigAsset::StaticClass());
+	FOutgoingReferenceFinder ReferenceFinder(const_cast<UStateTree*>(StateTree), UCameraRigAsset::StaticClass());
 	ReferenceFinder.CollectReferences();
 	ReferenceFinder.GetReferencesOfClass<UCameraRigAsset>(UsageInfo.CameraRigs);
+	ReferenceFinder.GetReferencesOfClass<UCameraRigProxyAsset>(UsageInfo.CameraRigProxies);
 }
 
 #undef LOCTEXT_NAMESPACE
diff --git a/Engine/Plugins/Cameras/GameplayCameras/Source/GameplayCameras/Public/Core/CameraAsset.h b/Engine/Plugins/Cameras/GameplayCameras/Source/GameplayCameras/Public/Core/CameraAsset.h
--- a/Engine/Plugins/Cameras/GameplayCameras/Source/GameplayCameras/Public/Core/CameraAsset.h
+++ b/Engine/Plugins/Cameras/GameplayCameras/Source/GameplayCameras/Public/Core/CameraAsset.h
@@ -157,7 +157,6 @@ protected:
 	virtual void PostLoad() override;
 	virtual void GetAssetRegistryTags(FAssetRegistryTagsContext Context) const override;
 #if WITH_EDITOR
-	virtual void GetAssetRegistryTagMetadata(TMap<FName, FAssetRegistryTagMetadata>& OutMetadata) const override;
 	virtual void PostEditChangeProperty(FPropertyChangedEvent& PropertyChangedEvent) override;
 #endif
 
diff --git a/Engine/Plugins/Cameras/GameplayCameras/Source/GameplayCameras/Public/Core/CameraDirector.h b/Engine/Plugins/Cameras/GameplayCameras/Source/GameplayCameras/Public/Core/CameraDirector.h
--- a/Engine/Plugins/Cameras/GameplayCameras/Source/GameplayCameras/Public/Core/CameraDirector.h
+++ b/Engine/Plugins/Cameras/GameplayCameras/Source/GameplayCameras/Public/Core/CameraDirector.h
@@ -5,12 +5,15 @@
 #include "CoreTypes.h"
 #include "UObject/Object.h"
 #include "Core/CameraDirectorEvaluator.h"
+#include "Core/CameraRigProxyTable.h"
+#include "UObject/AssetRegistryTagsContext.h"
 
 #include "CameraDirector.generated.h"
 
 #define UE_API GAMEPLAYCAMERAS_API
 
-class UCameraRigProxyTable;
+class UCameraRigAsset;
+class UCameraRigProxy;
 
 namespace UE::Cameras { class FCameraBuildLog; }
 
@@ -31,7 +34,10 @@ struct FCameraDirectorFactoryCreateParams
  */
 struct FCameraDirectorRigUsageInfo
 {
+	/** Camera rigs used by the camera director. */
 	TArray<UCameraRigAsset*> CameraRigs;
+	/** Camera rig proxies used by the camera director. */
+	TArray<UCameraRigProxyAsset*> CameraRigProxies;
 };
 
 /**
@@ -53,7 +59,10 @@ public:
 	UE_API void BuildCameraDirector(UE::Cameras::FCameraBuildLog& BuildLog);
 
 	/** Gets the list of camera rigs used by this camera director. */
-	UE_API void GatherRigUsageInfo(FCameraDirectorRigUsageInfo& UsageInfo);
+	UE_API void GatherRigUsageInfo(FCameraDirectorRigUsageInfo& UsageInfo) const;
+
+	/** Extend the owning camera asset's tags. */
+	UE_API void ExtendAssetRegistryTags(FAssetRegistryTagsContext Context) const;
 
 #if WITH_EDITOR
 	/** Called by the asset factories to setup new data before the editor opens. */
@@ -69,7 +78,10 @@ protected:
 	virtual void OnBuildCameraDirector(UE::Cameras::FCameraBuildLog& BuildLog) {}
 
 	/** Gets the list of camera rigs used by this camera director. */
-	virtual void OnGatherRigUsageInfo(FCameraDirectorRigUsageInfo& UsageInfo) {}
+	virtual void OnGatherRigUsageInfo(FCameraDirectorRigUsageInfo& UsageInfo) const {}
+
+	/** Extend the owning camera asset's tags. */
+	virtual void OnExtendAssetRegistryTags(FAssetRegistryTagsContext Context) const {}
 
 #if WITH_EDITOR
 	/** Called by the asset factories to setup new data before the editor opens. */
@@ -82,14 +94,8 @@ public:
 	 * The table that maps camera rig proxies (used in the evaluator Blueprint graph)
 	 * to actual camera rigs.
 	 */
-	UPROPERTY(EditAnywhere, Instanced, Category="Evaluation")
-	TObjectPtr<UCameraRigProxyTable> CameraRigProxyTable;
-
-protected:
-
-	/** Whether this type of director supports having a camera rig proxy table. */
-	UPROPERTY(Transient)
-	bool bHasCameraRigProxyTable = true;
+	UPROPERTY(EditAnywhere, Category="Evaluation")
+	FCameraRigProxyTable CameraRigProxyTable;
 };
 
 #undef UE_API
diff --git a/Engine/Plugins/Cameras/GameplayCameras/Source/GameplayCameras/Public/Core/CameraRigProxyTable.h b/Engine/Plugins/Cameras/GameplayCameras/Source/GameplayCameras/Public/Core/CameraRigProxyTable.h
--- a/Engine/Plugins/Cameras/GameplayCameras/Source/GameplayCameras/Public/Core/CameraRigProxyTable.h
+++ b/Engine/Plugins/Cameras/GameplayCameras/Source/GameplayCameras/Public/Core/CameraRigProxyTable.h
@@ -39,15 +39,13 @@ struct FCameraRigProxyTableEntry
 /**
  * A table that defines mappings between camera rig proxies and actual camera rigs.
  */
-UCLASS(MinimalAPI, EditInlineNew)
-class UCameraRigProxyTable : public UObject
+USTRUCT()
+struct FCameraRigProxyTable
 {
 	GENERATED_BODY()
 
 public:
 
-	UCameraRigProxyTable(const FObjectInitializer& ObjectInit);
-
 	/**
 	 * Resolves a given proxy to an actual camera rig.
 	 * Returns nullptr if the given proxy wasn't found, or not mapped to anything in the table.
diff --git a/Engine/Plugins/Cameras/GameplayCameras/Source/GameplayCameras/Public/Directors/BlueprintCameraDirector.h b/Engine/Plugins/Cameras/GameplayCameras/Source/GameplayCameras/Public/Directors/BlueprintCameraDirector.h
--- a/Engine/Plugins/Cameras/GameplayCameras/Source/GameplayCameras/Public/Directors/BlueprintCameraDirector.h
+++ b/Engine/Plugins/Cameras/GameplayCameras/Source/GameplayCameras/Public/Directors/BlueprintCameraDirector.h
@@ -181,6 +181,6 @@ protected:
 	// UCameraDirector interface.
 	virtual FCameraDirectorEvaluatorPtr OnBuildEvaluator(FCameraDirectorEvaluatorBuilder& Builder) const override;
 	virtual void OnBuildCameraDirector(UE::Cameras::FCameraBuildLog& BuildLog) override;
-	virtual void OnGatherRigUsageInfo(FCameraDirectorRigUsageInfo& UsageInfo) override;
+	virtual void OnGatherRigUsageInfo(FCameraDirectorRigUsageInfo& UsageInfo) const override;
 };
 
diff --git a/Engine/Plugins/Cameras/GameplayCameras/Source/GameplayCameras/Public/Directors/SingleCameraDirector.h b/Engine/Plugins/Cameras/GameplayCameras/Source/GameplayCameras/Public/Directors/SingleCameraDirector.h
--- a/Engine/Plugins/Cameras/GameplayCameras/Source/GameplayCameras/Public/Directors/SingleCameraDirector.h
+++ b/Engine/Plugins/Cameras/GameplayCameras/Source/GameplayCameras/Public/Directors/SingleCameraDirector.h
@@ -23,7 +23,7 @@ protected:
 	// UCameraDirector interface.
 	virtual FCameraDirectorEvaluatorPtr OnBuildEvaluator(FCameraDirectorEvaluatorBuilder& Builder) const override;
 	virtual void OnBuildCameraDirector(UE::Cameras::FCameraBuildLog& BuildLog) override;
-	virtual void OnGatherRigUsageInfo(FCameraDirectorRigUsageInfo& UsageInfo) override;
+	virtual void OnGatherRigUsageInfo(FCameraDirectorRigUsageInfo& UsageInfo) const override;
 
 public:
 
diff --git a/Engine/Plugins/Cameras/GameplayCameras/Source/GameplayCameras/Public/Directors/StateTreeCameraDirector.h b/Engine/Plugins/Cameras/GameplayCameras/Source/GameplayCameras/Public/Directors/StateTreeCameraDirector.h
--- a/Engine/Plugins/Cameras/GameplayCameras/Source/GameplayCameras/Public/Directors/StateTreeCameraDirector.h
+++ b/Engine/Plugins/Cameras/GameplayCameras/Source/GameplayCameras/Public/Directors/StateTreeCameraDirector.h
@@ -25,7 +25,7 @@ protected:
 	// UCameraDirector interface.
 	virtual FCameraDirectorEvaluatorPtr OnBuildEvaluator(FCameraDirectorEvaluatorBuilder& Builder) const override;
 	virtual void OnBuildCameraDirector(UE::Cameras::FCameraBuildLog& BuildLog) override;
-	virtual void OnGatherRigUsageInfo(FCameraDirectorRigUsageInfo& UsageInfo) override;
+	virtual void OnGatherRigUsageInfo(FCameraDirectorRigUsageInfo& UsageInfo) const override;
 
 public:
 
diff --git a/Engine/Plugins/Cameras/GameplayCameras/Source/GameplayCamerasEditor/Private/Customizations/CameraProxyTableDetailsCustomization.cpp b/Engine/Plugins/Cameras/GameplayCameras/Source/GameplayCamerasEditor/Private/Customizations/CameraProxyTableDetailsCustomization.cpp
deleted file mode 100644
diff --git a/Engine/Plugins/Cameras/GameplayCameras/Source/GameplayCamerasEditor/Private/Family/CameraAssetFamily.cpp b/Engine/Plugins/Cameras/GameplayCameras/Source/GameplayCamerasEditor/Private/Family/CameraAssetFamily.cpp
--- a/Engine/Plugins/Cameras/GameplayCameras/Source/GameplayCamerasEditor/Private/Family/CameraAssetFamily.cpp
+++ b/Engine/Plugins/Cameras/GameplayCameras/Source/GameplayCamerasEditor/Private/Family/CameraAssetFamily.cpp
@@ -62,6 +62,10 @@ void FCameraAssetFamily::FindAssetsOfType(UClass* InAssetType, TArray<FAssetData
 	}
 	else if (InAssetType == UCameraRigProxyAsset::StaticClass())
 	{
+		for (UCameraRigProxyAsset* CameraRigProxy : UsageInfo.CameraRigProxies)
+		{
+			OutAssets.Add(FAssetData(CameraRigProxy));
+		}
 	}
 }
 
diff --git a/Engine/Plugins/Cameras/GameplayCameras/Source/GameplayCamerasEditor/Private/GameplayCamerasEditorModule.cpp b/Engine/Plugins/Cameras/GameplayCameras/Source/GameplayCamerasEditor/Private/GameplayCamerasEditorModule.cpp
--- a/Engine/Plugins/Cameras/GameplayCameras/Source/GameplayCamerasEditor/Private/GameplayCamerasEditorModule.cpp
+++ b/Engine/Plugins/Cameras/GameplayCameras/Source/GameplayCamerasEditor/Private/GameplayCamerasEditorModule.cpp
@@ -17,7 +17,6 @@
 #include "ComponentVisualizers/GameplayCameraComponentVisualizer.h"
 #include "Customizations/CameraAssetReferenceDetailsCustomization.h"
 #include "Customizations/CameraParameterDetailsCustomizations.h"
-#include "Customizations/CameraProxyTableDetailsCustomization.h"
 #include "Customizations/CameraRigAssetReferenceDetailsCustomization.h"
 #include "Customizations/CameraShakeAssetReferenceDetailsCustomization.h"
 #include "Customizations/CameraVariableReferenceDetailsCustomizations.h"
@@ -37,7 +36,6 @@
 #include "Features/IModularFeatures.h"
 #include "GameFramework/GameplayCameraComponent.h"
 #include "GameplayCameras.h"
-#include "GameplayCamerasEditorSettings.h"
 #include "GameplayCamerasLiveEditManager.h"
 #include "IGameplayCamerasEditorModule.h"
 #include "IGameplayCamerasModule.h"
@@ -415,10 +413,6 @@ private:
 				"CameraAssetReference",
 				FOnGetPropertyTypeCustomizationInstance::CreateStatic(
 					&FCameraAssetReferenceDetailsCustomization::MakeInstance));
-		PropertyEditorModule.RegisterCustomPropertyTypeLayout(
-				"CameraRigProxyTableEntry", 
-				FOnGetPropertyTypeCustomizationInstance::CreateStatic(
-					&FCameraProxyTableEntryDetailsCustomization::MakeInstance));
 		PropertyEditorModule.RegisterCustomPropertyTypeLayout(
 				"CameraRigAssetReference", 
 				FOnGetPropertyTypeCustomizationInstance::CreateStatic(
@@ -446,7 +440,6 @@ private:
 			FCameraVariableReferenceDetailsCustomization::Unregister(*PropertyEditorModule);
 			FRichCurveDetailsCustomization::Unregister(*PropertyEditorModule);
 
-			PropertyEditorModule->UnregisterCustomPropertyTypeLayout("CameraRigProxyTableEntry");
 			PropertyEditorModule->UnregisterCustomPropertyTypeLayout("CameraRigAsset");
 			PropertyEditorModule->UnregisterCustomPropertyTypeLayout("CameraRigAssetReference");
 			PropertyEditorModule->UnregisterCustomPropertyTypeLayout("CameraShakeAssetReference");
diff --git a/Engine/Plugins/Cameras/GameplayCameras/Source/GameplayCamerasEditor/Private/Toolkits/BlueprintCameraDirectorAssetEditorMode.cpp b/Engine/Plugins/Cameras/GameplayCameras/Source/GameplayCamerasEditor/Private/Toolkits/BlueprintCameraDirectorAssetEditorMode.cpp
--- a/Engine/Plugins/Cameras/GameplayCameras/Source/GameplayCamerasEditor/Private/Toolkits/BlueprintCameraDirectorAssetEditorMode.cpp
+++ b/Engine/Plugins/Cameras/GameplayCameras/Source/GameplayCamerasEditor/Private/Toolkits/BlueprintCameraDirectorAssetEditorMode.cpp
@@ -2,7 +2,6 @@
 
 #include "Toolkits/BlueprintCameraDirectorAssetEditorMode.h"
 
-#include "Core/CameraRigProxyTable.h"
 #include "Directors/BlueprintCameraDirector.h"
 
 namespace UE::Cameras
diff --git a/Engine/Plugins/Cameras/GameplayCameras/Source/GameplayCamerasEditor/Public/Customizations/CameraProxyTableDetailsCustomization.h b/Engine/Plugins/Cameras/GameplayCameras/Source/GameplayCamerasEditor/Public/Customizations/CameraProxyTableDetailsCustomization.h
deleted file mode 100644
