Cameras: add camera rig proxies and directors to the "family shortcuts" widget

The camera director shortcuts points not to the actual camera director (which is inside the camera asset) but to any "external" director logic, such as a Blueprint or StateTree evaluator.
Since we don't know which directors may have an external asset, and what that asset is, we make the directors save the path to that asset (if any) inside a "well known" registry tag on the camera asset.

diff --git a/Engine/Plugins/Cameras/GameplayCameras/Source/GameplayCameras/Private/Directors/BlueprintCameraDirector.cpp b/Engine/Plugins/Cameras/GameplayCameras/Source/GameplayCameras/Private/Directors/BlueprintCameraDirector.cpp
--- a/Engine/Plugins/Cameras/GameplayCameras/Source/GameplayCameras/Private/Directors/BlueprintCameraDirector.cpp
+++ b/Engine/Plugins/Cameras/GameplayCameras/Source/GameplayCameras/Private/Directors/BlueprintCameraDirector.cpp
@@ -12,6 +12,7 @@
 #include "Core/CameraSystemEvaluator.h"
 #include "Core/RootCameraNode.h"
 #include "Engine/Blueprint.h"
+#include "Engine/BlueprintGeneratedClass.h"
 #include "Engine/World.h"
 #include "GameFramework/Actor.h"
 #include "GameFramework/ControllerGameplayCameraEvaluationComponent.h"
@@ -473,7 +474,8 @@ void UBlueprintCameraDirector::OnGatherRigUsageInfo(FCameraDirectorRigUsageInfo&
 		return;
 	}
 
-	FOutgoingReferenceFinder ReferenceFinder(EvaluatorBlueprint, UCameraRigAsset::StaticClass());
+	TArray<UClass*> RefClasses { UCameraRigAsset::StaticClass(), UCameraRigProxyAsset::StaticClass() };
+	FOutgoingReferenceFinder ReferenceFinder(EvaluatorBlueprint, RefClasses);
 	ReferenceFinder.CollectReferences();
 	ReferenceFinder.GetReferencesOfClass<UCameraRigAsset>(UsageInfo.CameraRigs);
 	ReferenceFinder.GetReferencesOfClass<UCameraRigProxyAsset>(UsageInfo.CameraRigProxies);
@@ -481,5 +483,19 @@ void UBlueprintCameraDirector::OnGatherRigUsageInfo(FCameraDirectorRigUsageInfo&
 #endif  // WITH_EDITORONLY_DATA
 }
 
+void UBlueprintCameraDirector::OnExtendAssetRegistryTags(FAssetRegistryTagsContext Context) const
+{
+#if WITH_EDITORONLY_DATA
+	if (CameraDirectorEvaluatorClass && CameraDirectorEvaluatorClass->ClassGeneratedBy)
+	{
+		FAssetRegistryTag ExternalDirectorTag;
+		ExternalDirectorTag.Type = FAssetRegistryTag::ETagType::TT_Hidden;
+		ExternalDirectorTag.Name = TEXT("ExternalDirector");
+		ExternalDirectorTag.Value = CameraDirectorEvaluatorClass->ClassGeneratedBy->GetPathName();
+		Context.AddTag(ExternalDirectorTag);
+	}
+#endif  // WITH_EDITORONLY_DATA
+}
+
 #undef LOCTEXT_NAMESPACE
 
diff --git a/Engine/Plugins/Cameras/GameplayCameras/Source/GameplayCameras/Private/Directors/StateTreeCameraDirector.cpp b/Engine/Plugins/Cameras/GameplayCameras/Source/GameplayCameras/Private/Directors/StateTreeCameraDirector.cpp
--- a/Engine/Plugins/Cameras/GameplayCameras/Source/GameplayCameras/Private/Directors/StateTreeCameraDirector.cpp
+++ b/Engine/Plugins/Cameras/GameplayCameras/Source/GameplayCameras/Private/Directors/StateTreeCameraDirector.cpp
@@ -243,11 +243,24 @@ void UStateTreeCameraDirector::OnGatherRigUsageInfo(FCameraDirectorRigUsageInfo&
 		return;
 	}
 
-	FOutgoingReferenceFinder ReferenceFinder(const_cast<UStateTree*>(StateTree), UCameraRigAsset::StaticClass());
+	TArray<UClass*> RefClasses { UCameraRigAsset::StaticClass(), UCameraRigProxyAsset::StaticClass() };
+	FOutgoingReferenceFinder ReferenceFinder(const_cast<UStateTree*>(StateTree), RefClasses);
 	ReferenceFinder.CollectReferences();
 	ReferenceFinder.GetReferencesOfClass<UCameraRigAsset>(UsageInfo.CameraRigs);
 	ReferenceFinder.GetReferencesOfClass<UCameraRigProxyAsset>(UsageInfo.CameraRigProxies);
 }
 
+void UStateTreeCameraDirector::OnExtendAssetRegistryTags(FAssetRegistryTagsContext Context) const
+{
+	if (const UStateTree* StateTree = StateTreeReference.GetStateTree())
+	{
+		FAssetRegistryTag ExternalDirectorTag;
+		ExternalDirectorTag.Type = FAssetRegistryTag::ETagType::TT_Hidden;
+		ExternalDirectorTag.Name = TEXT("ExternalDirector");
+		ExternalDirectorTag.Value = StateTree->GetPathName();
+		Context.AddTag(ExternalDirectorTag);
+	}
+}
+
 #undef LOCTEXT_NAMESPACE
 
diff --git a/Engine/Plugins/Cameras/GameplayCameras/Source/GameplayCameras/Public/Directors/BlueprintCameraDirector.h b/Engine/Plugins/Cameras/GameplayCameras/Source/GameplayCameras/Public/Directors/BlueprintCameraDirector.h
--- a/Engine/Plugins/Cameras/GameplayCameras/Source/GameplayCameras/Public/Directors/BlueprintCameraDirector.h
+++ b/Engine/Plugins/Cameras/GameplayCameras/Source/GameplayCameras/Public/Directors/BlueprintCameraDirector.h
@@ -182,5 +182,6 @@ protected:
 	virtual FCameraDirectorEvaluatorPtr OnBuildEvaluator(FCameraDirectorEvaluatorBuilder& Builder) const override;
 	virtual void OnBuildCameraDirector(UE::Cameras::FCameraBuildLog& BuildLog) override;
 	virtual void OnGatherRigUsageInfo(FCameraDirectorRigUsageInfo& UsageInfo) const override;
+	virtual void OnExtendAssetRegistryTags(FAssetRegistryTagsContext Context) const override;
 };
 
diff --git a/Engine/Plugins/Cameras/GameplayCameras/Source/GameplayCameras/Public/Directors/StateTreeCameraDirector.h b/Engine/Plugins/Cameras/GameplayCameras/Source/GameplayCameras/Public/Directors/StateTreeCameraDirector.h
--- a/Engine/Plugins/Cameras/GameplayCameras/Source/GameplayCameras/Public/Directors/StateTreeCameraDirector.h
+++ b/Engine/Plugins/Cameras/GameplayCameras/Source/GameplayCameras/Public/Directors/StateTreeCameraDirector.h
@@ -26,6 +26,7 @@ protected:
 	virtual FCameraDirectorEvaluatorPtr OnBuildEvaluator(FCameraDirectorEvaluatorBuilder& Builder) const override;
 	virtual void OnBuildCameraDirector(UE::Cameras::FCameraBuildLog& BuildLog) override;
 	virtual void OnGatherRigUsageInfo(FCameraDirectorRigUsageInfo& UsageInfo) const override;
+	virtual void OnExtendAssetRegistryTags(FAssetRegistryTagsContext Context) const override;
 
 public:
 
diff --git a/Engine/Plugins/Cameras/GameplayCameras/Source/GameplayCamerasEditor/Private/AssetTools/CameraRigProxyAssetEditor.cpp b/Engine/Plugins/Cameras/GameplayCameras/Source/GameplayCamerasEditor/Private/AssetTools/CameraRigProxyAssetEditor.cpp
--- a/Engine/Plugins/Cameras/GameplayCameras/Source/GameplayCamerasEditor/Private/AssetTools/CameraRigProxyAssetEditor.cpp
+++ b/Engine/Plugins/Cameras/GameplayCameras/Source/GameplayCamerasEditor/Private/AssetTools/CameraRigProxyAssetEditor.cpp
@@ -22,6 +22,8 @@ void UCameraRigProxyAssetEditor::GetObjectsToEdit(TArray<UObject*>& InObjectsToE
 TSharedPtr<FBaseAssetToolkit> UCameraRigProxyAssetEditor::CreateToolkit()
 {
 	using namespace UE::Cameras;
-	return MakeShared<FCameraRigProxyAssetEditorToolkit>(this);
+	TSharedPtr<FCameraRigProxyAssetEditorToolkit> Toolkit = MakeShared<FCameraRigProxyAssetEditorToolkit>(this);
+	Toolkit->SetCameraRigProxyAsset(CameraRigProxyAsset);
+	return Toolkit;
 }
 
diff --git a/Engine/Plugins/Cameras/GameplayCameras/Source/GameplayCamerasEditor/Private/Family/CameraAssetFamily.cpp b/Engine/Plugins/Cameras/GameplayCameras/Source/GameplayCamerasEditor/Private/Family/CameraAssetFamily.cpp
--- a/Engine/Plugins/Cameras/GameplayCameras/Source/GameplayCamerasEditor/Private/Family/CameraAssetFamily.cpp
+++ b/Engine/Plugins/Cameras/GameplayCameras/Source/GameplayCamerasEditor/Private/Family/CameraAssetFamily.cpp
@@ -6,7 +6,7 @@
 #include "Core/CameraDirector.h"
 #include "Core/CameraRigAsset.h"
 #include "Core/CameraRigProxyAsset.h"
-#include "Family/GameplayCamerasFamilyConstants.h"
+#include "Family/GameplayCamerasFamilyHelper.h"
 
 #define LOCTEXT_NAMESPACE "CameraAssetFamily"
 
@@ -27,6 +27,7 @@ UObject* FCameraAssetFamily::GetRootAsset() const
 void FCameraAssetFamily::GetAssetTypes(TArray<UClass*>& OutAssetTypes) const
 {
 	OutAssetTypes.Add(UCameraAsset::StaticClass());
+	OutAssetTypes.Add(UCameraDirector::StaticClass());
 	OutAssetTypes.Add(UCameraRigAsset::StaticClass());
 	OutAssetTypes.Add(UCameraRigProxyAsset::StaticClass());
 }
@@ -50,6 +51,14 @@ void FCameraAssetFamily::FindAssetsOfType(UClass* InAssetType, TArray<FAssetData
 		return;
 	}
 
+	if (InAssetType == UCameraDirector::StaticClass())
+	{
+		TArray<FAssetData> ThisAsset{ FAssetData(RootAsset) };
+		TArray<FAssetData> ExternalCameraDirectors;
+		FGameplayCamerasFamilyHelper::GetExternalCameraDirectorAssets(ThisAsset, OutAssets);
+		return;
+	}
+
 	FCameraDirectorRigUsageInfo UsageInfo;
 	CameraDirector->GatherRigUsageInfo(UsageInfo);
 
@@ -84,12 +93,12 @@ FText FCameraAssetFamily::GetAssetTypeTooltip(UClass* InAssetType) const
 
 const FSlateBrush* FCameraAssetFamily::GetAssetIcon(UClass* InAssetType) const
 {
-	return FGameplayCamerasFamilyConstants::GetAssetIcon(InAssetType);
+	return FGameplayCamerasFamilyHelper::GetAssetIcon(InAssetType);
 }
 
 FSlateColor FCameraAssetFamily::GetAssetTint(UClass* InAssetType) const
 {
-	return FGameplayCamerasFamilyConstants::GetAssetTint(InAssetType);
+	return FGameplayCamerasFamilyHelper::GetAssetTint(InAssetType);
 }
 
 }  // namespace UE::Cameras
diff --git a/Engine/Plugins/Cameras/GameplayCameras/Source/GameplayCamerasEditor/Private/Family/CameraRigAssetFamily.cpp b/Engine/Plugins/Cameras/GameplayCameras/Source/GameplayCamerasEditor/Private/Family/CameraRigAssetFamily.cpp
--- a/Engine/Plugins/Cameras/GameplayCameras/Source/GameplayCamerasEditor/Private/Family/CameraRigAssetFamily.cpp
+++ b/Engine/Plugins/Cameras/GameplayCameras/Source/GameplayCamerasEditor/Private/Family/CameraRigAssetFamily.cpp
@@ -6,7 +6,7 @@
 #include "Core/CameraAsset.h"
 #include "Core/CameraDirector.h"
 #include "Core/CameraRigAsset.h"
-#include "Family/GameplayCamerasFamilyConstants.h"
+#include "Family/GameplayCamerasFamilyHelper.h"
 #include "UObject/ReferencerFinder.h"
 
 #define LOCTEXT_NAMESPACE "CameraRigAssetFamily"
@@ -28,6 +28,7 @@ UObject* FCameraRigAssetFamily::GetRootAsset() const
 void FCameraRigAssetFamily::GetAssetTypes(TArray<UClass*>& OutAssetTypes) const
 {
 	OutAssetTypes.Add(UCameraAsset::StaticClass());
+	OutAssetTypes.Add(UCameraDirector::StaticClass());
 	OutAssetTypes.Add(UCameraRigAsset::StaticClass());
 }
 
@@ -44,52 +45,13 @@ void FCameraRigAssetFamily::FindAssetsOfType(UClass* InAssetType, TArray<FAssetD
 	}
 	else if (InAssetType == UCameraAsset::StaticClass())
 	{
-		const FName RootPackageName = RootAsset->GetPackage()->GetFName();
-
-		TArray<FAssetData> AllCameraAssets;
-		IAssetRegistry& AssetRegistry = FAssetRegistryModule::GetRegistry();
-		AssetRegistry.GetAssetsByClass(UCameraAsset::StaticClass()->GetClassPathName(), AllCameraAssets);
-		for (const FAssetData& CameraAsset : AllCameraAssets)
-		{
-			// By default, use the asset tags to know which camera asset uses which camera rigs.
-			// If the asset is loaded, it might have been modified and not saved yet, so use the in-memory
-			// object instead.
-			// If the asset doesn't have tags, it hasn't been saved since tags were added, so load it
-			// and also use it directly in memory.
-			bool bUseObjectDirectly = (CameraAsset.IsAssetLoaded());
-			if (!bUseObjectDirectly)
-			{
-				FString UsedCameraRigsTag = CameraAsset.GetTagValueRef<FString>(TEXT("UsedCameraRigs"));
-				if (!UsedCameraRigsTag.IsEmpty())
-				{
-					TArray<FString> UsedCameraRigs;
-					UsedCameraRigsTag.ParseIntoArray(UsedCameraRigs, TEXT(";"));
-					if (UsedCameraRigs.Contains(RootPackageName))
-					{
-						OutAssets.Add(CameraAsset);
-					}
-				}
-				else
-				{
-					bUseObjectDirectly = true;
-				}
-			}
-			if (bUseObjectDirectly)
-			{
-				if (UCameraAsset* LoadedCameraAsset = Cast<UCameraAsset>(CameraAsset.GetAsset()))
-				{
-					if (UCameraDirector* LoadedCameraDirector = LoadedCameraAsset->GetCameraDirector())
-					{
-						FCameraDirectorRigUsageInfo UsageInfo;
-						LoadedCameraDirector->GatherRigUsageInfo(UsageInfo);
-						if (UsageInfo.CameraRigs.Contains(RootAsset))
-						{
-							OutAssets.Add(CameraAsset);
-						}
-					}
-				}
-			}
-		}
+		FGameplayCamerasFamilyHelper::FindRelatedCameraAssets(RootAsset, OutAssets);
+	}
+	else if (InAssetType == UCameraDirector::StaticClass())
+	{
+		TArray<FAssetData> CameraAssets;
+		FGameplayCamerasFamilyHelper::FindRelatedCameraAssets(RootAsset, CameraAssets);
+		FGameplayCamerasFamilyHelper::GetExternalCameraDirectorAssets(CameraAssets, OutAssets);
 	}
 }
 
@@ -104,12 +66,12 @@ FText FCameraRigAssetFamily::GetAssetTypeTooltip(UClass* InAssetType) const
 
 const FSlateBrush* FCameraRigAssetFamily::GetAssetIcon(UClass* InAssetType) const
 {
-	return FGameplayCamerasFamilyConstants::GetAssetIcon(InAssetType);
+	return FGameplayCamerasFamilyHelper::GetAssetIcon(InAssetType);
 }
 
 FSlateColor FCameraRigAssetFamily::GetAssetTint(UClass* InAssetType) const
 {
-	return FGameplayCamerasFamilyConstants::GetAssetTint(InAssetType);
+	return FGameplayCamerasFamilyHelper::GetAssetTint(InAssetType);
 }
 
 }  // namespace UE::Cameras
diff --git a/Engine/Plugins/Cameras/GameplayCameras/Source/GameplayCamerasEditor/Private/Family/CameraRigProxyAssetFamily.cpp b/Engine/Plugins/Cameras/GameplayCameras/Source/GameplayCamerasEditor/Private/Family/CameraRigProxyAssetFamily.cpp
--- a/Engine/Plugins/Cameras/GameplayCameras/Source/GameplayCamerasEditor/Private/Family/CameraRigProxyAssetFamily.cpp
+++ b/Engine/Plugins/Cameras/GameplayCameras/Source/GameplayCamerasEditor/Private/Family/CameraRigProxyAssetFamily.cpp
@@ -2,9 +2,11 @@
 
 #include "Family/CameraRigProxyAssetFamily.h"
 
+#include "AssetRegistry/AssetRegistryModule.h"
 #include "Core/CameraAsset.h"
+#include "Core/CameraDirector.h"
 #include "Core/CameraRigProxyAsset.h"
-#include "Family/GameplayCamerasFamilyConstants.h"
+#include "Family/GameplayCamerasFamilyHelper.h"
 
 #define LOCTEXT_NAMESPACE "CameraRigProxyAssetFamily"
 
@@ -25,13 +27,30 @@ UObject* FCameraRigProxyAssetFamily::GetRootAsset() const
 void FCameraRigProxyAssetFamily::GetAssetTypes(TArray<UClass*>& OutAssetTypes) const
 {
 	OutAssetTypes.Add(UCameraAsset::StaticClass());
+	OutAssetTypes.Add(UCameraDirector::StaticClass());
 	OutAssetTypes.Add(UCameraRigProxyAsset::StaticClass());
 }
 
 void FCameraRigProxyAssetFamily::FindAssetsOfType(UClass* InAssetType, TArray<FAssetData>& OutAssets) const
 {
-	if (InAssetType == UCameraAsset::StaticClass())
+	if (!RootAsset)
+	{
+		return;
+	}
+
+	if (InAssetType == UCameraRigProxyAsset::StaticClass())
+	{
+		OutAssets.Add(RootAsset);
+	}
+	else if (InAssetType == UCameraAsset::StaticClass())
+	{
+		FGameplayCamerasFamilyHelper::FindRelatedCameraAssets(RootAsset, OutAssets);
+	}
+	else if (InAssetType == UCameraDirector::StaticClass())
 	{
+		TArray<FAssetData> CameraAssets;
+		FGameplayCamerasFamilyHelper::FindRelatedCameraAssets(RootAsset, CameraAssets);
+		FGameplayCamerasFamilyHelper::GetExternalCameraDirectorAssets(CameraAssets, OutAssets);
 	}
 }
 
@@ -46,12 +65,12 @@ FText FCameraRigProxyAssetFamily::GetAssetTypeTooltip(UClass* InAssetType) const
 
 const FSlateBrush* FCameraRigProxyAssetFamily::GetAssetIcon(UClass* InAssetType) const
 {
-	return FGameplayCamerasFamilyConstants::GetAssetIcon(InAssetType);
+	return FGameplayCamerasFamilyHelper::GetAssetIcon(InAssetType);
 }
 
 FSlateColor FCameraRigProxyAssetFamily::GetAssetTint(UClass* InAssetType) const
 {
-	return FGameplayCamerasFamilyConstants::GetAssetTint(InAssetType);
+	return FGameplayCamerasFamilyHelper::GetAssetTint(InAssetType);
 }
 
 }  // namespace UE::Cameras
diff --git a/Engine/Plugins/Cameras/GameplayCameras/Source/GameplayCamerasEditor/Private/Family/GameplayCamerasFamilyConstants.cpp b/Engine/Plugins/Cameras/GameplayCameras/Source/GameplayCamerasEditor/Private/Family/GameplayCamerasFamilyConstants.cpp
deleted file mode 100644
diff --git a/Engine/Plugins/Cameras/GameplayCameras/Source/GameplayCamerasEditor/Private/Family/GameplayCamerasFamilyConstants.h b/Engine/Plugins/Cameras/GameplayCameras/Source/GameplayCamerasEditor/Private/Family/GameplayCamerasFamilyConstants.h
deleted file mode 100644
diff --git a/Engine/Plugins/Cameras/GameplayCameras/Source/GameplayCamerasEditor/Private/Family/GameplayCamerasFamilyHelper.cpp b/Engine/Plugins/Cameras/GameplayCameras/Source/GameplayCamerasEditor/Private/Family/GameplayCamerasFamilyHelper.cpp
new file mode 100644
--- /dev/null
+++ b/Engine/Plugins/Cameras/GameplayCameras/Source/GameplayCamerasEditor/Private/Family/GameplayCamerasFamilyHelper.cpp
@@ -0,0 +1,171 @@
+// Copyright Epic Games, Inc. All Rights Reserved.
+
+#include "Family/GameplayCamerasFamilyHelper.h"
+
+#include "AssetRegistry/AssetData.h"
+#include "AssetRegistry/AssetRegistryModule.h"
+#include "Core/CameraAsset.h"
+#include "Core/CameraDirector.h"
+#include "Core/CameraRigAsset.h"
+#include "Core/CameraRigProxyAsset.h"
+#include "Styles/GameplayCamerasEditorStyle.h"
+#include "UObject/Package.h"
+
+namespace UE::Cameras
+{
+
+const FSlateBrush* FGameplayCamerasFamilyHelper::GetAssetIcon(UClass* InAssetType)
+{
+	if (InAssetType == UCameraAsset::StaticClass())
+	{
+		return FGameplayCamerasEditorStyle::Get()->GetBrush("Family.CameraAsset");
+	}
+	else if (InAssetType == UCameraDirector::StaticClass())
+	{
+		return FGameplayCamerasEditorStyle::Get()->GetBrush("Family.CameraDirector");
+	}
+	else if (InAssetType == UCameraRigAsset::StaticClass())
+	{
+		return FGameplayCamerasEditorStyle::Get()->GetBrush("Family.CameraRigAsset");
+	}
+	else if (InAssetType == UCameraRigProxyAsset::StaticClass())
+	{
+		return FGameplayCamerasEditorStyle::Get()->GetBrush("Family.CameraRigProxyAsset");
+	}
+	return nullptr;
+}
+
+FSlateColor FGameplayCamerasFamilyHelper::GetAssetTint(UClass* InAssetType)
+{
+	static const FColor Teal(23, 126, 137);
+	static const FColor MidnightGreen(8, 76, 97);
+	static const FColor Poppy(219, 58, 52);
+	static const FColor Sunglow(255, 200, 87);
+
+	if (InAssetType == UCameraAsset::StaticClass())
+	{
+		return FSlateColor(Sunglow);
+	}
+	else if (InAssetType == UCameraDirector::StaticClass())
+	{
+		return FSlateColor(Poppy);
+	}
+	else if (InAssetType == UCameraRigAsset::StaticClass())
+	{
+		return FSlateColor(Teal);
+	}
+	else if (InAssetType == UCameraRigProxyAsset::StaticClass())
+	{
+		return FSlateColor(MidnightGreen);
+	}
+	return FSlateColor();
+}
+
+namespace Internal
+{
+
+static void FindRelatedCameraAssets(UObject* Object, FName TagName, TFunctionRef<bool(const FCameraDirectorRigUsageInfo&)> ContainsPredicate, TArray<FAssetData>& OutCameraAssets)
+{
+	const FName RootPackageName = Object->GetPackage()->GetFName();
+
+	TArray<FAssetData> AllCameraAssets;
+	IAssetRegistry& AssetRegistry = FAssetRegistryModule::GetRegistry();
+	AssetRegistry.GetAssetsByClass(UCameraAsset::StaticClass()->GetClassPathName(), AllCameraAssets);
+	for (const FAssetData& CameraAsset : AllCameraAssets)
+	{
+		// By default, use the asset tags to know which camera asset uses which camera rigs.
+		// If the asset is loaded, it might have been modified and not saved yet, so use the in-memory
+		// object instead.
+		// If the asset doesn't have tags, it hasn't been saved since tags were added, so load it
+		// and also use it directly in memory.
+		bool bUseObjectDirectly = (CameraAsset.IsAssetLoaded());
+		if (!bUseObjectDirectly)
+		{
+			FString UsedReferencesTag = CameraAsset.GetTagValueRef<FString>(TagName);
+			if (!UsedReferencesTag.IsEmpty())
+			{
+				TArray<FString> UsedReferences;
+				UsedReferencesTag.ParseIntoArrayLines(UsedReferences);
+				if (UsedReferences.Contains(RootPackageName))
+				{
+					OutCameraAssets.Add(CameraAsset);
+				}
+			}
+			else if (!CameraAsset.FindTag(TagName))
+			{
+				bUseObjectDirectly = true;
+			}
+		}
+		if (bUseObjectDirectly)
+		{
+			if (UCameraAsset* LoadedCameraAsset = Cast<UCameraAsset>(CameraAsset.GetAsset()))
+			{
+				if (UCameraDirector* LoadedCameraDirector = LoadedCameraAsset->GetCameraDirector())
+				{
+					FCameraDirectorRigUsageInfo UsageInfo;
+					LoadedCameraDirector->GatherRigUsageInfo(UsageInfo);
+					if (ContainsPredicate(UsageInfo))
+					{
+						OutCameraAssets.Add(CameraAsset);
+					}
+				}
+			}
+		}
+	}
+}
+
+}  // namespace Internal
+
+void FGameplayCamerasFamilyHelper::FindRelatedCameraAssets(UCameraRigAsset* CameraRig, TArray<FAssetData>& OutCameraAssets)
+{
+	Internal::FindRelatedCameraAssets(
+			CameraRig, 
+			TEXT("UsedCameraRigs"),
+			[CameraRig](const FCameraDirectorRigUsageInfo& UsageInfo) { return UsageInfo.CameraRigs.Contains(CameraRig); },
+			OutCameraAssets);
+}
+
+void FGameplayCamerasFamilyHelper::FindRelatedCameraAssets(UCameraRigProxyAsset* CameraRigProxy, TArray<FAssetData>& OutCameraAssets)
+{
+	Internal::FindRelatedCameraAssets(
+			CameraRigProxy, 
+			TEXT("UsedCameraRigProxies"),
+			[CameraRigProxy](const FCameraDirectorRigUsageInfo& UsageInfo) { return UsageInfo.CameraRigProxies.Contains(CameraRigProxy); },
+			OutCameraAssets);
+}
+
+void FGameplayCamerasFamilyHelper::GetExternalCameraDirectorAssets(TArrayView<FAssetData> CameraAssets, TArray<FAssetData>& OutExternalCameraDirectors)
+{
+	IAssetRegistry& AssetRegistry = FAssetRegistryModule::GetRegistry();
+
+	for (const FAssetData& CameraAsset : CameraAssets)
+	{
+		FString ExternalDirectorName = CameraAsset.GetTagValueRef<FString>(TEXT("ExternalDirector"));
+		if (CameraAsset.IsAssetLoaded())
+		{
+			// If the camera asset is in memory, let's ask for up-to-date information.
+			if (UCameraAsset* LoadedCameraAsset = Cast<UCameraAsset>(CameraAsset.GetAsset()))
+			{
+				if (UCameraDirector* CameraDirector = LoadedCameraAsset->GetCameraDirector())
+				{
+					FAssetRegistryTagsContextData ContextData(LoadedCameraAsset);
+					FAssetRegistryTagsContext Context(ContextData);
+					CameraDirector->ExtendAssetRegistryTags(Context);
+
+					if (UObject::FAssetRegistryTag* NewTag = Context.FindTag(TEXT("ExternalDirector")))
+					{
+						ExternalDirectorName = NewTag->Value;
+					}
+				}
+			}
+		}
+		if (!ExternalDirectorName.IsEmpty())
+		{
+			FAssetData ExternalDirectorAsset = AssetRegistry.GetAssetByObjectPath(FSoftObjectPath(ExternalDirectorName));
+			OutExternalCameraDirectors.Add(ExternalDirectorAsset);
+		}
+	}
+}
+
+}  // namespace UE::Cameras
+
diff --git a/Engine/Plugins/Cameras/GameplayCameras/Source/GameplayCamerasEditor/Private/Family/GameplayCamerasFamilyHelper.h b/Engine/Plugins/Cameras/GameplayCameras/Source/GameplayCamerasEditor/Private/Family/GameplayCamerasFamilyHelper.h
new file mode 100644
--- /dev/null
+++ b/Engine/Plugins/Cameras/GameplayCameras/Source/GameplayCamerasEditor/Private/Family/GameplayCamerasFamilyHelper.h
@@ -0,0 +1,32 @@
+// Copyright Epic Games, Inc. All Rights Reserved.
+
+#pragma once
+
+#include "Styling/SlateColor.h"
+
+class UClass;
+class UObject;
+class UCameraAsset;
+class UCameraRigAsset;
+class UCameraRigProxyAsset;
+struct FAssetData;
+struct FSlateBrush;
+
+namespace UE::Cameras
+{
+
+class FGameplayCamerasFamilyHelper
+{
+public:
+
+	static const FSlateBrush* GetAssetIcon(UClass* InAssetType);
+
+	static FSlateColor GetAssetTint(UClass* InAssetType);
+
+	static void FindRelatedCameraAssets(UCameraRigAsset* CameraRig, TArray<FAssetData>& OutCameraAssets);
+	static void FindRelatedCameraAssets(UCameraRigProxyAsset* CameraRigProxy, TArray<FAssetData>& OutCameraAssets);
+	static void GetExternalCameraDirectorAssets(TArrayView<FAssetData> CameraAssets, TArray<FAssetData>& OutExternalCameraDirectors);
+};
+
+}  // namespace UE::Cameras
+
diff --git a/Engine/Plugins/Cameras/GameplayCameras/Source/GameplayCamerasEditor/Private/Styles/GameplayCamerasEditorStyle.cpp b/Engine/Plugins/Cameras/GameplayCameras/Source/GameplayCamerasEditor/Private/Styles/GameplayCamerasEditorStyle.cpp
--- a/Engine/Plugins/Cameras/GameplayCameras/Source/GameplayCamerasEditor/Private/Styles/GameplayCamerasEditorStyle.cpp
+++ b/Engine/Plugins/Cameras/GameplayCameras/Source/GameplayCamerasEditor/Private/Styles/GameplayCamerasEditorStyle.cpp
@@ -152,6 +152,7 @@ FGameplayCamerasEditorStyle::FGameplayCamerasEditorStyle()
 
 	// Family icons.
 	Set("Family.CameraAsset", new IMAGE_BRUSH_SVG("Icons/ContentBrowser-CameraKit", Icon20x20));
+	Set("Family.CameraDirector", new IMAGE_BRUSH_SVG("Icons/CameraEditor-CameraDirector", Icon20x20));
 	Set("Family.CameraRigAsset", new IMAGE_BRUSH_SVG("Icons/ContentBrowser-CameraRig", Icon20x20));
 	Set("Family.CameraRigProxyAsset", new IMAGE_BRUSH_SVG("Icons/ContentBrowser-CameraRigProxy", Icon20x20));
 
diff --git a/Engine/Plugins/Cameras/GameplayCameras/Source/GameplayCamerasEditor/Private/Toolkits/CameraRigProxyAssetEditorToolkit.cpp b/Engine/Plugins/Cameras/GameplayCameras/Source/GameplayCamerasEditor/Private/Toolkits/CameraRigProxyAssetEditorToolkit.cpp
--- a/Engine/Plugins/Cameras/GameplayCameras/Source/GameplayCamerasEditor/Private/Toolkits/CameraRigProxyAssetEditorToolkit.cpp
+++ b/Engine/Plugins/Cameras/GameplayCameras/Source/GameplayCamerasEditor/Private/Toolkits/CameraRigProxyAssetEditorToolkit.cpp
@@ -6,7 +6,9 @@
 #include "Core/CameraRigProxyAsset.h"
 #include "Framework/Docking/LayoutExtender.h"
 #include "Helpers/AssetTypeMenuOverlayHelper.h"
+#include "IGameplayCamerasFamily.h"
 #include "PropertyEditorModule.h"
+#include "Widgets/SCameraFamilyShortcutBar.h"
 
 #define LOCTEXT_NAMESPACE "CameraRigProxyAssetEditorToolkit"
 
@@ -36,6 +38,11 @@ FCameraRigProxyAssetEditorToolkit::~FCameraRigProxyAssetEditorToolkit()
 {
 }
 
+void FCameraRigProxyAssetEditorToolkit::SetCameraRigProxyAsset(UCameraRigProxyAsset* InCameraRigProxyAsset)
+{
+	CameraRigProxyAsset = InCameraRigProxyAsset;
+}
+
 void FCameraRigProxyAssetEditorToolkit::AddReferencedObjects(FReferenceCollector& Collector)
 {
 	Collector.AddReferencedObject(CameraRigProxyAsset);
@@ -77,6 +84,22 @@ void FCameraRigProxyAssetEditorToolkit::CreateWidgets()
 	DetailsView = CreateDetailsView();
 }
 
+void FCameraRigProxyAssetEditorToolkit::RegisterToolbar()
+{
+	TSharedRef<IGameplayCamerasFamily> Family = IGameplayCamerasFamily::CreateFamily(CameraRigProxyAsset).ToSharedRef();
+	TSharedRef<FExtender> ToolbarExtender = MakeShared<FExtender>();
+	AddToolbarExtender(ToolbarExtender);
+	ToolbarExtender->AddToolBarExtension(
+			"Asset",
+			EExtensionHook::After,
+			GetToolkitCommands(),
+			FToolBarExtensionDelegate::CreateLambda([this, Family](FToolBarBuilder& Builder)
+				{
+					AddToolbarWidget(SNew(SCameraFamilyShortcutBar, SharedThis(this), Family));
+				})
+			);
+}
+
 void FCameraRigProxyAssetEditorToolkit::PostInitAssetEditor()
 {
 	RegenerateMenusAndToolbars();
diff --git a/Engine/Plugins/Cameras/GameplayCameras/Source/GameplayCamerasEditor/Private/Widgets/SCameraFamilyAssetShortcut.cpp b/Engine/Plugins/Cameras/GameplayCameras/Source/GameplayCamerasEditor/Private/Widgets/SCameraFamilyAssetShortcut.cpp
--- a/Engine/Plugins/Cameras/GameplayCameras/Source/GameplayCamerasEditor/Private/Widgets/SCameraFamilyAssetShortcut.cpp
+++ b/Engine/Plugins/Cameras/GameplayCameras/Source/GameplayCamerasEditor/Private/Widgets/SCameraFamilyAssetShortcut.cpp
@@ -276,7 +276,7 @@ TSharedRef<SWidget> SCameraFamilyAssetShortcut::HandleGetDropdownMenuContent()
 		MenuBuilder.BeginSection("AssetSelection", LOCTEXT("AssetSelectionSection", "Select Asset"));
 		{
 			FAssetPickerConfig AssetPickerConfig;
-			AssetPickerConfig.Filter.ClassPaths.Add(FamilyAssetType->GetClassPathName());
+			AssetPickerConfig.bCanShowClasses = false;
 			for (const FAssetData& AssetData : AssetDatas)
 			{
 				AssetPickerConfig.Filter.SoftObjectPaths.Add(AssetData.GetSoftObjectPath());
diff --git a/Engine/Plugins/Cameras/GameplayCameras/Source/GameplayCamerasEditor/Public/Toolkits/CameraRigProxyAssetEditorToolkit.h b/Engine/Plugins/Cameras/GameplayCameras/Source/GameplayCamerasEditor/Public/Toolkits/CameraRigProxyAssetEditorToolkit.h
--- a/Engine/Plugins/Cameras/GameplayCameras/Source/GameplayCamerasEditor/Public/Toolkits/CameraRigProxyAssetEditorToolkit.h
+++ b/Engine/Plugins/Cameras/GameplayCameras/Source/GameplayCamerasEditor/Public/Toolkits/CameraRigProxyAssetEditorToolkit.h
@@ -26,12 +26,15 @@ public:
 	FCameraRigProxyAssetEditorToolkit(UCameraRigProxyAssetEditor* InOwningAssetEditor);
 	~FCameraRigProxyAssetEditorToolkit();
 
+	void SetCameraRigProxyAsset(UCameraRigProxyAsset* InCameraRigProxyAsset);
+
 protected:
 
 	// FBaseAssetToolkit interface
 	virtual void RegisterTabSpawners(const TSharedRef<class FTabManager>& TabManager) override;
 	virtual void UnregisterTabSpawners(const TSharedRef<FTabManager>& InTabManager) override;
 	virtual void CreateWidgets() override;
+	virtual void RegisterToolbar() override;
 	virtual void PostInitAssetEditor() override;
 	virtual void PostRegenerateMenusAndToolbars() override;
 
