Cameras: re-add the BP parameter structure for RunCameraDirector, mark it as deprecated

diff --git a/Engine/Plugins/Cameras/GameplayCameras/Source/GameplayCameras/Private/Directors/BlueprintCameraDirector.cpp b/Engine/Plugins/Cameras/GameplayCameras/Source/GameplayCameras/Private/Directors/BlueprintCameraDirector.cpp
--- a/Engine/Plugins/Cameras/GameplayCameras/Source/GameplayCameras/Private/Directors/BlueprintCameraDirector.cpp
+++ b/Engine/Plugins/Cameras/GameplayCameras/Source/GameplayCameras/Private/Directors/BlueprintCameraDirector.cpp
@@ -383,7 +383,12 @@ void UBlueprintCameraDirectorEvaluator::NativeRunCameraDirector(const UE::Camera
 			EvaluationContextOwner = EvaluationContext->GetOwner();
 		}
 
-		RunCameraDirector(Params.DeltaTime, EvaluationContextOwner);
+PRAGMA_DISABLE_DEPRECATION_WARNINGS
+		FBlueprintCameraDirectorEvaluationParams OldParams;
+		OldParams.DeltaTime = Params.DeltaTime;
+		OldParams.EvaluationContextOwner = EvaluationContextOwner;
+		RunCameraDirector(Params.DeltaTime, EvaluationContextOwner, OldParams);
+PRAGMA_ENABLE_DEPRECATION_WARNINGS
 	}
 }
 
diff --git a/Engine/Plugins/Cameras/GameplayCameras/Source/GameplayCameras/Public/Directors/BlueprintCameraDirector.h b/Engine/Plugins/Cameras/GameplayCameras/Source/GameplayCameras/Public/Directors/BlueprintCameraDirector.h
--- a/Engine/Plugins/Cameras/GameplayCameras/Source/GameplayCameras/Public/Directors/BlueprintCameraDirector.h
+++ b/Engine/Plugins/Cameras/GameplayCameras/Source/GameplayCameras/Public/Directors/BlueprintCameraDirector.h
@@ -14,6 +14,20 @@ class UCameraRigAsset;
 class UCameraRigProxyAsset;
 enum class ECameraRigLayer : uint8;
 
+USTRUCT(BlueprintType)
+struct UE_DEPRECATED(5.6, "This parameter structure is deprecated, parameters are now passed directly to RunCameraDirector.") FBlueprintCameraDirectorEvaluationParams
+{
+	GENERATED_BODY()
+
+	/** The elapsed time since the last evaluation. */
+	UPROPERTY(BlueprintReadWrite, Category="Evaluation", meta=(Deprecated, DeprecationMessage="Use the main DeltaTime parameter on RunCameraDirector"))
+	float DeltaTime = 0.f;
+
+	/** The owner (if any) of the evaluation context we are running inside of. */
+	UPROPERTY(BlueprintReadWrite, Category="Evaluation", meta=(Deprecated, DeprecationMessage="Use the main EvaluationContextOwner parameter on RunCameraDirector"))
+	TObjectPtr<UObject> EvaluationContextOwner;
+};
+
 /**
  * Base class for a Blueprint camera director evaluator.
  */
@@ -42,8 +56,10 @@ public:
 	 * Override this method in Blueprint to execute the custom logic that determines
 	 * what camera rig(s) should be active every frame.
 	 */
+PRAGMA_DISABLE_DEPRECATION_WARNINGS
 	UFUNCTION(BlueprintCallable, BlueprintImplementableEvent, Category="Evaluation")
-	void RunCameraDirector(float DeltaTime, UObject* EvaluationContextOwner);
+	void RunCameraDirector(float DeltaTime, UObject* EvaluationContextOwner, const FBlueprintCameraDirectorEvaluationParams& Params);
+PRAGMA_ENABLE_DEPRECATION_WARNINGS
 
 public:
 
