Cameras: correctly handle legacy content with rig proxy tables

diff --git a/Engine/Plugins/Cameras/GameplayCameras/Config/BaseGameplayCameras.ini b/Engine/Plugins/Cameras/GameplayCameras/Config/BaseGameplayCameras.ini
--- a/Engine/Plugins/Cameras/GameplayCameras/Config/BaseGameplayCameras.ini
+++ b/Engine/Plugins/Cameras/GameplayCameras/Config/BaseGameplayCameras.ini
@@ -27,3 +27,5 @@
 
 +FunctionRedirects=(OldName="GameplayCameraSystemActor.AutoManageActiveViewTarget", NewName="SetAutoSpawnedCameraSystemActorViewTarget")
 
++StructRedirects=(OldName="/Script/GameplayCameras.CameraRigProxyTableEntry",NewName="CameraRigProxyRedirectTableEntry")
+
diff --git a/Engine/Plugins/Cameras/GameplayCameras/Source/GameplayCameras/Private/Core/CameraDirector.cpp b/Engine/Plugins/Cameras/GameplayCameras/Source/GameplayCameras/Private/Core/CameraDirector.cpp
--- a/Engine/Plugins/Cameras/GameplayCameras/Source/GameplayCameras/Private/Core/CameraDirector.cpp
+++ b/Engine/Plugins/Cameras/GameplayCameras/Source/GameplayCameras/Private/Core/CameraDirector.cpp
@@ -4,9 +4,23 @@
 
 #include "Core/CameraRigAsset.h"
 #include "Core/CameraRigProxyAsset.h"
+#include "Core/CameraRigProxyRedirectTable.h"
 
 #include UE_INLINE_GENERATED_CPP_BY_NAME(CameraDirector)
 
+void UCameraDirector::PostLoad()
+{
+	Super::PostLoad();
+
+PRAGMA_DISABLE_DEPRECATION_WARNINGS
+	if (CameraRigProxyTable_DEPRECATED)
+	{
+		CameraRigProxyRedirectTable.Entries = CameraRigProxyTable_DEPRECATED->Entries;
+		CameraRigProxyTable_DEPRECATED = nullptr;
+	}
+PRAGMA_ENABLE_DEPRECATION_WARNINGS
+}
+
 FCameraDirectorEvaluatorPtr UCameraDirector::BuildEvaluator(FCameraDirectorEvaluatorBuilder& Builder) const
 {
 	using namespace UE::Cameras;
diff --git a/Engine/Plugins/Cameras/GameplayCameras/Source/GameplayCameras/Private/Core/CameraDirectorEvaluator.cpp b/Engine/Plugins/Cameras/GameplayCameras/Source/GameplayCameras/Private/Core/CameraDirectorEvaluator.cpp
--- a/Engine/Plugins/Cameras/GameplayCameras/Source/GameplayCameras/Private/Core/CameraDirectorEvaluator.cpp
+++ b/Engine/Plugins/Cameras/GameplayCameras/Source/GameplayCameras/Private/Core/CameraDirectorEvaluator.cpp
@@ -4,7 +4,7 @@
 
 #include "Core/CameraDirector.h"
 #include "Core/CameraEvaluationContext.h"
-#include "Core/CameraRigProxyTable.h"
+#include "Core/CameraRigProxyRedirectTable.h"
 #include "Core/CameraRigTransition.h"
 
 namespace UE::Cameras
@@ -75,9 +75,9 @@ const UCameraRigAsset* FCameraDirectorEvaluator::FindCameraRigByProxy(const UCam
 		return nullptr;
 	}
 
-	const FCameraRigProxyTable& ProxyTable = PrivateCameraDirector->CameraRigProxyTable;
+	const FCameraRigProxyRedirectTable& ProxyTable = PrivateCameraDirector->CameraRigProxyRedirectTable;
 
-	FCameraRigProxyTableResolveParams ResolveParams;
+	FCameraRigProxyResolveParams ResolveParams;
 	ResolveParams.CameraRigProxy = InProxy;
 	return ProxyTable.ResolveProxy(ResolveParams);
 }
diff --git a/Engine/Plugins/Cameras/GameplayCameras/Source/GameplayCameras/Private/Core/CameraRigProxyRedirectTable.cpp b/Engine/Plugins/Cameras/GameplayCameras/Source/GameplayCameras/Private/Core/CameraRigProxyRedirectTable.cpp
new file mode 100644
--- /dev/null
+++ b/Engine/Plugins/Cameras/GameplayCameras/Source/GameplayCameras/Private/Core/CameraRigProxyRedirectTable.cpp
@@ -0,0 +1,20 @@
+// Copyright Epic Games, Inc. All Rights Reserved.
+
+#include "Core/CameraRigProxyRedirectTable.h"
+
+#include UE_INLINE_GENERATED_CPP_BY_NAME(CameraRigProxyRedirectTable)
+
+UCameraRigAsset* FCameraRigProxyRedirectTable::ResolveProxy(const FCameraRigProxyResolveParams& InParams) const
+{
+	ensure(InParams.CameraRigProxy);
+
+	for (const FCameraRigProxyRedirectTableEntry& Entry : Entries)
+	{
+		if (InParams.CameraRigProxy == Entry.CameraRigProxy)
+		{
+			return Entry.CameraRig;
+		}
+	}
+	return nullptr;
+}
+
diff --git a/Engine/Plugins/Cameras/GameplayCameras/Source/GameplayCameras/Private/Core/CameraRigProxyTable.cpp b/Engine/Plugins/Cameras/GameplayCameras/Source/GameplayCameras/Private/Core/CameraRigProxyTable.cpp
deleted file mode 100644
diff --git a/Engine/Plugins/Cameras/GameplayCameras/Source/GameplayCameras/Public/Core/CameraDirector.h b/Engine/Plugins/Cameras/GameplayCameras/Source/GameplayCameras/Public/Core/CameraDirector.h
--- a/Engine/Plugins/Cameras/GameplayCameras/Source/GameplayCameras/Public/Core/CameraDirector.h
+++ b/Engine/Plugins/Cameras/GameplayCameras/Source/GameplayCameras/Public/Core/CameraDirector.h
@@ -5,7 +5,7 @@
 #include "CoreTypes.h"
 #include "UObject/Object.h"
 #include "Core/CameraDirectorEvaluator.h"
-#include "Core/CameraRigProxyTable.h"
+#include "Core/CameraRigProxyRedirectTable.h"
 #include "UObject/AssetRegistryTagsContext.h"
 
 #include "CameraDirector.generated.h"
@@ -69,6 +69,11 @@ public:
 	UE_API void FactoryCreateAsset(const FCameraDirectorFactoryCreateParams& InParams);
 #endif
 
+public:
+
+	// UObject interface.
+	UE_API virtual void PostLoad() override;
+
 protected:
 
 	/** Build the evaluator for this director. */
@@ -95,7 +100,14 @@ public:
 	 * to actual camera rigs.
 	 */
 	UPROPERTY(EditAnywhere, Category="Evaluation")
-	FCameraRigProxyTable CameraRigProxyTable;
+	FCameraRigProxyRedirectTable CameraRigProxyRedirectTable;
+
+private:
+
+PRAGMA_DISABLE_DEPRECATION_WARNINGS
+	UPROPERTY()
+	TObjectPtr<UCameraRigProxyTable> CameraRigProxyTable_DEPRECATED;
+PRAGMA_ENABLE_DEPRECATION_WARNINGS
 };
 
 #undef UE_API
diff --git a/Engine/Plugins/Cameras/GameplayCameras/Source/GameplayCameras/Public/Core/CameraRigProxyRedirectTable.h b/Engine/Plugins/Cameras/GameplayCameras/Source/GameplayCameras/Public/Core/CameraRigProxyRedirectTable.h
new file mode 100644
--- /dev/null
+++ b/Engine/Plugins/Cameras/GameplayCameras/Source/GameplayCameras/Public/Core/CameraRigProxyRedirectTable.h
@@ -0,0 +1,78 @@
+// Copyright Epic Games, Inc. All Rights Reserved.
+
+#pragma once
+
+#include "CoreTypes.h"
+#include "UObject/Object.h"
+
+#include "CameraRigProxyRedirectTable.generated.h"
+
+class UCameraRigAsset;
+class UCameraRigProxyAsset;
+
+/**
+ * Parameter structure for resolving a camera rig proxy from a proxy table.
+ */
+struct FCameraRigProxyResolveParams
+{
+	/** The camera rig proxy to resolve. */
+	const UCameraRigProxyAsset* CameraRigProxy = nullptr;
+};
+
+/**
+ * An entry in a camera rig proxy table.
+ */
+USTRUCT()
+struct FCameraRigProxyRedirectTableEntry
+{
+	GENERATED_BODY()
+
+	/** The camera rig proxy for this table entry. */
+	UPROPERTY(EditAnywhere, Category="Camera")
+	TObjectPtr<UCameraRigProxyAsset> CameraRigProxy;
+
+	/** The actual camera rig that should be mapped to the correspondig proxy. */
+	UPROPERTY(EditAnywhere, Category="Camera")
+	TObjectPtr<UCameraRigAsset> CameraRig;
+};
+
+/**
+ * A table that defines mappings between camera rig proxies and actual camera rigs.
+ */
+USTRUCT()
+struct FCameraRigProxyRedirectTable
+{
+	GENERATED_BODY()
+
+public:
+
+	/**
+	 * Resolves a given proxy to an actual camera rig.
+	 * Returns nullptr if the given proxy wasn't found, or not mapped to anything in the table.
+	 */
+	UCameraRigAsset* ResolveProxy(const FCameraRigProxyResolveParams& InParams) const;
+
+public:
+
+	// Internal API.
+
+	bool SerializeFromMismatchedTag(FPropertyTag const& Tag, FStructuredArchive::FSlot Slot);
+
+public:
+
+	/** The entries in the table. */
+	UPROPERTY(EditAnywhere, Category="Camera")
+	TArray<FCameraRigProxyRedirectTableEntry> Entries;
+};
+
+UCLASS()
+class UE_DEPRECATED(5.6, "Use FCameraRigProxyRedirectTable") UCameraRigProxyTable : public UObject
+{
+	GENERATED_BODY()
+
+public:
+
+	UPROPERTY()
+	TArray<FCameraRigProxyRedirectTableEntry> Entries;
+};
+
diff --git a/Engine/Plugins/Cameras/GameplayCameras/Source/GameplayCameras/Public/Core/CameraRigProxyTable.h b/Engine/Plugins/Cameras/GameplayCameras/Source/GameplayCameras/Public/Core/CameraRigProxyTable.h
deleted file mode 100644
