Cameras: add support preview evaluation during view rotation update

When using the GameplayCamerasPlayerCameraManager, a new project setting allows it to run a "preview update" to compute the view rotation early in the frame. This is suitable for use as the player controller's control rotation for when using last frame's value isn't adequate. This approach is more expensive than merely computing a yaw/pitch angle based on player input and camera rotation speed, but it does support complex situations such as blending between camera rigs that don't have the same orientation, but converge to an important point (e.g. target preservation).  In the future, the camera system will add support for the simple yaw/pitch approach, for games that have simpler setups.

The new option relies on running a "preview update" which, essentially, saves the state of the node evaluator tree, runs it, and restores the state, similar to IK updates. A first optimization pass is included in this change to reduce the size of these snapshots, although they are still too big and more optimization will have to be done later. Similarly, CPU-time optimization will have to be done too, e.g. by reducing the amount of virtual calls to serialize the state.

diff --git a/Engine/Plugins/Cameras/GameplayCameras/Source/GameplayCameras/Private/Core/BlendStackCameraNode.cpp b/Engine/Plugins/Cameras/GameplayCameras/Source/GameplayCameras/Private/Core/BlendStackCameraNode.cpp
--- a/Engine/Plugins/Cameras/GameplayCameras/Source/GameplayCameras/Private/Core/BlendStackCameraNode.cpp
+++ b/Engine/Plugins/Cameras/GameplayCameras/Source/GameplayCameras/Private/Core/BlendStackCameraNode.cpp
@@ -134,8 +134,8 @@ bool FBlendStackCameraNodeEvaluator::InitializeEntry(
 	NewEntry.CameraRig = CameraRig;
 	NewEntry.RootNode = EntryRootNode;
 	NewEntry.RootEvaluator = RootEvaluator->CastThisChecked<FBlendStackRootCameraNodeEvaluator>();
-	NewEntry.bWasContextInitialResultValid = EvaluationContext->GetInitialResult().bIsValid;
-	NewEntry.bIsFirstFrame = true;
+	NewEntry.Flags.bWasContextInitialResultValid = EvaluationContext->GetInitialResult().bIsValid;
+	NewEntry.Flags.bIsFirstFrame = true;
 
 	return true;
 }
@@ -163,7 +163,7 @@ void FBlendStackCameraNodeEvaluator::FreezeEntry(FCameraRigEntry& Entry)
 	RemoveListenedPackages(Entry);
 #endif
 	
-	Entry.bIsFrozen = true;
+	Entry.Flags.bIsFrozen = true;
 }
 
 FCameraRigEvaluationInfo FBlendStackCameraNodeEvaluator::GetActiveCameraRigEvaluationInfo() const
@@ -252,7 +252,7 @@ void FBlendStackCameraNodeEvaluator::ResolveEntries(TArray<FResolvedEntry>& OutR
 		}
 
 		// While we make these resolved entries, emit warnings and errors as needed.
-		if (!Entry.bIsFrozen)
+		if (!Entry.Flags.bIsFrozen)
 		{
 			// Check that we still have a valid context. If not, let's freeze the entry, since
 			// we won't be able to evaluate it anymore.
@@ -261,12 +261,12 @@ void FBlendStackCameraNodeEvaluator::ResolveEntries(TArray<FResolvedEntry>& OutR
 				FreezeEntry(Entry);
 
 #if UE_GAMEPLAY_CAMERAS_TRACE
-				if (Entry.bLogWarnings)
+				if (Entry.Flags.bLogWarnings)
 				{
 					UE_LOG(LogCameraSystem, Warning,
 							TEXT("Freezing camera rig '%s' because its evaluation context isn't valid anymore."),
 							*GetNameSafe(Entry.CameraRig));
-					Entry.bLogWarnings = false;
+					Entry.Flags.bLogWarnings = false;
 				}
 #endif  // UE_GAMEPLAY_CAMERAS_TRACE
 
@@ -278,12 +278,12 @@ void FBlendStackCameraNodeEvaluator::ResolveEntries(TArray<FResolvedEntry>& OutR
 			if (UNLIKELY(!ContextResult.bIsValid))
 			{
 #if UE_GAMEPLAY_CAMERAS_TRACE
-				if (Entry.bLogWarnings)
+				if (Entry.Flags.bLogWarnings)
 				{
 					UE_LOG(LogCameraSystem, Warning,
 							TEXT("Camera rig '%s' may experience a hitch because its initial result isn't valid."),
 							*GetNameSafe(Entry.CameraRig));
-					Entry.bLogWarnings = false;
+					Entry.Flags.bLogWarnings = false;
 				}
 #endif  // UE_GAMEPLAY_CAMERAS_TRACE
 
@@ -292,11 +292,11 @@ void FBlendStackCameraNodeEvaluator::ResolveEntries(TArray<FResolvedEntry>& OutR
 
 			// If the context was previously invalid, and this isn't the first frame, flag
 			// this update as a camera cut.
-			if (UNLIKELY(!Entry.bWasContextInitialResultValid && !Entry.bIsFirstFrame))
+			if (UNLIKELY(!Entry.Flags.bWasContextInitialResultValid && !Entry.Flags.bIsFirstFrame))
 			{
-				Entry.bForceCameraCut = true;
+				Entry.Flags.bForceCameraCut = true;
 			}
-			Entry.bWasContextInitialResultValid = true;
+			Entry.Flags.bWasContextInitialResultValid = true;
 
 			// Reset this entry's flags for this frame.
 			FCameraNodeEvaluationResult& CurResult = Entry.Result;
@@ -324,7 +324,7 @@ void FBlendStackCameraNodeEvaluator::ResolveEntries(TArray<FResolvedEntry>& OutR
 #if UE_GAMEPLAY_CAMERAS_TRACE
 		// This entry might have has warnings before. It's valid now, so let's
 		// re-enable warnings if it becomes invalid again in the future.
-		Entry.bLogWarnings = true;
+		Entry.Flags.bLogWarnings = true;
 #endif  // UE_GAMEPLAY_CAMERAS_TRACE
 	}
 }
@@ -334,8 +334,8 @@ void FBlendStackCameraNodeEvaluator::OnRunFinished(FCameraNodeEvaluationResult&
 	// Reset transient flags.
 	for (FCameraRigEntry& Entry : Entries)
 	{
-		Entry.bIsFirstFrame = false;
-		Entry.bForceCameraCut = false;
+		Entry.Flags.bIsFirstFrame = false;
+		Entry.Flags.bForceCameraCut = false;
 	}
 
 #if WITH_EDITOR || UE_GAMEPLAY_CAMERAS_DEBUG
@@ -501,6 +501,8 @@ void FBlendStackCameraNodeEvaluator::OnAddReferencedObjects(FReferenceCollector&
 
 void FBlendStackCameraNodeEvaluator::OnSerialize(const FCameraNodeEvaluatorSerializeParams& Params, FArchive& Ar)
 {
+	int32 NumEntriesToSerialize = Entries.Num();
+
 	if (Ar.IsSaving())
 	{
 		int32 NumEntries = Entries.Num();
@@ -511,18 +513,18 @@ void FBlendStackCameraNodeEvaluator::OnSerialize(const FCameraNodeEvaluatorSeria
 		int32 LoadedNumEntries = 0;
 		Ar << LoadedNumEntries;
 
-		ensure(LoadedNumEntries == Entries.Num());
+		ensureMsgf(
+				LoadedNumEntries == Entries.Num(),
+				TEXT("The number of entries changed since this blend stack was serialized!"));
+		NumEntriesToSerialize = LoadedNumEntries;
 	}
 
-	for (FCameraRigEntry& Entry : Entries)
+	for (int32 Index = 0; Index < NumEntriesToSerialize; ++Index)
 	{
+		FCameraRigEntry& Entry(Entries[Index]);
 		Entry.ContextResult.Serialize(Ar);
 		Entry.Result.Serialize(Ar);
-		Ar << Entry.bIsFirstFrame;
-		Ar << Entry.bIsFrozen;
-#if UE_GAMEPLAY_CAMERAS_TRACE
-		Ar << Entry.bLogWarnings;
-#endif  // UE_GAMEPLAY_CAMERAS_TRACE
+		Ar.SerializeBits(static_cast<void*>(&Entry.Flags), sizeof(FCameraRigEntry::Flags));
 	}
 }
 
@@ -556,7 +558,7 @@ void FBlendStackCameraNodeEvaluator::OnPostBuildAsset(const FGameplayCameraAsset
 					Index == Entries.Num() - 1);
 			if (!bInitialized)
 			{
-				Entry.bIsFrozen = true;
+				Entry.Flags.bIsFrozen = true;
 				continue;
 			}
 
diff --git a/Engine/Plugins/Cameras/GameplayCameras/Source/GameplayCameras/Private/Core/CameraPose.cpp b/Engine/Plugins/Cameras/GameplayCameras/Source/GameplayCameras/Private/Core/CameraPose.cpp
--- a/Engine/Plugins/Cameras/GameplayCameras/Source/GameplayCameras/Private/Core/CameraPose.cpp
+++ b/Engine/Plugins/Cameras/GameplayCameras/Source/GameplayCameras/Private/Core/CameraPose.cpp
@@ -456,17 +456,7 @@ void FCameraPose::InternalLerpChanged(const FCameraPose& ToPose, float Factor, c
 
 void FCameraPose::SerializeWithFlags(FArchive& Ar, FCameraPose& CameraPose)
 {
-	static FCameraPose DefaultCameraPose;
-
-	UScriptStruct* ClassStruct = FCameraPose::StaticStruct();
-	ClassStruct->SerializeItem(Ar, &CameraPose, &DefaultCameraPose);
-
-#define UE_CAMERA_POSE_FOR_PROPERTY(PropType, PropName)\
-	Ar << CameraPose.ChangedFlags.PropName;
-
-UE_CAMERA_POSE_FOR_ALL_PROPERTIES()
-
-#undef UE_CAMERA_POSE_FOR_PROPERTY
+	Ar.Serialize(static_cast<void*>(&CameraPose), sizeof(FCameraPose));
 }
 
 void FCameraPose::SerializeWithFlags(FArchive& Ar)
diff --git a/Engine/Plugins/Cameras/GameplayCameras/Source/GameplayCameras/Private/Core/CameraSystemEvaluator.cpp b/Engine/Plugins/Cameras/GameplayCameras/Source/GameplayCameras/Private/Core/CameraSystemEvaluator.cpp
--- a/Engine/Plugins/Cameras/GameplayCameras/Source/GameplayCameras/Private/Core/CameraSystemEvaluator.cpp
+++ b/Engine/Plugins/Cameras/GameplayCameras/Source/GameplayCameras/Private/Core/CameraSystemEvaluator.cpp
@@ -3,6 +3,7 @@
 #include "Core/CameraSystemEvaluator.h"
 
 #include "Camera/CameraTypes.h"
+#include "Core/BuiltInCameraVariables.h"
 #include "Core/CameraDirectorEvaluator.h"
 #include "Core/CameraEvaluationContext.h"
 #include "Core/CameraEvaluationService.h"
@@ -15,6 +16,8 @@
 #include "Debug/RootCameraDebugBlock.h"
 #include "GameplayCamerasSettings.h"
 #include "Math/ColorList.h"
+#include "Serialization/MemoryReader.h"
+#include "Serialization/MemoryWriter.h"
 #include "Services/CameraModifierService.h"
 #include "Services/CameraParameterSetterService.h"
 #include "Services/CameraShakeService.h"
@@ -424,6 +427,36 @@ void FCameraSystemEvaluator::PostUpdateServices(float DeltaTime, ECameraEvaluati
 	}
 }
 
+void FCameraSystemEvaluator::ViewRotationPreviewUpdate(const FCameraSystemEvaluationParams& Params, FCameraSystemViewRotationEvaluationResult& OutResult)
+{
+	SCOPE_CYCLE_COUNTER(CameraSystemEval_Total);
+
+	EvaluatorSnapshot.Reset();
+
+	FMemoryWriter Writer(EvaluatorSnapshot);
+	FCameraNodeEvaluatorHierarchy CameraSystemHierarchy(RootEvaluator);
+
+	FCameraNodeEvaluatorSerializeParams SerializeParams;
+	CameraSystemHierarchy.CallSerialize(SerializeParams, Writer);
+
+	{
+		FCameraNodeEvaluationParams NodeParams;
+		NodeParams.Evaluator = this;
+		NodeParams.DeltaTime = Params.DeltaTime;
+		NodeParams.EvaluationType = ECameraNodeEvaluationType::ViewRotationPreview;
+
+		RootNodeResult.Reset();
+
+		RootEvaluator->Run(NodeParams, RootNodeResult);
+
+		const FRotator3d& PreviewRotation = RootNodeResult.CameraPose.GetRotation();
+		OutResult.DeltaRotation += (PreviewRotation - OutResult.ViewRotation).GetNormalized();
+	}
+
+	FMemoryReader Reader(EvaluatorSnapshot);
+	CameraSystemHierarchy.CallSerialize(SerializeParams, Reader);
+}
+
 void FCameraSystemEvaluator::GetEvaluatedCameraView(FMinimalViewInfo& DesiredView)
 {
 	const FCameraPose& CameraPose = RootNodeResult.CameraPose;
diff --git a/Engine/Plugins/Cameras/GameplayCameras/Source/GameplayCameras/Private/Core/CameraVariableTable.cpp b/Engine/Plugins/Cameras/GameplayCameras/Source/GameplayCameras/Private/Core/CameraVariableTable.cpp
--- a/Engine/Plugins/Cameras/GameplayCameras/Source/GameplayCameras/Private/Core/CameraVariableTable.cpp
+++ b/Engine/Plugins/Cameras/GameplayCameras/Source/GameplayCameras/Private/Core/CameraVariableTable.cpp
@@ -444,7 +444,7 @@ void FCameraVariableTable::Serialize(FArchive& Ar)
 	{
 		Ar << Capacity;
 		Ar << Used;
-		Ar.Serialize(Memory, Capacity);
+		Ar.Serialize(Memory, Used);
 
 		int32 NumEntries = Entries.Num();
 		Ar << NumEntries;
@@ -455,7 +455,7 @@ void FCameraVariableTable::Serialize(FArchive& Ar)
 	{
 		// Load the data from the saved buffer. It should have a lesser-or-equal size to our buffer,
 		// because we might have grown since we were saved, but we wouldn't have shrunk (variable tables
-		// don't shrink). So we keep our present capacity and ignore the saved on, except to confirm
+		// don't shrink). So we keep our present capacity and ignore the saved one, except to confirm
 		// our expectations.
 		uint32 LoadedCapacity = 0;
 		Ar << LoadedCapacity;
@@ -465,8 +465,8 @@ void FCameraVariableTable::Serialize(FArchive& Ar)
 
 		ensureMsgf(LoadedCapacity <= Capacity, 
 				TEXT("Loading camera variable table data from a bigger, and probably different, table."));
-		LoadedCapacity = FMath::Min(LoadedCapacity, Capacity);
-		Ar.Serialize(Memory, LoadedCapacity);
+
+		Ar.Serialize(Memory, LoadedUsed);
 		Used = LoadedUsed;
 
 		// The number of entries may have changed, if new entries were created since we were saved.
diff --git a/Engine/Plugins/Cameras/GameplayCameras/Source/GameplayCameras/Private/Core/DefaultRootCameraNode.cpp b/Engine/Plugins/Cameras/GameplayCameras/Source/GameplayCameras/Private/Core/DefaultRootCameraNode.cpp
--- a/Engine/Plugins/Cameras/GameplayCameras/Source/GameplayCameras/Private/Core/DefaultRootCameraNode.cpp
+++ b/Engine/Plugins/Cameras/GameplayCameras/Source/GameplayCameras/Private/Core/DefaultRootCameraNode.cpp
@@ -79,7 +79,11 @@ void FDefaultRootCameraNodeEvaluator::OnRun(const FCameraNodeEvaluationParams& P
 	BaseLayer->Run(Params, OutResult);
 	MainLayer->Run(Params, OutResult);
 	GlobalLayer->Run(Params, OutResult);
-	VisualLayer->Run(Params, OutResult);
+	if (Params.EvaluationType != ECameraNodeEvaluationType::IK && 
+			Params.EvaluationType != ECameraNodeEvaluationType::ViewRotationPreview)
+	{
+		VisualLayer->Run(Params, OutResult);
+	}
 }
 
 FCameraRigInstanceID FDefaultRootCameraNodeEvaluator::OnActivateCameraRig(const FActivateCameraRigParams& Params)
diff --git a/Engine/Plugins/Cameras/GameplayCameras/Source/GameplayCameras/Private/Core/PersistentBlendStackCameraNode.cpp b/Engine/Plugins/Cameras/GameplayCameras/Source/GameplayCameras/Private/Core/PersistentBlendStackCameraNode.cpp
--- a/Engine/Plugins/Cameras/GameplayCameras/Source/GameplayCameras/Private/Core/PersistentBlendStackCameraNode.cpp
+++ b/Engine/Plugins/Cameras/GameplayCameras/Source/GameplayCameras/Private/Core/PersistentBlendStackCameraNode.cpp
@@ -27,7 +27,7 @@ FBlendStackEntryID FPersistentBlendStackCameraNodeEvaluator::Insert(const FBlend
 			const FCameraRigEntry& Entry(Entries[Index]);
 			const FCameraRigEntryExtraInfo& EntryExtraInfo(EntryExtraInfos[Index]);
 
-			if (!Entry.bIsFrozen &&
+			if (!Entry.Flags.bIsFrozen &&
 					Entry.CameraRig == Params.CameraRig &&
 					Entry.EvaluationContext == Params.EvaluationContext &&
 					EntryExtraInfo.StackOrder == Params.StackOrder)
@@ -233,11 +233,11 @@ void FPersistentBlendStackCameraNodeEvaluator::InternalUpdate(TArrayView<FResolv
 		FCameraRigEntry& Entry(ResolvedEntry.Entry);
 		FCameraRigEntryExtraInfo& EntryExtraInfo(EntryExtraInfos[ResolvedEntry.EntryIndex]);
 
-		if (!Entry.bIsFrozen)
+		if (!Entry.Flags.bIsFrozen)
 		{
 			FCameraNodeEvaluationParams CurParams(Params);
 			CurParams.EvaluationContext = ResolvedEntry.Context;
-			CurParams.bIsFirstFrame = Entry.bIsFirstFrame;
+			CurParams.bIsFirstFrame = Entry.Flags.bIsFirstFrame;
 
 			FCameraNodeEvaluationResult& CurResult(Entry.Result);
 
@@ -258,7 +258,7 @@ void FPersistentBlendStackCameraNodeEvaluator::InternalUpdate(TArrayView<FResolv
 				CurResult.ContextDataTable.OverrideAll(ContextResult.ContextDataTable);
 
 				// Setup flags.
-				CurResult.bIsCameraCut = OutResult.bIsCameraCut || ContextResult.bIsCameraCut || Entry.bForceCameraCut;
+				CurResult.bIsCameraCut = OutResult.bIsCameraCut || ContextResult.bIsCameraCut || Entry.Flags.bForceCameraCut;
 				CurResult.bIsValid = true;
 			}
 
@@ -342,10 +342,13 @@ void FPersistentBlendStackCameraNodeEvaluator::InternalUpdate(TArrayView<FResolv
 		}
 	}
 
-	for (int32 Index = EntriesToRemove.Num() - 1; Index >= 0; --Index)
+	if (!Params.IsStatelessEvaluation())
 	{
-		PopEntry(EntriesToRemove[Index]);
-		EntryExtraInfos.RemoveAt(EntriesToRemove[Index]);
+		for (int32 Index = EntriesToRemove.Num() - 1; Index >= 0; --Index)
+		{
+			PopEntry(EntriesToRemove[Index]);
+			EntryExtraInfos.RemoveAt(EntriesToRemove[Index]);
+		}
 	}
 }
 
@@ -375,7 +378,7 @@ const UCameraRigTransition* FPersistentBlendStackCameraNodeEvaluator::FindExitTr
 	// Find a transition that works for blending the given camera rig out.
 	return FCameraRigTransitionFinder::FindTransition(
 			Entry.CameraRig->ExitTransitions,
-			Entry.CameraRig, nullptr, Entry.bIsFrozen,
+			Entry.CameraRig, nullptr, Entry.Flags.bIsFrozen,
 			nullptr, nullptr);
 }
 
diff --git a/Engine/Plugins/Cameras/GameplayCameras/Source/GameplayCameras/Private/Core/TransientBlendStackCameraNode.cpp b/Engine/Plugins/Cameras/GameplayCameras/Source/GameplayCameras/Private/Core/TransientBlendStackCameraNode.cpp
--- a/Engine/Plugins/Cameras/GameplayCameras/Source/GameplayCameras/Private/Core/TransientBlendStackCameraNode.cpp
+++ b/Engine/Plugins/Cameras/GameplayCameras/Source/GameplayCameras/Private/Core/TransientBlendStackCameraNode.cpp
@@ -29,7 +29,7 @@ FBlendStackEntryID FTransientBlendStackCameraNodeEvaluator::Push(const FBlendSta
 	if (!Entries.IsEmpty())
 	{
 		FCameraRigEntry& TopEntry(Entries.Top());
-		if (!TopEntry.bIsFrozen 
+		if (!TopEntry.Flags.bIsFrozen 
 				&& TopEntry.EvaluationContext == Params.EvaluationContext)
 		{
 			// Don't push anything if what is being requested is already the active 
@@ -177,7 +177,7 @@ void FTransientBlendStackCameraNodeEvaluator::Freeze(const FBlendStackCameraFree
 		// Freeze any entries matching the given context and rig asset.
 		for (FCameraRigEntry& Entry : Entries)
 		{
-			if (!Entry.bIsFrozen && 
+			if (!Entry.Flags.bIsFrozen && 
 					Entry.CameraRig == Params.CameraRig &&
 					Entry.EvaluationContext == Params.EvaluationContext)
 			{
@@ -191,7 +191,7 @@ void FTransientBlendStackCameraNodeEvaluator::FreezeAll(TSharedPtr<const FCamera
 {
 	for (FCameraRigEntry& Entry : Entries)
 	{
-		if (!Entry.bIsFrozen && Entry.EvaluationContext == EvaluationContext)
+		if (!Entry.Flags.bIsFrozen && Entry.EvaluationContext == EvaluationContext)
 		{
 			FreezeEntry(Entry);
 		}
@@ -243,14 +243,14 @@ void FTransientBlendStackCameraNodeEvaluator::InternalPreBlendPrepare(TArrayView
 		FCameraRigEntry& Entry(ResolvedEntry.Entry);
 		FCameraRigEntryExtraInfo& EntryExtraInfo(EntryExtraInfos[ResolvedEntry.EntryIndex]);
 
-		if (UNLIKELY(Entry.bIsFrozen))
+		if (UNLIKELY(Entry.Flags.bIsFrozen))
 		{
 			continue;
 		}
 
 		FCameraNodeEvaluationParams CurParams(Params);
 		CurParams.EvaluationContext = ResolvedEntry.Context;
-		CurParams.bIsFirstFrame = Entry.bIsFirstFrame;
+		CurParams.bIsFirstFrame = Entry.Flags.bIsFirstFrame;
 
 		FCameraNodeEvaluationResult& CurResult(Entry.Result);
 
@@ -323,11 +323,11 @@ void FTransientBlendStackCameraNodeEvaluator::InternalPreBlendExecute(TArrayView
 
 		FCameraNodeEvaluationResult& CurResult(Entry.Result);
 
-		if (!Entry.bIsFrozen)
+		if (!Entry.Flags.bIsFrozen)
 		{
 			FCameraNodeEvaluationParams CurParams(Params);
 			CurParams.EvaluationContext = ResolvedEntry.Context;
-			CurParams.bIsFirstFrame = Entry.bIsFirstFrame;
+			CurParams.bIsFirstFrame = Entry.Flags.bIsFirstFrame;
 
 			FCameraNodePreBlendParams PreBlendParams(CurParams, CurResult.CameraPose, CurResult.VariableTable);
 			PreBlendParams.VariableTableFilter = VariableTableFilter;
@@ -359,7 +359,7 @@ void FTransientBlendStackCameraNodeEvaluator::InternalPreBlendExecute(TArrayView
 	for (FResolvedEntry& ResolvedEntry : ResolvedEntries)
 	{
 		FCameraRigEntry& Entry(ResolvedEntry.Entry);
-		if (!Entry.bIsFrozen)
+		if (!Entry.Flags.bIsFrozen)
 		{
 			FCameraNodeEvaluationResult& CurResult(Entry.Result);
 			CurResult.VariableTable.Override(PreBlendVariableTable, ECameraVariableTableFilter::KnownOnly);
@@ -374,7 +374,7 @@ void FTransientBlendStackCameraNodeEvaluator::InternalUpdate(TArrayView<FResolve
 		FCameraRigEntry& Entry(ResolvedEntry.Entry);
 		FCameraRigEntryExtraInfo& EntryExtraInfo(EntryExtraInfos[ResolvedEntry.EntryIndex]);
 
-		if (UNLIKELY(Entry.bIsFrozen))
+		if (UNLIKELY(Entry.Flags.bIsFrozen))
 		{
 			continue;
 		}
@@ -383,7 +383,7 @@ void FTransientBlendStackCameraNodeEvaluator::InternalUpdate(TArrayView<FResolve
 
 		FCameraNodeEvaluationParams CurParams(Params);
 		CurParams.EvaluationContext = ResolvedEntry.Context;
-		CurParams.bIsFirstFrame = Entry.bIsFirstFrame;
+		CurParams.bIsFirstFrame = Entry.Flags.bIsFirstFrame;
 
 		// Start with the input given to us.
 		CurResult.Reset();
@@ -394,7 +394,7 @@ void FTransientBlendStackCameraNodeEvaluator::InternalUpdate(TArrayView<FResolve
 		// Override it with whatever the evaluation context has set on its result.
 		const FCameraNodeEvaluationResult& ContextResult(Entry.ContextResult);
 		CurResult.CameraPose.OverrideChanged(ContextResult.CameraPose);
-		CurResult.bIsCameraCut = OutResult.bIsCameraCut || ContextResult.bIsCameraCut || Entry.bForceCameraCut;
+		CurResult.bIsCameraCut = OutResult.bIsCameraCut || ContextResult.bIsCameraCut || Entry.Flags.bForceCameraCut;
 		
 		CurResult.bIsValid = true;
 
@@ -421,11 +421,11 @@ void FTransientBlendStackCameraNodeEvaluator::InternalPostBlendExecute(TArrayVie
 
 		FCameraNodeEvaluationResult& CurResult(Entry.Result);
 
-		if (!Entry.bIsFrozen)
+		if (!Entry.Flags.bIsFrozen)
 		{
 			FCameraNodeEvaluationParams CurParams(Params);
 			CurParams.EvaluationContext = ResolvedEntry.Context;
-			CurParams.bIsFirstFrame = Entry.bIsFirstFrame;
+			CurParams.bIsFirstFrame = Entry.Flags.bIsFirstFrame;
 			FCameraNodeBlendParams BlendParams(CurParams, CurResult);
 
 			FCameraNodeBlendResult BlendResult(OutResult);
@@ -457,12 +457,15 @@ void FTransientBlendStackCameraNodeEvaluator::InternalPostBlendExecute(TArrayVie
 	}
 
 	// Pop out camera rigs that have been blended out.
-	const UBlendStackCameraNode* BlendStackNode = GetCameraNodeAs<UBlendStackCameraNode>();
-	if (BlendStackNode->BlendStackType == ECameraBlendStackType::IsolatedTransient && PopEntriesBelow != INDEX_NONE)
+	if (!Params.IsStatelessEvaluation())
 	{
-		PopEntries(PopEntriesBelow);
+		const UBlendStackCameraNode* BlendStackNode = GetCameraNodeAs<UBlendStackCameraNode>();
+		if (BlendStackNode->BlendStackType == ECameraBlendStackType::IsolatedTransient && PopEntriesBelow != INDEX_NONE)
+		{
+			PopEntries(PopEntriesBelow);
 
-		EntryExtraInfos.RemoveAt(0, PopEntriesBelow);
+			EntryExtraInfos.RemoveAt(0, PopEntriesBelow);
+		}
 	}
 }
 
@@ -521,7 +524,7 @@ const UCameraRigTransition* FTransientBlendStackCameraNodeEvaluator::FindTransit
 		TArray<const UCameraRigAsset*> FromCombinedCameraRigs;
 		UCombinedCameraRigsCameraNode::GetAllCombinationCameraRigs(TopEntry.CameraRig, FromCombinedCameraRigs);
 
-		const bool bFromFrozen = TopEntry.bIsFrozen;
+		const bool bFromFrozen = TopEntry.Flags.bIsFrozen;
 		const UCameraRigTransition* TransitionToUse = nullptr;
 
 		// Start by looking at exit transitions on the last active (top) camera rig.
diff --git a/Engine/Plugins/Cameras/GameplayCameras/Source/GameplayCameras/Private/GameFramework/ActivateCameraRigFunctions.cpp b/Engine/Plugins/Cameras/GameplayCameras/Source/GameplayCameras/Private/GameFramework/ActivateCameraRigFunctions.cpp
--- a/Engine/Plugins/Cameras/GameplayCameras/Source/GameplayCameras/Private/GameFramework/ActivateCameraRigFunctions.cpp
+++ b/Engine/Plugins/Cameras/GameplayCameras/Source/GameplayCameras/Private/GameFramework/ActivateCameraRigFunctions.cpp
@@ -51,8 +51,13 @@ void UActivateCameraRigFunctions::ActivateCameraRigImpl(UObject* WorldContextObj
 	UControllerGameplayCameraEvaluationComponent* CameraEvaluationComponent = nullptr;
 	if (IGameplayCameraSystemHost* FoundHost = IGameplayCameraSystemHost::FindActiveHost(PlayerController))
 	{
-		AActor* HostOwningActor = FoundHost->GetAsObject()->GetTypedOuter<AActor>();
-		check(HostOwningActor);
+		UObject* FoundHostObject = FoundHost->GetAsObject();
+		AActor* HostOwningActor = Cast<AActor>(FoundHostObject);
+		if (!HostOwningActor)
+		{
+			HostOwningActor = FoundHostObject->GetTypedOuter<AActor>();
+		}
+		ensure(HostOwningActor);
 
 		bool bComponentCreated = false;
 		CameraEvaluationComponent = UControllerGameplayCameraEvaluationComponent::FindOrAddComponent(HostOwningActor, &bComponentCreated);
diff --git a/Engine/Plugins/Cameras/GameplayCameras/Source/GameplayCameras/Private/GameFramework/GameplayCamerasPlayerCameraManager.cpp b/Engine/Plugins/Cameras/GameplayCameras/Source/GameplayCameras/Private/GameFramework/GameplayCamerasPlayerCameraManager.cpp
--- a/Engine/Plugins/Cameras/GameplayCameras/Source/GameplayCameras/Private/GameFramework/GameplayCamerasPlayerCameraManager.cpp
+++ b/Engine/Plugins/Cameras/GameplayCameras/Source/GameplayCameras/Private/GameFramework/GameplayCamerasPlayerCameraManager.cpp
@@ -15,6 +15,7 @@
 #include "GameFramework/ActorCameraEvaluationContext.h"
 #include "GameFramework/GameplayCameraComponentBase.h"
 #include "GameFramework/PlayerController.h"
+#include "GameplayCamerasSettings.h"
 #include "Services/CameraModifierService.h"
 #include "UObject/Package.h"
 #include "UObject/UObjectGlobals.h"
@@ -211,6 +212,12 @@ void AGameplayCamerasPlayerCameraManager::StopCameraModifierRig(FCameraRigInstan
 
 void AGameplayCamerasPlayerCameraManager::InitializeFor(APlayerController* PlayerController)
 {
+	if (!bOverrideViewRotationMode)
+	{
+		const UGameplayCamerasSettings* Settings = GetDefault<UGameplayCamerasSettings>();
+		ViewRotationMode = Settings->DefaultViewRotationMode;
+	}
+
 	EnsureCameraSystemHost();
 
 	Super::InitializeFor(PlayerController);
@@ -313,9 +320,36 @@ void AGameplayCamerasPlayerCameraManager::TeardownCameraSystemHost()
 
 void AGameplayCamerasPlayerCameraManager::ProcessViewRotation(float DeltaTime, FRotator& OutViewRotation, FRotator& OutDeltaRot)
 {
+	switch (ViewRotationMode)
+	{
+		case EGameplayCamerasViewRotationMode::PreviewUpdate:
+			RunViewRotationPreviewUpdate(DeltaTime, OutViewRotation, OutDeltaRot);
+			break;
+	}
+
 	Super::ProcessViewRotation(DeltaTime, OutViewRotation, OutDeltaRot);
 }
 
+void AGameplayCamerasPlayerCameraManager::RunViewRotationPreviewUpdate(float DeltaTime, FRotator& OutViewRotation, FRotator& OutDeltaRot)
+{
+	using namespace UE::Cameras;
+
+	if (HasCameraSystem())
+	{
+		FCameraSystemEvaluationParams Params;
+		Params.DeltaTime = DeltaTime;
+
+		FCameraSystemViewRotationEvaluationResult Result;
+		Result.ViewRotation = OutViewRotation;
+		Result.DeltaRotation = OutDeltaRot;
+
+		CameraSystemEvaluator->ViewRotationPreviewUpdate(Params, Result);
+
+		OutViewRotation = Result.ViewRotation;
+		OutDeltaRot = Result.DeltaRotation;
+	}
+}
+
 void AGameplayCamerasPlayerCameraManager::DoUpdateCamera(float DeltaTime)
 {
 	using namespace UE::Cameras;
diff --git a/Engine/Plugins/Cameras/GameplayCameras/Source/GameplayCameras/Public/Core/BlendStackCameraNode.h b/Engine/Plugins/Cameras/GameplayCameras/Source/GameplayCameras/Public/Core/BlendStackCameraNode.h
--- a/Engine/Plugins/Cameras/GameplayCameras/Source/GameplayCameras/Public/Core/BlendStackCameraNode.h
+++ b/Engine/Plugins/Cameras/GameplayCameras/Source/GameplayCameras/Public/Core/BlendStackCameraNode.h
@@ -210,19 +210,22 @@ protected:
 		FCameraNodeEvaluationResult ContextResult;
 		/** Evaluation result for this node tree. */
 		FCameraNodeEvaluationResult Result;
-		/** Whether this is the first frame this entry runs. */
-		bool bIsFirstFrame = false;
-		/** Whether the context's initial result was valid last frame. */
-		bool bWasContextInitialResultValid = false;
-		/** Whether to force a camera cut on this entry this frame. */
-		bool bForceCameraCut = false;
 
-		/** Whether this entry is frozen. */
-		bool bIsFrozen = false;
+		struct
+		{
+			/** Whether this is the first frame this entry runs. */
+			bool bIsFirstFrame:1 = false;
+			/** Whether the context's initial result was valid last frame. */
+			bool bWasContextInitialResultValid:1 = false;
+			/** Whether to force a camera cut on this entry this frame. */
+			bool bForceCameraCut:1 = false;
+			/** Whether this entry is frozen. */
+			bool bIsFrozen:1 = false;
 
 #if UE_GAMEPLAY_CAMERAS_TRACE
-		bool bLogWarnings = true;
+			bool bLogWarnings:1 = true;
 #endif  // UE_GAMEPLAY_CAMERAS_TRACE
+		} Flags;
 
 #if WITH_EDITOR
 		TArray<TWeakObjectPtr<const UPackage>, TInlineAllocator<4>> ListenedPackages;
diff --git a/Engine/Plugins/Cameras/GameplayCameras/Source/GameplayCameras/Public/Core/CameraNodeEvaluator.h b/Engine/Plugins/Cameras/GameplayCameras/Source/GameplayCameras/Public/Core/CameraNodeEvaluator.h
--- a/Engine/Plugins/Cameras/GameplayCameras/Source/GameplayCameras/Public/Core/CameraNodeEvaluator.h
+++ b/Engine/Plugins/Cameras/GameplayCameras/Source/GameplayCameras/Public/Core/CameraNodeEvaluator.h
@@ -15,8 +15,6 @@
 #include "GameplayCameras.h"
 #include "UObject/ObjectPtr.h"
 
-#define UE_API GAMEPLAYCAMERAS_API
-
 class FReferenceCollector;
 class UCameraNode;
 class UCameraRigAsset;
@@ -135,6 +133,8 @@ enum class ECameraNodeEvaluationType
 {
 	/** Normal evaluation. */
 	Standard,
+	/** View rotation evaluation. */
+	ViewRotationPreview,
 	/** Evaluation for IK aiming. */
 	IK,
 
@@ -159,12 +159,18 @@ struct FCameraNodeEvaluationParams
 	ECameraNodeEvaluationType EvaluationType = ECameraNodeEvaluationType::Standard;
 	/** Whether this is the first evaluation of this camera node hierarchy. */
 	bool bIsFirstFrame = false;
+
+	bool IsStatelessEvaluation() const
+	{
+		return EvaluationType == ECameraNodeEvaluationType::IK ||
+			EvaluationType == ECameraNodeEvaluationType::ViewRotationPreview;
+	}
 };
 
 /**
  * Input/output result structure for running a camera node evaluator.
  */
-struct FCameraNodeEvaluationResult
+struct GAMEPLAYCAMERAS_API FCameraNodeEvaluationResult
 {
 	/** The camera pose. */
 	FCameraPose CameraPose;
@@ -190,24 +196,24 @@ struct FCameraNodeEvaluationResult
 public:
 
 	/** Reset this result to its default (non-valid) state.  */
-	UE_API void Reset();
+	void Reset();
 
 	/** Reset all written-this-frame flags on the camera pose and tables. */
-	UE_API void ResetFrameFlags();
+	void ResetFrameFlags();
 
 	/** Override this result with the given other result. */
-	UE_API void OverrideAll(const FCameraNodeEvaluationResult& OtherResult, bool bIncludePrivateValues = false);
+	void OverrideAll(const FCameraNodeEvaluationResult& OtherResult, bool bIncludePrivateValues = false);
 
 	/** Interpolate this result towards the other given result. */
-	UE_API void LerpAll(const FCameraNodeEvaluationResult& ToResult, float BlendFactor, bool bIncludePrivateValues = false);
+	void LerpAll(const FCameraNodeEvaluationResult& ToResult, float BlendFactor, bool bIncludePrivateValues = false);
 
 	/** Serializes this result to the given archive. */
-	UE_API void Serialize(FArchive& Ar);
+	void Serialize(FArchive& Ar);
 
 public:
 
 	/** Collects objects from the context data table. */
-	UE_API void AddReferencedObjects(FReferenceCollector& Collector);
+	void AddReferencedObjects(FReferenceCollector& Collector);
 
 #if WITH_EDITOR || UE_GAMEPLAY_CAMERAS_DEBUG
 
@@ -215,10 +221,10 @@ public:
 
 	// Internal API.
 
-	UE_API void AddCameraPoseTrailPointIfNeeded();
-	UE_API void AddCameraPoseTrailPointIfNeeded(const FVector3d& Point);
-	UE_API void AppendCameraPoseLocationTrail(const FCameraNodeEvaluationResult& InResult);
-	UE_API TConstArrayView<FVector3d> GetCameraPoseLocationTrail() const;
+	void AddCameraPoseTrailPointIfNeeded();
+	void AddCameraPoseTrailPointIfNeeded(const FVector3d& Point);
+	void AppendCameraPoseLocationTrail(const FCameraNodeEvaluationResult& InResult);
+	TConstArrayView<FVector3d> GetCameraPoseLocationTrail() const;
 
 private:
 
@@ -241,7 +247,7 @@ struct FCameraOperationParams
 /**
  * Parameter structure for serializing the state of a camera node evaluator.
  */
-struct FCameraNodeEvaluatorSerializeParams
+struct GAMEPLAYCAMERAS_API FCameraNodeEvaluatorSerializeParams
 {
 };
 
@@ -265,8 +271,8 @@ class FCameraNodeEvaluator
 
 public:
 
-	FCameraNodeEvaluator() = default;
-	virtual ~FCameraNodeEvaluator() = default;
+	GAMEPLAYCAMERAS_API FCameraNodeEvaluator() = default;
+	GAMEPLAYCAMERAS_API virtual ~FCameraNodeEvaluator() = default;
 
 	/** Called to build any children evaluators. */
 	void Build(const FCameraNodeEvaluatorBuildParams& Params);
@@ -330,42 +336,42 @@ protected:
 protected:
 
 	/** Called to build any children evaluators. */
-	virtual void OnBuild(const FCameraNodeEvaluatorBuildParams& Params) {}
+	GAMEPLAYCAMERAS_API virtual void OnBuild(const FCameraNodeEvaluatorBuildParams& Params) {}
 
 	/** Initialize this evaluator. Children and descendants will be automatically initialized too. */
-	virtual void OnInitialize(const FCameraNodeEvaluatorInitializeParams& Params, FCameraNodeEvaluationResult& OutResult) {}
+	GAMEPLAYCAMERAS_API virtual void OnInitialize(const FCameraNodeEvaluatorInitializeParams& Params, FCameraNodeEvaluationResult& OutResult) {}
 
 	/** Collect referenced UObjects for this node. */
-	virtual void OnAddReferencedObjects(FReferenceCollector& Collector) {}
+	GAMEPLAYCAMERAS_API virtual void OnAddReferencedObjects(FReferenceCollector& Collector) {}
 
 	/** Get the list of children under this evaluator. */
-	virtual FCameraNodeEvaluatorChildrenView OnGetChildren() { return FCameraNodeEvaluatorChildrenView(); }
+	GAMEPLAYCAMERAS_API virtual FCameraNodeEvaluatorChildrenView OnGetChildren() { return FCameraNodeEvaluatorChildrenView(); }
 
 	/**
 	 * Called to update and store the blended parameters for this node.
 	 * Requires setting the ECameraNodeEvaluatorFlags::NeedsParameterUpdate flag.
 	 */
-	virtual void OnUpdateParameters(const FCameraBlendedParameterUpdateParams& Params, FCameraBlendedParameterUpdateResult& OutResult) {}
+	GAMEPLAYCAMERAS_API virtual void OnUpdateParameters(const FCameraBlendedParameterUpdateParams& Params, FCameraBlendedParameterUpdateResult& OutResult) {}
 
 	/**
 	 * Run this evaluator. This node evaluator is responsible for calling Run() on its children as appropriate.
 	 */
-	virtual void OnRun(const FCameraNodeEvaluationParams& Params, FCameraNodeEvaluationResult& OutResult) {}
+	GAMEPLAYCAMERAS_API virtual void OnRun(const FCameraNodeEvaluationParams& Params, FCameraNodeEvaluationResult& OutResult) {}
 
 	/** 
 	 * Execute an IK operation.
 	 * Requires setting the ECameraNodeEvaluatorFlags::SupportsOperations flag.
 	 */
-	virtual void OnExecuteOperation(const FCameraOperationParams& Params, FCameraOperation& Operation) {}
+	GAMEPLAYCAMERAS_API virtual void OnExecuteOperation(const FCameraOperationParams& Params, FCameraOperation& Operation) {}
 
 	/**
 	 * Serializes the state of this evaluator.
 	 * Requires setting the ECameraNodeEvaluatorFlags::NeedsSerialize flag, which is set by default.
 	 */
-	virtual void OnSerialize(const FCameraNodeEvaluatorSerializeParams& Params, FArchive& Ar) {}
+	GAMEPLAYCAMERAS_API virtual void OnSerialize(const FCameraNodeEvaluatorSerializeParams& Params, FArchive& Ar) {}
 
 #if WITH_EDITOR
-	virtual void OnDrawEditorPreview(const FCameraEditorPreviewDrawParams& Params, FCameraDebugRenderer& Renderer) {}
+	GAMEPLAYCAMERAS_API virtual void OnDrawEditorPreview(const FCameraEditorPreviewDrawParams& Params, FCameraDebugRenderer& Renderer) {}
 #endif  // WITH_EDITOR
 
 #if UE_GAMEPLAY_CAMERAS_DEBUG
@@ -429,4 +435,3 @@ EvaluatorType* FCameraNodeEvaluatorBuildParams::BuildEvaluatorAs(const UCameraNo
 #define UE_DEFINE_CAMERA_NODE_EVALUATOR(ClassName)\
 	UE_GAMEPLAY_CAMERAS_DEFINE_RTTI(ClassName)
 
-#undef UE_API
diff --git a/Engine/Plugins/Cameras/GameplayCameras/Source/GameplayCameras/Public/Core/CameraPose.h b/Engine/Plugins/Cameras/GameplayCameras/Source/GameplayCameras/Public/Core/CameraPose.h
--- a/Engine/Plugins/Cameras/GameplayCameras/Source/GameplayCameras/Public/Core/CameraPose.h
+++ b/Engine/Plugins/Cameras/GameplayCameras/Source/GameplayCameras/Public/Core/CameraPose.h
@@ -11,8 +11,6 @@
 
 #include "CameraPose.generated.h"
 
-#define UE_API GAMEPLAYCAMERAS_API
-
 class FArchive;
 struct FPostProcessSettings;
 enum EAspectRatioAxisConstraint : int;
@@ -59,10 +57,10 @@ enum EAspectRatioAxisConstraint : int;
 /**
  * Boolean flags for each of the properties inside FCameraPose.
  */
-struct FCameraPoseFlags
+struct GAMEPLAYCAMERAS_API FCameraPoseFlags
 {
 #define UE_CAMERA_POSE_FOR_PROPERTY(PropType, PropName)\
-	bool PropName = false;
+	bool PropName:1 = false;
 
 UE_CAMERA_POSE_FOR_ALL_PROPERTIES()
 
@@ -71,22 +69,22 @@ UE_CAMERA_POSE_FOR_ALL_PROPERTIES()
 public:
 
 	/** Returns a structure where all flags are set. */
-	static UE_API const FCameraPoseFlags& All();
+	static const FCameraPoseFlags& All();
 
 	/** Creates a new flags structure. */
-	UE_API FCameraPoseFlags();
+	FCameraPoseFlags();
 	/** Creates a new flags structure with all flags set to the given value. */
-	UE_API FCameraPoseFlags(bool bInValue);
+	FCameraPoseFlags(bool bInValue);
 
 	/** Sets all flags to the given value. */
-	UE_API FCameraPoseFlags& SetAllFlags(bool bInValue);
+	FCameraPoseFlags& SetAllFlags(bool bInValue);
 	/** Sets the flags that are set in OtherFlags, but checks that no flag is set on both structures. */
-	UE_API FCameraPoseFlags& ExclusiveCombine(const FCameraPoseFlags& OtherFlags);
+	FCameraPoseFlags& ExclusiveCombine(const FCameraPoseFlags& OtherFlags);
 
 	/** Combines the flags with an AND logical operation. */
-	UE_API FCameraPoseFlags& AND(const FCameraPoseFlags& OtherFlags);
+	FCameraPoseFlags& AND(const FCameraPoseFlags& OtherFlags);
 	/** Combines the flags with an OR logical operation. */
-	UE_API FCameraPoseFlags& OR(const FCameraPoseFlags& OtherFlags);
+	FCameraPoseFlags& OR(const FCameraPoseFlags& OtherFlags);
 };
 
 /**
@@ -96,16 +94,16 @@ public:
  * The ChangedFlags structure keeps track of which fields were changed via the setters.
  */
 USTRUCT()
-struct FCameraPose
+struct GAMEPLAYCAMERAS_API FCameraPose
 {
 	GENERATED_BODY()
 
 public:
 
-	UE_API FCameraPose();
+	FCameraPose();
 
 	/** Resets this camera pose to its default values, with all changed flags off. */
-	UE_API void Reset();
+	void Reset();
 
 public:
 
@@ -141,58 +139,58 @@ public:
 	/** Set the changed flags. */
 	void SetChangedFlags(const FCameraPoseFlags& InChangedFlags) { ChangedFlags = InChangedFlags; }
 	/** Set all fields as changed. */
-	UE_API void SetAllChangedFlags();
+	void SetAllChangedFlags();
 	/** Set all fields as clean. */
-	UE_API void ClearAllChangedFlags();
+	void ClearAllChangedFlags();
 
 public:
 
 	// Utility
 
 	/** Gets the transform of the camera. */
-	UE_API FTransform3d GetTransform() const;
+	FTransform3d GetTransform() const;
 
 	/** Sets the transform of the camera. */
-	UE_API void SetTransform(FTransform3d Transform, bool bForceSet = false);
+	void SetTransform(FTransform3d Transform, bool bForceSet = false);
 
 	/**
 	 * Computes the horizontal field of view of the camera.
 	 * The effective field of view can be driven by the FieldOfView property, or
 	 * the FocalLength property in combination with the sensor size.
 	 */
-	UE_API double GetEffectiveFieldOfView(bool bIncludeOverscan = true) const;
+	double GetEffectiveFieldOfView(bool bIncludeOverscan = true) const;
 
 	/** Gets the aspect ratio of the camera sensor. */
-	UE_API double GetSensorAspectRatio() const;
+	double GetSensorAspectRatio() const;
 
 	/** Gets the horizontal projection offset. */
-	UE_API double GetHorizontalProjectionOffset() const;
+	double GetHorizontalProjectionOffset() const;
 
 	/** Gets the vertical projection offset. */
-	UE_API double GetVerticalProjectionOffset() const;
+	double GetVerticalProjectionOffset() const;
 
 	/** Gets the aiming ray of the camera. */
-	UE_API FRay3d GetAimRay() const;
+	FRay3d GetAimRay() const;
 
 	/** Gets the aiming direction of the camera. */
-	UE_API FVector3d GetAimDir() const;
+	FVector3d GetAimDir() const;
 
 	/** Gets the location of the camera's target. */
-	UE_API FVector3d GetTarget() const;
+	FVector3d GetTarget() const;
 
 	/** Gets the location of the camera's target given a specific distance. */
-	UE_API FVector3d GetTarget(double InTargetDistance) const;
+	FVector3d GetTarget(double InTargetDistance) const;
 
 public:
 
 	/** Computes the horizontal field of view of a camera. */
-	static UE_API double GetEffectiveFieldOfView(float FocalLength, float FieldOfView, float SensorWidth, float SensorHeight, float SqueezeFactor = 1.f, float Overscan = 1.f);
+	static double GetEffectiveFieldOfView(float FocalLength, float FieldOfView, float SensorWidth, float SensorHeight, float SqueezeFactor = 1.f, float Overscan = 1.f);
 	
 	/** Computes the aspect ratio of a camera sensor. */
-	static UE_API double GetSensorAspectRatio(float SensorWidth, float SensorHeight);
+	static double GetSensorAspectRatio(float SensorWidth, float SensorHeight);
 
 	/** Gets the default sensor size. */
-	static UE_API void GetDefaultSensorSize(float& OutSensorWidth, float& OutSensorHeight);
+	static void GetDefaultSensorSize(float& OutSensorWidth, float& OutSensorHeight);
 
 	/**
 	 * Applies the necessary post-process settings given the current values
@@ -205,35 +203,35 @@ public:
 	 * @param bOverwriteSettings   Whether to overwrite values found to already be set
 	 * @return  Whether post-process settings were created.
 	 */
-	UE_API bool ApplyPhysicalCameraSettings(FPostProcessSettings& PostProcessSettings, bool bOverwriteSettings = false) const;
+	bool ApplyPhysicalCameraSettings(FPostProcessSettings& PostProcessSettings, bool bOverwriteSettings = false) const;
 
 public:
 
 	// Interpolation
 	
 	/** Takes all properties from OtherPose and sets them on this camera pose. */
-	UE_API void OverrideAll(const FCameraPose& OtherPose);
+	void OverrideAll(const FCameraPose& OtherPose);
 	/** Takes all changed properties from OtherPose and sets them on this camera pose. */
-	UE_API void OverrideChanged(const FCameraPose& OtherPose);
+	void OverrideChanged(const FCameraPose& OtherPose);
 	/** Interpolates all properties from ToPose using the given factor. */
-	UE_API void LerpAll(const FCameraPose& ToPose, float Factor);
+	void LerpAll(const FCameraPose& ToPose, float Factor);
 	/** Interpolates all changed properties from ToPose using the given factor. */
-	UE_API void LerpChanged(const FCameraPose& ToPose, float Factor);
+	void LerpChanged(const FCameraPose& ToPose, float Factor);
 	/** Interpolates changed properties from ToPose using the given factor. Only properties defined by InMask are taken into account. */
-	UE_API void LerpChanged(const FCameraPose& ToPose, float Factor, const FCameraPoseFlags& InMask, bool bInvertMask, FCameraPoseFlags& OutMask);
+	void LerpChanged(const FCameraPose& ToPose, float Factor, const FCameraPoseFlags& InMask, bool bInvertMask, FCameraPoseFlags& OutMask);
 
 public:
 
 	/** Serializes the given camera pose including the written-property flags. */
-	static UE_API void SerializeWithFlags(FArchive& Ar, FCameraPose& CameraPose);
+	static void SerializeWithFlags(FArchive& Ar, FCameraPose& CameraPose);
 
 	/** Serializes this camera pose including the written-property flags. */
-	UE_API void SerializeWithFlags(FArchive& Ar);
+	void SerializeWithFlags(FArchive& Ar);
 
 private:
 
-	UE_API void InternalOverrideChanged(const FCameraPose& OtherPose, bool bChangedOnly);
-	UE_API void InternalLerpChanged(const FCameraPose& ToPose, float Factor, const FCameraPoseFlags& InMask, bool bInvertMask, FCameraPoseFlags& OutMask, bool bChangedOnly);
+	void InternalOverrideChanged(const FCameraPose& OtherPose, bool bChangedOnly);
+	void InternalLerpChanged(const FCameraPose& ToPose, float Factor, const FCameraPoseFlags& InMask, bool bInvertMask, FCameraPoseFlags& OutMask, bool bChangedOnly);
 
 private:
 
@@ -360,4 +358,3 @@ private:
 	FCameraPoseFlags ChangedFlags;
 };
 
-#undef UE_API
diff --git a/Engine/Plugins/Cameras/GameplayCameras/Source/GameplayCameras/Public/Core/CameraSystemEvaluator.h b/Engine/Plugins/Cameras/GameplayCameras/Source/GameplayCameras/Public/Core/CameraSystemEvaluator.h
--- a/Engine/Plugins/Cameras/GameplayCameras/Source/GameplayCameras/Public/Core/CameraSystemEvaluator.h
+++ b/Engine/Plugins/Cameras/GameplayCameras/Source/GameplayCameras/Public/Core/CameraSystemEvaluator.h
@@ -108,6 +108,18 @@ public:
 	void Reset(const FCameraNodeEvaluationResult& NodeResult);
 };
 
+/**
+ * Result structure for view rotation updates of the camera system.
+ */
+struct FCameraSystemViewRotationEvaluationResult
+{
+	/** View rotation for the player. */
+	FRotator ViewRotation;
+
+	/** Control rotation for the player. */
+	FRotator DeltaRotation;
+};
+
 #if UE_GAMEPLAY_CAMERAS_DEBUG
 
 /**
@@ -221,6 +233,9 @@ public:
 	/** Run an update of the camera system. */
 	GAMEPLAYCAMERAS_API void Update(const FCameraSystemEvaluationParams& Params);
 
+	/** Run a view rotation preview update of the camera system. */
+	GAMEPLAYCAMERAS_API void ViewRotationPreviewUpdate(const FCameraSystemEvaluationParams& Params, FCameraSystemViewRotationEvaluationResult& OutResult);
+
 	/** Returns the root node evaluator. */
 	FRootCameraNodeEvaluator* GetRootNodeEvaluator() const { return RootEvaluator; }
 
@@ -292,6 +307,9 @@ private:
 	/** The current overall result of the camera system. */
 	FCameraSystemEvaluationResult Result;
 
+	/** Reusable buffer for node tree snapshots (e.g. view rotation preview update). */
+	TArray<uint8> EvaluatorSnapshot;
+
 	/** The role of this camera system. */
 	ECameraSystemEvaluatorRole Role = ECameraSystemEvaluatorRole::Game;
 
diff --git a/Engine/Plugins/Cameras/GameplayCameras/Source/GameplayCameras/Public/GameFramework/GameplayCamerasPlayerCameraManager.h b/Engine/Plugins/Cameras/GameplayCameras/Source/GameplayCameras/Public/GameFramework/GameplayCamerasPlayerCameraManager.h
--- a/Engine/Plugins/Cameras/GameplayCameras/Source/GameplayCameras/Public/GameFramework/GameplayCamerasPlayerCameraManager.h
+++ b/Engine/Plugins/Cameras/GameplayCameras/Source/GameplayCameras/Public/GameFramework/GameplayCamerasPlayerCameraManager.h
@@ -19,6 +19,27 @@ class FViewTargetContextReferencerService;
 
 }  // namespace UE::Cameras
 
+/**
+ * Defines how the GameplayCamerasPlayerCameraManager should handle the player's view rotation.
+ */
+UENUM()
+enum class EGameplayCamerasViewRotationMode
+{
+	/**
+	 * Don't do anything with the view rotation. This is suitable if the player camera input
+	 * and control rotation are handled by custom logic, and the Gameplay Cameras only use
+	 * the resulting control rotation (i.e. none of the camera rigs use player input nodes).
+	 */
+	None,
+	/**
+	 * Runs a "light" update of the whole camera system and comes up with a good approximation of
+	 * the resulting camera orientation. Computationally more expensive than just computing a
+	 * yaw and pitch angle, but handles blending camera rigs with different orientations while
+	 * preserving aim.
+	 */
+	PreviewUpdate
+};
+
 /**
  * A player camera manager that runs the GameplayCameras camera system.
  *
@@ -95,6 +116,18 @@ private:
 	void EnsureCameraSystemHost();
 	void TeardownCameraSystemHost();
 
+	void RunViewRotationPreviewUpdate(float DeltaTime, FRotator& OutViewRotation, FRotator& OutDeltaRot);
+
+public:
+
+	/** The view rotation handling mode to use. */
+	UPROPERTY(EditAnywhere, Category="Camera", meta=(EditCondition="bOverrideViewRotationMode"))
+	EGameplayCamerasViewRotationMode ViewRotationMode = EGameplayCamerasViewRotationMode::None;
+
+	/** Whether the default view rotation mode setting should be overriden. */
+	UPROPERTY(EditAnywhere, Category="Camera")
+	bool bOverrideViewRotationMode = false;
+
 private:
 
 	UPROPERTY(Transient)
diff --git a/Engine/Plugins/Cameras/GameplayCameras/Source/GameplayCameras/Public/GameplayCamerasSettings.h b/Engine/Plugins/Cameras/GameplayCameras/Source/GameplayCameras/Public/GameplayCamerasSettings.h
--- a/Engine/Plugins/Cameras/GameplayCameras/Source/GameplayCameras/Public/GameplayCamerasSettings.h
+++ b/Engine/Plugins/Cameras/GameplayCameras/Source/GameplayCameras/Public/GameplayCamerasSettings.h
@@ -4,6 +4,7 @@
 
 #include "CoreTypes.h"
 #include "Engine/DeveloperSettings.h"
+#include "GameFramework/GameplayCamerasPlayerCameraManager.h"
 #include "UObject/ObjectMacros.h"
 #include "UObject/Object.h"
 
@@ -30,6 +31,13 @@ public:
 	UPROPERTY(EditAnywhere, Config, Category="General")
 	bool bAutoBuildInPIE = true;
 
+	/** 
+	 * The default view rotation handling mode to use when the game's player controller uses a
+	 * GameplayCamerasPlayerCameraManager instance as its camera manager.
+	 */
+	UPROPERTY(EditAnywhere, Config, Category="General")
+	EGameplayCamerasViewRotationMode DefaultViewRotationMode = EGameplayCamerasViewRotationMode::None;
+
 	/**
 	 * The number of camera rigs combined in one frame past which the camera system emits a warning.
 	 */
