Cameras: fix LOCTEXT collision

diff --git a/Engine/Plugins/Cameras/GameplayCameras/Source/GameplayCamerasEditor/Private/Family/CameraAssetFamily.cpp b/Engine/Plugins/Cameras/GameplayCameras/Source/GameplayCamerasEditor/Private/Family/CameraAssetFamily.cpp
--- a/Engine/Plugins/Cameras/GameplayCameras/Source/GameplayCamerasEditor/Private/Family/CameraAssetFamily.cpp
+++ b/Engine/Plugins/Cameras/GameplayCameras/Source/GameplayCamerasEditor/Private/Family/CameraAssetFamily.cpp
@@ -86,7 +86,7 @@ FText FCameraAssetFamily::GetAssetTypeTooltip(UClass* InAssetType) const
 	}
 	else if (InAssetType == UCameraRigProxyAsset::StaticClass())
 	{
-		return LOCTEXT("CameraRigAssetTypeTooltip", "Open camera rig proxies referenced by this asset.");
+		return LOCTEXT("CameraRigProxyAssetTypeTooltip", "Open camera rig proxies referenced by this asset.");
 	}
 	return FText();
 }
