Cameras: make UActivateCameraRigFunctions accessible

#13093

diff --git a/Engine/Plugins/Cameras/GameplayCameras/Source/GameplayCameras/Public/GameFramework/ActivateCameraRigFunctions.h b/Engine/Plugins/Cameras/GameplayCameras/Source/GameplayCameras/Public/GameFramework/ActivateCameraRigFunctions.h
--- a/Engine/Plugins/Cameras/GameplayCameras/Source/GameplayCameras/Public/GameFramework/ActivateCameraRigFunctions.h
+++ b/Engine/Plugins/Cameras/GameplayCameras/Source/GameplayCameras/Public/GameFramework/ActivateCameraRigFunctions.h
@@ -18,7 +18,7 @@ enum class ECameraRigLayer : uint8;
  * meaningful initial result. They are activated on the camera system found to be running
  * on the given player controller.
  */
-UCLASS(MinimalAPI)
+UCLASS()
 class UActivateCameraRigFunctions : public UBlueprintFunctionLibrary
 {
 	GENERATED_BODY()
