Cameras: make UControllerGameplayCameraEvaluationComponent accessible

#13092

diff --git a/Engine/Plugins/Cameras/GameplayCameras/Source/GameplayCameras/Public/GameFramework/ControllerGameplayCameraEvaluationComponent.h b/Engine/Plugins/Cameras/GameplayCameras/Source/GameplayCameras/Public/GameFramework/ControllerGameplayCameraEvaluationComponent.h
--- a/Engine/Plugins/Cameras/GameplayCameras/Source/GameplayCameras/Public/GameFramework/ControllerGameplayCameraEvaluationComponent.h
+++ b/Engine/Plugins/Cameras/GameplayCameras/Source/GameplayCameras/Public/GameFramework/ControllerGameplayCameraEvaluationComponent.h
@@ -21,7 +21,7 @@ namespace UE::Cameras
  * A component, attached to a player controller, that can run camera rigs activated from
  * a global place like the Blueprint functions inside UActivateCameraRigFunctions.
  */
-UCLASS(Hidden, MinimalAPI)
+UCLASS(Hidden)
 class UControllerGameplayCameraEvaluationComponent : public UActorComponent
 {
 	GENERATED_BODY()
