Cameras: don't throw asserts or error messages when we have a camera actor without an asset specified on it

diff --git a/Engine/Plugins/Cameras/GameplayCameras/Source/GameplayCameras/Private/GameFramework/GameplayCameraComponent.cpp b/Engine/Plugins/Cameras/GameplayCameras/Source/GameplayCameras/Private/GameFramework/GameplayCameraComponent.cpp
--- a/Engine/Plugins/Cameras/GameplayCameras/Source/GameplayCameras/Private/GameFramework/GameplayCameraComponent.cpp
+++ b/Engine/Plugins/Cameras/GameplayCameras/Source/GameplayCameras/Private/GameFramework/GameplayCameraComponent.cpp
@@ -58,15 +58,7 @@ UCameraAsset* UGameplayCameraComponent::GetCameraAsset()
 
 bool UGameplayCameraComponent::OnValidateCameraEvaluationContextActivation()
 {
-	if (!CameraReference.IsValid())
-	{
-		FFrame::KismetExecutionMessage(
-				TEXT("Can't activate gameplay camera component: no camera asset was set!"),
-				ELogVerbosity::Error);
-		return false;
-	}
-	
-	return true;
+	return CameraReference.IsValid();
 }
 
 void UGameplayCameraComponent::OnUpdateCameraEvaluationContext(bool bForceApplyParameterOverrides)
diff --git a/Engine/Plugins/Cameras/GameplayCameras/Source/GameplayCameras/Private/GameFramework/GameplayCameraComponentBase.cpp b/Engine/Plugins/Cameras/GameplayCameras/Source/GameplayCameras/Private/GameFramework/GameplayCameraComponentBase.cpp
--- a/Engine/Plugins/Cameras/GameplayCameras/Source/GameplayCameras/Private/GameFramework/GameplayCameraComponentBase.cpp
+++ b/Engine/Plugins/Cameras/GameplayCameras/Source/GameplayCameras/Private/GameFramework/GameplayCameraComponentBase.cpp
@@ -248,7 +248,6 @@ void UGameplayCameraComponentBase::ActivateCameraEvaluationContext(APlayerContro
 
 	if (!OnValidateCameraEvaluationContextActivation())
 	{
-		// We expect the sub-class to print a useful error or warning message.
 		return;
 	}
 
@@ -645,7 +644,7 @@ void UGameplayCameraComponentBase::AutoManageEditorPreviewEvaluator()
 		EnsureCameraSystemHost();
 
 		ActivateCameraEvaluationContext(nullptr, this, EGameplayCameraComponentActivationMode::Push);
-		if (ensure(EvaluationContext))
+		if (EvaluationContext)
 		{
 			EvaluationContext->SetEditorPreviewCameraRigIndex(EditorPreviewCameraRigIndex);
 		}
diff --git a/Engine/Plugins/Cameras/GameplayCameras/Source/GameplayCameras/Private/GameFramework/GameplayCameraRigComponent.cpp b/Engine/Plugins/Cameras/GameplayCameras/Source/GameplayCameras/Private/GameFramework/GameplayCameraRigComponent.cpp
--- a/Engine/Plugins/Cameras/GameplayCameras/Source/GameplayCameras/Private/GameFramework/GameplayCameraRigComponent.cpp
+++ b/Engine/Plugins/Cameras/GameplayCameras/Source/GameplayCameras/Private/GameFramework/GameplayCameraRigComponent.cpp
@@ -58,15 +58,7 @@ UCameraAsset* UGameplayCameraRigComponent::GetCameraAsset()
 
 bool UGameplayCameraRigComponent::OnValidateCameraEvaluationContextActivation()
 {
-	if (!CameraRigReference.IsValid())
-	{
-		FFrame::KismetExecutionMessage(
-				TEXT("Can't activate gameplay camera rig component: no camera rig asset was set!"),
-				ELogVerbosity::Error);
-		return false;
-	}
-	
-	return true;
+	return CameraRigReference.IsValid();
 }
 
 void UGameplayCameraRigComponent::OnUpdateCameraEvaluationContext(bool bForceApplyParameterOverrides)
