Cameras: rigs could keep running with old data after their asset was rebuilt

diff --git a/Engine/Plugins/Cameras/GameplayCameras/Source/GameplayCameras/Private/GameFramework/GameplayCameraComponent.cpp b/Engine/Plugins/Cameras/GameplayCameras/Source/GameplayCameras/Private/GameFramework/GameplayCameraComponent.cpp
--- a/Engine/Plugins/Cameras/GameplayCameras/Source/GameplayCameras/Private/GameFramework/GameplayCameraComponent.cpp
+++ b/Engine/Plugins/Cameras/GameplayCameras/Source/GameplayCameras/Private/GameFramework/GameplayCameraComponent.cpp
@@ -84,15 +84,11 @@ void UGameplayCameraComponent::OnCameraAssetBuilt(const UCameraAsset* InCameraAs
 
 	// If our camera asset was just built, it may have some new parameters. We need to rebuild
 	// our variable table and context data table, and re-apply overrides.
+	CameraReference.RebuildParametersIfNeeded();
 	if (HasCameraEvaluationContext())
 	{
 		const FCameraAssetAllocationInfo& AllocationInfo = InCameraAsset->GetAllocationInfo();
-		FCameraNodeEvaluationResult& InitialResult = GetEvaluationContext()->GetInitialResult();
-		InitialResult.VariableTable.Initialize(AllocationInfo.VariableTableInfo);
-		InitialResult.ContextDataTable.Initialize(AllocationInfo.ContextDataTableInfo);
-
-		CameraReference.RebuildParametersIfNeeded();
-
+		ReinitializeCameraEvaluationContext(AllocationInfo.VariableTableInfo, AllocationInfo.ContextDataTableInfo);
 		UpdateCameraEvaluationContext(true);
 	}
 }
@@ -111,7 +107,7 @@ void UGameplayCameraComponent::PostEditChangeProperty( struct FPropertyChangedEv
 			if (PropertyChangedEvent.GetPropertyName() == TEXT("CameraAsset"))
 			{
 				// The camera asset has changed! Recreate the context.
-				RecreateCameraEvaluationContext();
+				RecreateEditorWorldCameraEvaluationContext();
 			}
 			else
 			{
diff --git a/Engine/Plugins/Cameras/GameplayCameras/Source/GameplayCameras/Private/GameFramework/GameplayCameraComponentBase.cpp b/Engine/Plugins/Cameras/GameplayCameras/Source/GameplayCameras/Private/GameFramework/GameplayCameraComponentBase.cpp
--- a/Engine/Plugins/Cameras/GameplayCameras/Source/GameplayCameras/Private/GameFramework/GameplayCameraComponentBase.cpp
+++ b/Engine/Plugins/Cameras/GameplayCameras/Source/GameplayCameras/Private/GameFramework/GameplayCameraComponentBase.cpp
@@ -6,13 +6,13 @@
 #include "Build/CameraBuildLog.h"
 #include "CineCameraComponent.h"
 #include "Core/CameraAsset.h"
+#include "Core/CameraEvaluationContextStack.h"
 #include "Core/CameraSystemEvaluator.h"
 #include "Core/RootCameraNode.h"
 #include "Debug/CameraDebugRenderer.h"
 #include "Engine/Canvas.h"
 #include "Engine/EngineTypes.h"
 #include "Engine/World.h"
-#include "GameplayCamerasDelegates.h"
 #include "GameplayCamerasSettings.h"
 #include "IGameplayCamerasLiveEditManager.h"
 #include "IGameplayCamerasModule.h"
@@ -431,19 +431,6 @@ void UGameplayCameraComponentBase::EndPlay(const EEndPlayReason::Type EndPlayRea
 	Super::EndPlay(EndPlayReason);
 }
 
-void UGameplayCameraComponentBase::OnUnregister()
-{
-	using namespace UE::Cameras;
-
-#if WITH_EDITOR
-
-	FGameplayCamerasDelegates::OnCameraAssetBuilt().RemoveAll(this);
-
-#endif  // WITH_EDITOR
-
-	Super::OnUnregister();
-}
-
 void UGameplayCameraComponentBase::Activate(bool bReset)
 {
 	// When auto-activing, this method gets called during OnRegister, before we have started playing.
@@ -550,12 +537,70 @@ void UGameplayCameraComponentBase::UpdateCameraEvaluationContext(bool bForceAppl
 #endif  // WITH_EDITOR
 }
 
-void UGameplayCameraComponentBase::RecreateCameraEvaluationContext()
+#if WITH_EDITOR
+
+void UGameplayCameraComponentBase::ReinitializeCameraEvaluationContext(
+			const FCameraVariableTableAllocationInfo& VariableTableAllocationInfo,
+			const FCameraContextDataTableAllocationInfo& ContextDataTableAllocationInfo)
 {
-	EvaluationContext = nullptr;
-	TryCreateCameraEvaluationContext(nullptr);
+	using namespace UE::Cameras;
+
+	if (EvaluationContext)
+	{
+		FCameraNodeEvaluationResult& InitialResult = EvaluationContext->GetInitialResult();
+		InitialResult.VariableTable.Initialize(VariableTableAllocationInfo);
+		InitialResult.ContextDataTable.Initialize(ContextDataTableAllocationInfo);
+
+		// Also freeze/remove any of our currently running camera rigs, because they might continue
+		// accessing variables and data that don't exist anymore.
+		if (CameraSystemEvaluator)
+		{
+			FRootCameraNodeEvaluator* RootEvaluator = CameraSystemEvaluator->GetRootNodeEvaluator();
+			RootEvaluator->DeactivateAllCameraRigs(EvaluationContext, true);
+		}
+	}
 }
 
+void UGameplayCameraComponentBase::RecreateEditorWorldCameraEvaluationContext()
+{
+	using namespace UE::Cameras;
+
+	if (!bIsEditorWorld)
+	{
+		return;
+	}
+
+	// We should only be calling this method to recreate the editor preview evaluator, so check that
+	// this is indeed the case.
+	if (EvaluationContext && CameraSystemEvaluator)
+	{
+		FCameraEvaluationContextStack& ContextStack = CameraSystemEvaluator->GetEvaluationContextStack();
+		TArray<TSharedPtr<FCameraEvaluationContext>> AllContexts;
+		ContextStack.GetAllContexts(AllContexts);
+		ensure(AllContexts.Num() == 1 && AllContexts[0] == EvaluationContext);
+	}
+
+	// Teardown and rebuild the evaluation context.
+	if (EvaluationContext)
+	{
+		if (CameraSystemEvaluator)
+		{
+			FRootCameraNodeEvaluator* RootEvaluator = CameraSystemEvaluator->GetRootNodeEvaluator();
+			RootEvaluator->DeactivateAllCameraRigs(EvaluationContext.ToSharedRef(), true);
+			CameraSystemEvaluator->RemoveEvaluationContext(EvaluationContext.ToSharedRef());
+		}
+		EvaluationContext = nullptr;
+
+		TryCreateCameraEvaluationContext(nullptr);
+		if (CameraSystemEvaluator)
+		{
+			CameraSystemEvaluator->PushEvaluationContext(EvaluationContext.ToSharedRef());
+		}
+	}
+}
+
+#endif  // WITH_EDITOR
+
 void UGameplayCameraComponentBase::UpdateOutputCameraComponent()
 {
 	using namespace UE::Cameras;
diff --git a/Engine/Plugins/Cameras/GameplayCameras/Source/GameplayCameras/Private/GameFramework/GameplayCameraRigComponent.cpp b/Engine/Plugins/Cameras/GameplayCameras/Source/GameplayCameras/Private/GameFramework/GameplayCameraRigComponent.cpp
--- a/Engine/Plugins/Cameras/GameplayCameras/Source/GameplayCameras/Private/GameFramework/GameplayCameraRigComponent.cpp
+++ b/Engine/Plugins/Cameras/GameplayCameras/Source/GameplayCameras/Private/GameFramework/GameplayCameraRigComponent.cpp
@@ -89,15 +89,11 @@ void UGameplayCameraRigComponent::OnCameraRigAssetBuilt(const UCameraRigAsset* I
 		TGuardValue<bool> ReentrancyGuard(bIsBuildingGeneratedCameraAsset, true);
 		GeneratedCameraAsset->BuildCamera();
 	}
+	CameraRigReference.RebuildParametersIfNeeded();
 	if (HasCameraEvaluationContext())
 	{
 		const FCameraObjectAllocationInfo& AllocationInfo = InCameraRigAsset->AllocationInfo;
-		FCameraNodeEvaluationResult& InitialResult = GetEvaluationContext()->GetInitialResult();
-		InitialResult.VariableTable.Initialize(AllocationInfo.VariableTableInfo);
-		InitialResult.ContextDataTable.Initialize(AllocationInfo.ContextDataTableInfo);
-
-		CameraRigReference.RebuildParametersIfNeeded();
-
+		ReinitializeCameraEvaluationContext(AllocationInfo.VariableTableInfo, AllocationInfo.ContextDataTableInfo);
 		UpdateCameraEvaluationContext(true);
 	}
 }
@@ -117,7 +113,7 @@ void UGameplayCameraRigComponent::PostEditChangeProperty( struct FPropertyChange
 			{
 				// The camera rig asset has changed! Recreate the context.
 				GeneratedCameraAsset = nullptr;
-				RecreateCameraEvaluationContext();
+				RecreateEditorWorldCameraEvaluationContext();
 			}
 			else
 			{
diff --git a/Engine/Plugins/Cameras/GameplayCameras/Source/GameplayCameras/Public/GameFramework/GameplayCameraComponentBase.h b/Engine/Plugins/Cameras/GameplayCameras/Source/GameplayCameras/Public/GameFramework/GameplayCameraComponentBase.h
--- a/Engine/Plugins/Cameras/GameplayCameras/Source/GameplayCameras/Public/GameFramework/GameplayCameraComponentBase.h
+++ b/Engine/Plugins/Cameras/GameplayCameras/Source/GameplayCameras/Public/GameFramework/GameplayCameraComponentBase.h
@@ -18,6 +18,8 @@ class FViewport;
 class UCameraAsset;
 class UCanvas;
 class UCineCameraComponent;
+struct FCameraContextDataTableAllocationInfo;
+struct FCameraVariableTableAllocationInfo;
 
 namespace UE::Cameras
 {
@@ -127,7 +129,6 @@ public:
 	virtual void OnRegister() override;
 	virtual void BeginPlay() override;
 	virtual void EndPlay(const EEndPlayReason::Type EndPlayReason) override;
-	virtual void OnUnregister() override;
 	virtual void Activate(bool bReset = false) override;
 	virtual void Deactivate() override;
 	virtual void TickComponent(float DeltaTime, ELevelTick TickType, FActorComponentTickFunction *ThisTickFunction) override;
@@ -167,7 +168,12 @@ protected:
 	void UpdateCameraEvaluationContext(bool bForceApplyParameterOverrides);
 	bool HasCameraEvaluationContext() const { return EvaluationContext.IsValid(); }
 
-	void RecreateCameraEvaluationContext();
+#if WITH_EDITOR
+	void ReinitializeCameraEvaluationContext(
+			const FCameraVariableTableAllocationInfo& VariableTableAllocationInfo,
+			const FCameraContextDataTableAllocationInfo& ContextDataTableAllocationInfo);
+	void RecreateEditorWorldCameraEvaluationContext();
+#endif  // WITH_EDITOR
 
 private:
 
