Cameras: add ability for the Camera Debugger panel to bind to specific camera system instances

The recent change that made all camera actors/components activate with their own camera system instance consequently made them all print debug info on top of each other. So there had to be a way for the user to specify which camera system instance they want to troubleshoot.

This change adds a debug ID to camera systems, and a repository singleton class to keep track of them. The Camera Debugger panel now has a combo box to pick the debug ID, except for editor worlds in which the component visualizer only renders when the parent actor is selected anyway.

diff --git a/Engine/Plugins/Cameras/GameplayCameras/Source/GameplayCameras/Private/Core/CameraSystemEvaluator.cpp b/Engine/Plugins/Cameras/GameplayCameras/Source/GameplayCameras/Private/Core/CameraSystemEvaluator.cpp
--- a/Engine/Plugins/Cameras/GameplayCameras/Source/GameplayCameras/Private/Core/CameraSystemEvaluator.cpp
+++ b/Engine/Plugins/Cameras/GameplayCameras/Source/GameplayCameras/Private/Core/CameraSystemEvaluator.cpp
@@ -30,7 +30,7 @@ DECLARE_CYCLE_STAT(TEXT("Camera System Eval"), CameraSystemEval_Total, STATGROUP
 namespace UE::Cameras
 {
 
-extern bool GGameplayCamerasDebugEnable;
+extern int32 GGameplayCamerasDebugEnableID;
 
 void FCameraSystemEvaluationResult::Reset()
 {
@@ -106,6 +106,13 @@ void FCameraSystemEvaluator::Initialize(const FCameraSystemEvaluatorCreateParams
 		InitParams.Evaluator = this;
 		RootEvaluator->Initialize(InitParams, RootNodeResult);
 	}
+
+#if UE_GAMEPLAY_CAMERAS_DEBUG
+	if (!DebugID.IsValid())
+	{
+		DebugID = FCameraSystemDebugRegistry::Get().RegisterCameraSystemEvaluator(SharedThis(this));
+	}
+#endif  // UE_GAMEPLAY_CAMERAS_DEBUG
 }
 
 FCameraSystemEvaluator::~FCameraSystemEvaluator()
@@ -121,6 +128,14 @@ FCameraSystemEvaluator::~FCameraSystemEvaluator()
 		}
 		EvaluationServices.Reset();
 	}
+
+#if UE_GAMEPLAY_CAMERAS_DEBUG
+	if (DebugID.IsValid())
+	{
+		FCameraSystemDebugRegistry::Get().UnregisterCameraSystemEvaluator(DebugID);
+		DebugID = FCameraSystemDebugID();
+	}
+#endif  // UE_GAMEPLAY_CAMERAS_TRACE
 }
 
 void FCameraSystemEvaluator::AddReferencedObjects(FReferenceCollector& Collector)
@@ -507,7 +522,7 @@ void FCameraSystemEvaluator::DrawEditorPreview(const FCameraSystemEditorPreviewP
 
 	Renderer.BeginDrawing();
 
-	if (GGameplayCamerasDebugEnable && RootDebugBlock)
+	if (RootDebugBlock)
 	{
 		RootDebugBlock->RootDebugDraw(Renderer);
 	}
@@ -546,7 +561,7 @@ bool FCameraSystemEvaluator::IsDebugTraceEnabled()
 bool FCameraSystemEvaluator::ShouldBuildOrDrawDebugBlocks()
 {
 	const bool bTraceEnabled = IsDebugTraceEnabled();
-	return bTraceEnabled || GGameplayCamerasDebugEnable;
+	return bTraceEnabled || GGameplayCamerasDebugEnableID >= 0;
 }
 
 void FCameraSystemEvaluator::BuildDebugBlocksIfNeeded()
diff --git a/Engine/Plugins/Cameras/GameplayCameras/Source/GameplayCameras/Private/Debug/CameraSystemDebugRegistry.cpp b/Engine/Plugins/Cameras/GameplayCameras/Source/GameplayCameras/Private/Debug/CameraSystemDebugRegistry.cpp
new file mode 100644
--- /dev/null
+++ b/Engine/Plugins/Cameras/GameplayCameras/Source/GameplayCameras/Private/Debug/CameraSystemDebugRegistry.cpp
@@ -0,0 +1,56 @@
+// Copyright Epic Games, Inc. All Rights Reserved.
+
+#include "Debug/CameraSystemDebugRegistry.h"
+
+#include "Core/CameraSystemEvaluator.h"
+#include "Serialization/Archive.h"
+
+#if UE_GAMEPLAY_CAMERAS_DEBUG
+
+namespace UE::Cameras
+{
+
+FArchive& operator <<(FArchive& Ar, FCameraSystemDebugID& InDebugID)
+{
+	Ar << InDebugID.Value;
+	return Ar;
+}
+
+FCameraSystemDebugRegistry& FCameraSystemDebugRegistry::Get()
+{
+	static FCameraSystemDebugRegistry Instance;
+	return Instance;
+}
+
+FCameraSystemDebugID FCameraSystemDebugRegistry::RegisterCameraSystemEvaluator(TSharedRef<FCameraSystemEvaluator> InEvaluator)
+{
+	const FString OwnerName = GetNameSafe(InEvaluator->GetOwner());
+	FEntry NewEntry{ InEvaluator, OwnerName };
+	const int32 NewIndex = Entries.Add(NewEntry);
+	return FCameraSystemDebugID(NewIndex);
+}
+
+void FCameraSystemDebugRegistry::UnregisterCameraSystemEvaluator(FCameraSystemDebugID InDebugID)
+{
+	if (ensure(InDebugID.IsValid()))
+	{
+		Entries.RemoveAt(InDebugID.Value);
+	}
+}
+
+void FCameraSystemDebugRegistry::GetRegisteredCameraSystemEvaluators(FRegisteredCameraSystems& OutEvaluators) const
+{
+	for (auto It = Entries.CreateConstIterator(); It; ++It)
+	{
+		const FEntry& Entry(*It);
+		if (TSharedPtr<FCameraSystemEvaluator> Evaluator = Entry.WeakEvaluator.Pin())
+		{
+			OutEvaluators.Add(Evaluator);
+		}
+	}
+}
+
+}  // namespace UE::Cameras
+
+#endif  // UE_GAMEPLAY_CAMERAS_DEBUG
+
diff --git a/Engine/Plugins/Cameras/GameplayCameras/Source/GameplayCameras/Private/Debug/RootCameraDebugBlock.cpp b/Engine/Plugins/Cameras/GameplayCameras/Source/GameplayCameras/Private/Debug/RootCameraDebugBlock.cpp
--- a/Engine/Plugins/Cameras/GameplayCameras/Source/GameplayCameras/Private/Debug/RootCameraDebugBlock.cpp
+++ b/Engine/Plugins/Cameras/GameplayCameras/Source/GameplayCameras/Private/Debug/RootCameraDebugBlock.cpp
@@ -27,11 +27,11 @@
 namespace UE::Cameras
 {
 
-bool GGameplayCamerasDebugEnable = false;
-static FAutoConsoleVariableRef CVarGameplayCamerasDebugEnable(
-	TEXT("GameplayCameras.Debug.Enable"),
-	GGameplayCamerasDebugEnable,
-	TEXT("(Default: false. Enables debug drawing for the GamplayCameras system."));
+int32 GGameplayCamerasDebugEnableID = INDEX_NONE;
+static FAutoConsoleVariableRef CVarGameplayCamerasDebugEnableID(
+	TEXT("GameplayCameras.Debug.EnableID"),
+	GGameplayCamerasDebugEnableID,
+	TEXT("(Default: -1. Enables debug drawing for the GamplayCameras system with the given debug ID."));
 
 FString GGameplayCamerasDebugCategories = "nodetree";
 static FAutoConsoleVariableRef CVarGameplayCamerasDebugCategories(
@@ -57,6 +57,8 @@ void FRootCameraDebugBlock::BuildDebugBlocks(const FCameraSystemEvaluator& Camer
 {
 	ensureMsgf(GetChildren().IsEmpty() && GetAttachments().IsEmpty(), TEXT("This root debug block has already been initialized!"));
 
+	DebugID = CameraSystem.GetDebugID();
+
 	// Debug block for showing the directors and context stack.
 	FCategoryTitleDebugBlock& DirectorTreeCategory = Builder.StartChildDebugBlock<FCategoryTitleDebugBlock>();
 	{
@@ -135,7 +137,8 @@ void FRootCameraDebugBlock::BuildDebugBlocks(const FCameraSystemEvaluator& Camer
 
 void FRootCameraDebugBlock::RootDebugDraw(FCameraDebugRenderer& Renderer)
 {
-	if (!GGameplayCamerasDebugEnable)
+	FCameraSystemDebugID WantedDebugID(GGameplayCamerasDebugEnableID);
+	if (WantedDebugID != DebugID && !WantedDebugID.IsAny())
 	{
 		return;
 	}
@@ -156,6 +159,11 @@ void FRootCameraDebugBlock::RootDebugDraw(FCameraDebugRenderer& Renderer)
 	Renderer.EndDrawing();
 }
 
+void FRootCameraDebugBlock::OnSerialize(FArchive& Ar)
+{
+	Ar << DebugID;
+}
+
 }  // namespace UE::Cameras
 
 #endif  // UE_GAMEPLAY_CAMERAS_DEBUG
diff --git a/Engine/Plugins/Cameras/GameplayCameras/Source/GameplayCameras/Private/GameFramework/GameplayCameraComponentBase.cpp b/Engine/Plugins/Cameras/GameplayCameras/Source/GameplayCameras/Private/GameFramework/GameplayCameraComponentBase.cpp
--- a/Engine/Plugins/Cameras/GameplayCameras/Source/GameplayCameras/Private/GameFramework/GameplayCameraComponentBase.cpp
+++ b/Engine/Plugins/Cameras/GameplayCameras/Source/GameplayCameras/Private/GameFramework/GameplayCameraComponentBase.cpp
@@ -473,7 +473,23 @@ void UGameplayCameraComponentBase::TickComponent(float DeltaTime, ELevelTick Tic
 		UpdateCameraEvaluationContext(false);
 	}
 
+#if WITH_EDITOR
+
+	if (bIsEditorWorld)
+	{
+		UpdateCameraSystemForEditorPreview(DeltaTime);
+	}
+	else
+	{
+		UpdateCameraSystem(DeltaTime);
+	}
+
+#else
+
 	UpdateCameraSystem(DeltaTime);
+
+#endif  // WITH_EDITOR
+
 	UpdateOutputCameraComponent();
 
 	CheckPendingDeactivation();
diff --git a/Engine/Plugins/Cameras/GameplayCameras/Source/GameplayCameras/Private/GameFramework/IGameplayCameraSystemHost.cpp b/Engine/Plugins/Cameras/GameplayCameras/Source/GameplayCameras/Private/GameFramework/IGameplayCameraSystemHost.cpp
--- a/Engine/Plugins/Cameras/GameplayCameras/Source/GameplayCameras/Private/GameFramework/IGameplayCameraSystemHost.cpp
+++ b/Engine/Plugins/Cameras/GameplayCameras/Source/GameplayCameras/Private/GameFramework/IGameplayCameraSystemHost.cpp
@@ -105,6 +105,22 @@ void IGameplayCameraSystemHost::UpdateCameraSystem(float DeltaTime)
 	}
 }
 
+#if WITH_EDITOR
+
+void IGameplayCameraSystemHost::UpdateCameraSystemForEditorPreview(float DeltaTime)
+{
+	using namespace UE::Cameras;
+
+	if (CameraSystemEvaluator.IsValid())
+	{
+		FCameraSystemEvaluationParams Params;
+		Params.DeltaTime = DeltaTime;
+		CameraSystemEvaluator->EditorPreviewUpdate(Params);
+	}
+}
+
+#endif  // WITH_EDITOR
+
 TSharedPtr<UE::Cameras::FCameraSystemEvaluator> IGameplayCameraSystemHost::GetCameraSystemEvaluator()
 {
 	return CameraSystemEvaluator;
diff --git a/Engine/Plugins/Cameras/GameplayCameras/Source/GameplayCameras/Public/Core/CameraSystemEvaluator.h b/Engine/Plugins/Cameras/GameplayCameras/Source/GameplayCameras/Public/Core/CameraSystemEvaluator.h
--- a/Engine/Plugins/Cameras/GameplayCameras/Source/GameplayCameras/Public/Core/CameraSystemEvaluator.h
+++ b/Engine/Plugins/Cameras/GameplayCameras/Source/GameplayCameras/Public/Core/CameraSystemEvaluator.h
@@ -10,6 +10,7 @@
 #include "Core/CameraPose.h"
 #include "CoreTypes.h"
 #include "Debug/CameraDebugBlockStorage.h"
+#include "Debug/CameraSystemDebugRegistry.h"
 #include "GameplayCameras.h"
 #include "Templates/SharedPointer.h"
 #include "UObject/GCObject.h"
@@ -190,6 +191,11 @@ public:
 	/** Gets the role of this camera system. */
 	ECameraSystemEvaluatorRole GetRole() const { return Role; }
 
+#if UE_GAMEPLAY_CAMERAS_DEBUG
+	/** Gets the debug ID for this camera system. */
+	FCameraSystemDebugID GetDebugID() const { return DebugID; }
+#endif  // UE_GAMEPLAY_CAMERAS_DEBUG
+
 public:
 
 	/** Push a new evaluation context on the stack. */
@@ -314,6 +320,9 @@ private:
 	ECameraSystemEvaluatorRole Role = ECameraSystemEvaluatorRole::Game;
 
 #if UE_GAMEPLAY_CAMERAS_DEBUG
+	/** The debug ID for this camera system. */
+	FCameraSystemDebugID DebugID;
+
 	/** Storage for debug drawing blocks. */
 	FCameraDebugBlockStorage DebugBlockStorage;
 
diff --git a/Engine/Plugins/Cameras/GameplayCameras/Source/GameplayCameras/Public/Debug/CameraSystemDebugRegistry.h b/Engine/Plugins/Cameras/GameplayCameras/Source/GameplayCameras/Public/Debug/CameraSystemDebugRegistry.h
new file mode 100644
--- /dev/null
+++ b/Engine/Plugins/Cameras/GameplayCameras/Source/GameplayCameras/Public/Debug/CameraSystemDebugRegistry.h
@@ -0,0 +1,102 @@
+// Copyright Epic Games, Inc. All Rights Reserved.
+
+#pragma once
+
+#include "Containers/SparseArray.h"
+#include "GameplayCameras.h"
+#include "Math/NumericLimits.h"
+#include "Templates/SharedPointerFwd.h"
+
+#if UE_GAMEPLAY_CAMERAS_DEBUG
+
+class FArchive;
+
+namespace UE::Cameras
+{
+
+class FCameraSystemEvaluator;
+
+/**
+ * An identifier for a camera system instance.
+ */
+struct FCameraSystemDebugID
+{
+	FCameraSystemDebugID() : Value(INVALID) {}
+	FCameraSystemDebugID(int32 InValue) : Value(InValue) {}
+
+	bool IsValid() const { return Value != INVALID; }
+	bool IsAny() const { return Value == ANY; }
+	operator bool() const { return IsValid(); }
+
+	int32 GetValue() const { return Value; }
+
+	static FCameraSystemDebugID Invalid() { return FCameraSystemDebugID(INVALID); }
+	static FCameraSystemDebugID Any() { return FCameraSystemDebugID(ANY); }
+
+public:
+
+	friend bool operator==(FCameraSystemDebugID A, FCameraSystemDebugID B)
+	{
+		return A.Value == B.Value;
+	}
+
+	friend bool operator!=(FCameraSystemDebugID A, FCameraSystemDebugID B)
+	{
+		return !(A == B);
+	}
+
+	friend uint32 GetTypeHash(FCameraSystemDebugID In)
+	{
+		return In.Value;
+	}
+
+	friend FString LexToString(const FCameraSystemDebugID In)
+	{
+		return LexToString(In.Value);
+	}
+
+	friend FArchive& operator <<(FArchive& Ar, FCameraSystemDebugID& InDebugID);
+
+private:
+
+	constexpr static int32 INVALID = INDEX_NONE;
+	constexpr static int32 ANY = MAX_int32;
+
+	int32 Value;
+
+	friend class FCameraSystemDebugRegistry;
+	friend class FRootCameraDebugBlock;
+};
+
+FArchive& operator <<(FArchive& Ar, FCameraSystemDebugID& InDebugID);
+
+/**
+ * A registry for any running camera system instance for which we may want to display debug info.
+ */
+class FCameraSystemDebugRegistry
+{
+public:
+
+	GAMEPLAYCAMERAS_API static FCameraSystemDebugRegistry& Get();
+
+	GAMEPLAYCAMERAS_API FCameraSystemDebugID RegisterCameraSystemEvaluator(TSharedRef<FCameraSystemEvaluator> InEvaluator);
+	GAMEPLAYCAMERAS_API void UnregisterCameraSystemEvaluator(FCameraSystemDebugID InDebugID);
+
+	using FRegisteredCameraSystems = TArray<TSharedPtr<FCameraSystemEvaluator>>;
+	GAMEPLAYCAMERAS_API void GetRegisteredCameraSystemEvaluators(FRegisteredCameraSystems& OutEvaluators) const;
+
+private:
+
+	struct FEntry
+	{
+		TWeakPtr<FCameraSystemEvaluator> WeakEvaluator;
+		FString EvaluatorOwnerName;
+	};
+
+	TSparseArray<FEntry> Entries;
+};
+
+}  // namespace UE::Cameras
+
+#endif  // UE_GAMEPLAY_CAMERAS_DEBUG
+
diff --git a/Engine/Plugins/Cameras/GameplayCameras/Source/GameplayCameras/Public/Debug/RootCameraDebugBlock.h b/Engine/Plugins/Cameras/GameplayCameras/Source/GameplayCameras/Public/Debug/RootCameraDebugBlock.h
--- a/Engine/Plugins/Cameras/GameplayCameras/Source/GameplayCameras/Public/Debug/RootCameraDebugBlock.h
+++ b/Engine/Plugins/Cameras/GameplayCameras/Source/GameplayCameras/Public/Debug/RootCameraDebugBlock.h
@@ -3,7 +3,7 @@
 #pragma once
 
 #include "Debug/CameraDebugBlock.h"
-#include "Debug/CameraDebugBlockStorage.h"
+#include "Debug/CameraSystemDebugRegistry.h"
 
 #if UE_GAMEPLAY_CAMERAS_DEBUG
 
@@ -14,7 +14,7 @@ class FCameraSystemEvaluator;
 struct FCameraDebugBlockBuildParams;
 struct FCameraDebugBlockBuilder;
 
-GAMEPLAYCAMERAS_API extern bool GGameplayCamerasDebugEnable;
+GAMEPLAYCAMERAS_API extern int32 GGameplayCamerasDebugEnableID;
 GAMEPLAYCAMERAS_API extern FString GGameplayCamerasDebugCategories;
 
 /**
@@ -31,6 +31,15 @@ public:
 
 	/** Initiate the debug drawing. */
 	GAMEPLAYCAMERAS_API void RootDebugDraw(FCameraDebugRenderer& Renderer);
+
+protected:
+
+	// FCameraDebugBlock interface.
+	virtual void OnSerialize(FArchive& Ar) override;
+
+private:
+
+	FCameraSystemDebugID DebugID; 
 };
 
 }  // namespace UE::Cameras
diff --git a/Engine/Plugins/Cameras/GameplayCameras/Source/GameplayCameras/Public/GameFramework/IGameplayCameraSystemHost.h b/Engine/Plugins/Cameras/GameplayCameras/Source/GameplayCameras/Public/GameFramework/IGameplayCameraSystemHost.h
--- a/Engine/Plugins/Cameras/GameplayCameras/Source/GameplayCameras/Public/GameFramework/IGameplayCameraSystemHost.h
+++ b/Engine/Plugins/Cameras/GameplayCameras/Source/GameplayCameras/Public/GameFramework/IGameplayCameraSystemHost.h
@@ -70,6 +70,11 @@ protected:
 	/** Updates the camera system, if it exists. */
 	GAMEPLAYCAMERAS_API void UpdateCameraSystem(float DeltaTime);
 
+#if WITH_EDITOR
+	/** Updates the camera system, if it exists, for an editor world preview. */
+	GAMEPLAYCAMERAS_API void UpdateCameraSystemForEditorPreview(float DeltaTime);
+#endif  // WITH_EDITOR
+
 protected:
 
 	/** The camera system evaluator. */
diff --git a/Engine/Plugins/Cameras/GameplayCameras/Source/GameplayCamerasEditor/Private/Debugger/SGameplayCamerasDebugger.cpp b/Engine/Plugins/Cameras/GameplayCameras/Source/GameplayCamerasEditor/Private/Debugger/SGameplayCamerasDebugger.cpp
--- a/Engine/Plugins/Cameras/GameplayCameras/Source/GameplayCamerasEditor/Private/Debugger/SGameplayCamerasDebugger.cpp
+++ b/Engine/Plugins/Cameras/GameplayCameras/Source/GameplayCamerasEditor/Private/Debugger/SGameplayCamerasDebugger.cpp
@@ -3,27 +3,34 @@
 #include "Debugger/SGameplayCamerasDebugger.h"
 
 #include "Commands/GameplayCamerasDebuggerCommands.h"
+#include "Core/CameraSystemEvaluator.h"
 #include "Debug/CameraDebugColors.h"
+#include "Debug/CameraSystemDebugRegistry.h"
 #include "Debug/RootCameraDebugBlock.h"
 #include "Debugger/SDebugCategoryButton.h"
 #include "Debugger/SDebugWidgetUtils.h"
+#include "Editor.h"
+#include "Engine/Engine.h"
+#include "Engine/EngineTypes.h"
 #include "Framework/Application/SlateApplication.h"
-#include "Framework/Docking/LayoutService.h"
 #include "Framework/Docking/TabManager.h"
 #include "Framework/MultiBox/MultiBoxBuilder.h"
+#include "GameFramework/Actor.h"
+#include "GameplayCamerasEditorSettings.h"
 #include "IGameplayCamerasEditorModule.h"
 #include "Modules/ModuleManager.h"
 #include "String/ParseTokens.h"
 #include "Styles/GameplayCamerasEditorStyle.h"
 #include "Styling/SlateTypes.h"
-#include "ToolMenus.h"
 #include "ToolMenuDelegates.h"
+#include "ToolMenuSection.h"
+#include "ToolMenus.h"
 #include "Widgets/Docking/SDockTab.h"
 #include "Widgets/Layout/SBox.h"
 #include "Widgets/Layout/SExpandableArea.h"
 #include "Widgets/Layout/SGridPanel.h"
-#include "Widgets/Text/STextBlock.h"
 #include "Widgets/SNullWidget.h"
+#include "Widgets/Text/STextBlock.h"
 #include "WorkspaceMenuStructure.h"
 #include "WorkspaceMenuStructureModule.h"
 
@@ -38,6 +45,89 @@ const FName SGameplayCamerasDebugger::WindowName(TEXT("GameplayCamerasDebugger")
 const FName SGameplayCamerasDebugger::MenubarName(TEXT("GameplayCamerasDebugger.Menubar"));
 const FName SGameplayCamerasDebugger::ToolbarName(TEXT("GameplayCamerasDebugger.Toolbar"));
 
+FGameplayCamerasDebuggerContext::FGameplayCamerasDebuggerContext()
+{
+	FEditorDelegates::MapChange.AddRaw(this, &FGameplayCamerasDebuggerContext::OnMapChange);
+	FEditorDelegates::BeginPIE.AddRaw(this, &FGameplayCamerasDebuggerContext::OnPieEvent);
+	FEditorDelegates::EndPIE.AddRaw(this, &FGameplayCamerasDebuggerContext::OnPieEvent);
+
+	if (GEngine)
+	{
+		GEngine->OnWorldAdded().AddRaw(this, &FGameplayCamerasDebuggerContext::OnWorldListChanged);
+		GEngine->OnWorldDestroyed().AddRaw(this, &FGameplayCamerasDebuggerContext::OnWorldListChanged);
+	}
+}
+
+FGameplayCamerasDebuggerContext::~FGameplayCamerasDebuggerContext()
+{
+	FEditorDelegates::MapChange.RemoveAll(this);
+	FEditorDelegates::BeginPIE.RemoveAll(this);
+	FEditorDelegates::EndPIE.RemoveAll(this);
+
+	if (GEngine)
+	{
+		GEngine->OnWorldAdded().RemoveAll(this);
+		GEngine->OnWorldDestroyed().RemoveAll(this);
+	}
+}
+
+UWorld* FGameplayCamerasDebuggerContext::GetContext()
+{
+	UpdateContext();
+	return WeakContext.Get();
+}
+
+void FGameplayCamerasDebuggerContext::UpdateContext()
+{
+	if (WeakContext.IsValid())
+	{
+		return;
+	}
+
+	const UGameplayCamerasEditorSettings* Settings = GetDefault<UGameplayCamerasEditorSettings>();
+
+	// Pick the first editor world we find, but if there's any PIE/SIE world, prefer those.
+	UWorld* NewContext = nullptr;
+	for (const FWorldContext& WorldContext : GEngine->GetWorldContexts())
+	{
+		if (WorldContext.WorldType == EWorldType::PIE)
+		{
+			NewContext = WorldContext.World();
+			break;
+		}
+		else if (WorldContext.WorldType == EWorldType::Editor)
+		{
+			if (NewContext == nullptr)
+			{
+				NewContext = WorldContext.World();
+			}
+		}
+	}
+	ensure(NewContext);
+	WeakContext = NewContext;
+}
+
+void FGameplayCamerasDebuggerContext::InvalidateContext()
+{
+	WeakContext = nullptr;
+	OnContextChangedEvent.Broadcast();
+}
+
+void FGameplayCamerasDebuggerContext::OnPieEvent(bool bIsSimulating)
+{
+	InvalidateContext();
+}
+
+void FGameplayCamerasDebuggerContext::OnMapChange(uint32 MapChangeFlags)
+{
+	InvalidateContext();
+}
+
+void FGameplayCamerasDebuggerContext::OnWorldListChanged(UWorld* InWorld)
+{
+	InvalidateContext();
+}
+
 void SGameplayCamerasDebugger::RegisterTabSpawners()
 {
 	TSharedRef<FGameplayCamerasEditorStyle> CamerasEditorStyle = FGameplayCamerasEditorStyle::Get();
@@ -74,10 +164,12 @@ TSharedRef<SDockTab> SGameplayCamerasDebugger::SpawnGameplayCamerasDebugger(cons
 
 SGameplayCamerasDebugger::SGameplayCamerasDebugger()
 {
+	DebugContext.OnContextChanged().AddRaw(this, &SGameplayCamerasDebugger::OnDebugContextChanged);
 }
 
 SGameplayCamerasDebugger::~SGameplayCamerasDebugger()
 {
+	DebugContext.OnContextChanged().RemoveAll(this);
 }
 
 void SGameplayCamerasDebugger::Construct(const FArguments& InArgs)
@@ -92,9 +184,9 @@ void SGameplayCamerasDebugger::Construct(const FArguments& InArgs)
 	TSharedRef<FUICommandList> CommandList = MakeShareable(new FUICommandList);
 	CommandList->MapAction(
 			Commands.EnableDebugInfo,
-			FExecuteAction::CreateLambda([]() { GGameplayCamerasDebugEnable = !GGameplayCamerasDebugEnable; }),
-			FCanExecuteAction(),
-			FIsActionChecked::CreateLambda([]() { return GGameplayCamerasDebugEnable; }));
+			FExecuteAction::CreateSP(this, &SGameplayCamerasDebugger::ToggleDebugDraw),
+			FCanExecuteAction::CreateSP(this, &SGameplayCamerasDebugger::CanToggleDebugDraw),
+			FIsActionChecked::CreateSP(this, &SGameplayCamerasDebugger::IsDebugDrawing));
 
 	// Build all UI elements.
 	TSharedRef<SWidget> MenubarContents = ConstructMenubar();
@@ -140,6 +232,52 @@ void SGameplayCamerasDebugger::Construct(const FArguments& InArgs)
 	{
 		SetActiveDebugCategoryPanel(FString(ActiveCategories[0]));
 	}
+
+	bRefreshDebugID = true;
+}
+
+void SGameplayCamerasDebugger::Tick(const FGeometry& AllottedGeometry, const double InCurrentTime, const float InDeltaTime)
+{
+	SCompoundWidget::Tick(AllottedGeometry, InCurrentTime, InDeltaTime);
+	
+	if (bRefreshDebugID)
+	{
+		// Auto-set the camera system debug ID when PIE starts/ends, and for other similar events.
+		UWorld* DebugWorld = DebugContext.GetContext();
+		if (DebugWorld)
+		{
+			const bool bIsEditorWorld = (
+					DebugWorld->WorldType == EWorldType::Editor || 
+					DebugWorld->WorldType == EWorldType::EditorPreview);
+
+			if (bIsEditorWorld)
+			{
+				DebugID = FCameraSystemDebugID::Any();
+			}
+			else
+			{
+				FCameraSystemDebugRegistry::FRegisteredCameraSystems CameraSystems;
+				FCameraSystemDebugRegistry::Get().GetRegisteredCameraSystemEvaluators(CameraSystems);
+
+				DebugID = FCameraSystemDebugID::Invalid();
+				for (TSharedPtr<FCameraSystemEvaluator> CameraSystem : CameraSystems)
+				{
+					UObject* CameraSystemOwner = CameraSystem->GetOwner();
+					if (CameraSystemOwner && CameraSystemOwner->GetWorld() == DebugWorld)
+					{
+						DebugID = CameraSystem->GetDebugID();
+						break;
+					}
+				}
+			}
+		}
+		else
+		{
+			DebugID = FCameraSystemDebugID::Invalid();
+		}
+
+		bRefreshDebugID = false;
+	}
 }
 
 void SGameplayCamerasDebugger::InitializeColorSchemeNames()
@@ -202,6 +340,16 @@ TSharedRef<SWidget> SGameplayCamerasDebugger::ConstructToolbar(TSharedRef<FUICom
 						TAttribute<FText>(),
 						TAttribute<FSlateIcon>::CreateSP(This, &SGameplayCamerasDebugger::GetToggleDebugDrawIcon));
 				MainSection.AddEntry(ToggleDebugInfo);
+
+				FToolMenuEntry BindComboEntry = FToolMenuEntry::InitComboButton(
+						"BindToCameraSystemsMenu",
+						FUIAction(),
+						FNewToolMenuDelegate::CreateSP(This, &SGameplayCamerasDebugger::GetCameraSystemPickerContent),
+						LOCTEXT("BindToCameraSystemsMenu", "Bind to..."),
+						LOCTEXT("BindToCameraSystemsMenuToolTip", "Pick a camera system instance to bind to"),
+						FSlateIcon(This->GameplayCamerasEditorStyleName, "Debugger.BindToCameraSystem"),
+						true);
+				MainSection.AddEntry(BindComboEntry);
 			}));
 	
 		Toolbar->AddDynamicSection(TEXT("DebugCategories"), FNewToolMenuDelegate::CreateLambda(
@@ -376,9 +524,31 @@ void SGameplayCamerasDebugger::ConstructDebugPanels()
 	}
 }
 
+void SGameplayCamerasDebugger::ToggleDebugDraw()
+{
+	if (GGameplayCamerasDebugEnableID < 0)
+	{
+		GGameplayCamerasDebugEnableID = DebugID.GetValue();
+	}
+	else
+	{
+		GGameplayCamerasDebugEnableID = -1;
+	}
+}
+
+bool SGameplayCamerasDebugger::CanToggleDebugDraw() const
+{
+	return DebugID.IsValid();
+}
+
+bool SGameplayCamerasDebugger::IsDebugDrawing() const
+{
+	return GGameplayCamerasDebugEnableID >= 0;
+}
+
 FText SGameplayCamerasDebugger::GetToggleDebugDrawText() const
 {
-	if (GGameplayCamerasDebugEnable)
+	if (GGameplayCamerasDebugEnableID >= 0)
 	{
 		return LOCTEXT("DebugInfoEnabled", "Debug Info Enabled");
 	}
@@ -390,7 +560,7 @@ FText SGameplayCamerasDebugger::GetToggleDebugDrawText() const
 
 FSlateIcon SGameplayCamerasDebugger::GetToggleDebugDrawIcon() const
 {
-	if (GGameplayCamerasDebugEnable)
+	if (GGameplayCamerasDebugEnableID >= 0)
 	{
 		return FSlateIcon(GameplayCamerasEditorStyleName, "Debugger.DebugInfoEnabled.Icon");
 	}
@@ -400,6 +570,97 @@ FSlateIcon SGameplayCamerasDebugger::GetToggleDebugDrawIcon() const
 	}
 }
 
+void SGameplayCamerasDebugger::GetCameraSystemPickerContent(UToolMenu* ToolMenu)
+{
+	FCameraSystemDebugRegistry::FRegisteredCameraSystems CameraSystems;
+	FCameraSystemDebugRegistry::Get().GetRegisteredCameraSystemEvaluators(CameraSystems);
+
+	UWorld* DebugWorld = DebugContext.GetContext();
+	FToolMenuSection& CameraSystemsSection = ToolMenu->AddSection(
+			"CameraSystems",
+			FText::Format(
+				LOCTEXT("BoundToWorldName", "Camera Systems in {0}"),
+				FText::FromName(DebugWorld->GetFName()))
+			);
+
+	if (DebugWorld)
+	{
+		const bool bIsEditorWorld = (
+				DebugWorld->WorldType == EWorldType::Editor || 
+				DebugWorld->WorldType == EWorldType::EditorPreview);
+
+		if (bIsEditorWorld)
+		{
+			CameraSystemsSection.AddMenuEntry(
+					TEXT("SelectToBindInEditorWorld"),
+					LOCTEXT("SelectToBindInEditorWorld", "Select actor to show debug info"),
+					LOCTEXT("SelectToBindInEditorWorldToolTip", "In editor worlds, debug info is shown for the selected camera actor."),
+					FSlateIcon(),
+					FUIAction(
+						FExecuteAction(),
+						FCanExecuteAction::CreateLambda([]() { return false; }))
+					);
+		}
+		else if (CameraSystems.Num() > 0)
+		{
+			for (TSharedPtr<FCameraSystemEvaluator> CameraSystem : CameraSystems)
+			{
+				UObject* CameraSystemOwner = CameraSystem->GetOwner();
+				if (CameraSystemOwner && CameraSystemOwner->GetWorld() == DebugWorld)
+				{
+					AActor* OwnerActor = Cast<AActor>(CameraSystemOwner);
+					if (!OwnerActor)
+					{
+						OwnerActor = CameraSystemOwner->GetTypedOuter<AActor>();
+					}
+					const FName OwnerName = (OwnerActor ? OwnerActor->GetFName() : CameraSystemOwner->GetFName());
+
+					CameraSystemsSection.AddMenuEntry(
+							NAME_None,
+							FText::Format(
+								LOCTEXT("BindToCameraSystem", "{0} (ID={1})"),
+								FText::FromName(OwnerName),
+								FText::FromString(LexToString(CameraSystem->GetDebugID()))),
+							LOCTEXT("BindToCameraSystemToolTip", "Bind to this camera system instance"),
+							FSlateIcon(),
+							FUIAction(
+								FExecuteAction::CreateSP(
+									this, &SGameplayCamerasDebugger::BindToCameraSystem, CameraSystem->GetDebugID()),
+								FCanExecuteAction(),
+								FIsActionChecked::CreateSP(
+									this, &SGameplayCamerasDebugger::IsBoundToCameraSystem, CameraSystem->GetDebugID())),
+							EUserInterfaceActionType::Check);
+				}
+			}
+		}
+		else
+		{
+			CameraSystemsSection.AddMenuEntry(
+					TEXT("NoCameraSystem"),
+					LOCTEXT("NoCameraSystem", "None"),
+					LOCTEXT("NoCameraSystemToolTip", "No camera systems found"),
+					FSlateIcon(),
+					FUIAction());
+		}
+	}
+}
+
+void SGameplayCamerasDebugger::BindToCameraSystem(FCameraSystemDebugID InDebugID)
+{
+	DebugID = InDebugID;
+	GGameplayCamerasDebugEnableID = InDebugID.GetValue();
+}
+
+bool SGameplayCamerasDebugger::IsBoundToCameraSystem(FCameraSystemDebugID InDebugID)
+{
+	return DebugID == InDebugID;
+}
+
+void SGameplayCamerasDebugger::OnDebugContextChanged()
+{
+	bRefreshDebugID = true;
+}
+
 bool SGameplayCamerasDebugger::IsDebugCategoryActive(FString InCategoryName)
 {
 	TArray<FStringView, TInlineAllocator<4>> ActiveCategories;
diff --git a/Engine/Plugins/Cameras/GameplayCameras/Source/GameplayCamerasEditor/Private/Debugger/SGameplayCamerasDebugger.h b/Engine/Plugins/Cameras/GameplayCameras/Source/GameplayCamerasEditor/Private/Debugger/SGameplayCamerasDebugger.h
--- a/Engine/Plugins/Cameras/GameplayCameras/Source/GameplayCamerasEditor/Private/Debugger/SGameplayCamerasDebugger.h
+++ b/Engine/Plugins/Cameras/GameplayCameras/Source/GameplayCamerasEditor/Private/Debugger/SGameplayCamerasDebugger.h
@@ -3,6 +3,7 @@
 #pragma once
 
 #include "CoreTypes.h"
+#include "Debug/CameraSystemDebugRegistry.h"
 #include "Widgets/SCompoundWidget.h"
 
 #include "SGameplayCamerasDebugger.generated.h"
@@ -19,6 +20,32 @@ struct FSlateIcon;
 namespace UE::Cameras
 {
 
+class FGameplayCamerasDebuggerContext
+{
+public:
+
+	FGameplayCamerasDebuggerContext();
+	~FGameplayCamerasDebuggerContext();
+
+	UWorld* GetContext();
+
+	FSimpleMulticastDelegate& OnContextChanged() { return OnContextChangedEvent; }
+
+private:
+
+	void UpdateContext();
+	void InvalidateContext();
+
+	void OnPieEvent(bool bIsSimulating);
+	void OnMapChange(uint32 MapChangeFlags);
+	void OnWorldListChanged(UWorld* InWorld);
+
+private:
+
+	TWeakObjectPtr<UWorld> WeakContext;
+	FSimpleMulticastDelegate OnContextChangedEvent;
+};
+
 class SGameplayCamerasDebugger : public SCompoundWidget
 {
 public:
@@ -41,6 +68,11 @@ public:
 
 	void Construct(const FArguments& InArgs);
 
+protected:
+
+	// SWidget interface.
+	virtual void Tick(const FGeometry& AllottedGeometry, const double InCurrentTime, const float InDeltaTime) override;
+
 protected:
 
 	static SGameplayCamerasDebugger* FromContext(UToolMenu* InMenu);
@@ -54,13 +86,27 @@ protected:
 	static bool IsDebugCategoryActive(FString InCategoryName);
 	void SetActiveDebugCategoryPanel(FString InCategoryName);
 
+	void ToggleDebugDraw();
+	bool CanToggleDebugDraw() const;
+	bool IsDebugDrawing() const;
 	FText GetToggleDebugDrawText() const;
 	FSlateIcon GetToggleDebugDrawIcon() const;
+	
+	void GetCameraSystemPickerContent(UToolMenu* ToolMenu);
+
+	void BindToCameraSystem(FCameraSystemDebugID InDebugID);
+	bool IsBoundToCameraSystem(FCameraSystemDebugID InDebugID);
+
+	void OnDebugContextChanged();
 
 private:
 
 	FName GameplayCamerasEditorStyleName;
 
+	FGameplayCamerasDebuggerContext DebugContext;
+	FCameraSystemDebugID DebugID;
+	bool bRefreshDebugID = false;
+
 	TSharedPtr<SBox> PanelHost;
 
 	TSharedPtr<SWidget> EmptyPanel;
diff --git a/Engine/Plugins/Cameras/GameplayCameras/Source/GameplayCamerasEditor/Private/Styles/GameplayCamerasEditorStyle.cpp b/Engine/Plugins/Cameras/GameplayCameras/Source/GameplayCamerasEditor/Private/Styles/GameplayCamerasEditorStyle.cpp
--- a/Engine/Plugins/Cameras/GameplayCameras/Source/GameplayCamerasEditor/Private/Styles/GameplayCamerasEditorStyle.cpp
+++ b/Engine/Plugins/Cameras/GameplayCameras/Source/GameplayCamerasEditor/Private/Styles/GameplayCamerasEditorStyle.cpp
@@ -128,6 +128,8 @@ FGameplayCamerasEditorStyle::FGameplayCamerasEditorStyle()
 	// Debugger tool icons.
 	Set("Debugger.TabIcon", new IMAGE_BRUSH_SVG("Icons/GameplayCameraSystem_16", Icon16x16));
 
+	Set("Debugger.BindToCameraSystem", new IMAGE_BRUSH_SVG("Icons/DebugLink", Icon16x16));
+
 	Set("Debugger.DebugInfoEnabled.Icon", new IMAGE_BRUSH_SVG("Icons/DebugInfo-ToggleCheck", Icon16x16, FStyleColors::AccentGreen));
 	Set("Debugger.DebugInfoDisabled.Icon", new IMAGE_BRUSH_SVG("Icons/DebugInfo-ToggleCheck", Icon16x16, FStyleColors::AccentGray));
 
