Cameras: allow building and running camera rigs and shakes with no root node

This make it more user-friendly for newly created assets to work right away instead of spitting out errors.

diff --git a/Engine/Plugins/Cameras/GameplayCameras/Source/GameplayCameras/Private/Build/CameraRigAssetBuilder.cpp b/Engine/Plugins/Cameras/GameplayCameras/Source/GameplayCameras/Private/Build/CameraRigAssetBuilder.cpp
--- a/Engine/Plugins/Cameras/GameplayCameras/Source/GameplayCameras/Private/Build/CameraRigAssetBuilder.cpp
+++ b/Engine/Plugins/Cameras/GameplayCameras/Source/GameplayCameras/Private/Build/CameraRigAssetBuilder.cpp
@@ -50,14 +50,6 @@ void FCameraRigAssetBuilder::BuildCameraRig(UCameraRigAsset* InCameraRig, FCusto
 
 void FCameraRigAssetBuilder::BuildCameraRigImpl()
 {
-	if (!CameraRig->RootNode)
-	{
-		BuildLog.AddMessage(EMessageSeverity::Error, CameraRig, 
-				FText::Format(LOCTEXT("MissingRootNode", "Camera rig '{0}' has no root node."), 
-					FText::FromString(GetPathNameSafe(CameraRig))));
-		return;
-	}
-
 	FCameraNodeHierarchyBuilder NodeBuilder(BuildLog, CameraRig);
 	NodeBuilder.PreBuild();
 
diff --git a/Engine/Plugins/Cameras/GameplayCameras/Source/GameplayCameras/Private/Build/CameraShakeAssetBuilder.cpp b/Engine/Plugins/Cameras/GameplayCameras/Source/GameplayCameras/Private/Build/CameraShakeAssetBuilder.cpp
--- a/Engine/Plugins/Cameras/GameplayCameras/Source/GameplayCameras/Private/Build/CameraShakeAssetBuilder.cpp
+++ b/Engine/Plugins/Cameras/GameplayCameras/Source/GameplayCameras/Private/Build/CameraShakeAssetBuilder.cpp
@@ -41,14 +41,6 @@ void FCameraShakeAssetBuilder::BuildCameraShake(UCameraShakeAsset* InCameraShake
 
 void FCameraShakeAssetBuilder::BuildCameraShakeImpl()
 {
-	if (!CameraShake->RootNode)
-	{
-		BuildLog.AddMessage(EMessageSeverity::Error, CameraShake, 
-				FText::Format(LOCTEXT("MissingRootNode", "Camera shake '{0}' has no root node."), 
-					FText::FromString(GetPathNameSafe(CameraShake))));
-		return;
-	}
-
 	FCameraNodeHierarchyBuilder NodeBuilder(BuildLog, CameraShake);
 	NodeBuilder.PreBuild();
 
diff --git a/Engine/Plugins/Cameras/GameplayCameras/Source/GameplayCameras/Private/Core/BlendStackCameraNode.cpp b/Engine/Plugins/Cameras/GameplayCameras/Source/GameplayCameras/Private/Core/BlendStackCameraNode.cpp
--- a/Engine/Plugins/Cameras/GameplayCameras/Source/GameplayCameras/Private/Core/BlendStackCameraNode.cpp
+++ b/Engine/Plugins/Cameras/GameplayCameras/Source/GameplayCameras/Private/Core/BlendStackCameraNode.cpp
@@ -73,7 +73,7 @@ FBlendStackCameraNodeEvaluator::~FBlendStackCameraNodeEvaluator()
 	PopEntries(Entries.Num());
 }
 
-bool FBlendStackCameraNodeEvaluator::InitializeEntry(
+void FBlendStackCameraNodeEvaluator::InitializeEntry(
 		FCameraRigEntry& NewEntry, 
 		const UCameraRigAsset* CameraRig,
 		TSharedPtr<const FCameraEvaluationContext> EvaluationContext,
@@ -88,10 +88,6 @@ bool FBlendStackCameraNodeEvaluator::InitializeEntry(
 	BuildParams.RootCameraNode = EntryRootNode;
 	BuildParams.AllocationInfo = &CameraRig->AllocationInfo.EvaluatorInfo;
 	FCameraNodeEvaluator* RootEvaluator = NewEntry.EvaluatorStorage.BuildEvaluatorTree(BuildParams);
-	if (!ensureMsgf(RootEvaluator, TEXT("No root evaluator was created for new camera rig!")))
-	{
-		return false;
-	}
 
 	// Allocate variable table and context data table.
 	NewEntry.ContextResult.VariableTable.Initialize(CameraRig->AllocationInfo.VariableTableInfo);
@@ -116,11 +112,14 @@ bool FBlendStackCameraNodeEvaluator::InitializeEntry(
 	}
 
 	// Initialize the node evaluators.
-	FCameraNodeEvaluatorInitializeParams InitParams(&NewEntry.EvaluatorHierarchy);
-	InitParams.Evaluator = OwningEvaluator;
-	InitParams.EvaluationContext = EvaluationContext;
-	InitParams.LastActiveCameraRigInfo = GetActiveCameraRigEvaluationInfo();
-	RootEvaluator->Initialize(InitParams, NewEntry.ContextResult);  // Initializing with the context result here.
+	if (RootEvaluator)
+	{
+		FCameraNodeEvaluatorInitializeParams InitParams(&NewEntry.EvaluatorHierarchy);
+		InitParams.Evaluator = OwningEvaluator;
+		InitParams.EvaluationContext = EvaluationContext;
+		InitParams.LastActiveCameraRigInfo = GetActiveCameraRigEvaluationInfo();
+		RootEvaluator->Initialize(InitParams, NewEntry.ContextResult);  // Initializing with the context result here.
+	}
 
 	// Set default values for unset entries in the variable table, so that pre-blending from default 
 	// values works.
@@ -133,11 +132,12 @@ bool FBlendStackCameraNodeEvaluator::InitializeEntry(
 	NewEntry.EvaluationContext = EvaluationContext;
 	NewEntry.CameraRig = CameraRig;
 	NewEntry.RootNode = EntryRootNode;
-	NewEntry.RootEvaluator = RootEvaluator->CastThisChecked<FBlendStackRootCameraNodeEvaluator>();
 	NewEntry.Flags.bWasContextInitialResultValid = EvaluationContext->GetInitialResult().bIsValid;
 	NewEntry.Flags.bIsFirstFrame = true;
-
-	return true;
+	if (RootEvaluator)
+	{
+		NewEntry.RootEvaluator = RootEvaluator->CastThisChecked<FBlendStackRootCameraNodeEvaluator>();
+	}
 }
 
 int32 FBlendStackCameraNodeEvaluator::IndexOfEntry(const FBlendStackEntryID EntryID) const
@@ -550,17 +550,12 @@ void FBlendStackCameraNodeEvaluator::OnPostBuildAsset(const FGameplayCameraAsset
 			Entry.RootNode->Blend = NewObject<UPopBlendCameraNode>(Entry.RootNode, NAME_None);
 
 			// Rebuild the evaluator tree.
-			const bool bInitialized = InitializeEntry(
+			InitializeEntry(
 					Entry,
 					Entry.CameraRig,
 					Entry.EvaluationContext.Pin(),
 					Entry.RootNode,
 					Index == Entries.Num() - 1);
-			if (!bInitialized)
-			{
-				Entry.Flags.bIsFrozen = true;
-				continue;
-			}
 
 			OnEntryReinitialized(Index);
 		}
diff --git a/Engine/Plugins/Cameras/GameplayCameras/Source/GameplayCameras/Private/Core/PersistentBlendStackCameraNode.cpp b/Engine/Plugins/Cameras/GameplayCameras/Source/GameplayCameras/Private/Core/PersistentBlendStackCameraNode.cpp
--- a/Engine/Plugins/Cameras/GameplayCameras/Source/GameplayCameras/Private/Core/PersistentBlendStackCameraNode.cpp
+++ b/Engine/Plugins/Cameras/GameplayCameras/Source/GameplayCameras/Private/Core/PersistentBlendStackCameraNode.cpp
@@ -56,16 +56,12 @@ FBlendStackEntryID FPersistentBlendStackCameraNodeEvaluator::Insert(const FBlend
 	}
 
 	FCameraRigEntry NewEntry;
-	const bool bInitialized = InitializeEntry(
+	InitializeEntry(
 			NewEntry, 
 			Params.CameraRig,
 			Params.EvaluationContext,
 			EntryRootNode,
 			false);
-	if (!bInitialized)
-	{
-		return FBlendStackEntryID();
-	}
 
 	FCameraRigEntryExtraInfo NewExtraInfo;
 	NewExtraInfo.StackOrder = Params.StackOrder;
diff --git a/Engine/Plugins/Cameras/GameplayCameras/Source/GameplayCameras/Private/Core/TransientBlendStackCameraNode.cpp b/Engine/Plugins/Cameras/GameplayCameras/Source/GameplayCameras/Private/Core/TransientBlendStackCameraNode.cpp
--- a/Engine/Plugins/Cameras/GameplayCameras/Source/GameplayCameras/Private/Core/TransientBlendStackCameraNode.cpp
+++ b/Engine/Plugins/Cameras/GameplayCameras/Source/GameplayCameras/Private/Core/TransientBlendStackCameraNode.cpp
@@ -102,16 +102,12 @@ FBlendStackEntryID FTransientBlendStackCameraNodeEvaluator::PushNewEntry(const F
 
 	// Make the new stack entry, and use its storage buffer to build the tree of evaluators.
 	FCameraRigEntry NewEntry;
-	const bool bInitialized = InitializeEntry(
+	InitializeEntry(
 			NewEntry, 
 			Params.CameraRig,
 			Params.EvaluationContext,
 			EntryRootNode,
 			true);
-	if (!bInitialized)
-	{
-		return FBlendStackEntryID();
-	}
 	
 #if WITH_EDITOR
 	// Listen to changes to the packages inside which this camera rig is defined. We will hot-reload the
diff --git a/Engine/Plugins/Cameras/GameplayCameras/Source/GameplayCameras/Private/Services/CameraShakeService.cpp b/Engine/Plugins/Cameras/GameplayCameras/Source/GameplayCameras/Private/Services/CameraShakeService.cpp
--- a/Engine/Plugins/Cameras/GameplayCameras/Source/GameplayCameras/Private/Services/CameraShakeService.cpp
+++ b/Engine/Plugins/Cameras/GameplayCameras/Source/GameplayCameras/Private/Services/CameraShakeService.cpp
@@ -40,7 +40,7 @@ private:
 
 	FShakeEntry* AddCameraShake(const FStartCameraShakeParams& Params);
 
-	bool InitializeEntry(
+	void InitializeEntry(
 		FShakeEntry& NewEntry, 
 		const UCameraShakeAsset* CameraShake,
 		TSharedPtr<const FCameraEvaluationContext> EvaluationContext);
@@ -160,11 +160,7 @@ FCameraShakeServiceCameraNodeEvaluator::FShakeEntry* FCameraShakeServiceCameraNo
 	ensure(Params.CameraShake);
 
 	FShakeEntry NewEntry;
-	const bool bInitialized = InitializeEntry(NewEntry, Params.CameraShake, ShakeContext);
-	if (!bInitialized)
-	{
-		return nullptr;
-	}
+	InitializeEntry(NewEntry, Params.CameraShake, ShakeContext);
 
 	NewEntry.ShakeScale = Params.ShakeScale;
 	NewEntry.PlaySpace = Params.PlaySpace;
@@ -289,7 +285,7 @@ void FCameraShakeServiceCameraNodeEvaluator::OnRun(const FCameraNodeEvaluationPa
 	}
 }
 
-bool FCameraShakeServiceCameraNodeEvaluator::InitializeEntry(
+void FCameraShakeServiceCameraNodeEvaluator::InitializeEntry(
 	FShakeEntry& NewEntry, 
 	const UCameraShakeAsset* CameraShake,
 	TSharedPtr<const FCameraEvaluationContext> EvaluationContext)
@@ -299,10 +295,6 @@ bool FCameraShakeServiceCameraNodeEvaluator::InitializeEntry(
 	BuildParams.RootCameraNode = CameraShake->RootNode;
 	BuildParams.AllocationInfo = &CameraShake->AllocationInfo.EvaluatorInfo;
 	FCameraNodeEvaluator* RootEvaluator = NewEntry.EvaluatorStorage.BuildEvaluatorTree(BuildParams);
-	if (!ensureMsgf(RootEvaluator, TEXT("No root evaluator was created for new camera rig!")))
-	{
-		return false;
-	}
 
 	// Generate the blend-in and blend-out evaluators.
 	FBlendCameraNodeEvaluator* BlendInEvaluator = nullptr;
@@ -333,20 +325,24 @@ bool FCameraShakeServiceCameraNodeEvaluator::InitializeEntry(
 	NewEntry.Result.ContextDataTable.OverrideAll(ContextResult.ContextDataTable);
 
 	// Initialize the node evaluators.
-	FCameraNodeEvaluatorInitializeParams InitParams(&NewEntry.EvaluatorHierarchy);
-	InitParams.Evaluator = OwningEvaluator;
-	InitParams.EvaluationContext = EvaluationContext;
-	RootEvaluator->Initialize(InitParams, NewEntry.Result);
+	if (RootEvaluator)
+	{
+		FCameraNodeEvaluatorInitializeParams InitParams(&NewEntry.EvaluatorHierarchy);
+		InitParams.Evaluator = OwningEvaluator;
+		InitParams.EvaluationContext = EvaluationContext;
+		RootEvaluator->Initialize(InitParams, NewEntry.Result);
+	}
 
 	// Wrap up!
 	NewEntry.EvaluationContext = EvaluationContext;
 	NewEntry.CameraShake = CameraShake;
-	NewEntry.RootEvaluator = RootEvaluator->CastThisChecked<FShakeCameraNodeEvaluator>();
 	NewEntry.BlendInEvaluator = BlendInEvaluator;
 	NewEntry.BlendOutEvaluator = BlendOutEvaluator;
 	NewEntry.bIsFirstFrame = true;
-
-	return true;
+	if (RootEvaluator)
+	{
+		NewEntry.RootEvaluator = RootEvaluator->CastThisChecked<FShakeCameraNodeEvaluator>();
+	}
 }
 
 void FCameraShakeServiceCameraNodeEvaluator::PopEntry(int32 EntryIndex)
diff --git a/Engine/Plugins/Cameras/GameplayCameras/Source/GameplayCameras/Public/Core/BlendStackCameraNode.h b/Engine/Plugins/Cameras/GameplayCameras/Source/GameplayCameras/Public/Core/BlendStackCameraNode.h
--- a/Engine/Plugins/Cameras/GameplayCameras/Source/GameplayCameras/Public/Core/BlendStackCameraNode.h
+++ b/Engine/Plugins/Cameras/GameplayCameras/Source/GameplayCameras/Public/Core/BlendStackCameraNode.h
@@ -143,7 +143,7 @@ protected:
 
 	struct FCameraRigEntry;
 
-	bool InitializeEntry(
+	void InitializeEntry(
 		FCameraRigEntry& NewEntry, 
 		const UCameraRigAsset* CameraRig,
 		TSharedPtr<const FCameraEvaluationContext> EvaluationContext,
