Cameras: remove warning when the first camera shake is created

diff --git a/Engine/Plugins/Cameras/GameplayCameras/Source/GameplayCameras/Private/Services/CameraShakeService.cpp b/Engine/Plugins/Cameras/GameplayCameras/Source/GameplayCameras/Private/Services/CameraShakeService.cpp
--- a/Engine/Plugins/Cameras/GameplayCameras/Source/GameplayCameras/Private/Services/CameraShakeService.cpp
+++ b/Engine/Plugins/Cameras/GameplayCameras/Source/GameplayCameras/Private/Services/CameraShakeService.cpp
@@ -395,6 +395,7 @@ void FCameraShakeService::EnsureShakeContextCreated()
 	if (!ShakeContext)
 	{
 		ShakeContext = MakeShared<FCameraEvaluationContext>();
+		ShakeContext->GetInitialResult().bIsValid = true;
 	}
 
 	// Create the camera rig that will contain and run all the camera shakes.
