Cameras: add missing GC collection for context results in blend stacks

diff --git a/Engine/Plugins/Cameras/GameplayCameras/Source/GameplayCameras/Private/Core/BlendStackCameraNode.cpp b/Engine/Plugins/Cameras/GameplayCameras/Source/GameplayCameras/Private/Core/BlendStackCameraNode.cpp
--- a/Engine/Plugins/Cameras/GameplayCameras/Source/GameplayCameras/Private/Core/BlendStackCameraNode.cpp
+++ b/Engine/Plugins/Cameras/GameplayCameras/Source/GameplayCameras/Private/Core/BlendStackCameraNode.cpp
@@ -495,6 +495,7 @@ void FBlendStackCameraNodeEvaluator::OnAddReferencedObjects(FReferenceCollector&
 	{
 		Collector.AddReferencedObject(Entry.CameraRig);
 		Collector.AddReferencedObject(Entry.RootNode);
+		Entry.ContextResult.AddReferencedObjects(Collector);
 		Entry.Result.AddReferencedObjects(Collector);
 	}
 }
