Cameras: protected from rig prefab nodes referencing their own camera rig

diff --git a/Engine/Plugins/Cameras/GameplayCameras/Source/GameplayCameras/Private/Nodes/Common/CameraRigCameraNode.cpp b/Engine/Plugins/Cameras/GameplayCameras/Source/GameplayCameras/Private/Nodes/Common/CameraRigCameraNode.cpp
--- a/Engine/Plugins/Cameras/GameplayCameras/Source/GameplayCameras/Private/Nodes/Common/CameraRigCameraNode.cpp
+++ b/Engine/Plugins/Cameras/GameplayCameras/Source/GameplayCameras/Private/Nodes/Common/CameraRigCameraNode.cpp
@@ -9,6 +9,7 @@
 #include "Core/CameraRigAssetReference.h"
 #include "Helpers/CameraObjectReferenceParameterOverrideEvaluator.h"
 #include "Helpers/CameraObjectInterfaceParameterOverrideHelper.h"
+#include "Logging/TokenizedMessage.h"
 
 #include UE_INLINE_GENERATED_CPP_BY_NAME(CameraRigCameraNode)
 
@@ -32,9 +33,10 @@ FCameraNodeEvaluatorChildrenView FCameraRigCameraNodeEvaluator::OnGetChildren()
 void FCameraRigCameraNodeEvaluator::OnBuild(const FCameraNodeEvaluatorBuildParams& Params)
 {
 	const UCameraRigCameraNode* CameraRigNode = GetCameraNodeAs<UCameraRigCameraNode>();
+	const UCameraRigAsset* OuterCameraRig = CameraRigNode->GetTypedOuter<UCameraRigAsset>();
 	if (const UCameraRigAsset* CameraRig = CameraRigNode->CameraRigReference.GetCameraRig())
 	{
-		if (CameraRig->RootNode)
+		if (CameraRig->RootNode && CameraRig != OuterCameraRig)
 		{
 			CameraRigRootEvaluator = Params.BuildEvaluator(CameraRig->RootNode);
 		}
@@ -89,11 +91,15 @@ void FCameraRigCameraNodeEvaluator::ApplyParameterOverrides(FCameraVariableTable
 
 void UCameraRigCameraNode::OnPreBuild(FCameraBuildLog& BuildLog)
 {
-	// Build the inner camera rig. Silently skip it if it's not set... but we will
+	// Build the inner camera rig. Silently skip it if it's not set or invalid... but we will
 	// report an error in OnBuild about it.
 	if (UCameraRigAsset* CameraRig = CameraRigReference.GetCameraRig())
 	{
-		CameraRig->BuildCameraRig(BuildLog);
+		const UCameraRigAsset* OuterCameraRig = GetTypedOuter<UCameraRigAsset>();
+		if (OuterCameraRig != CameraRig)
+		{
+			CameraRig->BuildCameraRig(BuildLog);
+		}
 	}
 
 	// Make sure the property bag of the camera rig reference is up to date.
@@ -107,23 +113,36 @@ void UCameraRigCameraNode::OnBuild(FCameraObjectBuildContext& BuildContext)
 	UCameraRigAsset* CameraRig = CameraRigReference.GetCameraRig();
 	if (!CameraRig)
 	{
-		BuildContext.BuildLog.AddMessage(EMessageSeverity::Error, this, 
+		BuildContext.BuildLog.AddMessage(EMessageSeverity::Warning, this, 
 				LOCTEXT("MissingCameraRig", "No camera rig specified on camera rig node."));
 		return;
 	}
 
-	// Whatever allocations our inner camera rig needs for its evaluators and
-	// their camera variables, we add that to our camera rig's allocation info.
-	BuildContext.AllocationInfo.Append(CameraRig->AllocationInfo);
+	const UCameraRigAsset* OuterCameraRig = GetTypedOuter<UCameraRigAsset>();
+	if (OuterCameraRig != CameraRig)
+	{
+		// Whatever allocations our inner camera rig needs for its evaluators and
+		// their camera variables, we add that to our camera rig's allocation info.
+		BuildContext.AllocationInfo.Append(CameraRig->AllocationInfo);
+	}
+	else
+	{
+		BuildContext.BuildLog.AddMessage(EMessageSeverity::Error, this, 
+				LOCTEXT("SelfReferenceError", "Circular camera rig references are forbidden."));
+	}
 }
 
 void UCameraRigCameraNode::GatherPackages(FCameraRigPackages& OutPackages) const
 {
 	if (const UCameraRigAsset* CameraRig = CameraRigReference.GetCameraRig())
 	{
+		const UCameraRigAsset* OuterCameraRig = GetTypedOuter<UCameraRigAsset>();
+		if (OuterCameraRig != CameraRig)
+		{
 #if WITH_EDITOR
-		CameraRig->GatherPackages(OutPackages);
+			CameraRig->GatherPackages(OutPackages);
 #endif  // WITH_EDITOR
+		}
 	}
 }
 
