Cameras: fix "Home" button in camera rig editor

When in the transitions graph, it would switch back to the node tree graph because it would match the root node there first. Now we just only check the currently opened graph.

diff --git a/Engine/Plugins/Cameras/GameplayCameras/Source/GameplayCamerasEditor/Private/Editors/SCameraRigAssetEditor.cpp b/Engine/Plugins/Cameras/GameplayCameras/Source/GameplayCamerasEditor/Private/Editors/SCameraRigAssetEditor.cpp
--- a/Engine/Plugins/Cameras/GameplayCameras/Source/GameplayCamerasEditor/Private/Editors/SCameraRigAssetEditor.cpp
+++ b/Engine/Plugins/Cameras/GameplayCameras/Source/GameplayCamerasEditor/Private/Editors/SCameraRigAssetEditor.cpp
@@ -2,10 +2,8 @@
 
 #include "Editors/SCameraRigAssetEditor.h"
 
-#include "Core/CameraNode.h"
 #include "Core/CameraRigAsset.h"
 #include "Core/CameraRigTransition.h"
-#include "EdGraph/EdGraphPin.h"
 #include "Editors/CameraObjectInterfaceParameterGraphNode.h"
 #include "Editors/CameraRigCameraNodeGraphSchema.h"
 #include "Editors/CameraRigTransitionGraphSchema.h"
@@ -14,7 +12,6 @@
 #include "Editors/ObjectTreeGraphNode.h"
 #include "Editors/SCameraNodeGraphEditor.h"
 #include "Editors/SObjectTreeGraphEditor.h"
-#include "ObjectEditorUtils.h"
 #include "Widgets/Layout/SBox.h"
 
 #define LOCTEXT_NAMESPACE "SCameraRigAssetEditor"
@@ -234,43 +231,51 @@ const FObjectTreeGraphConfig& SCameraRigAssetEditor::GetFocusedGraphConfig() con
 
 void SCameraRigAssetEditor::FocusHome()
 {
-	UObjectTreeGraph* Graph = nullptr;
-	TSharedPtr<SObjectTreeGraphEditor> GraphEditor = nullptr;
-
-	switch (CurrentMode)
-	{
-		case ECameraRigAssetEditorMode::NodeGraph:
-			Graph = NodeGraph;
-			GraphEditor = NodeGraphEditor;
-			break;
-		case ECameraRigAssetEditorMode::TransitionGraph:
-			Graph = TransitionGraph;
-			GraphEditor = TransitionGraphEditor;
-	}
-
-	if (Graph && GraphEditor)
-	{
-		FindAndJumpToObjectNode(CameraRigAsset);
-	}
+	FindAndJumpToObjectNode(CameraRigAsset, CurrentMode);
 }
 
 bool SCameraRigAssetEditor::FindAndJumpToObjectNode(UObject* InObject)
 {
-	if (UObjectTreeGraphNode* NodeGraphObjectNode = NodeGraph->FindObjectNode(InObject))
+	if (FindAndJumpToObjectNode(InObject, ECameraRigAssetEditorMode::NodeGraph))
 	{
 		SetEditorMode(ECameraRigAssetEditorMode::NodeGraph);
-		NodeGraphEditor->JumpToNode(NodeGraphObjectNode);
 		return true;
 	}
-	if (UObjectTreeGraphNode* TransitionGraphObjectNode = TransitionGraph->FindObjectNode(InObject))
+	if (FindAndJumpToObjectNode(InObject, ECameraRigAssetEditorMode::TransitionGraph))
 	{
 		SetEditorMode(ECameraRigAssetEditorMode::TransitionGraph);
-		TransitionGraphEditor->JumpToNode(TransitionGraphObjectNode);
 		return true;
 	}
 	return false;
 }
 
+bool SCameraRigAssetEditor::FindAndJumpToObjectNode(UObject* InObject, ECameraRigAssetEditorMode InEditorMode)
+{
+	UObjectTreeGraph* FocusGraph = nullptr;
+	TSharedPtr<SObjectTreeGraphEditor> FocusGraphEditor = nullptr;
+
+	switch (InEditorMode)
+	{
+		case ECameraRigAssetEditorMode::NodeGraph:
+			FocusGraph = NodeGraph;
+			FocusGraphEditor = NodeGraphEditor;
+			break;
+		case ECameraRigAssetEditorMode::TransitionGraph:
+			FocusGraph = TransitionGraph;
+			FocusGraphEditor = TransitionGraphEditor;
+	}
+
+	if (FocusGraph && FocusGraphEditor)
+	{
+		if (UObjectTreeGraphNode* NodeGraphObjectNode = FocusGraph->FindObjectNode(InObject))
+		{
+			FocusGraphEditor->JumpToNode(NodeGraphObjectNode);
+			return true;
+		}
+	}
+	return false;
+}
+
 FText SCameraRigAssetEditor::GetCameraRigAssetName(UObjectTreeGraph* ForGraph) const
 {
 	if (CameraRigAsset && ForGraph)
diff --git a/Engine/Plugins/Cameras/GameplayCameras/Source/GameplayCamerasEditor/Public/Editors/SCameraRigAssetEditor.h b/Engine/Plugins/Cameras/GameplayCameras/Source/GameplayCamerasEditor/Public/Editors/SCameraRigAssetEditor.h
--- a/Engine/Plugins/Cameras/GameplayCameras/Source/GameplayCamerasEditor/Public/Editors/SCameraRigAssetEditor.h
+++ b/Engine/Plugins/Cameras/GameplayCameras/Source/GameplayCamerasEditor/Public/Editors/SCameraRigAssetEditor.h
@@ -104,6 +104,8 @@ protected:
 	FText GetCameraRigAssetName(UObjectTreeGraph* ForGraph) const;
 	bool IsGraphEditorEnabled() const;
 
+	bool FindAndJumpToObjectNode(UObject* InObject, ECameraRigAssetEditorMode InEditorMode);
+
 private:
 
 	/** The asset being edited */
