Cameras: delete references to camera variables when a variable collection is deleted

diff --git a/Engine/Plugins/Cameras/GameplayCameras/Source/GameplayCamerasEditor/Private/GameplayCamerasEditorModule.cpp b/Engine/Plugins/Cameras/GameplayCameras/Source/GameplayCamerasEditor/Private/GameplayCamerasEditorModule.cpp
--- a/Engine/Plugins/Cameras/GameplayCameras/Source/GameplayCamerasEditor/Private/GameplayCamerasEditorModule.cpp
+++ b/Engine/Plugins/Cameras/GameplayCameras/Source/GameplayCamerasEditor/Private/GameplayCamerasEditorModule.cpp
@@ -15,6 +15,7 @@
 #include "Commands/GameplayCamerasDebuggerCommands.h"
 #include "Commands/ObjectTreeGraphEditorCommands.h"
 #include "ComponentVisualizers/GameplayCameraComponentVisualizer.h"
+#include "Core/CameraVariableCollection.h"
 #include "Customizations/CameraAssetReferenceDetailsCustomization.h"
 #include "Customizations/CameraParameterDetailsCustomizations.h"
 #include "Customizations/CameraRigAssetReferenceDetailsCustomization.h"
@@ -30,11 +31,12 @@
 #include "Debugger/SGameplayCamerasDebugger.h"
 #include "Directors/BlueprintCameraDirector.h"
 #include "EdGraph/EdGraph.h"
+#include "Editor.h"
 #include "Editor/UnrealEdEngine.h"
 #include "Editors/GameplayCamerasGraphPanelPinFactory.h"
 #include "Editors/SCameraVariablePicker.h"
 #include "Features/IModularFeatures.h"
-#include "GameFramework/GameplayCameraComponent.h"
+#include "GameFramework/GameplayCameraComponentBase.h"
 #include "GameplayCameras.h"
 #include "GameplayCamerasLiveEditManager.h"
 #include "IGameplayCamerasEditorModule.h"
@@ -45,6 +47,7 @@
 #include "K2Node_Event.h"
 #include "Kismet2/KismetEditorUtilities.h"
 #include "Modules/ModuleManager.h"
+#include "ObjectTools.h"
 #include "PropertyEditorModule.h"
 #include "Sequencer/CameraFramingZoneTrackEditor.h"
 #include "Sequencer/GameplayCameraComponentTrackEditor.h"
@@ -95,6 +98,8 @@ public:
 
 		FCoreDelegates::OnEnginePreExit.AddRaw(this, &FGameplayCamerasEditorModule::OnPreExit);
 
+		FEditorDelegates::OnPreForceDeleteObjects.AddRaw(this, &FGameplayCamerasEditorModule::OnPreForceDeleteObjects);
+
 		RegisterCameraDirectorEditors();
 		RegisterCoreDebugCategories();
 		RegisterRewindDebuggerFeatures();
@@ -135,6 +140,8 @@ public:
 
 		FCoreDelegates::OnPostEngineInit.RemoveAll(this);
 		FCoreDelegates::OnEnginePreExit.RemoveAll(this);
+
+		FEditorDelegates::OnPreForceDeleteObjects.RemoveAll(this);
 	}
 
 	virtual UCameraAssetEditor* CreateCameraAssetEditor(const EToolkitMode::Type Mode, const TSharedPtr<IToolkitHost>& InitToolkitHost, UCameraAsset* CameraAsset) override
@@ -268,6 +275,36 @@ private:
 
 		SGameplayCamerasDebugger::UnregisterTabSpawners();
 	}
+	
+	void OnPreForceDeleteObjects(const TArray<UObject*>& ObjectsToDelete)
+	{
+		TArray<UCameraVariableCollection*> VariableCollectionsToDelete;
+		for (UObject* Object : ObjectsToDelete)
+		{
+			if (UCameraVariableCollection* VariableCollection = Cast<UCameraVariableCollection>(Object))
+			{
+				VariableCollectionsToDelete.Add(VariableCollection);
+			}
+		}
+
+		if (VariableCollectionsToDelete.IsEmpty())
+		{
+			return;
+		}
+
+		// If any variable collection is being force-deleted, let's clear up references
+		// to variables from inside it.
+
+		TArray<UObject*> SubObjectsToDelete;
+		for (UCameraVariableCollection* VariableCollection : VariableCollectionsToDelete)
+		{
+			for (UCameraVariableAsset* Variable : VariableCollection->Variables)
+			{
+				SubObjectsToDelete.Add(Variable);
+			}
+		}
+		ObjectTools::ForceReplaceReferences(nullptr, SubObjectsToDelete);
+	}
 
 	void RegisterCameraDirectorEditors()
 	{
