Cameras: add actor factories for gameplay camera actors

When drag-and-dropping a camera or camera rig asset onto a level or a sequence, this creates the corresponding actor.

diff --git a/Engine/Plugins/Cameras/GameplayCameras/Source/GameplayCamerasEditor/Private/ActorFactories/GameplayCameraActorFactory.cpp b/Engine/Plugins/Cameras/GameplayCameras/Source/GameplayCamerasEditor/Private/ActorFactories/GameplayCameraActorFactory.cpp
new file mode 100644
--- /dev/null
+++ b/Engine/Plugins/Cameras/GameplayCameras/Source/GameplayCamerasEditor/Private/ActorFactories/GameplayCameraActorFactory.cpp
@@ -0,0 +1,60 @@
+// Copyright Epic Games, Inc. All Rights Reserved.
+
+#include "ActorFactories/GameplayCameraActorFactory.h"
+
+#include "Core/CameraAsset.h"
+#include "GameFramework/GameplayCameraActor.h"
+#include "GameFramework/GameplayCameraComponent.h"
+
+#include UE_INLINE_GENERATED_CPP_BY_NAME(GameplayCameraActorFactory)
+
+#define LOCTEXT_NAMESPACE "GameplayCameraActorFactory"
+
+UGameplayCameraActorFactory::UGameplayCameraActorFactory(const FObjectInitializer& ObjInit)
+	: Super(ObjInit)
+{
+	DisplayName = LOCTEXT("DisplayName", "Gameplay Camera Actor");
+	NewActorClass = AGameplayCameraActor::StaticClass();
+}
+
+AActor* UGameplayCameraActorFactory::SpawnActor(UObject* InAsset, ULevel* InLevel, const FTransform& InTransform, const FActorSpawnParameters& InSpawnParams)
+{
+	AGameplayCameraActor* NewActor = Cast<AGameplayCameraActor>(Super::SpawnActor(InAsset, InLevel, InTransform, InSpawnParams));
+	if (NewActor)
+	{
+		if (UCameraAsset* CameraAsset = Cast<UCameraAsset>(InAsset))
+		{
+			UGameplayCameraComponent* NewComponent = NewActor->GetCameraComponent();
+			NewComponent->CameraReference.SetCameraAsset(CameraAsset);
+		}
+	}
+	return NewActor;
+}
+
+bool UGameplayCameraActorFactory::CanCreateActorFrom(const FAssetData& AssetData, FText& OutErrorMsg)
+{
+	if (Super::CanCreateActorFrom(AssetData, OutErrorMsg))
+	{
+		return true;
+	}
+
+	if (AssetData.IsValid() && AssetData.IsInstanceOf(UCameraAsset::StaticClass()))
+	{
+		return true;
+	}
+
+	OutErrorMsg = LOCTEXT("NoCameraAsset", "A valid Gameplay Camera asset must be specified.");
+	return false;
+}
+
+UObject* UGameplayCameraActorFactory::GetAssetFromActorInstance(AActor* ActorInstance)
+{
+	if (AGameplayCameraActor* CameraActor = Cast<AGameplayCameraActor>(ActorInstance))
+	{
+		return CameraActor->GetCameraComponent()->CameraReference.GetCameraAsset();
+	}
+	return nullptr;
+}
+
+#undef LOCTEXT_NAMESPACE
+
diff --git a/Engine/Plugins/Cameras/GameplayCameras/Source/GameplayCamerasEditor/Private/ActorFactories/GameplayCameraActorFactory.h b/Engine/Plugins/Cameras/GameplayCameras/Source/GameplayCamerasEditor/Private/ActorFactories/GameplayCameraActorFactory.h
new file mode 100644
--- /dev/null
+++ b/Engine/Plugins/Cameras/GameplayCameras/Source/GameplayCamerasEditor/Private/ActorFactories/GameplayCameraActorFactory.h
@@ -0,0 +1,26 @@
+// Copyright Epic Games, Inc. All Rights Reserved.
+
+#pragma once
+
+#include "UObject/ObjectMacros.h"
+#include "ActorFactories/ActorFactory.h"
+
+#include "GameplayCameraActorFactory.generated.h"
+
+UCLASS()
+class UGameplayCameraActorFactory : public UActorFactory
+{
+	GENERATED_BODY()
+
+public:
+
+	UGameplayCameraActorFactory(const FObjectInitializer& ObjInit);
+
+public:
+
+	// UActorFactory interface.
+	virtual AActor* SpawnActor(UObject* InAsset, ULevel* InLevel, const FTransform& InTransform, const FActorSpawnParameters& InSpawnParams) override;
+	virtual bool CanCreateActorFrom( const FAssetData& AssetData, FText& OutErrorMsg ) override;
+	virtual UObject* GetAssetFromActorInstance(AActor* ActorInstance) override;
+};
+
diff --git a/Engine/Plugins/Cameras/GameplayCameras/Source/GameplayCamerasEditor/Private/ActorFactories/GameplayCameraRigActorFactory.cpp b/Engine/Plugins/Cameras/GameplayCameras/Source/GameplayCamerasEditor/Private/ActorFactories/GameplayCameraRigActorFactory.cpp
new file mode 100644
--- /dev/null
+++ b/Engine/Plugins/Cameras/GameplayCameras/Source/GameplayCamerasEditor/Private/ActorFactories/GameplayCameraRigActorFactory.cpp
@@ -0,0 +1,60 @@
+// Copyright Epic Games, Inc. All Rights Reserved.
+
+#include "ActorFactories/GameplayCameraRigActorFactory.h"
+
+#include "Core/CameraRigAsset.h"
+#include "GameFramework/GameplayCameraRigActor.h"
+#include "GameFramework/GameplayCameraRigComponent.h"
+
+#include UE_INLINE_GENERATED_CPP_BY_NAME(GameplayCameraRigActorFactory)
+
+#define LOCTEXT_NAMESPACE "GameplayCameraRigActorFactory"
+
+UGameplayCameraRigActorFactory::UGameplayCameraRigActorFactory(const FObjectInitializer& ObjInit)
+	: Super(ObjInit)
+{
+	DisplayName = LOCTEXT("DisplayName", "Gameplay Camera Rig Actor");
+	NewActorClass = AGameplayCameraRigActor::StaticClass();
+}
+
+AActor* UGameplayCameraRigActorFactory::SpawnActor(UObject* InAsset, ULevel* InLevel, const FTransform& InTransform, const FActorSpawnParameters& InSpawnParams)
+{
+	AGameplayCameraRigActor* NewActor = Cast<AGameplayCameraRigActor>(Super::SpawnActor(InAsset, InLevel, InTransform, InSpawnParams));
+	if (NewActor)
+	{
+		if (UCameraRigAsset* CameraRig = Cast<UCameraRigAsset>(InAsset))
+		{
+			UGameplayCameraRigComponent* NewComponent = NewActor->GetCameraRigComponent();
+			NewComponent->CameraRigReference.SetCameraRig(CameraRig);
+		}
+	}
+	return NewActor;
+}
+
+bool UGameplayCameraRigActorFactory::CanCreateActorFrom(const FAssetData& AssetData, FText& OutErrorMsg)
+{
+	if (Super::CanCreateActorFrom(AssetData, OutErrorMsg))
+	{
+		return true;
+	}
+
+	if (AssetData.IsValid() && AssetData.IsInstanceOf(UCameraRigAsset::StaticClass()))
+	{
+		return true;
+	}
+
+	OutErrorMsg = LOCTEXT("NoCameraRigAsset", "A valid Gameplay Camera Rig asset must be specified.");
+	return false;
+}
+
+UObject* UGameplayCameraRigActorFactory::GetAssetFromActorInstance(AActor* ActorInstance)
+{
+	if (AGameplayCameraRigActor* CameraRigActor = Cast<AGameplayCameraRigActor>(ActorInstance))
+	{
+		return CameraRigActor->GetCameraRigComponent()->CameraRigReference.GetCameraRig();
+	}
+	return nullptr;
+}
+
+#undef LOCTEXT_NAMESPACE
+
diff --git a/Engine/Plugins/Cameras/GameplayCameras/Source/GameplayCamerasEditor/Private/ActorFactories/GameplayCameraRigActorFactory.h b/Engine/Plugins/Cameras/GameplayCameras/Source/GameplayCamerasEditor/Private/ActorFactories/GameplayCameraRigActorFactory.h
new file mode 100644
--- /dev/null
+++ b/Engine/Plugins/Cameras/GameplayCameras/Source/GameplayCamerasEditor/Private/ActorFactories/GameplayCameraRigActorFactory.h
@@ -0,0 +1,26 @@
+// Copyright Epic Games, Inc. All Rights Reserved.
+
+#pragma once
+
+#include "UObject/ObjectMacros.h"
+#include "ActorFactories/ActorFactory.h"
+
+#include "GameplayCameraRigActorFactory.generated.h"
+
+UCLASS()
+class UGameplayCameraRigActorFactory : public UActorFactory
+{
+	GENERATED_BODY()
+
+public:
+
+	UGameplayCameraRigActorFactory(const FObjectInitializer& ObjInit);
+
+public:
+
+	// UActorFactory interface.
+	virtual AActor* SpawnActor(UObject* InAsset, ULevel* InLevel, const FTransform& InTransform, const FActorSpawnParameters& InSpawnParams) override;
+	virtual bool CanCreateActorFrom( const FAssetData& AssetData, FText& OutErrorMsg ) override;
+	virtual UObject* GetAssetFromActorInstance(AActor* ActorInstance) override;
+};
+
diff --git a/Engine/Plugins/Cameras/GameplayCameras/Source/GameplayCamerasEditor/Private/GameplayCamerasEditorModule.cpp b/Engine/Plugins/Cameras/GameplayCameras/Source/GameplayCamerasEditor/Private/GameplayCamerasEditorModule.cpp
--- a/Engine/Plugins/Cameras/GameplayCameras/Source/GameplayCamerasEditor/Private/GameplayCamerasEditorModule.cpp
+++ b/Engine/Plugins/Cameras/GameplayCameras/Source/GameplayCamerasEditor/Private/GameplayCamerasEditorModule.cpp
@@ -2,6 +2,8 @@
 
 #include "IGameplayCamerasEditorModule.h"
 
+#include "ActorFactories/GameplayCameraActorFactory.h"
+#include "ActorFactories/GameplayCameraRigActorFactory.h"
 #include "AssetTools/CameraAssetEditor.h"
 #include "AssetTools/CameraRigAssetEditor.h"
 #include "AssetTools/CameraRigProxyAssetEditor.h"
@@ -42,8 +44,8 @@
 #include "IGameplayCamerasEditorModule.h"
 #include "IGameplayCamerasModule.h"
 #include "IRewindDebuggerExtension.h"
-#include "ISettingsModule.h"
 #include "ISequencerModule.h"
+#include "ISettingsModule.h"
 #include "K2Node_Event.h"
 #include "Kismet2/KismetEditorUtilities.h"
 #include "Modules/ModuleManager.h"
@@ -52,6 +54,7 @@
 #include "Sequencer/CameraFramingZoneTrackEditor.h"
 #include "Sequencer/GameplayCameraComponentTrackEditor.h"
 #include "Styles/GameplayCamerasEditorStyle.h"
+#include "Subsystems/PlacementSubsystem.h"
 #include "ToolMenus.h"
 #include "Toolkits/BlueprintCameraDirectorAssetEditorMode.h"
 #include "Toolkits/CameraAssetEditorToolkit.h"
@@ -267,6 +270,16 @@ private:
 		using namespace UE::Cameras;
 
 		SGameplayCamerasDebugger::RegisterTabSpawners();
+
+		GameplayCameraActorFactory.Reset(NewObject<UGameplayCameraActorFactory>());
+		GameplayCameraRigActorFactory.Reset(NewObject<UGameplayCameraRigActorFactory>());
+		GEditor->ActorFactories.Add(GameplayCameraActorFactory.Get());
+		GEditor->ActorFactories.Add(GameplayCameraRigActorFactory.Get());
+		if (UPlacementSubsystem* PlacementSubsystem = GEditor->GetEditorSubsystem<UPlacementSubsystem>())
+		{
+			PlacementSubsystem->RegisterAssetFactory(GameplayCameraActorFactory.Get());
+			PlacementSubsystem->RegisterAssetFactory(GameplayCameraRigActorFactory.Get());
+		}
 	}
 
 	void OnPreExit()
@@ -274,6 +287,15 @@ private:
 		using namespace UE::Cameras;
 
 		SGameplayCamerasDebugger::UnregisterTabSpawners();
+
+		GEditor->ActorFactories.RemoveAll([this](const UActorFactory* ActorFactory)
+				{
+					return ActorFactory == GameplayCameraActorFactory.Get() || ActorFactory == GameplayCameraRigActorFactory.Get();
+				});
+		if (UPlacementSubsystem* PlacementSubsystem = GEditor->GetEditorSubsystem<UPlacementSubsystem>())
+		{
+			PlacementSubsystem->UnregisterAssetFactory(GameplayCameraActorFactory.Get());
+		}
 	}
 	
 	void OnPreForceDeleteObjects(const TArray<UObject*>& ObjectsToDelete)
@@ -626,6 +648,9 @@ private:
 	FDelegateHandle GameplayCameraComponentTrackCreateEditorHandle;
 	FDelegateHandle CameraFramingZoneTrackCreateEditorHandle;
 
+	TStrongObjectPtr<UGameplayCameraActorFactory> GameplayCameraActorFactory;
+	TStrongObjectPtr<UGameplayCameraRigActorFactory> GameplayCameraRigActorFactory;
+
 #if UE_GAMEPLAY_CAMERAS_TRACE
 	TSharedPtr<UE::Cameras::FCameraSystemTraceModule> TraceModule;
 	TSharedPtr<UE::Cameras::FCameraSystemRewindDebuggerExtension> RewindDebuggerExtension;
