Cameras: fix AGameplayCamerasPlayerCameraManager's management of the view target

The goal is to keep the view target in sync with the owner of the active camera evaluation context. This is especially important for scenarios where BP logic or Sequencer sets view targets back and forth.

- Handle the case of the evaluation context owner not being directly the view target, such as when it is a camera component on the view target.
- Add a delegate on the context stack to notify of changes, and use that in the camera manager class to sync the current view target. This handles the case of a gameplay camera being activated via the custom Blueprint methods that take a player controller pointer or index. These activations don't go through SetViewTarget.
- When a normal camera component or arbitrary actor is set as the view target, an evaluation context is created for them. We used to keep these contexts alive via a "referencer service", but move this directly to an array on the camera manager class. It's simpler and works better as we check at the end of every update if any context is done.
- Fix some possible asserts and crashes when using cinematic cameras that get unspawned before we get to run the camera system update.

diff --git a/Engine/Plugins/Cameras/GameplayCameras/Source/GameplayCameras/Private/Core/BlendStackRootCameraNode.cpp b/Engine/Plugins/Cameras/GameplayCameras/Source/GameplayCameras/Private/Core/BlendStackRootCameraNode.cpp
--- a/Engine/Plugins/Cameras/GameplayCameras/Source/GameplayCameras/Private/Core/BlendStackRootCameraNode.cpp
+++ b/Engine/Plugins/Cameras/GameplayCameras/Source/GameplayCameras/Private/Core/BlendStackRootCameraNode.cpp
@@ -166,11 +166,14 @@ void FBlendStackRootCameraNodeEvaluator::MergeCameraRig(const FCameraNodeEvaluat
 
 const UCameraRigAsset* FBlendStackRootCameraNodeEvaluator::FindInnermostCameraRigPrefab(const UCameraRigAsset* CameraRig)
 {
-	if (const UCameraRigCameraNode* PrefabNode = Cast<const UCameraRigCameraNode>(CameraRig->RootNode))
+	if (CameraRig)
 	{
-		if (const UCameraRigAsset* InnerCameraRig = PrefabNode->CameraRigReference.GetCameraRig())
+		if (const UCameraRigCameraNode* PrefabNode = Cast<const UCameraRigCameraNode>(CameraRig->RootNode))
 		{
-			return FindInnermostCameraRigPrefab(InnerCameraRig);
+			if (const UCameraRigAsset* InnerCameraRig = PrefabNode->CameraRigReference.GetCameraRig())
+			{
+				return FindInnermostCameraRigPrefab(InnerCameraRig);
+			}
 		}
 	}
 	return CameraRig;
diff --git a/Engine/Plugins/Cameras/GameplayCameras/Source/GameplayCameras/Private/Core/CameraEvaluationContextStack.cpp b/Engine/Plugins/Cameras/GameplayCameras/Source/GameplayCameras/Private/Core/CameraEvaluationContextStack.cpp
--- a/Engine/Plugins/Cameras/GameplayCameras/Source/GameplayCameras/Private/Core/CameraEvaluationContextStack.cpp
+++ b/Engine/Plugins/Cameras/GameplayCameras/Source/GameplayCameras/Private/Core/CameraEvaluationContextStack.cpp
@@ -53,6 +53,8 @@ void FCameraEvaluationContextStack::PushContext(TSharedRef<FCameraEvaluationCont
 			FContextEntry EntryCopy(MoveTemp(Entries[ExistingIndex]));
 			Entries.RemoveAt(ExistingIndex);
 			Entries.Add(MoveTemp(EntryCopy));
+
+			OnStackChangedEvent.Broadcast();
 		}
 		return;
 	}
@@ -66,6 +68,8 @@ void FCameraEvaluationContextStack::PushContext(TSharedRef<FCameraEvaluationCont
 	
 	NewEntry.WeakContext = Context;
 	Entries.Push(MoveTemp(NewEntry));
+
+	OnStackChangedEvent.Broadcast();
 }
 
 bool FCameraEvaluationContextStack::AddChildContext(TSharedRef<FCameraEvaluationContext> Context, TSharedPtr<FCameraEvaluationContext> ParentContext)
@@ -115,6 +119,8 @@ bool FCameraEvaluationContextStack::RemoveContext(TSharedRef<FCameraEvaluationCo
 			Context->Deactivate(DeactivateParams);
 
 			It.RemoveCurrent();
+
+			OnStackChangedEvent.Broadcast();
 			return true;
 		}
 	}
@@ -146,12 +152,18 @@ bool FCameraEvaluationContextStack::RemoveContextsOwnedBy(UObject* ContextOwner,
 		}
 	}
 
+	if (bRemovedAny)
+	{
+		OnStackChangedEvent.Broadcast();
+	}
+
 	return bRemovedAny;
 }
 
 void FCameraEvaluationContextStack::PopContext()
 {
 	Entries.Pop();
+	OnStackChangedEvent.Broadcast();
 }
 
 void FCameraEvaluationContextStack::GetAllContexts(TArray<TSharedPtr<FCameraEvaluationContext>>& OutContexts) const
@@ -176,6 +188,7 @@ void FCameraEvaluationContextStack::Reset()
 		}
 	}
 	Entries.Reset();
+	OnStackChangedEvent.Broadcast();
 }
 
 void FCameraEvaluationContextStack::Initialize(FCameraSystemEvaluator& InEvaluator)
@@ -199,6 +212,7 @@ void FCameraEvaluationContextStack::OnEndCameraSystemUpdate()
 	// Reset all written-this-frame flags on evaluation contexts, so we properly get those flags set
 	// regardless of when, during next frame, they set their variables. This is because various 
 	// gameplay systems, Blueprint scripting, whatever, might set variables at any time.
+	bool bRemovedAny = false;
 	TArray<TSharedPtr<FCameraEvaluationContext>> ContextsToVisit;
 	for (auto It = Entries.CreateIterator(); It; ++It)
 	{
@@ -209,6 +223,7 @@ void FCameraEvaluationContextStack::OnEndCameraSystemUpdate()
 		else
 		{
 			It.RemoveCurrent();
+			bRemovedAny = true;
 		}
 	}
 	while (!ContextsToVisit.IsEmpty())
@@ -225,6 +240,11 @@ void FCameraEvaluationContextStack::OnEndCameraSystemUpdate()
 			}
 		}
 	}
+
+	if (bRemovedAny)
+	{
+		OnStackChangedEvent.Broadcast();
+	}
 }
 
 }  // namespace UE::Cameras
diff --git a/Engine/Plugins/Cameras/GameplayCameras/Source/GameplayCameras/Private/GameFramework/ActorCameraEvaluationContext.cpp b/Engine/Plugins/Cameras/GameplayCameras/Source/GameplayCameras/Private/GameFramework/ActorCameraEvaluationContext.cpp
--- a/Engine/Plugins/Cameras/GameplayCameras/Source/GameplayCameras/Private/GameFramework/ActorCameraEvaluationContext.cpp
+++ b/Engine/Plugins/Cameras/GameplayCameras/Source/GameplayCameras/Private/GameFramework/ActorCameraEvaluationContext.cpp
@@ -95,14 +95,13 @@ UE_DEFINE_CAMERA_NODE_EVALUATOR(FCameraComponentCameraNodeEvaluator)
 void FCameraComponentCameraNodeEvaluator::OnRun(const FCameraNodeEvaluationParams& Params, FCameraNodeEvaluationResult& OutResult)
 {
 	UCameraComponent* CameraComponent = Cast<UCameraComponent>(Params.EvaluationContext->GetOwner());
-	if (!ensure(CameraComponent))
+
+	if (CameraComponent)
 	{
-		return;
+		FMinimalViewInfo CameraView;
+		CameraComponent->GetCameraView(Params.DeltaTime, CameraView);
+		FActorCameraEvaluationContext::ApplyMinimalViewInfo(CameraView, OutResult);
 	}
-
-	FMinimalViewInfo CameraView;
-	CameraComponent->GetCameraView(Params.DeltaTime, CameraView);
-	FActorCameraEvaluationContext::ApplyMinimalViewInfo(CameraView, OutResult);
 }
 
 class FCalcCameraActorCameraNodeEvaluator : public FCameraNodeEvaluator
@@ -120,14 +119,13 @@ UE_DEFINE_CAMERA_NODE_EVALUATOR(FCalcCameraActorCameraNodeEvaluator)
 void FCalcCameraActorCameraNodeEvaluator::OnRun(const FCameraNodeEvaluationParams& Params, FCameraNodeEvaluationResult& OutResult)
 {
 	AActor* Actor = Cast<AActor>(Params.EvaluationContext->GetOwner());
-	if (!ensure(Actor))
+
+	if (Actor)
 	{
-		return;
+		FMinimalViewInfo CameraView;
+		Actor->CalcCamera(Params.DeltaTime, CameraView);
+		FActorCameraEvaluationContext::ApplyMinimalViewInfo(CameraView, OutResult);
 	}
-
-	FMinimalViewInfo CameraView;
-	Actor->CalcCamera(Params.DeltaTime, CameraView);
-	FActorCameraEvaluationContext::ApplyMinimalViewInfo(CameraView, OutResult);
 }
 
 }  // namespace UE::Cameras
diff --git a/Engine/Plugins/Cameras/GameplayCameras/Source/GameplayCameras/Private/GameFramework/GameplayCamerasPlayerCameraManager.cpp b/Engine/Plugins/Cameras/GameplayCameras/Source/GameplayCameras/Private/GameFramework/GameplayCamerasPlayerCameraManager.cpp
--- a/Engine/Plugins/Cameras/GameplayCameras/Source/GameplayCameras/Private/GameFramework/GameplayCamerasPlayerCameraManager.cpp
+++ b/Engine/Plugins/Cameras/GameplayCameras/Source/GameplayCameras/Private/GameFramework/GameplayCamerasPlayerCameraManager.cpp
@@ -9,6 +9,7 @@
 #include "Core/CameraEvaluationService.h"
 #include "Core/CameraRigTransition.h"
 #include "Core/CameraSystemEvaluator.h"
+#include "Core/RootCameraNode.h"
 #include "Core/RootCameraNodeCameraRigEvent.h"
 #include "Engine/Canvas.h"
 #include "Engine/Engine.h"
@@ -23,6 +24,10 @@
 namespace UE::Cameras
 {
 
+/**
+ * A blend node evaluator for UViewTargetTransitionParamsBlendCameraNode, which emulates the basic engine
+ * blend curves for view targets.
+ */
 class FViewTargetTransitionParamsBlendCameraNodeEvaluator : public FSimpleBlendCameraNodeEvaluator
 {
 	UE_DECLARE_BLEND_CAMERA_NODE_EVALUATOR_EX(GAMEPLAYCAMERAS_API, FViewTargetTransitionParamsBlendCameraNodeEvaluator, FSimpleBlendCameraNodeEvaluator)
@@ -64,57 +69,6 @@ void FViewTargetTransitionParamsBlendCameraNodeEvaluator::OnComputeBlendFactor(c
 
 UE_DEFINE_BLEND_CAMERA_NODE_EVALUATOR(FViewTargetTransitionParamsBlendCameraNodeEvaluator)
 
-class FViewTargetContextReferencerService : public FCameraEvaluationService
-{
-public:
-
-	FViewTargetContextReferencerService()
-	{
-		SetEvaluationServiceFlags(ECameraEvaluationServiceFlags::NeedsRootCameraNodeEvents);
-	}
-
-	void AddViewTargetContext(AActor* InViewTarget, TSharedRef<FCameraEvaluationContext> InContext)
-	{
-		Entries.Add(FEntry{ InViewTarget, InContext });
-	}
-
-protected:
-
-	virtual void OnInitialize(const FCameraEvaluationServiceInitializeParams& Params) override
-	{
-		Evaluator = Params.Evaluator;
-	}
-
-	virtual void OnRootCameraNodeEvent(const FRootCameraNodeCameraRigEvent& InEvent) override
-	{
-		if (InEvent.EventType == ERootCameraNodeCameraRigEventType::Deactivated)
-		{
-			TSharedPtr<const FCameraEvaluationContext> Context = InEvent.CameraRigInfo.EvaluationContext;
-			if (Context)
-			{
-				for (auto It = Entries.CreateIterator(); It; ++It)
-				{
-					if (It->Context == Context)
-					{
-						It.RemoveCurrent();
-					}
-				}
-			}
-		}
-	}
-
-private:
-
-	struct FEntry
-	{
-		TWeakObjectPtr<AActor> WeakViewTarget;
-		TSharedPtr<FCameraEvaluationContext> Context;
-	};
-
-	FCameraSystemEvaluator* Evaluator = nullptr;
-	TArray<FEntry> Entries;
-};
-
 }  // namespace UE::Cameras
 
 AGameplayCamerasPlayerCameraManager::AGameplayCamerasPlayerCameraManager(const FObjectInitializer& ObjectInitializer)
@@ -124,7 +78,7 @@ AGameplayCamerasPlayerCameraManager::AGameplayCamerasPlayerCameraManager(const F
 
 void AGameplayCamerasPlayerCameraManager::BeginDestroy()
 {
-	TeardownCameraSystemHost();
+	DestroyCameraSystem();
 
 	Super::BeginDestroy();
 }
@@ -168,7 +122,7 @@ void AGameplayCamerasPlayerCameraManager::ReleasePlayerController()
 
 	OriginalCameraManager = nullptr;
 
-	TeardownCameraSystemHost();
+	DestroyCameraSystem();
 
 	PCOwner = nullptr;
 }
@@ -212,13 +166,20 @@ void AGameplayCamerasPlayerCameraManager::StopCameraModifierRig(FCameraRigInstan
 
 void AGameplayCamerasPlayerCameraManager::InitializeFor(APlayerController* PlayerController)
 {
+	using namespace UE::Cameras;
+
 	if (!bOverrideViewRotationMode)
 	{
 		const UGameplayCamerasSettings* Settings = GetDefault<UGameplayCamerasSettings>();
 		ViewRotationMode = Settings->DefaultViewRotationMode;
 	}
 
-	EnsureCameraSystemHost();
+	EnsureCameraSystemInitialized();
+	if (ensure(CameraSystemEvaluator))
+	{
+		FCameraEvaluationContextStack& ContextStack = CameraSystemEvaluator->GetEvaluationContextStack();
+		ContextStack.OnStackChanged().AddUObject(this, &AGameplayCamerasPlayerCameraManager::OnContextStackChanged);
+	}
 
 	Super::InitializeFor(PlayerController);
 }
@@ -231,9 +192,14 @@ void AGameplayCamerasPlayerCameraManager::SetViewTarget(AActor* NewViewTarget, F
 	// If that context owner isn't an actor, and isn't inside an actor (like a component), we use the player
 	// controller as the view target.
 
+	ensure(bIsSettingNewViewTarget == false);
+	TGuardValue<bool> ReentrancyGuard(bIsSettingNewViewTarget, true);
+	FCameraEvaluationContextStack& ContextStack = CameraSystemEvaluator->GetEvaluationContextStack();
+
+	// If the view target is null, this is sort of a shortcut for "we're done with the current view target",
+	// so pop the context stack and reactivate the previous context.
 	if (NewViewTarget == nullptr)
 	{
-		FCameraEvaluationContextStack& ContextStack = CameraSystemEvaluator->GetEvaluationContextStack();
 		ContextStack.PopContext();
 
 		if (TSharedPtr<FCameraEvaluationContext> NewActiveContext = ContextStack.GetActiveContext())
@@ -257,64 +223,57 @@ void AGameplayCamerasPlayerCameraManager::SetViewTarget(AActor* NewViewTarget, F
 		return;
 	}
 
-	if (UGameplayCameraComponentBase* GameplayCameraComponent = NewViewTarget->FindComponentByClass<UGameplayCameraComponentBase>())
-	{
-		GameplayCameraComponent->ActivateCameraForPlayerController(PCOwner);
-	}
-	else if (UCameraComponent* CameraComponent = NewViewTarget->FindComponentByClass<UCameraComponent>())
+	// See if we can find the view target in the context stack. If so, reactivate it instead of potentially
+	// making a new context for the same thing.
+	bool bFoundContext = false;
+	TArray<TSharedPtr<FCameraEvaluationContext>> CurrentContexts;
+	ContextStack.GetAllContexts(CurrentContexts);
+	for (TSharedPtr<FCameraEvaluationContext> CurrentContext : CurrentContexts)
 	{
-		TSharedRef<FActorCameraEvaluationContext> NewContext = MakeShared<FActorCameraEvaluationContext>(CameraComponent);
-		CameraSystemEvaluator->PushEvaluationContext(NewContext);
-	}
-	else
-	{
-		TSharedRef<FActorCameraEvaluationContext> NewContext = MakeShared<FActorCameraEvaluationContext>(NewViewTarget);
-		CameraSystemEvaluator->PushEvaluationContext(NewContext);
+		UObject* CurrentContextOwner = CurrentContext ? CurrentContext->GetOwner() : nullptr;
+		if (CurrentContextOwner && 
+				(CurrentContext->GetOwner() == NewViewTarget || 
+				 CurrentContextOwner->GetTypedOuter<AActor>() == NewViewTarget))
+		{
+			// This will move the context to the top if it's already in the stack (which it is, we
+			// found it there).
+			ContextStack.PushContext(CurrentContext.ToSharedRef());
+			bFoundContext = true;
+		}
 	}
 
-	// If transition parameters were given, override the next activation for the new evaluation context.
-	TSharedPtr<FCameraEvaluationContext> NextContext = CameraSystemEvaluator->GetEvaluationContextStack().GetActiveContext();
-	if (NextContext)
+	if (!bFoundContext)
 	{
-		ViewTargetContextReferencerService->AddViewTargetContext(NewViewTarget, NextContext.ToSharedRef());
-
-		if (TransitionParams.BlendTime > 0.f)
+		if (UGameplayCameraComponentBase* GameplayCameraComponent = NewViewTarget->FindComponentByClass<UGameplayCameraComponentBase>())
 		{
-			UViewTargetTransitionParamsBlendCameraNode* BlendNode = NewObject<UViewTargetTransitionParamsBlendCameraNode>(GetTransientPackage());
-			BlendNode->TransitionParams = TransitionParams;
-
-			UCameraRigTransition* Transition = NewObject<UCameraRigTransition>(GetTransientPackage());
-			Transition->Blend = BlendNode;
-
-			FCameraDirectorEvaluator* DirectorEvaluator = NextContext->GetDirectorEvaluator();
-			DirectorEvaluator->OverrideNextActivationTransition(Transition);
+			GameplayCameraComponent->ActivateCameraForPlayerController(PCOwner);
+		}
+		else if (UCameraComponent* CameraComponent = NewViewTarget->FindComponentByClass<UCameraComponent>())
+		{
+			TSharedRef<FActorCameraEvaluationContext> NewContext = MakeShared<FActorCameraEvaluationContext>(CameraComponent);
+			CameraSystemEvaluator->PushEvaluationContext(NewContext);
+			ViewTargetContexts.Add(NewContext);
+		}
+		else
+		{
+			TSharedRef<FActorCameraEvaluationContext> NewContext = MakeShared<FActorCameraEvaluationContext>(NewViewTarget);
+			CameraSystemEvaluator->PushEvaluationContext(NewContext);
+			ViewTargetContexts.Add(NewContext);
 		}
 	}
-}
 
-void AGameplayCamerasPlayerCameraManager::EnsureCameraSystemHost()
-{
-	using namespace UE::Cameras;
-
-	if (!HasCameraSystem())
+	// If transition parameters were given, override the next activation for the new evaluation context.
+	TSharedPtr<FCameraEvaluationContext> NextContext = CameraSystemEvaluator->GetEvaluationContextStack().GetActiveContext();
+	if (NextContext && TransitionParams.BlendTime > 0.f)
 	{
-		InitializeCameraSystem();
+		UViewTargetTransitionParamsBlendCameraNode* BlendNode = NewObject<UViewTargetTransitionParamsBlendCameraNode>(GetTransientPackage());
+		BlendNode->TransitionParams = TransitionParams;
 
-		ViewTargetContextReferencerService = MakeShared<FViewTargetContextReferencerService>();
-		CameraSystemEvaluator->RegisterEvaluationService(ViewTargetContextReferencerService.ToSharedRef());
-	}
-}
+		UCameraRigTransition* Transition = NewObject<UCameraRigTransition>(GetTransientPackage());
+		Transition->Blend = BlendNode;
 
-void AGameplayCamerasPlayerCameraManager::TeardownCameraSystemHost()
-{
-	if (HasCameraSystem())
-	{
-		if (ensure(ViewTargetContextReferencerService))
-		{
-			CameraSystemEvaluator->UnregisterEvaluationService(ViewTargetContextReferencerService.ToSharedRef());
-		}
-
-		DestroyCameraSystem();
+		FCameraDirectorEvaluator* DirectorEvaluator = NextContext->GetDirectorEvaluator();
+		DirectorEvaluator->OverrideNextActivationTransition(Transition);
 	}
 }
 
@@ -370,6 +329,55 @@ void AGameplayCamerasPlayerCameraManager::DoUpdateCamera(float DeltaTime)
 		FillCameraCache(DesiredView);
 
 		LastFrameDesiredView = DesiredView;
+
+		CleanUpViewTargetContexts();
+	}
+}
+
+void AGameplayCamerasPlayerCameraManager::OnContextStackChanged()
+{
+	using namespace UE::Cameras;
+
+	// When the context stack changes, such as when a gameplay camera component activates directly
+	// against our camera system host, we want to update the view target so that it's always in sync
+	// with whichever owns the active evaluation context.
+	//
+	// This is as opposed to going through SetViewTarget or some other APlayerCameraManager method.
+
+	if (ensure(CameraSystemEvaluator) && !bIsSettingNewViewTarget)
+	{
+		TGuardValue<bool> ReentrancyGuard(bIsSettingNewViewTarget, true);
+
+		FCameraEvaluationContextStack& ContextStack = CameraSystemEvaluator->GetEvaluationContextStack();
+		TSharedPtr<FCameraEvaluationContext> ActiveContext = ContextStack.GetActiveContext();
+		UObject* ActiveContextOwner = ActiveContext->GetOwner();
+		if (ActiveContextOwner)
+		{
+			ViewTarget.SetNewTarget(ActiveContextOwner->GetTypedOuter<AActor>());
+		}
+		else
+		{
+			ViewTarget.SetNewTarget(nullptr);
+		}
+
+		ViewTarget.CheckViewTarget(PCOwner);
+		BlendParams = FViewTargetTransitionParams();
+	}
+}
+
+void AGameplayCamerasPlayerCameraManager::CleanUpViewTargetContexts()
+{
+	using namespace UE::Cameras;
+
+	FRootCameraNodeEvaluator* RootEvaluator = CameraSystemEvaluator->GetRootNodeEvaluator();
+
+	for (auto It = ViewTargetContexts.CreateIterator(); It; ++It)
+	{
+		TSharedRef<FCameraEvaluationContext> Context(*It);
+		if (!RootEvaluator->HasAnyRunningCameraRig(Context))
+		{
+			It.RemoveCurrent();
+		}
 	}
 }
 
diff --git a/Engine/Plugins/Cameras/GameplayCameras/Source/GameplayCameras/Public/Core/CameraEvaluationContextStack.h b/Engine/Plugins/Cameras/GameplayCameras/Source/GameplayCameras/Public/Core/CameraEvaluationContextStack.h
--- a/Engine/Plugins/Cameras/GameplayCameras/Source/GameplayCameras/Public/Core/CameraEvaluationContextStack.h
+++ b/Engine/Plugins/Cameras/GameplayCameras/Source/GameplayCameras/Public/Core/CameraEvaluationContextStack.h
@@ -64,6 +64,9 @@ public:
 	/** Empties the stack of all contexts. */
 	void Reset();
 
+	/** An event raised when the stack changes. */ 
+	FSimpleMulticastDelegate& OnStackChanged() { return OnStackChangedEvent; }
+
 public:
 
 	// Internal API
@@ -84,6 +87,9 @@ private:
 	/** The owner evaluator. */
 	FCameraSystemEvaluator* Evaluator = nullptr;
 
+	/** An event raised when the stack changes. */ 
+	FSimpleMulticastDelegate OnStackChangedEvent;
+
 #if UE_GAMEPLAY_CAMERAS_DEBUG
 	friend class FCameraDirectorTreeDebugBlock;
 #endif  // UE_GAMEPLAY_CAMERAS_DEBUG
diff --git a/Engine/Plugins/Cameras/GameplayCameras/Source/GameplayCameras/Public/GameFramework/GameplayCamerasPlayerCameraManager.h b/Engine/Plugins/Cameras/GameplayCameras/Source/GameplayCameras/Public/GameFramework/GameplayCamerasPlayerCameraManager.h
--- a/Engine/Plugins/Cameras/GameplayCameras/Source/GameplayCameras/Public/GameFramework/GameplayCamerasPlayerCameraManager.h
+++ b/Engine/Plugins/Cameras/GameplayCameras/Source/GameplayCameras/Public/GameFramework/GameplayCamerasPlayerCameraManager.h
@@ -14,6 +14,7 @@ struct FCameraRigInstanceID;
 namespace UE::Cameras
 {
 
+class FCameraEvaluationContext;
 class FCameraSystemEvaluator;
 class FViewTargetContextReferencerService;
 
@@ -113,11 +114,11 @@ protected:
 
 private:
 
-	void EnsureCameraSystemHost();
-	void TeardownCameraSystemHost();
-
 	void RunViewRotationPreviewUpdate(float DeltaTime, FRotator& OutViewRotation, FRotator& OutDeltaRot);
 
+	void OnContextStackChanged();
+	void CleanUpViewTargetContexts();
+
 public:
 
 	/** The view rotation handling mode to use. */
@@ -133,9 +134,11 @@ private:
 	UPROPERTY(Transient)
 	TObjectPtr<APlayerCameraManager> OriginalCameraManager;
 
-	TSharedPtr<UE::Cameras::FViewTargetContextReferencerService> ViewTargetContextReferencerService;
+	TArray<TSharedRef<UE::Cameras::FCameraEvaluationContext>> ViewTargetContexts;
 
 	FMinimalViewInfo LastFrameDesiredView;
+
+	bool bIsSettingNewViewTarget = false;
 };
 
 /**
