Cameras: add "auto-binding" for showing the debug info of whatever camera system is currently in control

Debug ID 0 is now an "auto ID" that maps to either the player camera manager or the current view target, if they host a camera system.

The camera debugger panel gains a menu entry for it when PIE is running, and some work has been done for Rewind Debugger to record all running camera system instances and allow switching between them at replay time.

diff --git a/Engine/Plugins/Cameras/GameplayCameras/Source/GameplayCameras/Private/Core/CameraSystemEvaluator.cpp b/Engine/Plugins/Cameras/GameplayCameras/Source/GameplayCameras/Private/Core/CameraSystemEvaluator.cpp
--- a/Engine/Plugins/Cameras/GameplayCameras/Source/GameplayCameras/Private/Core/CameraSystemEvaluator.cpp
+++ b/Engine/Plugins/Cameras/GameplayCameras/Source/GameplayCameras/Private/Core/CameraSystemEvaluator.cpp
@@ -30,7 +30,7 @@ DECLARE_CYCLE_STAT(TEXT("Camera System Eval"), CameraSystemEval_Total, STATGROUP
 namespace UE::Cameras
 {
 
-extern int32 GGameplayCamerasDebugEnableID;
+extern bool GGameplayCamerasDebugEnable;
 
 void FCameraSystemEvaluationResult::Reset()
 {
@@ -561,7 +561,7 @@ bool FCameraSystemEvaluator::IsDebugTraceEnabled()
 bool FCameraSystemEvaluator::ShouldBuildOrDrawDebugBlocks()
 {
 	const bool bTraceEnabled = IsDebugTraceEnabled();
-	return bTraceEnabled || GGameplayCamerasDebugEnableID >= 0;
+	return bTraceEnabled || GGameplayCamerasDebugEnable;
 }
 
 void FCameraSystemEvaluator::BuildDebugBlocksIfNeeded()
@@ -600,7 +600,7 @@ void FCameraSystemEvaluator::DebugUpdate(const FCameraSystemDebugUpdateParams& P
 #endif
 
 	FCameraDebugRenderer Renderer(OwnerWorld, Params.CanvasObject, Params.bIsDebugCameraEnabled);
-	RootDebugBlock->RootDebugDraw(Renderer);
+	RootDebugBlock->RootDebugDraw(Renderer, Params.bForceDraw);
 }
 
 #endif  // UE_GAMEPLAY_CAMERAS_DEBUG
diff --git a/Engine/Plugins/Cameras/GameplayCameras/Source/GameplayCameras/Private/Debug/CameraSystemDebugRegistry.cpp b/Engine/Plugins/Cameras/GameplayCameras/Source/GameplayCameras/Private/Debug/CameraSystemDebugRegistry.cpp
--- a/Engine/Plugins/Cameras/GameplayCameras/Source/GameplayCameras/Private/Debug/CameraSystemDebugRegistry.cpp
+++ b/Engine/Plugins/Cameras/GameplayCameras/Source/GameplayCameras/Private/Debug/CameraSystemDebugRegistry.cpp
@@ -27,14 +27,14 @@ FCameraSystemDebugID FCameraSystemDebugRegistry::RegisterCameraSystemEvaluator(T
 	const FString OwnerName = GetNameSafe(InEvaluator->GetOwner());
 	FEntry NewEntry{ InEvaluator, OwnerName };
 	const int32 NewIndex = Entries.Add(NewEntry);
-	return FCameraSystemDebugID(NewIndex);
+	return FCameraSystemDebugID(NewIndex + 1);
 }
 
 void FCameraSystemDebugRegistry::UnregisterCameraSystemEvaluator(FCameraSystemDebugID InDebugID)
 {
-	if (ensure(InDebugID.IsValid()))
+	if (ensure(InDebugID.IsValid() && !InDebugID.IsAny() && !InDebugID.IsAuto()))
 	{
-		Entries.RemoveAt(InDebugID.Value);
+		Entries.RemoveAt(InDebugID.Value - 1);
 	}
 }
 
diff --git a/Engine/Plugins/Cameras/GameplayCameras/Source/GameplayCameras/Private/Debug/CameraSystemTrace.cpp b/Engine/Plugins/Cameras/GameplayCameras/Source/GameplayCameras/Private/Debug/CameraSystemTrace.cpp
--- a/Engine/Plugins/Cameras/GameplayCameras/Source/GameplayCameras/Private/Debug/CameraSystemTrace.cpp
+++ b/Engine/Plugins/Cameras/GameplayCameras/Source/GameplayCameras/Private/Debug/CameraSystemTrace.cpp
@@ -29,6 +29,7 @@ UE_TRACE_CHANNEL(CameraSystemChannel)
 UE_TRACE_EVENT_BEGIN(CameraSystem, CameraSystemEvaluation)
 	UE_TRACE_EVENT_FIELD(uint64, Cycle)
 	UE_TRACE_EVENT_FIELD(double, RecordingTime)
+	UE_TRACE_EVENT_FIELD(int32, CameraSystemDebugID)
 	UE_TRACE_EVENT_FIELD(double, EvaluatedLocationX)
 	UE_TRACE_EVENT_FIELD(double, EvaluatedLocationY)
 	UE_TRACE_EVENT_FIELD(double, EvaluatedLocationZ)
@@ -281,7 +282,7 @@ bool FCameraSystemTrace::IsTraceEnabled()
 	return GGameplayCamerasDebugTrace || UE_TRACE_CHANNELEXPR_IS_ENABLED(CameraSystemChannel);
 }
 
-void FCameraSystemTrace::TraceEvaluation(UWorld* InWorld, const FCameraSystemEvaluationResult& InResult, FCameraDebugBlock& InRootDebugBlock)
+void FCameraSystemTrace::TraceEvaluation(UWorld* InWorld, const FCameraSystemEvaluationResult& InResult, FRootCameraDebugBlock& InRootDebugBlock)
 {
 	if (!IsTraceEnabled())
 	{
@@ -296,6 +297,7 @@ void FCameraSystemTrace::TraceEvaluation(UWorld* InWorld, const FCameraSystemEva
 	UE_TRACE_LOG(CameraSystem, CameraSystemEvaluation, CameraSystemChannel)
 		<< CameraSystemEvaluation.Cycle(FPlatformTime::Cycles64())
 		<< CameraSystemEvaluation.RecordingTime(FObjectTrace::GetWorldElapsedTime(InWorld))
+		<< CameraSystemEvaluation.CameraSystemDebugID(InRootDebugBlock.GetDebugID().GetValue())
 		<< CameraSystemEvaluation.EvaluatedLocationX(InResult.CameraPose.GetLocation().X)
 		<< CameraSystemEvaluation.EvaluatedLocationY(InResult.CameraPose.GetLocation().Y)
 		<< CameraSystemEvaluation.EvaluatedLocationZ(InResult.CameraPose.GetLocation().Z)
diff --git a/Engine/Plugins/Cameras/GameplayCameras/Source/GameplayCameras/Private/Debug/RootCameraDebugBlock.cpp b/Engine/Plugins/Cameras/GameplayCameras/Source/GameplayCameras/Private/Debug/RootCameraDebugBlock.cpp
--- a/Engine/Plugins/Cameras/GameplayCameras/Source/GameplayCameras/Private/Debug/RootCameraDebugBlock.cpp
+++ b/Engine/Plugins/Cameras/GameplayCameras/Source/GameplayCameras/Private/Debug/RootCameraDebugBlock.cpp
@@ -27,17 +27,23 @@
 namespace UE::Cameras
 {
 
-int32 GGameplayCamerasDebugEnableID = INDEX_NONE;
-static FAutoConsoleVariableRef CVarGameplayCamerasDebugEnableID(
-	TEXT("GameplayCameras.Debug.EnableID"),
-	GGameplayCamerasDebugEnableID,
-	TEXT("(Default: -1. Enables debug drawing for the GamplayCameras system with the given debug ID."));
+bool GGameplayCamerasDebugEnable = false;
+static FAutoConsoleVariableRef CVarGameplayCamerasDebugEnable(
+	TEXT("GameplayCameras.Debug.Enable"),
+	GGameplayCamerasDebugEnable,
+	TEXT("(Default: false. Enables debug drawing for the GameplayCameras system."));
+
+int32 GGameplayCamerasDebugSystemID = INDEX_NONE;
+static FAutoConsoleVariableRef CVarGameplayCamerasDebugSystemID(
+	TEXT("GameplayCameras.Debug.SystemID"),
+	GGameplayCamerasDebugSystemID,
+	TEXT("(Default: -1. Specifies the GameplayCameras system instance to show debug drawing for."));
 
 FString GGameplayCamerasDebugCategories = "nodetree";
 static FAutoConsoleVariableRef CVarGameplayCamerasDebugCategories(
 	TEXT("GameplayCameras.Debug.Categories"),
 	GGameplayCamerasDebugCategories,
-	TEXT("(Default: nodes. Specifies which debug categories to display the GamplayCameras system."));
+	TEXT("(Default: nodes. Specifies which debug categories to display the GameplayCameras system."));
 
 bool GGameplayCamerasDebugPoseStatsShowUnchanged = false;
 static FAutoConsoleVariableRef CVarGameplayCamerasDebugPoseStatsShowUnchanged(
@@ -135,10 +141,16 @@ void FRootCameraDebugBlock::BuildDebugBlocks(const FCameraSystemEvaluator& Camer
 	AddChild(&Builder.BuildDebugBlock<FViewfinderDebugBlock>());
 }
 
-void FRootCameraDebugBlock::RootDebugDraw(FCameraDebugRenderer& Renderer)
+void FRootCameraDebugBlock::RootDebugDraw(FCameraDebugRenderer& Renderer, bool bForceDraw)
 {
-	FCameraSystemDebugID WantedDebugID(GGameplayCamerasDebugEnableID);
-	if (WantedDebugID != DebugID && !WantedDebugID.IsAny())
+	if (!GGameplayCamerasDebugEnable)
+	{
+		return;
+	}
+
+	FCameraSystemDebugID WantedDebugID(GGameplayCamerasDebugSystemID);
+	const bool bDoDebugDraw = (bForceDraw || WantedDebugID.IsAny() || WantedDebugID == DebugID);
+	if (!bDoDebugDraw)
 	{
 		return;
 	}
diff --git a/Engine/Plugins/Cameras/GameplayCameras/Source/GameplayCameras/Private/GameFramework/IGameplayCameraSystemHost.cpp b/Engine/Plugins/Cameras/GameplayCameras/Source/GameplayCameras/Private/GameFramework/IGameplayCameraSystemHost.cpp
--- a/Engine/Plugins/Cameras/GameplayCameras/Source/GameplayCameras/Private/GameFramework/IGameplayCameraSystemHost.cpp
+++ b/Engine/Plugins/Cameras/GameplayCameras/Source/GameplayCameras/Private/GameFramework/IGameplayCameraSystemHost.cpp
@@ -2,14 +2,23 @@
 
 #include "GameFramework/IGameplayCameraSystemHost.h"
 
+#include "Camera/PlayerCameraManager.h"
 #include "Core/CameraEvaluationContext.h"
 #include "Core/CameraEvaluationContextStack.h"
 #include "Core/CameraSystemEvaluator.h"
+#include "Debug/CameraSystemDebugRegistry.h"
 #include "Debug/DebugDrawService.h"
 #include "Engine/Canvas.h"
 #include "Engine/World.h"
 #include "GameFramework/PlayerController.h"
 
+namespace UE::Cameras
+{
+
+extern int32 GGameplayCamerasDebugSystemID;
+
+}  // namespace UE::Cameras
+
 void IGameplayCameraSystemHost::InitializeCameraSystem()
 {
 	using namespace UE::Cameras;
@@ -134,18 +143,48 @@ void IGameplayCameraSystemHost::DebugDraw(UCanvas* Canvas, APlayerController* Pl
 
 	if (CameraSystemEvaluator.IsValid())
 	{
-		// We're looking from the outside if we are not the view target, or if we don't have a player
-		// anymore (which happens in spectator mode like with the debug camera).
-		bool bIsDebugCameraEnabled = false;
-		if (TSharedPtr<FCameraEvaluationContext> ActiveContext = CameraSystemEvaluator->GetEvaluationContextStack().GetActiveContext())
+		UWorld* OwnerWorld = nullptr;
+		if (UObject* CameraSystemOwner = CameraSystemEvaluator->GetOwner())
+		{
+			OwnerWorld = CameraSystemOwner->GetWorld();
+		}
+
+		// Find the actual player controller as best we can.
+		APlayerController* ActualPlayerController = PlayerController;
+		if (!ActualPlayerController)
 		{
-			const APlayerController* ActivePlayerController = ActiveContext->GetPlayerController();
-			bIsDebugCameraEnabled = !ActivePlayerController || !ActivePlayerController->Player;
+			const FCameraEvaluationContextStack& ContextStack = CameraSystemEvaluator->GetEvaluationContextStack();
+			if (TSharedPtr<FCameraEvaluationContext> ActiveContext = ContextStack.GetActiveContext())
+			{
+				ActualPlayerController = ActiveContext->GetPlayerController();
+			}
 		}
 
+		if (!ActualPlayerController && OwnerWorld)
+		{
+			ActualPlayerController = OwnerWorld->GetFirstPlayerController();
+		}
+
+		const UObject* ThisAsObject = GetAsObject();
+		const AActor* ViewTarget = ActualPlayerController ? ActualPlayerController->GetViewTarget() : nullptr;
+		const bool bThisIsCameraManager = (ActualPlayerController && ThisAsObject == ActualPlayerController->PlayerCameraManager);
+		const bool bThisIsViewTarget = (ThisAsObject && ViewTarget && ThisAsObject->GetTypedOuter<AActor>() == ViewTarget);
+
+		// We're looking from the outside if we are not the view target, or if we don't have a player
+		// anymore (which happens in spectator mode like with the debug camera).
+		bool bIsDebugCameraEnabled = (
+				(!bThisIsCameraManager && !bThisIsViewTarget) ||
+				!ActualPlayerController || !ActualPlayerController->Player);
+
+		// Force draw this host's camera system if the wanted debug ID is "auto" and we are the
+		// view target or camera manager.
+		FCameraSystemDebugID WantedDebugID(GGameplayCamerasDebugSystemID);
+		bool bForceDraw = WantedDebugID.IsAuto() && (bThisIsCameraManager || bThisIsViewTarget);
+
 		FCameraSystemDebugUpdateParams DebugUpdateParams;
 		DebugUpdateParams.CanvasObject = Canvas;
 		DebugUpdateParams.bIsDebugCameraEnabled = bIsDebugCameraEnabled;
+		DebugUpdateParams.bForceDraw = bForceDraw;
 		CameraSystemEvaluator->DebugUpdate(DebugUpdateParams);
 	}
 }
diff --git a/Engine/Plugins/Cameras/GameplayCameras/Source/GameplayCameras/Public/Core/CameraSystemEvaluator.h b/Engine/Plugins/Cameras/GameplayCameras/Source/GameplayCameras/Public/Core/CameraSystemEvaluator.h
--- a/Engine/Plugins/Cameras/GameplayCameras/Source/GameplayCameras/Public/Core/CameraSystemEvaluator.h
+++ b/Engine/Plugins/Cameras/GameplayCameras/Source/GameplayCameras/Public/Core/CameraSystemEvaluator.h
@@ -133,6 +133,9 @@ struct FCameraSystemDebugUpdateParams
 
 	/** Whether the debug camera is enabled, giving an "outside" view of camera system. */
 	bool bIsDebugCameraEnabled = false;
+
+	/** Whether to force drawing debug info for this camera system. */
+	bool bForceDraw = false;
 };
 
 #endif  // UE_GAMEPLAY_CAMERAS_DEBUG
diff --git a/Engine/Plugins/Cameras/GameplayCameras/Source/GameplayCameras/Public/Debug/CameraSystemDebugRegistry.h b/Engine/Plugins/Cameras/GameplayCameras/Source/GameplayCameras/Public/Debug/CameraSystemDebugRegistry.h
--- a/Engine/Plugins/Cameras/GameplayCameras/Source/GameplayCameras/Public/Debug/CameraSystemDebugRegistry.h
+++ b/Engine/Plugins/Cameras/GameplayCameras/Source/GameplayCameras/Public/Debug/CameraSystemDebugRegistry.h
@@ -24,13 +24,15 @@ struct FCameraSystemDebugID
 	FCameraSystemDebugID() : Value(INVALID) {}
 	FCameraSystemDebugID(int32 InValue) : Value(InValue) {}
 
-	bool IsValid() const { return Value != INVALID; }
+	bool IsValid() const { return Value >= 0; }
+	bool IsAuto() const { return Value == AUTO; }
 	bool IsAny() const { return Value == ANY; }
 	operator bool() const { return IsValid(); }
 
 	int32 GetValue() const { return Value; }
 
 	static FCameraSystemDebugID Invalid() { return FCameraSystemDebugID(INVALID); }
+	static FCameraSystemDebugID Auto() { return FCameraSystemDebugID(AUTO); }
 	static FCameraSystemDebugID Any() { return FCameraSystemDebugID(ANY); }
 
 public:
@@ -59,7 +61,8 @@ public:
 
 private:
 
-	constexpr static int32 INVALID = INDEX_NONE;
+	constexpr static int32 INVALID = -1;
+	constexpr static int32 AUTO = 0;
 	constexpr static int32 ANY = MAX_int32;
 
 	int32 Value;
diff --git a/Engine/Plugins/Cameras/GameplayCameras/Source/GameplayCameras/Public/Debug/CameraSystemTrace.h b/Engine/Plugins/Cameras/GameplayCameras/Source/GameplayCameras/Public/Debug/CameraSystemTrace.h
--- a/Engine/Plugins/Cameras/GameplayCameras/Source/GameplayCameras/Public/Debug/CameraSystemTrace.h
+++ b/Engine/Plugins/Cameras/GameplayCameras/Source/GameplayCameras/Public/Debug/CameraSystemTrace.h
@@ -16,6 +16,7 @@ namespace UE::Cameras
 
 class FCameraDebugBlock;
 class FCameraDebugBlockStorage;
+class FRootCameraDebugBlock;
 struct FCameraSystemEvaluationResult;
 
 /**
@@ -34,7 +35,7 @@ public:
 	/** Returns whether tracing of camera system evaluation is enabled. */
 	GAMEPLAYCAMERAS_API static bool IsTraceEnabled();
 	/** Records one frame of camera system evaluation. */
-	GAMEPLAYCAMERAS_API static void TraceEvaluation(UWorld* InWorld, const FCameraSystemEvaluationResult& InResult, FCameraDebugBlock& InRootDebugBlock);
+	GAMEPLAYCAMERAS_API static void TraceEvaluation(UWorld* InWorld, const FCameraSystemEvaluationResult& InResult, FRootCameraDebugBlock& InRootDebugBlock);
 	/** Reads back one frame of camera system evaluation. */
 	GAMEPLAYCAMERAS_API static FCameraDebugBlock* ReadEvaluationTrace(TArray<uint8> InSerializedBlocks, FCameraDebugBlockStorage& InStorage);
 };
diff --git a/Engine/Plugins/Cameras/GameplayCameras/Source/GameplayCameras/Public/Debug/RootCameraDebugBlock.h b/Engine/Plugins/Cameras/GameplayCameras/Source/GameplayCameras/Public/Debug/RootCameraDebugBlock.h
--- a/Engine/Plugins/Cameras/GameplayCameras/Source/GameplayCameras/Public/Debug/RootCameraDebugBlock.h
+++ b/Engine/Plugins/Cameras/GameplayCameras/Source/GameplayCameras/Public/Debug/RootCameraDebugBlock.h
@@ -13,8 +13,10 @@ namespace UE::Cameras
 class FCameraSystemEvaluator;
 struct FCameraDebugBlockBuildParams;
 struct FCameraDebugBlockBuilder;
+struct FCameraSystemDebugID;
 
-GAMEPLAYCAMERAS_API extern int32 GGameplayCamerasDebugEnableID;
+GAMEPLAYCAMERAS_API extern bool GGameplayCamerasDebugEnable;
+GAMEPLAYCAMERAS_API extern int32 GGameplayCamerasDebugSystemID;
 GAMEPLAYCAMERAS_API extern FString GGameplayCamerasDebugCategories;
 
 /**
@@ -30,7 +32,10 @@ public:
 	GAMEPLAYCAMERAS_API void BuildDebugBlocks(const FCameraSystemEvaluator& CameraSystem, const FCameraDebugBlockBuildParams& Params, FCameraDebugBlockBuilder& Builder);
 
 	/** Initiate the debug drawing. */
-	GAMEPLAYCAMERAS_API void RootDebugDraw(FCameraDebugRenderer& Renderer);
+	GAMEPLAYCAMERAS_API void RootDebugDraw(FCameraDebugRenderer& Renderer, bool bForceDraw = false);
+
+	/** Gets the debug ID of the camera system that generated this debug info. */
+	const FCameraSystemDebugID& GetDebugID() const { return DebugID; }
 
 protected:
 
diff --git a/Engine/Plugins/Cameras/GameplayCameras/Source/GameplayCamerasEditor/Private/Debugger/SGameplayCamerasDebugger.cpp b/Engine/Plugins/Cameras/GameplayCameras/Source/GameplayCamerasEditor/Private/Debugger/SGameplayCamerasDebugger.cpp
--- a/Engine/Plugins/Cameras/GameplayCameras/Source/GameplayCamerasEditor/Private/Debugger/SGameplayCamerasDebugger.cpp
+++ b/Engine/Plugins/Cameras/GameplayCameras/Source/GameplayCamerasEditor/Private/Debugger/SGameplayCamerasDebugger.cpp
@@ -16,6 +16,7 @@
 #include "Framework/Docking/TabManager.h"
 #include "Framework/MultiBox/MultiBoxBuilder.h"
 #include "GameFramework/Actor.h"
+#include "GameFramework/PlayerController.h"
 #include "GameplayCamerasEditorSettings.h"
 #include "IGameplayCamerasEditorModule.h"
 #include "Modules/ModuleManager.h"
@@ -185,7 +186,7 @@ void SGameplayCamerasDebugger::Construct(const FArguments& InArgs)
 	CommandList->MapAction(
 			Commands.EnableDebugInfo,
 			FExecuteAction::CreateSP(this, &SGameplayCamerasDebugger::ToggleDebugDraw),
-			FCanExecuteAction::CreateSP(this, &SGameplayCamerasDebugger::CanToggleDebugDraw),
+			FCanExecuteAction(),
 			FIsActionChecked::CreateSP(this, &SGameplayCamerasDebugger::IsDebugDrawing));
 
 	// Build all UI elements.
@@ -243,6 +244,7 @@ void SGameplayCamerasDebugger::Tick(const FGeometry& AllottedGeometry, const dou
 	if (bRefreshDebugID)
 	{
 		// Auto-set the camera system debug ID when PIE starts/ends, and for other similar events.
+		FCameraSystemDebugID DebugID;
 		UWorld* DebugWorld = DebugContext.GetContext();
 		if (DebugWorld)
 		{
@@ -256,26 +258,11 @@ void SGameplayCamerasDebugger::Tick(const FGeometry& AllottedGeometry, const dou
 			}
 			else
 			{
-				FCameraSystemDebugRegistry::FRegisteredCameraSystems CameraSystems;
-				FCameraSystemDebugRegistry::Get().GetRegisteredCameraSystemEvaluators(CameraSystems);
-
-				DebugID = FCameraSystemDebugID::Invalid();
-				for (TSharedPtr<FCameraSystemEvaluator> CameraSystem : CameraSystems)
-				{
-					UObject* CameraSystemOwner = CameraSystem->GetOwner();
-					if (CameraSystemOwner && CameraSystemOwner->GetWorld() == DebugWorld)
-					{
-						DebugID = CameraSystem->GetDebugID();
-						break;
-					}
-				}
+				DebugID = FCameraSystemDebugID::Auto();
 			}
 		}
-		else
-		{
-			DebugID = FCameraSystemDebugID::Invalid();
-		}
 
+		GGameplayCamerasDebugSystemID = DebugID.GetValue();
 		bRefreshDebugID = false;
 	}
 }
@@ -526,29 +513,17 @@ void SGameplayCamerasDebugger::ConstructDebugPanels()
 
 void SGameplayCamerasDebugger::ToggleDebugDraw()
 {
-	if (GGameplayCamerasDebugEnableID < 0)
-	{
-		GGameplayCamerasDebugEnableID = DebugID.GetValue();
-	}
-	else
-	{
-		GGameplayCamerasDebugEnableID = -1;
-	}
-}
-
-bool SGameplayCamerasDebugger::CanToggleDebugDraw() const
-{
-	return DebugID.IsValid();
+	GGameplayCamerasDebugEnable = !GGameplayCamerasDebugEnable;
 }
 
 bool SGameplayCamerasDebugger::IsDebugDrawing() const
 {
-	return GGameplayCamerasDebugEnableID >= 0;
+	return GGameplayCamerasDebugEnable;
 }
 
 FText SGameplayCamerasDebugger::GetToggleDebugDrawText() const
 {
-	if (GGameplayCamerasDebugEnableID >= 0)
+	if (GGameplayCamerasDebugEnable)
 	{
 		return LOCTEXT("DebugInfoEnabled", "Debug Info Enabled");
 	}
@@ -560,7 +535,7 @@ FText SGameplayCamerasDebugger::GetToggleDebugDrawText() const
 
 FSlateIcon SGameplayCamerasDebugger::GetToggleDebugDrawIcon() const
 {
-	if (GGameplayCamerasDebugEnableID >= 0)
+	if (GGameplayCamerasDebugEnable)
 	{
 		return FSlateIcon(GameplayCamerasEditorStyleName, "Debugger.DebugInfoEnabled.Icon");
 	}
@@ -603,6 +578,21 @@ void SGameplayCamerasDebugger::GetCameraSystemPickerContent(UToolMenu* ToolMenu)
 		}
 		else if (CameraSystems.Num() > 0)
 		{
+			CameraSystemsSection.AddMenuEntry(
+					TEXT("AutoBindToViewTarget"),
+					LOCTEXT("AutoBindToViewTarget", "Auto-bind to the current view target"),
+					LOCTEXT("AutoBindToViewTargetToolTip", "Show the debug info for the view target of the local player."),
+					FSlateIcon(),
+					FUIAction(
+						FExecuteAction::CreateSP(
+							this, &SGameplayCamerasDebugger::BindToCameraSystem, FCameraSystemDebugID::Auto()),
+						FCanExecuteAction(),
+						FIsActionChecked::CreateSP(
+							this, &SGameplayCamerasDebugger::IsBoundToCameraSystem, FCameraSystemDebugID::Auto())),
+						EUserInterfaceActionType::Check);
+
+			CameraSystemsSection.AddSeparator(NAME_None);
+
 			for (TSharedPtr<FCameraSystemEvaluator> CameraSystem : CameraSystems)
 			{
 				UObject* CameraSystemOwner = CameraSystem->GetOwner();
@@ -647,13 +637,12 @@ void SGameplayCamerasDebugger::GetCameraSystemPickerContent(UToolMenu* ToolMenu)
 
 void SGameplayCamerasDebugger::BindToCameraSystem(FCameraSystemDebugID InDebugID)
 {
-	DebugID = InDebugID;
-	GGameplayCamerasDebugEnableID = InDebugID.GetValue();
+	GGameplayCamerasDebugSystemID = InDebugID.GetValue();
 }
 
 bool SGameplayCamerasDebugger::IsBoundToCameraSystem(FCameraSystemDebugID InDebugID)
 {
-	return DebugID == InDebugID;
+	return GGameplayCamerasDebugSystemID == InDebugID.GetValue();
 }
 
 void SGameplayCamerasDebugger::OnDebugContextChanged()
diff --git a/Engine/Plugins/Cameras/GameplayCameras/Source/GameplayCamerasEditor/Private/Debugger/SGameplayCamerasDebugger.h b/Engine/Plugins/Cameras/GameplayCameras/Source/GameplayCamerasEditor/Private/Debugger/SGameplayCamerasDebugger.h
--- a/Engine/Plugins/Cameras/GameplayCameras/Source/GameplayCamerasEditor/Private/Debugger/SGameplayCamerasDebugger.h
+++ b/Engine/Plugins/Cameras/GameplayCameras/Source/GameplayCamerasEditor/Private/Debugger/SGameplayCamerasDebugger.h
@@ -87,7 +87,6 @@ protected:
 	void SetActiveDebugCategoryPanel(FString InCategoryName);
 
 	void ToggleDebugDraw();
-	bool CanToggleDebugDraw() const;
 	bool IsDebugDrawing() const;
 	FText GetToggleDebugDrawText() const;
 	FSlateIcon GetToggleDebugDrawIcon() const;
@@ -104,7 +103,6 @@ private:
 	FName GameplayCamerasEditorStyleName;
 
 	FGameplayCamerasDebuggerContext DebugContext;
-	FCameraSystemDebugID DebugID;
 	bool bRefreshDebugID = false;
 
 	TSharedPtr<SBox> PanelHost;
diff --git a/Engine/Plugins/Cameras/GameplayCameras/Source/GameplayCamerasEditor/Private/Trace/CameraSystemRewindDebuggerExtension.cpp b/Engine/Plugins/Cameras/GameplayCameras/Source/GameplayCamerasEditor/Private/Trace/CameraSystemRewindDebuggerExtension.cpp
--- a/Engine/Plugins/Cameras/GameplayCameras/Source/GameplayCamerasEditor/Private/Trace/CameraSystemRewindDebuggerExtension.cpp
+++ b/Engine/Plugins/Cameras/GameplayCameras/Source/GameplayCamerasEditor/Private/Trace/CameraSystemRewindDebuggerExtension.cpp
@@ -53,13 +53,17 @@ void FCameraSystemRewindDebuggerExtension::Update(float DeltaTime, IRewindDebugg
 	if (FrameProvider.GetFrameFromTime(ETraceFrameType::TraceFrameType_Game, CurrentTraceTime, Frame))
 	{
 		const FCameraSystemTraceFrameData* FoundFrameData = nullptr;
+		const FCameraSystemDebugID WantedDebugID(GGameplayCamerasDebugSystemID);
 
 		const FCameraSystemTraceTimeline* CameraSystemTimeline = CameraSystemProvider->GetTimeline();
 		CameraSystemTimeline->EnumerateEvents(
 				Frame.StartTime, Frame.EndTime, 
-				[&FoundFrameData](double InStartTime, double InEndTime, uint32 InDepth, const FCameraSystemTraceFrameData& FrameData)  
+				[&FoundFrameData, WantedDebugID](double InStartTime, double InEndTime, uint32 InDepth, const FCameraSystemTraceFrameData& FrameData)  
 				{
-					FoundFrameData = &FrameData;
+					if (WantedDebugID.IsAny() || WantedDebugID.GetValue() == FrameData.CameraSystemID)
+					{
+						FoundFrameData = &FrameData;
+					}
 					return TraceServices::EEventEnumerate::Continue;
 				});
 
diff --git a/Engine/Plugins/Cameras/GameplayCameras/Source/GameplayCamerasEditor/Private/Trace/CameraSystemTraceAnalyzer.cpp b/Engine/Plugins/Cameras/GameplayCameras/Source/GameplayCamerasEditor/Private/Trace/CameraSystemTraceAnalyzer.cpp
--- a/Engine/Plugins/Cameras/GameplayCameras/Source/GameplayCamerasEditor/Private/Trace/CameraSystemTraceAnalyzer.cpp
+++ b/Engine/Plugins/Cameras/GameplayCameras/Source/GameplayCamerasEditor/Private/Trace/CameraSystemTraceAnalyzer.cpp
@@ -31,6 +31,7 @@ bool FCameraSystemTraceAnalyzer::OnEvent(uint16 RouteId, EStyle Style, const FOn
 
 	uint64 Cycle = Context.EventData.GetValue<uint64>("Cycle");
 	double RecordingTime = Context.EventData.GetValue<double>("RecordingTime");
+	FrameData.CameraSystemID = Context.EventData.GetValue<int32>("CameraSystemDebugID");
 	FrameData.EvaluatedLocation.X = Context.EventData.GetValue<double>("EvaluatedLocationX");
 	FrameData.EvaluatedLocation.Y = Context.EventData.GetValue<double>("EvaluatedLocationY");
 	FrameData.EvaluatedLocation.Z = Context.EventData.GetValue<double>("EvaluatedLocationZ");
diff --git a/Engine/Plugins/Cameras/GameplayCameras/Source/GameplayCamerasEditor/Private/Trace/CameraSystemTraceProvider.h b/Engine/Plugins/Cameras/GameplayCameras/Source/GameplayCamerasEditor/Private/Trace/CameraSystemTraceProvider.h
--- a/Engine/Plugins/Cameras/GameplayCameras/Source/GameplayCamerasEditor/Private/Trace/CameraSystemTraceProvider.h
+++ b/Engine/Plugins/Cameras/GameplayCameras/Source/GameplayCamerasEditor/Private/Trace/CameraSystemTraceProvider.h
@@ -19,6 +19,7 @@ namespace UE::Cameras
 /** Data for one frame of tracing. */
 struct FCameraSystemTraceFrameData
 {
+	int32 CameraSystemID;
 	FVector3d EvaluatedLocation;
 	FRotator3d EvaluatedRotation;
 	float EvaluatedFieldOfView;
