Cameras: fix unit tests

- An empty camera rig isn't supposed to spit out an error anymore.
- FCameraSystemEvaluator must now be created as a shared pointer since it needs a weak pointer for the debug registry.
- Prevent the context stack from firing a "changed" event while we're destroying the camera system.

diff --git a/Engine/Plugins/Cameras/GameplayCameras/Source/GameplayCameras/Private/Core/CameraSystemEvaluator.cpp b/Engine/Plugins/Cameras/GameplayCameras/Source/GameplayCameras/Private/Core/CameraSystemEvaluator.cpp
--- a/Engine/Plugins/Cameras/GameplayCameras/Source/GameplayCameras/Private/Core/CameraSystemEvaluator.cpp
+++ b/Engine/Plugins/Cameras/GameplayCameras/Source/GameplayCameras/Private/Core/CameraSystemEvaluator.cpp
@@ -117,6 +117,7 @@ void FCameraSystemEvaluator::Initialize(const FCameraSystemEvaluatorCreateParams
 
 FCameraSystemEvaluator::~FCameraSystemEvaluator()
 {
+	ContextStack.OnStackChanged().Clear();
 	ContextStack.Reset();
 
 	{
diff --git a/Engine/Plugins/Cameras/GameplayCameras/Source/GameplayCameras/Private/Tests/CameraRigAssetBuilderTests.cpp b/Engine/Plugins/Cameras/GameplayCameras/Source/GameplayCameras/Private/Tests/CameraRigAssetBuilderTests.cpp
--- a/Engine/Plugins/Cameras/GameplayCameras/Source/GameplayCameras/Private/Tests/CameraRigAssetBuilderTests.cpp
+++ b/Engine/Plugins/Cameras/GameplayCameras/Source/GameplayCameras/Private/Tests/CameraRigAssetBuilderTests.cpp
@@ -15,19 +15,11 @@ bool FCameraRigAssetBuilderNullTest::RunTest(const FString& Parameters)
 {
 	using namespace UE::Cameras::Test;
 
-	UCameraRigAsset* CameraRig = FCameraRigAssetTestBuilder(TEXT("InvalidTest")).Get();
+	UCameraRigAsset* CameraRig = FCameraRigAssetTestBuilder(TEXT("EmptyTest")).Get();
 	UTEST_EQUAL("Dirty status", CameraRig->BuildStatus, ECameraBuildStatus::Dirty);
 
-	FStringFormatOrderedArguments ErrorArgs;
-	ErrorArgs.Add(CameraRig->GetPathName());
-	AddExpectedMessage(
-			FString::Format(TEXT("Camera rig '{0}' has no root node."), ErrorArgs),
-			ELogVerbosity::Error,
-			EAutomationExpectedMessageFlags::Contains,
-			1,
-			false);
 	CameraRig->BuildCameraRig();
-	UTEST_EQUAL("Error status", CameraRig->BuildStatus, ECameraBuildStatus::WithErrors);
+	UTEST_EQUAL("Clean status", CameraRig->BuildStatus, ECameraBuildStatus::Clean);
 
 	return true;
 }
diff --git a/Engine/Plugins/Cameras/GameplayCameras/Source/GameplayCameras/Private/Tests/CameraSystemUpdateTests.cpp b/Engine/Plugins/Cameras/GameplayCameras/Source/GameplayCameras/Private/Tests/CameraSystemUpdateTests.cpp
--- a/Engine/Plugins/Cameras/GameplayCameras/Source/GameplayCameras/Private/Tests/CameraSystemUpdateTests.cpp
+++ b/Engine/Plugins/Cameras/GameplayCameras/Source/GameplayCameras/Private/Tests/CameraSystemUpdateTests.cpp
@@ -29,18 +29,18 @@ bool FCameraSystemFrameFlagsTest::RunTest(const FString& Parameters)
 		.BuildCameraAsset()
 		.Get();
 
-	FCameraSystemEvaluator Evaluator;
-	Evaluator.Initialize();
-	Evaluator.PushEvaluationContext(EvaluationContext);
+	TSharedRef<FCameraSystemEvaluator> Evaluator = FCameraSystemEvaluatorBuilder::Build();
+	Evaluator->Initialize();
+	Evaluator->PushEvaluationContext(EvaluationContext);
 
 	EvaluationContext->GetInitialResult().bIsValid = true;
 
 	FCameraSystemEvaluationParams Params;
 	Params.DeltaTime = 0.3f;
-	Evaluator.Update(Params);
+	Evaluator->Update(Params);
 
 	FCameraRigEvaluationInfo CameraRigInfo;
-	FRootCameraNodeEvaluator* RootEvaluator = Evaluator.GetRootNodeEvaluator();
+	FRootCameraNodeEvaluator* RootEvaluator = Evaluator->GetRootNodeEvaluator();
 	RootEvaluator->GetActiveCameraRigInfo(CameraRigInfo);
 	UTEST_NOT_NULL("RootEvaluator", CameraRigInfo.RootEvaluator);
 
@@ -55,7 +55,7 @@ bool FCameraSystemFrameFlagsTest::RunTest(const FString& Parameters)
 		UTEST_TRUE("IsFirstFrame", Update.bIsFirstFrame);
 	}
 
-	Evaluator.Update(Params);
+	Evaluator->Update(Params);
 
 	{
 		UTEST_EQUAL("NumReceivedUpdates", UpdateTracker->ReceivedUpdates.Num(), 2);
@@ -67,7 +67,7 @@ bool FCameraSystemFrameFlagsTest::RunTest(const FString& Parameters)
 
 	EvaluationContext->GetInitialResult().bIsCameraCut = true;
 
-	Evaluator.Update(Params);
+	Evaluator->Update(Params);
 
 	{
 		UTEST_EQUAL("NumReceivedUpdates", UpdateTracker->ReceivedUpdates.Num(), 3);
diff --git a/Engine/Plugins/Cameras/GameplayCameras/Source/GameplayCameras/Private/Tests/GameplayCamerasTestBuilder.h b/Engine/Plugins/Cameras/GameplayCameras/Source/GameplayCameras/Private/Tests/GameplayCamerasTestBuilder.h
--- a/Engine/Plugins/Cameras/GameplayCameras/Source/GameplayCameras/Private/Tests/GameplayCamerasTestBuilder.h
+++ b/Engine/Plugins/Cameras/GameplayCameras/Source/GameplayCameras/Private/Tests/GameplayCamerasTestBuilder.h
@@ -10,6 +10,7 @@
 #include "Core/CameraParameters.h"
 #include "Core/CameraRigAsset.h"
 #include "Core/CameraRigTransition.h"
+#include "Core/CameraSystemEvaluator.h"
 #include "Core/CameraVariableReferences.h"
 #include "Core/CameraVariableTableFwd.h"
 #include "Directors/SingleCameraDirector.h"
@@ -738,5 +739,21 @@ private:
 	TSharedPtr<FNamedObjectRegistry> NamedObjectRegistry;
 };
 
+/**
+ * Builder class for a camera system evaluator.
+ */
+class FCameraSystemEvaluatorBuilder
+{
+public:
+
+	/** Makes a new camera system evaluator. */
+	static TSharedRef<FCameraSystemEvaluator> Build(UObject* OwnerObject = nullptr)
+	{
+		TSharedRef<FCameraSystemEvaluator> NewEvaluator = MakeShared<FCameraSystemEvaluator>();
+		NewEvaluator->Initialize(OwnerObject);
+		return NewEvaluator;
+	}
+};
+
 }  // namespace UE::Cameras::Test
 
