Cameras: add option on camera components to set control rotation

This fixes a regression caused by having removed the use of the camera system actor: there wasn't anything left to set the control rotation when using player input camera nodes.

diff --git a/Engine/Plugins/Cameras/GameplayCameras/Source/GameplayCameras/Private/GameFramework/GameplayCameraComponentBase.cpp b/Engine/Plugins/Cameras/GameplayCameras/Source/GameplayCameras/Private/GameFramework/GameplayCameraComponentBase.cpp
--- a/Engine/Plugins/Cameras/GameplayCameras/Source/GameplayCameras/Private/GameFramework/GameplayCameraComponentBase.cpp
+++ b/Engine/Plugins/Cameras/GameplayCameras/Source/GameplayCameras/Private/GameFramework/GameplayCameraComponentBase.cpp
@@ -482,11 +482,13 @@ void UGameplayCameraComponentBase::TickComponent(float DeltaTime, ELevelTick Tic
 	else
 	{
 		UpdateCameraSystem(DeltaTime);
+		UpdateControlRotationIfNeeded();
 	}
 
 #else
 
 	UpdateCameraSystem(DeltaTime);
+	UpdateControlRotationIfNeeded();
 
 #endif  // WITH_EDITOR
 
@@ -553,6 +555,38 @@ void UGameplayCameraComponentBase::UpdateCameraEvaluationContext(bool bForceAppl
 #endif  // WITH_EDITOR
 }
 
+void UGameplayCameraComponentBase::UpdateControlRotationIfNeeded()
+{
+	using namespace UE::Cameras;
+
+	if (!HasCameraSystem() || !HasCameraEvaluationContext() || !bSetControlRotationWhenViewTarget)
+	{
+		return;
+	}
+
+	APlayerController* PlayerController = EvaluationContext->GetPlayerController();
+	if (!PlayerController)
+	{
+		return;
+	}
+
+	// If the player camera manager is hosting a camera system, it probably already handles control
+	// rotation in its own way.
+	if (IGameplayCameraSystemHost* CameraManagerHost = Cast<IGameplayCameraSystemHost>(PlayerController->PlayerCameraManager))
+	{
+		return;
+	}
+
+	// Set control rotation if we are the view target.
+	AActor* OwnerActor = GetOwner();
+	if (OwnerActor && PlayerController->GetViewTarget() == OwnerActor)
+	{
+		const FCameraSystemEvaluationResult& Result = CameraSystemEvaluator->GetEvaluatedResult();
+		const FRotator3d& ControlRotation = Result.CameraPose.GetRotation();
+		PlayerController->SetControlRotation(ControlRotation);
+	}
+}
+
 #if WITH_EDITOR
 
 void UGameplayCameraComponentBase::ReinitializeCameraEvaluationContext(
diff --git a/Engine/Plugins/Cameras/GameplayCameras/Source/GameplayCameras/Public/GameFramework/GameplayCameraComponentBase.h b/Engine/Plugins/Cameras/GameplayCameras/Source/GameplayCameras/Public/GameFramework/GameplayCameraComponentBase.h
--- a/Engine/Plugins/Cameras/GameplayCameras/Source/GameplayCameras/Public/GameFramework/GameplayCameraComponentBase.h
+++ b/Engine/Plugins/Cameras/GameplayCameras/Source/GameplayCameras/Public/GameFramework/GameplayCameraComponentBase.h
@@ -168,6 +168,8 @@ protected:
 	void UpdateCameraEvaluationContext(bool bForceApplyParameterOverrides);
 	bool HasCameraEvaluationContext() const { return EvaluationContext.IsValid(); }
 
+	void UpdateControlRotationIfNeeded();
+
 #if WITH_EDITOR
 	void ReinitializeCameraEvaluationContext(
 			const FCameraVariableTableAllocationInfo& VariableTableAllocationInfo,
@@ -201,6 +203,15 @@ public:
 	UPROPERTY(BlueprintReadWrite, EditAnywhere, Category=Activation, meta=(EditCondition="bAutoActivate"))
 	TEnumAsByte<EAutoReceiveInput::Type> AutoActivateForPlayer;
 
+	/**
+	 * Specifies whether this component should set the player controller's control rotation 
+	 * to the computed point of view's orientation every frame. This is only used when a 
+	 * player controller is associated with this component, and the view target is that
+	 * component.
+	 */
+	UPROPERTY(BlueprintReadWrite, EditAnywhere, Category=Camera)
+	bool bSetControlRotationWhenViewTarget = true;
+
 #if WITH_EDITORONLY_DATA
 
 	/** Whether to run this camera in editor. */
