Cameras: don't initialize the camera system twice

diff --git a/Engine/Plugins/Cameras/GameplayCameras/Source/GameplayCameras/Private/Tests/CameraSystemUpdateTests.cpp b/Engine/Plugins/Cameras/GameplayCameras/Source/GameplayCameras/Private/Tests/CameraSystemUpdateTests.cpp
--- a/Engine/Plugins/Cameras/GameplayCameras/Source/GameplayCameras/Private/Tests/CameraSystemUpdateTests.cpp
+++ b/Engine/Plugins/Cameras/GameplayCameras/Source/GameplayCameras/Private/Tests/CameraSystemUpdateTests.cpp
@@ -30,7 +30,6 @@ bool FCameraSystemFrameFlagsTest::RunTest(const FString& Parameters)
 		.Get();
 
 	TSharedRef<FCameraSystemEvaluator> Evaluator = FCameraSystemEvaluatorBuilder::Build();
-	Evaluator->Initialize();
 	Evaluator->PushEvaluationContext(EvaluationContext);
 
 	EvaluationContext->GetInitialResult().bIsValid = true;
