Cameras: make the camera components visible in the "+Add" component menu

diff --git a/Engine/Plugins/Cameras/GameplayCameras/Source/GameplayCameras/Public/GameFramework/GameplayCameraComponent.h b/Engine/Plugins/Cameras/GameplayCameras/Source/GameplayCameras/Public/GameFramework/GameplayCameraComponent.h
--- a/Engine/Plugins/Cameras/GameplayCameras/Source/GameplayCameras/Public/GameFramework/GameplayCameraComponent.h
+++ b/Engine/Plugins/Cameras/GameplayCameras/Source/GameplayCameras/Public/GameFramework/GameplayCameraComponent.h
@@ -10,7 +10,9 @@
 /**
  * A component that can run a camera asset inside its own camera evaluation context.
  */
-UCLASS(Blueprintable, MinimalAPI)
+UCLASS(Blueprintable, MinimalAPI, 
+		ClassGroup=Camera, 
+		meta=(BlueprintSpawnableComponent))
 class UGameplayCameraComponent : public UGameplayCameraComponentBase
 {
 	GENERATED_BODY()
diff --git a/Engine/Plugins/Cameras/GameplayCameras/Source/GameplayCameras/Public/GameFramework/GameplayCameraRigComponent.h b/Engine/Plugins/Cameras/GameplayCameras/Source/GameplayCameras/Public/GameFramework/GameplayCameraRigComponent.h
--- a/Engine/Plugins/Cameras/GameplayCameras/Source/GameplayCameras/Public/GameFramework/GameplayCameraRigComponent.h
+++ b/Engine/Plugins/Cameras/GameplayCameras/Source/GameplayCameras/Public/GameFramework/GameplayCameraRigComponent.h
@@ -10,7 +10,9 @@
 /**
  * A component that can run a camera rig asset inside its own camera evaluation context.
  */
-UCLASS(Blueprintable, MinimalAPI)
+UCLASS(Blueprintable, MinimalAPI, 
+		ClassGroup=Camera, 
+		meta=(BlueprintSpawnableComponent))
 class UGameplayCameraRigComponent : public UGameplayCameraComponentBase
 {
 	GENERATED_BODY()
