Updated code using script to make sure x_api is set on functions/globalvars instead of on types

diff --git a/Engine/Plugins/Cameras/GameplayCameras/Source/GameplayCameras/Public/Core/CameraNodeEvaluator.h b/Engine/Plugins/Cameras/GameplayCameras/Source/GameplayCameras/Public/Core/CameraNodeEvaluator.h
--- a/Engine/Plugins/Cameras/GameplayCameras/Source/GameplayCameras/Public/Core/CameraNodeEvaluator.h
+++ b/Engine/Plugins/Cameras/GameplayCameras/Source/GameplayCameras/Public/Core/CameraNodeEvaluator.h
@@ -15,6 +15,8 @@
 #include "GameplayCameras.h"
 #include "UObject/ObjectPtr.h"
 
+#define UE_API GAMEPLAYCAMERAS_API
+
 class FReferenceCollector;
 class UCameraNode;
 class UCameraRigAsset;
@@ -170,7 +172,7 @@ struct FCameraNodeEvaluationParams
 /**
  * Input/output result structure for running a camera node evaluator.
  */
-struct GAMEPLAYCAMERAS_API FCameraNodeEvaluationResult
+struct FCameraNodeEvaluationResult
 {
 	/** The camera pose. */
 	FCameraPose CameraPose;
@@ -196,24 +198,24 @@ struct GAMEPLAYCAMERAS_API FCameraNodeEvaluationResult
 public:
 
 	/** Reset this result to its default (non-valid) state.  */
-	void Reset();
+	UE_API void Reset();
 
 	/** Reset all written-this-frame flags on the camera pose and tables. */
-	void ResetFrameFlags();
+	UE_API void ResetFrameFlags();
 
 	/** Override this result with the given other result. */
-	void OverrideAll(const FCameraNodeEvaluationResult& OtherResult, bool bIncludePrivateValues = false);
+	UE_API void OverrideAll(const FCameraNodeEvaluationResult& OtherResult, bool bIncludePrivateValues = false);
 
 	/** Interpolate this result towards the other given result. */
-	void LerpAll(const FCameraNodeEvaluationResult& ToResult, float BlendFactor, bool bIncludePrivateValues = false);
+	UE_API void LerpAll(const FCameraNodeEvaluationResult& ToResult, float BlendFactor, bool bIncludePrivateValues = false);
 
 	/** Serializes this result to the given archive. */
-	void Serialize(FArchive& Ar);
+	UE_API void Serialize(FArchive& Ar);
 
 public:
 
 	/** Collects objects from the context data table. */
-	void AddReferencedObjects(FReferenceCollector& Collector);
+	UE_API void AddReferencedObjects(FReferenceCollector& Collector);
 
 #if WITH_EDITOR || UE_GAMEPLAY_CAMERAS_DEBUG
 
@@ -221,10 +223,10 @@ public:
 
 	// Internal API.
 
-	void AddCameraPoseTrailPointIfNeeded();
-	void AddCameraPoseTrailPointIfNeeded(const FVector3d& Point);
-	void AppendCameraPoseLocationTrail(const FCameraNodeEvaluationResult& InResult);
-	TConstArrayView<FVector3d> GetCameraPoseLocationTrail() const;
+	UE_API void AddCameraPoseTrailPointIfNeeded();
+	UE_API void AddCameraPoseTrailPointIfNeeded(const FVector3d& Point);
+	UE_API void AppendCameraPoseLocationTrail(const FCameraNodeEvaluationResult& InResult);
+	UE_API TConstArrayView<FVector3d> GetCameraPoseLocationTrail() const;
 
 private:
 
@@ -247,7 +249,7 @@ struct FCameraOperationParams
 /**
  * Parameter structure for serializing the state of a camera node evaluator.
  */
-struct GAMEPLAYCAMERAS_API FCameraNodeEvaluatorSerializeParams
+struct FCameraNodeEvaluatorSerializeParams
 {
 };
 
@@ -271,8 +273,8 @@ class FCameraNodeEvaluator
 
 public:
 
-	GAMEPLAYCAMERAS_API FCameraNodeEvaluator() = default;
-	GAMEPLAYCAMERAS_API virtual ~FCameraNodeEvaluator() = default;
+	FCameraNodeEvaluator() = default;
+	virtual ~FCameraNodeEvaluator() = default;
 
 	/** Called to build any children evaluators. */
 	void Build(const FCameraNodeEvaluatorBuildParams& Params);
@@ -336,42 +338,42 @@ protected:
 protected:
 
 	/** Called to build any children evaluators. */
-	GAMEPLAYCAMERAS_API virtual void OnBuild(const FCameraNodeEvaluatorBuildParams& Params) {}
+	virtual void OnBuild(const FCameraNodeEvaluatorBuildParams& Params) {}
 
 	/** Initialize this evaluator. Children and descendants will be automatically initialized too. */
-	GAMEPLAYCAMERAS_API virtual void OnInitialize(const FCameraNodeEvaluatorInitializeParams& Params, FCameraNodeEvaluationResult& OutResult) {}
+	virtual void OnInitialize(const FCameraNodeEvaluatorInitializeParams& Params, FCameraNodeEvaluationResult& OutResult) {}
 
 	/** Collect referenced UObjects for this node. */
-	GAMEPLAYCAMERAS_API virtual void OnAddReferencedObjects(FReferenceCollector& Collector) {}
+	virtual void OnAddReferencedObjects(FReferenceCollector& Collector) {}
 
 	/** Get the list of children under this evaluator. */
-	GAMEPLAYCAMERAS_API virtual FCameraNodeEvaluatorChildrenView OnGetChildren() { return FCameraNodeEvaluatorChildrenView(); }
+	virtual FCameraNodeEvaluatorChildrenView OnGetChildren() { return FCameraNodeEvaluatorChildrenView(); }
 
 	/**
 	 * Called to update and store the blended parameters for this node.
 	 * Requires setting the ECameraNodeEvaluatorFlags::NeedsParameterUpdate flag.
 	 */
-	GAMEPLAYCAMERAS_API virtual void OnUpdateParameters(const FCameraBlendedParameterUpdateParams& Params, FCameraBlendedParameterUpdateResult& OutResult) {}
+	virtual void OnUpdateParameters(const FCameraBlendedParameterUpdateParams& Params, FCameraBlendedParameterUpdateResult& OutResult) {}
 
 	/**
 	 * Run this evaluator. This node evaluator is responsible for calling Run() on its children as appropriate.
 	 */
-	GAMEPLAYCAMERAS_API virtual void OnRun(const FCameraNodeEvaluationParams& Params, FCameraNodeEvaluationResult& OutResult) {}
+	virtual void OnRun(const FCameraNodeEvaluationParams& Params, FCameraNodeEvaluationResult& OutResult) {}
 
 	/** 
 	 * Execute an IK operation.
 	 * Requires setting the ECameraNodeEvaluatorFlags::SupportsOperations flag.
 	 */
-	GAMEPLAYCAMERAS_API virtual void OnExecuteOperation(const FCameraOperationParams& Params, FCameraOperation& Operation) {}
+	virtual void OnExecuteOperation(const FCameraOperationParams& Params, FCameraOperation& Operation) {}
 
 	/**
 	 * Serializes the state of this evaluator.
 	 * Requires setting the ECameraNodeEvaluatorFlags::NeedsSerialize flag, which is set by default.
 	 */
-	GAMEPLAYCAMERAS_API virtual void OnSerialize(const FCameraNodeEvaluatorSerializeParams& Params, FArchive& Ar) {}
+	virtual void OnSerialize(const FCameraNodeEvaluatorSerializeParams& Params, FArchive& Ar) {}
 
 #if WITH_EDITOR
-	GAMEPLAYCAMERAS_API virtual void OnDrawEditorPreview(const FCameraEditorPreviewDrawParams& Params, FCameraDebugRenderer& Renderer) {}
+	virtual void OnDrawEditorPreview(const FCameraEditorPreviewDrawParams& Params, FCameraDebugRenderer& Renderer) {}
 #endif  // WITH_EDITOR
 
 #if UE_GAMEPLAY_CAMERAS_DEBUG
@@ -435,3 +437,4 @@ EvaluatorType* FCameraNodeEvaluatorBuildParams::BuildEvaluatorAs(const UCameraNo
 #define UE_DEFINE_CAMERA_NODE_EVALUATOR(ClassName)\
 	UE_GAMEPLAY_CAMERAS_DEFINE_RTTI(ClassName)
 
+#undef UE_API
diff --git a/Engine/Plugins/Cameras/GameplayCameras/Source/GameplayCameras/Public/Core/CameraPose.h b/Engine/Plugins/Cameras/GameplayCameras/Source/GameplayCameras/Public/Core/CameraPose.h
--- a/Engine/Plugins/Cameras/GameplayCameras/Source/GameplayCameras/Public/Core/CameraPose.h
+++ b/Engine/Plugins/Cameras/GameplayCameras/Source/GameplayCameras/Public/Core/CameraPose.h
@@ -11,6 +11,8 @@
 
 #include "CameraPose.generated.h"
 
+#define UE_API GAMEPLAYCAMERAS_API
+
 class FArchive;
 struct FPostProcessSettings;
 enum EAspectRatioAxisConstraint : int;
@@ -57,7 +59,7 @@ enum EAspectRatioAxisConstraint : int;
 /**
  * Boolean flags for each of the properties inside FCameraPose.
  */
-struct GAMEPLAYCAMERAS_API FCameraPoseFlags
+struct FCameraPoseFlags
 {
 #define UE_CAMERA_POSE_FOR_PROPERTY(PropType, PropName)\
 	bool PropName:1 = false;
@@ -69,22 +71,22 @@ UE_CAMERA_POSE_FOR_ALL_PROPERTIES()
 public:
 
 	/** Returns a structure where all flags are set. */
-	static const FCameraPoseFlags& All();
+	static UE_API const FCameraPoseFlags& All();
 
 	/** Creates a new flags structure. */
-	FCameraPoseFlags();
+	UE_API FCameraPoseFlags();
 	/** Creates a new flags structure with all flags set to the given value. */
-	FCameraPoseFlags(bool bInValue);
+	UE_API FCameraPoseFlags(bool bInValue);
 
 	/** Sets all flags to the given value. */
-	FCameraPoseFlags& SetAllFlags(bool bInValue);
+	UE_API FCameraPoseFlags& SetAllFlags(bool bInValue);
 	/** Sets the flags that are set in OtherFlags, but checks that no flag is set on both structures. */
-	FCameraPoseFlags& ExclusiveCombine(const FCameraPoseFlags& OtherFlags);
+	UE_API FCameraPoseFlags& ExclusiveCombine(const FCameraPoseFlags& OtherFlags);
 
 	/** Combines the flags with an AND logical operation. */
-	FCameraPoseFlags& AND(const FCameraPoseFlags& OtherFlags);
+	UE_API FCameraPoseFlags& AND(const FCameraPoseFlags& OtherFlags);
 	/** Combines the flags with an OR logical operation. */
-	FCameraPoseFlags& OR(const FCameraPoseFlags& OtherFlags);
+	UE_API FCameraPoseFlags& OR(const FCameraPoseFlags& OtherFlags);
 };
 
 /**
@@ -94,16 +96,16 @@ public:
  * The ChangedFlags structure keeps track of which fields were changed via the setters.
  */
 USTRUCT()
-struct GAMEPLAYCAMERAS_API FCameraPose
+struct FCameraPose
 {
 	GENERATED_BODY()
 
 public:
 
-	FCameraPose();
+	UE_API FCameraPose();
 
 	/** Resets this camera pose to its default values, with all changed flags off. */
-	void Reset();
+	UE_API void Reset();
 
 public:
 
@@ -139,58 +141,58 @@ public:
 	/** Set the changed flags. */
 	void SetChangedFlags(const FCameraPoseFlags& InChangedFlags) { ChangedFlags = InChangedFlags; }
 	/** Set all fields as changed. */
-	void SetAllChangedFlags();
+	UE_API void SetAllChangedFlags();
 	/** Set all fields as clean. */
-	void ClearAllChangedFlags();
+	UE_API void ClearAllChangedFlags();
 
 public:
 
 	// Utility
 
 	/** Gets the transform of the camera. */
-	FTransform3d GetTransform() const;
+	UE_API FTransform3d GetTransform() const;
 
 	/** Sets the transform of the camera. */
-	void SetTransform(FTransform3d Transform, bool bForceSet = false);
+	UE_API void SetTransform(FTransform3d Transform, bool bForceSet = false);
 
 	/**
 	 * Computes the horizontal field of view of the camera.
 	 * The effective field of view can be driven by the FieldOfView property, or
 	 * the FocalLength property in combination with the sensor size.
 	 */
-	double GetEffectiveFieldOfView(bool bIncludeOverscan = true) const;
+	UE_API double GetEffectiveFieldOfView(bool bIncludeOverscan = true) const;
 
 	/** Gets the aspect ratio of the camera sensor. */
-	double GetSensorAspectRatio() const;
+	UE_API double GetSensorAspectRatio() const;
 
 	/** Gets the horizontal projection offset. */
-	double GetHorizontalProjectionOffset() const;
+	UE_API double GetHorizontalProjectionOffset() const;
 
 	/** Gets the vertical projection offset. */
-	double GetVerticalProjectionOffset() const;
+	UE_API double GetVerticalProjectionOffset() const;
 
 	/** Gets the aiming ray of the camera. */
-	FRay3d GetAimRay() const;
+	UE_API FRay3d GetAimRay() const;
 
 	/** Gets the aiming direction of the camera. */
-	FVector3d GetAimDir() const;
+	UE_API FVector3d GetAimDir() const;
 
 	/** Gets the location of the camera's target. */
-	FVector3d GetTarget() const;
+	UE_API FVector3d GetTarget() const;
 
 	/** Gets the location of the camera's target given a specific distance. */
-	FVector3d GetTarget(double InTargetDistance) const;
+	UE_API FVector3d GetTarget(double InTargetDistance) const;
 
 public:
 
 	/** Computes the horizontal field of view of a camera. */
-	static double GetEffectiveFieldOfView(float FocalLength, float FieldOfView, float SensorWidth, float SensorHeight, float SqueezeFactor = 1.f, float Overscan = 1.f);
+	static UE_API double GetEffectiveFieldOfView(float FocalLength, float FieldOfView, float SensorWidth, float SensorHeight, float SqueezeFactor = 1.f, float Overscan = 1.f);
 	
 	/** Computes the aspect ratio of a camera sensor. */
-	static double GetSensorAspectRatio(float SensorWidth, float SensorHeight);
+	static UE_API double GetSensorAspectRatio(float SensorWidth, float SensorHeight);
 
 	/** Gets the default sensor size. */
-	static void GetDefaultSensorSize(float& OutSensorWidth, float& OutSensorHeight);
+	static UE_API void GetDefaultSensorSize(float& OutSensorWidth, float& OutSensorHeight);
 
 	/**
 	 * Applies the necessary post-process settings given the current values
@@ -203,35 +205,35 @@ public:
 	 * @param bOverwriteSettings   Whether to overwrite values found to already be set
 	 * @return  Whether post-process settings were created.
 	 */
-	bool ApplyPhysicalCameraSettings(FPostProcessSettings& PostProcessSettings, bool bOverwriteSettings = false) const;
+	UE_API bool ApplyPhysicalCameraSettings(FPostProcessSettings& PostProcessSettings, bool bOverwriteSettings = false) const;
 
 public:
 
 	// Interpolation
 	
 	/** Takes all properties from OtherPose and sets them on this camera pose. */
-	void OverrideAll(const FCameraPose& OtherPose);
+	UE_API void OverrideAll(const FCameraPose& OtherPose);
 	/** Takes all changed properties from OtherPose and sets them on this camera pose. */
-	void OverrideChanged(const FCameraPose& OtherPose);
+	UE_API void OverrideChanged(const FCameraPose& OtherPose);
 	/** Interpolates all properties from ToPose using the given factor. */
-	void LerpAll(const FCameraPose& ToPose, float Factor);
+	UE_API void LerpAll(const FCameraPose& ToPose, float Factor);
 	/** Interpolates all changed properties from ToPose using the given factor. */
-	void LerpChanged(const FCameraPose& ToPose, float Factor);
+	UE_API void LerpChanged(const FCameraPose& ToPose, float Factor);
 	/** Interpolates changed properties from ToPose using the given factor. Only properties defined by InMask are taken into account. */
-	void LerpChanged(const FCameraPose& ToPose, float Factor, const FCameraPoseFlags& InMask, bool bInvertMask, FCameraPoseFlags& OutMask);
+	UE_API void LerpChanged(const FCameraPose& ToPose, float Factor, const FCameraPoseFlags& InMask, bool bInvertMask, FCameraPoseFlags& OutMask);
 
 public:
 
 	/** Serializes the given camera pose including the written-property flags. */
-	static void SerializeWithFlags(FArchive& Ar, FCameraPose& CameraPose);
+	static UE_API void SerializeWithFlags(FArchive& Ar, FCameraPose& CameraPose);
 
 	/** Serializes this camera pose including the written-property flags. */
-	void SerializeWithFlags(FArchive& Ar);
+	UE_API void SerializeWithFlags(FArchive& Ar);
 
 private:
 
-	void InternalOverrideChanged(const FCameraPose& OtherPose, bool bChangedOnly);
-	void InternalLerpChanged(const FCameraPose& ToPose, float Factor, const FCameraPoseFlags& InMask, bool bInvertMask, FCameraPoseFlags& OutMask, bool bChangedOnly);
+	UE_API void InternalOverrideChanged(const FCameraPose& OtherPose, bool bChangedOnly);
+	UE_API void InternalLerpChanged(const FCameraPose& ToPose, float Factor, const FCameraPoseFlags& InMask, bool bInvertMask, FCameraPoseFlags& OutMask, bool bChangedOnly);
 
 private:
 
@@ -358,3 +360,4 @@ private:
 	FCameraPoseFlags ChangedFlags;
 };
 
+#undef UE_API
diff --git a/Engine/Plugins/ChaosClothAsset/Source/ChaosClothAsset/Public/ChaosClothAsset/CollectionClothSimMorphTargetFacade.h b/Engine/Plugins/ChaosClothAsset/Source/ChaosClothAsset/Public/ChaosClothAsset/CollectionClothSimMorphTargetFacade.h
--- a/Engine/Plugins/ChaosClothAsset/Source/ChaosClothAsset/Public/ChaosClothAsset/CollectionClothSimMorphTargetFacade.h
+++ b/Engine/Plugins/ChaosClothAsset/Source/ChaosClothAsset/Public/ChaosClothAsset/CollectionClothSimMorphTargetFacade.h
@@ -4,6 +4,8 @@
 
 #include "Templates/SharedPointer.h"
 
+#define UE_API CHAOSCLOTHASSET_API
+
 struct FManagedArrayCollection;
 
 namespace UE::Chaos::ClothAsset
@@ -13,7 +15,7 @@ namespace UE::Chaos::ClothAsset
 	 * Constructed from FCollectionClothConstFacade.
 	 * Const access (read only) version.
 	 */
-	class CHAOSCLOTHASSET_API FCollectionClothSimMorphTargetConstFacade
+	class FCollectionClothSimMorphTargetConstFacade
 	{
 	public:
 		FCollectionClothSimMorphTargetConstFacade() = delete;
@@ -27,18 +29,18 @@ namespace UE::Chaos::ClothAsset
 		virtual ~FCollectionClothSimMorphTargetConstFacade() = default;
 
 		/** Return the name for this morph target */
-		const FString& GetSimMorphTargetName() const;
+		UE_API const FString& GetSimMorphTargetName() const;
 
 		//~ Sim Morph Target Vertices Group
 		/** Return the total number of vertices for this morph target. */
-		int32 GetNumSimMorphTargetVertices() const;
+		UE_API int32 GetNumSimMorphTargetVertices() const;
 		/** Return the morph target vertices offset for this morph target in the morph target vertices for the collection. */
-		int32 GetimMorphTargetVerticesOffset() const;
-		TConstArrayView<FVector3f> GetSimMorphTargetPositionDelta() const;
-		TConstArrayView<FVector3f> GetSimMorphTargetTangentZDelta() const;
-		TConstArrayView<int32> GetSimMorphTargetSimVertex3DIndex() const;
+		UE_API int32 GetimMorphTargetVerticesOffset() const;
+		UE_API TConstArrayView<FVector3f> GetSimMorphTargetPositionDelta() const;
+		UE_API TConstArrayView<FVector3f> GetSimMorphTargetTangentZDelta() const;
+		UE_API TConstArrayView<int32> GetSimMorphTargetSimVertex3DIndex() const;
 
-		bool IsEmpty() const;
+		UE_API bool IsEmpty() const;
 
 		/** Return the Morph Target index this facade has been created with. */
 		int32 GetMorphTargetIndex() const { return MorphTargetIndex; }
@@ -46,7 +48,7 @@ namespace UE::Chaos::ClothAsset
 	protected:
 		friend class FCollectionClothSimMorphTargetFacade;  // For other instances access
 		friend class FCollectionClothConstFacade;
-		FCollectionClothSimMorphTargetConstFacade(const TSharedRef<const class FClothCollection>& InClothCollection, int32 InPatternIndex);
+		UE_API FCollectionClothSimMorphTargetConstFacade(const TSharedRef<const class FClothCollection>& InClothCollection, int32 InPatternIndex);
 
 		static constexpr int32 GetBaseElementIndex() { return 0; }
 		int32 GetElementIndex() const { return GetBaseElementIndex() + MorphTargetIndex; }
@@ -60,7 +62,7 @@ namespace UE::Chaos::ClothAsset
 	 * Constructed from FCollectionClothFacade.
 	 * Non-const access (read/write) version.
 	 */
-	class CHAOSCLOTHASSET_API FCollectionClothSimMorphTargetFacade final : public FCollectionClothSimMorphTargetConstFacade
+	class FCollectionClothSimMorphTargetFacade final : public FCollectionClothSimMorphTargetConstFacade
 	{
 	public:
 		FCollectionClothSimMorphTargetFacade() = delete;
@@ -74,34 +76,36 @@ namespace UE::Chaos::ClothAsset
 		virtual ~FCollectionClothSimMorphTargetFacade() override = default;
 
 		/** Remove all geometry from this morph target. */
-		void Reset();
+		UE_API void Reset();
 
 		/** Initialize from another morph target. Assumes all indices match between source and target. */
-		void Initialize(const FCollectionClothSimMorphTargetConstFacade& Other, int32 SimVertex3DOffset);
+		UE_API void Initialize(const FCollectionClothSimMorphTargetConstFacade& Other, int32 SimVertex3DOffset);
 
 		/** Initialize */
-		void Initialize(const FString& Name, const TConstArrayView<FVector3f>& PositionDeltas, const TConstArrayView<FVector3f>& TangentZDeltas, const TConstArrayView<int32>& SimVertex3DIndices);
+		UE_API void Initialize(const FString& Name, const TConstArrayView<FVector3f>& PositionDeltas, const TConstArrayView<FVector3f>& TangentZDeltas, const TConstArrayView<int32>& SimVertex3DIndices);
 
 		/** Set the name for this morph target. */
-		void SetSimMorphTargetName(const FString& MorphTargetName);
+		UE_API void SetSimMorphTargetName(const FString& MorphTargetName);
 
 		//~ Sim Morph Target Vertices Group
 		/** Grow or shrink the space reserved for morph target vertices for this morph target within the cloth collection and return its start index. */
-		void SetNumSimMorphTargetVertices(int32 NumMorphTargetVertices);
-		void RemoveSimMorphTargetVertices(const TArray<int32>& SortedDeletionList);
-		TArrayView<FVector3f> GetSimMorphTargetPositionDelta();
-		TArrayView<FVector3f> GetSimMorphTargetTangentZDelta();
-		TArrayView<int32> GetSimMorphTargetSimVertex3DIndex();
+		UE_API void SetNumSimMorphTargetVertices(int32 NumMorphTargetVertices);
+		UE_API void RemoveSimMorphTargetVertices(const TArray<int32>& SortedDeletionList);
+		UE_API TArrayView<FVector3f> GetSimMorphTargetPositionDelta();
+		UE_API TArrayView<FVector3f> GetSimMorphTargetTangentZDelta();
+		UE_API TArrayView<int32> GetSimMorphTargetSimVertex3DIndex();
 
 		/** Remove all morph target vertices with invalid indices. */
-		void Compact();
+		UE_API void Compact();
 
 	protected:
 		friend class FCollectionClothFacade;
-		FCollectionClothSimMorphTargetFacade(const TSharedRef<class FClothCollection>& InClothCollection, int32 InMorphTargetIndex);
+		UE_API FCollectionClothSimMorphTargetFacade(const TSharedRef<class FClothCollection>& InClothCollection, int32 InMorphTargetIndex);
 
-		void SetDefaults();
+		UE_API void SetDefaults();
 
 		TSharedRef<class FClothCollection> GetClothCollection() { return ConstCastSharedRef<class FClothCollection>(ClothCollection); }
 	};
 }  // End namespace UE::Chaos::ClothAsset
+
+#undef UE_API
diff --git a/Engine/Plugins/Experimental/Dataflow/Source/DataflowNodes/Public/Dataflow/DataflowCollectionEditSkinWeightsNode.h b/Engine/Plugins/Experimental/Dataflow/Source/DataflowNodes/Public/Dataflow/DataflowCollectionEditSkinWeightsNode.h
--- a/Engine/Plugins/Experimental/Dataflow/Source/DataflowNodes/Public/Dataflow/DataflowCollectionEditSkinWeightsNode.h
+++ b/Engine/Plugins/Experimental/Dataflow/Source/DataflowNodes/Public/Dataflow/DataflowCollectionEditSkinWeightsNode.h
@@ -1,4 +1,4 @@
-﻿// Copyright Epic Games, Inc. All Rights Reserved.
+// Copyright Epic Games, Inc. All Rights Reserved.
 
 #pragma once 
 
@@ -12,9 +12,11 @@
 
 #include "DataflowCollectionEditSkinWeightsNode.generated.h"
 
+#define UE_API DATAFLOWNODES_API
+
 /** Dataflow skin weights data */
 USTRUCT()
-struct DATAFLOWNODES_API FDataflowSkinWeightData
+struct FDataflowSkinWeightData
 {
 	GENERATED_BODY()
 
@@ -27,14 +29,14 @@ struct DATAFLOWNODES_API FDataflowSkinWeightData
 
 /** Edit skin weights vertex properties. */
 USTRUCT(Meta = (DataflowCollection))
-struct DATAFLOWNODES_API FDataflowCollectionEditSkinWeightsNode : public FDataflowPrimitiveNode
+struct FDataflowCollectionEditSkinWeightsNode : public FDataflowPrimitiveNode
 {
 	GENERATED_BODY()
 	DATAFLOW_NODE_DEFINE_INTERNAL(FDataflowCollectionEditSkinWeightsNode, "EditSkinWeights", "Collection", "Edit skin weights and save it to collection")
 
 public:
 
-	FDataflowCollectionEditSkinWeightsNode(const UE::Dataflow::FNodeParameters& InParam, FGuid InGuid = FGuid::NewGuid());
+	UE_API FDataflowCollectionEditSkinWeightsNode(const UE::Dataflow::FNodeParameters& InParam, FGuid InGuid = FGuid::NewGuid());
 
 	UPROPERTY(Meta = (DataflowInput, DataflowOutput, DataflowPassthrough = "Collection"))
 	FManagedArrayCollection Collection;
@@ -71,35 +73,35 @@ public:
 	FDataflowBoneSelectionChangedNotifyDelegate OnBoneSelectionChanged;
 
 	/** Report the vertex weights back onto the property ones */
-	void ReportVertexWeights(const TArray<TArray<int32>>& SetupIndices, const TArray<TArray<float>>& SetupWeights, const TArray<TArray<int32>>& FinalIndices, const TArray<TArray<float>>& FinalWeights);
+	UE_API void ReportVertexWeights(const TArray<TArray<int32>>& SetupIndices, const TArray<TArray<float>>& SetupWeights, const TArray<TArray<int32>>& FinalIndices, const TArray<TArray<float>>& FinalWeights);
 
 	/** Extract the vertex weights back from the property ones */
-	void ExtractVertexWeights(const TArray<TArray<int32>>& SetupIndices, const TArray<TArray<float>>& SetupWeights, TArrayView<TArray<int32>> FinalIndices, TArrayView<TArray<float>> FinalWeights) const;
+	UE_API void ExtractVertexWeights(const TArray<TArray<int32>>& SetupIndices, const TArray<TArray<float>>& SetupWeights, TArrayView<TArray<int32>> FinalIndices, TArrayView<TArray<float>> FinalWeights) const;
 
 	/** Fill the weight attribute values from the collection */
-	bool FillAttributeWeights(const FManagedArrayCollection& SelectedCollection,
+	UE_API bool FillAttributeWeights(const FManagedArrayCollection& SelectedCollection,
 		const FCollectionAttributeKey& IndicesAttributeKey, const FCollectionAttributeKey& WeightsAttributeKey, TArray<TArray<int32>>& AttributeIndices, TArray<TArray<float>>& AttributeWeights) const;
 
 	/** Get the weights attribute key to retrieve/set the weight values*/
-	FCollectionAttributeKey GetBoneIndicesKey(UE::Dataflow::FContext& Context) const;
+	UE_API FCollectionAttributeKey GetBoneIndicesKey(UE::Dataflow::FContext& Context) const;
 
 	/** Get the weights attribute key to retrieve/set the weight values*/
-    FCollectionAttributeKey GetBoneWeightsKey(UE::Dataflow::FContext& Context) const;
+    UE_API FCollectionAttributeKey GetBoneWeightsKey(UE::Dataflow::FContext& Context) const;
 
 	/** Validate the skeletal mesh construction */
 	void ValidateSkeletalMesh() {bValidSkeletalMesh = true;}
 private:
 
 	//~ Begin FDataflowPrimitiveNode interface
-	virtual TArray<UE::Dataflow::FRenderingParameter> GetRenderParametersImpl() const override;
-	virtual void Evaluate(UE::Dataflow::FContext& Context, const FDataflowOutput* Out) const override;
-	virtual void AddPrimitiveComponents(const TSharedPtr<const FManagedArrayCollection> RenderCollection,
+	UE_API virtual TArray<UE::Dataflow::FRenderingParameter> GetRenderParametersImpl() const override;
+	UE_API virtual void Evaluate(UE::Dataflow::FContext& Context, const FDataflowOutput* Out) const override;
+	UE_API virtual void AddPrimitiveComponents(const TSharedPtr<const FManagedArrayCollection> RenderCollection,
 		TObjectPtr<UObject> NodeOwner, TObjectPtr<AActor> RootActor, TArray<TObjectPtr<UPrimitiveComponent>>& PrimitiveComponents) override;
-	virtual void OnInvalidate() override;
+	UE_API virtual void OnInvalidate() override;
 #if WITH_EDITOR
 	virtual bool CanDebugDraw() const override { return true; }
-	virtual bool CanDebugDrawViewMode(const FName& ViewModeName) const override;
-	virtual void DebugDraw(UE::Dataflow::FContext& Context,
+	UE_API virtual bool CanDebugDrawViewMode(const FName& ViewModeName) const override;
+	UE_API virtual void DebugDraw(UE::Dataflow::FContext& Context,
 		IDataflowDebugDrawInterface& DataflowRenderingInterface,
 		const FDebugDrawParameters& DebugDrawParameters) const override;
 #endif
@@ -112,3 +114,5 @@ private:
 	/** Valid skeletal mesh boolean to trigger the construction */
 	bool bValidSkeletalMesh = false;
 };
+
+#undef UE_API
diff --git a/Engine/Plugins/Experimental/TechAudioTools/Source/TechAudioToolsMetaSound/Public/ViewModels/MetaSoundViewModel.h b/Engine/Plugins/Experimental/TechAudioTools/Source/TechAudioToolsMetaSound/Public/ViewModels/MetaSoundViewModel.h
--- a/Engine/Plugins/Experimental/TechAudioTools/Source/TechAudioToolsMetaSound/Public/ViewModels/MetaSoundViewModel.h
+++ b/Engine/Plugins/Experimental/TechAudioTools/Source/TechAudioToolsMetaSound/Public/ViewModels/MetaSoundViewModel.h
@@ -8,6 +8,8 @@
 
 #include "MetaSoundViewModel.generated.h"
 
+#define UE_API TECHAUDIOTOOLSMETASOUND_API
+
 TECHAUDIOTOOLSMETASOUND_API DECLARE_LOG_CATEGORY_EXTERN(LogTechAudioToolsMetaSound, Log, All);
 
 class UMetaSoundBuilderBase;
@@ -19,8 +21,8 @@ class UMetaSoundOutputViewModel;
  * Can be initialized using a MetaSound Builder or a MetaSound asset. Creates member viewmodels for each input and output in the
  * MetaSound upon initialization.
  */
-UCLASS(DisplayName = "MetaSound Viewmodel")
-class TECHAUDIOTOOLSMETASOUND_API UMetaSoundViewModel : public UMVVMViewModelBase
+UCLASS(MinimalAPI, DisplayName = "MetaSound Viewmodel")
+class UMetaSoundViewModel : public UMVVMViewModelBase
 {
 	GENERATED_BODY()
 
@@ -36,27 +38,27 @@ private:
 public:
 	// Returns the object name of the initialized builder as text.
 	UFUNCTION(BlueprintCallable, FieldNotify, Category = "MetaSound Viewmodel")
-	FText GetBuilderNameAsText() const;
+	UE_API FText GetBuilderNameAsText() const;
 
 	// Contains MetaSound Input Viewmodels for each input of the initialized MetaSound.
 	UFUNCTION(BlueprintCallable, FieldNotify, DisplayName = "Get Input Viewmodels", Category = "MetaSound Viewmodel")
-	virtual TArray<UMetaSoundInputViewModel*> GetInputViewModels() const;
+	UE_API virtual TArray<UMetaSoundInputViewModel*> GetInputViewModels() const;
 
 	// Contains MetaSound Output ViewModels for each output of the initialized MetaSound.
 	UFUNCTION(BlueprintCallable, FieldNotify, DisplayName = "Get Output Viewmodels", Category = "MetaSound Viewmodel")
-	virtual TArray<UMetaSoundOutputViewModel*> GetOutputViewModels() const;
+	UE_API virtual TArray<UMetaSoundOutputViewModel*> GetOutputViewModels() const;
 
 	// Initializes the viewmodel using the given MetaSound asset.
 	UFUNCTION(BlueprintCallable, DisplayName = "Initialize MetaSound", Category = "Audio|MetaSound Viewmodel")
-	virtual void InitializeMetaSound(const TScriptInterface<IMetaSoundDocumentInterface> InMetaSound);
+	UE_API virtual void InitializeMetaSound(const TScriptInterface<IMetaSoundDocumentInterface> InMetaSound);
 
 	// Initializes the viewmodel using the given builder.
 	UFUNCTION(BlueprintCallable, DisplayName = "Initialize Builder", Category = "Audio|MetaSound Viewmodel")
-	virtual void Initialize(UMetaSoundBuilderBase* InBuilder);
+	UE_API virtual void Initialize(UMetaSoundBuilderBase* InBuilder);
 
 	// Resets this MetaSoundViewModel instance to an uninitialized state.
 	UFUNCTION(BlueprintCallable, Category = "Audio|MetaSound Viewmodel")
-	virtual void Reset();
+	UE_API virtual void Reset();
 
 	// Returns a reference to the initialized MetaSound's Builder.
 	UFUNCTION(BlueprintCallable, Category = "Audio|MetaSound Viewmodel")
@@ -66,19 +68,19 @@ public:
 	void SetIsInitialized(const bool bInIsInitialized) { UE_MVVM_SET_PROPERTY_VALUE(bIsInitialized, bInIsInitialized); }
 
 protected:
-	virtual void InitializeProperties(const FMetasoundFrontendDocument& FrontendDocument);
-	virtual void ResetProperties();
+	UE_API virtual void InitializeProperties(const FMetasoundFrontendDocument& FrontendDocument);
+	UE_API virtual void ResetProperties();
 
 	// Called upon initialization. Creates viewmodel instances for all inputs and outputs of the initialized MetaSound.
-	void CreateMemberViewModels();
+	UE_API void CreateMemberViewModels();
 
 	// Creates a single MetaSoundInputViewModel instance for the given input.
-	virtual void CreateInputViewModel(const FMetasoundFrontendClassInput& InInput);
-    virtual UMetaSoundInputViewModel* CreateInputViewModelInstance();
+	UE_API virtual void CreateInputViewModel(const FMetasoundFrontendClassInput& InInput);
+    UE_API virtual UMetaSoundInputViewModel* CreateInputViewModelInstance();
 
 	// Creates a single MetaSoundOutputViewModel instance for the given output.
-	void CreateOutputViewModel(const FMetasoundFrontendClassOutput& InOutput);
-	virtual UMetaSoundOutputViewModel* CreateOutputViewModelInstance();
+	UE_API void CreateOutputViewModel(const FMetasoundFrontendClassOutput& InOutput);
+	UE_API virtual UMetaSoundOutputViewModel* CreateOutputViewModelInstance();
 
 	UPROPERTY(Transient)
 	TObjectPtr<UMetaSoundBuilderBase> Builder;
@@ -94,8 +96,8 @@ protected:
  * Viewmodel class for MetaSound inputs. Allows widgets in UMG to bind to MetaSound literals. Useful for creating knobs, sliders, and other
  * widgets for setting MetaSound input parameters.
  */
-UCLASS(DisplayName = "MetaSound Input Viewmodel")
-class TECHAUDIOTOOLSMETASOUND_API UMetaSoundInputViewModel : public UMVVMViewModelBase
+UCLASS(MinimalAPI, DisplayName = "MetaSound Input Viewmodel")
+class UMetaSoundInputViewModel : public UMVVMViewModelBase
 {
 	GENERATED_BODY()
 
@@ -128,29 +130,29 @@ protected:
 	TObjectPtr<UMetaSoundBuilderBase> Builder;
 
 public:
-	virtual void InitializeInput(UMetaSoundBuilderBase* InBuilder, const FMetasoundFrontendClassInput& InInput);
-	virtual void ResetInput();
+	UE_API virtual void InitializeInput(UMetaSoundBuilderBase* InBuilder, const FMetasoundFrontendClassInput& InInput);
+	UE_API virtual void ResetInput();
 
 	FName GetInputName() const { return InputName; }
 
-	void SetIsInitialized(const bool bInIsInitialized);
-	void SetInputName(const FName& InInputName);
-	void SetDataType(const FName& InDataType);
-	void SetIsArray(const bool bInIsArray);
+	UE_API void SetIsInitialized(const bool bInIsInitialized);
+	UE_API void SetInputName(const FName& InInputName);
+	UE_API void SetDataType(const FName& InDataType);
+	UE_API void SetIsArray(const bool bInIsArray);
 
 	// Returns the value of this input's MetaSound Literal as a text value.
 	UFUNCTION(BlueprintCallable, FieldNotify, Category = "Audio|MetaSound Input Viewmodel")
-	FText GetLiteralValueAsText() const;
+	UE_API FText GetLiteralValueAsText() const;
 
-	void SetLiteral(const FMetasoundFrontendLiteral& InLiteral);
+	UE_API void SetLiteral(const FMetasoundFrontendLiteral& InLiteral);
 };
 
 /**
  * Viewmodel class for MetaSound outputs. Allows widgets in UMG to bind to data from a MetaSound output. Useful for driving visual parameters
  * using MetaSound outputs.
  */
-UCLASS(DisplayName = "MetaSound Output Viewmodel")
-class TECHAUDIOTOOLSMETASOUND_API UMetaSoundOutputViewModel : public UMVVMViewModelBase
+UCLASS(MinimalAPI, DisplayName = "MetaSound Output Viewmodel")
+class UMetaSoundOutputViewModel : public UMVVMViewModelBase
 {
 	GENERATED_BODY()
 
@@ -179,17 +181,19 @@ protected:
 	TObjectPtr<UMetaSoundBuilderBase> Builder;
 
 public:
-	virtual void InitializeOutput(UMetaSoundBuilderBase* InBuilder, const FMetasoundFrontendClassOutput& InOutput);
-	virtual void ResetOutput();
+	UE_API virtual void InitializeOutput(UMetaSoundBuilderBase* InBuilder, const FMetasoundFrontendClassOutput& InOutput);
+	UE_API virtual void ResetOutput();
 
 	FName GetOutputName() const { return OutputName; }
 
-	void SetIsInitialized(const bool bInIsInitialized);
-	void SetOutputName(const FName& InOutputName);
-	void SetDataType(const FName& InDataType);
-	void SetIsArray(const bool bInIsArray);
-	void SetMetaSoundOutput(const FMetaSoundOutput& InMetaSoundOutput);
+	UE_API void SetIsInitialized(const bool bInIsInitialized);
+	UE_API void SetOutputName(const FName& InOutputName);
+	UE_API void SetDataType(const FName& InDataType);
+	UE_API void SetIsArray(const bool bInIsArray);
+	UE_API void SetMetaSoundOutput(const FMetaSoundOutput& InMetaSoundOutput);
 
 	UFUNCTION()
-	void OnOutputValueChanged(FName InOutputName, const FMetaSoundOutput& InOutput);
+	UE_API void OnOutputValueChanged(FName InOutputName, const FMetaSoundOutput& InOutput);
 };
+
+#undef UE_API
diff --git a/Engine/Plugins/Experimental/WaterAdvanced/Source/WaterAdvanced/Public/FFTOceanPatchSubsystem.h b/Engine/Plugins/Experimental/WaterAdvanced/Source/WaterAdvanced/Public/FFTOceanPatchSubsystem.h
--- a/Engine/Plugins/Experimental/WaterAdvanced/Source/WaterAdvanced/Public/FFTOceanPatchSubsystem.h
+++ b/Engine/Plugins/Experimental/WaterAdvanced/Source/WaterAdvanced/Public/FFTOceanPatchSubsystem.h
@@ -24,7 +24,7 @@ public:
 	UE_API virtual TStatId GetStatId() const override;
 	
 	UE_API TObjectPtr<UTextureRenderTarget2D> GetOceanNormalRT(UWorld* World);
-	UE_API TObjectPtr<UNiagaraComponent> GetOceanSystem() { return FFTOceanSystem; }
+	TObjectPtr<UNiagaraComponent> GetOceanSystem() { return FFTOceanSystem; }
 protected:
 	UPROPERTY(BlueprintReadOnly, VisibleAnywhere, Category="Shallow Water")
 	TObjectPtr<UShallowWaterSettings> Settings;
diff --git a/Engine/Plugins/PCG/Source/PCG/Public/Compute/PCGComputeSource.h b/Engine/Plugins/PCG/Source/PCG/Public/Compute/PCGComputeSource.h
--- a/Engine/Plugins/PCG/Source/PCG/Public/Compute/PCGComputeSource.h
+++ b/Engine/Plugins/PCG/Source/PCG/Public/Compute/PCGComputeSource.h
@@ -8,15 +8,17 @@
 
 #include "PCGComputeSource.generated.h"
 
-UCLASS()
-class PCG_API UPCGComputeSource : public UComputeSource, public IPCGNodeSourceTextProvider
+#define UE_API PCG_API
+
+UCLASS(MinimalAPI)
+class UPCGComputeSource : public UComputeSource, public IPCGNodeSourceTextProvider
 {
 	GENERATED_BODY()
 
 public:
 	//~ Begin UComputeSource Interface.
-	FString GetSource() const override;
-	FString GetVirtualPath() const override;
+	UE_API FString GetSource() const override;
+	UE_API FString GetVirtualPath() const override;
 	//~ End UComputeSource Interface.
 
 #if WITH_EDITOR
@@ -29,7 +31,7 @@ public:
 	bool IsShaderTextReadOnly() const override { return false; }
 	//~End IPCGNodeSourceTextProvider interface
 
-	void SetSource(const FString& InSource);
+	UE_API void SetSource(const FString& InSource);
 #endif
 
 #if WITH_EDITORONLY_DATA
@@ -38,3 +40,5 @@ protected:
 	FString Source;
 #endif
 };
+
+#undef UE_API
diff --git a/Engine/Source/Runtime/Engine/Public/PrimitiveUniformShaderParameters.h b/Engine/Source/Runtime/Engine/Public/PrimitiveUniformShaderParameters.h
--- a/Engine/Source/Runtime/Engine/Public/PrimitiveUniformShaderParameters.h
+++ b/Engine/Source/Runtime/Engine/Public/PrimitiveUniformShaderParameters.h
@@ -96,7 +96,3 @@ public:
 /** Global primitive uniform buffer resource containing identity transformations. */
 extern ENGINE_API TGlobalResource<FIdentityPrimitiveUniformBuffer> GIdentityPrimitiveUniformBuffer;
 
-#if UE_ENABLE_INCLUDE_ORDER_DEPRECATED_IN_5_3
-#include "PrimitiveUniformShaderParametersBuilder.h"
-#include "PrimitiveSceneShaderData.h"
-#endif
diff --git a/Engine/Source/Runtime/MassEntity/Public/MassQueryExecutor.h b/Engine/Source/Runtime/MassEntity/Public/MassQueryExecutor.h
--- a/Engine/Source/Runtime/MassEntity/Public/MassQueryExecutor.h
+++ b/Engine/Source/Runtime/MassEntity/Public/MassQueryExecutor.h
@@ -5,6 +5,8 @@
 #include "Misc/NotNull.h"
 #include "MassExecutionContext.h"
 
+#define UE_API MASSENTITY_API
+
 class UMassProcessor;
 
 namespace UE::Mass
@@ -20,7 +22,7 @@ struct FQueryDefinitionBase
 /**
  * A MassEntityQuery wrapper with type-safe data access.
  */
-struct MASSENTITY_API FQueryExecutor
+struct FQueryExecutor
 {
 	template<typename T, typename = typename TEnableIf<TIsDerivedFrom<T, FQueryExecutor>::IsDerived>::Type>
 	static TSharedPtr<T> CreateQuery(FMassEntityQuery& InQuery, UObject* InLogOwner = nullptr)
@@ -42,7 +44,7 @@ struct MASSENTITY_API FQueryExecutor
 
 	friend ::UMassProcessor;
 
-	explicit FQueryExecutor(FMassEntityQuery& InQuery, UObject* InLogOwner = nullptr);
+	UE_API explicit FQueryExecutor(FMassEntityQuery& InQuery, UObject* InLogOwner = nullptr);
 	virtual ~FQueryExecutor() = default;
 
 	/** Override with logic to perform against the entities returned by this query. */
@@ -119,12 +121,12 @@ protected:
 		});
 	}
 
-	FQueryExecutor();
+	UE_API FQueryExecutor();
 private:
-	static FMassEntityQuery DummyQuery;
+	static UE_API FMassEntityQuery DummyQuery;
 
-	void ConfigureQuery(FMassSubsystemRequirements& ProcessorRequirements);
-	void CallExecute(FMassExecutionContext& Context);
+	UE_API void ConfigureQuery(FMassSubsystemRequirements& ProcessorRequirements);
+	UE_API void CallExecute(FMassExecutionContext& Context);
 
 	TNotNull<FMassEntityQuery*> BoundQuery;
 	TWeakObjectPtr<UObject> LogOwner;
@@ -133,7 +135,7 @@ private:
 	FQueryDefinitionBase* AccessorsPtr;
 
 #if WITH_MASSENTITY_DEBUG
-	void ValidateAccessors();
+	UE_API void ValidateAccessors();
 	uint32 DebugSize = 0;
 #endif
 };
@@ -793,3 +795,5 @@ struct FConstSubsystemAccess
 
 
 } // namespace UE::Mass
+
+#undef UE_API
