Cameras: fix camera parameter setter component not finding a system host

diff --git a/Engine/Plugins/Cameras/GameplayCameras/Source/GameplayCameras/Private/GameFramework/IGameplayCameraSystemHost.cpp b/Engine/Plugins/Cameras/GameplayCameras/Source/GameplayCameras/Private/GameFramework/IGameplayCameraSystemHost.cpp
--- a/Engine/Plugins/Cameras/GameplayCameras/Source/GameplayCameras/Private/GameFramework/IGameplayCameraSystemHost.cpp
+++ b/Engine/Plugins/Cameras/GameplayCameras/Source/GameplayCameras/Private/GameFramework/IGameplayCameraSystemHost.cpp
@@ -56,18 +56,21 @@ TScriptInterface<IGameplayCameraSystemHost> IGameplayCameraSystemHost::GetAsScri
 
 IGameplayCameraSystemHost* IGameplayCameraSystemHost::FindActiveHost(APlayerController* PlayerController)
 {
-	if (PlayerController->PlayerCameraManager)
+	if (PlayerController)
 	{
-		if (IGameplayCameraSystemHost* CameraManagerHost = Cast<IGameplayCameraSystemHost>(PlayerController->PlayerCameraManager))
+		if (PlayerController->PlayerCameraManager)
 		{
-			return CameraManagerHost;
+			if (IGameplayCameraSystemHost* CameraManagerHost = Cast<IGameplayCameraSystemHost>(PlayerController->PlayerCameraManager))
+			{
+				return CameraManagerHost;
+			}
 		}
-	}
-	else if (AActor* ViewTarget = PlayerController->GetViewTarget())
-	{
-		if (IGameplayCameraSystemHost* ViewTargetHost = ViewTarget->FindComponentByInterface<IGameplayCameraSystemHost>())
+		if (AActor* ViewTarget = PlayerController->GetViewTarget())
 		{
-			return ViewTargetHost;
+			if (IGameplayCameraSystemHost* ViewTargetHost = ViewTarget->FindComponentByInterface<IGameplayCameraSystemHost>())
+			{
+				return ViewTargetHost;
+			}
 		}
 	}
 	return nullptr;
