Cameras: fix initialization problems with the base/global/visual rig management component

It can now get a proper player controller when being initialized around a camera system run by a view target (as opposed to a camera system run by a camera manager).

diff --git a/Engine/Plugins/Cameras/GameplayCameras/Source/GameplayCameras/Private/GameFramework/ActivateCameraRigFunctions.cpp b/Engine/Plugins/Cameras/GameplayCameras/Source/GameplayCameras/Private/GameFramework/ActivateCameraRigFunctions.cpp
--- a/Engine/Plugins/Cameras/GameplayCameras/Source/GameplayCameras/Private/GameFramework/ActivateCameraRigFunctions.cpp
+++ b/Engine/Plugins/Cameras/GameplayCameras/Source/GameplayCameras/Private/GameFramework/ActivateCameraRigFunctions.cpp
@@ -63,7 +63,7 @@ void UActivateCameraRigFunctions::ActivateCameraRigImpl(UObject* WorldContextObj
 		CameraEvaluationComponent = UControllerGameplayCameraEvaluationComponent::FindOrAddComponent(HostOwningActor, &bComponentCreated);
 		if (bComponentCreated)
 		{
-			CameraEvaluationComponent->Initialize(FoundHost->GetAsScriptInterface());
+			CameraEvaluationComponent->Initialize(FoundHost->GetAsScriptInterface(), PlayerController);
 		}
 	}
 
diff --git a/Engine/Plugins/Cameras/GameplayCameras/Source/GameplayCameras/Private/GameFramework/ControllerGameplayCameraEvaluationComponent.cpp b/Engine/Plugins/Cameras/GameplayCameras/Source/GameplayCameras/Private/GameFramework/ControllerGameplayCameraEvaluationComponent.cpp
--- a/Engine/Plugins/Cameras/GameplayCameras/Source/GameplayCameras/Private/GameFramework/ControllerGameplayCameraEvaluationComponent.cpp
+++ b/Engine/Plugins/Cameras/GameplayCameras/Source/GameplayCameras/Private/GameFramework/ControllerGameplayCameraEvaluationComponent.cpp
@@ -17,10 +17,31 @@ UControllerGameplayCameraEvaluationComponent::UControllerGameplayCameraEvaluatio
 	bAutoActivate = true;
 }
 
-void UControllerGameplayCameraEvaluationComponent::Initialize(TScriptInterface<IGameplayCameraSystemHost> InCameraSystemHost)
+void UControllerGameplayCameraEvaluationComponent::Initialize(TScriptInterface<IGameplayCameraSystemHost> InCameraSystemHost, APlayerController* InPlayerController)
 {
+	using namespace UE::Cameras;
+
+	if (!ensureMsgf(
+				CameraSystemHost == nullptr,
+				TEXT("This component has already been initialized!")))
+	{
+		return;
+	}
+
 	ensure(InCameraSystemHost);
 	CameraSystemHost = InCameraSystemHost;
+
+	APlayerController* PlayerController = InPlayerController;
+	if (!PlayerController)
+	{
+		PlayerController = GetOwner<APlayerController>();
+	}
+
+	FCameraEvaluationContextInitializeParams InitParams;
+	InitParams.Owner = this;
+	InitParams.PlayerController = PlayerController;
+	EvaluationContext = MakeShared<FCameraEvaluationContext>(InitParams);
+	EvaluationContext->GetInitialResult().bIsValid = true;	
 }
 
 void UControllerGameplayCameraEvaluationComponent::ActivateCameraRig(UCameraRigAsset* CameraRig, ECameraRigLayer EvaluationLayer)
@@ -56,13 +77,9 @@ void UControllerGameplayCameraEvaluationComponent::ActivateCameraRigs()
 {
 	using namespace UE::Cameras;
 
-	if (!ensure(CameraSystemHost))
-	{
-		return;
-	}
-	
-	EnsureEvaluationContext();
-	if (!EvaluationContext)
+	if (!ensureMsgf(
+				CameraSystemHost && EvaluationContext,
+				TEXT("This component hasn't been initialized!")))
 	{
 		return;
 	}
@@ -85,23 +102,6 @@ void UControllerGameplayCameraEvaluationComponent::ActivateCameraRigs()
 	}
 }
 
-void UControllerGameplayCameraEvaluationComponent::EnsureEvaluationContext()
-{
-	using namespace UE::Cameras;
-
-	if (!EvaluationContext.IsValid())
-	{
-		// TODO: we won't find a player controller this way if we are under a view target.
-		APlayerController* PlayerController = GetOwner<APlayerController>();
-
-		FCameraEvaluationContextInitializeParams InitParams;
-		InitParams.Owner = this;
-		InitParams.PlayerController = PlayerController;
-		EvaluationContext = MakeShared<FCameraEvaluationContext>(InitParams);
-		EvaluationContext->GetInitialResult().bIsValid = true;	
-	}
-}
-
 UControllerGameplayCameraEvaluationComponent* UControllerGameplayCameraEvaluationComponent::FindComponent(AActor* OwnerActor)
 {
 	return OwnerActor->FindComponentByClass<UControllerGameplayCameraEvaluationComponent>();
diff --git a/Engine/Plugins/Cameras/GameplayCameras/Source/GameplayCameras/Public/GameFramework/ControllerGameplayCameraEvaluationComponent.h b/Engine/Plugins/Cameras/GameplayCameras/Source/GameplayCameras/Public/GameFramework/ControllerGameplayCameraEvaluationComponent.h
--- a/Engine/Plugins/Cameras/GameplayCameras/Source/GameplayCameras/Public/GameFramework/ControllerGameplayCameraEvaluationComponent.h
+++ b/Engine/Plugins/Cameras/GameplayCameras/Source/GameplayCameras/Public/GameFramework/ControllerGameplayCameraEvaluationComponent.h
@@ -30,7 +30,8 @@ public:
 
 	UControllerGameplayCameraEvaluationComponent(const FObjectInitializer& ObjectInitializer);
 
-	void Initialize(TScriptInterface<IGameplayCameraSystemHost> InCameraSystemHost);
+	/** Initializes this component's evaluation context. */
+	void Initialize(TScriptInterface<IGameplayCameraSystemHost> InCameraSystemHost, APlayerController* InPlayerController = nullptr);
 
 	/** Activates a new camera rig. */
 	void ActivateCameraRig(UCameraRigAsset* CameraRig, ECameraRigLayer EvaluationLayer);
@@ -49,7 +50,6 @@ public:
 private:
 
 	void ActivateCameraRigs();
-	void EnsureEvaluationContext();
 
 private:
 
