Cameras: allow running the target ray-cast node when there's no player controller

diff --git a/Engine/Plugins/Cameras/GameplayCameras/Source/GameplayCameras/Private/Nodes/Common/TargetRayCastCameraNode.cpp b/Engine/Plugins/Cameras/GameplayCameras/Source/GameplayCameras/Private/Nodes/Common/TargetRayCastCameraNode.cpp
--- a/Engine/Plugins/Cameras/GameplayCameras/Source/GameplayCameras/Private/Nodes/Common/TargetRayCastCameraNode.cpp
+++ b/Engine/Plugins/Cameras/GameplayCameras/Source/GameplayCameras/Private/Nodes/Common/TargetRayCastCameraNode.cpp
@@ -88,12 +88,12 @@ void FTargetRayCastCameraNodeEvaluator::OnRun(const FCameraNodeEvaluationParams&
 	}
 
 	UWorld* World = Params.EvaluationContext->GetWorld();
-	APlayerController* PlayerController = Params.EvaluationContext->GetPlayerController();
-	if (!World || !PlayerController)
+	if (!ensure(World))
 	{
 		return;
 	}
 
+	APlayerController* PlayerController = Params.EvaluationContext->GetPlayerController();
 	RunLineTrace(World, PlayerController, OutResult);
 }
 
@@ -102,7 +102,7 @@ void FTargetRayCastCameraNodeEvaluator::RunLineTrace(UWorld* World, APlayerContr
 	static FName LineTraceTag(TEXT("TargetRayCast"));
 	static FName LineTraceOwnerTag(TEXT("TargetRayCastCameraNode"));
 
-	APawn* Pawn = PlayerController->GetPawn();
+	APawn* Pawn = PlayerController ? PlayerController->GetPawn() : nullptr;
 
 	FCameraPose& CameraPose = OutResult.CameraPose;
 	const FVector3d RayStart = CameraPose.GetLocation();
