Cameras: make node evaluator and evaluation service flags inclusive by default (i.e. opt-out)

This fixes the problem of someone implementing one of the API methods, and that method not being called until they think of adding the flag.
Now, all flags are enabled by default, and people can opt-out of them to optimize the camera system evaluation.

diff --git a/Engine/Plugins/Cameras/GameplayCameras/Source/GameplayCameras/Private/Core/BlendStackCameraNode.cpp b/Engine/Plugins/Cameras/GameplayCameras/Source/GameplayCameras/Private/Core/BlendStackCameraNode.cpp
--- a/Engine/Plugins/Cameras/GameplayCameras/Source/GameplayCameras/Private/Core/BlendStackCameraNode.cpp
+++ b/Engine/Plugins/Cameras/GameplayCameras/Source/GameplayCameras/Private/Core/BlendStackCameraNode.cpp
@@ -62,6 +62,8 @@ UE_DEFINE_CAMERA_NODE_EVALUATOR(FBlendStackCameraNodeEvaluator)
 
 FBlendStackCameraNodeEvaluator::FBlendStackCameraNodeEvaluator()
 {
+	SetNodeEvaluatorFlags(ECameraNodeEvaluatorFlags::NeedsSerialize);
+
 #if WITH_EDITOR || UE_GAMEPLAY_CAMERAS_DEBUG
 	bAutoCameraPoseMovementTrail = false;
 #endif  // WITH_EDITOR || UE_GAMEPLAY_CAMERAS_DEBUG
diff --git a/Engine/Plugins/Cameras/GameplayCameras/Source/GameplayCameras/Private/Core/BlendStackRootCameraNode.cpp b/Engine/Plugins/Cameras/GameplayCameras/Source/GameplayCameras/Private/Core/BlendStackRootCameraNode.cpp
--- a/Engine/Plugins/Cameras/GameplayCameras/Source/GameplayCameras/Private/Core/BlendStackRootCameraNode.cpp
+++ b/Engine/Plugins/Cameras/GameplayCameras/Source/GameplayCameras/Private/Core/BlendStackRootCameraNode.cpp
@@ -56,7 +56,7 @@ UE_DEFINE_CAMERA_DEBUG_BLOCK_WITH_FIELDS(FBlendStackRootCameraDebugBlock)
 
 FBlendStackRootCameraNodeEvaluator::FBlendStackRootCameraNodeEvaluator()
 {
-	AddNodeEvaluatorFlags(ECameraNodeEvaluatorFlags::NeedsParameterUpdate);
+	SetNodeEvaluatorFlags(ECameraNodeEvaluatorFlags::NeedsParameterUpdate);
 }
 
 FCameraNodeEvaluatorChildrenView FBlendStackRootCameraNodeEvaluator::OnGetChildren()
diff --git a/Engine/Plugins/Cameras/GameplayCameras/Source/GameplayCameras/Private/Nodes/Attach/AttachToActorCameraNode.cpp b/Engine/Plugins/Cameras/GameplayCameras/Source/GameplayCameras/Private/Nodes/Attach/AttachToActorCameraNode.cpp
--- a/Engine/Plugins/Cameras/GameplayCameras/Source/GameplayCameras/Private/Nodes/Attach/AttachToActorCameraNode.cpp
+++ b/Engine/Plugins/Cameras/GameplayCameras/Source/GameplayCameras/Private/Nodes/Attach/AttachToActorCameraNode.cpp
@@ -47,6 +47,8 @@ UE_DEFINE_CAMERA_DEBUG_BLOCK_WITH_FIELDS(FAttachToActorCameraDebugBlock)
 
 void FAttachToActorCameraNodeEvaluator::OnInitialize(const FCameraNodeEvaluatorInitializeParams& Params, FCameraNodeEvaluationResult& OutResult)
 {
+	SetNodeEvaluatorFlags(ECameraNodeEvaluatorFlags::NeedsSerialize);
+
 	const UAttachToActorCameraNode* AttachNode = GetCameraNodeAs<UAttachToActorCameraNode>();
 	AttachmentReader.Initialize(AttachNode->Attachment, AttachNode->AttachmentDataID);
 	AttachToLocationReader.Initialize(AttachNode->AttachToLocation);
diff --git a/Engine/Plugins/Cameras/GameplayCameras/Source/GameplayCameras/Private/Nodes/Attach/AttachToActorGroupCameraNode.cpp b/Engine/Plugins/Cameras/GameplayCameras/Source/GameplayCameras/Private/Nodes/Attach/AttachToActorGroupCameraNode.cpp
--- a/Engine/Plugins/Cameras/GameplayCameras/Source/GameplayCameras/Private/Nodes/Attach/AttachToActorGroupCameraNode.cpp
+++ b/Engine/Plugins/Cameras/GameplayCameras/Source/GameplayCameras/Private/Nodes/Attach/AttachToActorGroupCameraNode.cpp
@@ -31,6 +31,8 @@ UE_DEFINE_CAMERA_NODE_EVALUATOR(FAttachToActorGroupCameraNodeEvaluator)
 
 void FAttachToActorGroupCameraNodeEvaluator::OnInitialize(const FCameraNodeEvaluatorInitializeParams& Params, FCameraNodeEvaluationResult& OutResult)
 {
+	SetNodeEvaluatorFlags(ECameraNodeEvaluatorFlags::None);
+
 	const UAttachToActorGroupCameraNode* AttachNode = GetCameraNodeAs<UAttachToActorGroupCameraNode>();
 	AttachmentsReader.Initialize(AttachNode->Attachments, AttachNode->AttachmentsDataID);
 }
diff --git a/Engine/Plugins/Cameras/GameplayCameras/Source/GameplayCameras/Private/Nodes/Attach/AttachToPlayerPawnCameraNode.cpp b/Engine/Plugins/Cameras/GameplayCameras/Source/GameplayCameras/Private/Nodes/Attach/AttachToPlayerPawnCameraNode.cpp
--- a/Engine/Plugins/Cameras/GameplayCameras/Source/GameplayCameras/Private/Nodes/Attach/AttachToPlayerPawnCameraNode.cpp
+++ b/Engine/Plugins/Cameras/GameplayCameras/Source/GameplayCameras/Private/Nodes/Attach/AttachToPlayerPawnCameraNode.cpp
@@ -35,6 +35,8 @@ UE_DEFINE_CAMERA_NODE_EVALUATOR(FAttachToPlayerPawnCameraNodeEvaluator)
 
 void FAttachToPlayerPawnCameraNodeEvaluator::OnInitialize(const FCameraNodeEvaluatorInitializeParams& Params, FCameraNodeEvaluationResult& OutResult)
 {
+	SetNodeEvaluatorFlags(ECameraNodeEvaluatorFlags::None);
+
 	const UAttachToPlayerPawnCameraNode* AttachNode = GetCameraNodeAs<UAttachToPlayerPawnCameraNode>();
 	AttachToLocationReader.Initialize(AttachNode->AttachToLocation);
 	AttachToRotationReader.Initialize(AttachNode->AttachToRotation);
diff --git a/Engine/Plugins/Cameras/GameplayCameras/Source/GameplayCameras/Private/Nodes/Blends/InterruptedBlendCameraNode.cpp b/Engine/Plugins/Cameras/GameplayCameras/Source/GameplayCameras/Private/Nodes/Blends/InterruptedBlendCameraNode.cpp
--- a/Engine/Plugins/Cameras/GameplayCameras/Source/GameplayCameras/Private/Nodes/Blends/InterruptedBlendCameraNode.cpp
+++ b/Engine/Plugins/Cameras/GameplayCameras/Source/GameplayCameras/Private/Nodes/Blends/InterruptedBlendCameraNode.cpp
@@ -7,6 +7,18 @@ namespace UE::Cameras
 
 UE_DEFINE_BLEND_CAMERA_NODE_EVALUATOR(FInterruptedBlendCameraNodeEvaluator)
 
+FInterruptedBlendCameraNodeEvaluator::FInterruptedBlendCameraNodeEvaluator()
+{
+	SetNodeEvaluatorFlags(ECameraNodeEvaluatorFlags::None);
+}
+
+FInterruptedBlendCameraNodeEvaluator::FInterruptedBlendCameraNodeEvaluator(FBlendCameraNodeEvaluator* InChildBlend, FBlendCameraNodeEvaluator* InFrozenBlend) 
+	: ChildBlend(InChildBlend)
+	, FrozenBlend(InFrozenBlend)
+{
+	SetNodeEvaluatorFlags(ECameraNodeEvaluatorFlags::None);
+}
+
 FCameraNodeEvaluatorChildrenView FInterruptedBlendCameraNodeEvaluator::OnGetChildren()
 {
 	return FCameraNodeEvaluatorChildrenView{ ChildBlend, FrozenBlend };
diff --git a/Engine/Plugins/Cameras/GameplayCameras/Source/GameplayCameras/Private/Nodes/Blends/LocationRotationBlendCameraNode.cpp b/Engine/Plugins/Cameras/GameplayCameras/Source/GameplayCameras/Private/Nodes/Blends/LocationRotationBlendCameraNode.cpp
--- a/Engine/Plugins/Cameras/GameplayCameras/Source/GameplayCameras/Private/Nodes/Blends/LocationRotationBlendCameraNode.cpp
+++ b/Engine/Plugins/Cameras/GameplayCameras/Source/GameplayCameras/Private/Nodes/Blends/LocationRotationBlendCameraNode.cpp
@@ -16,6 +16,13 @@ class FLocationRotationBlendCameraNodeEvaluator : public FBlendCameraNodeEvaluat
 {
 	UE_DECLARE_BLEND_CAMERA_NODE_EVALUATOR(GAMEPLAYCAMERAS_API, FLocationRotationBlendCameraNodeEvaluator)
 
+public:
+
+	FLocationRotationBlendCameraNodeEvaluator()
+	{
+		SetNodeEvaluatorFlags(ECameraNodeEvaluatorFlags::None);
+	}
+
 protected:
 
 	// FCameraNodeEvaluator interface.
diff --git a/Engine/Plugins/Cameras/GameplayCameras/Source/GameplayCameras/Private/Nodes/Blends/OrbitBlendCameraNode.cpp b/Engine/Plugins/Cameras/GameplayCameras/Source/GameplayCameras/Private/Nodes/Blends/OrbitBlendCameraNode.cpp
--- a/Engine/Plugins/Cameras/GameplayCameras/Source/GameplayCameras/Private/Nodes/Blends/OrbitBlendCameraNode.cpp
+++ b/Engine/Plugins/Cameras/GameplayCameras/Source/GameplayCameras/Private/Nodes/Blends/OrbitBlendCameraNode.cpp
@@ -24,6 +24,13 @@ class FOrbitBlendCameraNodeEvaluator : public FBlendCameraNodeEvaluator
 {
 	UE_DECLARE_BLEND_CAMERA_NODE_EVALUATOR(GAMEPLAYCAMERAS_API, FOrbitBlendCameraNodeEvaluator)
 
+public:
+
+	FOrbitBlendCameraNodeEvaluator()
+	{
+		SetNodeEvaluatorFlags(ECameraNodeEvaluatorFlags::NeedsSerialize);
+	}
+
 protected:
 
 	// FCameraNodeEvaluator interface.
diff --git a/Engine/Plugins/Cameras/GameplayCameras/Source/GameplayCameras/Private/Nodes/Blends/PopBlendCameraNode.cpp b/Engine/Plugins/Cameras/GameplayCameras/Source/GameplayCameras/Private/Nodes/Blends/PopBlendCameraNode.cpp
--- a/Engine/Plugins/Cameras/GameplayCameras/Source/GameplayCameras/Private/Nodes/Blends/PopBlendCameraNode.cpp
+++ b/Engine/Plugins/Cameras/GameplayCameras/Source/GameplayCameras/Private/Nodes/Blends/PopBlendCameraNode.cpp
@@ -11,7 +11,15 @@ class FPopBlendCameraNodeEvaluator : public FBlendCameraNodeEvaluator
 {
 	UE_DECLARE_BLEND_CAMERA_NODE_EVALUATOR(GAMEPLAYCAMERAS_API, FPopBlendCameraNodeEvaluator)
 
+public:
+
+	FPopBlendCameraNodeEvaluator()
+	{
+		SetNodeEvaluatorFlags(ECameraNodeEvaluatorFlags::None);
+	}
+
 protected:
+
 	virtual void OnRun(const FCameraNodeEvaluationParams& Params, FCameraNodeEvaluationResult& OutResult) override;
 	virtual void OnBlendParameters(const FCameraNodePreBlendParams& Params, FCameraNodePreBlendResult& OutResult) override;
 	virtual void OnBlendResults(const FCameraNodeBlendParams& Params, FCameraNodeBlendResult& OutResult) override;
diff --git a/Engine/Plugins/Cameras/GameplayCameras/Source/GameplayCameras/Private/Nodes/Blends/ReverseBlendCameraNode.cpp b/Engine/Plugins/Cameras/GameplayCameras/Source/GameplayCameras/Private/Nodes/Blends/ReverseBlendCameraNode.cpp
--- a/Engine/Plugins/Cameras/GameplayCameras/Source/GameplayCameras/Private/Nodes/Blends/ReverseBlendCameraNode.cpp
+++ b/Engine/Plugins/Cameras/GameplayCameras/Source/GameplayCameras/Private/Nodes/Blends/ReverseBlendCameraNode.cpp
@@ -7,6 +7,17 @@ namespace UE::Cameras
 
 UE_DEFINE_BLEND_CAMERA_NODE_EVALUATOR(FReverseBlendCameraNodeEvaluator)
 
+FReverseBlendCameraNodeEvaluator::FReverseBlendCameraNodeEvaluator()
+{
+	SetNodeEvaluatorFlags(ECameraNodeEvaluatorFlags::None);
+}
+
+FReverseBlendCameraNodeEvaluator::FReverseBlendCameraNodeEvaluator(FBlendCameraNodeEvaluator* InChildBlend) 
+	: ChildBlend(InChildBlend)
+{
+	SetNodeEvaluatorFlags(ECameraNodeEvaluatorFlags::None);
+}
+
 FCameraNodeEvaluatorChildrenView FReverseBlendCameraNodeEvaluator::OnGetChildren()
 {
 	return FCameraNodeEvaluatorChildrenView{ ChildBlend };
diff --git a/Engine/Plugins/Cameras/GameplayCameras/Source/GameplayCameras/Private/Nodes/Blends/SimpleBlendCameraNode.cpp b/Engine/Plugins/Cameras/GameplayCameras/Source/GameplayCameras/Private/Nodes/Blends/SimpleBlendCameraNode.cpp
--- a/Engine/Plugins/Cameras/GameplayCameras/Source/GameplayCameras/Private/Nodes/Blends/SimpleBlendCameraNode.cpp
+++ b/Engine/Plugins/Cameras/GameplayCameras/Source/GameplayCameras/Private/Nodes/Blends/SimpleBlendCameraNode.cpp
@@ -19,6 +19,11 @@ UE_DECLARE_CAMERA_DEBUG_BLOCK_END()
 
 UE_DEFINE_CAMERA_DEBUG_BLOCK_WITH_FIELDS(FSimpleBlendCameraDebugBlock)
 
+FSimpleBlendCameraNodeEvaluator::FSimpleBlendCameraNodeEvaluator()
+{
+	SetNodeEvaluatorFlags(ECameraNodeEvaluatorFlags::NeedsSerialize);
+}
+
 void FSimpleBlendCameraNodeEvaluator::OnRun(const FCameraNodeEvaluationParams& Params, FCameraNodeEvaluationResult& OutResult)
 {
 	FSimpleBlendCameraNodeEvaluationResult FactorResult;
diff --git a/Engine/Plugins/Cameras/GameplayCameras/Source/GameplayCameras/Private/Nodes/Collision/CollisionPushCameraNode.cpp b/Engine/Plugins/Cameras/GameplayCameras/Source/GameplayCameras/Private/Nodes/Collision/CollisionPushCameraNode.cpp
--- a/Engine/Plugins/Cameras/GameplayCameras/Source/GameplayCameras/Private/Nodes/Collision/CollisionPushCameraNode.cpp
+++ b/Engine/Plugins/Cameras/GameplayCameras/Source/GameplayCameras/Private/Nodes/Collision/CollisionPushCameraNode.cpp
@@ -33,9 +33,9 @@ class FCollisionPushCameraNodeEvaluator : public FCameraNodeEvaluator
 
 protected:
 
+	// FCameraNodeEvaluator interface.
 	virtual void OnInitialize(const FCameraNodeEvaluatorInitializeParams& Params, FCameraNodeEvaluationResult& OutResult) override;
 	virtual void OnRun(const FCameraNodeEvaluationParams& Params, FCameraNodeEvaluationResult& OutResult) override;
-
 #if UE_GAMEPLAY_CAMERAS_DEBUG
 	virtual void OnBuildDebugBlocks(const FCameraDebugBlockBuildParams& Params, FCameraDebugBlockBuilder& Builder) override;
 #endif  // UE_GAMEPLAY_CAMERAS_DEBUG
@@ -100,6 +100,8 @@ UE_DEFINE_CAMERA_DEBUG_BLOCK_WITH_FIELDS(FCollisionPushCameraDebugBlock)
 
 void FCollisionPushCameraNodeEvaluator::OnInitialize(const FCameraNodeEvaluatorInitializeParams& Params, FCameraNodeEvaluationResult& OutResult)
 {
+	SetNodeEvaluatorFlags(ECameraNodeEvaluatorFlags::None);
+
 	const UCollisionPushCameraNode* CollisionPushNode = GetCameraNodeAs<UCollisionPushCameraNode>();
 
 	EnableCollisionReader.Initialize(CollisionPushNode->EnableCollision, true);
diff --git a/Engine/Plugins/Cameras/GameplayCameras/Source/GameplayCameras/Private/Nodes/Collision/OcclusionMaterialCameraNode.cpp b/Engine/Plugins/Cameras/GameplayCameras/Source/GameplayCameras/Private/Nodes/Collision/OcclusionMaterialCameraNode.cpp
--- a/Engine/Plugins/Cameras/GameplayCameras/Source/GameplayCameras/Private/Nodes/Collision/OcclusionMaterialCameraNode.cpp
+++ b/Engine/Plugins/Cameras/GameplayCameras/Source/GameplayCameras/Private/Nodes/Collision/OcclusionMaterialCameraNode.cpp
@@ -38,6 +38,7 @@ public:
 
 protected:
 
+	// FCameraNodeEvaluator interface.
 	virtual void OnInitialize(const FCameraNodeEvaluatorInitializeParams& Params, FCameraNodeEvaluationResult& OutResult) override;
 	virtual void OnRun(const FCameraNodeEvaluationParams& Params, FCameraNodeEvaluationResult& OutResult) override;
 
@@ -73,6 +74,8 @@ FOcclusionMaterialCameraNodeEvaluator::~FOcclusionMaterialCameraNodeEvaluator()
 
 void FOcclusionMaterialCameraNodeEvaluator::OnInitialize(const FCameraNodeEvaluatorInitializeParams& Params, FCameraNodeEvaluationResult& OutResult)
 {
+	SetNodeEvaluatorFlags(ECameraNodeEvaluatorFlags::None);
+
 	const UOcclusionMaterialCameraNode* OcclusionMaterialNode = GetCameraNodeAs<UOcclusionMaterialCameraNode>();
 	OcclusionSphereRadiusReader.Initialize(OcclusionMaterialNode->OcclusionSphereRadius);
 	OcclusionTargetOffsetReader.Initialize(OcclusionMaterialNode->OcclusionTargetOffset);
diff --git a/Engine/Plugins/Cameras/GameplayCameras/Source/GameplayCameras/Private/Nodes/Common/ArrayCameraNode.cpp b/Engine/Plugins/Cameras/GameplayCameras/Source/GameplayCameras/Private/Nodes/Common/ArrayCameraNode.cpp
--- a/Engine/Plugins/Cameras/GameplayCameras/Source/GameplayCameras/Private/Nodes/Common/ArrayCameraNode.cpp
+++ b/Engine/Plugins/Cameras/GameplayCameras/Source/GameplayCameras/Private/Nodes/Common/ArrayCameraNode.cpp
@@ -13,8 +13,16 @@ class FArrayCameraNodeEvaluator : public FCameraNodeEvaluator
 {
 	UE_DECLARE_CAMERA_NODE_EVALUATOR(GAMEPLAYCAMERAS_API, FArrayCameraNodeEvaluator)
 
+public:
+
+	FArrayCameraNodeEvaluator()
+	{
+		SetNodeEvaluatorFlags(ECameraNodeEvaluatorFlags::None);
+	}
+
 protected:
 
+	// FCameraNodeEvaluator interface.
 	virtual FCameraNodeEvaluatorChildrenView OnGetChildren() override;
 	virtual void OnBuild(const FCameraNodeEvaluatorBuildParams& Params) override;
 	virtual void OnRun(const FCameraNodeEvaluationParams& Params, FCameraNodeEvaluationResult& OutResult) override;
diff --git a/Engine/Plugins/Cameras/GameplayCameras/Source/GameplayCameras/Private/Nodes/Common/AutoFocusCameraNode.cpp b/Engine/Plugins/Cameras/GameplayCameras/Source/GameplayCameras/Private/Nodes/Common/AutoFocusCameraNode.cpp
--- a/Engine/Plugins/Cameras/GameplayCameras/Source/GameplayCameras/Private/Nodes/Common/AutoFocusCameraNode.cpp
+++ b/Engine/Plugins/Cameras/GameplayCameras/Source/GameplayCameras/Private/Nodes/Common/AutoFocusCameraNode.cpp
@@ -21,6 +21,7 @@ class FAutoFocusCameraNodeEvaluator : public FCameraNodeEvaluator
 
 protected:
 
+	// FCameraNodeEvaluator interface.
 	virtual void OnInitialize(const FCameraNodeEvaluatorInitializeParams& Params, FCameraNodeEvaluationResult& OutResult) override;
 	virtual void OnRun(const FCameraNodeEvaluationParams& Params, FCameraNodeEvaluationResult& OutResult) override;
 	virtual void OnSerialize(const FCameraNodeEvaluatorSerializeParams& Params, FArchive& Ar) override;
@@ -53,6 +54,8 @@ UE_DEFINE_CAMERA_DEBUG_BLOCK_WITH_FIELDS(FAutoFocusCameraDebugBlock)
 
 void FAutoFocusCameraNodeEvaluator::OnInitialize(const FCameraNodeEvaluatorInitializeParams& Params, FCameraNodeEvaluationResult& OutResult)
 {
+	SetNodeEvaluatorFlags(ECameraNodeEvaluatorFlags::NeedsSerialize);
+
 	const UAutoFocusCameraNode* AutoFocusNode = GetCameraNodeAs<UAutoFocusCameraNode>();
 
 	EnableAutoFocusReader.Initialize(AutoFocusNode->EnableAutoFocus);
diff --git a/Engine/Plugins/Cameras/GameplayCameras/Source/GameplayCameras/Private/Nodes/Common/BodyParametersCameraNode.cpp b/Engine/Plugins/Cameras/GameplayCameras/Source/GameplayCameras/Private/Nodes/Common/BodyParametersCameraNode.cpp
--- a/Engine/Plugins/Cameras/GameplayCameras/Source/GameplayCameras/Private/Nodes/Common/BodyParametersCameraNode.cpp
+++ b/Engine/Plugins/Cameras/GameplayCameras/Source/GameplayCameras/Private/Nodes/Common/BodyParametersCameraNode.cpp
@@ -16,6 +16,7 @@ class FBodyParametersCameraNodeEvaluator : public FCameraNodeEvaluator
 
 protected:
 
+	// FCameraNodeEvaluator interface.
 	virtual void OnInitialize(const FCameraNodeEvaluatorInitializeParams& Params, FCameraNodeEvaluationResult& OutResult) override;
 	virtual void OnRun(const FCameraNodeEvaluationParams& Params, FCameraNodeEvaluationResult& OutResult) override;
 
@@ -29,6 +30,8 @@ UE_DEFINE_CAMERA_NODE_EVALUATOR(FBodyParametersCameraNodeEvaluator)
 
 void FBodyParametersCameraNodeEvaluator::OnInitialize(const FCameraNodeEvaluatorInitializeParams& Params, FCameraNodeEvaluationResult& OutResult)
 {
+	SetNodeEvaluatorFlags(ECameraNodeEvaluatorFlags::None);
+
 	const UBodyParametersCameraNode* BodyParametersNode = GetCameraNodeAs<UBodyParametersCameraNode>();
 	ShutterSpeedReader.Initialize(BodyParametersNode->ShutterSpeed);
 	ISOReader.Initialize(BodyParametersNode->ISO);
diff --git a/Engine/Plugins/Cameras/GameplayCameras/Source/GameplayCameras/Private/Nodes/Common/BoomArmCameraNode.cpp b/Engine/Plugins/Cameras/GameplayCameras/Source/GameplayCameras/Private/Nodes/Common/BoomArmCameraNode.cpp
--- a/Engine/Plugins/Cameras/GameplayCameras/Source/GameplayCameras/Private/Nodes/Common/BoomArmCameraNode.cpp
+++ b/Engine/Plugins/Cameras/GameplayCameras/Source/GameplayCameras/Private/Nodes/Common/BoomArmCameraNode.cpp
@@ -25,6 +25,7 @@ class FBoomArmCameraNodeEvaluator : public FCameraNodeEvaluator
 
 protected:
 
+	// FCameraNodeEvaluator interface.
 	virtual void OnBuild(const FCameraNodeEvaluatorBuildParams& Params) override;
 	virtual void OnInitialize(const FCameraNodeEvaluatorInitializeParams& Params, FCameraNodeEvaluationResult& OutResult) override;
 	virtual FCameraNodeEvaluatorChildrenView OnGetChildren() override;
@@ -79,7 +80,7 @@ void FBoomArmCameraNodeEvaluator::OnBuild(const FCameraNodeEvaluatorBuildParams&
 
 void FBoomArmCameraNodeEvaluator::OnInitialize(const FCameraNodeEvaluatorInitializeParams& Params, FCameraNodeEvaluationResult& OutResult)
 {
-	AddNodeEvaluatorFlags(ECameraNodeEvaluatorFlags::SupportsOperations);
+	SetNodeEvaluatorFlags(ECameraNodeEvaluatorFlags::SupportsOperations);
 
 	const UBoomArmCameraNode* BoomArmNode = GetCameraNodeAs<UBoomArmCameraNode>();
 	BoomOffsetReader.Initialize(BoomArmNode->BoomOffset);
diff --git a/Engine/Plugins/Cameras/GameplayCameras/Source/GameplayCameras/Private/Nodes/Common/CameraRigCameraNode.cpp b/Engine/Plugins/Cameras/GameplayCameras/Source/GameplayCameras/Private/Nodes/Common/CameraRigCameraNode.cpp
--- a/Engine/Plugins/Cameras/GameplayCameras/Source/GameplayCameras/Private/Nodes/Common/CameraRigCameraNode.cpp
+++ b/Engine/Plugins/Cameras/GameplayCameras/Source/GameplayCameras/Private/Nodes/Common/CameraRigCameraNode.cpp
@@ -22,7 +22,7 @@ UE_DEFINE_CAMERA_NODE_EVALUATOR(FCameraRigCameraNodeEvaluator)
 
 FCameraRigCameraNodeEvaluator::FCameraRigCameraNodeEvaluator()
 {
-	AddNodeEvaluatorFlags(ECameraNodeEvaluatorFlags::NeedsParameterUpdate);
+	SetNodeEvaluatorFlags(ECameraNodeEvaluatorFlags::NeedsParameterUpdate);
 }
 
 FCameraNodeEvaluatorChildrenView FCameraRigCameraNodeEvaluator::OnGetChildren()
diff --git a/Engine/Plugins/Cameras/GameplayCameras/Source/GameplayCameras/Private/Nodes/Common/ClippingPlanesCameraNode.cpp b/Engine/Plugins/Cameras/GameplayCameras/Source/GameplayCameras/Private/Nodes/Common/ClippingPlanesCameraNode.cpp
--- a/Engine/Plugins/Cameras/GameplayCameras/Source/GameplayCameras/Private/Nodes/Common/ClippingPlanesCameraNode.cpp
+++ b/Engine/Plugins/Cameras/GameplayCameras/Source/GameplayCameras/Private/Nodes/Common/ClippingPlanesCameraNode.cpp
@@ -17,6 +17,7 @@ class FClippingPlanesCameraNodeEvaluator : public FCameraNodeEvaluator
 
 protected:
 
+	// FCameraNodeEvaluator interface.
 	virtual void OnInitialize(const FCameraNodeEvaluatorInitializeParams& Params, FCameraNodeEvaluationResult& OutResult) override;
 	virtual void OnRun(const FCameraNodeEvaluationParams& Params, FCameraNodeEvaluationResult& OutResult) override;
 
@@ -30,6 +31,8 @@ UE_DEFINE_CAMERA_NODE_EVALUATOR(FClippingPlanesCameraNodeEvaluator)
 
 void FClippingPlanesCameraNodeEvaluator::OnInitialize(const FCameraNodeEvaluatorInitializeParams& Params, FCameraNodeEvaluationResult& OutResult)
 {
+	SetNodeEvaluatorFlags(ECameraNodeEvaluatorFlags::None);
+
 	const UClippingPlanesCameraNode* ClippingPlanesNode = GetCameraNodeAs<UClippingPlanesCameraNode>();
 	NearPlaneReader.Initialize(ClippingPlanesNode->NearPlane);
 	FarPlaneReader.Initialize(ClippingPlanesNode->FarPlane);
diff --git a/Engine/Plugins/Cameras/GameplayCameras/Source/GameplayCameras/Private/Nodes/Common/DampenPositionCameraNode.cpp b/Engine/Plugins/Cameras/GameplayCameras/Source/GameplayCameras/Private/Nodes/Common/DampenPositionCameraNode.cpp
--- a/Engine/Plugins/Cameras/GameplayCameras/Source/GameplayCameras/Private/Nodes/Common/DampenPositionCameraNode.cpp
+++ b/Engine/Plugins/Cameras/GameplayCameras/Source/GameplayCameras/Private/Nodes/Common/DampenPositionCameraNode.cpp
@@ -23,6 +23,7 @@ class FDampenPositionCameraNodeEvaluator : public FCameraNodeEvaluator
 
 protected:
 
+	// FCameraNodeEvaluator interface.
 	virtual void OnInitialize(const FCameraNodeEvaluatorInitializeParams& Params, FCameraNodeEvaluationResult& OutResult) override;
 	virtual void OnRun(const FCameraNodeEvaluationParams& Params, FCameraNodeEvaluationResult& OutResult) override;
 	virtual void OnSerialize(const FCameraNodeEvaluatorSerializeParams& Params, FArchive& Ar) override;
@@ -75,6 +76,8 @@ UE_DEFINE_CAMERA_DEBUG_BLOCK_WITH_FIELDS(FDampenPositionCameraDebugBlock)
 
 void FDampenPositionCameraNodeEvaluator::OnInitialize(const FCameraNodeEvaluatorInitializeParams& Params, FCameraNodeEvaluationResult& OutResult)
 {
+	SetNodeEvaluatorFlags(ECameraNodeEvaluatorFlags::NeedsSerialize);
+
 	const UDampenPositionCameraNode* DampenNode = GetCameraNodeAs<UDampenPositionCameraNode>();
 
 	ForwardDampingFactorReader.Initialize(DampenNode->ForwardDampingFactor);
diff --git a/Engine/Plugins/Cameras/GameplayCameras/Source/GameplayCameras/Private/Nodes/Common/DampenRotationCameraNode.cpp b/Engine/Plugins/Cameras/GameplayCameras/Source/GameplayCameras/Private/Nodes/Common/DampenRotationCameraNode.cpp
--- a/Engine/Plugins/Cameras/GameplayCameras/Source/GameplayCameras/Private/Nodes/Common/DampenRotationCameraNode.cpp
+++ b/Engine/Plugins/Cameras/GameplayCameras/Source/GameplayCameras/Private/Nodes/Common/DampenRotationCameraNode.cpp
@@ -66,6 +66,8 @@ UE_DEFINE_CAMERA_DEBUG_BLOCK_WITH_FIELDS(FDampenRotationCameraDebugBlock)
 
 void FDampenRotationCameraNodeEvaluator::OnInitialize(const FCameraNodeEvaluatorInitializeParams& Params, FCameraNodeEvaluationResult& OutResult)
 {
+	SetNodeEvaluatorFlags(ECameraNodeEvaluatorFlags::NeedsSerialize);
+
 	const UDampenRotationCameraNode* DampenNode = GetCameraNodeAs<UDampenRotationCameraNode>();
 
 	YawDampingFactorReader.Initialize(DampenNode->YawDampingFactor);
diff --git a/Engine/Plugins/Cameras/GameplayCameras/Source/GameplayCameras/Private/Nodes/Common/FieldOfViewCameraNode.cpp b/Engine/Plugins/Cameras/GameplayCameras/Source/GameplayCameras/Private/Nodes/Common/FieldOfViewCameraNode.cpp
--- a/Engine/Plugins/Cameras/GameplayCameras/Source/GameplayCameras/Private/Nodes/Common/FieldOfViewCameraNode.cpp
+++ b/Engine/Plugins/Cameras/GameplayCameras/Source/GameplayCameras/Private/Nodes/Common/FieldOfViewCameraNode.cpp
@@ -17,6 +17,7 @@ class FFieldOfViewCameraNodeEvaluator : public FCameraNodeEvaluator
 
 protected:
 
+	// FCameraNodeEvaluator interface.
 	virtual void OnInitialize(const FCameraNodeEvaluatorInitializeParams& Params, FCameraNodeEvaluationResult& OutResult) override;
 	virtual void OnRun(const FCameraNodeEvaluationParams& Params, FCameraNodeEvaluationResult& OutResult) override;
 
@@ -29,6 +30,8 @@ UE_DEFINE_CAMERA_NODE_EVALUATOR(FFieldOfViewCameraNodeEvaluator)
 
 void FFieldOfViewCameraNodeEvaluator::OnInitialize(const FCameraNodeEvaluatorInitializeParams& Params, FCameraNodeEvaluationResult& OutResult)
 {
+	SetNodeEvaluatorFlags(ECameraNodeEvaluatorFlags::None);
+
 	const UFieldOfViewCameraNode* FieldOfViewNode = GetCameraNodeAs<UFieldOfViewCameraNode>();
 	FieldOfViewReader.Initialize(FieldOfViewNode->FieldOfView);
 }
diff --git a/Engine/Plugins/Cameras/GameplayCameras/Source/GameplayCameras/Private/Nodes/Common/FilmbackCameraNode.cpp b/Engine/Plugins/Cameras/GameplayCameras/Source/GameplayCameras/Private/Nodes/Common/FilmbackCameraNode.cpp
--- a/Engine/Plugins/Cameras/GameplayCameras/Source/GameplayCameras/Private/Nodes/Common/FilmbackCameraNode.cpp
+++ b/Engine/Plugins/Cameras/GameplayCameras/Source/GameplayCameras/Private/Nodes/Common/FilmbackCameraNode.cpp
@@ -16,6 +16,7 @@ class FFilmbackCameraNodeEvaluator : public FCameraNodeEvaluator
 
 protected:
 
+	// FCameraNodeEvaluator interface.
 	virtual void OnInitialize(const FCameraNodeEvaluatorInitializeParams& Params, FCameraNodeEvaluationResult& OutResult) override;
 	virtual void OnRun(const FCameraNodeEvaluationParams& Params, FCameraNodeEvaluationResult& OutResult) override;
 
@@ -34,6 +35,8 @@ UE_DEFINE_CAMERA_NODE_EVALUATOR(FFilmbackCameraNodeEvaluator)
 
 void FFilmbackCameraNodeEvaluator::OnInitialize(const FCameraNodeEvaluatorInitializeParams& Params, FCameraNodeEvaluationResult& OutResult)
 {
+	SetNodeEvaluatorFlags(ECameraNodeEvaluatorFlags::None);
+
 	const UFilmbackCameraNode* FilmbackNode = GetCameraNodeAs<UFilmbackCameraNode>();
 	SensorWidthReader.Initialize(FilmbackNode->SensorWidth);
 	SensorHeightReader.Initialize(FilmbackNode->SensorHeight);
diff --git a/Engine/Plugins/Cameras/GameplayCameras/Source/GameplayCameras/Private/Nodes/Common/LensCalibrationCameraNode.cpp b/Engine/Plugins/Cameras/GameplayCameras/Source/GameplayCameras/Private/Nodes/Common/LensCalibrationCameraNode.cpp
--- a/Engine/Plugins/Cameras/GameplayCameras/Source/GameplayCameras/Private/Nodes/Common/LensCalibrationCameraNode.cpp
+++ b/Engine/Plugins/Cameras/GameplayCameras/Source/GameplayCameras/Private/Nodes/Common/LensCalibrationCameraNode.cpp
@@ -19,6 +19,7 @@ class FLensCalibrationCameraNodeEvaluator : public FCameraNodeEvaluator
 
 protected:
 
+	// FCameraNodeEvaluator interface.
 	virtual void OnInitialize(const FCameraNodeEvaluatorInitializeParams& Params, FCameraNodeEvaluationResult& OutResult) override;
 	virtual void OnRun(const FCameraNodeEvaluationParams& Params, FCameraNodeEvaluationResult& OutResult) override;
 	virtual void OnAddReferencedObjects(FReferenceCollector& Collector) override;
@@ -33,6 +34,8 @@ UE_DEFINE_CAMERA_NODE_EVALUATOR(FLensCalibrationCameraNodeEvaluator)
 
 void FLensCalibrationCameraNodeEvaluator::OnInitialize(const FCameraNodeEvaluatorInitializeParams& Params, FCameraNodeEvaluationResult& OutResult)
 {
+	SetNodeEvaluatorFlags(ECameraNodeEvaluatorFlags::None);
+
 	const ULensCalibrationCameraNode* LensCalibrationNode = GetCameraNodeAs<ULensCalibrationCameraNode>();
 	if (LensCalibrationNode->LensFile)
 	{
diff --git a/Engine/Plugins/Cameras/GameplayCameras/Source/GameplayCameras/Private/Nodes/Common/LensParametersCameraNode.cpp b/Engine/Plugins/Cameras/GameplayCameras/Source/GameplayCameras/Private/Nodes/Common/LensParametersCameraNode.cpp
--- a/Engine/Plugins/Cameras/GameplayCameras/Source/GameplayCameras/Private/Nodes/Common/LensParametersCameraNode.cpp
+++ b/Engine/Plugins/Cameras/GameplayCameras/Source/GameplayCameras/Private/Nodes/Common/LensParametersCameraNode.cpp
@@ -17,6 +17,7 @@ class FLensParametersCameraNodeEvaluator : public FCameraNodeEvaluator
 
 protected:
 
+	// FCameraNodeEvaluator interface.
 	virtual void OnInitialize(const FCameraNodeEvaluatorInitializeParams& Params, FCameraNodeEvaluationResult& OutResult) override;
 	virtual void OnRun(const FCameraNodeEvaluationParams& Params, FCameraNodeEvaluationResult& OutResult) override;
 
@@ -32,6 +33,8 @@ UE_DEFINE_CAMERA_NODE_EVALUATOR(FLensParametersCameraNodeEvaluator)
 
 void FLensParametersCameraNodeEvaluator::OnInitialize(const FCameraNodeEvaluatorInitializeParams& Params, FCameraNodeEvaluationResult& OutResult)
 {
+	SetNodeEvaluatorFlags(ECameraNodeEvaluatorFlags::None);
+
 	const ULensParametersCameraNode* LensParametersNode = GetCameraNodeAs<ULensParametersCameraNode>();
 	FocalLengthReader.Initialize(LensParametersNode->FocalLength);
 	FocusDistanceReader.Initialize(LensParametersNode->FocusDistance);
diff --git a/Engine/Plugins/Cameras/GameplayCameras/Source/GameplayCameras/Private/Nodes/Common/OffsetCameraNode.cpp b/Engine/Plugins/Cameras/GameplayCameras/Source/GameplayCameras/Private/Nodes/Common/OffsetCameraNode.cpp
--- a/Engine/Plugins/Cameras/GameplayCameras/Source/GameplayCameras/Private/Nodes/Common/OffsetCameraNode.cpp
+++ b/Engine/Plugins/Cameras/GameplayCameras/Source/GameplayCameras/Private/Nodes/Common/OffsetCameraNode.cpp
@@ -16,6 +16,7 @@ class FOffsetCameraNodeEvaluator : public FCameraNodeEvaluator
 
 protected:
 
+	// FCameraNodeEvaluator interface.
 	virtual void OnInitialize(const FCameraNodeEvaluatorInitializeParams& Params, FCameraNodeEvaluationResult& OutResult) override;
 	virtual void OnRun(const FCameraNodeEvaluationParams& Params, FCameraNodeEvaluationResult& OutResult) override;
 
@@ -29,6 +30,8 @@ UE_DEFINE_CAMERA_NODE_EVALUATOR(FOffsetCameraNodeEvaluator)
 
 void FOffsetCameraNodeEvaluator::OnInitialize(const FCameraNodeEvaluatorInitializeParams& Params, FCameraNodeEvaluationResult& OutResult)
 {
+	SetNodeEvaluatorFlags(ECameraNodeEvaluatorFlags::None);
+
 	const UOffsetCameraNode* OffsetNode = GetCameraNodeAs<UOffsetCameraNode>();
 	TranslationReader.Initialize(OffsetNode->TranslationOffset);
 	RotationReader.Initialize(OffsetNode->RotationOffset);
diff --git a/Engine/Plugins/Cameras/GameplayCameras/Source/GameplayCameras/Private/Nodes/Common/OrthographicCameraMode.cpp b/Engine/Plugins/Cameras/GameplayCameras/Source/GameplayCameras/Private/Nodes/Common/OrthographicCameraMode.cpp
--- a/Engine/Plugins/Cameras/GameplayCameras/Source/GameplayCameras/Private/Nodes/Common/OrthographicCameraMode.cpp
+++ b/Engine/Plugins/Cameras/GameplayCameras/Source/GameplayCameras/Private/Nodes/Common/OrthographicCameraMode.cpp
@@ -18,6 +18,7 @@ class FOrthographicCameraNodeEvaluator : public FCameraNodeEvaluator
 
 protected:
 
+	// FCameraNodeEvaluator interface.
 	virtual void OnInitialize(const FCameraNodeEvaluatorInitializeParams& Params, FCameraNodeEvaluationResult& OutResult) override;
 	virtual void OnRun(const FCameraNodeEvaluationParams& Params, FCameraNodeEvaluationResult& OutResult) override;
 
@@ -31,6 +32,8 @@ UE_DEFINE_CAMERA_NODE_EVALUATOR(FOrthographicCameraNodeEvaluator)
 
 void FOrthographicCameraNodeEvaluator::OnInitialize(const FCameraNodeEvaluatorInitializeParams& Params, FCameraNodeEvaluationResult& OutResult)
 {
+	SetNodeEvaluatorFlags(ECameraNodeEvaluatorFlags::None);
+
 	const UOrthographicCameraNode* OrthographicNode = GetCameraNodeAs<UOrthographicCameraNode>();
 	EnableOrthographicModeReader.Initialize(OrthographicNode->EnableOrthographicMode);
 	OrthographicWidthReader.Initialize(OrthographicNode->OrthographicWidth);
diff --git a/Engine/Plugins/Cameras/GameplayCameras/Source/GameplayCameras/Private/Nodes/Common/PostProcessCameraNode.cpp b/Engine/Plugins/Cameras/GameplayCameras/Source/GameplayCameras/Private/Nodes/Common/PostProcessCameraNode.cpp
--- a/Engine/Plugins/Cameras/GameplayCameras/Source/GameplayCameras/Private/Nodes/Common/PostProcessCameraNode.cpp
+++ b/Engine/Plugins/Cameras/GameplayCameras/Source/GameplayCameras/Private/Nodes/Common/PostProcessCameraNode.cpp
@@ -13,8 +13,16 @@ class FPostProcessCameraNodeEvaluator : public FCameraNodeEvaluator
 {
 	UE_DECLARE_CAMERA_NODE_EVALUATOR(GAMEPLAYCAMERAS_API, FPostProcessCameraNodeEvaluator)
 
+public:
+
+	FPostProcessCameraNodeEvaluator()
+	{
+		SetNodeEvaluatorFlags(ECameraNodeEvaluatorFlags::None);
+	}
+
 protected:
 
+	// FCameraNodeEvaluator interface.
 	virtual void OnRun(const FCameraNodeEvaluationParams& Params, FCameraNodeEvaluationResult& OutResult) override;
 };
 
diff --git a/Engine/Plugins/Cameras/GameplayCameras/Source/GameplayCameras/Private/Nodes/Common/SetLocationCameraNode.cpp b/Engine/Plugins/Cameras/GameplayCameras/Source/GameplayCameras/Private/Nodes/Common/SetLocationCameraNode.cpp
--- a/Engine/Plugins/Cameras/GameplayCameras/Source/GameplayCameras/Private/Nodes/Common/SetLocationCameraNode.cpp
+++ b/Engine/Plugins/Cameras/GameplayCameras/Source/GameplayCameras/Private/Nodes/Common/SetLocationCameraNode.cpp
@@ -17,6 +17,7 @@ class FSetLocationCameraNodeEvaluator : public FCameraNodeEvaluator
 
 protected:
 
+	// FCameraNodeEvaluator interface.
 	virtual void OnInitialize(const FCameraNodeEvaluatorInitializeParams& Params, FCameraNodeEvaluationResult& OutResult) override;
 	virtual void OnRun(const FCameraNodeEvaluationParams& Params, FCameraNodeEvaluationResult& OutResult) override;
 
@@ -29,6 +30,8 @@ UE_DEFINE_CAMERA_NODE_EVALUATOR(FSetLocationCameraNodeEvaluator)
 
 void FSetLocationCameraNodeEvaluator::OnInitialize(const FCameraNodeEvaluatorInitializeParams& Params, FCameraNodeEvaluationResult& OutResult)
 {
+	SetNodeEvaluatorFlags(ECameraNodeEvaluatorFlags::None);
+
 	const USetLocationCameraNode* LocationNode = GetCameraNodeAs<USetLocationCameraNode>();
 	LocationReader.Initialize(LocationNode->Location);
 }
diff --git a/Engine/Plugins/Cameras/GameplayCameras/Source/GameplayCameras/Private/Nodes/Common/SetRotationCameraNode.cpp b/Engine/Plugins/Cameras/GameplayCameras/Source/GameplayCameras/Private/Nodes/Common/SetRotationCameraNode.cpp
--- a/Engine/Plugins/Cameras/GameplayCameras/Source/GameplayCameras/Private/Nodes/Common/SetRotationCameraNode.cpp
+++ b/Engine/Plugins/Cameras/GameplayCameras/Source/GameplayCameras/Private/Nodes/Common/SetRotationCameraNode.cpp
@@ -17,6 +17,7 @@ class FSetRotationCameraNodeEvaluator : public FCameraNodeEvaluator
 
 protected:
 
+	// FCameraNodeEvaluator interface.
 	virtual void OnInitialize(const FCameraNodeEvaluatorInitializeParams& Params, FCameraNodeEvaluationResult& OutResult) override;
 	virtual void OnRun(const FCameraNodeEvaluationParams& Params, FCameraNodeEvaluationResult& OutResult) override;
 
@@ -29,6 +30,8 @@ UE_DEFINE_CAMERA_NODE_EVALUATOR(FSetRotationCameraNodeEvaluator)
 
 void FSetRotationCameraNodeEvaluator::OnInitialize(const FCameraNodeEvaluatorInitializeParams& Params, FCameraNodeEvaluationResult& OutResult)
 {
+	SetNodeEvaluatorFlags(ECameraNodeEvaluatorFlags::None);
+
 	const USetRotationCameraNode* RotationNode = GetCameraNodeAs<USetRotationCameraNode>();
 	RotationReader.Initialize(RotationNode->Rotation);
 }
diff --git a/Engine/Plugins/Cameras/GameplayCameras/Source/GameplayCameras/Private/Nodes/Common/SplineFieldOfViewCameraNode.cpp b/Engine/Plugins/Cameras/GameplayCameras/Source/GameplayCameras/Private/Nodes/Common/SplineFieldOfViewCameraNode.cpp
--- a/Engine/Plugins/Cameras/GameplayCameras/Source/GameplayCameras/Private/Nodes/Common/SplineFieldOfViewCameraNode.cpp
+++ b/Engine/Plugins/Cameras/GameplayCameras/Source/GameplayCameras/Private/Nodes/Common/SplineFieldOfViewCameraNode.cpp
@@ -28,6 +28,7 @@ public:
 
 protected:
 
+	// FCameraNodeEvaluator interface.
 	virtual void OnInitialize(const FCameraNodeEvaluatorInitializeParams& Params, FCameraNodeEvaluationResult& OutResult) override;
 	virtual void OnRun(const FCameraNodeEvaluationParams& Params, FCameraNodeEvaluationResult& OutResult) override;
 
@@ -48,6 +49,8 @@ UE_DEFINE_CAMERA_NODE_EVALUATOR(FSplineFieldOfViewCameraNodeEvaluator)
 
 void FSplineFieldOfViewCameraNodeEvaluator::OnInitialize(const FCameraNodeEvaluatorInitializeParams& Params, FCameraNodeEvaluationResult& OutResult)
 {
+	SetNodeEvaluatorFlags(ECameraNodeEvaluatorFlags::None);
+
 	const USplineFieldOfViewCameraNode* FieldOfViewNode = GetCameraNodeAs<USplineFieldOfViewCameraNode>();
 	SplineInputReader.Initialize(FieldOfViewNode->SplineInput);
 
diff --git a/Engine/Plugins/Cameras/GameplayCameras/Source/GameplayCameras/Private/Nodes/Common/SplineOffsetCameraNode.cpp b/Engine/Plugins/Cameras/GameplayCameras/Source/GameplayCameras/Private/Nodes/Common/SplineOffsetCameraNode.cpp
--- a/Engine/Plugins/Cameras/GameplayCameras/Source/GameplayCameras/Private/Nodes/Common/SplineOffsetCameraNode.cpp
+++ b/Engine/Plugins/Cameras/GameplayCameras/Source/GameplayCameras/Private/Nodes/Common/SplineOffsetCameraNode.cpp
@@ -27,6 +27,7 @@ public:
 
 protected:
 
+	// FCameraNodeEvaluator interface.
 	virtual void OnInitialize(const FCameraNodeEvaluatorInitializeParams& Params, FCameraNodeEvaluationResult& OutResult) override;
 	virtual void OnRun(const FCameraNodeEvaluationParams& Params, FCameraNodeEvaluationResult& OutResult) override;
 
@@ -53,6 +54,8 @@ UE_DEFINE_CAMERA_NODE_EVALUATOR(FSplineOffsetCameraNodeEvaluator)
 
 void FSplineOffsetCameraNodeEvaluator::OnInitialize(const FCameraNodeEvaluatorInitializeParams& Params, FCameraNodeEvaluationResult& OutResult)
 {
+	SetNodeEvaluatorFlags(ECameraNodeEvaluatorFlags::None);
+
 	const USplineOffsetCameraNode* SplineOffsetNode = GetCameraNodeAs<USplineOffsetCameraNode>();
 	SplineInputReader.Initialize(SplineOffsetNode->SplineInput);
 
diff --git a/Engine/Plugins/Cameras/GameplayCameras/Source/GameplayCameras/Private/Nodes/Common/SplineOrbitCameraNode.cpp b/Engine/Plugins/Cameras/GameplayCameras/Source/GameplayCameras/Private/Nodes/Common/SplineOrbitCameraNode.cpp
--- a/Engine/Plugins/Cameras/GameplayCameras/Source/GameplayCameras/Private/Nodes/Common/SplineOrbitCameraNode.cpp
+++ b/Engine/Plugins/Cameras/GameplayCameras/Source/GameplayCameras/Private/Nodes/Common/SplineOrbitCameraNode.cpp
@@ -50,6 +50,7 @@ public:
 
 protected:
 
+	// FCameraNodeEvaluator interface.
 	virtual void OnBuild(const FCameraNodeEvaluatorBuildParams& Params) override;
 	virtual void OnInitialize(const FCameraNodeEvaluatorInitializeParams& Params, FCameraNodeEvaluationResult& OutResult) override;
 	virtual FCameraNodeEvaluatorChildrenView OnGetChildren() override;
@@ -123,7 +124,7 @@ void FSplineOrbitCameraNodeEvaluator::OnBuild(const FCameraNodeEvaluatorBuildPar
 
 void FSplineOrbitCameraNodeEvaluator::OnInitialize(const FCameraNodeEvaluatorInitializeParams& Params, FCameraNodeEvaluationResult& OutResult)
 {
-	AddNodeEvaluatorFlags(ECameraNodeEvaluatorFlags::SupportsOperations);
+	SetNodeEvaluatorFlags(ECameraNodeEvaluatorFlags::SupportsOperations);
 
 	RebuildCurves();
 
diff --git a/Engine/Plugins/Cameras/GameplayCameras/Source/GameplayCameras/Private/Nodes/Common/TargetRayCastCameraNode.cpp b/Engine/Plugins/Cameras/GameplayCameras/Source/GameplayCameras/Private/Nodes/Common/TargetRayCastCameraNode.cpp
--- a/Engine/Plugins/Cameras/GameplayCameras/Source/GameplayCameras/Private/Nodes/Common/TargetRayCastCameraNode.cpp
+++ b/Engine/Plugins/Cameras/GameplayCameras/Source/GameplayCameras/Private/Nodes/Common/TargetRayCastCameraNode.cpp
@@ -69,6 +69,8 @@ UE_DEFINE_CAMERA_DEBUG_BLOCK_WITH_FIELDS(FTargetRayCastCameraDebugBlock)
 
 void FTargetRayCastCameraNodeEvaluator::OnInitialize(const FCameraNodeEvaluatorInitializeParams& Params, FCameraNodeEvaluationResult& OutResult)
 {
+	SetNodeEvaluatorFlags(ECameraNodeEvaluatorFlags::None);
+
 	const UTargetRayCastCameraNode* RayCastNode = GetCameraNodeAs<UTargetRayCastCameraNode>();
 	AutoFocusReader.Initialize(RayCastNode->AutoFocus);
 }
diff --git a/Engine/Plugins/Cameras/GameplayCameras/Source/GameplayCameras/Private/Nodes/Framing/BaseFramingCameraNode.cpp b/Engine/Plugins/Cameras/GameplayCameras/Source/GameplayCameras/Private/Nodes/Framing/BaseFramingCameraNode.cpp
--- a/Engine/Plugins/Cameras/GameplayCameras/Source/GameplayCameras/Private/Nodes/Framing/BaseFramingCameraNode.cpp
+++ b/Engine/Plugins/Cameras/GameplayCameras/Source/GameplayCameras/Private/Nodes/Framing/BaseFramingCameraNode.cpp
@@ -60,6 +60,8 @@ UE_DEFINE_CAMERA_NODE_EVALUATOR(FBaseFramingCameraNodeEvaluator)
 
 void FBaseFramingCameraNodeEvaluator::OnInitialize(const FCameraNodeEvaluatorInitializeParams& Params, FCameraNodeEvaluationResult& OutResult)
 {
+	SetNodeEvaluatorFlags(ECameraNodeEvaluatorFlags::None);
+
 	const UBaseFramingCameraNode* BaseFramingNode = GetCameraNodeAs<UBaseFramingCameraNode>();
 
 	Readers.TargetInfos.Initialize(BaseFramingNode->TargetInfos, BaseFramingNode->TargetInfosDataID);
diff --git a/Engine/Plugins/Cameras/GameplayCameras/Source/GameplayCameras/Private/Nodes/Input/AutoRotateInput2DCameraNode.cpp b/Engine/Plugins/Cameras/GameplayCameras/Source/GameplayCameras/Private/Nodes/Input/AutoRotateInput2DCameraNode.cpp
--- a/Engine/Plugins/Cameras/GameplayCameras/Source/GameplayCameras/Private/Nodes/Input/AutoRotateInput2DCameraNode.cpp
+++ b/Engine/Plugins/Cameras/GameplayCameras/Source/GameplayCameras/Private/Nodes/Input/AutoRotateInput2DCameraNode.cpp
@@ -36,6 +36,7 @@ public:
 
 protected:
 
+	// FCameraNodeEvaluator interface.
 	virtual void OnBuild(const FCameraNodeEvaluatorBuildParams& Params) override;
 	virtual FCameraNodeEvaluatorChildrenView OnGetChildren() override;
 	virtual void OnInitialize(const FCameraNodeEvaluatorInitializeParams& Params, FCameraNodeEvaluationResult& OutResult) override;
@@ -101,6 +102,8 @@ FCameraNodeEvaluatorChildrenView FAutoRotateInput2DCameraNodeEvaluator::OnGetChi
 
 void FAutoRotateInput2DCameraNodeEvaluator::OnInitialize(const FCameraNodeEvaluatorInitializeParams& Params, FCameraNodeEvaluationResult& OutResult)
 {
+	SetNodeEvaluatorFlags(ECameraNodeEvaluatorFlags::None);
+
 	const FCameraNodeEvaluationResult& InitialResult = Params.EvaluationContext->GetInitialResult();
 	LastContextLocation = InitialResult.CameraPose.GetLocation();
 
diff --git a/Engine/Plugins/Cameras/GameplayCameras/Source/GameplayCameras/Private/Nodes/Input/CameraRigInput1DSlot.cpp b/Engine/Plugins/Cameras/GameplayCameras/Source/GameplayCameras/Private/Nodes/Input/CameraRigInput1DSlot.cpp
--- a/Engine/Plugins/Cameras/GameplayCameras/Source/GameplayCameras/Private/Nodes/Input/CameraRigInput1DSlot.cpp
+++ b/Engine/Plugins/Cameras/GameplayCameras/Source/GameplayCameras/Private/Nodes/Input/CameraRigInput1DSlot.cpp
@@ -19,7 +19,10 @@ UE_DEFINE_CAMERA_NODE_EVALUATOR(FCameraRigInput1DSlotEvaluator)
 
 FCameraRigInput1DSlotEvaluator::FCameraRigInput1DSlotEvaluator()
 {
-	AddNodeEvaluatorFlags(ECameraNodeEvaluatorFlags::NeedsParameterUpdate | ECameraNodeEvaluatorFlags::SupportsOperations);
+	SetNodeEvaluatorFlags(
+			ECameraNodeEvaluatorFlags::NeedsParameterUpdate | 
+			ECameraNodeEvaluatorFlags::NeedsSerialize | 
+			ECameraNodeEvaluatorFlags::SupportsOperations);
 }
 
 void FCameraRigInput1DSlotEvaluator::OnInitialize(const FCameraNodeEvaluatorInitializeParams& Params, FCameraNodeEvaluationResult& OutResult)
diff --git a/Engine/Plugins/Cameras/GameplayCameras/Source/GameplayCameras/Private/Nodes/Input/CameraRigInput2DSlot.cpp b/Engine/Plugins/Cameras/GameplayCameras/Source/GameplayCameras/Private/Nodes/Input/CameraRigInput2DSlot.cpp
--- a/Engine/Plugins/Cameras/GameplayCameras/Source/GameplayCameras/Private/Nodes/Input/CameraRigInput2DSlot.cpp
+++ b/Engine/Plugins/Cameras/GameplayCameras/Source/GameplayCameras/Private/Nodes/Input/CameraRigInput2DSlot.cpp
@@ -19,7 +19,10 @@ UE_DEFINE_CAMERA_NODE_EVALUATOR(FCameraRigInput2DSlotEvaluator)
 
 FCameraRigInput2DSlotEvaluator::FCameraRigInput2DSlotEvaluator()
 {
-	AddNodeEvaluatorFlags(ECameraNodeEvaluatorFlags::NeedsParameterUpdate | ECameraNodeEvaluatorFlags::SupportsOperations);
+	SetNodeEvaluatorFlags(
+			ECameraNodeEvaluatorFlags::NeedsParameterUpdate | 
+			ECameraNodeEvaluatorFlags::NeedsSerialize | 
+			ECameraNodeEvaluatorFlags::SupportsOperations);
 }
 
 void FCameraRigInput2DSlotEvaluator::OnInitialize(const FCameraNodeEvaluatorInitializeParams& Params, FCameraNodeEvaluationResult& OutResult)
diff --git a/Engine/Plugins/Cameras/GameplayCameras/Source/GameplayCameras/Private/Nodes/Input/InputAxisBinding2DCameraNode.cpp b/Engine/Plugins/Cameras/GameplayCameras/Source/GameplayCameras/Private/Nodes/Input/InputAxisBinding2DCameraNode.cpp
--- a/Engine/Plugins/Cameras/GameplayCameras/Source/GameplayCameras/Private/Nodes/Input/InputAxisBinding2DCameraNode.cpp
+++ b/Engine/Plugins/Cameras/GameplayCameras/Source/GameplayCameras/Private/Nodes/Input/InputAxisBinding2DCameraNode.cpp
@@ -21,6 +21,7 @@ class FInputAxisBinding2DCameraNodeEvaluator : public FCameraRigInput2DSlotEvalu
 
 protected:
 
+	// FCameraNodeEvaluator interface.
 	virtual void OnInitialize(const FCameraNodeEvaluatorInitializeParams& Params, FCameraNodeEvaluationResult& OutResult) override;
 	virtual void OnUpdateParameters(const FCameraBlendedParameterUpdateParams& Params, FCameraBlendedParameterUpdateResult& OutResult) override;
 
diff --git a/Engine/Plugins/Cameras/GameplayCameras/Source/GameplayCameras/Private/Nodes/Shakes/CameraShakeCameraNode.cpp b/Engine/Plugins/Cameras/GameplayCameras/Source/GameplayCameras/Private/Nodes/Shakes/CameraShakeCameraNode.cpp
--- a/Engine/Plugins/Cameras/GameplayCameras/Source/GameplayCameras/Private/Nodes/Shakes/CameraShakeCameraNode.cpp
+++ b/Engine/Plugins/Cameras/GameplayCameras/Source/GameplayCameras/Private/Nodes/Shakes/CameraShakeCameraNode.cpp
@@ -52,7 +52,7 @@ UE_DEFINE_CAMERA_NODE_EVALUATOR(FCameraShakeCameraNodeEvaluator)
 
 FCameraShakeCameraNodeEvaluator::FCameraShakeCameraNodeEvaluator()
 {
-	AddNodeEvaluatorFlags(ECameraNodeEvaluatorFlags::NeedsParameterUpdate);
+	SetNodeEvaluatorFlags(ECameraNodeEvaluatorFlags::NeedsParameterUpdate);
 }
 
 FCameraNodeEvaluatorChildrenView FCameraShakeCameraNodeEvaluator::OnGetChildren()
@@ -104,7 +104,7 @@ void FCameraShakeCameraNodeEvaluator::OnRun(const FCameraNodeEvaluationParams& P
 		CameraShakeRootEvaluator->Run(Params, OutResult);
 	}
 	// If evaluating the shake later in the visual layer, keep requesting that the shake service
-	// maintains this shake alive. We already put the shake's parameters inside the vairable table so
+	// maintains this shake alive. We already put the shake's parameters inside the variable table so
 	// its values will blend with whoever else wants this shake, and it will eventually run with the
 	// blended values.
 	else if (CameraShakeService)
diff --git a/Engine/Plugins/Cameras/GameplayCameras/Source/GameplayCameras/Private/Nodes/Shakes/CompositeShakeCameraNode.cpp b/Engine/Plugins/Cameras/GameplayCameras/Source/GameplayCameras/Private/Nodes/Shakes/CompositeShakeCameraNode.cpp
--- a/Engine/Plugins/Cameras/GameplayCameras/Source/GameplayCameras/Private/Nodes/Shakes/CompositeShakeCameraNode.cpp
+++ b/Engine/Plugins/Cameras/GameplayCameras/Source/GameplayCameras/Private/Nodes/Shakes/CompositeShakeCameraNode.cpp
@@ -13,6 +13,11 @@ class FCompositeShakeCameraNodeEvaluator : public FShakeCameraNodeEvaluator
 
 public:
 
+	FCompositeShakeCameraNodeEvaluator()
+	{
+		SetNodeEvaluatorFlags(ECameraNodeEvaluatorFlags::None);
+	}
+
 protected:
 
 	// FShakeCameraNodeEvaluator interface.
diff --git a/Engine/Plugins/Cameras/GameplayCameras/Source/GameplayCameras/Private/Nodes/Shakes/EnvelopeShakeCameraNode.cpp b/Engine/Plugins/Cameras/GameplayCameras/Source/GameplayCameras/Private/Nodes/Shakes/EnvelopeShakeCameraNode.cpp
--- a/Engine/Plugins/Cameras/GameplayCameras/Source/GameplayCameras/Private/Nodes/Shakes/EnvelopeShakeCameraNode.cpp
+++ b/Engine/Plugins/Cameras/GameplayCameras/Source/GameplayCameras/Private/Nodes/Shakes/EnvelopeShakeCameraNode.cpp
@@ -13,8 +13,6 @@ class FEnvelopeShakeCameraNodeEvaluator : public FShakeCameraNodeEvaluator
 {
 	UE_DECLARE_SHAKE_CAMERA_NODE_EVALUATOR(GAMEPLAYCAMERAS_API, FEnvelopeShakeCameraNodeEvaluator)
 
-public:
-
 protected:
 
 	// FShakeCameraNodeEvaluator interface.
@@ -52,6 +50,8 @@ FCameraNodeEvaluatorChildrenView FEnvelopeShakeCameraNodeEvaluator::OnGetChildre
 
 void FEnvelopeShakeCameraNodeEvaluator::OnInitialize(const FCameraNodeEvaluatorInitializeParams& Params, FCameraNodeEvaluationResult& OutResult)
 {
+	SetNodeEvaluatorFlags(ECameraNodeEvaluatorFlags::NeedsSerialize);
+
 	const UEnvelopeShakeCameraNode* EnvelopeNode = GetCameraNodeAs<UEnvelopeShakeCameraNode>();
 
 	EaseInTime = EnvelopeNode->EaseInTime.GetValue(OutResult.VariableTable);
diff --git a/Engine/Plugins/Cameras/GameplayCameras/Source/GameplayCameras/Private/Nodes/Shakes/PerlinNoiseLocationShakeCameraNode.cpp b/Engine/Plugins/Cameras/GameplayCameras/Source/GameplayCameras/Private/Nodes/Shakes/PerlinNoiseLocationShakeCameraNode.cpp
--- a/Engine/Plugins/Cameras/GameplayCameras/Source/GameplayCameras/Private/Nodes/Shakes/PerlinNoiseLocationShakeCameraNode.cpp
+++ b/Engine/Plugins/Cameras/GameplayCameras/Source/GameplayCameras/Private/Nodes/Shakes/PerlinNoiseLocationShakeCameraNode.cpp
@@ -59,6 +59,8 @@ UE_DEFINE_SHAKE_CAMERA_NODE_EVALUATOR(FPerlinNoiseLocationShakeCameraNodeEvaluat
 
 void FPerlinNoiseLocationShakeCameraNodeEvaluator::OnInitialize(const FCameraNodeEvaluatorInitializeParams& Params, FCameraNodeEvaluationResult& OutResult)
 {
+	SetNodeEvaluatorFlags(ECameraNodeEvaluatorFlags::NeedsSerialize);
+
 	const UPerlinNoiseLocationShakeCameraNode* ShakeNode = GetCameraNodeAs<UPerlinNoiseLocationShakeCameraNode>();
 
 	AmplitudeMultiplier.Initialize(ShakeNode->AmplitudeMultiplier);
diff --git a/Engine/Plugins/Cameras/GameplayCameras/Source/GameplayCameras/Private/Nodes/Shakes/PerlinNoiseRotationShakeCameraNode.cpp b/Engine/Plugins/Cameras/GameplayCameras/Source/GameplayCameras/Private/Nodes/Shakes/PerlinNoiseRotationShakeCameraNode.cpp
--- a/Engine/Plugins/Cameras/GameplayCameras/Source/GameplayCameras/Private/Nodes/Shakes/PerlinNoiseRotationShakeCameraNode.cpp
+++ b/Engine/Plugins/Cameras/GameplayCameras/Source/GameplayCameras/Private/Nodes/Shakes/PerlinNoiseRotationShakeCameraNode.cpp
@@ -59,6 +59,8 @@ UE_DEFINE_SHAKE_CAMERA_NODE_EVALUATOR(FPerlinNoiseRotationShakeCameraNodeEvaluat
 
 void FPerlinNoiseRotationShakeCameraNodeEvaluator::OnInitialize(const FCameraNodeEvaluatorInitializeParams& Params, FCameraNodeEvaluationResult& OutResult)
 {
+	SetNodeEvaluatorFlags(ECameraNodeEvaluatorFlags::NeedsSerialize);
+
 	const UPerlinNoiseRotationShakeCameraNode* ShakeNode = GetCameraNodeAs<UPerlinNoiseRotationShakeCameraNode>();
 
 	AmplitudeMultiplier.Initialize(ShakeNode->AmplitudeMultiplier);
diff --git a/Engine/Plugins/Cameras/GameplayCameras/Source/GameplayCameras/Private/Nodes/Utility/BlueprintCameraNode.cpp b/Engine/Plugins/Cameras/GameplayCameras/Source/GameplayCameras/Private/Nodes/Utility/BlueprintCameraNode.cpp
--- a/Engine/Plugins/Cameras/GameplayCameras/Source/GameplayCameras/Private/Nodes/Utility/BlueprintCameraNode.cpp
+++ b/Engine/Plugins/Cameras/GameplayCameras/Source/GameplayCameras/Private/Nodes/Utility/BlueprintCameraNode.cpp
@@ -68,6 +68,8 @@ UE_DEFINE_CAMERA_DEBUG_BLOCK_WITH_FIELDS(FBlueprintCameraDebugBlock)
 
 void FBlueprintCameraNodeEvaluator::OnInitialize(const FCameraNodeEvaluatorInitializeParams& Params, FCameraNodeEvaluationResult& OutResult)
 {
+	SetNodeEvaluatorFlags(ECameraNodeEvaluatorFlags::None);
+
 	const UBlueprintCameraNode* BlueprintNode = GetCameraNodeAs<UBlueprintCameraNode>();
 	if (!ensure(BlueprintNode))
 	{
diff --git a/Engine/Plugins/Cameras/GameplayCameras/Source/GameplayCameras/Private/Services/CameraModifierService.cpp b/Engine/Plugins/Cameras/GameplayCameras/Source/GameplayCameras/Private/Services/CameraModifierService.cpp
--- a/Engine/Plugins/Cameras/GameplayCameras/Source/GameplayCameras/Private/Services/CameraModifierService.cpp
+++ b/Engine/Plugins/Cameras/GameplayCameras/Source/GameplayCameras/Private/Services/CameraModifierService.cpp
@@ -14,6 +14,8 @@ UE_DEFINE_CAMERA_EVALUATION_SERVICE(FCameraModifierService)
 
 void FCameraModifierService::OnInitialize(const FCameraEvaluationServiceInitializeParams& Params)
 {
+	SetEvaluationServiceFlags(ECameraEvaluationServiceFlags::None);
+
 	ensure(Evaluator == nullptr);
 	Evaluator = Params.Evaluator;
 }
diff --git a/Engine/Plugins/Cameras/GameplayCameras/Source/GameplayCameras/Private/Services/CameraShakeService.cpp b/Engine/Plugins/Cameras/GameplayCameras/Source/GameplayCameras/Private/Services/CameraShakeService.cpp
--- a/Engine/Plugins/Cameras/GameplayCameras/Source/GameplayCameras/Private/Services/CameraShakeService.cpp
+++ b/Engine/Plugins/Cameras/GameplayCameras/Source/GameplayCameras/Private/Services/CameraShakeService.cpp
@@ -359,6 +359,8 @@ UE_DEFINE_CAMERA_EVALUATION_SERVICE(FCameraShakeService)
 
 void FCameraShakeService::OnInitialize(const FCameraEvaluationServiceInitializeParams& Params)
 {
+	SetEvaluationServiceFlags(ECameraEvaluationServiceFlags::None);
+
 	ensure(Evaluator == nullptr);
 	Evaluator = Params.Evaluator;
 }
diff --git a/Engine/Plugins/Cameras/GameplayCameras/Source/GameplayCameras/Private/Services/PlayerControlRotationService.cpp b/Engine/Plugins/Cameras/GameplayCameras/Source/GameplayCameras/Private/Services/PlayerControlRotationService.cpp
--- a/Engine/Plugins/Cameras/GameplayCameras/Source/GameplayCameras/Private/Services/PlayerControlRotationService.cpp
+++ b/Engine/Plugins/Cameras/GameplayCameras/Source/GameplayCameras/Private/Services/PlayerControlRotationService.cpp
@@ -42,6 +42,7 @@ UE_DEFINE_CAMERA_DEBUG_BLOCK_WITH_FIELDS(FPlayerControlRotationDebugBlock)
 FPlayerControlRotationEvaluationService::FPlayerControlRotationEvaluationService()
 	: PreviousAxisBindingValue(EForceInit::ForceInit)
 {
+	SetEvaluationServiceFlags(ECameraEvaluationServiceFlags::NeedsPostUpdate);
 }
 
 FPlayerControlRotationEvaluationService::FPlayerControlRotationEvaluationService(const FPlayerControlRotationParams& InParams)
diff --git a/Engine/Plugins/Cameras/GameplayCameras/Source/GameplayCameras/Public/Core/CameraEvaluationService.h b/Engine/Plugins/Cameras/GameplayCameras/Source/GameplayCameras/Public/Core/CameraEvaluationService.h
--- a/Engine/Plugins/Cameras/GameplayCameras/Source/GameplayCameras/Public/Core/CameraEvaluationService.h
+++ b/Engine/Plugins/Cameras/GameplayCameras/Source/GameplayCameras/Public/Core/CameraEvaluationService.h
@@ -22,7 +22,9 @@ enum class ECameraEvaluationServiceFlags
 	None = 0,
 	NeedsPreUpdate = 1 << 0,
 	NeedsPostUpdate = 1 << 1,
-	NeedsRootCameraNodeEvents = 1 << 2
+	NeedsRootCameraNodeEvents = 1 << 2,
+
+	Default = NeedsPreUpdate | NeedsPostUpdate | NeedsRootCameraNodeEvents
 };
 ENUM_CLASS_FLAGS(ECameraEvaluationServiceFlags);
 
@@ -123,7 +125,10 @@ protected:
 
 protected:
 
-	/** Sets the flags on this service. */
+	/**
+	 * Sets the flags on this service.
+	 * Can be called from the constructor, or during OnInitialize().
+	 */
 	void SetEvaluationServiceFlags(ECameraEvaluationServiceFlags InFlags);
 
 private:
diff --git a/Engine/Plugins/Cameras/GameplayCameras/Source/GameplayCameras/Public/Core/CameraNodeEvaluator.h b/Engine/Plugins/Cameras/GameplayCameras/Source/GameplayCameras/Public/Core/CameraNodeEvaluator.h
--- a/Engine/Plugins/Cameras/GameplayCameras/Source/GameplayCameras/Public/Core/CameraNodeEvaluator.h
+++ b/Engine/Plugins/Cameras/GameplayCameras/Source/GameplayCameras/Public/Core/CameraNodeEvaluator.h
@@ -43,7 +43,7 @@ enum class ECameraNodeEvaluatorFlags
 	NeedsSerialize = 1 << 1,
 	SupportsOperations = 1 << 2,
 
-	Default = NeedsSerialize
+	Default = NeedsParameterUpdate | NeedsSerialize | SupportsOperations
 };
 ENUM_CLASS_FLAGS(ECameraNodeEvaluatorFlags)
 
@@ -329,12 +329,19 @@ public:
 
 protected:
 
-	/** Adds flags for this evaluator. */
-	void AddNodeEvaluatorFlags(ECameraNodeEvaluatorFlags InFlags);
-
-	/** Sets the flags for this evaluator. */
+	/**
+	 * Sets the flags for this evaluator.
+	 * Can be called from the constructor, or during OnInitialize().
+	 */
 	void SetNodeEvaluatorFlags(ECameraNodeEvaluatorFlags InFlags);
 
+	/**
+	 * Adds flags for this evaluator.
+	 * Evaluators default to having all flags enabled, so this is mostly only for re-adding a
+	 * flag that was removed by a base class.
+	 */
+	void AddNodeEvaluatorFlags(ECameraNodeEvaluatorFlags InFlags);
+
 protected:
 
 	/** Called to build any children evaluators. */
diff --git a/Engine/Plugins/Cameras/GameplayCameras/Source/GameplayCameras/Public/Nodes/Blends/InterruptedBlendCameraNode.h b/Engine/Plugins/Cameras/GameplayCameras/Source/GameplayCameras/Public/Nodes/Blends/InterruptedBlendCameraNode.h
--- a/Engine/Plugins/Cameras/GameplayCameras/Source/GameplayCameras/Public/Nodes/Blends/InterruptedBlendCameraNode.h
+++ b/Engine/Plugins/Cameras/GameplayCameras/Source/GameplayCameras/Public/Nodes/Blends/InterruptedBlendCameraNode.h
@@ -21,12 +21,8 @@ class FInterruptedBlendCameraNodeEvaluator : public FBlendCameraNodeEvaluator
 
 public:
 
-	FInterruptedBlendCameraNodeEvaluator()
-	{}
-	FInterruptedBlendCameraNodeEvaluator(FBlendCameraNodeEvaluator* InChildBlend, FBlendCameraNodeEvaluator* InFrozenBlend) 
-		: ChildBlend(InChildBlend)
-		, FrozenBlend(InFrozenBlend)
-	{}
+	FInterruptedBlendCameraNodeEvaluator();
+	FInterruptedBlendCameraNodeEvaluator(FBlendCameraNodeEvaluator* InChildBlend, FBlendCameraNodeEvaluator* InFrozenBlend);
 
 protected:
 
diff --git a/Engine/Plugins/Cameras/GameplayCameras/Source/GameplayCameras/Public/Nodes/Blends/ReverseBlendCameraNode.h b/Engine/Plugins/Cameras/GameplayCameras/Source/GameplayCameras/Public/Nodes/Blends/ReverseBlendCameraNode.h
--- a/Engine/Plugins/Cameras/GameplayCameras/Source/GameplayCameras/Public/Nodes/Blends/ReverseBlendCameraNode.h
+++ b/Engine/Plugins/Cameras/GameplayCameras/Source/GameplayCameras/Public/Nodes/Blends/ReverseBlendCameraNode.h
@@ -20,11 +20,8 @@ class FReverseBlendCameraNodeEvaluator : public FBlendCameraNodeEvaluator
 
 public:
 
-	FReverseBlendCameraNodeEvaluator()
-	{}
-	FReverseBlendCameraNodeEvaluator(FBlendCameraNodeEvaluator* InChildBlend) 
-		: ChildBlend(InChildBlend)
-	{}
+	FReverseBlendCameraNodeEvaluator();
+	FReverseBlendCameraNodeEvaluator(FBlendCameraNodeEvaluator* InChildBlend);
 
 protected:
 
diff --git a/Engine/Plugins/Cameras/GameplayCameras/Source/GameplayCameras/Public/Nodes/Blends/SimpleBlendCameraNode.h b/Engine/Plugins/Cameras/GameplayCameras/Source/GameplayCameras/Public/Nodes/Blends/SimpleBlendCameraNode.h
--- a/Engine/Plugins/Cameras/GameplayCameras/Source/GameplayCameras/Public/Nodes/Blends/SimpleBlendCameraNode.h
+++ b/Engine/Plugins/Cameras/GameplayCameras/Source/GameplayCameras/Public/Nodes/Blends/SimpleBlendCameraNode.h
@@ -25,7 +25,7 @@ class USimpleFixedTimeBlendCameraNode : public USimpleBlendCameraNode
 
 public:
 
-	void SetBlendTime(float BlendTimeIn) { BlendTime = BlendTimeIn; }
+	GAMEPLAYCAMERAS_API void SetBlendTime(float BlendTimeIn) { BlendTime = BlendTimeIn; }
 
 public:
 
@@ -51,6 +51,8 @@ class FSimpleBlendCameraNodeEvaluator : public FBlendCameraNodeEvaluator
 
 public:
 
+	FSimpleBlendCameraNodeEvaluator();
+
 	/** Gets the last evaluated blend factor. */
 	float GetBlendFactor() const { return BlendFactor; }
 
diff --git a/Engine/Plugins/Cameras/GameplayCameras/Source/GameplayCameras/Public/Nodes/Framing/BaseFramingCameraNode.h b/Engine/Plugins/Cameras/GameplayCameras/Source/GameplayCameras/Public/Nodes/Framing/BaseFramingCameraNode.h
--- a/Engine/Plugins/Cameras/GameplayCameras/Source/GameplayCameras/Public/Nodes/Framing/BaseFramingCameraNode.h
+++ b/Engine/Plugins/Cameras/GameplayCameras/Source/GameplayCameras/Public/Nodes/Framing/BaseFramingCameraNode.h
@@ -159,6 +159,7 @@ class FBaseFramingCameraNodeEvaluator : public FCameraNodeEvaluator
 
 protected:
 
+	// FCameraNodeEvaluator interface.
 	virtual void OnInitialize(const FCameraNodeEvaluatorInitializeParams& Params, FCameraNodeEvaluationResult& OutResult) override;
 
 #if WITH_EDITOR
