Updated code using script to make sure x_api is set on functions/globalvars instead of on types

diff --git a/Engine/Plugins/Cameras/GameplayCameras/Source/GameplayCameras/Public/Nodes/Blends/SimpleBlendCameraNode.h b/Engine/Plugins/Cameras/GameplayCameras/Source/GameplayCameras/Public/Nodes/Blends/SimpleBlendCameraNode.h
--- a/Engine/Plugins/Cameras/GameplayCameras/Source/GameplayCameras/Public/Nodes/Blends/SimpleBlendCameraNode.h
+++ b/Engine/Plugins/Cameras/GameplayCameras/Source/GameplayCameras/Public/Nodes/Blends/SimpleBlendCameraNode.h
@@ -25,7 +25,7 @@ class USimpleFixedTimeBlendCameraNode : public USimpleBlendCameraNode
 
 public:
 
-	GAMEPLAYCAMERAS_API void SetBlendTime(float BlendTimeIn) { BlendTime = BlendTimeIn; }
+	void SetBlendTime(float BlendTimeIn) { BlendTime = BlendTimeIn; }
 
 public:
 
diff --git a/Engine/Source/Runtime/HeadMountedDisplay/Public/IXRTrackingSystem.h b/Engine/Source/Runtime/HeadMountedDisplay/Public/IXRTrackingSystem.h
--- a/Engine/Source/Runtime/HeadMountedDisplay/Public/IXRTrackingSystem.h
+++ b/Engine/Source/Runtime/HeadMountedDisplay/Public/IXRTrackingSystem.h
@@ -67,7 +67,7 @@ public:
 	 * Other devices can have arbitrary ids defined by each system.
 	 * If a tracking system does not support tracking HMDs, device ID zero should be treated as invalid.
 	 */
-	static HEADMOUNTEDDISPLAY_API const int32 HMDDeviceId = 0;
+	static constexpr int32 HMDDeviceId = 0;
 
 	/**
 	 * Whether or not the system supports positional tracking (either via sensor or other means)
