Cameras: normalize yaw/pitch for spline orbit node

When the pitch value comes from the player control rotation, it's not normalized and that creates problems when evaluating the curves.

diff --git a/Engine/Plugins/Cameras/GameplayCameras/Source/GameplayCameras/Private/Nodes/Common/SplineOrbitCameraNode.cpp b/Engine/Plugins/Cameras/GameplayCameras/Source/GameplayCameras/Private/Nodes/Common/SplineOrbitCameraNode.cpp
--- a/Engine/Plugins/Cameras/GameplayCameras/Source/GameplayCameras/Private/Nodes/Common/SplineOrbitCameraNode.cpp
+++ b/Engine/Plugins/Cameras/GameplayCameras/Source/GameplayCameras/Private/Nodes/Common/SplineOrbitCameraNode.cpp
@@ -190,6 +190,7 @@ void FSplineOrbitCameraNodeEvaluator::OnRun(const FCameraNodeEvaluationParams& P
 		const FRotator3d ControlRotation = PlayerController->GetControlRotation();
 		OrbitRotation = ControlRotation;
 	}
+	OrbitRotation.Normalize();
 
 	const USplineOrbitCameraNode* SplineOrbitNode = GetCameraNodeAs<USplineOrbitCameraNode>();
 
