Cameras: fix crash when connecting a camera variable reference to a camera rig parameter

diff --git a/Engine/Plugins/Cameras/GameplayCameras/Source/GameplayCameras/Private/Build/CameraObjectInterfaceParameterBuilder.cpp b/Engine/Plugins/Cameras/GameplayCameras/Source/GameplayCameras/Private/Build/CameraObjectInterfaceParameterBuilder.cpp
--- a/Engine/Plugins/Cameras/GameplayCameras/Source/GameplayCameras/Private/Build/CameraObjectInterfaceParameterBuilder.cpp
+++ b/Engine/Plugins/Cameras/GameplayCameras/Source/GameplayCameras/Private/Build/CameraObjectInterfaceParameterBuilder.cpp
@@ -6,6 +6,7 @@
 #include "Core/CameraNode.h"
 #include "Core/CameraObjectInterfaceParameterDefinition.h"
 #include "Core/CameraParameters.h"
+#include "Core/CameraVariableReferences.h"
 #include "Core/ICustomCameraNodeParameterProvider.h"
 #include "StructUtils/PropertyBag.h"
 #include "UObject/UnrealType.h"
@@ -238,11 +239,17 @@ void FCameraObjectInterfaceParameterBuilder::SetDefaultParameterValues(const UBa
 					case ECameraVariableType::ValueName:\
 						{\
 							using CameraParameterType = F##ValueName##CameraParameter;\
-							if (ensure(StructProperty->Struct == CameraParameterType::StaticStruct()))\
+							using CameraVariableReferenceType = F##ValueName##CameraVariableReference;\
+							if (StructProperty->Struct == CameraParameterType::StaticStruct())\
 							{\
 								CameraParameterType* CameraParameterPtr = StructProperty->ContainerPtrToValuePtr<CameraParameterType>(CameraNode);\
 								RawSourceValuePtr = static_cast<void*>(&CameraParameterPtr->Value);\
 							}\
+							else if (StructProperty->Struct == CameraVariableReferenceType::StaticStruct())\
+							{\
+								CameraVariableReferenceType* VariableReferencePtr = StructProperty->ContainerPtrToValuePtr<CameraVariableReferenceType>(CameraNode);\
+								RawSourceValuePtr = VariableReferencePtr->Variable ? VariableReferencePtr->Variable->GetDefaultValuePtr() : nullptr;\
+							}\
 						}\
 						break;
 					UE_CAMERA_VARIABLE_FOR_ALL_TYPES()
