Cameras: fixes and improvements to auto-rotate node

- Make it work when yaw/pitch is used from Control Rotation.
- Add MovementOrFacing mode
- Add optional user-defined direction vector

diff --git a/Engine/Plugins/Cameras/GameplayCameras/Source/GameplayCameras/Private/Nodes/Input/AutoRotateInput2DCameraNode.cpp b/Engine/Plugins/Cameras/GameplayCameras/Source/GameplayCameras/Private/Nodes/Input/AutoRotateInput2DCameraNode.cpp
--- a/Engine/Plugins/Cameras/GameplayCameras/Source/GameplayCameras/Private/Nodes/Input/AutoRotateInput2DCameraNode.cpp
+++ b/Engine/Plugins/Cameras/GameplayCameras/Source/GameplayCameras/Private/Nodes/Input/AutoRotateInput2DCameraNode.cpp
@@ -7,13 +7,14 @@
 #include "Core/CameraEvaluationContext.h"
 #include "Core/CameraOperation.h"
 #include "Core/CameraParameterReader.h"
-#include "Core/CameraRigAsset.h"
 #include "Core/CameraValueInterpolator.h"
-#include "Core/CameraVariableTableFwd.h"
+#include "Core/CameraVariableReferenceReader.h"
 #include "Debug/CameraDebugBlock.h"
 #include "Debug/CameraDebugBlockBuilder.h"
 #include "Debug/CameraDebugRenderer.h"
+#include "GameFramework/PlayerController.h"
 #include "HAL/IConsoleManager.h"
+#include "ValueInterpolators/CriticalDamperValueInterpolator.h"
 
 #include UE_INLINE_GENERATED_CPP_BY_NAME(AutoRotateInput2DCameraNode)
 
@@ -49,9 +50,11 @@ protected:
 private:
 
 	void DeactivateAutoRotate(FCameraNodeEvaluationResult& OutResult);
+	APlayerController* GetPlayerController(TSharedPtr<const FCameraEvaluationContext> EvaluationContext) const;
 
 private:
 
+	TCameraVariableReferenceReader<FVector3d> DirectionVectorReader;
 	TCameraParameterReader<float> WaitTimeReader;
 	TCameraParameterReader<float> DeactivationThresholdReader;
 	TCameraParameterReader<bool> FreezeControlRotationReader;
@@ -112,9 +115,15 @@ void FAutoRotateInput2DCameraNodeEvaluator::OnInitialize(const FCameraNodeEvalua
 	{
 		LastInputValue = InputNodeEvaluator->GetInputValue();
 	}
+	else if (APlayerController* PlayerController = GetPlayerController(Params.EvaluationContext))
+	{
+		const FRotator3d ControlRotation = PlayerController->GetControlRotation();
+		LastInputValue = FVector2d(ControlRotation.Yaw, ControlRotation.Pitch);
+	}
 
 	const UAutoRotateInput2DCameraNode* AutoRotateNode = GetCameraNodeAs<UAutoRotateInput2DCameraNode>();
 
+	DirectionVectorReader.Initialize(AutoRotateNode->DirectionVector);
 	WaitTimeReader.Initialize(AutoRotateNode->WaitTime);
 	DeactivationThresholdReader.Initialize(AutoRotateNode->DeactivationThreshold);
 	FreezeControlRotationReader.Initialize(AutoRotateNode->FreezeControlRotation);
@@ -137,6 +146,11 @@ void FAutoRotateInput2DCameraNodeEvaluator::OnRun(const FCameraNodeEvaluationPar
 
 		InputValue = InputNodeEvaluator->GetInputValue();
 	}
+	else if (APlayerController* PlayerController = GetPlayerController(Params.EvaluationContext))
+	{
+		const FRotator3d ControlRotation = PlayerController->GetControlRotation();
+		InputValue = FVector2d(ControlRotation.Yaw, ControlRotation.Pitch);
+	}
 
 	// Bail out if auto-rotate is disabled.
 	const bool bAutoRotateEnabled = EnableAutoRotateReader.Get(OutResult.VariableTable);
@@ -171,11 +185,43 @@ void FAutoRotateInput2DCameraNodeEvaluator::OnRun(const FCameraNodeEvaluationPar
 	}
 
 	// Figure out which direction we should auto-rotate towards.
-	FVector3d AutoRotateDir = InitialResult.CameraPose.GetAimDir();
-	if (AutoRotateNode->Direction == ECameraAutoRotateDirection::Movement)
+	bool bHasAutoRotateDir = false;
+	const FVector3d ContextAimDir = InitialResult.CameraPose.GetAimDir();
+	FVector3d AutoRotateDir = ContextAimDir;
+	if (DirectionVectorReader.IsDriven())
+	{
+		bHasAutoRotateDir = DirectionVectorReader.TryGet(OutResult.VariableTable, AutoRotateDir);
+	}
+	else
 	{
-		AutoRotateDir = ContextMovement.GetSafeNormal(UE_SMALL_NUMBER, AutoRotateDir);
+		switch (AutoRotateNode->Direction)
+		{
+			case ECameraAutoRotateDirection::Facing:
+			default:
+				bHasAutoRotateDir = true;
+				break;
+			case ECameraAutoRotateDirection::Movement:
+				if (!ContextMovement.IsNearlyZero())
+				{
+					AutoRotateDir = ContextMovement.GetSafeNormal(UE_SMALL_NUMBER, AutoRotateDir);
+					bHasAutoRotateDir = true;
+				}
+				break;
+			case ECameraAutoRotateDirection::MovementOrFacing:
+				if (!ContextMovement.IsNearlyZero())
+				{
+					AutoRotateDir = ContextMovement.GetSafeNormal(UE_SMALL_NUMBER, AutoRotateDir);
+				}
+				bHasAutoRotateDir = true;
+				break;
+		}
 	}
+	if (!bHasAutoRotateDir)
+	{
+		DeactivateAutoRotate(OutResult);
+		return;
+	}
+
 	const FRotator3d AutoRotateRot = AutoRotateDir.ToOrientationRotator();
 
 	// Figure out how much work we have to do.
@@ -255,6 +301,14 @@ void FAutoRotateInput2DCameraNodeEvaluator::OnRun(const FCameraNodeEvaluationPar
 			bDeactivateAutoRotate = true;
 		}
 	}
+	else if (APlayerController* PlayerController = GetPlayerController(Params.EvaluationContext))
+	{
+		const FVector2d NewInputValue(InputValue.X + NewDeltaVector.X, InputValue.Y + NewDeltaVector.Y);
+		const FRotator3d NewControlRotation(NewInputValue.Y, NewInputValue.X, 0.0);
+		PlayerController->SetControlRotation(NewControlRotation);
+
+		LastInputValue = NewInputValue;
+	}
 
 	if (Interpolator->IsFinished() || bDeactivateAutoRotate)
 	{
@@ -286,6 +340,15 @@ void FAutoRotateInput2DCameraNodeEvaluator::DeactivateAutoRotate(FCameraNodeEval
 	}
 }
 
+APlayerController* FAutoRotateInput2DCameraNodeEvaluator::GetPlayerController(TSharedPtr<const FCameraEvaluationContext> EvaluationContext) const
+{
+	if (EvaluationContext)
+	{
+		return EvaluationContext->GetPlayerController();
+	}
+	return nullptr;
+}
+
 #if UE_GAMEPLAY_CAMERAS_DEBUG
 
 void FAutoRotateInput2DCameraNodeEvaluator::OnBuildDebugBlocks(const FCameraDebugBlockBuildParams& Params, FCameraDebugBlockBuilder& Builder)
@@ -318,6 +381,10 @@ UAutoRotateInput2DCameraNode::UAutoRotateInput2DCameraNode(const FObjectInitiali
 	: Super(ObjInit)
 {
 	AddNodeFlags(ECameraNodeFlags::CustomGetChildren);
+
+	UCriticalDamperValueInterpolator* DefaultInterpolator = ObjInit.CreateDefaultSubobject<UCriticalDamperValueInterpolator>(this, "Interpolator");
+	DefaultInterpolator->DampingFactor = 10.f;
+	Interpolator = DefaultInterpolator;
 }
 
 FCameraNodeChildrenView UAutoRotateInput2DCameraNode::OnGetChildren()
diff --git a/Engine/Plugins/Cameras/GameplayCameras/Source/GameplayCameras/Public/Nodes/Input/AutoRotateInput2DCameraNode.h b/Engine/Plugins/Cameras/GameplayCameras/Source/GameplayCameras/Public/Nodes/Input/AutoRotateInput2DCameraNode.h
--- a/Engine/Plugins/Cameras/GameplayCameras/Source/GameplayCameras/Public/Nodes/Input/AutoRotateInput2DCameraNode.h
+++ b/Engine/Plugins/Cameras/GameplayCameras/Source/GameplayCameras/Public/Nodes/Input/AutoRotateInput2DCameraNode.h
@@ -3,6 +3,7 @@
 #pragma once
 
 #include "Core/CameraParameters.h"
+#include "Core/CameraVariableReferences.h"
 #include "Math/MathFwd.h"
 #include "Nodes/Input/Input2DCameraNode.h"
 
@@ -16,8 +17,16 @@ enum class ECameraAutoRotateDirection
 {
 	/** Re-align towards the evaluation context's facing. */
 	Facing,
-	/** Re-align towards the evaluation context's movement direction. */
-	Movement
+	/** 
+	 * Re-align towards the evaluation context's movement direction. 
+	 * Doesn't do anything when there is no movement.
+	 */
+	Movement,
+	/**
+	 * Re-align towards the evaluation context's movement direction if there is movement,
+	 * or towards its facing otherwise.
+	 */
+	MovementOrFacing
 };
 
 /**
@@ -35,6 +44,10 @@ public:
 	UPROPERTY(EditAnywhere, Category="Auto-Rotate")
 	ECameraAutoRotateDirection Direction;
 
+	/** An override for the direction to re-align towards. */
+	UPROPERTY(EditAnywhere, Category="Auto-Rotate")
+	FVector3dCameraVariableReference DirectionVector;
+
 	/** The time, in seconds, to wait before re-aligning. */
 	UPROPERTY(EditAnywhere, Category="Auto-Rotate")
 	FFloatCameraParameter WaitTime = 1.f;
