Cameras: camera reference widget is now disabled if overriden with a rig parameter

diff --git a/Engine/Plugins/Cameras/GameplayCameras/Source/GameplayCameras/Public/Core/CameraVariableReferences.h b/Engine/Plugins/Cameras/GameplayCameras/Source/GameplayCameras/Public/Core/CameraVariableReferences.h
--- a/Engine/Plugins/Cameras/GameplayCameras/Source/GameplayCameras/Public/Core/CameraVariableReferences.h
+++ b/Engine/Plugins/Cameras/GameplayCameras/Source/GameplayCameras/Public/Core/CameraVariableReferences.h
@@ -16,6 +16,7 @@ namespace UE::Cameras
 	F##ValueName##CameraVariableReference(VariableAssetType* InVariable) : Variable(InVariable) {}\
 	bool IsValid() const { return VariableID.IsValid(); }\
 	bool HasVariable() const { return Variable != nullptr; }\
+	bool HasNonUserOverride() const { return VariableID.IsValid() && (!Variable || Variable->GetVariableID() != VariableID); }\
 	const F##ValueName##CameraVariableReference::ValueType* GetValue(const UE::Cameras::FCameraVariableTable& VariableTable) const;
 
 // All camera references have:
diff --git a/Engine/Plugins/Cameras/GameplayCameras/Source/GameplayCamerasEditor/Private/Customizations/CameraVariableReferenceDetailsCustomizations.cpp b/Engine/Plugins/Cameras/GameplayCameras/Source/GameplayCamerasEditor/Private/Customizations/CameraVariableReferenceDetailsCustomizations.cpp
--- a/Engine/Plugins/Cameras/GameplayCameras/Source/GameplayCamerasEditor/Private/Customizations/CameraVariableReferenceDetailsCustomizations.cpp
+++ b/Engine/Plugins/Cameras/GameplayCameras/Source/GameplayCamerasEditor/Private/Customizations/CameraVariableReferenceDetailsCustomizations.cpp
@@ -82,6 +82,7 @@ void FCameraVariableReferenceDetailsCustomization::CustomizeHeader(TSharedRef<IP
 		.HasDownArrow(true)
 		.ContentPadding(1.f)
 		.ToolTipText(LOCTEXT("SetVariable_ToolTip", "Selects the camera variable"))
+		.IsEnabled(this, &FCameraVariableReferenceDetailsCustomization::IsCameraVariableBrowserEnabled)
 		.ButtonContent()
 		[
 			SNew(SHorizontalBox)
@@ -169,6 +170,22 @@ TSharedRef<SWidget> FCameraVariableReferenceDetailsCustomization::BuildCameraVar
 	return MenuBuilder.MakeWidget();
 }
 
+bool FCameraVariableReferenceDetailsCustomization::IsCameraVariableBrowserEnabled() const
+{
+	TArray<void*> RawData;
+	StructProperty->AccessRawData(RawData);
+
+	for (int32 ValueIndex = 0; ValueIndex < RawData.Num(); ++ValueIndex)
+	{
+		if (HasNonUserOverride(RawData[ValueIndex]))
+		{
+			return false;
+		}
+	}
+
+	return true;
+}
+
 FText FCameraVariableReferenceDetailsCustomization::GetVariableName() const
 {
 	UObject* VariableObject = nullptr;
@@ -247,6 +264,11 @@ void FCameraVariableReferenceDetailsCustomization::OnResetToDefault()
 }
 
 #define UE_CAMERA_VARIABLE_FOR_TYPE(ValueType, ValueName)\
+bool F##ValueName##CameraVariableReferenceDetailsCustomization::HasNonUserOverride(void* InRawData) const\
+{\
+	F##ValueName##CameraVariableReference* TypedData = reinterpret_cast<F##ValueName##CameraVariableReference*>(InRawData);\
+	return TypedData->HasNonUserOverride();\
+}\
 void F##ValueName##CameraVariableReferenceDetailsCustomization::SetReferenceVariable(void* InRawData, UCameraVariableAsset* InVariable)\
 {\
 	F##ValueName##CameraVariableReference* TypedData = reinterpret_cast<F##ValueName##CameraVariableReference*>(InRawData);\
diff --git a/Engine/Plugins/Cameras/GameplayCameras/Source/GameplayCamerasEditor/Public/Customizations/CameraVariableReferenceDetailsCustomizations.h b/Engine/Plugins/Cameras/GameplayCameras/Source/GameplayCamerasEditor/Public/Customizations/CameraVariableReferenceDetailsCustomizations.h
--- a/Engine/Plugins/Cameras/GameplayCameras/Source/GameplayCamerasEditor/Public/Customizations/CameraVariableReferenceDetailsCustomizations.h
+++ b/Engine/Plugins/Cameras/GameplayCameras/Source/GameplayCamerasEditor/Public/Customizations/CameraVariableReferenceDetailsCustomizations.h
@@ -33,12 +33,15 @@ public:
 
 protected:
 
+	virtual bool HasNonUserOverride(void* InRawData) const = 0;
 	virtual void SetReferenceVariable(void* InRawData, UCameraVariableAsset* InVariable) = 0;
 
 private:
 
 	TSharedRef<SWidget> BuildCameraVariableBrowser();
 
+	bool IsCameraVariableBrowserEnabled() const;
+
 	FText GetVariableName() const;
 
 	bool CanClearVariable() const;
@@ -64,6 +67,7 @@ protected:
 class F##ValueName##CameraVariableReferenceDetailsCustomization : public FCameraVariableReferenceDetailsCustomization\
 {\
 protected:\
+	virtual bool HasNonUserOverride(void* InRawData) const override;\
 	virtual void SetReferenceVariable(void* InRawData, UCameraVariableAsset* InVariable) override;\
 };
 UE_CAMERA_VARIABLE_FOR_ALL_TYPES()
