Cameras: re-add previously removed Blueprint APIs as deprecated APIs

We don't guarantee backwards compatibility for Experimental plugins but this should cover a majority of cases.

diff --git a/Engine/Plugins/Cameras/GameplayCameras/Source/GameplayCameras/Private/Directors/BlueprintCameraDirector.cpp b/Engine/Plugins/Cameras/GameplayCameras/Source/GameplayCameras/Private/Directors/BlueprintCameraDirector.cpp
--- a/Engine/Plugins/Cameras/GameplayCameras/Source/GameplayCameras/Private/Directors/BlueprintCameraDirector.cpp
+++ b/Engine/Plugins/Cameras/GameplayCameras/Source/GameplayCameras/Private/Directors/BlueprintCameraDirector.cpp
@@ -355,7 +355,11 @@ void UBlueprintCameraDirectorEvaluator::NativeActivateCameraDirector(const UE::C
 			EvaluationContextOwner = EvaluationContext->GetOwner();
 		}
 
-		ActivateCameraDirector(EvaluationContextOwner);
+PRAGMA_DISABLE_DEPRECATION_WARNINGS
+		FBlueprintCameraDirectorActivateParams OldParams;
+		OldParams.EvaluationContextOwner = EvaluationContextOwner;
+		ActivateCameraDirector(EvaluationContextOwner, OldParams);
+PRAGMA_ENABLE_DEPRECATION_WARNINGS
 	}
 }
 
@@ -369,7 +373,11 @@ void UBlueprintCameraDirectorEvaluator::NativeDeactivateCameraDirector(const UE:
 			EvaluationContextOwner = EvaluationContext->GetOwner();
 		}
 
-		DeactivateCameraDirector(EvaluationContextOwner);
+PRAGMA_DISABLE_DEPRECATION_WARNINGS
+	FBlueprintCameraDirectorDeactivateParams OldParams;
+		OldParams.EvaluationContextOwner = EvaluationContextOwner;
+		DeactivateCameraDirector(EvaluationContextOwner, OldParams);
+PRAGMA_ENABLE_DEPRECATION_WARNINGS
 	}
 }
 
diff --git a/Engine/Plugins/Cameras/GameplayCameras/Source/GameplayCameras/Private/GameFramework/BlueprintCameraPose.cpp b/Engine/Plugins/Cameras/GameplayCameras/Source/GameplayCameras/Private/GameFramework/BlueprintCameraPose.cpp
--- a/Engine/Plugins/Cameras/GameplayCameras/Source/GameplayCameras/Private/GameFramework/BlueprintCameraPose.cpp
+++ b/Engine/Plugins/Cameras/GameplayCameras/Source/GameplayCameras/Private/GameFramework/BlueprintCameraPose.cpp
@@ -123,3 +123,40 @@ FBlueprintCameraPose UBlueprintCameraPoseFunctionLibrary::MakeCameraPoseFromCine
 	return Result;
 }
 
+// Deprecated methods.
+
+FBlueprintCameraPose UBlueprintCameraPoseFunctionLibrary::SetLocation(const FBlueprintCameraPose& CameraPose, const FVector& Location)
+{
+	FBlueprintCameraPose Result(CameraPose);
+	Result.Location = Location;
+	return Result;
+}
+
+FBlueprintCameraPose UBlueprintCameraPoseFunctionLibrary::SetRotation(const FBlueprintCameraPose& CameraPose, const FRotator& Rotation)
+{
+	FBlueprintCameraPose Result(CameraPose);
+	Result.Rotation = Rotation;
+	return Result;
+}
+
+FBlueprintCameraPose UBlueprintCameraPoseFunctionLibrary::SetTargetDistance(const FBlueprintCameraPose& CameraPose, double TargetDistance)
+{
+	FBlueprintCameraPose Result(CameraPose);
+	Result.TargetDistance = TargetDistance;
+	return Result;
+}
+
+FBlueprintCameraPose UBlueprintCameraPoseFunctionLibrary::SetFieldOfView(const FBlueprintCameraPose& CameraPose, float FieldOfView)
+{
+	FBlueprintCameraPose Result(CameraPose);
+	Result.FieldOfView = FieldOfView;
+	return Result;
+}
+
+FBlueprintCameraPose UBlueprintCameraPoseFunctionLibrary::SetFocalLength(const FBlueprintCameraPose& CameraPose, float FocalLength)
+{
+	FBlueprintCameraPose Result(CameraPose);
+	Result.FocalLength = FocalLength;
+	return Result;
+}
+
diff --git a/Engine/Plugins/Cameras/GameplayCameras/Source/GameplayCameras/Private/Nodes/Utility/BlueprintCameraNode.cpp b/Engine/Plugins/Cameras/GameplayCameras/Source/GameplayCameras/Private/Nodes/Utility/BlueprintCameraNode.cpp
--- a/Engine/Plugins/Cameras/GameplayCameras/Source/GameplayCameras/Private/Nodes/Utility/BlueprintCameraNode.cpp
+++ b/Engine/Plugins/Cameras/GameplayCameras/Source/GameplayCameras/Private/Nodes/Utility/BlueprintCameraNode.cpp
@@ -214,6 +214,7 @@ void UBlueprintCameraNodeEvaluator::SetupExecution(TSharedPtr<const FCameraEvalu
 
 	ensure(!CameraData.IsValid());
 	CameraData = FBlueprintCameraEvaluationDataRef::MakeExternalRef(&OutResult);
+	VariableTable = FBlueprintCameraEvaluationDataRef::MakeExternalRef(&OutResult);;
 
 	ensure(!CurrentContext.IsValid());
 	CurrentContext = EvaluationContext;
@@ -221,6 +222,7 @@ void UBlueprintCameraNodeEvaluator::SetupExecution(TSharedPtr<const FCameraEvalu
 
 void UBlueprintCameraNodeEvaluator::TeardownExecution()
 {
+	VariableTable = FBlueprintCameraEvaluationDataRef();
 	CameraData = FBlueprintCameraEvaluationDataRef();
 	CurrentContext = nullptr;
 }
@@ -260,9 +262,9 @@ FBlueprintCameraPose UBlueprintCameraNodeEvaluator::GetCurrentCameraPose() const
 	return UBlueprintCameraEvaluationDataFunctionLibrary::GetCameraPose(CameraData);
 }
 
-void UBlueprintCameraNodeEvaluator::SetCurrentCameraPose(const FBlueprintCameraPose& CameraPose)
+void UBlueprintCameraNodeEvaluator::SetCurrentCameraPose(const FBlueprintCameraPose& InCameraPose)
 {
-	UBlueprintCameraEvaluationDataFunctionLibrary::SetCameraPose(CameraData, CameraPose);
+	UBlueprintCameraEvaluationDataFunctionLibrary::SetCameraPose(CameraData, InCameraPose);
 }
 
 void UBlueprintCameraNodeEvaluator::SetDefaultOwningCameraRigParameters(FBlueprintCameraEvaluationDataRef TargetCameraData) const
diff --git a/Engine/Plugins/Cameras/GameplayCameras/Source/GameplayCameras/Public/Directors/BlueprintCameraDirector.h b/Engine/Plugins/Cameras/GameplayCameras/Source/GameplayCameras/Public/Directors/BlueprintCameraDirector.h
--- a/Engine/Plugins/Cameras/GameplayCameras/Source/GameplayCameras/Public/Directors/BlueprintCameraDirector.h
+++ b/Engine/Plugins/Cameras/GameplayCameras/Source/GameplayCameras/Public/Directors/BlueprintCameraDirector.h
@@ -28,6 +28,26 @@ struct UE_DEPRECATED(5.6, "This parameter structure is deprecated, parameters ar
 	TObjectPtr<UObject> EvaluationContextOwner;
 };
 
+USTRUCT(BlueprintType)
+struct UE_DEPRECATED(5.6, "This parameter structure is deprecated, parameters are now passed directly to ActivateCameraDirector.") FBlueprintCameraDirectorActivateParams
+{
+	GENERATED_BODY()
+
+	/** The owner (if any) of the evaluation context we are running inside of. */
+	UPROPERTY(BlueprintReadWrite, Category="Evaluation")
+	TObjectPtr<UObject> EvaluationContextOwner;
+};
+
+USTRUCT(BlueprintType)
+struct UE_DEPRECATED(5.6, "This parameter structure is deprecated, parameters are now passed directly to DeactivateCameraDirector.") FBlueprintCameraDirectorDeactivateParams
+{
+	GENERATED_BODY()
+
+	/** The owner (if any) of the evaluation context we were running inside of. */
+	UPROPERTY(BlueprintReadWrite, Category="Evaluation")
+	TObjectPtr<UObject> EvaluationContextOwner;
+};
+
 /**
  * Base class for a Blueprint camera director evaluator.
  */
@@ -42,15 +62,19 @@ public:
 	 * Override this method in Blueprint to execute custom logic when this
 	 * camera director gets activated.
 	 */
+PRAGMA_DISABLE_DEPRECATION_WARNINGS
 	UFUNCTION(BlueprintCallable, BlueprintImplementableEvent, Category="Camera Director|Activation")
-	void ActivateCameraDirector(UObject* EvaluationContextOwner);
+	void ActivateCameraDirector(UObject* EvaluationContextOwner, const FBlueprintCameraDirectorActivateParams& Params);
+PRAGMA_ENABLE_DEPRECATION_WARNINGS
 
 	/**
 	 * Override this method in Blueprint to execute custom logic when this
 	 * camera director gets deactivated.
 	 */
+PRAGMA_DISABLE_DEPRECATION_WARNINGS
 	UFUNCTION(BlueprintCallable, BlueprintImplementableEvent, Category="Camera Director|Activation")
-	void DeactivateCameraDirector(UObject* EvaluationContextOwner);
+	void DeactivateCameraDirector(UObject* EvaluationContextOwner, const FBlueprintCameraDirectorDeactivateParams& Params);
+PRAGMA_ENABLE_DEPRECATION_WARNINGS
 	
 	/**
 	 * Override this method in Blueprint to execute the custom logic that determines
diff --git a/Engine/Plugins/Cameras/GameplayCameras/Source/GameplayCameras/Public/GameFramework/BlueprintCameraPose.h b/Engine/Plugins/Cameras/GameplayCameras/Source/GameplayCameras/Public/GameFramework/BlueprintCameraPose.h
--- a/Engine/Plugins/Cameras/GameplayCameras/Source/GameplayCameras/Public/GameFramework/BlueprintCameraPose.h
+++ b/Engine/Plugins/Cameras/GameplayCameras/Source/GameplayCameras/Public/GameFramework/BlueprintCameraPose.h
@@ -195,6 +195,38 @@ public:
 	/** Creates a new camera pose given a cine-camera component. */
 	UFUNCTION(BlueprintCallable, Category="Camera")
 	static FBlueprintCameraPose MakeCameraPoseFromCineCameraComponent(const UCineCameraComponent* CameraComponent);
+
+public:
+
+	UFUNCTION(BlueprintPure, Category="Camera", meta=(DeprecatedFunction, DeprecationMessage="Please use the standard Break node"))
+	static inline FVector GetLocation(const FBlueprintCameraPose& CameraPose) { return CameraPose.Location; }
+
+	UFUNCTION(BlueprintPure, Category="Camera", meta=(DeprecatedFunction, DeprecationMessage="Please use the standard Break node"))
+	static inline FRotator GetRotation(const FBlueprintCameraPose& CameraPose) { return CameraPose.Rotation; }
+
+	UFUNCTION(BlueprintPure, Category="Camera", meta=(DeprecatedFunction, DeprecationMessage="Please use the standard Break node"))
+	static inline double GetTargetDistance(const FBlueprintCameraPose& CameraPose) { return CameraPose.TargetDistance; }
+
+	UFUNCTION(BlueprintPure, Category="Camera", meta=(DeprecatedFunction, DeprecationMessage="Please use the standard Break node"))
+	static inline double GetFieldOfView(const FBlueprintCameraPose& CameraPose) { return CameraPose.FieldOfView; }
+
+	UFUNCTION(BlueprintPure, Category="Camera", meta=(DeprecatedFunction, DeprecationMessage="Please use the standard Break node"))
+	static inline double GetFocalLength(const FBlueprintCameraPose& CameraPose) { return CameraPose.FocalLength; }
+	
+	UFUNCTION(BlueprintCallable, Category="Camera", meta=(DeprecatedFunction, DeprecationMessage="Please use the standard Set Members node"))
+	static FBlueprintCameraPose SetLocation(const FBlueprintCameraPose& CameraPose, const FVector& Location);
+
+	UFUNCTION(BlueprintCallable, Category="Camera", meta=(DeprecatedFunction, DeprecationMessage="Please use the standard Set Members node"))
+	static FBlueprintCameraPose SetRotation(const FBlueprintCameraPose& CameraPose, const FRotator& Rotation);
+
+	UFUNCTION(BlueprintCallable, Category="Camera", meta=(DeprecatedFunction, DeprecationMessage="Please use the standard Set Members node"))
+	static FBlueprintCameraPose SetTargetDistance(const FBlueprintCameraPose& CameraPose, double TargetDistance);
+
+	UFUNCTION(BlueprintCallable, Category="Camera", meta=(DeprecatedFunction, DeprecationMessage="Please use the standard Set Members node"))
+	static FBlueprintCameraPose SetFieldOfView(const FBlueprintCameraPose& CameraPose, float FieldOfView);
+
+	UFUNCTION(BlueprintCallable, Category="Camera", meta=(DeprecatedFunction, DeprecationMessage="Please use the standard Set Members node"))
+	static FBlueprintCameraPose SetFocalLength(const FBlueprintCameraPose& CameraPose, float FocalLength);
 };
 
 #undef UE_API
diff --git a/Engine/Plugins/Cameras/GameplayCameras/Source/GameplayCameras/Public/Nodes/Utility/BlueprintCameraNode.h b/Engine/Plugins/Cameras/GameplayCameras/Source/GameplayCameras/Public/Nodes/Utility/BlueprintCameraNode.h
--- a/Engine/Plugins/Cameras/GameplayCameras/Source/GameplayCameras/Public/Nodes/Utility/BlueprintCameraNode.h
+++ b/Engine/Plugins/Cameras/GameplayCameras/Source/GameplayCameras/Public/Nodes/Utility/BlueprintCameraNode.h
@@ -100,6 +100,26 @@ protected:
 	UPROPERTY(BlueprintReadOnly, Category="Evaluation")
 	FBlueprintCameraEvaluationDataRef CameraData;
 
+public:
+
+	// Deprecated methods.
+
+	UFUNCTION(BlueprintPure, BlueprintGetter, meta=(DeprecatedFunction, DeprecationMessage="Please use GetCurrentCameraPose"))
+	FBlueprintCameraPose GetCameraPose() const { return GetCurrentCameraPose(); }
+
+	UFUNCTION(BlueprintSetter, meta=(DeprecatedFunction, DeprecationMessage="Please use SetCurrentCameraPose"))
+	void SetCameraPose(UPARAM(DisplayName="Camera Pose") const FBlueprintCameraPose& InCameraPose) { SetCurrentCameraPose(InCameraPose); }
+
+protected:
+
+	// Deprecated fields.
+
+	UPROPERTY(BlueprintGetter=GetCameraPose, BlueprintSetter=SetCameraPose, Category="Evaluation", meta=(DeprecatedProperty, DeprecationMessage="Please use CameraData, or GetCurrentCameraPose and SetCurrentCameraPose"))
+	FBlueprintCameraPose CameraPose;
+
+	UPROPERTY(BlueprintReadOnly, Category="Evaluation", meta=(DeprecatedProperty, DeprecationMessage="Please use CameraData"))
+	FBlueprintCameraEvaluationDataRef VariableTable;
+
 private:
 
 	TSharedPtr<const UE::Cameras::FCameraEvaluationContext> CurrentContext;
