Cameras: expose debug block builder APIs to 3rd party code

diff --git a/Engine/Plugins/Cameras/GameplayCameras/Source/GameplayCameras/Public/Debug/CameraDebugBlockBuilder.h b/Engine/Plugins/Cameras/GameplayCameras/Source/GameplayCameras/Public/Debug/CameraDebugBlockBuilder.h
--- a/Engine/Plugins/Cameras/GameplayCameras/Source/GameplayCameras/Public/Debug/CameraDebugBlockBuilder.h
+++ b/Engine/Plugins/Cameras/GameplayCameras/Source/GameplayCameras/Public/Debug/CameraDebugBlockBuilder.h
@@ -84,7 +84,7 @@ public:
 	}
 
 	/** Ends the currently active debug drawing block. */
-	void EndChildDebugBlock();
+	GAMEPLAYCAMERAS_API void EndChildDebugBlock();
 
 	/** Gets current hierarchy level. */
 	int32 GetHierarchyLevel() const { return CurrentHierarchy.Num(); }
@@ -95,21 +95,21 @@ public:
 	 * This implies that a node evaluator will visit its children "manually", otherwise these
 	 * children node evaluators won't have any debugging information available.
 	 */
-	void SkipChildren();
+	GAMEPLAYCAMERAS_API void SkipChildren();
 	/** Gets visiting flags. */
 	ECameraDebugBlockBuildVisitFlags GetVisitFlags() const { return VisitFlags; }
 	/** Resets visiting flags. */
-	void ResetVisitFlags();
+	GAMEPLAYCAMERAS_API void ResetVisitFlags();
 
 	/** Temporarily overrides the attachment/children to work on the new parent. */
-	void StartParentDebugBlockOverride(FCameraDebugBlock& InNewParent);
+	GAMEPLAYCAMERAS_API void StartParentDebugBlockOverride(FCameraDebugBlock& InNewParent);
 	/** Ends a temporary attachment/children override. */
-	void EndParentDebugBlockOverride();
+	GAMEPLAYCAMERAS_API void EndParentDebugBlockOverride();
 
 private:
 
-	void OnAttachDebugBlock(FCameraDebugBlock* InNewBlock);
-	void OnStartChildDebugBlock(FCameraDebugBlock* InNewBlock);
+	GAMEPLAYCAMERAS_API void OnAttachDebugBlock(FCameraDebugBlock* InNewBlock);
+	GAMEPLAYCAMERAS_API void OnStartChildDebugBlock(FCameraDebugBlock* InNewBlock);
 
 private:
 
