Fixed BeginPlay call before the "UControllerGameplayCameraEvaluationComponent" component got initialized
- The component found was being registered before being initialized, which was causing the call of the BeginPlay of said component that ended up raising an ensure.

#rnx

diff --git a/Engine/Plugins/Cameras/GameplayCameras/Source/GameplayCameras/Private/GameFramework/ControllerGameplayCameraEvaluationComponent.cpp b/Engine/Plugins/Cameras/GameplayCameras/Source/GameplayCameras/Private/GameFramework/ControllerGameplayCameraEvaluationComponent.cpp
--- a/Engine/Plugins/Cameras/GameplayCameras/Source/GameplayCameras/Private/GameFramework/ControllerGameplayCameraEvaluationComponent.cpp
+++ b/Engine/Plugins/Cameras/GameplayCameras/Source/GameplayCameras/Private/GameFramework/ControllerGameplayCameraEvaluationComponent.cpp
@@ -42,6 +42,8 @@ void UControllerGameplayCameraEvaluationComponent::Initialize(TScriptInterface<I
 	InitParams.PlayerController = PlayerController;
 	EvaluationContext = MakeShared<FCameraEvaluationContext>(InitParams);
 	EvaluationContext->GetInitialResult().bIsValid = true;	
+	
+	RegisterComponent();
 }
 
 void UControllerGameplayCameraEvaluationComponent::ActivateCameraRig(UCameraRigAsset* CameraRig, ECameraRigLayer EvaluationLayer)
@@ -114,7 +116,6 @@ UControllerGameplayCameraEvaluationComponent* UControllerGameplayCameraEvaluatio
 	{
 		ControllerComponent = NewObject<UControllerGameplayCameraEvaluationComponent>(
 				OwnerActor, TEXT("ControllerGameplayCameraEvaluationComponent"), RF_Transient);
-		ControllerComponent->RegisterComponent();
 		if (bOutCreated)
 		{
 			*bOutCreated = true;
