Cameras: fix incorrect reset-to-default behavior on the camera parameter editor UI

Also add a "Go To Variable" menu entry to open the referenced camera variable.

diff --git a/Engine/Plugins/Cameras/GameplayCameras/Source/GameplayCamerasEditor/Private/Customizations/CameraParameterDetailsCustomizations.cpp b/Engine/Plugins/Cameras/GameplayCameras/Source/GameplayCamerasEditor/Private/Customizations/CameraParameterDetailsCustomizations.cpp
--- a/Engine/Plugins/Cameras/GameplayCameras/Source/GameplayCamerasEditor/Private/Customizations/CameraParameterDetailsCustomizations.cpp
+++ b/Engine/Plugins/Cameras/GameplayCameras/Source/GameplayCamerasEditor/Private/Customizations/CameraParameterDetailsCustomizations.cpp
@@ -20,6 +20,9 @@
 #include "ScopedTransaction.h"
 #include "Styles/GameplayCamerasEditorStyle.h"
 #include "Styling/StyleColors.h"
+#include "Subsystems/AssetEditorSubsystem.h"
+#include "Toolkits/CameraVariableCollectionEditorToolkit.h"
+#include "Toolkits/ToolkitManager.h"
 #include "UObject/Object.h"
 #include "Widgets/Images/SImage.h"
 #include "Widgets/Input/SComboButton.h"
@@ -116,7 +119,7 @@ void FCameraParameterDetailsCustomization::CustomizeHeader(TSharedRef<IPropertyH
 			.ContentPadding(1.f)
 			.ButtonStyle(FAppStyle::Get(), "SimpleButton")
 			.IsEnabled(this, &FCameraParameterDetailsCustomization::IsCameraVariableBrowserEnabled)
-			.ToolTipText(LOCTEXT("SetVariable_ToolTip", "Selects a camera variable to drive this parameter"))
+			.ToolTipText(this, &FCameraParameterDetailsCustomization::GetCameraVariableBrowserToolTip)
 			.ButtonContent()
 			[
 				SNew(SHorizontalBox)
@@ -166,6 +169,11 @@ void FCameraParameterDetailsCustomization::CustomizeHeader(TSharedRef<IPropertyH
 			.OnGetMenuContent(this, &FCameraParameterDetailsCustomization::BuildCameraVariableBrowser)
 		]
 	];
+
+	HeaderRow.OverrideResetToDefault(
+			FResetToDefaultOverride::Create(
+				FIsResetToDefaultVisible::CreateSP(this, &FCameraParameterDetailsCustomization::IsResetToDefaultVisible),
+				FResetToDefaultHandler::CreateSP(this, &FCameraParameterDetailsCustomization::OnResetToDefault)));
 }
 
 void FCameraParameterDetailsCustomization::CustomizeChildren(TSharedRef<IPropertyHandle> PropertyHandle, IDetailChildrenBuilder& ChildBuilder, IPropertyTypeCustomizationUtils& CustomizationUtils)
@@ -251,6 +259,14 @@ TSharedRef<SWidget> FCameraParameterDetailsCustomization::BuildCameraVariableBro
 
 	MenuBuilder.BeginSection(NAME_None, LOCTEXT("CameraVariableOperations", "Current Parameter"));
 	{
+		MenuBuilder.AddMenuEntry(
+			LOCTEXT("GoToVariable", "Go to variable"),
+			LOCTEXT("GoToVariable_ToolTip", "Open the referenced camera variable collection asset"),
+			FSlateIcon(FAppStyle::GetAppStyleSetName(), "Icons.BrowseContent"),
+			FUIAction(
+				FExecuteAction::CreateSP(this, &FCameraParameterDetailsCustomization::OnGoToVariable),
+				FCanExecuteAction::CreateSP(this, &FCameraParameterDetailsCustomization::CanGoToVariable))
+			);
 		MenuBuilder.AddMenuEntry(
 			LOCTEXT("ClearVariable", "Clear"),
 			LOCTEXT("ClearVariable_ToolTip", "Clears the variable from the camera parameter"),
@@ -300,6 +316,22 @@ bool FCameraParameterDetailsCustomization::IsCameraVariableBrowserEnabled() cons
 	return !VariableInfo.bHasNonUserOverride;
 }
 
+FText FCameraParameterDetailsCustomization::GetCameraVariableBrowserToolTip() const
+{
+	if (!VariableInfo.bHasNonUserOverride)
+	{
+		return LOCTEXT(
+				"SetVariable_ToolTip", 
+				"Selects a camera variable to drive this parameter");
+	}
+	else
+	{
+		return LOCTEXT(
+				"SetVariableDisabled_ToolTip",
+				"This parameter is exposed by the camera rig and cannot be also driven by a camera variable");
+	}
+}
+
 FText FCameraParameterDetailsCustomization::GetVariableInfoText() const
 {
 	return VariableInfo.InfoText;
@@ -344,6 +376,35 @@ FOptionalSize FCameraParameterDetailsCustomization::GetVariableErrorTextMaxWidth
 	return bShowVariableErrorText ? FOptionalSize((LayoutBoxWidth - FixedSpace) / 3.f) : FOptionalSize(0);
 }
 
+bool FCameraParameterDetailsCustomization::CanGoToVariable() const
+{
+	UObject* VariableObject = nullptr;
+	FPropertyAccess::Result PropertyAccessResult = VariableProperty->GetValue(VariableObject);
+	return PropertyAccessResult == FPropertyAccess::Success;
+}
+
+void FCameraParameterDetailsCustomization::OnGoToVariable()
+{
+	UObject* VariableObject = nullptr;
+	FPropertyAccess::Result PropertyAccessResult = VariableProperty->GetValue(VariableObject);
+	if (VariableObject && PropertyAccessResult == FPropertyAccess::Success)
+	{
+		UCameraVariableCollection* VariableCollection = VariableObject->GetTypedOuter<UCameraVariableCollection>();
+		if (VariableCollection)
+		{
+			GEditor->GetEditorSubsystem<UAssetEditorSubsystem>()->OpenEditorForAsset(VariableCollection);
+
+			TSharedPtr<IToolkit> FoundToolkit = FToolkitManager::Get().FindEditorForAsset(VariableCollection);
+			if (FoundToolkit)
+			{
+				TSharedPtr<FCameraVariableCollectionEditorToolkit> VariableCollectionToolkit = 
+					StaticCastSharedPtr<FCameraVariableCollectionEditorToolkit>(FoundToolkit);
+				VariableCollectionToolkit->FocusWindow(VariableObject);
+			}
+		}
+	}
+}
+
 bool FCameraParameterDetailsCustomization::CanClearVariable() const
 {
 	return VariableProperty->CanResetToDefault();
@@ -420,3 +481,4 @@ UE_CAMERA_VARIABLE_FOR_ALL_TYPES()
 }  // namespace UE::Cameras
 
 #undef LOCTEXT_NAMESPACE
+
diff --git a/Engine/Plugins/Cameras/GameplayCameras/Source/GameplayCamerasEditor/Private/Editors/SCameraVariableCollectionEditor.cpp b/Engine/Plugins/Cameras/GameplayCameras/Source/GameplayCamerasEditor/Private/Editors/SCameraVariableCollectionEditor.cpp
--- a/Engine/Plugins/Cameras/GameplayCameras/Source/GameplayCamerasEditor/Private/Editors/SCameraVariableCollectionEditor.cpp
+++ b/Engine/Plugins/Cameras/GameplayCameras/Source/GameplayCamerasEditor/Private/Editors/SCameraVariableCollectionEditor.cpp
@@ -9,6 +9,7 @@
 #include "ScopedTransaction.h"
 #include "Styles/GameplayCamerasEditorStyle.h"
 #include "ToolMenus.h"
+#include "Types/SlateEnums.h"
 #include "Widgets/Input/SSearchBox.h"
 #include "Widgets/Layout/SBorder.h"
 #include "Widgets/SBoxPanel.h"
@@ -246,6 +247,11 @@ void SCameraVariableCollectionEditor::GetSelectedVariables(TArray<UCameraVariabl
 	ListView->GetSelectedItems(OutSelection);
 }
 
+void SCameraVariableCollectionEditor::SelectVariable(UCameraVariableAsset* InItem)
+{
+	ListView->SetSelection(InItem, ESelectInfo::Direct);
+}
+
 void SCameraVariableCollectionEditor::RequestRenameVariable(UCameraVariableAsset* InItem, FSimpleDelegate InOnRenamedItem)
 {
 	bDeferredRequestRenameItem = true;
diff --git a/Engine/Plugins/Cameras/GameplayCameras/Source/GameplayCamerasEditor/Private/Toolkits/CameraVariableCollectionEditorToolkit.cpp b/Engine/Plugins/Cameras/GameplayCameras/Source/GameplayCamerasEditor/Private/Toolkits/CameraVariableCollectionEditorToolkit.cpp
--- a/Engine/Plugins/Cameras/GameplayCameras/Source/GameplayCamerasEditor/Private/Toolkits/CameraVariableCollectionEditorToolkit.cpp
+++ b/Engine/Plugins/Cameras/GameplayCameras/Source/GameplayCamerasEditor/Private/Toolkits/CameraVariableCollectionEditorToolkit.cpp
@@ -345,6 +345,16 @@ bool FCameraVariableCollectionEditorToolkit::CanDeleteVariable()
 	return !Selection.IsEmpty();
 }
 
+void FCameraVariableCollectionEditorToolkit::FocusWindow(UObject* ObjectToFocusOn)
+{
+	FBaseAssetToolkit::FocusWindow(ObjectToFocusOn);
+
+	if (UCameraVariableAsset* VariableToFocusOn = Cast<UCameraVariableAsset>(ObjectToFocusOn))
+	{
+		VariableCollectionEditorWidget->SelectVariable(VariableToFocusOn);
+	}
+}
+
 void FCameraVariableCollectionEditorToolkit::PostUndo(bool bSuccess)
 {
 	VariableCollectionEditorWidget->RequestListRefresh();
diff --git a/Engine/Plugins/Cameras/GameplayCameras/Source/GameplayCamerasEditor/Public/Customizations/CameraParameterDetailsCustomizations.h b/Engine/Plugins/Cameras/GameplayCameras/Source/GameplayCamerasEditor/Public/Customizations/CameraParameterDetailsCustomizations.h
--- a/Engine/Plugins/Cameras/GameplayCameras/Source/GameplayCamerasEditor/Public/Customizations/CameraParameterDetailsCustomizations.h
+++ b/Engine/Plugins/Cameras/GameplayCameras/Source/GameplayCamerasEditor/Public/Customizations/CameraParameterDetailsCustomizations.h
@@ -74,6 +74,7 @@ private:
 
 	bool IsValueEditorEnabled() const;
 	bool IsCameraVariableBrowserEnabled() const;
+	FText GetCameraVariableBrowserToolTip() const;
 
 	FText GetVariableInfoText() const;
 	EVisibility GetVariableInfoTextVisibility() const;
@@ -83,6 +84,9 @@ private:
 	EVisibility GetVariableErrorTextVisibility() const;
 	FOptionalSize GetVariableErrorTextMaxWidth() const;
 
+	bool CanGoToVariable() const;
+	void OnGoToVariable();
+
 	bool CanClearVariable() const;
 	void OnClearVariable();
 
diff --git a/Engine/Plugins/Cameras/GameplayCameras/Source/GameplayCamerasEditor/Public/Editors/SCameraVariableCollectionEditor.h b/Engine/Plugins/Cameras/GameplayCameras/Source/GameplayCamerasEditor/Public/Editors/SCameraVariableCollectionEditor.h
--- a/Engine/Plugins/Cameras/GameplayCameras/Source/GameplayCamerasEditor/Public/Editors/SCameraVariableCollectionEditor.h
+++ b/Engine/Plugins/Cameras/GameplayCameras/Source/GameplayCamerasEditor/Public/Editors/SCameraVariableCollectionEditor.h
@@ -48,6 +48,9 @@ public:
 	/** Gets the selected variables in the list view. */
 	void GetSelectedVariables(TArray<UCameraVariableAsset*>& OutSelection) const;
 
+	/** Selects the given variable if it is in the list view. */
+	void SelectVariable(UCameraVariableAsset* InItem);
+
 	/** Enter editing mode for the given variable's name. */
 	void RequestRenameVariable(UCameraVariableAsset* InItem, FSimpleDelegate InOnRenamedItem);
 
diff --git a/Engine/Plugins/Cameras/GameplayCameras/Source/GameplayCamerasEditor/Public/Toolkits/CameraVariableCollectionEditorToolkit.h b/Engine/Plugins/Cameras/GameplayCameras/Source/GameplayCamerasEditor/Public/Toolkits/CameraVariableCollectionEditorToolkit.h
--- a/Engine/Plugins/Cameras/GameplayCameras/Source/GameplayCamerasEditor/Public/Toolkits/CameraVariableCollectionEditorToolkit.h
+++ b/Engine/Plugins/Cameras/GameplayCameras/Source/GameplayCamerasEditor/Public/Toolkits/CameraVariableCollectionEditorToolkit.h
@@ -31,6 +31,9 @@ public:
 	FCameraVariableCollectionEditorToolkit(UCameraVariableCollectionEditor* InOwningAssetEditor);
 	~FCameraVariableCollectionEditorToolkit();
 
+	// IAssetEditorInstance interface
+	virtual void FocusWindow(UObject* ObjectToFocusOn = nullptr) override;
+
 protected:
 
 	// FBaseAssetToolkit interface
