Cameras: add GetPlayerController method on the Blueprint camera node evaluator

diff --git a/Engine/Plugins/Cameras/GameplayCameras/Source/GameplayCameras/Private/Nodes/Utility/BlueprintCameraNode.cpp b/Engine/Plugins/Cameras/GameplayCameras/Source/GameplayCameras/Private/Nodes/Utility/BlueprintCameraNode.cpp
--- a/Engine/Plugins/Cameras/GameplayCameras/Source/GameplayCameras/Private/Nodes/Utility/BlueprintCameraNode.cpp
+++ b/Engine/Plugins/Cameras/GameplayCameras/Source/GameplayCameras/Private/Nodes/Utility/BlueprintCameraNode.cpp
@@ -278,6 +278,16 @@ void UBlueprintCameraNodeEvaluator::SetDefaultOwningCameraRigParameters(FBluepri
 	}
 }
 
+APlayerController* UBlueprintCameraNodeEvaluator::GetPlayerController() const
+{
+	if (CurrentContext)
+	{
+		return CurrentContext->GetPlayerController();
+	}
+
+	return nullptr;
+}
+
 UWorld* UBlueprintCameraNodeEvaluator::GetWorld() const
 {
 	if (UWorld* CachedWorld = WeakCachedWorld.Get())
diff --git a/Engine/Plugins/Cameras/GameplayCameras/Source/GameplayCameras/Public/Nodes/Utility/BlueprintCameraNode.h b/Engine/Plugins/Cameras/GameplayCameras/Source/GameplayCameras/Public/Nodes/Utility/BlueprintCameraNode.h
--- a/Engine/Plugins/Cameras/GameplayCameras/Source/GameplayCameras/Public/Nodes/Utility/BlueprintCameraNode.h
+++ b/Engine/Plugins/Cameras/GameplayCameras/Source/GameplayCameras/Public/Nodes/Utility/BlueprintCameraNode.h
@@ -60,6 +60,10 @@ public:
 	UFUNCTION(BlueprintCallable, Category="Evaluation")
 	void SetDefaultOwningCameraRigParameters(FBlueprintCameraEvaluationDataRef TargetCameraData) const;
 
+	/** Gets the player controller that the node is running for, if any. */
+	UFUNCTION(BlueprintPure, Category="Evaluation")
+	APlayerController* GetPlayerController() const;
+
 public:
 
 	using FCameraNodeEvaluatorInitializeParams = UE::Cameras::FCameraNodeEvaluatorInitializeParams;
