Cameras: re-add camera rigs' DisplayName as a deprecated property so that the camera asset toolkit can generate camera rig packages with proper names

diff --git a/Engine/Plugins/Cameras/GameplayCameras/Source/GameplayCameras/Public/Core/CameraObjectInterface.h b/Engine/Plugins/Cameras/GameplayCameras/Source/GameplayCameras/Public/Core/CameraObjectInterface.h
--- a/Engine/Plugins/Cameras/GameplayCameras/Source/GameplayCameras/Public/Core/CameraObjectInterface.h
+++ b/Engine/Plugins/Cameras/GameplayCameras/Source/GameplayCameras/Public/Core/CameraObjectInterface.h
@@ -201,5 +201,18 @@ public:
 
 	/** Returns whether an exposed parameter with the given name exists. */
 	GAMEPLAYCAMERAS_API bool HasBlendableParameter(const FString& ParameterName) const;
+
+public:
+
+	// Deprecated methods.
+	UE_DEPRECATED(5.6, "Camera rigs are now all standalone assets and don't need a separate display name.")
+	FString GetDisplayName() const { return DisplayName_DEPRECATED; }
+
+private:
+
+	// Deprecated properties.
+
+	UPROPERTY()
+	FString DisplayName_DEPRECATED;
 };
 
diff --git a/Engine/Plugins/Cameras/GameplayCameras/Source/GameplayCamerasEditor/Private/Toolkits/CameraAssetEditorToolkit.cpp b/Engine/Plugins/Cameras/GameplayCameras/Source/GameplayCamerasEditor/Private/Toolkits/CameraAssetEditorToolkit.cpp
--- a/Engine/Plugins/Cameras/GameplayCameras/Source/GameplayCamerasEditor/Private/Toolkits/CameraAssetEditorToolkit.cpp
+++ b/Engine/Plugins/Cameras/GameplayCameras/Source/GameplayCamerasEditor/Private/Toolkits/CameraAssetEditorToolkit.cpp
@@ -572,7 +572,14 @@ void FCameraAssetEditorToolkit::UpgradeLegacyCameraAssets()
 		SlowTask.EnterProgressFrame();
 
 		// Name the new camera rigs like this: "<CameraAsset>_<CameraRig>"
-		const FString CameraRigPackageName = FString::Printf(TEXT("%s_%s"), *CameraRigsBaseName, *CameraRig->GetName());
+PRAGMA_DISABLE_DEPRECATION_WARNINGS
+		FString CameraRigName = CameraRig->Interface.GetDisplayName();
+		if (CameraRigName.IsEmpty())
+		{
+			CameraRigName = CameraRig->GetName();
+		}
+PRAGMA_ENABLE_DEPRECATION_WARNINGS
+		const FString CameraRigPackageName = FString::Printf(TEXT("%s_%s"), *CameraRigsBaseName, *CameraRigName);
 
 		FString CameraRigPackagePath;
 		FString CameraRigAssetName;
