Cameras: make the whole plugin backwards compatible with 5.5

diff --git a/Engine/Plugins/Cameras/GameplayCameras/Source/GameplayCameras/Private/Build/CameraObjectInterfaceParameterBuilder.cpp b/Engine/Plugins/Cameras/GameplayCameras/Source/GameplayCameras/Private/Build/CameraObjectInterfaceParameterBuilder.cpp
--- a/Engine/Plugins/Cameras/GameplayCameras/Source/GameplayCameras/Private/Build/CameraObjectInterfaceParameterBuilder.cpp
+++ b/Engine/Plugins/Cameras/GameplayCameras/Source/GameplayCameras/Private/Build/CameraObjectInterfaceParameterBuilder.cpp
@@ -8,6 +8,7 @@
 #include "Core/CameraParameters.h"
 #include "Core/CameraVariableReferences.h"
 #include "Core/ICustomCameraNodeParameterProvider.h"
+#include "Misc/EngineVersionComparison.h"
 #include "StructUtils/PropertyBag.h"
 #include "UObject/UnrealType.h"
 
@@ -176,7 +177,9 @@ void FCameraObjectInterfaceParameterBuilder::AppendDefaultParameterProperties(co
 			FPropertyBagPropertyDesc NewProperty(Definition.ParameterName, ContainerType, PropertyType, PropertyTypeObject);
 			// Make the property bag match the camera interface parameter GUIDs.
 			NewProperty.ID = Definition.ParameterGuid;
+#if UE_VERSION_NEWER_THAN_OR_EQUAL(5,6,0)
 			NewProperty.PropertyFlags |= PropertyFlags;
+#endif
 
 			OutProperties.Add(NewProperty);
 		}
diff --git a/Engine/Plugins/Cameras/GameplayCameras/Source/GameplayCameras/Private/Compat/PostProcessUtils.cpp b/Engine/Plugins/Cameras/GameplayCameras/Source/GameplayCameras/Private/Compat/PostProcessUtils.cpp
new file mode 100644
--- /dev/null
+++ b/Engine/Plugins/Cameras/GameplayCameras/Source/GameplayCameras/Private/Compat/PostProcessUtils.cpp
@@ -0,0 +1,595 @@
+// Copyright Epic Games, Inc. All Rights Reserved.
+
+#include "Compat/PostProcessUtils.h"
+
+#if !UE_VERSION_NEWER_THAN_OR_EQUAL(5,6,0)
+
+#include "Engine/Scene.h"
+
+bool FPostProcessUtils::OverridePostProcessSettings(FPostProcessSettings& ThisFrom, const FPostProcessSettings& OtherTo)
+{
+	bool bAnyOverwritten = false;
+
+#define UE_SET_PP(Name)\
+	if (OtherTo.bOverride_##Name)\
+	{\
+		ThisFrom.bOverride_##Name = true;\
+		ThisFrom.Name = OtherTo.Name;\
+		bAnyOverwritten = true;\
+	}
+
+	{
+		UE_SET_PP(TemperatureType);
+		UE_SET_PP(WhiteTemp);
+		UE_SET_PP(WhiteTint);
+
+		UE_SET_PP(ColorSaturation);
+		UE_SET_PP(ColorContrast);
+		UE_SET_PP(ColorGamma);
+		UE_SET_PP(ColorGain);
+		UE_SET_PP(ColorOffset);
+
+		UE_SET_PP(ColorSaturationShadows);
+		UE_SET_PP(ColorContrastShadows);
+		UE_SET_PP(ColorGammaShadows);
+		UE_SET_PP(ColorGainShadows);
+		UE_SET_PP(ColorOffsetShadows);
+
+		UE_SET_PP(ColorSaturationMidtones);
+		UE_SET_PP(ColorContrastMidtones);
+		UE_SET_PP(ColorGammaMidtones);
+		UE_SET_PP(ColorGainMidtones);
+		UE_SET_PP(ColorOffsetMidtones);
+
+		UE_SET_PP(ColorSaturationHighlights);
+		UE_SET_PP(ColorContrastHighlights);
+		UE_SET_PP(ColorGammaHighlights);
+		UE_SET_PP(ColorGainHighlights);
+		UE_SET_PP(ColorOffsetHighlights);
+
+		UE_SET_PP(ColorCorrectionShadowsMax);
+		UE_SET_PP(ColorCorrectionHighlightsMin);
+		UE_SET_PP(ColorCorrectionHighlightsMax);
+
+		UE_SET_PP(BlueCorrection);
+		UE_SET_PP(ExpandGamut);
+		UE_SET_PP(ToneCurveAmount);
+
+		UE_SET_PP(FilmSlope);
+		UE_SET_PP(FilmToe);
+		UE_SET_PP(FilmShoulder);
+		UE_SET_PP(FilmBlackClip);
+		UE_SET_PP(FilmWhiteClip);
+
+		UE_SET_PP(SceneColorTint);
+		UE_SET_PP(SceneFringeIntensity);
+		UE_SET_PP(ChromaticAberrationStartOffset);
+		UE_SET_PP(BloomIntensity);
+		UE_SET_PP(BloomThreshold);
+		UE_SET_PP(Bloom1Tint);
+		UE_SET_PP(BloomSizeScale);
+		UE_SET_PP(Bloom1Size);
+		UE_SET_PP(Bloom2Tint);
+		UE_SET_PP(Bloom2Size);
+		UE_SET_PP(Bloom3Tint);
+		UE_SET_PP(Bloom3Size);
+		UE_SET_PP(Bloom4Tint);
+		UE_SET_PP(Bloom4Size);
+		UE_SET_PP(Bloom5Tint);
+		UE_SET_PP(Bloom5Size);
+		UE_SET_PP(Bloom6Tint);
+		UE_SET_PP(Bloom6Size);
+		UE_SET_PP(BloomDirtMaskIntensity);
+		UE_SET_PP(BloomDirtMaskTint);
+		UE_SET_PP(BloomConvolutionScatterDispersion);
+		UE_SET_PP(BloomConvolutionSize);
+		UE_SET_PP(BloomConvolutionCenterUV);
+		UE_SET_PP(BloomConvolutionPreFilterMin);
+		UE_SET_PP(BloomConvolutionPreFilterMax);
+		UE_SET_PP(BloomConvolutionPreFilterMult);
+		UE_SET_PP(AmbientCubemapIntensity);
+		UE_SET_PP(AmbientCubemapTint);
+		UE_SET_PP(CameraShutterSpeed);
+		UE_SET_PP(CameraISO);
+		UE_SET_PP(AutoExposureLowPercent);
+		UE_SET_PP(AutoExposureHighPercent);
+		UE_SET_PP(AutoExposureMinBrightness);
+		UE_SET_PP(AutoExposureMaxBrightness);
+		UE_SET_PP(AutoExposureSpeedUp);
+		UE_SET_PP(AutoExposureSpeedDown);
+		UE_SET_PP(AutoExposureBias);
+		UE_SET_PP(HistogramLogMin);
+		UE_SET_PP(HistogramLogMax);
+		UE_SET_PP(LocalExposureMethod);
+		UE_SET_PP(LocalExposureContrastScale_DEPRECATED);
+		UE_SET_PP(LocalExposureHighlightContrastScale);
+		UE_SET_PP(LocalExposureShadowContrastScale);
+		UE_SET_PP(LocalExposureHighlightThreshold);
+		UE_SET_PP(LocalExposureShadowThreshold);
+		UE_SET_PP(LocalExposureDetailStrength);
+		UE_SET_PP(LocalExposureBlurredLuminanceBlend);
+		UE_SET_PP(LocalExposureBlurredLuminanceKernelSizePercent);
+		UE_SET_PP(LocalExposureMiddleGreyBias);
+		UE_SET_PP(LensFlareIntensity);
+		UE_SET_PP(LensFlareTint);
+		UE_SET_PP(LensFlareBokehSize);
+		UE_SET_PP(LensFlareThreshold);
+		UE_SET_PP(VignetteIntensity);
+		UE_SET_PP(Sharpen);
+		UE_SET_PP(FilmGrainIntensity);
+		UE_SET_PP(FilmGrainIntensityShadows);
+		UE_SET_PP(FilmGrainIntensityMidtones);
+		UE_SET_PP(FilmGrainIntensityHighlights);
+		UE_SET_PP(FilmGrainShadowsMax);
+		UE_SET_PP(FilmGrainHighlightsMin);
+		UE_SET_PP(FilmGrainHighlightsMax);
+		UE_SET_PP(FilmGrainTexelSize);
+		UE_SET_PP(AmbientOcclusionIntensity);
+		UE_SET_PP(AmbientOcclusionStaticFraction);
+		UE_SET_PP(AmbientOcclusionRadius);
+		UE_SET_PP(AmbientOcclusionFadeDistance);
+		UE_SET_PP(AmbientOcclusionFadeRadius);
+		UE_SET_PP(AmbientOcclusionDistance_DEPRECATED);
+		UE_SET_PP(AmbientOcclusionPower);
+		UE_SET_PP(AmbientOcclusionBias);
+		UE_SET_PP(AmbientOcclusionQuality);
+		UE_SET_PP(AmbientOcclusionMipBlend);
+		UE_SET_PP(AmbientOcclusionMipScale);
+		UE_SET_PP(AmbientOcclusionMipThreshold);
+		UE_SET_PP(AmbientOcclusionTemporalBlendWeight);
+		UE_SET_PP(IndirectLightingColor);
+		UE_SET_PP(IndirectLightingIntensity);
+
+		UE_SET_PP(DepthOfFieldFocalDistance);
+
+		UE_SET_PP(DepthOfFieldFstop);
+		UE_SET_PP(DepthOfFieldMinFstop);
+		UE_SET_PP(DepthOfFieldSensorWidth);
+		UE_SET_PP(DepthOfFieldSqueezeFactor);
+		UE_SET_PP(DepthOfFieldDepthBlurRadius);
+		UE_SET_PP(DepthOfFieldUseHairDepth)
+		UE_SET_PP(DepthOfFieldDepthBlurAmount);
+		UE_SET_PP(DepthOfFieldFocalRegion);
+		UE_SET_PP(DepthOfFieldNearTransitionRegion);
+		UE_SET_PP(DepthOfFieldFarTransitionRegion);
+		UE_SET_PP(DepthOfFieldScale);
+		UE_SET_PP(DepthOfFieldNearBlurSize);
+		UE_SET_PP(DepthOfFieldFarBlurSize);
+		UE_SET_PP(DepthOfFieldOcclusion);
+		UE_SET_PP(DepthOfFieldSkyFocusDistance);
+		UE_SET_PP(DepthOfFieldVignetteSize);
+
+		UE_SET_PP(MotionBlurAmount);
+		UE_SET_PP(MotionBlurMax);
+		UE_SET_PP(MotionBlurPerObjectSize);
+		UE_SET_PP(ScreenSpaceReflectionQuality);
+		UE_SET_PP(ScreenSpaceReflectionIntensity);
+		UE_SET_PP(ScreenSpaceReflectionMaxRoughness);
+
+		UE_SET_PP(TranslucencyType);
+		UE_SET_PP(RayTracingTranslucencyMaxRoughness);
+		UE_SET_PP(RayTracingTranslucencyRefractionRays);
+		UE_SET_PP(RayTracingTranslucencySamplesPerPixel);
+		UE_SET_PP(RayTracingTranslucencyShadows);
+		UE_SET_PP(RayTracingTranslucencyRefraction);
+
+		UE_SET_PP(DynamicGlobalIlluminationMethod);
+		UE_SET_PP(LumenSurfaceCacheResolution);
+		UE_SET_PP(LumenSceneLightingQuality);
+		UE_SET_PP(LumenSceneDetail);
+		UE_SET_PP(LumenSceneViewDistance);
+		UE_SET_PP(LumenSceneLightingUpdateSpeed);
+		UE_SET_PP(LumenFinalGatherQuality);
+		UE_SET_PP(LumenFinalGatherLightingUpdateSpeed);
+		UE_SET_PP(LumenFinalGatherScreenTraces);
+		UE_SET_PP(LumenMaxTraceDistance);
+
+		UE_SET_PP(LumenDiffuseColorBoost);
+		UE_SET_PP(LumenSkylightLeaking);
+		UE_SET_PP(LumenFullSkylightLeakingDistance);
+
+		UE_SET_PP(LumenRayLightingMode);
+		UE_SET_PP(LumenReflectionsScreenTraces);
+		UE_SET_PP(LumenFrontLayerTranslucencyReflections);
+		UE_SET_PP(LumenMaxRoughnessToTraceReflections);
+		UE_SET_PP(LumenMaxReflectionBounces);
+		UE_SET_PP(LumenMaxRefractionBounces);
+		UE_SET_PP(ReflectionMethod);
+		UE_SET_PP(LumenReflectionQuality);
+		UE_SET_PP(RayTracingAO);
+		UE_SET_PP(RayTracingAOSamplesPerPixel);
+		UE_SET_PP(RayTracingAOIntensity);
+		UE_SET_PP(RayTracingAORadius);
+
+		UE_SET_PP(PathTracingMaxBounces);
+		UE_SET_PP(PathTracingSamplesPerPixel);
+		UE_SET_PP(PathTracingMaxPathIntensity);
+		UE_SET_PP(PathTracingEnableEmissiveMaterials);
+		UE_SET_PP(PathTracingEnableReferenceDOF);
+		UE_SET_PP(PathTracingEnableReferenceAtmosphere);
+		UE_SET_PP(PathTracingEnableDenoiser);
+		UE_SET_PP(PathTracingIncludeEmissive);
+		UE_SET_PP(PathTracingIncludeDiffuse);
+		UE_SET_PP(PathTracingIncludeIndirectDiffuse);
+		UE_SET_PP(PathTracingIncludeSpecular);
+		UE_SET_PP(PathTracingIncludeIndirectSpecular);
+		UE_SET_PP(PathTracingIncludeVolume);
+		UE_SET_PP(PathTracingIncludeIndirectVolume);
+
+		UE_SET_PP(DepthOfFieldBladeCount);
+
+		// This is no bOverride_AmbientCubemap so just see if it is set.
+		if (OtherTo.AmbientCubemap)
+		{
+			ThisFrom.AmbientCubemap = OtherTo.AmbientCubemap;
+			bAnyOverwritten = true;
+		}
+
+		UE_SET_PP(ColorGradingIntensity);
+		UE_SET_PP(ColorGradingLUT);
+
+		UE_SET_PP(BloomDirtMask);
+		UE_SET_PP(BloomMethod);
+		UE_SET_PP(BloomConvolutionTexture)
+		UE_SET_PP(FilmGrainTexture)
+
+		UE_SET_PP(BloomConvolutionBufferScale);
+
+		UE_SET_PP(AutoExposureBiasCurve);
+		UE_SET_PP(AutoExposureMeterMask);
+		UE_SET_PP(LocalExposureHighlightContrastCurve);
+		UE_SET_PP(LocalExposureShadowContrastCurve);
+		UE_SET_PP(LensFlareBokehShape);
+
+		if (OtherTo.bOverride_LensFlareTints)
+		{
+			for (uint32 i = 0; i < 8; ++i)
+			{
+				ThisFrom.LensFlareTints[i] = OtherTo.LensFlareTints[i];
+			}
+			bAnyOverwritten = true;
+		}
+
+		if (OtherTo.bOverride_MobileHQGaussian)
+		{
+			ThisFrom.bMobileHQGaussian = OtherTo.bMobileHQGaussian;
+			bAnyOverwritten = true;
+		}
+
+		UE_SET_PP(AutoExposureMethod);
+		UE_SET_PP(AmbientOcclusionRadiusInWS);
+		UE_SET_PP(MotionBlurTargetFPS);
+		UE_SET_PP(AutoExposureApplyPhysicalCameraExposure);
+		UE_SET_PP(UserFlags);
+	}
+
+#undef UE_SET_PP
+
+	// Override blendable objects.
+	bAnyOverwritten |= (ThisFrom.WeightedBlendables.Array.Num() > 0 || OtherTo.WeightedBlendables.Array.Num() > 0);
+	ThisFrom.WeightedBlendables = OtherTo.WeightedBlendables;
+
+	return bAnyOverwritten;
+}
+
+bool FPostProcessUtils::BlendPostProcessSettings(FPostProcessSettings& ThisFrom, const FPostProcessSettings& OtherTo, float BlendFactor)
+{
+	if (UNLIKELY(BlendFactor <= 0.f))
+	{
+		return false;
+	}
+
+	BlendFactor = FMath::Clamp(BlendFactor, 0.f, 1.f);
+	const bool bShouldFlip = (BlendFactor >= 0.5f);
+
+	bool bAnyOverwritten = false;
+
+#define UE_SET_PP(Name)\
+	if ((OtherTo.bOverride_##Name || ThisFrom.bOverride_##Name) && bShouldFlip)\
+	{\
+		ThisFrom.bOverride_##Name = true;\
+		ThisFrom.Name = OtherTo.Name;\
+		bAnyOverwritten = true;\
+	}
+
+#define UE_LERP_PP(Name)\
+	if (OtherTo.bOverride_##Name || ThisFrom.bOverride_##Name)\
+	{\
+		ThisFrom.bOverride_##Name = true;\
+		ThisFrom.Name = FMath::Lerp(ThisFrom.Name, OtherTo.Name, BlendFactor);\
+		bAnyOverwritten = true;\
+	}
+
+	{
+		UE_SET_PP(TemperatureType);
+		UE_LERP_PP(WhiteTemp);
+		UE_LERP_PP(WhiteTint);
+
+		UE_LERP_PP(ColorSaturation);
+		UE_LERP_PP(ColorContrast);
+		UE_LERP_PP(ColorGamma);
+		UE_LERP_PP(ColorGain);
+		UE_LERP_PP(ColorOffset);
+
+		UE_LERP_PP(ColorSaturationShadows);
+		UE_LERP_PP(ColorContrastShadows);
+		UE_LERP_PP(ColorGammaShadows);
+		UE_LERP_PP(ColorGainShadows);
+		UE_LERP_PP(ColorOffsetShadows);
+
+		UE_LERP_PP(ColorSaturationMidtones);
+		UE_LERP_PP(ColorContrastMidtones);
+		UE_LERP_PP(ColorGammaMidtones);
+		UE_LERP_PP(ColorGainMidtones);
+		UE_LERP_PP(ColorOffsetMidtones);
+
+		UE_LERP_PP(ColorSaturationHighlights);
+		UE_LERP_PP(ColorContrastHighlights);
+		UE_LERP_PP(ColorGammaHighlights);
+		UE_LERP_PP(ColorGainHighlights);
+		UE_LERP_PP(ColorOffsetHighlights);
+
+		UE_LERP_PP(ColorCorrectionShadowsMax);
+		UE_LERP_PP(ColorCorrectionHighlightsMin);
+		UE_LERP_PP(ColorCorrectionHighlightsMax);
+
+		UE_LERP_PP(BlueCorrection);
+		UE_LERP_PP(ExpandGamut);
+		UE_LERP_PP(ToneCurveAmount);
+
+		UE_LERP_PP(FilmSlope);
+		UE_LERP_PP(FilmToe);
+		UE_LERP_PP(FilmShoulder);
+		UE_LERP_PP(FilmBlackClip);
+		UE_LERP_PP(FilmWhiteClip);
+
+		UE_LERP_PP(SceneColorTint);
+		UE_LERP_PP(SceneFringeIntensity);
+		UE_LERP_PP(ChromaticAberrationStartOffset);
+		UE_LERP_PP(BloomIntensity);
+		UE_LERP_PP(BloomThreshold);
+		UE_LERP_PP(Bloom1Tint);
+		UE_LERP_PP(BloomSizeScale);
+		UE_LERP_PP(Bloom1Size);
+		UE_LERP_PP(Bloom2Tint);
+		UE_LERP_PP(Bloom2Size);
+		UE_LERP_PP(Bloom3Tint);
+		UE_LERP_PP(Bloom3Size);
+		UE_LERP_PP(Bloom4Tint);
+		UE_LERP_PP(Bloom4Size);
+		UE_LERP_PP(Bloom5Tint);
+		UE_LERP_PP(Bloom5Size);
+		UE_LERP_PP(Bloom6Tint);
+		UE_LERP_PP(Bloom6Size);
+		UE_LERP_PP(BloomDirtMaskIntensity);
+		UE_LERP_PP(BloomDirtMaskTint);
+		UE_LERP_PP(BloomConvolutionScatterDispersion);
+		UE_LERP_PP(BloomConvolutionSize);
+		UE_LERP_PP(BloomConvolutionCenterUV);
+		UE_LERP_PP(BloomConvolutionPreFilterMin);
+		UE_LERP_PP(BloomConvolutionPreFilterMax);
+		UE_LERP_PP(BloomConvolutionPreFilterMult);
+		UE_LERP_PP(AmbientCubemapIntensity);
+		UE_LERP_PP(AmbientCubemapTint);
+		UE_LERP_PP(CameraShutterSpeed);
+		UE_LERP_PP(CameraISO);
+		UE_LERP_PP(AutoExposureLowPercent);
+		UE_LERP_PP(AutoExposureHighPercent);
+		UE_LERP_PP(AutoExposureMinBrightness);
+		UE_LERP_PP(AutoExposureMaxBrightness);
+		UE_LERP_PP(AutoExposureSpeedUp);
+		UE_LERP_PP(AutoExposureSpeedDown);
+		UE_LERP_PP(AutoExposureBias);
+		UE_LERP_PP(HistogramLogMin);
+		UE_LERP_PP(HistogramLogMax);
+		UE_SET_PP(LocalExposureMethod);
+		UE_LERP_PP(LocalExposureContrastScale_DEPRECATED);
+		UE_LERP_PP(LocalExposureHighlightContrastScale);
+		UE_LERP_PP(LocalExposureShadowContrastScale);
+		UE_LERP_PP(LocalExposureHighlightThreshold);
+		UE_LERP_PP(LocalExposureShadowThreshold);
+		UE_LERP_PP(LocalExposureDetailStrength);
+		UE_LERP_PP(LocalExposureBlurredLuminanceBlend);
+		UE_LERP_PP(LocalExposureBlurredLuminanceKernelSizePercent);
+		UE_LERP_PP(LocalExposureMiddleGreyBias);
+		UE_LERP_PP(LensFlareIntensity);
+		UE_LERP_PP(LensFlareTint);
+		UE_LERP_PP(LensFlareBokehSize);
+		UE_LERP_PP(LensFlareThreshold);
+		UE_LERP_PP(VignetteIntensity);
+		UE_LERP_PP(Sharpen);
+		UE_LERP_PP(FilmGrainIntensity);
+		UE_LERP_PP(FilmGrainIntensityShadows);
+		UE_LERP_PP(FilmGrainIntensityMidtones);
+		UE_LERP_PP(FilmGrainIntensityHighlights);
+		UE_LERP_PP(FilmGrainShadowsMax);
+		UE_LERP_PP(FilmGrainHighlightsMin);
+		UE_LERP_PP(FilmGrainHighlightsMax);
+		UE_LERP_PP(FilmGrainTexelSize);
+		UE_LERP_PP(AmbientOcclusionIntensity);
+		UE_LERP_PP(AmbientOcclusionStaticFraction);
+		UE_LERP_PP(AmbientOcclusionRadius);
+		UE_LERP_PP(AmbientOcclusionFadeDistance);
+		UE_LERP_PP(AmbientOcclusionFadeRadius);
+		UE_LERP_PP(AmbientOcclusionDistance_DEPRECATED);
+		UE_LERP_PP(AmbientOcclusionPower);
+		UE_LERP_PP(AmbientOcclusionBias);
+		UE_LERP_PP(AmbientOcclusionQuality);
+		UE_LERP_PP(AmbientOcclusionMipBlend);
+		UE_LERP_PP(AmbientOcclusionMipScale);
+		UE_LERP_PP(AmbientOcclusionMipThreshold);
+		UE_LERP_PP(AmbientOcclusionTemporalBlendWeight);
+		UE_LERP_PP(IndirectLightingColor);
+		UE_LERP_PP(IndirectLightingIntensity);
+
+		if (OtherTo.bOverride_DepthOfFieldFocalDistance)
+		{
+			if (ThisFrom.DepthOfFieldFocalDistance == 0.0f || OtherTo.DepthOfFieldFocalDistance == 0.0f)
+			{
+				ThisFrom.DepthOfFieldFocalDistance = OtherTo.DepthOfFieldFocalDistance;
+			}
+			else
+			{
+				ThisFrom.DepthOfFieldFocalDistance = FMath::Lerp(ThisFrom.DepthOfFieldFocalDistance, OtherTo.DepthOfFieldFocalDistance, BlendFactor);
+			}
+			bAnyOverwritten = true;
+		}
+		UE_LERP_PP(DepthOfFieldFstop);
+		UE_LERP_PP(DepthOfFieldMinFstop);
+		UE_LERP_PP(DepthOfFieldSensorWidth);
+		UE_LERP_PP(DepthOfFieldSqueezeFactor);
+		UE_LERP_PP(DepthOfFieldDepthBlurRadius);
+		UE_SET_PP(DepthOfFieldUseHairDepth)
+		UE_LERP_PP(DepthOfFieldDepthBlurAmount);
+		UE_LERP_PP(DepthOfFieldFocalRegion);
+		UE_LERP_PP(DepthOfFieldNearTransitionRegion);
+		UE_LERP_PP(DepthOfFieldFarTransitionRegion);
+		UE_LERP_PP(DepthOfFieldScale);
+		UE_LERP_PP(DepthOfFieldNearBlurSize);
+		UE_LERP_PP(DepthOfFieldFarBlurSize);
+		UE_LERP_PP(DepthOfFieldOcclusion);
+		UE_LERP_PP(DepthOfFieldSkyFocusDistance);
+		UE_LERP_PP(DepthOfFieldVignetteSize);
+		UE_LERP_PP(MotionBlurAmount);
+		UE_LERP_PP(MotionBlurMax);
+		UE_LERP_PP(MotionBlurPerObjectSize);
+		UE_LERP_PP(ScreenSpaceReflectionQuality);
+		UE_LERP_PP(ScreenSpaceReflectionIntensity);
+		UE_LERP_PP(ScreenSpaceReflectionMaxRoughness);
+
+		UE_SET_PP(TranslucencyType);
+		UE_SET_PP(RayTracingTranslucencyMaxRoughness);
+		UE_SET_PP(RayTracingTranslucencyRefractionRays);
+		UE_SET_PP(RayTracingTranslucencySamplesPerPixel);
+		UE_SET_PP(RayTracingTranslucencyShadows);
+		UE_SET_PP(RayTracingTranslucencyRefraction);
+
+		UE_SET_PP(DynamicGlobalIlluminationMethod);
+		UE_SET_PP(LumenSurfaceCacheResolution);
+		UE_SET_PP(LumenSceneLightingQuality);
+		UE_SET_PP(LumenSceneDetail);
+		UE_SET_PP(LumenSceneViewDistance);
+		UE_SET_PP(LumenSceneLightingUpdateSpeed);
+		UE_SET_PP(LumenFinalGatherQuality);
+		UE_SET_PP(LumenFinalGatherLightingUpdateSpeed);
+		UE_SET_PP(LumenFinalGatherScreenTraces);
+		UE_SET_PP(LumenMaxTraceDistance);
+
+		UE_LERP_PP(LumenDiffuseColorBoost);
+		UE_LERP_PP(LumenSkylightLeaking);
+		UE_LERP_PP(LumenFullSkylightLeakingDistance);
+
+		UE_SET_PP(LumenRayLightingMode);
+		UE_SET_PP(LumenReflectionsScreenTraces);
+		UE_SET_PP(LumenFrontLayerTranslucencyReflections);
+		UE_SET_PP(LumenMaxRoughnessToTraceReflections);
+		UE_SET_PP(LumenMaxReflectionBounces);
+		UE_SET_PP(LumenMaxRefractionBounces);
+		UE_SET_PP(ReflectionMethod);
+		UE_SET_PP(LumenReflectionQuality);
+		UE_SET_PP(RayTracingAO);
+		UE_SET_PP(RayTracingAOSamplesPerPixel);
+		UE_SET_PP(RayTracingAOIntensity);
+		UE_SET_PP(RayTracingAORadius);
+
+		UE_SET_PP(PathTracingMaxBounces);
+		UE_SET_PP(PathTracingSamplesPerPixel);
+		UE_LERP_PP(PathTracingMaxPathIntensity);
+		UE_SET_PP(PathTracingEnableEmissiveMaterials);
+		UE_SET_PP(PathTracingEnableReferenceDOF);
+		UE_SET_PP(PathTracingEnableReferenceAtmosphere);
+		UE_SET_PP(PathTracingEnableDenoiser);
+		UE_SET_PP(PathTracingIncludeEmissive);
+		UE_SET_PP(PathTracingIncludeDiffuse);
+		UE_SET_PP(PathTracingIncludeIndirectDiffuse);
+		UE_SET_PP(PathTracingIncludeSpecular);
+		UE_SET_PP(PathTracingIncludeIndirectSpecular);
+		UE_SET_PP(PathTracingIncludeVolume);
+		UE_SET_PP(PathTracingIncludeIndirectVolume);
+
+		UE_SET_PP(DepthOfFieldBladeCount);
+
+		// No cubemap blending (only supported for FFinalPostProcessSettings)
+		if (bShouldFlip)
+		{
+			ThisFrom.AmbientCubemap = OtherTo.AmbientCubemap;
+			bAnyOverwritten = true;
+		}
+
+		// No color grading texture blending (only supported for FFinalPostProcessSettings)
+		UE_LERP_PP(ColorGradingIntensity);
+		UE_SET_PP(ColorGradingLUT);
+
+		UE_SET_PP(BloomDirtMask);
+		UE_SET_PP(BloomMethod);
+		UE_SET_PP(BloomConvolutionTexture)
+		UE_SET_PP(FilmGrainTexture)
+
+		// Flipping this instead of blending, as per the comment in SceneView.cpp
+		UE_SET_PP(BloomConvolutionBufferScale);
+
+		UE_SET_PP(AutoExposureBiasCurve);
+		UE_SET_PP(AutoExposureMeterMask);
+		UE_SET_PP(LocalExposureHighlightContrastCurve);
+		UE_SET_PP(LocalExposureShadowContrastCurve);
+		UE_SET_PP(LensFlareBokehShape);
+
+		if (OtherTo.bOverride_LensFlareTints)
+		{
+			for (uint32 i = 0; i < 8; ++i)
+			{
+				ThisFrom.LensFlareTints[i] = FMath::Lerp(ThisFrom.LensFlareTints[i], OtherTo.LensFlareTints[i], BlendFactor);
+			}
+			bAnyOverwritten = true;
+		}
+
+		if (OtherTo.bOverride_MobileHQGaussian && bShouldFlip)
+		{
+			ThisFrom.bMobileHQGaussian = OtherTo.bMobileHQGaussian;
+			bAnyOverwritten = true;
+		}
+
+		UE_SET_PP(AutoExposureMethod);
+		UE_SET_PP(AmbientOcclusionRadiusInWS);
+		UE_SET_PP(MotionBlurTargetFPS);
+		UE_SET_PP(AutoExposureApplyPhysicalCameraExposure);
+		UE_SET_PP(UserFlags);
+	}
+
+#undef UE_SET_PP
+#undef UE_LERP_PP
+
+	// Blend out old objects, blend in new objects.
+	// This uses an O(n2) lookup for now since 99% of the time we have zero blendables, and the rest
+	// of the time we have only one or two.
+	bAnyOverwritten |= (ThisFrom.WeightedBlendables.Array.Num() > 0 || OtherTo.WeightedBlendables.Array.Num() > 0);
+	for (FWeightedBlendable& ThisBlendable : ThisFrom.WeightedBlendables.Array)
+	{
+		const bool bFoundInOther = OtherTo.WeightedBlendables.Array.ContainsByPredicate(
+				[&ThisBlendable](const FWeightedBlendable& OtherBlendable)
+				{
+					return ThisBlendable.Object == OtherBlendable.Object;
+				});
+		if (!bFoundInOther)
+		{
+			ThisBlendable.Weight *= (1.0f - BlendFactor);
+		}
+	}
+	for (const FWeightedBlendable& OtherBlendable : OtherTo.WeightedBlendables.Array)
+	{
+		const bool bFoundInThis = ThisFrom.WeightedBlendables.Array.ContainsByPredicate(
+				[&OtherBlendable](FWeightedBlendable& ThisBlendable)
+				{
+					return ThisBlendable.Object == OtherBlendable.Object;
+				});
+		if (!bFoundInThis)
+		{
+			ThisFrom.WeightedBlendables.Array.Emplace(OtherBlendable.Weight * BlendFactor, OtherBlendable.Object);
+		}
+	}
+
+	return bAnyOverwritten;
+}
+
+#endif  // !5.6.0
+
diff --git a/Engine/Plugins/Cameras/GameplayCameras/Source/GameplayCameras/Private/Compat/PostProcessUtils.h b/Engine/Plugins/Cameras/GameplayCameras/Source/GameplayCameras/Private/Compat/PostProcessUtils.h
new file mode 100644
--- /dev/null
+++ b/Engine/Plugins/Cameras/GameplayCameras/Source/GameplayCameras/Private/Compat/PostProcessUtils.h
@@ -0,0 +1,23 @@
+// Copyright Epic Games, Inc. All Rights Reserved.
+
+#pragma once
+
+#include "CoreTypes.h"
+#include "Misc/EngineVersionComparison.h"
+
+struct FPostProcessSettings;
+
+#if !UE_VERSION_NEWER_THAN_OR_EQUAL(5,6,0)
+
+/**
+ * Copy of UE 5.6's FPostProcessUtils class for backwards compatible code.
+ * Implementation only supports post process settings that existed up to UE 5.5.
+ */
+struct FPostProcessUtils
+{
+	static bool OverridePostProcessSettings(FPostProcessSettings& ThisFrom, const FPostProcessSettings& OtherTo);
+	static bool BlendPostProcessSettings(FPostProcessSettings& ThisFrom, const FPostProcessSettings& OtherTo, float BlendFactor);
+};
+
+#endif  // <5.6.0
+
diff --git a/Engine/Plugins/Cameras/GameplayCameras/Source/GameplayCameras/Private/Core/PostProcessSettingsCollection.cpp b/Engine/Plugins/Cameras/GameplayCameras/Source/GameplayCameras/Private/Core/PostProcessSettingsCollection.cpp
--- a/Engine/Plugins/Cameras/GameplayCameras/Source/GameplayCameras/Private/Core/PostProcessSettingsCollection.cpp
+++ b/Engine/Plugins/Cameras/GameplayCameras/Source/GameplayCameras/Private/Core/PostProcessSettingsCollection.cpp
@@ -2,14 +2,24 @@
 
 #include "Core/PostProcessSettingsCollection.h"
 
+#include "Misc/EngineVersionComparison.h"
+
+#if UE_VERSION_NEWER_THAN_OR_EQUAL(5,6,0)
 #include "Engine/PostProcessUtils.h"
+#else
+#include "Compat/PostProcessUtils.h"
+#endif
 
 namespace UE::Cameras
 {
 
 void FPostProcessSettingsCollection::Reset()
 {
+#if UE_VERSION_NEWER_THAN_OR_EQUAL(5,6,0)
 	PostProcessSettings = FPostProcessSettings::GetDefault();
+#else
+	PostProcessSettings = FPostProcessSettings();
+#endif
 }
 
 void FPostProcessSettingsCollection::OverrideAll(const FPostProcessSettingsCollection& OtherCollection)
@@ -51,7 +61,12 @@ void FPostProcessSettingsCollection::InternalLerp(const FPostProcessSettings& To
 void FPostProcessSettingsCollection::Serialize(FArchive& Ar)
 {
 	UScriptStruct* PostProcessSettingsStruct = FPostProcessSettings::StaticStruct();
+#if UE_VERSION_NEWER_THAN_OR_EQUAL(5,6,0)
 	PostProcessSettingsStruct->SerializeItem(Ar, &PostProcessSettings, &FPostProcessSettings::GetDefault());
+#else
+	FPostProcessSettings DefaultSettings;
+	PostProcessSettingsStruct->SerializeItem(Ar, &PostProcessSettings, &DefaultSettings);
+#endif
 }
 
 }  // namespace UE::Cameras
diff --git a/Engine/Plugins/Cameras/GameplayCameras/Source/GameplayCameras/Private/Debug/CameraDebugRenderer.cpp b/Engine/Plugins/Cameras/GameplayCameras/Source/GameplayCameras/Private/Debug/CameraDebugRenderer.cpp
--- a/Engine/Plugins/Cameras/GameplayCameras/Source/GameplayCameras/Private/Debug/CameraDebugRenderer.cpp
+++ b/Engine/Plugins/Cameras/GameplayCameras/Source/GameplayCameras/Private/Debug/CameraDebugRenderer.cpp
@@ -15,6 +15,7 @@
 #include "Engine/World.h"
 #include "HAL/IConsoleManager.h"
 #include "Math/Box2D.h"
+#include "Misc/EngineVersionComparison.h"
 #include "SceneView.h"
 
 #if UE_GAMEPLAY_CAMERAS_DEBUG
@@ -645,7 +646,11 @@ void FCameraDebugRenderer::DrawText(const FVector3d& WorldPosition, const FVecto
 
 ULineBatchComponent* FCameraDebugRenderer::GetDebugLineBatcher() const
 {
+#if UE_VERSION_NEWER_THAN_OR_EQUAL(5,6,0)
 	return World ? World->GetLineBatcher(UWorld::ELineBatcherType::Foreground) : nullptr;
+#else
+	return World ? World->ForegroundLineBatcher : nullptr;
+#endif
 }
 
 void FCameraDebugRenderer::SkipAttachedBlocks()
diff --git a/Engine/Plugins/Cameras/GameplayCameras/Source/GameplayCameras/Private/Debug/PostProcessSettingsDebugBlock.cpp b/Engine/Plugins/Cameras/GameplayCameras/Source/GameplayCameras/Private/Debug/PostProcessSettingsDebugBlock.cpp
--- a/Engine/Plugins/Cameras/GameplayCameras/Source/GameplayCameras/Private/Debug/PostProcessSettingsDebugBlock.cpp
+++ b/Engine/Plugins/Cameras/GameplayCameras/Source/GameplayCameras/Private/Debug/PostProcessSettingsDebugBlock.cpp
@@ -11,6 +11,7 @@
 #include "Engine/Texture2D.h"
 #include "Engine/TextureCube.h"
 #include "Math/ColorList.h"
+#include "Misc/EngineVersionComparison.h"
 
 #if UE_GAMEPLAY_CAMERAS_DEBUG
 
@@ -135,8 +136,10 @@ void FPostProcessSettingsDebugBlock::OnDebugDraw(const FCameraDebugBlockDrawPara
 		UE_DRAW_PP(LocalExposureDetailStrength);
 		UE_DRAW_PP(LocalExposureBlurredLuminanceBlend);
 		UE_DRAW_PP(LocalExposureBlurredLuminanceKernelSizePercent);
+#if UE_VERSION_NEWER_THAN_OR_EQUAL(5,6,0)
 		UE_DRAW_PP(LocalExposureHighlightThresholdStrength);
 		UE_DRAW_PP(LocalExposureShadowThresholdStrength);
+#endif
 		UE_DRAW_PP(LocalExposureMiddleGreyBias);
 		UE_DRAW_PP(LensFlareIntensity);
 		UE_DRAW_PP(LensFlareTint);
@@ -186,6 +189,7 @@ void FPostProcessSettingsDebugBlock::OnDebugDraw(const FCameraDebugBlockDrawPara
 		UE_DRAW_PP(DepthOfFieldOcclusion);
 		UE_DRAW_PP(DepthOfFieldSkyFocusDistance);
 		UE_DRAW_PP(DepthOfFieldVignetteSize);
+#if UE_VERSION_NEWER_THAN_OR_EQUAL(5,6,0)
 		UE_DRAW_PP(DepthOfFieldAspectRatioScalar);
 		UE_DRAW_PP(DepthOfFieldPetzvalBokeh);
 		UE_DRAW_PP(DepthOfFieldPetzvalBokehFalloff);
@@ -203,6 +207,7 @@ void FPostProcessSettingsDebugBlock::OnDebugDraw(const FCameraDebugBlockDrawPara
 			}
 			Renderer.AddText(TEXT("\n"));
 		}
+#endif
 
 		UE_DRAW_PP(MotionBlurAmount);
 		UE_DRAW_PP(MotionBlurMax);
@@ -231,7 +236,9 @@ void FPostProcessSettingsDebugBlock::OnDebugDraw(const FCameraDebugBlockDrawPara
 
 		UE_DRAW_PP(LumenDiffuseColorBoost);
 		UE_DRAW_PP(LumenSkylightLeaking);
+#if UE_VERSION_NEWER_THAN_OR_EQUAL(5,6,0)
 		UE_DRAW_PP(LumenSkylightLeakingTint);
+#endif
 		UE_DRAW_PP(LumenFullSkylightLeakingDistance);
 
 		UE_DRAW_PP(LumenRayLightingMode);
diff --git a/Engine/Plugins/Cameras/GameplayCameras/Source/GameplayCameras/Private/GameFramework/GameplayCameraComponentBase.cpp b/Engine/Plugins/Cameras/GameplayCameras/Source/GameplayCameras/Private/GameFramework/GameplayCameraComponentBase.cpp
--- a/Engine/Plugins/Cameras/GameplayCameras/Source/GameplayCameras/Private/GameFramework/GameplayCameraComponentBase.cpp
+++ b/Engine/Plugins/Cameras/GameplayCameras/Source/GameplayCameras/Private/GameFramework/GameplayCameraComponentBase.cpp
@@ -18,10 +18,16 @@
 #include "IGameplayCamerasModule.h"
 #include "Kismet/GameplayStatics.h"
 #include "Misc/AssertionMacros.h"
-#include "PrimitiveDrawInterface.h"
+#include "Misc/EngineVersionComparison.h"
 #include "SceneView.h"
 #include "ShowFlags.h"
 
+#if UE_VERSION_NEWER_THAN_OR_EQUAL(5,6,0)
+#include "PrimitiveDrawInterface.h"
+#else
+#include "SceneManagement.h"
+#endif
+
 #if WITH_EDITOR
 #include "Editor.h"
 #include "LevelEditorViewport.h"
@@ -792,7 +798,11 @@ void UGameplayCameraComponentBase::OnDrawVisualizationHUD(const FViewport* Viewp
 	{
 		const AActor* OwnerActor = GetOwner();
 
+#if UE_VERSION_NEWER_THAN_OR_EQUAL(5,6,0)
 		const AActor* ViewActor = SceneView->ViewActor.Get();
+#else
+		const AActor* ViewActor = SceneView->ViewActor;
+#endif
 		const bool bIsLockedToCamera = (ViewActor == OwnerActor);
 
 		FCameraSystemEditorPreviewParams Params;
diff --git a/Engine/Plugins/Cameras/GameplayCameras/Source/GameplayCameras/Private/GameplayCamerasModule.cpp b/Engine/Plugins/Cameras/GameplayCameras/Source/GameplayCameras/Private/GameplayCamerasModule.cpp
--- a/Engine/Plugins/Cameras/GameplayCameras/Source/GameplayCameras/Private/GameplayCamerasModule.cpp
+++ b/Engine/Plugins/Cameras/GameplayCameras/Source/GameplayCameras/Private/GameplayCamerasModule.cpp
@@ -5,6 +5,7 @@
 #include "Debug/CameraDebugColors.h"
 #include "GameplayCameras.h"
 #include "Logging/MessageLog.h"
+#include "Misc/EngineVersionComparison.h"
 #include "Modules/ModuleManager.h"
 #include "MovieScene/MovieSceneCameraParameterInstantiator.h"
 #include "MovieSceneSection.h"
@@ -122,15 +123,19 @@ private:
 	void RegisterMovieSceneIntegration()
 	{
 #if WITH_EDITORONLY_DATA
+#if UE_VERSION_NEWER_THAN_OR_EQUAL(5,6,0)
 		OnMovieSceneSectionAddedToTrackHandle = UMovieSceneSection::GetOnSectionAddedToTrack().AddStatic(
 				&UMovieSceneCameraParameterInstantiator::OnMovieSceneSectionAddedToTrack);
+#endif
 #endif
 	}
 
 	void UnregisterMovieSceneIntegration()
 	{
 #if WITH_EDITORONLY_DATA
+#if UE_VERSION_NEWER_THAN_OR_EQUAL(5,6,0)
 		UMovieSceneSection::GetOnSectionAddedToTrack().Remove(OnMovieSceneSectionAddedToTrackHandle);
+#endif
 #endif
 	}
 
diff --git a/Engine/Plugins/Cameras/GameplayCameras/Source/GameplayCameras/Private/Helpers/CameraObjectInterfaceParameterOverrideHelper.cpp b/Engine/Plugins/Cameras/GameplayCameras/Source/GameplayCameras/Private/Helpers/CameraObjectInterfaceParameterOverrideHelper.cpp
--- a/Engine/Plugins/Cameras/GameplayCameras/Source/GameplayCameras/Private/Helpers/CameraObjectInterfaceParameterOverrideHelper.cpp
+++ b/Engine/Plugins/Cameras/GameplayCameras/Source/GameplayCameras/Private/Helpers/CameraObjectInterfaceParameterOverrideHelper.cpp
@@ -7,6 +7,7 @@
 #include "Core/CameraObjectInterfaceParameterDefinition.h"
 #include "Core/CameraParameters.h"
 #include "Core/CameraVariableTable.h"
+#include "Misc/EngineVersionComparison.h"
 #include "StructUtils/PropertyBag.h"
 
 namespace UE::Cameras
@@ -67,7 +68,11 @@ void ApplyBlendableParameterOverride(
 		return;
 	}
 
+#if UE_VERSION_NEWER_THAN_OR_EQUAL(5,6,0)
 	TValueOrError<FStructView, EPropertyBagResult> ParameterValueOrError = PropertyBag.GetValueStruct(PropertyBagPropertyDesc);
+#else
+	TValueOrError<FStructView, EPropertyBagResult> ParameterValueOrError = PropertyBag.GetValueStruct(PropertyBagPropertyDesc.Name);
+#endif
 	if (!ensureMsgf(
 				ParameterValueOrError.HasValue() && !ParameterValueOrError.HasError(),
 				TEXT("Camera parameter has no valid value! Error: %s"),
@@ -195,13 +200,21 @@ void ApplyDataParameterSingleOverride(
 	{
 		case ECameraContextDataType::Name:
 			{
+#if UE_VERSION_NEWER_THAN_OR_EQUAL(5,6,0)
 				TValueOrError<FName, EPropertyBagResult> ParameterValueOrError = PropertyBag.GetValueName(PropertyBagPropertyDesc);
+#else
+				TValueOrError<FName, EPropertyBagResult> ParameterValueOrError = PropertyBag.GetValueName(PropertyBagPropertyDesc.Name);
+#endif
 				ApplyDataParameterOverride(ParameterDefinition, ParameterValueOrError, ContextDataTable);
 			}
 			break;
 		case ECameraContextDataType::String:
 			{
+#if UE_VERSION_NEWER_THAN_OR_EQUAL(5,6,0)
 				TValueOrError<FString, EPropertyBagResult> ParameterValueOrError = PropertyBag.GetValueString(PropertyBagPropertyDesc);
+#else
+				TValueOrError<FString, EPropertyBagResult> ParameterValueOrError = PropertyBag.GetValueString(PropertyBagPropertyDesc.Name);
+#endif
 				ApplyDataParameterOverride(ParameterDefinition, ParameterValueOrError, ContextDataTable);
 			}
 			break;
@@ -210,7 +223,11 @@ void ApplyDataParameterSingleOverride(
 				const UEnum* EnumType = Cast<const UEnum>(ParameterDefinition.DataTypeObject);
 				if (ensure(EnumType))
 				{
+#if UE_VERSION_NEWER_THAN_OR_EQUAL(5,6,0)
 					TValueOrError<uint8, EPropertyBagResult> ParameterValueOrError = PropertyBag.GetValueEnum(PropertyBagPropertyDesc, EnumType);
+#else
+					TValueOrError<uint8, EPropertyBagResult> ParameterValueOrError = PropertyBag.GetValueEnum(PropertyBagPropertyDesc.Name, EnumType);
+#endif
 					ApplyDataParameterOverride(ParameterDefinition, ParameterValueOrError, ContextDataTable);
 				}
 			}
@@ -220,20 +237,32 @@ void ApplyDataParameterSingleOverride(
 				const UScriptStruct* StructType = Cast<const UScriptStruct>(ParameterDefinition.DataTypeObject);
 				if (ensure(StructType))
 				{
+#if UE_VERSION_NEWER_THAN_OR_EQUAL(5,6,0)
 					TValueOrError<FStructView, EPropertyBagResult> ParameterValueOrError = PropertyBag.GetValueStruct(PropertyBagPropertyDesc, StructType);
+#else
+					TValueOrError<FStructView, EPropertyBagResult> ParameterValueOrError = PropertyBag.GetValueStruct(PropertyBagPropertyDesc.Name, StructType);
+#endif
 					ApplyDataParameterOverride(ParameterDefinition, ParameterValueOrError, ContextDataTable);
 				}
 			}
 			break;
 		case ECameraContextDataType::Object:
 			{
+#if UE_VERSION_NEWER_THAN_OR_EQUAL(5,6,0)
 				TValueOrError<UObject*, EPropertyBagResult> ParameterValueOrError = PropertyBag.GetValueObject(PropertyBagPropertyDesc);
+#else
+				TValueOrError<UObject*, EPropertyBagResult> ParameterValueOrError = PropertyBag.GetValueObject(PropertyBagPropertyDesc.Name);
+#endif
 				ApplyDataParameterOverride(ParameterDefinition, ParameterValueOrError, ContextDataTable);
 			}
 			break;
 		case ECameraContextDataType::Class:
 			{
+#if UE_VERSION_NEWER_THAN_OR_EQUAL(5,6,0)
 				TValueOrError<UClass*, EPropertyBagResult> ParameterValueOrError = PropertyBag.GetValueClass(PropertyBagPropertyDesc);
+#else
+				TValueOrError<UClass*, EPropertyBagResult> ParameterValueOrError = PropertyBag.GetValueClass(PropertyBagPropertyDesc.Name);
+#endif
 				ApplyDataParameterOverride(ParameterDefinition, ParameterValueOrError, ContextDataTable);
 			}
 			break;
@@ -250,7 +279,11 @@ void ApplyDataParameterArrayOverride(
 		const FPropertyBagPropertyDesc& PropertyBagPropertyDesc,
 		FCameraContextDataTable& ContextDataTable)
 {
+#if UE_VERSION_NEWER_THAN_OR_EQUAL(5,6,0)
 	TValueOrError<const FPropertyBagArrayRef, EPropertyBagResult> ArrayOrError = PropertyBag.GetArrayRef(PropertyBagPropertyDesc);
+#else
+	TValueOrError<const FPropertyBagArrayRef, EPropertyBagResult> ArrayOrError = PropertyBag.GetArrayRef(PropertyBagPropertyDesc.Name);
+#endif
 	if (!ensureMsgf(
 				ArrayOrError.HasValue() && !ArrayOrError.HasError(),
 				TEXT("Camera parameter has no valid value! Error: %s"),
diff --git a/Engine/Plugins/Cameras/GameplayCameras/Source/GameplayCameras/Private/MovieScene/MovieSceneCameraParameterInstantiator.cpp b/Engine/Plugins/Cameras/GameplayCameras/Source/GameplayCameras/Private/MovieScene/MovieSceneCameraParameterInstantiator.cpp
--- a/Engine/Plugins/Cameras/GameplayCameras/Source/GameplayCameras/Private/MovieScene/MovieSceneCameraParameterInstantiator.cpp
+++ b/Engine/Plugins/Cameras/GameplayCameras/Source/GameplayCameras/Private/MovieScene/MovieSceneCameraParameterInstantiator.cpp
@@ -2,6 +2,8 @@
 
 #include "MovieScene/MovieSceneCameraParameterInstantiator.h"
 
+#if UE_VERSION_NEWER_THAN_OR_EQUAL(5,6,0)
+
 #include "Core/CameraAsset.h"
 #include "Core/CameraAssetReference.h"
 #include "Core/CameraRigAsset.h"
@@ -271,3 +273,5 @@ void UMovieSceneCameraParameterInstantiator::OnRun(FSystemTaskPrerequisites& InP
 	.RunInline_PerEntity(&Linker->EntityManager, FSetupCameraParameterOverrideTask(PreAnimatedStorage));
 }
 
+#endif  // 5.6.0
+
diff --git a/Engine/Plugins/Cameras/GameplayCameras/Source/GameplayCameras/Private/MovieScene/MovieSceneCameraParameterInstantiator.h b/Engine/Plugins/Cameras/GameplayCameras/Source/GameplayCameras/Private/MovieScene/MovieSceneCameraParameterInstantiator.h
--- a/Engine/Plugins/Cameras/GameplayCameras/Source/GameplayCameras/Private/MovieScene/MovieSceneCameraParameterInstantiator.h
+++ b/Engine/Plugins/Cameras/GameplayCameras/Source/GameplayCameras/Private/MovieScene/MovieSceneCameraParameterInstantiator.h
@@ -2,6 +2,10 @@
 
 #pragma once
 
+#include "Misc/EngineVersionComparison.h"
+
+#if UE_VERSION_NEWER_THAN_OR_EQUAL(5,6,0)
+
 #include "EntitySystem/IMovieSceneEntityDecorator.h"
 #include "EntitySystem/MovieSceneEntityInstantiatorSystem.h"
 #include "UObject/ObjectMacros.h"
@@ -59,4 +63,5 @@ private:
 	TSharedPtr<UE::Cameras::FPreAnimatedCameraParameterStorage> PreAnimatedStorage;
 };
 
+#endif  // UE_VERSION_NEWER_THAN_OR_EQUAL(5,6,0)
 
diff --git a/Engine/Plugins/Cameras/GameplayCameras/Source/GameplayCamerasEditor/Private/Compat/EditorCompat.h b/Engine/Plugins/Cameras/GameplayCameras/Source/GameplayCamerasEditor/Private/Compat/EditorCompat.h
new file mode 100644
--- /dev/null
+++ b/Engine/Plugins/Cameras/GameplayCameras/Source/GameplayCamerasEditor/Private/Compat/EditorCompat.h
@@ -0,0 +1,25 @@
+// Copyright Epic Games, Inc. All Rights Reserved.
+
+#pragma once
+
+#include "Math/MathFwd.h"
+#include "Misc/EngineVersionComparison.h"
+
+// Useful aliases for changes in editor APIs.
+
+#if UE_VERSION_NEWER_THAN_OR_EQUAL(5,6,0)
+
+using FSlateCompatVector2f = FVector2f;
+
+using FPerformGraphActionLocation = const FVector2f&;
+
+#else
+
+using FSlateCompatVector2f = FVector2D;
+
+using FPerformGraphActionLocation = const FVector2D;
+
+using FDelegateUserObjectConst = const void*;
+
+#endif  // 5.6.0
+
diff --git a/Engine/Plugins/Cameras/GameplayCameras/Source/GameplayCamerasEditor/Private/Compat/SResizeBox.cpp b/Engine/Plugins/Cameras/GameplayCameras/Source/GameplayCamerasEditor/Private/Compat/SResizeBox.cpp
new file mode 100644
--- /dev/null
+++ b/Engine/Plugins/Cameras/GameplayCameras/Source/GameplayCamerasEditor/Private/Compat/SResizeBox.cpp
@@ -0,0 +1,119 @@
+// Copyright Epic Games, Inc. All Rights Reserved.
+
+#include "Compat/SResizeBox.h"
+
+#if !UE_VERSION_NEWER_THAN_OR_EQUAL(5,6,0)
+
+#include "Brushes/SlateColorBrush.h"
+#include "Widgets/Layout/SBox.h"
+
+void SVerticalResizeBox::Construct(const FArguments& InArgs)
+{
+	ContentHeight = InArgs._ContentHeight;
+	HandleHeight = InArgs._HandleHeight;
+	HandleColor = InArgs._HandleColor;
+	HandleHighlightColor = InArgs._HandleHighlightColor;
+	HandleBrush = FSlateColorBrush(FLinearColor::White);
+	ContentHeightChanged = InArgs._ContentHeightChanged;
+
+	ChildSlot
+	[
+		SNew(SBox)
+		.HeightOverride(this, &SVerticalResizeBox::GetHeightOverride)
+		.Padding(FMargin(0, 0, 0, HandleHeight))
+		[
+			InArgs._Content.Widget
+		]
+	];
+}
+
+FReply SVerticalResizeBox::OnMouseButtonDown(const FGeometry& MyGeometry, const FPointerEvent& MouseEvent)
+{
+	if (MouseEvent.GetEffectingButton() == EKeys::LeftMouseButton)
+	{
+		FVector2f MouseLocation = MyGeometry.AbsoluteToLocal(MouseEvent.GetScreenSpacePosition());
+		if (MyGeometry.GetLocalSize().Y - MouseLocation.Y < HandleHeight)
+		{
+			DragStartLocation = MouseLocation.Y;
+			DragStartContentHeight = ContentHeight.Get();
+			return FReply::Handled().CaptureMouse(SharedThis(this));
+		}
+	}
+	return FReply::Unhandled();
+}
+
+FReply SVerticalResizeBox::OnMouseButtonUp(const FGeometry& MyGeometry, const FPointerEvent& MouseEvent)
+{
+	if (this->HasMouseCapture())
+	{
+		return FReply::Handled().ReleaseMouseCapture();
+	}
+	return FReply::Unhandled();
+}
+
+FReply SVerticalResizeBox::OnMouseMove(const FGeometry& MyGeometry, const FPointerEvent& MouseEvent)
+{
+	FVector2f MouseLocation = MyGeometry.AbsoluteToLocal(MouseEvent.GetScreenSpacePosition());
+	LastMouseLocation = MouseLocation.Y;
+	if (this->HasMouseCapture())
+	{
+		float NewContentHeight = DragStartContentHeight + (MouseLocation.Y - DragStartLocation);
+		if (ContentHeight.IsBound() && ContentHeightChanged.IsBound())
+		{
+			ContentHeightChanged.Execute(NewContentHeight);
+		}
+		else
+		{
+			ContentHeight = NewContentHeight;
+		}
+		return FReply::Handled();
+	}
+	return FReply::Unhandled();
+}
+
+FCursorReply SVerticalResizeBox::OnCursorQuery(const FGeometry& MyGeometry, const FPointerEvent& CursorEvent) const
+{
+	FVector2f CursorLocation = MyGeometry.AbsoluteToLocal(CursorEvent.GetScreenSpacePosition());
+	if (MyGeometry.GetLocalSize().Y - CursorLocation.Y < HandleHeight)
+	{
+		return FCursorReply::Cursor(EMouseCursor::ResizeUpDown);
+	}
+	return FCursorReply::Unhandled();
+}
+
+int32 SVerticalResizeBox::OnPaint(const FPaintArgs& Args, const FGeometry& AllottedGeometry, const FSlateRect& MyClippingRect, FSlateWindowElementList& OutDrawElements, int32 LayerId, const FWidgetStyle& InWidgetStyle, bool bParentEnabled) const
+{
+	FLinearColor HandleBoxColor;
+	int32 HandleLayerId = LayerId + 1;
+	FVector2f LocalSize = AllottedGeometry.GetLocalSize();
+	if (IsHovered() && LastMouseLocation.IsSet() && LastMouseLocation.GetValue() >= LocalSize.Y - HandleHeight && LastMouseLocation.GetValue() <= LocalSize.Y)
+	{
+		HandleBoxColor = HandleHighlightColor.Get();
+	}
+	else
+	{
+		HandleBoxColor = HandleColor.Get();
+	}
+
+	FVector2f HandleLocation(0, AllottedGeometry.GetLocalSize().Y - HandleHeight);
+	FVector2f HandleSize(AllottedGeometry.GetLocalSize().X, HandleHeight);
+	FSlateDrawElement::MakeBox
+	(
+		OutDrawElements,
+		HandleLayerId,
+		AllottedGeometry.ToPaintGeometry(HandleSize, FSlateLayoutTransform(HandleLocation)),
+		&HandleBrush,
+		ESlateDrawEffect::None,
+		HandleBoxColor
+	);
+
+	return SCompoundWidget::OnPaint(Args, AllottedGeometry, MyClippingRect, OutDrawElements, HandleLayerId, InWidgetStyle, bParentEnabled);
+}
+
+FOptionalSize SVerticalResizeBox::GetHeightOverride() const
+{
+	return ContentHeight.Get() + HandleHeight;
+}
+
+#endif  // <5.6.0
+
diff --git a/Engine/Plugins/Cameras/GameplayCameras/Source/GameplayCamerasEditor/Private/Compat/SResizeBox.h b/Engine/Plugins/Cameras/GameplayCameras/Source/GameplayCamerasEditor/Private/Compat/SResizeBox.h
new file mode 100644
--- /dev/null
+++ b/Engine/Plugins/Cameras/GameplayCameras/Source/GameplayCamerasEditor/Private/Compat/SResizeBox.h
@@ -0,0 +1,62 @@
+// Copyright Epic Games, Inc. All Rights Reserved.
+
+#pragma once
+
+#include "Misc/EngineVersionComparison.h"
+
+#if !UE_VERSION_NEWER_THAN_OR_EQUAL(5,6,0)
+
+#include "Widgets/SCompoundWidget.h"
+#include "Widgets/DeclarativeSyntaxSupport.h"
+#include "Styling/SlateBrush.h"
+#include "Types/SlateStructs.h"
+
+class SVerticalResizeBox : public SCompoundWidget
+{
+public:
+	DECLARE_DELEGATE_OneParam(FOnContentHeightChanged, float);
+
+public:
+	SLATE_BEGIN_ARGS(SVerticalResizeBox)
+		: _HandleHeight(5)
+		, _ContentHeight(50)
+		, _HandleColor(FLinearColor(0.0f, 0.0f, 0.0f, 0.0f))
+		, _HandleHighlightColor(FLinearColor(1.0f, 1.0f, 1.0f, 0.5f))
+	{}
+		SLATE_ARGUMENT(float, HandleHeight)
+		SLATE_ATTRIBUTE(float, ContentHeight)
+		SLATE_ATTRIBUTE(FLinearColor, HandleColor)
+		SLATE_ATTRIBUTE(FLinearColor, HandleHighlightColor)
+		SLATE_EVENT(FOnContentHeightChanged, ContentHeightChanged)
+		SLATE_DEFAULT_SLOT(FArguments, Content);
+	SLATE_END_ARGS();
+
+	void Construct(const FArguments& InArgs);
+
+	virtual FReply OnMouseButtonDown(const FGeometry& MyGeometry, const FPointerEvent& MouseEvent) override;
+	virtual FReply OnMouseButtonUp(const FGeometry& MyGeometry, const FPointerEvent& MouseEvent) override;
+	virtual FReply OnMouseMove(const FGeometry& MyGeometry, const FPointerEvent& MouseEvent) override;
+	virtual FCursorReply OnCursorQuery(const FGeometry& MyGeometry, const FPointerEvent& CursorEvent) const override;
+	virtual int32 OnPaint(const FPaintArgs& Args, const FGeometry& AllottedGeometry, const FSlateRect& MyClippingRect, FSlateWindowElementList& OutDrawElements, int32 LayerId, const FWidgetStyle& InWidgetStyle, bool bParentEnabled) const override;
+
+private:
+	FOptionalSize GetHeightOverride() const;
+
+private:
+	TOptional<float> LastMouseLocation;
+
+	TAttribute<float> ContentHeight;
+	float HandleHeight;
+
+	float DragStartLocation;
+	float DragStartContentHeight;
+
+	TAttribute<FLinearColor> HandleColor;
+	TAttribute<FLinearColor> HandleHighlightColor;
+	FSlateBrush HandleBrush;
+
+	FOnContentHeightChanged ContentHeightChanged;
+};
+
+#endif  // <5.6.0
+
diff --git a/Engine/Plugins/Cameras/GameplayCameras/Source/GameplayCamerasEditor/Private/Customizations/RichCurveDetailsCustomizations.cpp b/Engine/Plugins/Cameras/GameplayCameras/Source/GameplayCamerasEditor/Private/Customizations/RichCurveDetailsCustomizations.cpp
--- a/Engine/Plugins/Cameras/GameplayCameras/Source/GameplayCamerasEditor/Private/Customizations/RichCurveDetailsCustomizations.cpp
+++ b/Engine/Plugins/Cameras/GameplayCameras/Source/GameplayCamerasEditor/Private/Customizations/RichCurveDetailsCustomizations.cpp
@@ -16,13 +16,13 @@
 #include "ICurveEditorBounds.h"
 #include "IDetailChildrenBuilder.h"
 #include "IGameplayCamerasEditorModule.h"
+#include "Misc/EngineVersionComparison.h"
 #include "PropertyEditorModule.h"
 #include "RichCurveEditorModel.h"
 #include "SCurveEditorPanel.h"
 #include "SCurveKeyDetailPanel.h"
 #include "Slate/SlateTextures.h"
 #include "Styles/GameplayCamerasEditorStyle.h"
-#include "SResizeBox.h"
 #include "Toolkits/ToolkitManager.h"
 #include "Tree/CurveEditorTree.h"
 #include "Tree/SCurveEditorTree.h"
@@ -33,6 +33,12 @@
 #include "Widgets/SBoxPanel.h"
 #include "Widgets/SViewport.h"
 
+#if UE_VERSION_NEWER_THAN_OR_EQUAL(5,6,0)
+#include "SResizeBox.h"
+#else
+#include "Compat/SResizeBox.h"
+#endif
+
 #define LOCTEXT_NAMESPACE "RichCurveDetailsCustomizations"
 
 namespace UE::Cameras
diff --git a/Engine/Plugins/Cameras/GameplayCameras/Source/GameplayCamerasEditor/Private/Editors/CameraNodeGraphDragDropOp.cpp b/Engine/Plugins/Cameras/GameplayCameras/Source/GameplayCamerasEditor/Private/Editors/CameraNodeGraphDragDropOp.cpp
--- a/Engine/Plugins/Cameras/GameplayCameras/Source/GameplayCamerasEditor/Private/Editors/CameraNodeGraphDragDropOp.cpp
+++ b/Engine/Plugins/Cameras/GameplayCameras/Source/GameplayCamerasEditor/Private/Editors/CameraNodeGraphDragDropOp.cpp
@@ -38,7 +38,7 @@ FReply FCameraNodeGraphInterfaceParameterDragDropOp::ExecuteDragOver(TSharedPtr<
 	return FReply::Handled();
 }
 
-FReply FCameraNodeGraphInterfaceParameterDragDropOp::ExecuteDrop(TSharedPtr<SGraphEditor> GraphEditor, const UE::Slate::FDeprecateVector2DParameter& NewLocation)
+FReply FCameraNodeGraphInterfaceParameterDragDropOp::ExecuteDrop(TSharedPtr<SGraphEditor> GraphEditor, const FSlateCompatVector2f& NewLocation)
 {
 	if (!InterfaceParameter || InterfaceParameter->bHasGraphNode)
 	{
diff --git a/Engine/Plugins/Cameras/GameplayCameras/Source/GameplayCamerasEditor/Private/Editors/CameraNodeGraphSchema.cpp b/Engine/Plugins/Cameras/GameplayCameras/Source/GameplayCamerasEditor/Private/Editors/CameraNodeGraphSchema.cpp
--- a/Engine/Plugins/Cameras/GameplayCameras/Source/GameplayCamerasEditor/Private/Editors/CameraNodeGraphSchema.cpp
+++ b/Engine/Plugins/Cameras/GameplayCameras/Source/GameplayCamerasEditor/Private/Editors/CameraNodeGraphSchema.cpp
@@ -512,7 +512,7 @@ FCameraNodeGraphSchemaAction_NewInterfaceParameterNode::FCameraNodeGraphSchemaAc
 {
 }
 
-UEdGraphNode* FCameraNodeGraphSchemaAction_NewInterfaceParameterNode::PerformAction(UEdGraph* ParentGraph, UEdGraphPin* FromPin, const FVector2f& Location, bool bSelectNewNode)
+UEdGraphNode* FCameraNodeGraphSchemaAction_NewInterfaceParameterNode::PerformAction(UEdGraph* ParentGraph, UEdGraphPin* FromPin, FPerformGraphActionLocation Location, bool bSelectNewNode)
 {
 	UObjectTreeGraph* ObjectTreeGraph = Cast<UObjectTreeGraph>(ParentGraph);
 	if (!ensure(ObjectTreeGraph))
@@ -585,7 +585,7 @@ FCameraNodeGraphSchemaAction_AddInterfaceParameterNode::FCameraNodeGraphSchemaAc
 {
 }
 
-UEdGraphNode* FCameraNodeGraphSchemaAction_AddInterfaceParameterNode::PerformAction(UEdGraph* ParentGraph, UEdGraphPin* FromPin, const FVector2f& Location, bool bSelectNewNode)
+UEdGraphNode* FCameraNodeGraphSchemaAction_AddInterfaceParameterNode::PerformAction(UEdGraph* ParentGraph, UEdGraphPin* FromPin, FPerformGraphActionLocation Location, bool bSelectNewNode)
 {
 	if (!InterfaceParameter || InterfaceParameter->bHasGraphNode)
 	{
diff --git a/Engine/Plugins/Cameras/GameplayCameras/Source/GameplayCamerasEditor/Private/Editors/ObjectTreeDragDropOp.cpp b/Engine/Plugins/Cameras/GameplayCameras/Source/GameplayCamerasEditor/Private/Editors/ObjectTreeDragDropOp.cpp
--- a/Engine/Plugins/Cameras/GameplayCameras/Source/GameplayCamerasEditor/Private/Editors/ObjectTreeDragDropOp.cpp
+++ b/Engine/Plugins/Cameras/GameplayCameras/Source/GameplayCamerasEditor/Private/Editors/ObjectTreeDragDropOp.cpp
@@ -60,7 +60,7 @@ FReply FObjectTreeClassDragDropOp::ExecuteDragOver(TSharedPtr<SGraphEditor> Grap
 	return FReply::Handled();
 }
 
-FReply FObjectTreeClassDragDropOp::ExecuteDrop(TSharedPtr<SGraphEditor> GraphEditor, const UE::Slate::FDeprecateVector2DParameter& NewLocation)
+FReply FObjectTreeClassDragDropOp::ExecuteDrop(TSharedPtr<SGraphEditor> GraphEditor, const FSlateCompatVector2f& NewLocation)
 {
 	const FScopedTransaction Transaction(LOCTEXT("DropObjectClasses", "Drop New Nodes"));
 
@@ -69,7 +69,7 @@ FReply FObjectTreeClassDragDropOp::ExecuteDrop(TSharedPtr<SGraphEditor> GraphEdi
 
 	GraphEditor->ClearSelectionSet();
 
-	FVector2f CurLocation = NewLocation;
+	FSlateCompatVector2f CurLocation = NewLocation;
 	for (UClass* PlaceableClass : PlaceableClasses)
 	{
 		FObjectTreeGraphSchemaAction_NewNode Action;
@@ -77,7 +77,7 @@ FReply FObjectTreeClassDragDropOp::ExecuteDrop(TSharedPtr<SGraphEditor> GraphEdi
 		UEdGraphNode* NewNode = Action.PerformAction(Graph, nullptr, CurLocation, false);
 		GraphEditor->SetNodeSelection(NewNode, true);
 
-		CurLocation += FVector2f(20.0f, 20.0f);
+		CurLocation += FSlateCompatVector2f(20.0f, 20.0f);
 	}
 
 	return FReply::Handled();
diff --git a/Engine/Plugins/Cameras/GameplayCameras/Source/GameplayCamerasEditor/Private/Editors/ObjectTreeGraphCommentNode.cpp b/Engine/Plugins/Cameras/GameplayCameras/Source/GameplayCamerasEditor/Private/Editors/ObjectTreeGraphCommentNode.cpp
--- a/Engine/Plugins/Cameras/GameplayCameras/Source/GameplayCamerasEditor/Private/Editors/ObjectTreeGraphCommentNode.cpp
+++ b/Engine/Plugins/Cameras/GameplayCameras/Source/GameplayCamerasEditor/Private/Editors/ObjectTreeGraphCommentNode.cpp
@@ -99,7 +99,7 @@ void UObjectTreeGraphCommentNode::PostPlacedNewNode()
 	}
 }
 
-void UObjectTreeGraphCommentNode::ResizeNode(const FVector2f& NewSize)
+void UObjectTreeGraphCommentNode::ResizeNode(const FSlateCompatVector2f& NewSize)
 {
 	Super::ResizeNode(NewSize);
 
diff --git a/Engine/Plugins/Cameras/GameplayCameras/Source/GameplayCamerasEditor/Private/Editors/ObjectTreeGraphNode.cpp b/Engine/Plugins/Cameras/GameplayCameras/Source/GameplayCamerasEditor/Private/Editors/ObjectTreeGraphNode.cpp
--- a/Engine/Plugins/Cameras/GameplayCameras/Source/GameplayCamerasEditor/Private/Editors/ObjectTreeGraphNode.cpp
+++ b/Engine/Plugins/Cameras/GameplayCameras/Source/GameplayCamerasEditor/Private/Editors/ObjectTreeGraphNode.cpp
@@ -58,11 +58,13 @@ FLinearColor UObjectTreeGraphNode::GetNodeTitleColor() const
 	return NodeContext.ObjectClassConfigs.NodeTitleColor().Get(NodeContext.GraphConfig.DefaultGraphNodeTitleColor);
 }
 
+#if UE_VERSION_NEWER_THAN_OR_EQUAL(5,6,0)
 FLinearColor UObjectTreeGraphNode::GetNodeTitleTextColor() const
 {
 	const FNodeContext NodeContext = GetNodeContext();
 	return NodeContext.ObjectClassConfigs.NodeTitleTextColor().Get(NodeContext.GraphConfig.DefaultGraphNodeTitleTextColor);
 }
+#endif
 
 FLinearColor UObjectTreeGraphNode::GetNodeBodyTintColor() const
 {
diff --git a/Engine/Plugins/Cameras/GameplayCameras/Source/GameplayCamerasEditor/Private/Editors/ObjectTreeGraphSchema.cpp b/Engine/Plugins/Cameras/GameplayCameras/Source/GameplayCamerasEditor/Private/Editors/ObjectTreeGraphSchema.cpp
--- a/Engine/Plugins/Cameras/GameplayCameras/Source/GameplayCamerasEditor/Private/Editors/ObjectTreeGraphSchema.cpp
+++ b/Engine/Plugins/Cameras/GameplayCameras/Source/GameplayCamerasEditor/Private/Editors/ObjectTreeGraphSchema.cpp
@@ -1457,7 +1457,11 @@ FObjectTreeGraphSchemaAction_NewNode::FObjectTreeGraphSchemaAction_NewNode(FText
 {
 }
 
+#if UE_VERSION_NEWER_THAN_OR_EQUAL(5,6,0)
 UEdGraphNode* FObjectTreeGraphSchemaAction_NewNode::PerformAction(UEdGraph* ParentGraph, UEdGraphPin* FromPin, const FVector2f& Location, bool bSelectNewNode)
+#else
+UEdGraphNode* FObjectTreeGraphSchemaAction_NewNode::PerformAction(UEdGraph* ParentGraph, UEdGraphPin* FromPin, const FVector2D Location, bool bSelectNewNode)
+#endif
 {
 	UObjectTreeGraph* ObjectTreeGraph = Cast<UObjectTreeGraph>(ParentGraph);
 	if (!ensure(ObjectTreeGraph))
@@ -1535,7 +1539,11 @@ FObjectTreeGraphSchemaAction_NewComment::FObjectTreeGraphSchemaAction_NewComment
 {
 }
 
+#if UE_VERSION_NEWER_THAN_OR_EQUAL(5,6,0)
 UEdGraphNode* FObjectTreeGraphSchemaAction_NewComment::PerformAction(UEdGraph* ParentGraph, UEdGraphPin* FromPin, const FVector2f& Location, bool bSelectNewNode)
+#else
+UEdGraphNode* FObjectTreeGraphSchemaAction_NewComment::PerformAction(UEdGraph* ParentGraph, UEdGraphPin* FromPin, const FVector2D Location, bool bSelectNewNode)
+#endif
 {
 	UObjectTreeGraph* ObjectTreeGraph = Cast<UObjectTreeGraph>(ParentGraph);
 	if (!ensure(ObjectTreeGraph))
diff --git a/Engine/Plugins/Cameras/GameplayCameras/Source/GameplayCamerasEditor/Private/Editors/SCameraNodeGraphEditor.cpp b/Engine/Plugins/Cameras/GameplayCameras/Source/GameplayCamerasEditor/Private/Editors/SCameraNodeGraphEditor.cpp
--- a/Engine/Plugins/Cameras/GameplayCameras/Source/GameplayCamerasEditor/Private/Editors/SCameraNodeGraphEditor.cpp
+++ b/Engine/Plugins/Cameras/GameplayCameras/Source/GameplayCamerasEditor/Private/Editors/SCameraNodeGraphEditor.cpp
@@ -2,6 +2,7 @@
 
 #include "Editors/SCameraNodeGraphEditor.h"
 
+#include "Compat/EditorCompat.h"
 #include "Editors/CameraNodeGraphDragDropOp.h"
 #include "SGraphPanel.h"
 
@@ -24,7 +25,7 @@ FReply SCameraNodeGraphEditor::OnDrop(const FGeometry& MyGeometry, const FDragDr
 	if (InterfaceParameterOp)
 	{
 		SGraphPanel* GraphPanel = GraphEditor->GetGraphPanel();
-		FVector2D NewLocation = GraphPanel->PanelCoordToGraphCoord(MyGeometry.AbsoluteToLocal(DragDropEvent.GetScreenSpacePosition()));
+		FSlateCompatVector2f NewLocation = GraphPanel->PanelCoordToGraphCoord(MyGeometry.AbsoluteToLocal(DragDropEvent.GetScreenSpacePosition()));
 
 		return InterfaceParameterOp->ExecuteDrop(GraphEditor, NewLocation);
 	}
diff --git a/Engine/Plugins/Cameras/GameplayCameras/Source/GameplayCamerasEditor/Private/Editors/SObjectTreeGraphCommentNode.cpp b/Engine/Plugins/Cameras/GameplayCameras/Source/GameplayCamerasEditor/Private/Editors/SObjectTreeGraphCommentNode.cpp
--- a/Engine/Plugins/Cameras/GameplayCameras/Source/GameplayCamerasEditor/Private/Editors/SObjectTreeGraphCommentNode.cpp
+++ b/Engine/Plugins/Cameras/GameplayCameras/Source/GameplayCamerasEditor/Private/Editors/SObjectTreeGraphCommentNode.cpp
@@ -12,7 +12,7 @@ void SObjectTreeGraphCommentNode::Construct(const FArguments& InArgs)
 	ObjectGraphNode = InArgs._GraphNode;
 }
 
-void SObjectTreeGraphCommentNode::MoveTo(const FVector2f& NewPosition, FNodeSet& NodeFilter, bool bMarkDirty)
+void SObjectTreeGraphCommentNode::MoveTo(const FSlateCompatVector2f& NewPosition, FNodeSet& NodeFilter, bool bMarkDirty)
 {
 	SGraphNodeComment::MoveTo(NewPosition, NodeFilter, bMarkDirty);
 
diff --git a/Engine/Plugins/Cameras/GameplayCameras/Source/GameplayCamerasEditor/Private/Editors/SObjectTreeGraphEditor.cpp b/Engine/Plugins/Cameras/GameplayCameras/Source/GameplayCamerasEditor/Private/Editors/SObjectTreeGraphEditor.cpp
--- a/Engine/Plugins/Cameras/GameplayCameras/Source/GameplayCamerasEditor/Private/Editors/SObjectTreeGraphEditor.cpp
+++ b/Engine/Plugins/Cameras/GameplayCameras/Source/GameplayCamerasEditor/Private/Editors/SObjectTreeGraphEditor.cpp
@@ -237,7 +237,7 @@ FReply SObjectTreeGraphEditor::OnDrop(const FGeometry& MyGeometry, const FDragDr
 	if (ObjectClassOp)
 	{
 		SGraphPanel* GraphPanel = GraphEditor->GetGraphPanel();
-		FVector2D NewLocation = GraphPanel->PanelCoordToGraphCoord(MyGeometry.AbsoluteToLocal(DragDropEvent.GetScreenSpacePosition()));
+		FSlateCompatVector2f NewLocation = GraphPanel->PanelCoordToGraphCoord(MyGeometry.AbsoluteToLocal(DragDropEvent.GetScreenSpacePosition()));
 
 		return ObjectClassOp->ExecuteDrop(GraphEditor, NewLocation);
 	}
@@ -324,7 +324,7 @@ FString SObjectTreeGraphEditor::ExportNodesToText(const FGraphPanelSelectionSet&
 	return Schema->ExportNodesToText(Nodes, bOnlyCanDuplicateNodes, bOnlyCanDeleteNodes);
 }
 
-void SObjectTreeGraphEditor::ImportNodesFromText(const FVector2D& Location, const FString& TextToImport)
+void SObjectTreeGraphEditor::ImportNodesFromText(const FSlateCompatVector2f& Location, const FString& TextToImport)
 {
 	// Start a transaction and flag things as modified.
 	const FScopedTransaction Transaction(LOCTEXT("PasteNodes", "Paste Nodes"));
@@ -491,7 +491,12 @@ void SObjectTreeGraphEditor::PasteNodes()
 	FString TextToImport;
 	FPlatformApplicationMisc::ClipboardPaste(TextToImport);
 	
-	ImportNodesFromText(GraphEditor->GetPasteLocation2f(), TextToImport);
+#if UE_VERSION_NEWER_THAN_OR_EQUAL(5,6,0)
+	FVector2f PasteLocation = GraphEditor->GetPasteLocation2f();
+#else
+	FVector2D PasteLocation = GraphEditor->GetPasteLocation();
+#endif
+	ImportNodesFromText(PasteLocation, TextToImport);
 }
 
 bool SObjectTreeGraphEditor::CanPasteNodes()
diff --git a/Engine/Plugins/Cameras/GameplayCameras/Source/GameplayCamerasEditor/Private/Editors/SObjectTreeGraphNode.cpp b/Engine/Plugins/Cameras/GameplayCameras/Source/GameplayCamerasEditor/Private/Editors/SObjectTreeGraphNode.cpp
--- a/Engine/Plugins/Cameras/GameplayCameras/Source/GameplayCamerasEditor/Private/Editors/SObjectTreeGraphNode.cpp
+++ b/Engine/Plugins/Cameras/GameplayCameras/Source/GameplayCamerasEditor/Private/Editors/SObjectTreeGraphNode.cpp
@@ -27,7 +27,7 @@ void SObjectTreeGraphNode::Construct(const FArguments& InArgs)
 	UpdateGraphNode();
 }
 
-void SObjectTreeGraphNode::MoveTo(const FVector2f& NewPosition, FNodeSet& NodeFilter, bool bMarkDirty)
+void SObjectTreeGraphNode::MoveTo(const FSlateCompatVector2f& NewPosition, FNodeSet& NodeFilter, bool bMarkDirty)
 {
 	SGraphNode::MoveTo(NewPosition, NodeFilter, bMarkDirty);
 
diff --git a/Engine/Plugins/Cameras/GameplayCameras/Source/GameplayCamerasEditor/Private/GameplayCamerasEditorModule.cpp b/Engine/Plugins/Cameras/GameplayCameras/Source/GameplayCamerasEditor/Private/GameplayCamerasEditorModule.cpp
--- a/Engine/Plugins/Cameras/GameplayCameras/Source/GameplayCamerasEditor/Private/GameplayCamerasEditorModule.cpp
+++ b/Engine/Plugins/Cameras/GameplayCameras/Source/GameplayCamerasEditor/Private/GameplayCamerasEditorModule.cpp
@@ -601,16 +601,20 @@ private:
 	{
 		using namespace UE::Cameras;
 
+#if UE_VERSION_NEWER_THAN_OR_EQUAL(5,6,0)
 		ISequencerModule& SequencerModule = FModuleManager::Get().LoadModuleChecked<ISequencerModule>("Sequencer");
 		GameplayCameraComponentTrackCreateEditorHandle = SequencerModule.RegisterTrackEditor(FOnCreateTrackEditor::CreateStatic(&FGameplayCameraComponentTrackEditor::CreateTrackEditor));
 		CameraFramingZoneTrackCreateEditorHandle = SequencerModule.RegisterPropertyTrackEditor<FCameraFramingZoneTrackEditor>();
+#endif
 	}
 
 	void UnregisterSequencerTracks()
 	{
+#if UE_VERSION_NEWER_THAN_OR_EQUAL(5,6,0)
 		ISequencerModule& SequencerModule = FModuleManager::Get().LoadModuleChecked<ISequencerModule>("Sequencer");
 		SequencerModule.UnRegisterTrackEditor(GameplayCameraComponentTrackCreateEditorHandle);
 		SequencerModule.UnRegisterTrackEditor(CameraFramingZoneTrackCreateEditorHandle);
+#endif
 	}
 
 	void InitializeLiveEditManager()
diff --git a/Engine/Plugins/Cameras/GameplayCameras/Source/GameplayCamerasEditor/Private/Sequencer/GameplayCameraComponentTrackEditor.cpp b/Engine/Plugins/Cameras/GameplayCameras/Source/GameplayCamerasEditor/Private/Sequencer/GameplayCameraComponentTrackEditor.cpp
--- a/Engine/Plugins/Cameras/GameplayCameras/Source/GameplayCamerasEditor/Private/Sequencer/GameplayCameraComponentTrackEditor.cpp
+++ b/Engine/Plugins/Cameras/GameplayCameras/Source/GameplayCamerasEditor/Private/Sequencer/GameplayCameraComponentTrackEditor.cpp
@@ -2,6 +2,8 @@
 
 #include "Sequencer/GameplayCameraComponentTrackEditor.h"
 
+#if UE_VERSION_NEWER_THAN_OR_EQUAL(5,6,0)
+
 #include "Core/CameraAsset.h"
 #include "Core/CameraObjectInterfaceParameterDefinition.h"
 #include "Core/CameraParameters.h"
@@ -418,3 +420,5 @@ void FGameplayCameraComponentTrackEditor::BuildAddParameterTrackSubMenuItems(FMe
 
 #undef LOCTEXT_NAMESPACE
 
+#endif  // >=5.6.0
+
diff --git a/Engine/Plugins/Cameras/GameplayCameras/Source/GameplayCamerasEditor/Private/Sequencer/GameplayCameraComponentTrackEditor.h b/Engine/Plugins/Cameras/GameplayCameras/Source/GameplayCamerasEditor/Private/Sequencer/GameplayCameraComponentTrackEditor.h
--- a/Engine/Plugins/Cameras/GameplayCameras/Source/GameplayCamerasEditor/Private/Sequencer/GameplayCameraComponentTrackEditor.h
+++ b/Engine/Plugins/Cameras/GameplayCameras/Source/GameplayCamerasEditor/Private/Sequencer/GameplayCameraComponentTrackEditor.h
@@ -2,6 +2,10 @@
 
 #pragma once
 
+#include "Misc/EngineVersionComparison.h"
+
+#if UE_VERSION_NEWER_THAN_OR_EQUAL(5,6,0)
+
 #include "ISequencerTrackEditor.h"
 #include "MovieSceneTrackEditor.h"
 #include "Templates/SharedPointerFwd.h"
@@ -71,3 +75,5 @@ private:
 	bool CanAddCameraParameterTrack(FPropertyMenuData PropertyMenuData, FGuid ObjectBinding) const;
 };
 
+#endif  // >= 5.6.0
+
diff --git a/Engine/Plugins/Cameras/GameplayCameras/Source/GameplayCamerasEditor/Public/Editors/CameraNodeGraphDragDropOp.h b/Engine/Plugins/Cameras/GameplayCameras/Source/GameplayCamerasEditor/Public/Editors/CameraNodeGraphDragDropOp.h
--- a/Engine/Plugins/Cameras/GameplayCameras/Source/GameplayCamerasEditor/Public/Editors/CameraNodeGraphDragDropOp.h
+++ b/Engine/Plugins/Cameras/GameplayCameras/Source/GameplayCamerasEditor/Public/Editors/CameraNodeGraphDragDropOp.h
@@ -2,6 +2,7 @@
 
 #pragma once
 
+#include "Compat/EditorCompat.h"
 #include "DragAndDrop/DecoratedDragDropOp.h"
 
 class SGraphEditor;
@@ -16,7 +17,7 @@ public:
 	static TSharedRef<FCameraNodeGraphInterfaceParameterDragDropOp> New(UCameraObjectInterfaceParameterBase* InInterfaceParameter);
 
 	FReply ExecuteDragOver(TSharedPtr<SGraphEditor> GraphEditor);
-	FReply ExecuteDrop(TSharedPtr<SGraphEditor> GraphEditor, const UE::Slate::FDeprecateVector2DParameter& NewLocation);
+	FReply ExecuteDrop(TSharedPtr<SGraphEditor> GraphEditor, const FSlateCompatVector2f& NewLocation);
 
 private:
 
diff --git a/Engine/Plugins/Cameras/GameplayCameras/Source/GameplayCamerasEditor/Public/Editors/CameraNodeGraphSchema.h b/Engine/Plugins/Cameras/GameplayCameras/Source/GameplayCamerasEditor/Public/Editors/CameraNodeGraphSchema.h
--- a/Engine/Plugins/Cameras/GameplayCameras/Source/GameplayCamerasEditor/Public/Editors/CameraNodeGraphSchema.h
+++ b/Engine/Plugins/Cameras/GameplayCameras/Source/GameplayCamerasEditor/Public/Editors/CameraNodeGraphSchema.h
@@ -2,6 +2,7 @@
 
 #pragma once
 
+#include "Compat/EditorCompat.h"
 #include "Core/CameraObjectInterface.h"
 #include "Core/CameraObjectInterfaceParameterDefinition.h"
 #include "EdGraph/EdGraphPin.h"
@@ -83,8 +84,7 @@ public:
 	// FEdGraphSchemaAction interface.
 	static FName StaticGetTypeId() { static FName Type("FCameraNodeGraphSchemaAction_NewInterfaceParameterNode"); return Type; }
 	virtual FName GetTypeId() const override { return StaticGetTypeId(); }
-	using FEdGraphSchemaAction::PerformAction; // Prevent hiding of deprecated base class function with FVector2D
-	virtual UEdGraphNode* PerformAction(UEdGraph* ParentGraph, UEdGraphPin* FromPin, const FVector2f& Location, bool bSelectNewNode = true) override;
+	virtual UEdGraphNode* PerformAction(UEdGraph* ParentGraph, UEdGraphPin* FromPin, FPerformGraphActionLocation Location, bool bSelectNewNode = true) override;
 };
 
 USTRUCT()
@@ -105,7 +105,6 @@ public:
 	// FEdGraphSchemaAction interface.
 	static FName StaticGetTypeId() { static FName Type("FCameraNodeGraphSchemaAction_AddInterfaceParameterNode"); return Type; }
 	virtual FName GetTypeId() const override { return StaticGetTypeId(); }
-	using FEdGraphSchemaAction::PerformAction; // Prevent hiding of deprecated base class function with FVector2D
-	virtual UEdGraphNode* PerformAction(UEdGraph* ParentGraph, UEdGraphPin* FromPin, const FVector2f& Location, bool bSelectNewNode = true) override;
+	virtual UEdGraphNode* PerformAction(UEdGraph* ParentGraph, UEdGraphPin* FromPin, FPerformGraphActionLocation Location, bool bSelectNewNode = true) override;
 };
 
diff --git a/Engine/Plugins/Cameras/GameplayCameras/Source/GameplayCamerasEditor/Public/Editors/ObjectTreeDragDropOp.h b/Engine/Plugins/Cameras/GameplayCameras/Source/GameplayCamerasEditor/Public/Editors/ObjectTreeDragDropOp.h
--- a/Engine/Plugins/Cameras/GameplayCameras/Source/GameplayCamerasEditor/Public/Editors/ObjectTreeDragDropOp.h
+++ b/Engine/Plugins/Cameras/GameplayCameras/Source/GameplayCamerasEditor/Public/Editors/ObjectTreeDragDropOp.h
@@ -2,6 +2,7 @@
 
 #pragma once
 
+#include "Compat/EditorCompat.h"
 #include "DragAndDrop/DecoratedDragDropOp.h"
 
 class UObjectTreeGraph;
@@ -23,7 +24,7 @@ public:
 	TArrayView<UClass* const> GetObjectClasses() const { return ObjectClasses; }
 
 	FReply ExecuteDragOver(TSharedPtr<SGraphEditor> GraphEditor);
-	FReply ExecuteDrop(TSharedPtr<SGraphEditor> GraphEditor, const UE::Slate::FDeprecateVector2DParameter& NewLocation);
+	FReply ExecuteDrop(TSharedPtr<SGraphEditor> GraphEditor, const FSlateCompatVector2f& NewLocation);
 
 private:
 
diff --git a/Engine/Plugins/Cameras/GameplayCameras/Source/GameplayCamerasEditor/Public/Editors/ObjectTreeGraphCommentNode.h b/Engine/Plugins/Cameras/GameplayCameras/Source/GameplayCamerasEditor/Public/Editors/ObjectTreeGraphCommentNode.h
--- a/Engine/Plugins/Cameras/GameplayCameras/Source/GameplayCamerasEditor/Public/Editors/ObjectTreeGraphCommentNode.h
+++ b/Engine/Plugins/Cameras/GameplayCameras/Source/GameplayCamerasEditor/Public/Editors/ObjectTreeGraphCommentNode.h
@@ -2,6 +2,7 @@
 
 #pragma once
 
+#include "Compat/EditorCompat.h"
 #include "EdGraphNode_Comment.h"
 #include "UObject/WeakObjectPtrFwd.h"
 
@@ -28,7 +29,7 @@ public:
 	virtual TSharedPtr<SGraphNode> CreateVisualWidget() override;
 	virtual void GetNodeContextMenuActions(UToolMenu* Menu, UGraphNodeContextMenuContext* Context) const override;
 	virtual void PostPlacedNewNode() override;
-	virtual void ResizeNode(const FVector2f& NewSize) override;
+	virtual void ResizeNode(const FSlateCompatVector2f& NewSize) override;
 	virtual void OnRenameNode(const FString& NewName) override;
 
 public:
diff --git a/Engine/Plugins/Cameras/GameplayCameras/Source/GameplayCamerasEditor/Public/Editors/ObjectTreeGraphNode.h b/Engine/Plugins/Cameras/GameplayCameras/Source/GameplayCamerasEditor/Public/Editors/ObjectTreeGraphNode.h
--- a/Engine/Plugins/Cameras/GameplayCameras/Source/GameplayCamerasEditor/Public/Editors/ObjectTreeGraphNode.h
+++ b/Engine/Plugins/Cameras/GameplayCameras/Source/GameplayCamerasEditor/Public/Editors/ObjectTreeGraphNode.h
@@ -5,6 +5,7 @@
 #include "Containers/Array.h"
 #include "EdGraph/EdGraphNode.h"
 #include "Editors/ObjectTreeGraphConfig.h"
+#include "Misc/EngineVersionComparison.h"
 #include "UObject/UObjectGlobals.h"
 
 #include "ObjectTreeGraphNode.generated.h"
@@ -63,7 +64,9 @@ public:
 	virtual FText GetNodeTitle(ENodeTitleType::Type TitleType) const override;
 	virtual TSharedPtr<SGraphNode> CreateVisualWidget() override;
 	virtual FLinearColor GetNodeTitleColor() const override;
+#if UE_VERSION_NEWER_THAN_OR_EQUAL(5,6,0)
 	virtual FLinearColor GetNodeTitleTextColor() const override;
+#endif
 	virtual FLinearColor GetNodeBodyTintColor() const override;
 	virtual FText GetTooltipText() const override;
 	virtual void AllocateDefaultPins() override;
diff --git a/Engine/Plugins/Cameras/GameplayCameras/Source/GameplayCamerasEditor/Public/Editors/ObjectTreeGraphSchema.h b/Engine/Plugins/Cameras/GameplayCameras/Source/GameplayCamerasEditor/Public/Editors/ObjectTreeGraphSchema.h
--- a/Engine/Plugins/Cameras/GameplayCameras/Source/GameplayCamerasEditor/Public/Editors/ObjectTreeGraphSchema.h
+++ b/Engine/Plugins/Cameras/GameplayCameras/Source/GameplayCamerasEditor/Public/Editors/ObjectTreeGraphSchema.h
@@ -2,6 +2,7 @@
 
 #pragma once
 
+#include "Compat/EditorCompat.h"
 #include "CoreTypes.h"
 #include "EdGraph/EdGraphSchema.h"
 #include "GraphEditor.h"
@@ -165,8 +166,7 @@ public:
 	// FEdGraphSchemaAction interface.
 	static FName StaticGetTypeId() { static FName Type("FObjectTreeGraphSchemaAction_NewNode"); return Type; }
 	virtual FName GetTypeId() const override { return StaticGetTypeId(); }
-	using FEdGraphSchemaAction::PerformAction; // Prevent hiding of deprecated base class function with FVector2D
-	virtual UEdGraphNode* PerformAction(UEdGraph* ParentGraph, UEdGraphPin* FromPin, const FVector2f& Location, bool bSelectNewNode = true) override;
+	virtual UEdGraphNode* PerformAction(UEdGraph* ParentGraph, UEdGraphPin* FromPin, FPerformGraphActionLocation Location, bool bSelectNewNode = true) override;
 
 protected:
 
@@ -189,6 +189,6 @@ public:
 	// FEdGraphSchemaAction interface.
 	static FName StaticGetTypeId() { static FName Type("FObjectTreeGraphSchemaAction_NewComment"); return Type; }
 	virtual FName GetTypeId() const override { return StaticGetTypeId(); } 
-	virtual UEdGraphNode* PerformAction(UEdGraph* ParentGraph, UEdGraphPin* FromPin, const FVector2f& Location, bool bSelectNewNode = true) override;
+	virtual UEdGraphNode* PerformAction(UEdGraph* ParentGraph, UEdGraphPin* FromPin, FPerformGraphActionLocation Location, bool bSelectNewNode = true) override;
 };
 
diff --git a/Engine/Plugins/Cameras/GameplayCameras/Source/GameplayCamerasEditor/Public/Editors/SCameraRigAssetEditor.h b/Engine/Plugins/Cameras/GameplayCameras/Source/GameplayCamerasEditor/Public/Editors/SCameraRigAssetEditor.h
--- a/Engine/Plugins/Cameras/GameplayCameras/Source/GameplayCamerasEditor/Public/Editors/SCameraRigAssetEditor.h
+++ b/Engine/Plugins/Cameras/GameplayCameras/Source/GameplayCamerasEditor/Public/Editors/SCameraRigAssetEditor.h
@@ -2,6 +2,7 @@
 
 #pragma once
 
+#include "Compat/EditorCompat.h"
 #include "Core/CameraRigAsset.h"
 #include "CoreTypes.h"
 #include "GraphEditor.h"
diff --git a/Engine/Plugins/Cameras/GameplayCameras/Source/GameplayCamerasEditor/Public/Editors/SCameraRigTransitionEditor.h b/Engine/Plugins/Cameras/GameplayCameras/Source/GameplayCamerasEditor/Public/Editors/SCameraRigTransitionEditor.h
--- a/Engine/Plugins/Cameras/GameplayCameras/Source/GameplayCamerasEditor/Public/Editors/SCameraRigTransitionEditor.h
+++ b/Engine/Plugins/Cameras/GameplayCameras/Source/GameplayCamerasEditor/Public/Editors/SCameraRigTransitionEditor.h
@@ -2,6 +2,7 @@
 
 #pragma once
 
+#include "Compat/EditorCompat.h"
 #include "CoreTypes.h"
 #include "Editors/CameraRigTransitionGraphSchema.h"
 #include "Editors/ObjectTreeGraphConfig.h"
diff --git a/Engine/Plugins/Cameras/GameplayCameras/Source/GameplayCamerasEditor/Public/Editors/SObjectTreeGraphCommentNode.h b/Engine/Plugins/Cameras/GameplayCameras/Source/GameplayCamerasEditor/Public/Editors/SObjectTreeGraphCommentNode.h
--- a/Engine/Plugins/Cameras/GameplayCameras/Source/GameplayCamerasEditor/Public/Editors/SObjectTreeGraphCommentNode.h
+++ b/Engine/Plugins/Cameras/GameplayCameras/Source/GameplayCamerasEditor/Public/Editors/SObjectTreeGraphCommentNode.h
@@ -2,6 +2,7 @@
 
 #pragma once
 
+#include "Compat/EditorCompat.h"
 #include "SGraphNodeComment.h"
 
 class UObjectTreeGraphCommentNode;
@@ -26,7 +27,7 @@ public:
 public:
 
 	// SNodePanel::SNode interface.
-	virtual void MoveTo(const FVector2f& NewPosition, FNodeSet& NodeFilter, bool bMarkDirty=true) override;
+	virtual void MoveTo(const FSlateCompatVector2f& NewPosition, FNodeSet& NodeFilter, bool bMarkDirty=true) override;
 
 protected:
 
diff --git a/Engine/Plugins/Cameras/GameplayCameras/Source/GameplayCamerasEditor/Public/Editors/SObjectTreeGraphEditor.h b/Engine/Plugins/Cameras/GameplayCameras/Source/GameplayCamerasEditor/Public/Editors/SObjectTreeGraphEditor.h
--- a/Engine/Plugins/Cameras/GameplayCameras/Source/GameplayCamerasEditor/Public/Editors/SObjectTreeGraphEditor.h
+++ b/Engine/Plugins/Cameras/GameplayCameras/Source/GameplayCamerasEditor/Public/Editors/SObjectTreeGraphEditor.h
@@ -2,6 +2,7 @@
 
 #pragma once
 
+#include "Compat/EditorCompat.h"
 #include "EditorUndoClient.h"
 #include "GraphEditor.h"
 #include "Widgets/SCompoundWidget.h"
@@ -78,7 +79,7 @@ protected:
 	void OnDoubleClicked();
 
 	FString ExportNodesToText(const FGraphPanelSelectionSet& Nodes, bool bOnlyCanDuplicateNodes, bool bOnlyCanDeleteNodes);
-	void ImportNodesFromText(const FVector2D& Location, const FString& TextToImport);
+	void ImportNodesFromText(const FSlateCompatVector2f& Location, const FString& TextToImport);
 	bool CanImportNodesFromText(const FString& TextToImport);
 	void DeleteNodes(TArrayView<UEdGraphNode*> NodesToDelete);
 
diff --git a/Engine/Plugins/Cameras/GameplayCameras/Source/GameplayCamerasEditor/Public/Editors/SObjectTreeGraphNode.h b/Engine/Plugins/Cameras/GameplayCameras/Source/GameplayCamerasEditor/Public/Editors/SObjectTreeGraphNode.h
--- a/Engine/Plugins/Cameras/GameplayCameras/Source/GameplayCamerasEditor/Public/Editors/SObjectTreeGraphNode.h
+++ b/Engine/Plugins/Cameras/GameplayCameras/Source/GameplayCamerasEditor/Public/Editors/SObjectTreeGraphNode.h
@@ -2,6 +2,7 @@
 
 #pragma once
 
+#include "Compat/EditorCompat.h"
 #include "SGraphNode.h"
 #include "SGraphPin.h"
 
@@ -37,7 +38,7 @@ public:
 	virtual const FSlateBrush* GetNodeBodyBrush() const override;
 
 	// SNodePanel::SNode interface.
-	virtual void MoveTo(const FVector2f& NewPosition, FNodeSet& NodeFilter, bool bMarkDirty=true) override;
+	virtual void MoveTo(const FSlateCompatVector2f& NewPosition, FNodeSet& NodeFilter, bool bMarkDirty=true) override;
 
 protected:
 
diff --git a/Engine/Plugins/Cameras/GameplayCameras/Source/GameplayCamerasUncookedOnly/Private/BlueprintGraph/K2Node_SetCameraRigParameter.cpp b/Engine/Plugins/Cameras/GameplayCameras/Source/GameplayCamerasUncookedOnly/Private/BlueprintGraph/K2Node_SetCameraRigParameter.cpp
--- a/Engine/Plugins/Cameras/GameplayCameras/Source/GameplayCamerasUncookedOnly/Private/BlueprintGraph/K2Node_SetCameraRigParameter.cpp
+++ b/Engine/Plugins/Cameras/GameplayCameras/Source/GameplayCamerasUncookedOnly/Private/BlueprintGraph/K2Node_SetCameraRigParameter.cpp
@@ -13,6 +13,7 @@
 #include "Kismet/BlueprintInstancedStructLibrary.h"
 #include "Kismet2/BlueprintEditorUtils.h"
 #include "KismetCompiler.h"
+#include "Misc/EngineVersionComparison.h"
 #include "Styling/AppStyle.h"
 #include "Styling/SlateIconFinder.h"
 
@@ -101,10 +102,12 @@ FSlateIcon UK2Node_SetCameraRigParameter::GetIconAndTint(FLinearColor& OutColor)
 		const UEdGraphSchema_K2* K2Schema = GetDefault<UEdGraphSchema_K2>();
 		OutColor = K2Schema->GetPinTypeColor(ParameterValuePin->PinType);
 
+#if UE_VERSION_NEWER_THAN_OR_EQUAL(5,6,0)
 		if (UStruct* Struct = Cast<UStruct>(ParameterValuePin->PinType.PinSubCategoryObject.Get()))
 		{
 			return FSlateIconFinder::FindIconForClass(Struct);
 		}
+#endif
 	}
 
 	return FSlateIcon(FAppStyle::GetAppStyleSetName(), "Kismet.AllClasses.VariableIcon");
