Cameras: fix legacy camera asset upgrade when a rig referencer hasn't been loaded

diff --git a/Engine/Plugins/Cameras/GameplayCameras/Source/GameplayCamerasEditor/Private/Toolkits/CameraAssetEditorToolkit.cpp b/Engine/Plugins/Cameras/GameplayCameras/Source/GameplayCamerasEditor/Private/Toolkits/CameraAssetEditorToolkit.cpp
--- a/Engine/Plugins/Cameras/GameplayCameras/Source/GameplayCamerasEditor/Private/Toolkits/CameraAssetEditorToolkit.cpp
+++ b/Engine/Plugins/Cameras/GameplayCameras/Source/GameplayCamerasEditor/Private/Toolkits/CameraAssetEditorToolkit.cpp
@@ -539,6 +539,11 @@ void FCameraAssetEditorToolkit::UpgradeLegacyCameraAssets()
 		for (FName OnDiskReferencer : OnDiskReferencers)
 		{
 			UPackage* ReferencerPackage = FindPackage(nullptr, *OnDiskReferencer.ToString());
+			if (!ReferencerPackage)
+			{
+				FPackagePath PackagePath = FPackagePath::FromPackageNameChecked(OnDiskReferencer);
+				ReferencerPackage = LoadPackage(nullptr, PackagePath, LOAD_None);
+			}
 			if (!ensure(ReferencerPackage))
 			{
 				continue;
