Cameras: prevent infinite recursion when there's a cyclic reference of camera rigs

diff --git a/Engine/Plugins/Cameras/GameplayCameras/Source/GameplayCameras/Private/Core/BlendStackRootCameraNode.cpp b/Engine/Plugins/Cameras/GameplayCameras/Source/GameplayCameras/Private/Core/BlendStackRootCameraNode.cpp
--- a/Engine/Plugins/Cameras/GameplayCameras/Source/GameplayCameras/Private/Core/BlendStackRootCameraNode.cpp
+++ b/Engine/Plugins/Cameras/GameplayCameras/Source/GameplayCameras/Private/Core/BlendStackRootCameraNode.cpp
@@ -168,25 +168,18 @@ const UCameraRigAsset* FBlendStackRootCameraNodeEvaluator::FindInnermostCameraRi
 {
 	if (CameraRig)
 	{
-		if (const UCameraRigCameraNode* PrefabNode = Cast<const UCameraRigCameraNode>(CameraRig->RootNode))
-		{
-			if (const UCameraRigAsset* InnerCameraRig = PrefabNode->CameraRigReference.GetCameraRig())
-			{
-				return FindInnermostCameraRigPrefab(InnerCameraRig);
-			}
-		}
+		TArray<TObjectPtr<const UCameraRigCameraNode>> PrefabTrail;
+		return BuildNestedPrefabTrail(CameraRig, PrefabTrail);
 	}
-	return CameraRig;
+	return nullptr;
 }
 
 FCameraNodeEvaluator* FBlendStackRootCameraNodeEvaluator::FindInnermostCameraRigEvaluator(FCameraNodeEvaluator* CameraNodeEvaluator)
 {
-	if (FCameraRigCameraNodeEvaluator* PrefabNodeEvaluator = CameraNodeEvaluator->CastThis<FCameraRigCameraNodeEvaluator>())
+	if (CameraNodeEvaluator)
 	{
-		if (FCameraNodeEvaluator* InnerNodeEvaluator = PrefabNodeEvaluator->GetCameraRigRootEvaluator())
-		{
-			return FindInnermostCameraRigEvaluator(InnerNodeEvaluator);
-		}
+		TArray<FCameraRigCameraNodeEvaluator*> EvaluatorTrail;
+		return BuildNestedEvaluatorTrail(CameraNodeEvaluator, EvaluatorTrail);
 	}
 	return CameraNodeEvaluator;
 }
@@ -218,18 +211,38 @@ void FBlendStackRootCameraNodeEvaluator::InitializeBlendedParameterOverridesStac
 	BlendedParameterOverridesStack.Add(MoveTemp(InitialParameterOverrides));
 }
 
-void FBlendStackRootCameraNodeEvaluator::BuildNestedPrefabTrail(const UCameraRigAsset* CameraRig, TArray<TObjectPtr<const UCameraRigCameraNode>>& OutPrefabNodes)
+const UCameraRigAsset* FBlendStackRootCameraNodeEvaluator::BuildNestedPrefabTrail(const UCameraRigAsset* CameraRig, TArray<TObjectPtr<const UCameraRigCameraNode>>& OutPrefabNodes)
 {
 	if (const UCameraRigCameraNode* PrefabNode = Cast<const UCameraRigCameraNode>(CameraRig->RootNode))
 	{
-		if (const UCameraRigAsset* InnerCameraRig = PrefabNode->CameraRigReference.GetCameraRig())
+		if (ensureMsgf(!OutPrefabNodes.Contains(PrefabNode), TEXT("Circular camera rig prefab reference detected!")))
 		{
 			OutPrefabNodes.Add(PrefabNode);
-			return BuildNestedPrefabTrail(InnerCameraRig, OutPrefabNodes);
+
+			if (const UCameraRigAsset* InnerCameraRig = PrefabNode->CameraRigReference.GetCameraRig())
+			{
+				return BuildNestedPrefabTrail(InnerCameraRig, OutPrefabNodes);
+			}
 		}
 	}
+	return CameraRig;
+}
+
+FCameraNodeEvaluator* FBlendStackRootCameraNodeEvaluator::BuildNestedEvaluatorTrail(FCameraNodeEvaluator* CameraNodeEvaluator, TArray<FCameraRigCameraNodeEvaluator*>& OutPrefabEvaluators)
+{
+	if (FCameraRigCameraNodeEvaluator* PrefabNodeEvaluator = CameraNodeEvaluator->CastThis<FCameraRigCameraNodeEvaluator>())
+	{
+		if (ensureMsgf(!OutPrefabEvaluators.Contains(PrefabNodeEvaluator), TEXT("Circular camera rig prefab reference detected!")))
+		{
+			OutPrefabEvaluators.Add(PrefabNodeEvaluator);
 
-	ensure(CameraRig == BlendablePrefabCameraRig);
+			if (FCameraNodeEvaluator* InnerNodeEvaluator = PrefabNodeEvaluator->GetCameraRigRootEvaluator())
+			{
+				return BuildNestedEvaluatorTrail(InnerNodeEvaluator, OutPrefabEvaluators);
+			}
+		}
+	}
+	return CameraNodeEvaluator;
 }
 
 void FBlendStackRootCameraNodeEvaluator::OnUpdateParameters(const FCameraBlendedParameterUpdateParams& Params, FCameraBlendedParameterUpdateResult& OutResult)
diff --git a/Engine/Plugins/Cameras/GameplayCameras/Source/GameplayCameras/Private/Core/BlendStackRootCameraNode.h b/Engine/Plugins/Cameras/GameplayCameras/Source/GameplayCameras/Private/Core/BlendStackRootCameraNode.h
--- a/Engine/Plugins/Cameras/GameplayCameras/Source/GameplayCameras/Private/Core/BlendStackRootCameraNode.h
+++ b/Engine/Plugins/Cameras/GameplayCameras/Source/GameplayCameras/Private/Core/BlendStackRootCameraNode.h
@@ -46,6 +46,7 @@ namespace UE::Cameras
 {
 
 class FBlendCameraNodeEvaluator;
+class FCameraRigCameraNodeEvaluator;
 
 enum class ECameraRigMergingEligibility
 {
@@ -95,8 +96,10 @@ private:
 	static const UCameraRigAsset* FindInnermostCameraRigPrefab(const UCameraRigAsset* CameraRig);
 	static FCameraNodeEvaluator* FindInnermostCameraRigEvaluator(FCameraNodeEvaluator* CameraNodeEvaluator);
 
+	static const UCameraRigAsset* BuildNestedPrefabTrail(const UCameraRigAsset* CameraRig, TArray<TObjectPtr<const UCameraRigCameraNode>>& OutPrefabNodes);
+	static FCameraNodeEvaluator* BuildNestedEvaluatorTrail(FCameraNodeEvaluator* CameraNodeEvaluator, TArray<FCameraRigCameraNodeEvaluator*>& OutPrefabEvaluators);
+
 	void InitializeBlendedParameterOverridesStack();
-	void BuildNestedPrefabTrail(const UCameraRigAsset* CameraRig, TArray<TObjectPtr<const UCameraRigCameraNode>>& OutPrefabNodes);
 	void RunBlendedParameterOverridesStack(const FCameraBlendedParameterUpdateParams& Params, FCameraBlendedParameterUpdateResult& OutResult);
 
 private:
