Cameras: add missing initialization call for blend nodes in merged camera rigs

diff --git a/Engine/Plugins/Cameras/GameplayCameras/Source/GameplayCameras/Private/Core/BlendStackRootCameraNode.cpp b/Engine/Plugins/Cameras/GameplayCameras/Source/GameplayCameras/Private/Core/BlendStackRootCameraNode.cpp
--- a/Engine/Plugins/Cameras/GameplayCameras/Source/GameplayCameras/Private/Core/BlendStackRootCameraNode.cpp
+++ b/Engine/Plugins/Cameras/GameplayCameras/Source/GameplayCameras/Private/Core/BlendStackRootCameraNode.cpp
@@ -127,7 +127,12 @@ ECameraRigMergingEligibility FBlendStackRootCameraNodeEvaluator::CompareCameraRi
 	return ECameraRigMergingEligibility::EligibleForMerge;
 }
 
-void FBlendStackRootCameraNodeEvaluator::MergeCameraRig(const FCameraNodeEvaluatorBuildParams& Params, const UCameraRigAsset* CameraRig, const UBlendCameraNode* Blend)
+void FBlendStackRootCameraNodeEvaluator::MergeCameraRig(
+		const FCameraNodeEvaluatorBuildParams& BuildParams,
+		const FCameraNodeEvaluatorInitializeParams& InitParams,
+		FCameraNodeEvaluationResult& InitResult,
+		const UCameraRigAsset* CameraRig,
+		const UBlendCameraNode* Blend)
 {
 	if (!ensureMsgf(CameraRig, TEXT("No camera rig given.")))
 	{
@@ -153,7 +158,11 @@ void FBlendStackRootCameraNodeEvaluator::MergeCameraRig(const FCameraNodeEvaluat
 
 	if (Blend)
 	{
-		BlendedParameterOverrides.BlendEvaluator = Params.BuildEvaluatorAs<FBlendCameraNodeEvaluator>(Blend);
+		BlendedParameterOverrides.BlendEvaluator = BuildParams.BuildEvaluatorAs<FBlendCameraNodeEvaluator>(Blend);
+	}
+	if (BlendedParameterOverrides.BlendEvaluator)
+	{
+		BlendedParameterOverrides.BlendEvaluator->Initialize(InitParams, InitResult);
 	}
 
 	BlendedParameterOverridesStack.Add(MoveTemp(BlendedParameterOverrides));
diff --git a/Engine/Plugins/Cameras/GameplayCameras/Source/GameplayCameras/Private/Core/BlendStackRootCameraNode.h b/Engine/Plugins/Cameras/GameplayCameras/Source/GameplayCameras/Private/Core/BlendStackRootCameraNode.h
--- a/Engine/Plugins/Cameras/GameplayCameras/Source/GameplayCameras/Private/Core/BlendStackRootCameraNode.h
+++ b/Engine/Plugins/Cameras/GameplayCameras/Source/GameplayCameras/Private/Core/BlendStackRootCameraNode.h
@@ -73,7 +73,12 @@ public:
 
 	ECameraRigMergingEligibility CompareCameraRigForMerging(const UCameraRigAsset* CameraRig) const;
 
-	void MergeCameraRig(const FCameraNodeEvaluatorBuildParams& Params, const UCameraRigAsset* CameraRig, const UBlendCameraNode* Blend);
+	void MergeCameraRig(
+			const FCameraNodeEvaluatorBuildParams& BuildParams, 
+			const FCameraNodeEvaluatorInitializeParams& InitParams, 
+			FCameraNodeEvaluationResult& InitResult,
+			const UCameraRigAsset* CameraRig, 
+			const UBlendCameraNode* Blend);
 
 protected:
 
diff --git a/Engine/Plugins/Cameras/GameplayCameras/Source/GameplayCameras/Private/Core/TransientBlendStackCameraNode.cpp b/Engine/Plugins/Cameras/GameplayCameras/Source/GameplayCameras/Private/Core/TransientBlendStackCameraNode.cpp
--- a/Engine/Plugins/Cameras/GameplayCameras/Source/GameplayCameras/Private/Core/TransientBlendStackCameraNode.cpp
+++ b/Engine/Plugins/Cameras/GameplayCameras/Source/GameplayCameras/Private/Core/TransientBlendStackCameraNode.cpp
@@ -137,7 +137,12 @@ FBlendStackEntryID FTransientBlendStackCameraNodeEvaluator::PushMergedEntry(cons
 	FCameraRigEntry& TopEntry = Entries.Top();
 	FCameraNodeEvaluatorBuilder Builder(TopEntry.EvaluatorStorage);
 	FCameraNodeEvaluatorBuildParams BuildParams(Builder);
-	TopEntry.RootEvaluator->MergeCameraRig(BuildParams, PushParams.CameraRig, Blend);
+
+	FCameraNodeEvaluatorInitializeParams InitParams;
+	InitParams.Evaluator = OwningEvaluator;
+	InitParams.EvaluationContext = TopEntry.EvaluationContext.Pin();
+
+	TopEntry.RootEvaluator->MergeCameraRig(BuildParams, InitParams, TopEntry.Result, PushParams.CameraRig, Blend);
 
 	// Merging a camera rig the first time changes the evaluator tree by removing any prefab nodes
 	// near the root... so we need to rebuild our hierarchy cache.
