Cameras: update content browser icons, add icon for camera shake assets

diff --git a/Engine/Plugins/Cameras/GameplayCameras/Source/GameplayCamerasEditor/Private/Styles/GameplayCamerasEditorStyle.cpp b/Engine/Plugins/Cameras/GameplayCameras/Source/GameplayCamerasEditor/Private/Styles/GameplayCamerasEditorStyle.cpp
--- a/Engine/Plugins/Cameras/GameplayCameras/Source/GameplayCamerasEditor/Private/Styles/GameplayCamerasEditorStyle.cpp
+++ b/Engine/Plugins/Cameras/GameplayCameras/Source/GameplayCamerasEditor/Private/Styles/GameplayCamerasEditorStyle.cpp
@@ -45,6 +45,8 @@ FGameplayCamerasEditorStyle::FGameplayCamerasEditorStyle()
 	Set("ClassThumbnail.CameraRigProxyAsset", new IMAGE_BRUSH_SVG("Icons/ContentBrowser-CameraRigProxy", Icon64x64));
 	Set("ClassIcon.CameraVariableCollection", new IMAGE_BRUSH_SVG("Icons/ContentBrowser-CameraVariableCollection", Icon16x16));
 	Set("ClassThumbnail.CameraVariableCollection", new IMAGE_BRUSH_SVG("Icons/ContentBrowser-CameraVariableCollection", Icon64x64));
+	Set("ClassIcon.CameraShakeAsset", new IMAGE_BRUSH_SVG("Icons/ContentBrowser-CameraShake", Icon16x16));
+	Set("ClassThumbnail.CameraShakeAsset", new IMAGE_BRUSH_SVG("Icons/ContentBrowser-CameraShake", Icon64x64));
 
 	// Camera actors and components.
 	Set("ClassIcon.GameplayCameraComponent", new IMAGE_BRUSH_SVG("Icons/GameplayCamera_16", Icon16x16));
