Cameras: re-add error messages for activating camera components with invalid data

Only print this message in non-editor worlds. For instance, creating a new camera actor in editor with no camera asset set yet won't print an error message as it tries to activate its editor preview.

Also, remove the override CalcCamera() method from the base GPC actor, since we want the view from a GPC component to look the same regardless of its base actor.

diff --git a/Engine/Plugins/Cameras/GameplayCameras/Source/GameplayCameras/Private/GameFramework/GameplayCameraActorBase.cpp b/Engine/Plugins/Cameras/GameplayCameras/Source/GameplayCameras/Private/GameFramework/GameplayCameraActorBase.cpp
--- a/Engine/Plugins/Cameras/GameplayCameras/Source/GameplayCameras/Private/GameFramework/GameplayCameraActorBase.cpp
+++ b/Engine/Plugins/Cameras/GameplayCameras/Source/GameplayCameras/Private/GameFramework/GameplayCameraActorBase.cpp
@@ -16,17 +16,7 @@ AGameplayCameraActorBase::AGameplayCameraActorBase(const FObjectInitializer& Obj
 
 void AGameplayCameraActorBase::CalcCamera(float DeltaTime, FMinimalViewInfo& OutResult)
 {
-	using namespace UE::Cameras;
-
-	UGameplayCameraComponentBase* CameraComponentBase = GetCameraComponentBase();
-	if (ensure(CameraComponentBase))
-	{
-		CameraComponentBase->CalcCamera(DeltaTime, OutResult);
-	}
-	else
-	{
-		Super::CalcCamera(DeltaTime, OutResult);
-	}
+	Super::CalcCamera(DeltaTime, OutResult);
 }
 
 #undef LOCTEXT_NAMESPACE
diff --git a/Engine/Plugins/Cameras/GameplayCameras/Source/GameplayCameras/Private/GameFramework/GameplayCameraComponent.cpp b/Engine/Plugins/Cameras/GameplayCameras/Source/GameplayCameras/Private/GameFramework/GameplayCameraComponent.cpp
--- a/Engine/Plugins/Cameras/GameplayCameras/Source/GameplayCameras/Private/GameFramework/GameplayCameraComponent.cpp
+++ b/Engine/Plugins/Cameras/GameplayCameras/Source/GameplayCameras/Private/GameFramework/GameplayCameraComponent.cpp
@@ -58,7 +58,14 @@ UCameraAsset* UGameplayCameraComponent::GetCameraAsset()
 
 bool UGameplayCameraComponent::OnValidateCameraEvaluationContextActivation()
 {
-	return CameraReference.IsValid();
+	const bool bIsValid = CameraReference.IsValid();
+	if (!bIsValid && !IsEditorWorld())
+	{
+		FFrame::KismetExecutionMessage(
+				TEXT("Can't activate gameplay camera component: no camera asset was set!"),
+				ELogVerbosity::Error);
+	}
+	return bIsValid;
 }
 
 void UGameplayCameraComponent::OnUpdateCameraEvaluationContext(bool bForceApplyParameterOverrides)
diff --git a/Engine/Plugins/Cameras/GameplayCameras/Source/GameplayCameras/Private/GameFramework/GameplayCameraComponentBase.cpp b/Engine/Plugins/Cameras/GameplayCameras/Source/GameplayCameras/Private/GameFramework/GameplayCameraComponentBase.cpp
--- a/Engine/Plugins/Cameras/GameplayCameras/Source/GameplayCameras/Private/GameFramework/GameplayCameraComponentBase.cpp
+++ b/Engine/Plugins/Cameras/GameplayCameras/Source/GameplayCameras/Private/GameFramework/GameplayCameraComponentBase.cpp
@@ -593,6 +593,15 @@ void UGameplayCameraComponentBase::UpdateControlRotationIfNeeded()
 	}
 }
 
+bool UGameplayCameraComponentBase::IsEditorWorld() const
+{
+#if WITH_EDITOR
+	return bIsEditorWorld;
+#else
+	return false;
+#endif  // WITH_EDITOR
+}
+
 #if WITH_EDITOR
 
 void UGameplayCameraComponentBase::ReinitializeCameraEvaluationContext(
diff --git a/Engine/Plugins/Cameras/GameplayCameras/Source/GameplayCameras/Private/GameFramework/GameplayCameraRigComponent.cpp b/Engine/Plugins/Cameras/GameplayCameras/Source/GameplayCameras/Private/GameFramework/GameplayCameraRigComponent.cpp
--- a/Engine/Plugins/Cameras/GameplayCameras/Source/GameplayCameras/Private/GameFramework/GameplayCameraRigComponent.cpp
+++ b/Engine/Plugins/Cameras/GameplayCameras/Source/GameplayCameras/Private/GameFramework/GameplayCameraRigComponent.cpp
@@ -58,7 +58,14 @@ UCameraAsset* UGameplayCameraRigComponent::GetCameraAsset()
 
 bool UGameplayCameraRigComponent::OnValidateCameraEvaluationContextActivation()
 {
-	return CameraRigReference.IsValid();
+	const bool bIsValid = CameraRigReference.IsValid();
+	if (!bIsValid && !IsEditorWorld())
+	{
+		FFrame::KismetExecutionMessage(
+				TEXT("Can't activate gameplay camera rig component: no camera rig asset was set!"),
+				ELogVerbosity::Error);
+	}
+	return bIsValid;
 }
 
 void UGameplayCameraRigComponent::OnUpdateCameraEvaluationContext(bool bForceApplyParameterOverrides)
diff --git a/Engine/Plugins/Cameras/GameplayCameras/Source/GameplayCameras/Public/GameFramework/GameplayCameraComponentBase.h b/Engine/Plugins/Cameras/GameplayCameras/Source/GameplayCameras/Public/GameFramework/GameplayCameraComponentBase.h
--- a/Engine/Plugins/Cameras/GameplayCameras/Source/GameplayCameras/Public/GameFramework/GameplayCameraComponentBase.h
+++ b/Engine/Plugins/Cameras/GameplayCameras/Source/GameplayCameras/Public/GameFramework/GameplayCameraComponentBase.h
@@ -168,6 +168,7 @@ protected:
 	void UpdateCameraEvaluationContext(bool bForceApplyParameterOverrides);
 	bool HasCameraEvaluationContext() const { return EvaluationContext.IsValid(); }
 
+	bool IsEditorWorld() const;
 	void UpdateControlRotationIfNeeded();
 
 #if WITH_EDITOR
