Cameras: pass camera data as a const-ref in the BP interop methods

This was tripping the BP compiler who made it an error to connect a BP camera node evaluator's "CameraData" variable, because it's read-only.

diff --git a/Engine/Plugins/Cameras/GameplayCameras/Source/GameplayCameras/Private/GameFramework/CameraRigParameterInterop.cpp b/Engine/Plugins/Cameras/GameplayCameras/Source/GameplayCameras/Private/GameFramework/CameraRigParameterInterop.cpp
--- a/Engine/Plugins/Cameras/GameplayCameras/Source/GameplayCameras/Private/GameFramework/CameraRigParameterInterop.cpp
+++ b/Engine/Plugins/Cameras/GameplayCameras/Source/GameplayCameras/Private/GameFramework/CameraRigParameterInterop.cpp
@@ -23,7 +23,7 @@ void UCameraRigParameterInterop::GetCameraParameter(const FBlueprintCameraEvalua
 	checkNoEntry();
 }
 
-void UCameraRigParameterInterop::SetCameraParameter(FBlueprintCameraEvaluationDataRef& CameraData, UCameraRigAsset* CameraRig, FName ParameterName, const int32& NewValue)
+void UCameraRigParameterInterop::SetCameraParameter(const FBlueprintCameraEvaluationDataRef& CameraData, UCameraRigAsset* CameraRig, FName ParameterName, const int32& NewValue)
 {
 	checkNoEntry();
 }
diff --git a/Engine/Plugins/Cameras/GameplayCameras/Source/GameplayCameras/Public/GameFramework/CameraRigParameterInterop.h b/Engine/Plugins/Cameras/GameplayCameras/Source/GameplayCameras/Public/GameFramework/CameraRigParameterInterop.h
--- a/Engine/Plugins/Cameras/GameplayCameras/Source/GameplayCameras/Public/GameFramework/CameraRigParameterInterop.h
+++ b/Engine/Plugins/Cameras/GameplayCameras/Source/GameplayCameras/Public/GameFramework/CameraRigParameterInterop.h
@@ -34,7 +34,7 @@ public:
 	static void GetCameraParameter(UPARAM(Ref) const FBlueprintCameraEvaluationDataRef& CameraData, UCameraRigAsset* CameraRig, FName ParameterName, int32& ReturnValue);
 
 	UFUNCTION(BlueprintCallable, CustomThunk, Category="Camera", meta=(BlueprintInternalUseOnly="true", CustomStructureParam="NewValue"))
-	static void SetCameraParameter(UPARAM(Ref) FBlueprintCameraEvaluationDataRef& CameraData, UCameraRigAsset* CameraRig, FName ParameterName, UPARAM(Ref) const int32& NewValue);
+	static void SetCameraParameter(UPARAM(Ref) const FBlueprintCameraEvaluationDataRef& CameraData, UCameraRigAsset* CameraRig, FName ParameterName, UPARAM(Ref) const int32& NewValue);
 
 private:
 
