Cameras: when setting the control rotation, only use the pre-visual result

This avoids including unwanted things like camera shakes into the control rotation.

Also, disable setting control rotation by default, since it's more likely that a project would start with its own control rotation handling.

diff --git a/Engine/Plugins/Cameras/GameplayCameras/Source/GameplayCameras/Private/Core/CameraSystemEvaluator.cpp b/Engine/Plugins/Cameras/GameplayCameras/Source/GameplayCameras/Private/Core/CameraSystemEvaluator.cpp
--- a/Engine/Plugins/Cameras/GameplayCameras/Source/GameplayCameras/Private/Core/CameraSystemEvaluator.cpp
+++ b/Engine/Plugins/Cameras/GameplayCameras/Source/GameplayCameras/Private/Core/CameraSystemEvaluator.cpp
@@ -246,6 +246,7 @@ void FCameraSystemEvaluator::UpdateImpl(float DeltaTime, ECameraNodeEvaluationTy
 	if (UNLIKELY(!ActiveContext.IsValid()))
 	{
 		Result.bIsValid = false;
+		PreVisualResult.bIsValid = false;
 		return;
 	}
 
@@ -274,6 +275,7 @@ void FCameraSystemEvaluator::UpdateImpl(float DeltaTime, ECameraNodeEvaluationTy
 	PostUpdateServices(DeltaTime, ECameraEvaluationServiceFlags::None);
 
 	// Harvest the result.
+	PreVisualResult.Reset(RootEvaluator->GetPreVisualLayerResult());
 	Result.Reset(RootNodeResult);
 
 	// Generate debug information if needed.
diff --git a/Engine/Plugins/Cameras/GameplayCameras/Source/GameplayCameras/Private/Core/DefaultRootCameraNode.cpp b/Engine/Plugins/Cameras/GameplayCameras/Source/GameplayCameras/Private/Core/DefaultRootCameraNode.cpp
--- a/Engine/Plugins/Cameras/GameplayCameras/Source/GameplayCameras/Private/Core/DefaultRootCameraNode.cpp
+++ b/Engine/Plugins/Cameras/GameplayCameras/Source/GameplayCameras/Private/Core/DefaultRootCameraNode.cpp
@@ -79,6 +79,7 @@ void FDefaultRootCameraNodeEvaluator::OnRun(const FCameraNodeEvaluationParams& P
 	BaseLayer->Run(Params, OutResult);
 	MainLayer->Run(Params, OutResult);
 	GlobalLayer->Run(Params, OutResult);
+	SetPreVisualLayerResult(OutResult);
 	if (Params.EvaluationType != ECameraNodeEvaluationType::IK && 
 			Params.EvaluationType != ECameraNodeEvaluationType::ViewRotationPreview)
 	{
diff --git a/Engine/Plugins/Cameras/GameplayCameras/Source/GameplayCameras/Private/Core/RootCameraNode.cpp b/Engine/Plugins/Cameras/GameplayCameras/Source/GameplayCameras/Private/Core/RootCameraNode.cpp
--- a/Engine/Plugins/Cameras/GameplayCameras/Source/GameplayCameras/Private/Core/RootCameraNode.cpp
+++ b/Engine/Plugins/Cameras/GameplayCameras/Source/GameplayCameras/Private/Core/RootCameraNode.cpp
@@ -100,6 +100,16 @@ void FRootCameraNodeEvaluator::RunSingleCameraRig(const FSingleCameraRigEvaluati
 	OnRunSingleCameraRig(Params, OutResult);
 }
 
+const FCameraNodeEvaluationResult& FRootCameraNodeEvaluator::GetPreVisualLayerResult() const
+{
+	return PreVisualResult;
+}
+
+void FRootCameraNodeEvaluator::SetPreVisualLayerResult(const FCameraNodeEvaluationResult& InResult)
+{
+	PreVisualResult.OverrideAll(InResult, true);
+}
+
 void FRootCameraNodeEvaluator::BroadcastCameraRigEvent(const FRootCameraNodeCameraRigEvent& InEvent) const
 {
 	if (ensure(OwningEvaluator))
diff --git a/Engine/Plugins/Cameras/GameplayCameras/Source/GameplayCameras/Private/GameFramework/GameplayCameraComponentBase.cpp b/Engine/Plugins/Cameras/GameplayCameras/Source/GameplayCameras/Private/GameFramework/GameplayCameraComponentBase.cpp
--- a/Engine/Plugins/Cameras/GameplayCameras/Source/GameplayCameras/Private/GameFramework/GameplayCameraComponentBase.cpp
+++ b/Engine/Plugins/Cameras/GameplayCameras/Source/GameplayCameras/Private/GameFramework/GameplayCameraComponentBase.cpp
@@ -587,7 +587,7 @@ void UGameplayCameraComponentBase::UpdateControlRotationIfNeeded()
 	AActor* OwnerActor = GetOwner();
 	if (OwnerActor && PlayerController->GetViewTarget() == OwnerActor)
 	{
-		const FCameraSystemEvaluationResult& Result = CameraSystemEvaluator->GetEvaluatedResult();
+		const FCameraSystemEvaluationResult& Result = CameraSystemEvaluator->GetPreVisualLayerEvaluatedResult();
 		const FRotator3d& ControlRotation = Result.CameraPose.GetRotation();
 		PlayerController->SetControlRotation(ControlRotation);
 	}
diff --git a/Engine/Plugins/Cameras/GameplayCameras/Source/GameplayCameras/Public/Core/CameraSystemEvaluator.h b/Engine/Plugins/Cameras/GameplayCameras/Source/GameplayCameras/Public/Core/CameraSystemEvaluator.h
--- a/Engine/Plugins/Cameras/GameplayCameras/Source/GameplayCameras/Public/Core/CameraSystemEvaluator.h
+++ b/Engine/Plugins/Cameras/GameplayCameras/Source/GameplayCameras/Public/Core/CameraSystemEvaluator.h
@@ -251,6 +251,9 @@ public:
 	/** Gets the evaluated result. */
 	const FCameraSystemEvaluationResult& GetEvaluatedResult() const { return Result; }
 
+	/** Gets the evaluated result without the contribution of the visual layer. */
+	const FCameraSystemEvaluationResult& GetPreVisualLayerEvaluatedResult() const { return PreVisualResult; }
+
 	/** Get the last evaluated camera. */
 	GAMEPLAYCAMERAS_API void GetEvaluatedCameraView(FMinimalViewInfo& DesiredView);
 
@@ -316,6 +319,9 @@ private:
 	/** The current overall result of the camera system. */
 	FCameraSystemEvaluationResult Result;
 
+	/** The current overall result of the camera system without the contribution of the visual layer. */
+	FCameraSystemEvaluationResult PreVisualResult;
+
 	/** Reusable buffer for node tree snapshots (e.g. view rotation preview update). */
 	TArray<uint8> EvaluatorSnapshot;
 
diff --git a/Engine/Plugins/Cameras/GameplayCameras/Source/GameplayCameras/Public/Core/RootCameraNode.h b/Engine/Plugins/Cameras/GameplayCameras/Source/GameplayCameras/Public/Core/RootCameraNode.h
--- a/Engine/Plugins/Cameras/GameplayCameras/Source/GameplayCameras/Public/Core/RootCameraNode.h
+++ b/Engine/Plugins/Cameras/GameplayCameras/Source/GameplayCameras/Public/Core/RootCameraNode.h
@@ -145,6 +145,9 @@ public:
 	/** Gets the variable table containing the blended camera object interfaces parameters. */
 	const FCameraVariableTable* GetBlendedParameters() const;
 
+	/** Returns the evaluation result without the contribution of the visual layer. */
+	const FCameraNodeEvaluationResult& GetPreVisualLayerResult() const;
+
 	/**
 	 * Builds the hierarchy of the system for a given single camera rig.
 	 * This is expected to return the nodes of all the layers, except for the main layer which
@@ -199,6 +202,8 @@ protected:
 
 protected:
 
+	void SetPreVisualLayerResult(const FCameraNodeEvaluationResult& InResult);
+
 	void BroadcastCameraRigEvent(const FRootCameraNodeCameraRigEvent& InEvent) const;
 
 private:
@@ -206,6 +211,9 @@ private:
 	/** The camera system that owns this root node. */
 	FCameraSystemEvaluator* OwningEvaluator = nullptr;
 
+	/** Evaluation result before the contribution of the visual layer. */
+	FCameraNodeEvaluationResult PreVisualResult;
+
 	/** The delegate to notify when an event happens. */
 	FOnRootCameraNodeCameraRigEvent OnCameraRigEventDelegate;
 };
diff --git a/Engine/Plugins/Cameras/GameplayCameras/Source/GameplayCameras/Public/GameFramework/GameplayCameraComponentBase.h b/Engine/Plugins/Cameras/GameplayCameras/Source/GameplayCameras/Public/GameFramework/GameplayCameraComponentBase.h
--- a/Engine/Plugins/Cameras/GameplayCameras/Source/GameplayCameras/Public/GameFramework/GameplayCameraComponentBase.h
+++ b/Engine/Plugins/Cameras/GameplayCameras/Source/GameplayCameras/Public/GameFramework/GameplayCameraComponentBase.h
@@ -211,7 +211,7 @@ public:
 	 * component.
 	 */
 	UPROPERTY(BlueprintReadWrite, EditAnywhere, Category=Camera)
-	bool bSetControlRotationWhenViewTarget = true;
+	bool bSetControlRotationWhenViewTarget = false;
 
 #if WITH_EDITORONLY_DATA
 
