Cameras: fix crash when running editor builds without an editor

diff --git a/Engine/Plugins/Cameras/GameplayCameras/Source/GameplayCameras/Private/GameFramework/GameplayCameraComponentBase.cpp b/Engine/Plugins/Cameras/GameplayCameras/Source/GameplayCameras/Private/GameFramework/GameplayCameraComponentBase.cpp
--- a/Engine/Plugins/Cameras/GameplayCameras/Source/GameplayCameras/Private/GameFramework/GameplayCameraComponentBase.cpp
+++ b/Engine/Plugins/Cameras/GameplayCameras/Source/GameplayCameras/Private/GameFramework/GameplayCameraComponentBase.cpp
@@ -198,10 +198,16 @@ void UGameplayCameraComponentBase::DeactivateCameraEvaluationContext(bool bImmed
 
 bool UGameplayCameraComponentBase::CanRunCameraSystem() const
 {
+	using namespace UE::Cameras;
+
 #if WITH_EDITOR
-	IGameplayCamerasModule& Module = IGameplayCamerasModule::Get();
-	const bool bCanRunInEditor = bRunInEditor && Module.GetLiveEditManager()->CanRunInEditor();
-	return (!bIsEditorWorld || bCanRunInEditor);
+	IGameplayCamerasModule& GameplayCamerasModule = IGameplayCamerasModule::Get();
+	if (TSharedPtr<IGameplayCamerasLiveEditManager> LiveEditManager = GameplayCamerasModule.GetLiveEditManager())
+	{
+		const bool bCanRunInEditor = bRunInEditor && GameplayCamerasModule.GetLiveEditManager()->CanRunInEditor();
+		return (!bIsEditorWorld || bCanRunInEditor);
+	}
+	return !bIsEditorWorld;
 #else
 	return true;
 #endif  // WITH_EDITOR
