Cameras: backwards compatible Rewind Debugger API change

diff --git a/Engine/Plugins/Cameras/GameplayCameras/Source/GameplayCamerasEditor/Private/Trace/CameraSystemRewindDebuggerTrack.cpp b/Engine/Plugins/Cameras/GameplayCameras/Source/GameplayCamerasEditor/Private/Trace/CameraSystemRewindDebuggerTrack.cpp
--- a/Engine/Plugins/Cameras/GameplayCameras/Source/GameplayCamerasEditor/Private/Trace/CameraSystemRewindDebuggerTrack.cpp
+++ b/Engine/Plugins/Cameras/GameplayCameras/Source/GameplayCamerasEditor/Private/Trace/CameraSystemRewindDebuggerTrack.cpp
@@ -156,7 +156,12 @@ void FCameraSystemRewindDebuggerTrackCreator::GetTrackTypesInternal(TArray<Rewin
 	Types.Add({GetNameInternal(), LOCTEXT("DisplayName", "Gameplay Camera System")});
 }
 
+
+#if UE_VERSION_NEWER_THAN_OR_EQUAL(5,7,0)
 TSharedPtr<RewindDebugger::FRewindDebuggerTrack> FCameraSystemRewindDebuggerTrackCreator::CreateTrackInternal(const RewindDebugger::FObjectId& InObjectId) const
+#else
+TSharedPtr<RewindDebugger::FRewindDebuggerTrack> FCameraSystemRewindDebuggerTrackCreator::CreateTrackInternal(uint64 InObjectId) const
+#endif
 {
 	return MakeShared<FCameraSystemRewindDebuggerTrack>();
 }
diff --git a/Engine/Plugins/Cameras/GameplayCameras/Source/GameplayCamerasEditor/Private/Trace/CameraSystemRewindDebuggerTrack.h b/Engine/Plugins/Cameras/GameplayCameras/Source/GameplayCamerasEditor/Private/Trace/CameraSystemRewindDebuggerTrack.h
--- a/Engine/Plugins/Cameras/GameplayCameras/Source/GameplayCamerasEditor/Private/Trace/CameraSystemRewindDebuggerTrack.h
+++ b/Engine/Plugins/Cameras/GameplayCameras/Source/GameplayCamerasEditor/Private/Trace/CameraSystemRewindDebuggerTrack.h
@@ -5,6 +5,7 @@
 #include "GameplayCameras.h"
 #include "IRewindDebuggerTrackCreator.h"
 #include "Math/Color.h"
+#include "Misc/EngineVersionComparison.h"
 #include "RewindDebuggerTrack.h"
 
 #if UE_GAMEPLAY_CAMERAS_TRACE
@@ -62,7 +63,11 @@ protected:
 	virtual FName GetTargetTypeNameInternal() const override;
 	virtual FName GetNameInternal() const override;
 	virtual void GetTrackTypesInternal(TArray<RewindDebugger::FRewindDebuggerTrackType>& Types) const override;
+#if UE_VERSION_NEWER_THAN_OR_EQUAL(5,7,0)
 	virtual TSharedPtr<RewindDebugger::FRewindDebuggerTrack> CreateTrackInternal(const RewindDebugger::FObjectId& InObjectId) const override;
+#else
+	virtual TSharedPtr<RewindDebugger::FRewindDebuggerTrack> CreateTrackInternal(uint64 InObjectId) const override;
+#endif
 };
 
 }  // namespace UE::Cameras
