Cameras: fix the tick group for the camera components

diff --git a/Engine/Plugins/Cameras/GameplayCameras/Source/GameplayCameras/Private/GameFramework/GameplayCameraComponentBase.cpp b/Engine/Plugins/Cameras/GameplayCameras/Source/GameplayCameras/Private/GameFramework/GameplayCameraComponentBase.cpp
--- a/Engine/Plugins/Cameras/GameplayCameras/Source/GameplayCameras/Private/GameFramework/GameplayCameraComponentBase.cpp
+++ b/Engine/Plugins/Cameras/GameplayCameras/Source/GameplayCameras/Private/GameFramework/GameplayCameraComponentBase.cpp
@@ -45,6 +45,7 @@ UGameplayCameraComponentBase::UGameplayCameraComponentBase(const FObjectInitiali
 	bWantsOnUpdateTransform = true;
 
 	PrimaryComponentTick.bCanEverTick = true;
+	PrimaryComponentTick.TickGroup = TG_PostPhysics;
 
 	OutputCameraComponent = ObjectInit.CreateDefaultSubobject<UCineCameraComponent>(this, TEXT("OutputCameraComponent"), true);
 	OutputCameraComponent->SetupAttachment(this);
