Cameras: don't pre-blend camera variable table entries that are not flagged for it

diff --git a/Engine/Plugins/Cameras/GameplayCameras/Source/GameplayCameras/Private/Core/PersistentBlendStackCameraNode.cpp b/Engine/Plugins/Cameras/GameplayCameras/Source/GameplayCameras/Private/Core/PersistentBlendStackCameraNode.cpp
--- a/Engine/Plugins/Cameras/GameplayCameras/Source/GameplayCameras/Private/Core/PersistentBlendStackCameraNode.cpp
+++ b/Engine/Plugins/Cameras/GameplayCameras/Source/GameplayCameras/Private/Core/PersistentBlendStackCameraNode.cpp
@@ -281,13 +281,15 @@ void FPersistentBlendStackCameraNodeEvaluator::InternalUpdate(TArrayView<FResolv
 			if (EntryBlendEvaluator)
 			{
 				FCameraNodePreBlendParams PreBlendParams(CurParams, CurResult.CameraPose, CurResult.VariableTable);
+				PreBlendParams.VariableTableFilter = ECameraVariableTableFilter::InputOnly;
+
 				FCameraNodePreBlendResult PreBlendResult(OutResult.VariableTable);
 
 				EntryBlendEvaluator->BlendParameters(PreBlendParams, PreBlendResult);
 			}
 			else
 			{
-				OutResult.VariableTable.Override(CurResult.VariableTable, ECameraVariableTableFilter::None);
+				OutResult.VariableTable.Override(CurResult.VariableTable, ECameraVariableTableFilter::InputOnly);
 			}
 
 			// Run the camera rig's root node.
diff --git a/Engine/Plugins/Cameras/GameplayCameras/Source/GameplayCameras/Private/Core/TransientBlendStackCameraNode.cpp b/Engine/Plugins/Cameras/GameplayCameras/Source/GameplayCameras/Private/Core/TransientBlendStackCameraNode.cpp
--- a/Engine/Plugins/Cameras/GameplayCameras/Source/GameplayCameras/Private/Core/TransientBlendStackCameraNode.cpp
+++ b/Engine/Plugins/Cameras/GameplayCameras/Source/GameplayCameras/Private/Core/TransientBlendStackCameraNode.cpp
@@ -315,7 +315,7 @@ void FTransientBlendStackCameraNodeEvaluator::InternalPreBlendExecute(TArrayView
 		return;
 	}
 
-	constexpr ECameraVariableTableFilter VariableTableFilter = ECameraVariableTableFilter::None;
+	constexpr ECameraVariableTableFilter VariableTableFilter = ECameraVariableTableFilter::InputOnly;
 
 	for (FResolvedEntry& ResolvedEntry : ResolvedEntries)
 	{
