Cameras: fix broken "inline" evaluation of camera shake prefabs

diff --git a/Engine/Plugins/Cameras/GameplayCameras/Source/GameplayCameras/Private/Nodes/Shakes/CameraShakeCameraNode.cpp b/Engine/Plugins/Cameras/GameplayCameras/Source/GameplayCameras/Private/Nodes/Shakes/CameraShakeCameraNode.cpp
--- a/Engine/Plugins/Cameras/GameplayCameras/Source/GameplayCameras/Private/Nodes/Shakes/CameraShakeCameraNode.cpp
+++ b/Engine/Plugins/Cameras/GameplayCameras/Source/GameplayCameras/Private/Nodes/Shakes/CameraShakeCameraNode.cpp
@@ -43,7 +43,7 @@ private:
 
 	ECameraShakeEvaluationMode EvaluationMode = ECameraShakeEvaluationMode::VisualLayer;
 
-	FCameraNodeEvaluator* CameraShakeRootEvaluator = nullptr;
+	FShakeCameraNodeEvaluator* CameraShakeRootEvaluator = nullptr;
 
 	TSharedPtr<FCameraShakeService> CameraShakeService;
 };
@@ -71,7 +71,7 @@ void FCameraShakeCameraNodeEvaluator::OnBuild(const FCameraNodeEvaluatorBuildPar
 		const UCameraShakeAsset* CameraShake = CameraShakeNode->CameraShakeReference.GetCameraShake();
 		if (CameraShake && CameraShake->RootNode)
 		{
-			CameraShakeRootEvaluator = Params.BuildEvaluator(CameraShake->RootNode);
+			CameraShakeRootEvaluator = Params.BuildEvaluatorAs<FShakeCameraNodeEvaluator>(CameraShake->RootNode);
 		}
 	}
 }
@@ -102,6 +102,10 @@ void FCameraShakeCameraNodeEvaluator::OnRun(const FCameraNodeEvaluationParams& P
 	if (CameraShakeRootEvaluator)
 	{
 		CameraShakeRootEvaluator->Run(Params, OutResult);
+
+		FCameraNodeShakeParams ShakeParams(Params);
+		FCameraNodeShakeResult ShakeResult(OutResult);
+		CameraShakeRootEvaluator->ShakeResult(ShakeParams, ShakeResult);
 	}
 	// If evaluating the shake later in the visual layer, keep requesting that the shake service
 	// maintains this shake alive. We already put the shake's parameters inside the variable table so
