Cameras: fix camera rig actor parameter keyframing

diff --git a/Engine/Plugins/Cameras/GameplayCameras/Source/GameplayCameras/Private/MovieScene/MovieSceneCameraParameterInstantiator.cpp b/Engine/Plugins/Cameras/GameplayCameras/Source/GameplayCameras/Private/MovieScene/MovieSceneCameraParameterInstantiator.cpp
--- a/Engine/Plugins/Cameras/GameplayCameras/Source/GameplayCameras/Private/MovieScene/MovieSceneCameraParameterInstantiator.cpp
+++ b/Engine/Plugins/Cameras/GameplayCameras/Source/GameplayCameras/Private/MovieScene/MovieSceneCameraParameterInstantiator.cpp
@@ -239,7 +239,8 @@ void UMovieSceneCameraParameterInstantiator::OnMovieSceneSectionAddedToTrack(UMo
 	// of the camera component, so use that to determine if this is a camera parameter track.
 	const FMovieScenePropertyBinding& PropertyBinding = PropertyTrack->GetPropertyBinding();
 	const FString& PropertyPath = PropertyBinding.PropertyPath.ToString();
-	if (PropertyPath.StartsWith(TEXT("CameraReference.Parameters.Value")))
+	if (PropertyPath.StartsWith(TEXT("CameraReference.Parameters.Value")) ||
+		PropertyPath.StartsWith(TEXT("CameraRigReference.Parameters.Value")))
 	{
 		UMovieSceneCameraParameterDecoration* Decoration = NewSection->GetOrCreateDecoration<UMovieSceneCameraParameterDecoration>();
 	}
