Cameras: framing nodes weren't looking at the player pawn anymore when no target was specified

diff --git a/Engine/Plugins/Cameras/GameplayCameras/Source/GameplayCameras/Private/Nodes/Framing/BaseFramingCameraNode.cpp b/Engine/Plugins/Cameras/GameplayCameras/Source/GameplayCameras/Private/Nodes/Framing/BaseFramingCameraNode.cpp
--- a/Engine/Plugins/Cameras/GameplayCameras/Source/GameplayCameras/Private/Nodes/Framing/BaseFramingCameraNode.cpp
+++ b/Engine/Plugins/Cameras/GameplayCameras/Source/GameplayCameras/Private/Nodes/Framing/BaseFramingCameraNode.cpp
@@ -156,6 +156,7 @@ bool FBaseFramingCameraNodeEvaluator::AcquireTargetInfo(TSharedPtr<const FCamera
 		{
 			OutInfo.LocalBounds = RootComponent->Bounds;
 		}
+		OutInfos.Add(OutInfo);
 		return true;
 	}
 
