Cameras: fix gameplay tag transition condition

When a tag query is specified, it should fail if there is no previous/next camera rig.

diff --git a/Engine/Plugins/Cameras/GameplayCameras/Source/GameplayCameras/Private/Transitions/GameplayTagTransitionConditions.cpp b/Engine/Plugins/Cameras/GameplayCameras/Source/GameplayCameras/Private/Transitions/GameplayTagTransitionConditions.cpp
--- a/Engine/Plugins/Cameras/GameplayCameras/Source/GameplayCameras/Private/Transitions/GameplayTagTransitionConditions.cpp
+++ b/Engine/Plugins/Cameras/GameplayCameras/Source/GameplayCameras/Private/Transitions/GameplayTagTransitionConditions.cpp
@@ -12,6 +12,8 @@ bool UGameplayTagTransitionCondition::OnTransitionMatches(const FCameraRigTransi
 
 	if (!PreviousGameplayTagQuery.IsEmpty())
 	{
+		bPreviousMatches = false;
+
 		if (Params.FromCameraRig)
 		{
 			FGameplayTagContainer TagContainer;
@@ -21,16 +23,14 @@ bool UGameplayTagTransitionCondition::OnTransitionMatches(const FCameraRigTransi
 				bPreviousMatches = true;
 			}
 		}
-		else
-		{
-			bPreviousMatches = false;
-		}
 	}
 
 	bool bNextMatches = true;
 
 	if (!NextGameplayTagQuery.IsEmpty())
 	{
+		bNextMatches = false;
+
 		if (Params.ToCameraRig)
 		{
 			FGameplayTagContainer TagContainer;
@@ -40,10 +40,6 @@ bool UGameplayTagTransitionCondition::OnTransitionMatches(const FCameraRigTransi
 				bNextMatches = true;
 			}
 		}
-		else
-		{
-			bNextMatches = false;
-		}
 	}
 
 	return bPreviousMatches && bNextMatches;
