Cameras: fix CustomThunk implementation for get/set parameter functions

First, the macro for getting the parameter name was wrong.

Second, wildcard pins don't support constant/default values, so the SetCameraParameter node always needs something connected into "NewValue", otherwise it crashes at runtime. Fixing this would require any of the user-accessible parameter setter nodes to ask for something to be connected into their value pin, instead of letting people specify values, but that's not very user friendly or efficient. So this change creates a "Make Literal Xxx" node for those cases.

This is a temporary fix. For later engine versions, we should look into expanding the Blueprint compiler and support default values on wildcard nodes.

diff --git a/Engine/Plugins/Cameras/GameplayCameras/Source/GameplayCameras/Private/GameFramework/CameraRigParameterInterop.cpp b/Engine/Plugins/Cameras/GameplayCameras/Source/GameplayCameras/Private/GameFramework/CameraRigParameterInterop.cpp
--- a/Engine/Plugins/Cameras/GameplayCameras/Source/GameplayCameras/Private/GameFramework/CameraRigParameterInterop.cpp
+++ b/Engine/Plugins/Cameras/GameplayCameras/Source/GameplayCameras/Private/GameFramework/CameraRigParameterInterop.cpp
@@ -32,7 +32,7 @@ DEFINE_FUNCTION(UCameraRigParameterInterop::execGetCameraParameter)
 {
 	P_GET_STRUCT_REF(FBlueprintCameraEvaluationDataRef, CameraData);
 	P_GET_OBJECT(UCameraRigAsset, CameraRig);
-	P_GET_STRUCT(FName, ParameterName);
+	P_GET_PROPERTY(FNameProperty, ParameterName);
 
 	// Read wildcard value input.
 	Stack.MostRecentPropertyAddress = nullptr;
@@ -113,7 +113,7 @@ DEFINE_FUNCTION(UCameraRigParameterInterop::execSetCameraParameter)
 {
 	P_GET_STRUCT_REF(FBlueprintCameraEvaluationDataRef, CameraData);
 	P_GET_OBJECT(UCameraRigAsset, CameraRig);
-	P_GET_STRUCT(FName, ParameterName);
+	P_GET_PROPERTY(FNameProperty, ParameterName);
 
 	// Read wildcard value input.
 	Stack.MostRecentPropertyAddress = nullptr;
@@ -186,5 +186,30 @@ DEFINE_FUNCTION(UCameraRigParameterInterop::execSetCameraParameter)
 	}
 }
 
+FVector UCameraRigParameterInteropLibrary::MakeLiteralVector(FVector Value)
+{
+	return Value;
+}
+
+FVector3f UCameraRigParameterInteropLibrary::MakeLiteralVector3f(FVector3f Value)
+{
+	return Value;
+}
+
+FVector2D UCameraRigParameterInteropLibrary::MakeLiteralVector2D(FVector2D Value)
+{
+	return Value;
+}
+
+FRotator UCameraRigParameterInteropLibrary::MakeLiteralRotator(FRotator Value)
+{
+	return Value;
+}
+
+FLinearColor UCameraRigParameterInteropLibrary::MakeLiteralLinearColor(FLinearColor Value)
+{
+	return Value;
+}
+
 #undef LOCTEXT_NAMESPACE
 
diff --git a/Engine/Plugins/Cameras/GameplayCameras/Source/GameplayCameras/Public/GameFramework/CameraRigParameterInterop.h b/Engine/Plugins/Cameras/GameplayCameras/Source/GameplayCameras/Public/GameFramework/CameraRigParameterInterop.h
--- a/Engine/Plugins/Cameras/GameplayCameras/Source/GameplayCameras/Public/GameFramework/CameraRigParameterInterop.h
+++ b/Engine/Plugins/Cameras/GameplayCameras/Source/GameplayCameras/Public/GameFramework/CameraRigParameterInterop.h
@@ -42,3 +42,29 @@ private:
 	DECLARE_FUNCTION(execSetCameraParameter);
 };
 
+/**
+ * Internal Blueprint function library for creating default values of the SetCameraParameter node above.
+ */
+UCLASS(MinimalAPI, Hidden)
+class UCameraRigParameterInteropLibrary : public UBlueprintFunctionLibrary
+{
+	GENERATED_BODY()
+
+public:
+
+	UFUNCTION(BlueprintPure, Category="Camera", meta=(BlueprintThreadSafe, BlueprintInternalUseOnly="true"))
+	static FVector MakeLiteralVector(FVector Value);
+
+	UFUNCTION(BlueprintPure, Category="Camera", meta=(BlueprintThreadSafe, BlueprintInternalUseOnly="true"))
+	static FVector3f MakeLiteralVector3f(FVector3f Value);
+
+	UFUNCTION(BlueprintPure, Category="Camera", meta=(BlueprintThreadSafe, BlueprintInternalUseOnly="true"))
+	static FVector2D MakeLiteralVector2D(FVector2D Value);
+	
+	UFUNCTION(BlueprintPure, Category="Camera", meta=(BlueprintThreadSafe, BlueprintInternalUseOnly="true"))
+	static FRotator MakeLiteralRotator(FRotator Value);
+
+	UFUNCTION(BlueprintPure, Category="Camera", meta=(BlueprintThreadSafe, BlueprintInternalUseOnly="true"))
+	static FLinearColor MakeLiteralLinearColor(FLinearColor Value);
+};
+
diff --git a/Engine/Plugins/Cameras/GameplayCameras/Source/GameplayCamerasUncookedOnly/Private/BlueprintGraph/K2Node_CameraRigBase.cpp b/Engine/Plugins/Cameras/GameplayCameras/Source/GameplayCamerasUncookedOnly/Private/BlueprintGraph/K2Node_CameraRigBase.cpp
--- a/Engine/Plugins/Cameras/GameplayCameras/Source/GameplayCamerasUncookedOnly/Private/BlueprintGraph/K2Node_CameraRigBase.cpp
+++ b/Engine/Plugins/Cameras/GameplayCameras/Source/GameplayCamerasUncookedOnly/Private/BlueprintGraph/K2Node_CameraRigBase.cpp
@@ -12,6 +12,9 @@
 #include "GameFramework/BlueprintCameraEvaluationDataRef.h"
 #include "GameFramework/CameraRigParameterInterop.h"
 #include "GameplayCamerasDelegates.h"
+#include "K2Node_CallFunction.h"
+#include "K2Node_EnumLiteral.h"
+#include "Kismet/KismetSystemLibrary.h"
 #include "Kismet2/BlueprintEditorUtils.h"
 #include "KismetCompiler.h"
 #include "Subsystems/AssetEditorSubsystem.h"
@@ -237,5 +240,86 @@ FEdGraphPinType UK2Node_CameraRigBase::MakeDataParameterPinType(ECameraContextDa
 	return PinType;
 }
 
+UK2Node_CallFunction* UK2Node_CameraRigBase::CreateMakeLiteralNode(FKismetCompilerContext& CompilerContext, UEdGraph* SourceGraph, UK2Node* SourceNode, UClass* FunctionLibraryClass, const TCHAR* FunctionName, UEdGraphPin* SourceValuePin)
+{
+	UK2Node_CallFunction* MakeLiteralNode = CompilerContext.SpawnIntermediateNode<UK2Node_CallFunction>(SourceNode, SourceGraph);
+	MakeLiteralNode->FunctionReference.SetExternalMember(FunctionName, FunctionLibraryClass);
+	MakeLiteralNode->AllocateDefaultPins();
+	CompilerContext.MessageLog.NotifyIntermediateObjectCreation(MakeLiteralNode, SourceGraph);
+
+	UEdGraphPin* LiteralValuePin = MakeLiteralNode->FindPinChecked(TEXT("Value"));
+	LiteralValuePin->DefaultValue = SourceValuePin->DefaultValue;
+	LiteralValuePin->DefaultTextValue = SourceValuePin->DefaultTextValue;
+	LiteralValuePin->AutogeneratedDefaultValue = SourceValuePin->AutogeneratedDefaultValue;
+	LiteralValuePin->DefaultObject = SourceValuePin->DefaultObject;
+
+	return MakeLiteralNode;
+}
+
+UK2Node_CallFunction* UK2Node_CameraRigBase::MakeLiteralValueForPin(FKismetCompilerContext& CompilerContext, UEdGraph* SourceGraph, UK2Node* SourceNode, UEdGraphPin* SourceValuePin)
+{
+#define UE_NEW_MAKE_LITERAL_NODE(FunctionLibraryClassName, MakeLiteralFunctionName)\
+	CreateMakeLiteralNode(\
+			CompilerContext, SourceGraph, SourceNode,\
+			FunctionLibraryClassName::StaticClass(), TEXT(#MakeLiteralFunctionName),\
+			SourceValuePin);
+
+#define UE_CREATE_MAKE_LITERAL_NODE_IF_PIN_CATEGORY(PinCategoryName, MakeLiteralFunctionName)\
+	if (SourceValuePin->PinType.PinCategory == UEdGraphSchema_K2::PinCategoryName)\
+	{\
+		return UE_NEW_MAKE_LITERAL_NODE(UKismetSystemLibrary, MakeLiteralFunctionName);\
+	}
+
+#define UE_CREATE_MAKE_LITERAL_NODE_IF_PIN_SUB_CATEGORY_OBJECT(SubCategoryObject, MakeLiteralFunctionName)\
+	if (SourceValuePin->PinType.PinSubCategoryObject == SubCategoryObject)\
+	{\
+		return UE_NEW_MAKE_LITERAL_NODE(UCameraRigParameterInteropLibrary, MakeLiteralFunctionName);\
+	}
+
+	// Support for all types with and inline widget for editing the default value.
+	// Start with those supported by the built-in FNodeFactory class.
+	UE_CREATE_MAKE_LITERAL_NODE_IF_PIN_CATEGORY(PC_Boolean, MakeLiteralBool)
+	UE_CREATE_MAKE_LITERAL_NODE_IF_PIN_CATEGORY(PC_Text, MakeLiteralText)
+	UE_CREATE_MAKE_LITERAL_NODE_IF_PIN_CATEGORY(PC_Int, MakeLiteralInt)
+	UE_CREATE_MAKE_LITERAL_NODE_IF_PIN_CATEGORY(PC_Int64, MakeLiteralInt64)
+	UE_CREATE_MAKE_LITERAL_NODE_IF_PIN_CATEGORY(PC_Name, MakeLiteralName)
+	UE_CREATE_MAKE_LITERAL_NODE_IF_PIN_CATEGORY(PC_String, MakeLiteralString)
+	UE_CREATE_MAKE_LITERAL_NODE_IF_PIN_CATEGORY(PC_Text, MakeLiteralText)
+
+	UE_CREATE_MAKE_LITERAL_NODE_IF_PIN_CATEGORY(PC_Float, MakeLiteralFloat)
+	UE_CREATE_MAKE_LITERAL_NODE_IF_PIN_CATEGORY(PC_Double, MakeLiteralDouble)
+	if (SourceValuePin->PinType.PinCategory == UEdGraphSchema_K2::PC_Real && 
+			SourceValuePin->PinType.PinSubCategory == UEdGraphSchema_K2::PC_Float)
+	{
+		return UE_NEW_MAKE_LITERAL_NODE(UKismetSystemLibrary, MakeLiteralFloat);
+	}
+	if (SourceValuePin->PinType.PinCategory == UEdGraphSchema_K2::PC_Real && 
+			SourceValuePin->PinType.PinSubCategory == UEdGraphSchema_K2::PC_Double)
+	{
+		return UE_NEW_MAKE_LITERAL_NODE(UKismetSystemLibrary, MakeLiteralDouble);
+	}
+
+	if (SourceValuePin->PinType.PinCategory == UEdGraphSchema_K2::PC_Byte &&
+		SourceValuePin->PinType.PinSubCategoryObject == nullptr)  // Not an enum
+	{
+		return UE_NEW_MAKE_LITERAL_NODE(UKismetSystemLibrary, MakeLiteralByte);
+	}
+
+	if (SourceValuePin->PinType.PinCategory == UEdGraphSchema_K2::PC_Struct)
+	{
+		UE_CREATE_MAKE_LITERAL_NODE_IF_PIN_SUB_CATEGORY_OBJECT(TBaseStructure<FLinearColor>::Get(), MakeLiteralLinearColor)
+		UE_CREATE_MAKE_LITERAL_NODE_IF_PIN_SUB_CATEGORY_OBJECT(TBaseStructure<FVector>::Get(), MakeLiteralVector)
+		UE_CREATE_MAKE_LITERAL_NODE_IF_PIN_SUB_CATEGORY_OBJECT(TVariantStructure<FVector3f>::Get(), MakeLiteralVector3f)
+		UE_CREATE_MAKE_LITERAL_NODE_IF_PIN_SUB_CATEGORY_OBJECT(TBaseStructure<FVector2D>::Get(), MakeLiteralVector2D)
+		UE_CREATE_MAKE_LITERAL_NODE_IF_PIN_SUB_CATEGORY_OBJECT(TBaseStructure<FRotator>::Get(), MakeLiteralRotator)
+	}
+
+	return nullptr;
+
+#undef UE_NEW_MAKE_LITERAL_NODE
+#undef UE_CREATE_MAKE_LITERAL_NODE_IF_PIN_CATEGORY
+#undef UE_CREATE_MAKE_LITERAL_NODE_IF_PIN_SUB_CATEGORY_OBJECT
+}
+
 #undef LOCTEXT_NAMESPACE
 
diff --git a/Engine/Plugins/Cameras/GameplayCameras/Source/GameplayCamerasUncookedOnly/Private/BlueprintGraph/K2Node_CameraRigBase.h b/Engine/Plugins/Cameras/GameplayCameras/Source/GameplayCamerasUncookedOnly/Private/BlueprintGraph/K2Node_CameraRigBase.h
--- a/Engine/Plugins/Cameras/GameplayCameras/Source/GameplayCamerasUncookedOnly/Private/BlueprintGraph/K2Node_CameraRigBase.h
+++ b/Engine/Plugins/Cameras/GameplayCameras/Source/GameplayCamerasUncookedOnly/Private/BlueprintGraph/K2Node_CameraRigBase.h
@@ -14,6 +14,7 @@
 class UCameraRigAsset;
 class UCameraObjectInterfaceBlendableParameter;
 class UCameraObjectInterfaceDataParameter;
+class UK2Node_CallFunction;
 
 /**
  * Utility base class for Blueprint nodes that can set camera rig parameters.
@@ -63,6 +64,12 @@ protected:
 
 	void OnCameraRigAssetBuilt(const UCameraRigAsset* InBuiltCameraRig);
 
+protected:
+
+	static UK2Node_CallFunction* CreateMakeLiteralNode(FKismetCompilerContext& CompilerContext, UEdGraph* SourceGraph, UK2Node* SourceNode, UClass* FunctionLibraryClass, const TCHAR* FunctionName, UEdGraphPin* SourceValuePin);
+
+	static UK2Node_CallFunction* MakeLiteralValueForPin(FKismetCompilerContext& CompilerContext, UEdGraph* SourceGraph, UK2Node* SourceNode, UEdGraphPin* InValuePin);
+
 protected:
 
 	UPROPERTY()
diff --git a/Engine/Plugins/Cameras/GameplayCameras/Source/GameplayCamerasUncookedOnly/Private/BlueprintGraph/K2Node_SetCameraRigParameter.cpp b/Engine/Plugins/Cameras/GameplayCameras/Source/GameplayCamerasUncookedOnly/Private/BlueprintGraph/K2Node_SetCameraRigParameter.cpp
--- a/Engine/Plugins/Cameras/GameplayCameras/Source/GameplayCamerasUncookedOnly/Private/BlueprintGraph/K2Node_SetCameraRigParameter.cpp
+++ b/Engine/Plugins/Cameras/GameplayCameras/Source/GameplayCamerasUncookedOnly/Private/BlueprintGraph/K2Node_SetCameraRigParameter.cpp
@@ -175,6 +175,7 @@ void UK2Node_SetCameraRigParameter::ExpandNode(FKismetCompilerContext& CompilerC
 	UK2Node_CallFunction* CallSetParameter = CompilerContext.SpawnIntermediateNode<UK2Node_CallFunction>(this, SourceGraph);
 	CallSetParameter->FunctionReference.SetExternalMember(TEXT("SetCameraParameter"), UCameraRigParameterInterop::StaticClass());
 	CallSetParameter->AllocateDefaultPins();
+	CompilerContext.MessageLog.NotifyIntermediateObjectCreation(CallSetParameter, SourceGraph);
 	UEdGraphPin* FirstExecPin = CallSetParameter->GetExecPin();
 
 	// Connect the camera evaluation result argument.
@@ -189,17 +190,28 @@ void UK2Node_SetCameraRigParameter::ExpandNode(FKismetCompilerContext& CompilerC
 	UEdGraphPin* CallSetParameterNamePin = CallSetParameter->FindPinChecked(TEXT("ParameterName"));
 	CallSetParameterNamePin->DefaultValue = CameraParameterName;
 
-	// Set or connect the parameter value argument.
+	// Set or connect the parameter value argument. Wildcard pins don't accept default values so if the value doesn't
+	// come from a connection, we need to use a "MakeLiteralXxx" node.
 	UEdGraphPin* CallSetParameterValuePin = CallSetParameter->FindPinChecked(TEXT("NewValue"));
 	CallSetParameterValuePin->PinType = CameraParameterValuePin->PinType;
-	CallSetParameterValuePin->DefaultValue = CameraParameterValuePin->DefaultValue;
-	CallSetParameterValuePin->DefaultTextValue = CameraParameterValuePin->DefaultTextValue;
-	CallSetParameterValuePin->AutogeneratedDefaultValue = CameraParameterValuePin->AutogeneratedDefaultValue;
-	CallSetParameterValuePin->DefaultObject = CameraParameterValuePin->DefaultObject;
 	if (CameraParameterValuePin->LinkedTo.Num() > 0)
 	{
 		CompilerContext.MovePinLinksToIntermediate(*CameraParameterValuePin, *CallSetParameterValuePin);
 	}
+	else if (UK2Node_CallFunction* MakeLiteral = MakeLiteralValueForPin(CompilerContext, SourceGraph, this, CameraParameterValuePin))
+	{
+		MakeLiteral->GetReturnValuePin()->MakeLinkTo(CallSetParameterValuePin);
+	}
+	else
+	{
+		FText MissingParameterConnectionMsg = FText::Format(
+				LOCTEXT(
+					"ErrorRequiresLiteral", 
+					"SetCameraRigParameter node @@ parameter value '{0}' must have an input connected into it "
+					"(try connecting a MakeStruct, EnumLiteral, or appropriate node)."),
+				FText::FromString(CameraParameterName));
+		CompilerContext.MessageLog.Error(*MissingParameterConnectionMsg.ToString(), this);
+	}
 
 	// Set extra type pin for enums.
 	if (ExtraEnumTypePinObject)
diff --git a/Engine/Plugins/Cameras/GameplayCameras/Source/GameplayCamerasUncookedOnly/Private/BlueprintGraph/K2Node_SetCameraRigParameters.cpp b/Engine/Plugins/Cameras/GameplayCameras/Source/GameplayCamerasUncookedOnly/Private/BlueprintGraph/K2Node_SetCameraRigParameters.cpp
--- a/Engine/Plugins/Cameras/GameplayCameras/Source/GameplayCamerasUncookedOnly/Private/BlueprintGraph/K2Node_SetCameraRigParameters.cpp
+++ b/Engine/Plugins/Cameras/GameplayCameras/Source/GameplayCamerasUncookedOnly/Private/BlueprintGraph/K2Node_SetCameraRigParameters.cpp
@@ -176,17 +176,28 @@ void UK2Node_SetCameraRigParameters::ExpandNode(FKismetCompilerContext& Compiler
 		UEdGraphPin* CallSetParameterNamePin = CallSetParameter->FindPinChecked(TEXT("ParameterName"));
 		CallSetParameterNamePin->DefaultValue = BlendableParameter->InterfaceParameterName;
 
-		// Set or connect the parameter value argument.
+		// Set or connect the parameter value argument. Wildcard pins don't accept default values so if the value doesn't
+		// come from a connection, we need to use a "MakeLiteralXxx" node.
 		UEdGraphPin* CallSetParameterValuePin = CallSetParameter->FindPinChecked(TEXT("NewValue"));
 		CallSetParameterValuePin->PinType = RigParameterPin->PinType;
-		CallSetParameterValuePin->DefaultValue = RigParameterPin->DefaultValue;
-		CallSetParameterValuePin->DefaultTextValue = RigParameterPin->DefaultTextValue;
-		CallSetParameterValuePin->AutogeneratedDefaultValue = RigParameterPin->AutogeneratedDefaultValue;
-		CallSetParameterValuePin->DefaultObject = RigParameterPin->DefaultObject;
 		if (RigParameterPin->LinkedTo.Num() > 0)
 		{
 			CompilerContext.MovePinLinksToIntermediate(*RigParameterPin, *CallSetParameterValuePin);
 		}
+		else if (UK2Node_CallFunction* MakeLiteral = MakeLiteralValueForPin(CompilerContext, SourceGraph, this, RigParameterPin))
+		{
+			MakeLiteral->GetReturnValuePin()->MakeLinkTo(CallSetParameterValuePin);
+		}
+		else
+		{
+			FText MissingParameterConnectionMsg = FText::Format(
+					LOCTEXT(
+						"ErrorRequiresLiteral", 
+						"SetCameraRigParameter node @@ parameter value '{0}' must have an input connected into it "
+						"(try connecting a MakeStruct, EnumLiteral, or appropriate node)."),
+					FText::FromString(BlendableParameter->InterfaceParameterName));
+			CompilerContext.MessageLog.Error(*MissingParameterConnectionMsg.ToString(), this);
+		}
 
 		// Chain the execution.
 		UEdGraphPin* CallSetParameterExecPin = CallSetParameter->GetExecPin();
@@ -240,17 +251,28 @@ void UK2Node_SetCameraRigParameters::ExpandNode(FKismetCompilerContext& Compiler
 		UEdGraphPin* CallSetParameterNamePin = CallSetParameter->FindPinChecked(TEXT("ParameterName"));
 		CallSetParameterNamePin->DefaultValue = DataParameter->InterfaceParameterName;
 
-		// Set or connect the parameter value argument.
+		// Set or connect the parameter value argument. Wildcard pins don't accept default values so if the value doesn't
+		// come from a connection, we need to use a "MakeLiteralXxx" node.
 		UEdGraphPin* CallSetParameterValuePin = CallSetParameter->FindPinChecked(TEXT("NewValue"));
 		CallSetParameterValuePin->PinType = RigParameterPin->PinType;
-		CallSetParameterValuePin->DefaultValue = RigParameterPin->DefaultValue;
-		CallSetParameterValuePin->DefaultTextValue = RigParameterPin->DefaultTextValue;
-		CallSetParameterValuePin->AutogeneratedDefaultValue = RigParameterPin->AutogeneratedDefaultValue;
-		CallSetParameterValuePin->DefaultObject = RigParameterPin->DefaultObject;
 		if (RigParameterPin->LinkedTo.Num() > 0)
 		{
 			CompilerContext.MovePinLinksToIntermediate(*RigParameterPin, *CallSetParameterValuePin);
 		}
+		else if (UK2Node_CallFunction* MakeLiteral = MakeLiteralValueForPin(CompilerContext, SourceGraph, this, RigParameterPin))
+		{
+			MakeLiteral->GetReturnValuePin()->MakeLinkTo(CallSetParameterValuePin);
+		}
+		else
+		{
+			FText MissingParameterConnectionMsg = FText::Format(
+					LOCTEXT(
+						"ErrorRequiresLiteral", 
+						"SetCameraRigParameter node @@ parameter value '{0}' must have an input connected into it "
+						"(try connecting a MakeStruct, EnumLiteral, or appropriate node)."),
+					FText::FromString(DataParameter->InterfaceParameterName));
+			CompilerContext.MessageLog.Error(*MissingParameterConnectionMsg.ToString(), this);
+		}
 
 		// Chain the execution.
 		if (PreviousThenPin)
