Cameras: use the new parameter on CineCameraComponent to truly disable physical camera settings

The CineCameraComponent had a recent change that caused exposure to be computed even when FocusMethod was disabled.

diff --git a/Engine/Plugins/Cameras/GameplayCameras/Source/GameplayCameras/Private/GameFramework/GameplayCameraComponentBase.cpp b/Engine/Plugins/Cameras/GameplayCameras/Source/GameplayCameras/Private/GameFramework/GameplayCameraComponentBase.cpp
--- a/Engine/Plugins/Cameras/GameplayCameras/Source/GameplayCameras/Private/GameFramework/GameplayCameraComponentBase.cpp
+++ b/Engine/Plugins/Cameras/GameplayCameras/Source/GameplayCameras/Private/GameFramework/GameplayCameraComponentBase.cpp
@@ -705,7 +705,9 @@ void UGameplayCameraComponentBase::UpdateOutputCameraComponent()
 			OutputCameraComponent->AspectRatioAxisConstraint = Result.CameraPose.GetAspectRatioAxisConstraint();
 
 			OutputCameraComponent->FocusSettings.ManualFocusDistance = Result.CameraPose.GetFocusDistance();
-			OutputCameraComponent->FocusSettings.FocusMethod = (Result.CameraPose.GetEnablePhysicalCamera() ? ECameraFocusMethod::Manual : ECameraFocusMethod::Disable);
+			OutputCameraComponent->FocusSettings.FocusMethod = (Result.CameraPose.GetEnablePhysicalCamera() ? ECameraFocusMethod::Manual : ECameraFocusMethod::DoNotOverride);
+
+			OutputCameraComponent->ExposureMethod = (Result.CameraPose.GetEnablePhysicalCamera() ? ECameraExposureMethod::Enabled : ECameraExposureMethod::DoNotOverride);
 
 			OutputCameraComponent->ProjectionMode = Result.CameraPose.GetProjectionMode();
 			OutputCameraComponent->OrthoWidth = Result.CameraPose.GetOrthographicWidth();
