Cameras: fix crash when a camera's enter or exit transitions have empty pins

diff --git a/Engine/Plugins/Cameras/GameplayCameras/Source/GameplayCameras/Private/Helpers/CameraRigTransitionFinder.cpp b/Engine/Plugins/Cameras/GameplayCameras/Source/GameplayCameras/Private/Helpers/CameraRigTransitionFinder.cpp
--- a/Engine/Plugins/Cameras/GameplayCameras/Source/GameplayCameras/Private/Helpers/CameraRigTransitionFinder.cpp
+++ b/Engine/Plugins/Cameras/GameplayCameras/Source/GameplayCameras/Private/Helpers/CameraRigTransitionFinder.cpp
@@ -23,10 +23,13 @@ const UCameraRigTransition* FCameraRigTransitionFinder::FindTransition(
 	// The transition should be used if all its conditions pass.
 	for (TObjectPtr<const UCameraRigTransition> Transition : Transitions)
 	{
-		const bool bConditionsPass = Transition->AllConditionsMatch(MatchParams);
-		if (bConditionsPass)
+		if (Transition)
 		{
-			return Transition;
+			const bool bConditionsPass = Transition->AllConditionsMatch(MatchParams);
+			if (bConditionsPass)
+			{
+				return Transition;
+			}
 		}
 	}
 
