Cameras: fix assert when adding an enum data parameter to a camera rig

diff --git a/Engine/Plugins/Cameras/GameplayCameras/Source/GameplayCamerasEditor/Private/Toolkits/CameraObjectInterfaceParametersToolkit.cpp b/Engine/Plugins/Cameras/GameplayCameras/Source/GameplayCamerasEditor/Private/Toolkits/CameraObjectInterfaceParametersToolkit.cpp
--- a/Engine/Plugins/Cameras/GameplayCameras/Source/GameplayCamerasEditor/Private/Toolkits/CameraObjectInterfaceParametersToolkit.cpp
+++ b/Engine/Plugins/Cameras/GameplayCameras/Source/GameplayCamerasEditor/Private/Toolkits/CameraObjectInterfaceParametersToolkit.cpp
@@ -534,6 +534,11 @@ protected:
 		{
 			NewDataType = ECameraContextDataType::Enum;
 		}
+		else if (PinType.PinCategory == UEdGraphSchema_K2::PC_Byte && 
+				PinType.PinSubCategoryObject != nullptr && PinType.PinSubCategoryObject->IsA<UEnum>())
+		{
+			NewDataType = ECameraContextDataType::Enum;
+		}
 		else if (PinType.PinCategory == UEdGraphSchema_K2::PC_Struct)
 		{
 			NewDataType = ECameraContextDataType::Struct;
