Cameras: fix camera rig parameter nodes not being added to the list of graph objects

diff --git a/Engine/Plugins/Cameras/GameplayCameras/Source/GameplayCamerasEditor/Private/Editors/CameraNodeGraphSchema.cpp b/Engine/Plugins/Cameras/GameplayCameras/Source/GameplayCamerasEditor/Private/Editors/CameraNodeGraphSchema.cpp
--- a/Engine/Plugins/Cameras/GameplayCameras/Source/GameplayCamerasEditor/Private/Editors/CameraNodeGraphSchema.cpp
+++ b/Engine/Plugins/Cameras/GameplayCameras/Source/GameplayCamerasEditor/Private/Editors/CameraNodeGraphSchema.cpp
@@ -565,6 +565,8 @@ UEdGraphNode* FCameraNodeGraphSchemaAction_NewInterfaceParameterNode::PerformAct
 
 	UCameraObjectInterfaceParameterGraphNode* NewGraphNode = Schema->CreateInterfaceParameterNode(ObjectTreeGraph, NewInterfaceParameter);
 
+	Schema->AddConnectableObject(ObjectTreeGraph, NewInterfaceParameter);
+
 	NewGraphNode->NodePosX = Location.X;
 	NewGraphNode->NodePosY = Location.Y;
 	NewGraphNode->OnGraphNodeMoved(false);
@@ -616,6 +618,8 @@ UEdGraphNode* FCameraNodeGraphSchemaAction_AddInterfaceParameterNode::PerformAct
 	
 	UCameraObjectInterfaceParameterGraphNode* NewGraphNode = Schema->CreateInterfaceParameterNode(ParentGraph, InterfaceParameter);
 
+	Schema->AddConnectableObject(ObjectTreeGraph, InterfaceParameter);
+
 	NewGraphNode->NodePosX = Location.X;
 	NewGraphNode->NodePosY = Location.Y;
 	NewGraphNode->OnGraphNodeMoved(false);
