Cameras: expose camera node evaluator flag setters

diff --git a/Engine/Plugins/Cameras/GameplayCameras/Source/GameplayCameras/Public/Core/CameraNodeEvaluator.h b/Engine/Plugins/Cameras/GameplayCameras/Source/GameplayCameras/Public/Core/CameraNodeEvaluator.h
--- a/Engine/Plugins/Cameras/GameplayCameras/Source/GameplayCameras/Public/Core/CameraNodeEvaluator.h
+++ b/Engine/Plugins/Cameras/GameplayCameras/Source/GameplayCameras/Public/Core/CameraNodeEvaluator.h
@@ -333,14 +333,14 @@ protected:
 	 * Sets the flags for this evaluator.
 	 * Can be called from the constructor, or during OnInitialize().
 	 */
-	void SetNodeEvaluatorFlags(ECameraNodeEvaluatorFlags InFlags);
+	GAMEPLAYCAMERAS_API void SetNodeEvaluatorFlags(ECameraNodeEvaluatorFlags InFlags);
 
 	/**
 	 * Adds flags for this evaluator.
 	 * Evaluators default to having all flags enabled, so this is mostly only for re-adding a
 	 * flag that was removed by a base class.
 	 */
-	void AddNodeEvaluatorFlags(ECameraNodeEvaluatorFlags InFlags);
+	GAMEPLAYCAMERAS_API void AddNodeEvaluatorFlags(ECameraNodeEvaluatorFlags InFlags);
 
 protected:
 
