Cameras: fix crash when changing the type of a camera rig parameter

When a camera rig asset is built, we do two things: reconstruct the evaluation context, and re-create any running instances of that camera rig. But the former was freezing any running instances, which removed the listeners for build-events, so the latter was not happening. We now keep those listeners. Then, when re-creating the running instances, we remove possibly invalid camera variable and context data entries even for those we just froze, and for the transient blend stack cache.

diff --git a/Engine/Plugins/Cameras/GameplayCameras/Source/GameplayCameras/Private/Core/BlendStackCameraNode.cpp b/Engine/Plugins/Cameras/GameplayCameras/Source/GameplayCameras/Private/Core/BlendStackCameraNode.cpp
--- a/Engine/Plugins/Cameras/GameplayCameras/Source/GameplayCameras/Private/Core/BlendStackCameraNode.cpp
+++ b/Engine/Plugins/Cameras/GameplayCameras/Source/GameplayCameras/Private/Core/BlendStackCameraNode.cpp
@@ -160,10 +160,6 @@ void FBlendStackCameraNodeEvaluator::FreezeEntry(FCameraRigEntry& Entry)
 	Entry.RootNode = nullptr;
 
 	Entry.EvaluationContext.Reset();
-
-#if WITH_EDITOR
-	RemoveListenedPackages(Entry);
-#endif
 	
 	Entry.Flags.bIsFrozen = true;
 }
@@ -539,8 +535,14 @@ void FBlendStackCameraNodeEvaluator::OnPostBuildAsset(const FGameplayCameraAsset
 	{
 		FCameraRigEntry& Entry(Entries[Index]);
 		const bool bRebuildEntry = Entry.ListenedPackages.Contains(BuildEvent.AssetPackage);
-		if (bRebuildEntry)
+		if (!bRebuildEntry)
 		{
+			continue;
+		}
+
+		if (!Entry.Flags.bIsFrozen)
+		{
+			// Destroy the existing entry evaluators.
 			Entry.EvaluatorStorage.DestroyEvaluatorTree();
 			Entry.EvaluatorHierarchy.Reset();
 
@@ -559,9 +561,16 @@ void FBlendStackCameraNodeEvaluator::OnPostBuildAsset(const FGameplayCameraAsset
 					Entry.EvaluationContext.Pin(),
 					Entry.RootNode,
 					Index == Entries.Num() - 1);
-
-			OnEntryReinitialized(Index);
 		}
+		else
+		{
+			// Leave things frozen, but also remove all the frozen data, in case the new data
+			// has different types or names for the same variables, because the user changed them.
+			Entry.Result.VariableTable.UnsetAllValues();
+			Entry.Result.ContextDataTable.UnsetAllValues();
+		}
+
+		OnEntryReinitialized(Index);
 	}
 }
 
diff --git a/Engine/Plugins/Cameras/GameplayCameras/Source/GameplayCameras/Private/Core/TransientBlendStackCameraNode.cpp b/Engine/Plugins/Cameras/GameplayCameras/Source/GameplayCameras/Private/Core/TransientBlendStackCameraNode.cpp
--- a/Engine/Plugins/Cameras/GameplayCameras/Source/GameplayCameras/Private/Core/TransientBlendStackCameraNode.cpp
+++ b/Engine/Plugins/Cameras/GameplayCameras/Source/GameplayCameras/Private/Core/TransientBlendStackCameraNode.cpp
@@ -609,5 +609,17 @@ const UCameraRigTransition* FTransientBlendStackCameraNodeEvaluator::FindTransit
 	return nullptr;
 }
 
+#if WITH_EDITOR
+
+void FTransientBlendStackCameraNodeEvaluator::OnEntryReinitialized(int32 EntryIndex)
+{
+	// Empty our pre-blend variable table in case it has variables that have been changed or are
+	// not valid anymore, such as the user changing the type of a camera rig parameter.
+	FCameraVariableTableAllocationInfo EmptyAllocationInfo;
+	PreBlendVariableTable.Initialize(EmptyAllocationInfo);
+}
+
+#endif
+
 }  // namespace UE::Cameras
 
diff --git a/Engine/Plugins/Cameras/GameplayCameras/Source/GameplayCameras/Public/Core/TransientBlendStackCameraNode.h b/Engine/Plugins/Cameras/GameplayCameras/Source/GameplayCameras/Public/Core/TransientBlendStackCameraNode.h
--- a/Engine/Plugins/Cameras/GameplayCameras/Source/GameplayCameras/Public/Core/TransientBlendStackCameraNode.h
+++ b/Engine/Plugins/Cameras/GameplayCameras/Source/GameplayCameras/Public/Core/TransientBlendStackCameraNode.h
@@ -77,6 +77,11 @@ protected:
 	virtual void OnRun(const FCameraNodeEvaluationParams& Params, FCameraNodeEvaluationResult& OutResult) override;
 	virtual void OnSerialize(const FCameraNodeEvaluatorSerializeParams& Params, FArchive& Ar) override;
 
+#if WITH_EDITOR
+	// FBlendStackCameraNodeEvaluator interface.
+	virtual void OnEntryReinitialized(int32 EntryIndex) override;
+#endif
+
 private:
 
 	// Update methods.
