Cameras: use the proxy table to better populate the shortcut bar for camera assets

diff --git a/Engine/Plugins/Cameras/GameplayCameras/Source/GameplayCamerasEditor/Private/Family/CameraAssetFamily.cpp b/Engine/Plugins/Cameras/GameplayCameras/Source/GameplayCamerasEditor/Private/Family/CameraAssetFamily.cpp
--- a/Engine/Plugins/Cameras/GameplayCameras/Source/GameplayCamerasEditor/Private/Family/CameraAssetFamily.cpp
+++ b/Engine/Plugins/Cameras/GameplayCameras/Source/GameplayCamerasEditor/Private/Family/CameraAssetFamily.cpp
@@ -59,19 +59,40 @@ void FCameraAssetFamily::FindAssetsOfType(UClass* InAssetType, TArray<FAssetData
 		return;
 	}
 
-	FCameraDirectorRigUsageInfo UsageInfo;
-	CameraDirector->GatherRigUsageInfo(UsageInfo);
+	TSet<UCameraRigAsset*> AllCameraRigs;
+	TSet<UCameraRigProxyAsset*> AllCameraRigProxies;
+
+	// Get camera rigs and proxies from the director's logic.
+	{
+		FCameraDirectorRigUsageInfo UsageInfo;
+		CameraDirector->GatherRigUsageInfo(UsageInfo);
+
+		AllCameraRigs.Append(UsageInfo.CameraRigs);
+		AllCameraRigProxies.Append(UsageInfo.CameraRigProxies);
+	}
+	// Add camera rigs and proxies from the proxy table.
+	for (const FCameraRigProxyRedirectTableEntry& Entry : CameraDirector->CameraRigProxyRedirectTable.Entries)
+	{
+		if (Entry.CameraRigProxy)
+		{
+			AllCameraRigProxies.Add(Entry.CameraRigProxy);
+		}
+		if (Entry.CameraRig)
+		{
+			AllCameraRigs.Add(Entry.CameraRig);
+		}
+	}
 
 	if (InAssetType == UCameraRigAsset::StaticClass())
 	{
-		for (UCameraRigAsset* CameraRig : UsageInfo.CameraRigs)
+		for (UCameraRigAsset* CameraRig : AllCameraRigs)
 		{
 			OutAssets.Add(FAssetData(CameraRig));
 		}
 	}
 	else if (InAssetType == UCameraRigProxyAsset::StaticClass())
 	{
-		for (UCameraRigProxyAsset* CameraRigProxy : UsageInfo.CameraRigProxies)
+		for (UCameraRigProxyAsset* CameraRigProxy : AllCameraRigProxies)
 		{
 			OutAssets.Add(FAssetData(CameraRigProxy));
 		}
