Cameras: don't warn about a camera asset not using any rigs if it's using proxies, and build any rigs in the proxy table.

diff --git a/Engine/Plugins/Cameras/GameplayCameras/Source/GameplayCameras/Private/Build/CameraAssetBuilder.cpp b/Engine/Plugins/Cameras/GameplayCameras/Source/GameplayCameras/Private/Build/CameraAssetBuilder.cpp
--- a/Engine/Plugins/Cameras/GameplayCameras/Source/GameplayCameras/Private/Build/CameraAssetBuilder.cpp
+++ b/Engine/Plugins/Cameras/GameplayCameras/Source/GameplayCameras/Private/Build/CameraAssetBuilder.cpp
@@ -49,29 +49,45 @@ void FCameraAssetBuilder::BuildCamera(UCameraAsset* InCameraAsset, FCustomBuildS
 
 void FCameraAssetBuilder::BuildCameraImpl()
 {
-	TArray<UCameraRigAsset*> CameraRigs;
+	TSet<UCameraRigAsset*> AllCameraRigs;
+	TSet<UCameraRigProxyAsset*> AllCameraRigProxies;
 
 	// Build the camera director and get the list of camera rigs it references.
-	if (UCameraDirector* CameraDirector = CameraAsset->GetCameraDirector())
+	UCameraDirector* CameraDirector = CameraAsset->GetCameraDirector();
+	if (CameraDirector)
 	{
 		CameraDirector->BuildCameraDirector(BuildLog);
 
 		FCameraDirectorRigUsageInfo UsageInfo;
 		CameraDirector->GatherRigUsageInfo(UsageInfo);
-		CameraRigs = UsageInfo.CameraRigs;
+
+		AllCameraRigs.Append(UsageInfo.CameraRigs);
+		AllCameraRigProxies.Append(UsageInfo.CameraRigProxies);
 	}
 	else
 	{
 		BuildLog.AddMessage(EMessageSeverity::Error, LOCTEXT("MissingDirector", "Camera has no director set."));
 	}
 
-	if (CameraRigs.IsEmpty())
+	if (AllCameraRigs.IsEmpty() && AllCameraRigProxies.IsEmpty())
+	{
+		BuildLog.AddMessage(EMessageSeverity::Warning, LOCTEXT("MissingRigs", "Camera director isn't using any camera rigs or proxies."));
+	}
+
+	// Get also the rigs from the proxy table.
+	if (CameraDirector)
 	{
-		BuildLog.AddMessage(EMessageSeverity::Warning, LOCTEXT("MissingRigs", "Camera isn't using any camera rigs."));
+		for (const FCameraRigProxyRedirectTableEntry& Entry : CameraDirector->CameraRigProxyRedirectTable.Entries)
+		{
+			if (Entry.CameraRig)
+			{
+				AllCameraRigs.Add(Entry.CameraRig);
+			}
+		}
 	}
 
 	// Build each of the camera rigs.
-	for (UCameraRigAsset* CameraRig : CameraRigs)
+	for (UCameraRigAsset* CameraRig : AllCameraRigs)
 	{
 		FCameraRigAssetBuilder CameraRigBuilder(BuildLog);
 		CameraRigBuilder.BuildCameraRig(CameraRig);
@@ -82,7 +98,7 @@ void FCameraAssetBuilder::BuildCameraImpl()
 	TArray<FCameraObjectInterfaceParameterDefinition> ParameterDefinitions;
 	TArray<TObjectPtr<const UCameraRigAsset>> ParameterOwners;
 
-	for (const UCameraRigAsset* CameraRig : CameraRigs)
+	for (const UCameraRigAsset* CameraRig : AllCameraRigs)
 	{
 		for (const FCameraObjectInterfaceParameterDefinition& Definition : CameraRig->GetParameterDefinitions())
 		{
@@ -101,13 +117,13 @@ void FCameraAssetBuilder::BuildCameraImpl()
 	// Get the list of all the camera rigs' interface parameters, and rebuild our
 	// parameters property bag.
 	TArray<FPropertyBagPropertyDesc> DefaultParameterProperties;
-	for (const UCameraRigAsset* CameraRig : CameraRigs)
+	for (const UCameraRigAsset* CameraRig : AllCameraRigs)
 	{
 		FCameraObjectInterfaceParameterBuilder::AppendDefaultParameterProperties(CameraRig, DefaultParameterProperties);
 	}
 	FInstancedPropertyBag DefaultParameters;
 	DefaultParameters.AddProperties(DefaultParameterProperties);
-	for (const UCameraRigAsset* CameraRig : CameraRigs)
+	for (const UCameraRigAsset* CameraRig : AllCameraRigs)
 	{
 		FCameraObjectInterfaceParameterBuilder::SetDefaultParameterValues(CameraRig, DefaultParameters);
 	}
@@ -121,7 +137,7 @@ void FCameraAssetBuilder::BuildCameraImpl()
 	// Accumulate all the camera rigs' allocation infos and store that on the asset.
 	FCameraAssetAllocationInfo AllocationInfo;
 
-	for (const UCameraRigAsset* CameraRig : CameraRigs)
+	for (const UCameraRigAsset* CameraRig : AllCameraRigs)
 	{
 		AllocationInfo.VariableTableInfo.Combine(CameraRig->AllocationInfo.VariableTableInfo);
 		AllocationInfo.ContextDataTableInfo.Combine(CameraRig->AllocationInfo.ContextDataTableInfo);
