Cameras: some changes to input slots

- Run input nodes in OnRun. OnUpdateParameters is only meant for pre-blending input parameters, not running actual logic, if only because it is run in isolation on just a few nodes that opt-in.
- Move the revert axis options onto the base input slot class.
- Add a speed option on the input slot class. Pre-blend that instead of pre-blending the actual player input.
- Move user-exposed option for accumulation onto the axis-binding specific node, in preparation of adding an actual accumulation node for more complex input chains.
- Move axis-binding code to a helper class in preparation of a "raw" input action node.

diff --git a/Engine/Plugins/Cameras/GameplayCameras/Source/GameplayCameras/Private/Nodes/Input/CameraRigInput1DSlot.cpp b/Engine/Plugins/Cameras/GameplayCameras/Source/GameplayCameras/Private/Nodes/Input/CameraRigInput1DSlot.cpp
--- a/Engine/Plugins/Cameras/GameplayCameras/Source/GameplayCameras/Private/Nodes/Input/CameraRigInput1DSlot.cpp
+++ b/Engine/Plugins/Cameras/GameplayCameras/Source/GameplayCameras/Private/Nodes/Input/CameraRigInput1DSlot.cpp
@@ -27,10 +27,14 @@ FCameraRigInput1DSlotEvaluator::FCameraRigInput1DSlotEvaluator()
 
 void FCameraRigInput1DSlotEvaluator::OnInitialize(const FCameraNodeEvaluatorInitializeParams& Params, FCameraNodeEvaluationResult& OutResult)
 {
-	TransientInputValue = 0.f;
+	const UCameraRigInput1DSlot* SlotNode = GetCameraNodeAs<UCameraRigInput1DSlot>();
+
+	RevertAxisReader.Initialize(SlotNode->RevertAxis);
+	SpeedReader.Initialize(SlotNode->Speed);
+
+	DeltaInputValue = 0.f;
 	InputValue = 0.f;
 
-	const UCameraRigInput1DSlot* SlotNode = GetCameraNodeAs<UCameraRigInput1DSlot>();
 	if (SlotNode->GetVariableID().IsValid() && Params.LastActiveCameraRigInfo.LastResult)
 	{
 		const FCameraVariableTable& LastActiveRigVariableTable = Params.LastActiveCameraRigInfo.LastResult->VariableTable;
@@ -41,9 +45,10 @@ void FCameraRigInput1DSlotEvaluator::OnInitialize(const FCameraNodeEvaluatorInit
 void FCameraRigInput1DSlotEvaluator::OnUpdateParameters(const FCameraBlendedParameterUpdateParams& Params, FCameraBlendedParameterUpdateResult& OutResult)
 {
 	const UCameraRigInput1DSlot* SlotNode = GetCameraNodeAs<UCameraRigInput1DSlot>();
-	if (SlotNode->InputSlotParameters.bIsPreBlended)
+	if (SlotNode->bIsPreBlended)
 	{
-		OutResult.VariableTable.SetValue<double>(SlotNode->GetTransientVariableID(), TransientInputValue);
+		const double Speed = SpeedReader.Get(OutResult.VariableTable);
+		OutResult.VariableTable.SetValue<double>(SlotNode->GetSpeedVariableID(), Speed);
 	}
 }
 
@@ -51,18 +56,20 @@ void FCameraRigInput1DSlotEvaluator::OnRun(const FCameraNodeEvaluationParams& Pa
 {
 	const UCameraRigInput1DSlot* SlotNode = GetCameraNodeAs<UCameraRigInput1DSlot>();
 
-	if (SlotNode->InputSlotParameters.bIsPreBlended)
-	{
-		TransientInputValue = OutResult.VariableTable.GetValue<double>(SlotNode->GetTransientVariableID());
-	}
+	const double Speed = 
+		SlotNode->bIsPreBlended ?
+			OutResult.VariableTable.GetValue<double>(SlotNode->GetSpeedVariableID()) :
+			SpeedReader.Get(OutResult.VariableTable);
+
+	double FinalDelta(DeltaInputValue * Speed * Params.DeltaTime);
 
-	if (SlotNode->InputSlotParameters.bIsAccumulated)
+	if (bIsAccumulated)
 	{
-		InputValue += TransientInputValue;
+		InputValue += FinalDelta;
 	}
 	else
 	{
-		InputValue = TransientInputValue;
+		InputValue = FinalDelta;
 	}
 
 	InputValue = SlotNode->Normalize.NormalizeValue(InputValue);
@@ -88,11 +95,22 @@ void FCameraRigInput1DSlotEvaluator::OnSerialize(const FCameraNodeEvaluatorSeria
 {
 	Super::OnSerialize(Params, Ar);
 
-	Ar << TransientInputValue;
+	Ar << DeltaInputValue;
 }
 
 }  // namespace UE::Cameras
 
+void UCameraRigInput1DSlot::PostLoad()
+{
+	if (!InputSlotParameters_DEPRECATED.bIsPreBlended)
+	{
+		InputSlotParameters_DEPRECATED.bIsPreBlended = true;
+		bIsPreBlended = false;
+	}
+
+	Super::PostLoad();
+}
+
 void UCameraRigInput1DSlot::OnBuild(FCameraObjectBuildContext& BuildContext) 
 {
 	using namespace UE::Cameras;
@@ -107,7 +125,7 @@ void UCameraRigInput1DSlot::OnBuild(FCameraObjectBuildContext& BuildContext)
 	{
 		VariableDefinition = CustomVariable.Variable->GetVariableDefinition();
 	}
-	else if (InputSlotParameters.bIsPreBlended)
+	else if (bIsPreBlended)
 	{
 		BuildContext.BuildLog.AddMessage(
 				EMessageSeverity::Error,
@@ -124,11 +142,11 @@ void UCameraRigInput1DSlot::OnBuild(FCameraObjectBuildContext& BuildContext)
 		FCameraVariableTableAllocationInfo& VariableTableInfo = BuildContext.AllocationInfo.VariableTableInfo;
 		VariableTableInfo.VariableDefinitions.Add(VariableDefinition);
 
-		FCameraVariableDefinition TransientVariableDefinition = VariableDefinition.CreateVariant(TEXT("Transient"));
-		VariableTableInfo.VariableDefinitions.Add(TransientVariableDefinition);
+		FCameraVariableDefinition SpeedVariableDefinition = VariableDefinition.CreateVariant(TEXT("Speed"));
+		VariableTableInfo.VariableDefinitions.Add(SpeedVariableDefinition);
 
 		VariableID = VariableDefinition.VariableID;
-		TransientVariableID = TransientVariableDefinition.VariableID;
+		SpeedVariableID = SpeedVariableDefinition.VariableID;
 	}
 }
 
diff --git a/Engine/Plugins/Cameras/GameplayCameras/Source/GameplayCameras/Private/Nodes/Input/CameraRigInput2DSlot.cpp b/Engine/Plugins/Cameras/GameplayCameras/Source/GameplayCameras/Private/Nodes/Input/CameraRigInput2DSlot.cpp
--- a/Engine/Plugins/Cameras/GameplayCameras/Source/GameplayCameras/Private/Nodes/Input/CameraRigInput2DSlot.cpp
+++ b/Engine/Plugins/Cameras/GameplayCameras/Source/GameplayCameras/Private/Nodes/Input/CameraRigInput2DSlot.cpp
@@ -27,10 +27,15 @@ FCameraRigInput2DSlotEvaluator::FCameraRigInput2DSlotEvaluator()
 
 void FCameraRigInput2DSlotEvaluator::OnInitialize(const FCameraNodeEvaluatorInitializeParams& Params, FCameraNodeEvaluationResult& OutResult)
 {
-	TransientInputValue = FVector2d::ZeroVector;
+	const UCameraRigInput2DSlot* SlotNode = GetCameraNodeAs<UCameraRigInput2DSlot>();
+
+	RevertAxisXReader.Initialize(SlotNode->RevertAxisX);
+	RevertAxisYReader.Initialize(SlotNode->RevertAxisY);
+	SpeedReader.Initialize(SlotNode->Speed);
+
+	DeltaInputValue = FVector2d::ZeroVector;
 	InputValue = FVector2d::ZeroVector;
 
-	const UCameraRigInput2DSlot* SlotNode = GetCameraNodeAs<UCameraRigInput2DSlot>();
 	if (SlotNode->GetVariableID().IsValid() && Params.LastActiveCameraRigInfo.LastResult)
 	{
 		const FCameraVariableTable& LastActiveRigVariableTable = Params.LastActiveCameraRigInfo.LastResult->VariableTable;
@@ -41,9 +46,10 @@ void FCameraRigInput2DSlotEvaluator::OnInitialize(const FCameraNodeEvaluatorInit
 void FCameraRigInput2DSlotEvaluator::OnUpdateParameters(const FCameraBlendedParameterUpdateParams& Params, FCameraBlendedParameterUpdateResult& OutResult)
 {
 	const UCameraRigInput2DSlot* SlotNode = GetCameraNodeAs<UCameraRigInput2DSlot>();
-	if (SlotNode->InputSlotParameters.bIsPreBlended)
+	if (SlotNode->bIsPreBlended)
 	{
-		OutResult.VariableTable.SetValue<FVector2d>(SlotNode->GetTransientVariableID(), TransientInputValue);
+		const FVector2d Speed = SpeedReader.Get(OutResult.VariableTable);
+		OutResult.VariableTable.SetValue<FVector2d>(SlotNode->GetSpeedVariableID(), Speed);
 	}
 }
 
@@ -51,18 +57,29 @@ void FCameraRigInput2DSlotEvaluator::OnRun(const FCameraNodeEvaluationParams& Pa
 {
 	const UCameraRigInput2DSlot* SlotNode = GetCameraNodeAs<UCameraRigInput2DSlot>();
 
-	if (SlotNode->InputSlotParameters.bIsPreBlended)
+	const FVector2d Speed = 
+		SlotNode->bIsPreBlended ?
+			OutResult.VariableTable.GetValue<FVector2d>(SlotNode->GetSpeedVariableID()) :
+			SpeedReader.Get(OutResult.VariableTable);
+
+	FVector2d FinalDelta(DeltaInputValue.X * Speed.X * Params.DeltaTime, DeltaInputValue.Y * Speed.Y * Params.DeltaTime);
+
+	if (RevertAxisXReader.Get(OutResult.VariableTable))
+	{
+		FinalDelta.X = -FinalDelta.X;
+	}
+	if (RevertAxisYReader.Get(OutResult.VariableTable))
 	{
-		TransientInputValue = OutResult.VariableTable.GetValue<FVector2d>(SlotNode->GetTransientVariableID());
+		FinalDelta.Y = -FinalDelta.Y;
 	}
 
-	if (SlotNode->InputSlotParameters.bIsAccumulated)
+	if (bIsAccumulated)
 	{
-		InputValue += TransientInputValue;
+		InputValue += FinalDelta;
 	}
 	else
 	{
-		InputValue = TransientInputValue;
+		InputValue = FinalDelta;
 	}
 
 	InputValue.X = SlotNode->NormalizeX.NormalizeValue(InputValue.X);
@@ -95,11 +112,22 @@ void FCameraRigInput2DSlotEvaluator::OnSerialize(const FCameraNodeEvaluatorSeria
 {
 	Super::OnSerialize(Params, Ar);
 
-	Ar << TransientInputValue;
+	Ar << DeltaInputValue;
 }
 
 }  // namespace UE::Cameras
 
+void UCameraRigInput2DSlot::PostLoad()
+{
+	if (!InputSlotParameters_DEPRECATED.bIsPreBlended)
+	{
+		InputSlotParameters_DEPRECATED.bIsPreBlended = true;
+		bIsPreBlended = false;
+	}
+
+	Super::PostLoad();
+}
+
 void UCameraRigInput2DSlot::OnBuild(FCameraObjectBuildContext& BuildContext)
 {
 	using namespace UE::Cameras;
@@ -114,7 +142,7 @@ void UCameraRigInput2DSlot::OnBuild(FCameraObjectBuildContext& BuildContext)
 	{
 		VariableDefinition = CustomVariable.Variable->GetVariableDefinition();
 	}
-	else if (InputSlotParameters.bIsPreBlended)
+	else if (bIsPreBlended)
 	{
 		BuildContext.BuildLog.AddMessage(
 				EMessageSeverity::Error,
@@ -131,11 +159,11 @@ void UCameraRigInput2DSlot::OnBuild(FCameraObjectBuildContext& BuildContext)
 		FCameraVariableTableAllocationInfo& VariableTableInfo = BuildContext.AllocationInfo.VariableTableInfo;
 		VariableTableInfo.VariableDefinitions.Add(VariableDefinition);
 
-		FCameraVariableDefinition TransientVariableDefinition = VariableDefinition.CreateVariant(TEXT("Transient"));
-		VariableTableInfo.VariableDefinitions.Add(TransientVariableDefinition);
+		FCameraVariableDefinition SpeedVariableDefinition = VariableDefinition.CreateVariant(TEXT("Speed"));
+		VariableTableInfo.VariableDefinitions.Add(SpeedVariableDefinition);
 
 		VariableID = VariableDefinition.VariableID;
-		TransientVariableID = TransientVariableDefinition.VariableID;
+		SpeedVariableID = SpeedVariableDefinition.VariableID;
 	}
 }
 
diff --git a/Engine/Plugins/Cameras/GameplayCameras/Source/GameplayCameras/Private/Nodes/Input/InputAxisBinding2DCameraNode.cpp b/Engine/Plugins/Cameras/GameplayCameras/Source/GameplayCameras/Private/Nodes/Input/InputAxisBinding2DCameraNode.cpp
--- a/Engine/Plugins/Cameras/GameplayCameras/Source/GameplayCameras/Private/Nodes/Input/InputAxisBinding2DCameraNode.cpp
+++ b/Engine/Plugins/Cameras/GameplayCameras/Source/GameplayCameras/Private/Nodes/Input/InputAxisBinding2DCameraNode.cpp
@@ -3,14 +3,9 @@
 #include "Nodes/Input/InputAxisBinding2DCameraNode.h"
 
 #include "Components/InputComponent.h"
-#include "Core/CameraEvaluationContext.h"
 #include "Core/CameraParameterReader.h"
-#include "Core/CameraSystemEvaluator.h"
 #include "EnhancedInputComponent.h"
-#include "GameFramework/Actor.h"
-#include "GameplayCameras.h"
-#include "InputAction.h"
-#include "UObject/Package.h"
+#include "Nodes/Input/InputAxisBindingHelpers.h"
 
 namespace UE::Cameras
 {
@@ -23,16 +18,13 @@ protected:
 
 	// FCameraNodeEvaluator interface.
 	virtual void OnInitialize(const FCameraNodeEvaluatorInitializeParams& Params, FCameraNodeEvaluationResult& OutResult) override;
-	virtual void OnUpdateParameters(const FCameraBlendedParameterUpdateParams& Params, FCameraBlendedParameterUpdateResult& OutResult) override;
+	virtual void OnRun(const FCameraNodeEvaluationParams& Params, FCameraNodeEvaluationResult& OutResult) override;
 
 private:
 
-	TObjectPtr<UEnhancedInputComponent> InputComponent;
-
-	TCameraParameterReader<bool> RevertAxisXReader;
-	TCameraParameterReader<bool> RevertAxisYReader;
 	TCameraParameterReader<FVector2d> MultiplierReader;
 
+	TObjectPtr<UEnhancedInputComponent> InputComponent;
 	TArray<FEnhancedInputActionValueBinding*> AxisValueBindings;
 };
 
@@ -40,87 +32,37 @@ UE_DEFINE_CAMERA_NODE_EVALUATOR(FInputAxisBinding2DCameraNodeEvaluator)
 
 void FInputAxisBinding2DCameraNodeEvaluator::OnInitialize(const FCameraNodeEvaluatorInitializeParams& Params, FCameraNodeEvaluationResult& OutResult)
 {
-	UObject* ContextOwner = Params.EvaluationContext->GetOwner();
-	if (ContextOwner)
-	{
-		if (AActor* ContextOwnerActor = Cast<AActor>(ContextOwner))
-		{
-			InputComponent = Cast<UEnhancedInputComponent>(ContextOwnerActor->InputComponent);
-		}
-		else if (AActor* OuterActor = ContextOwner->GetTypedOuter<AActor>())
-		{
-			InputComponent = Cast<UEnhancedInputComponent>(OuterActor->InputComponent);
-		}
-	}
-
 	const UInputAxisBinding2DCameraNode* AxisBindingNode = GetCameraNodeAs<UInputAxisBinding2DCameraNode>();
 
-	RevertAxisXReader.Initialize(AxisBindingNode->RevertAxisX);
-	RevertAxisYReader.Initialize(AxisBindingNode->RevertAxisY);
 	MultiplierReader.Initialize(AxisBindingNode->Multiplier);
 
-	if (InputComponent)
-	{
-		for (TObjectPtr<UInputAction> AxisAction : AxisBindingNode->AxisActions)
-		{
-			FEnhancedInputActionValueBinding* AxisValueBinding = &InputComponent->BindActionValue(AxisAction);
-			AxisValueBindings.Add(AxisValueBinding);
-		}
-	}
-	else if (Params.Evaluator->GetRole() == ECameraSystemEvaluatorRole::Game)
-	{
-		UE_LOG(LogCameraSystem, Error, TEXT("No input component found on context owner '%s' for node '%s' in '%s'."),
-				*GetNameSafe(ContextOwner), 
-				*GetNameSafe(AxisBindingNode),
-				*GetNameSafe(AxisBindingNode ? AxisBindingNode->GetOutermost() : nullptr));
-	}
+	InputComponent = FInputAxisBindingHelpers::FindInputComponent(Params);
+	FInputAxisBindingHelpers::BindActionValues(Params, AxisBindingNode, InputComponent, AxisBindingNode->AxisActions, AxisValueBindings);
 
 	Super::OnInitialize(Params, OutResult);
 }
 
-void FInputAxisBinding2DCameraNodeEvaluator::OnUpdateParameters(const FCameraBlendedParameterUpdateParams& Params, FCameraBlendedParameterUpdateResult& OutResult)
+void FInputAxisBinding2DCameraNodeEvaluator::OnRun(const FCameraNodeEvaluationParams& Params, FCameraNodeEvaluationResult& OutResult)
 {
-	const UInputAxisBinding2DCameraNode* AxisBindingNode = GetCameraNodeAs<UInputAxisBinding2DCameraNode>();
-
-	FVector2d HighestValue(FVector2d::ZeroVector);
-	double HighestSquaredLenth = 0.f;
-
-	for (FEnhancedInputActionValueBinding* AxisValueBinding : AxisValueBindings)
-	{
-		if (!AxisValueBinding)
-		{
-			continue;
-		}
-
-		const FVector2d Value = AxisValueBinding->GetValue().Get<FVector2D>();
-		const double ValueSquaredLength = Value.SquaredLength();
-		if (ValueSquaredLength > HighestSquaredLenth)
-		{
-			HighestValue = Value;
-		}
-	}
-
 	const FVector2d Multiplier = MultiplierReader.Get(OutResult.VariableTable);
-	TransientInputValue = FVector2d(HighestValue.X * Multiplier.X, HighestValue.Y * Multiplier.Y);
 
-	if (RevertAxisXReader.Get(OutResult.VariableTable))
-	{
-		TransientInputValue.X = -TransientInputValue.X;
-	}
-	if (RevertAxisYReader.Get(OutResult.VariableTable))
-	{
-		TransientInputValue.Y = -TransientInputValue.Y;
-	}
+	const FVector2d HighestValue = FInputAxisBindingHelpers::GetHighestValue(AxisValueBindings);
+	DeltaInputValue = FVector2d(HighestValue.X * Multiplier.X, HighestValue.Y * Multiplier.Y);
 
-	Super::OnUpdateParameters(Params, OutResult);
+	Super::OnRun(Params, OutResult);
 }
 
 }  // namespace UE::Cameras
 
-UInputAxisBinding2DCameraNode::UInputAxisBinding2DCameraNode(const FObjectInitializer& ObjInit)
-	: Super(ObjInit)
+void UInputAxisBinding2DCameraNode::PostLoad()
 {
-	Multiplier = FVector2D(1, 1);
+	if (!InputSlotParameters_DEPRECATED.bIsAccumulated)
+	{
+		InputSlotParameters_DEPRECATED.bIsAccumulated = true;
+		bIsAccumulated = false;
+	}
+
+	Super::PostLoad();
 }
 
 FCameraNodeEvaluatorPtr UInputAxisBinding2DCameraNode::OnBuildEvaluator(FCameraNodeEvaluatorBuilder& Builder) const
diff --git a/Engine/Plugins/Cameras/GameplayCameras/Source/GameplayCameras/Private/Nodes/Input/InputAxisBindingHelpers.cpp b/Engine/Plugins/Cameras/GameplayCameras/Source/GameplayCameras/Private/Nodes/Input/InputAxisBindingHelpers.cpp
new file mode 100644
--- /dev/null
+++ b/Engine/Plugins/Cameras/GameplayCameras/Source/GameplayCameras/Private/Nodes/Input/InputAxisBindingHelpers.cpp
@@ -0,0 +1,82 @@
+// Copyright Epic Games, Inc. All Rights Reserved.
+
+#include "Nodes/Input/InputAxisBindingHelpers.h"
+
+#include "Components/InputComponent.h"
+#include "Core/CameraEvaluationContext.h"
+#include "Core/CameraSystemEvaluator.h"
+#include "EnhancedInputComponent.h"
+#include "GameFramework/Actor.h"
+#include "UObject/Package.h"
+
+namespace UE::Cameras
+{
+
+UEnhancedInputComponent* FInputAxisBindingHelpers::FindInputComponent(const UE::Cameras::FCameraNodeEvaluatorInitializeParams& Params)
+{
+	UObject* ContextOwner = Params.EvaluationContext->GetOwner();
+	if (ContextOwner)
+	{
+		if (AActor* ContextOwnerActor = Cast<AActor>(ContextOwner))
+		{
+			return Cast<UEnhancedInputComponent>(ContextOwnerActor->InputComponent);
+		}
+		else if (AActor* OuterActor = ContextOwner->GetTypedOuter<AActor>())
+		{
+			return Cast<UEnhancedInputComponent>(OuterActor->InputComponent);
+		}
+	}
+	return nullptr;
+}
+
+void FInputAxisBindingHelpers::BindActionValues(
+		const UE::Cameras::FCameraNodeEvaluatorInitializeParams& Params,
+		const UCameraNode* CameraNode,
+		UEnhancedInputComponent* InputComponent,
+		const TArray<TObjectPtr<UInputAction>>& AxisActions,
+		TArray<FEnhancedInputActionValueBinding*>& OutAxisValueBindings)
+{
+	UObject* ContextOwner = Params.EvaluationContext->GetOwner();
+
+	if (InputComponent)
+	{
+		for (TObjectPtr<UInputAction> AxisAction : AxisActions)
+		{
+			FEnhancedInputActionValueBinding* AxisValueBinding = &InputComponent->BindActionValue(AxisAction);
+			OutAxisValueBindings.Add(AxisValueBinding);
+		}
+	}
+	else if (Params.Evaluator->GetRole() == ECameraSystemEvaluatorRole::Game)
+	{
+		UE_LOG(LogCameraSystem, Error, TEXT("No input component found on context owner '%s' for node '%s' in '%s'."),
+				*GetNameSafe(ContextOwner), 
+				*GetNameSafe(CameraNode),
+				*GetNameSafe(CameraNode ? CameraNode->GetOutermost() : nullptr));
+	}
+}
+
+FVector2d FInputAxisBindingHelpers::GetHighestValue(const TArray<FEnhancedInputActionValueBinding*>& AxisValueBindings)
+{
+	FVector2d HighestValue(FVector2d::ZeroVector);
+	double HighestSquaredLenth = 0.f;
+
+	for (FEnhancedInputActionValueBinding* AxisValueBinding : AxisValueBindings)
+	{
+		if (!AxisValueBinding)
+		{
+			continue;
+		}
+
+		const FVector2d Value = AxisValueBinding->GetValue().Get<FVector2D>();
+		const double ValueSquaredLength = Value.SquaredLength();
+		if (ValueSquaredLength > HighestSquaredLenth)
+		{
+			HighestValue = Value;
+		}
+	}
+
+	return HighestValue;
+}
+
+}  // namespace UE::Cameras
+
diff --git a/Engine/Plugins/Cameras/GameplayCameras/Source/GameplayCameras/Private/Nodes/Input/InputAxisBindingHelpers.h b/Engine/Plugins/Cameras/GameplayCameras/Source/GameplayCameras/Private/Nodes/Input/InputAxisBindingHelpers.h
new file mode 100644
--- /dev/null
+++ b/Engine/Plugins/Cameras/GameplayCameras/Source/GameplayCameras/Private/Nodes/Input/InputAxisBindingHelpers.h
@@ -0,0 +1,35 @@
+// Copyright Epic Games, Inc. All Rights Reserved.
+
+#pragma once
+
+#include "CoreTypes.h"
+#include "Containers/Array.h"
+
+class UCameraNode;
+class UEnhancedInputComponent;
+class UInputAction;
+struct FEnhancedInputActionValueBinding;
+
+namespace UE::Cameras { struct FCameraNodeEvaluatorInitializeParams; }
+
+namespace UE::Cameras
+{
+
+class FInputAxisBindingHelpers
+{
+public:
+
+	static UEnhancedInputComponent* FindInputComponent(const UE::Cameras::FCameraNodeEvaluatorInitializeParams& Params);
+
+	static void BindActionValues(
+			const UE::Cameras::FCameraNodeEvaluatorInitializeParams& Params,
+			const UCameraNode* CameraNode,
+			UEnhancedInputComponent* InputComponent,
+			const TArray<TObjectPtr<UInputAction>>& AxisActions,
+			TArray<FEnhancedInputActionValueBinding*>& OutAxisValueBindings);
+
+	static FVector2d GetHighestValue(const TArray<FEnhancedInputActionValueBinding*>& AxisValueBindings);
+};
+
+}  // namespace UE::Cameras
+
diff --git a/Engine/Plugins/Cameras/GameplayCameras/Source/GameplayCameras/Public/Nodes/Input/CameraRigInput1DSlot.h b/Engine/Plugins/Cameras/GameplayCameras/Source/GameplayCameras/Public/Nodes/Input/CameraRigInput1DSlot.h
--- a/Engine/Plugins/Cameras/GameplayCameras/Source/GameplayCameras/Public/Nodes/Input/CameraRigInput1DSlot.h
+++ b/Engine/Plugins/Cameras/GameplayCameras/Source/GameplayCameras/Public/Nodes/Input/CameraRigInput1DSlot.h
@@ -4,6 +4,8 @@
 
 #include "Core/BuiltInCameraVariables.h"
 #include "Core/CameraNodeEvaluator.h"
+#include "Core/CameraParameterReader.h"
+#include "Core/CameraParameters.h"
 #include "Core/CameraVariableReferences.h"
 #include "Nodes/Input/CameraRigInputSlotTypes.h"
 #include "Nodes/Input/Input1DCameraNode.h"
@@ -23,10 +25,6 @@ class UCameraRigInput1DSlot : public UInput1DCameraNode
 
 public:
 
-	/** Input processing parameters. */
-	UPROPERTY(EditAnywhere, Category="Input", meta=(ShowOnlyInnerProperties))
-	FCameraRigInputSlotParameters InputSlotParameters;
-
 	/** Clamping of the final input value. */
 	UPROPERTY(EditAnywhere, Category="Input")
 	FCameraParameterClamping Clamp;
@@ -35,6 +33,14 @@ public:
 	UPROPERTY(EditAnywhere, Category="Input")
 	FCameraParameterNormalization Normalize;
 
+	/** Whether to revert the axis. */
+	UPROPERTY(EditAnywhere, Category="Input")
+	FBooleanCameraParameter RevertAxis = false;
+
+	/** A speed, in units/seconds, to use on the input value. */
+	UPROPERTY(EditAnywhere, Category="Input")
+	FDoubleCameraParameter Speed = 1.0;
+
 	/** The variable to use to blend with other input slots. */
 	UPROPERTY(EditAnywhere, Category="Input")
 	EBuiltInDoubleCameraVariable BuiltInVariable = EBuiltInDoubleCameraVariable::Yaw;
@@ -43,10 +49,25 @@ public:
 	UPROPERTY(EditAnywhere, Category="Input", meta=(EditCondition="BuiltInVariable == EBuiltInDoubleCameraVariable::None"))
 	FDoubleCameraVariableReference CustomVariable;
 
+	/** Whether the multiplier should be pre-blended. */
+	UPROPERTY(EditAnywhere, Category="Input")
+	bool bIsPreBlended = true;
+
+
+	// Deprecated
+
+	UPROPERTY()
+	FCameraRigInputSlotParameters InputSlotParameters_DEPRECATED;
+
 public:
 
 	FCameraVariableID GetVariableID() const { return VariableID; }
-	FCameraVariableID GetTransientVariableID() const { return TransientVariableID; }
+	FCameraVariableID GetSpeedVariableID() const { return SpeedVariableID; }
+
+public:
+
+	// UObject interface.
+	virtual void PostLoad() override;
 
 protected:
 
@@ -57,7 +78,7 @@ protected:
 private:
 
 	UPROPERTY()
-	FCameraVariableID TransientVariableID;
+	FCameraVariableID SpeedVariableID;
 	UPROPERTY()
 	FCameraVariableID VariableID;
 };
@@ -86,7 +107,11 @@ protected:
 
 protected:
 
-	double TransientInputValue = 0.f;
+	TCameraParameterReader<bool> RevertAxisReader;
+	TCameraParameterReader<double> SpeedReader;
+
+	double DeltaInputValue = 0.f;
+	bool bIsAccumulated = true;
 };
 
 }  // namespace UE::Cameras
diff --git a/Engine/Plugins/Cameras/GameplayCameras/Source/GameplayCameras/Public/Nodes/Input/CameraRigInput2DSlot.h b/Engine/Plugins/Cameras/GameplayCameras/Source/GameplayCameras/Public/Nodes/Input/CameraRigInput2DSlot.h
--- a/Engine/Plugins/Cameras/GameplayCameras/Source/GameplayCameras/Public/Nodes/Input/CameraRigInput2DSlot.h
+++ b/Engine/Plugins/Cameras/GameplayCameras/Source/GameplayCameras/Public/Nodes/Input/CameraRigInput2DSlot.h
@@ -4,6 +4,8 @@
 
 #include "Core/BuiltInCameraVariables.h"
 #include "Core/CameraNodeEvaluator.h"
+#include "Core/CameraParameterReader.h"
+#include "Core/CameraParameters.h"
 #include "Core/CameraVariableReferences.h"
 #include "Nodes/Input/CameraRigInputSlotTypes.h"
 #include "Nodes/Input/Input2DCameraNode.h"
@@ -23,10 +25,6 @@ class UCameraRigInput2DSlot : public UInput2DCameraNode
 
 public:
 
-	/** Input processing parameters. */
-	UPROPERTY(EditAnywhere, Category="Input", meta=(ShowOnlyInnerProperties))
-	FCameraRigInputSlotParameters InputSlotParameters;
-
 	/** Clamping of the final input value. */
 	UPROPERTY(EditAnywhere, Category="Input")
 	FCameraParameterClamping ClampX;
@@ -43,6 +41,18 @@ public:
 	UPROPERTY(EditAnywhere, Category="Input")
 	FCameraParameterNormalization NormalizeY;
 
+	/** Whether to revert the X axis. */
+	UPROPERTY(EditAnywhere, Category="Input")
+	FBooleanCameraParameter RevertAxisX = false;
+
+	/** Whether to revert the Y axis. */
+	UPROPERTY(EditAnywhere, Category="Input")
+	FBooleanCameraParameter RevertAxisY = false;
+
+	/** A speed, in units/seconds, to use on the input value. */
+	UPROPERTY(EditAnywhere, Category="Input")
+	FVector2dCameraParameter Speed = FVector2d{ 1.0, 1.0 };
+
 	/** The variable to use to blend with other input slots. */
 	UPROPERTY(EditAnywhere, Category="Input")
 	EBuiltInVector2dCameraVariable BuiltInVariable = EBuiltInVector2dCameraVariable::YawPitch;
@@ -51,10 +61,25 @@ public:
 	UPROPERTY(EditAnywhere, Category="Input", meta=(EditCondition="BuiltInVariable == EBuiltInVector2dCameraVariable::None"))
 	FVector2dCameraVariableReference CustomVariable;
 
+	/** Whether parameters such as speed should be pre-blended. */
+	UPROPERTY(EditAnywhere, Category="Input")
+	bool bIsPreBlended = true;
+
+
+	// Deprecated
+
+	UPROPERTY()
+	FCameraRigInputSlotParameters InputSlotParameters_DEPRECATED;
+
 public:
 
 	FCameraVariableID GetVariableID() const { return VariableID; }
-	FCameraVariableID GetTransientVariableID() const { return TransientVariableID; }
+	FCameraVariableID GetSpeedVariableID() const { return SpeedVariableID; }
+
+public:
+
+	// UObject interface.
+	virtual void PostLoad() override;
 
 protected:
 
@@ -65,7 +90,7 @@ protected:
 private:
 
 	UPROPERTY()
-	FCameraVariableID TransientVariableID;
+	FCameraVariableID SpeedVariableID;
 	UPROPERTY()
 	FCameraVariableID VariableID;
 };
@@ -94,7 +119,12 @@ protected:
 
 protected:
 
-	FVector2d TransientInputValue = FVector2d::ZeroVector;
+	TCameraParameterReader<bool> RevertAxisXReader;
+	TCameraParameterReader<bool> RevertAxisYReader;
+	TCameraParameterReader<FVector2d> SpeedReader;
+
+	FVector2d DeltaInputValue = FVector2d::ZeroVector;
+	bool bIsAccumulated = true;
 };
 
 }  // namespace UE::Cameras
diff --git a/Engine/Plugins/Cameras/GameplayCameras/Source/GameplayCameras/Public/Nodes/Input/InputAxisBinding2DCameraNode.h b/Engine/Plugins/Cameras/GameplayCameras/Source/GameplayCameras/Public/Nodes/Input/InputAxisBinding2DCameraNode.h
--- a/Engine/Plugins/Cameras/GameplayCameras/Source/GameplayCameras/Public/Nodes/Input/InputAxisBinding2DCameraNode.h
+++ b/Engine/Plugins/Cameras/GameplayCameras/Source/GameplayCameras/Public/Nodes/Input/InputAxisBinding2DCameraNode.h
@@ -11,7 +11,9 @@
 class UInputAction;
 
 /**
- * An input node that reads player input from an input action.
+ * An input node that reads player input from an input action and accumulates
+ * it into a usable input value. Basically a Raw Input Axis Binding node combined
+ * with an Input Accumulator node.
  */
 UCLASS(MinimalAPI, meta=(CameraNodeCategories="Input"))
 class UInputAxisBinding2DCameraNode : public UCameraRigInput2DSlot
@@ -24,21 +26,18 @@ public:
 	UPROPERTY(EditAnywhere, Category="Input")
 	TArray<TObjectPtr<UInputAction>> AxisActions;
 
-	/** Whether to revert the X axis. */
-	UPROPERTY(EditAnywhere, Category="Input Processing")
-	FBooleanCameraParameter RevertAxisX = false;
-
-	/** Whether to revert the Y axis. */
-	UPROPERTY(EditAnywhere, Category="Input Processing")
-	FBooleanCameraParameter RevertAxisY = false;
-
 	/** A multiplier to use on the input values. */
-	UPROPERTY(EditAnywhere, Category="Input Processing")
-	FVector2dCameraParameter Multiplier;
+	UPROPERTY(EditAnywhere, Category="Input")
+	FVector2dCameraParameter Multiplier = FVector2d(1.0, 1.0);
+
+	/** Whether the player input is accumulated from frame to frame. */
+	UPROPERTY(EditAnywhere, Category="Input")
+	bool bIsAccumulated = true;
 
 public:
 
-	UInputAxisBinding2DCameraNode(const FObjectInitializer& ObjInit);
+	// UObject interface.
+	virtual void PostLoad() override;
 	
 protected:
 
