Cameras: add raw input action node for complex input chains

diff --git a/Engine/Plugins/Cameras/GameplayCameras/Source/GameplayCameras/Private/Nodes/Input/RawInputAxisBinding2DCameraNode.cpp b/Engine/Plugins/Cameras/GameplayCameras/Source/GameplayCameras/Private/Nodes/Input/RawInputAxisBinding2DCameraNode.cpp
new file mode 100644
--- /dev/null
+++ b/Engine/Plugins/Cameras/GameplayCameras/Source/GameplayCameras/Private/Nodes/Input/RawInputAxisBinding2DCameraNode.cpp
@@ -0,0 +1,62 @@
+// Copyright Epic Games, Inc. All Rights Reserved.
+
+#include "Nodes/Input/RawInputAxisBinding2DCameraNode.h"
+
+#include "Components/InputComponent.h"
+#include "Core/CameraParameterReader.h"
+#include "EnhancedInputComponent.h"
+#include "Nodes/Input/InputAxisBindingHelpers.h"
+
+namespace UE::Cameras
+{
+
+class FRawInputAxisBinding2DCameraNodeEvaluator : public FInput2DCameraNodeEvaluator
+{
+	UE_DECLARE_CAMERA_NODE_EVALUATOR_EX(GAMEPLAYCAMERAS_API, FRawInputAxisBinding2DCameraNodeEvaluator, FInput2DCameraNodeEvaluator)
+
+protected:
+
+	// FCameraNodeEvaluator interface.
+	virtual void OnInitialize(const FCameraNodeEvaluatorInitializeParams& Params, FCameraNodeEvaluationResult& OutResult) override;
+	virtual void OnRun(const FCameraNodeEvaluationParams& Params, FCameraNodeEvaluationResult& OutResult) override;
+
+private:
+
+	TCameraParameterReader<FVector2d> MultiplierReader;
+
+	TObjectPtr<UEnhancedInputComponent> InputComponent;
+	TArray<FEnhancedInputActionValueBinding*> AxisValueBindings;
+};
+
+UE_DEFINE_CAMERA_NODE_EVALUATOR(FRawInputAxisBinding2DCameraNodeEvaluator)
+
+void FRawInputAxisBinding2DCameraNodeEvaluator::OnInitialize(const FCameraNodeEvaluatorInitializeParams& Params, FCameraNodeEvaluationResult& OutResult)
+{
+	const URawInputAxisBinding2DCameraNode* AxisBindingNode = GetCameraNodeAs<URawInputAxisBinding2DCameraNode>();
+
+	MultiplierReader.Initialize(AxisBindingNode->Multiplier);
+
+	InputComponent = FInputAxisBindingHelpers::FindInputComponent(Params);
+	FInputAxisBindingHelpers::BindActionValues(Params, AxisBindingNode, InputComponent, AxisBindingNode->AxisActions, AxisValueBindings);
+
+	Super::OnInitialize(Params, OutResult);
+}
+
+void FRawInputAxisBinding2DCameraNodeEvaluator::OnRun(const FCameraNodeEvaluationParams& Params, FCameraNodeEvaluationResult& OutResult)
+{
+	const FVector2d Multiplier = MultiplierReader.Get(OutResult.VariableTable);
+
+	const FVector2d HighestValue = FInputAxisBindingHelpers::GetHighestValue(AxisValueBindings);
+	InputValue = FVector2d(HighestValue.X * Multiplier.X, HighestValue.Y * Multiplier.Y);
+
+	Super::OnRun(Params, OutResult);
+}
+
+}  // namespace UE::Cameras
+
+FCameraNodeEvaluatorPtr URawInputAxisBinding2DCameraNode::OnBuildEvaluator(FCameraNodeEvaluatorBuilder& Builder) const
+{
+	using namespace UE::Cameras;
+	return Builder.BuildEvaluator<FRawInputAxisBinding2DCameraNodeEvaluator>();
+}
+
diff --git a/Engine/Plugins/Cameras/GameplayCameras/Source/GameplayCameras/Public/Nodes/Input/RawInputAxisBinding2DCameraNode.h b/Engine/Plugins/Cameras/GameplayCameras/Source/GameplayCameras/Public/Nodes/Input/RawInputAxisBinding2DCameraNode.h
new file mode 100644
--- /dev/null
+++ b/Engine/Plugins/Cameras/GameplayCameras/Source/GameplayCameras/Public/Nodes/Input/RawInputAxisBinding2DCameraNode.h
@@ -0,0 +1,37 @@
+// Copyright Epic Games, Inc. All Rights Reserved.
+
+#pragma once
+
+#include "Core/CameraNodeEvaluator.h"
+#include "Core/CameraParameters.h"
+#include "Math/MathFwd.h"
+#include "Nodes/Input/Input2DCameraNode.h"
+
+#include "RawInputAxisBinding2DCameraNode.generated.h"
+
+class UInputAction;
+
+/**
+ * An input node that reads player input from an input action.
+ */
+UCLASS(MinimalAPI, meta=(CameraNodeCategories="Input"))
+class URawInputAxisBinding2DCameraNode : public UInput2DCameraNode
+{
+	GENERATED_BODY()
+
+public:
+
+	/** The axis input action(s) to read from. */
+	UPROPERTY(EditAnywhere, Category="Input")
+	TArray<TObjectPtr<UInputAction>> AxisActions;
+
+	/** A multiplier to use on the input values. */
+	UPROPERTY(EditAnywhere, Category="Input")
+	FVector2dCameraParameter Multiplier = FVector2d(1.0, 1.0);
+
+protected:
+
+	// UCameraNode interface.
+	virtual FCameraNodeEvaluatorPtr OnBuildEvaluator(FCameraNodeEvaluatorBuilder& Builder) const override;
+};
+
