Cameras: add input accumulator node for complex input chains

diff --git a/Engine/Plugins/Cameras/GameplayCameras/Source/GameplayCameras/Private/Nodes/Input/InputAccumulator2DCameraNode.cpp b/Engine/Plugins/Cameras/GameplayCameras/Source/GameplayCameras/Private/Nodes/Input/InputAccumulator2DCameraNode.cpp
new file mode 100644
--- /dev/null
+++ b/Engine/Plugins/Cameras/GameplayCameras/Source/GameplayCameras/Private/Nodes/Input/InputAccumulator2DCameraNode.cpp
@@ -0,0 +1,83 @@
+// Copyright Epic Games, Inc. All Rights Reserved.
+
+#include "Nodes/Input/InputAccumulator2DCameraNode.h"
+
+#include "Nodes/Input/CameraRigInput2DSlot.h"
+
+#define LOCTEXT_NAMESPACE "InputAccumulator2DCameraNode"
+
+namespace UE::Cameras
+{
+
+class FInputAccumulator2DCameraNodeEvaluator : public FCameraRigInput2DSlotEvaluator
+{
+	UE_DECLARE_CAMERA_NODE_EVALUATOR_EX(GAMEPLAYCAMERAS_API, FInputAccumulator2DCameraNodeEvaluator, FCameraRigInput2DSlotEvaluator)
+
+public:
+
+	FInputAccumulator2DCameraNodeEvaluator();
+
+protected:
+
+	// FCameraNodeEvaluator interface.
+	virtual void OnBuild(const FCameraNodeEvaluatorBuildParams& Params) override;
+	virtual FCameraNodeEvaluatorChildrenView OnGetChildren() override;
+	virtual void OnUpdateParameters(const FCameraBlendedParameterUpdateParams& Params, FCameraBlendedParameterUpdateResult& OutResult) override;
+	virtual void OnRun(const FCameraNodeEvaluationParams& Params, FCameraNodeEvaluationResult& OutResult) override;
+
+private:
+
+	FInput2DCameraNodeEvaluator* InputSlotEvaluator = nullptr;
+};
+
+UE_DEFINE_CAMERA_NODE_EVALUATOR(FInputAccumulator2DCameraNodeEvaluator)
+
+FInputAccumulator2DCameraNodeEvaluator::FInputAccumulator2DCameraNodeEvaluator()
+{
+	bIsAccumulated = true;
+}
+
+void FInputAccumulator2DCameraNodeEvaluator::OnBuild(const FCameraNodeEvaluatorBuildParams& Params)
+{
+	Super::OnBuild(Params);
+
+	const UInputAccumulator2DCameraNode* AccumulatorNode = GetCameraNodeAs<UInputAccumulator2DCameraNode>();
+	InputSlotEvaluator = Params.BuildEvaluatorAs<FInput2DCameraNodeEvaluator>(AccumulatorNode->InputSlot);
+}
+
+FCameraNodeEvaluatorChildrenView FInputAccumulator2DCameraNodeEvaluator::OnGetChildren()
+{
+	return FCameraNodeEvaluatorChildrenView{ InputSlotEvaluator };
+}
+
+void FInputAccumulator2DCameraNodeEvaluator::OnUpdateParameters(const FCameraBlendedParameterUpdateParams& Params, FCameraBlendedParameterUpdateResult& OutResult)
+{
+	if (InputSlotEvaluator)
+	{
+		InputSlotEvaluator->UpdateParameters(Params, OutResult);
+		DeltaInputValue = InputSlotEvaluator->GetInputValue();
+	}
+
+	Super::OnUpdateParameters(Params, OutResult);
+}
+
+void FInputAccumulator2DCameraNodeEvaluator::OnRun(const FCameraNodeEvaluationParams& Params, FCameraNodeEvaluationResult& OutResult)
+{
+	if (InputSlotEvaluator)
+	{
+		InputSlotEvaluator->Run(Params, OutResult);
+	}
+
+	Super::OnRun(Params, OutResult);
+}
+
+}  // namespace UE::Cameras
+
+FCameraNodeEvaluatorPtr UInputAccumulator2DCameraNode::OnBuildEvaluator(FCameraNodeEvaluatorBuilder& Builder) const
+{
+	using namespace UE::Cameras;
+	return Builder.BuildEvaluator<FInputAccumulator2DCameraNodeEvaluator>();
+}
+
+#undef LOCTEXT_NAMESPACE
+
diff --git a/Engine/Plugins/Cameras/GameplayCameras/Source/GameplayCameras/Public/Nodes/Input/InputAccumulator2DCameraNode.h b/Engine/Plugins/Cameras/GameplayCameras/Source/GameplayCameras/Public/Nodes/Input/InputAccumulator2DCameraNode.h
new file mode 100644
--- /dev/null
+++ b/Engine/Plugins/Cameras/GameplayCameras/Source/GameplayCameras/Public/Nodes/Input/InputAccumulator2DCameraNode.h
@@ -0,0 +1,33 @@
+// Copyright Epic Games, Inc. All Rights Reserved.
+
+#pragma once
+
+#include "Core/BuiltInCameraVariables.h"
+#include "Core/CameraParameters.h"
+#include "Core/CameraVariableReferences.h"
+#include "Math/MathFwd.h"
+#include "Nodes/Input/CameraRigInputSlotTypes.h"
+#include "Nodes/Input/CameraRigInput2DSlot.h"
+
+#include "InputAccumulator2DCameraNode.generated.h"
+
+/**
+ * A camera node that accumulates player input.
+ */
+UCLASS(MinimalAPI, meta=(CameraNodeCategories="Input"))
+class UInputAccumulator2DCameraNode : public UCameraRigInput2DSlot
+{
+	GENERATED_BODY()
+
+public:
+
+	/** Input node to accumulate values from. */
+	UPROPERTY()
+	TObjectPtr<UInput2DCameraNode> InputSlot;
+
+protected:
+
+	// UCameraNode interface.
+	virtual FCameraNodeEvaluatorPtr OnBuildEvaluator(FCameraNodeEvaluatorBuilder& Builder) const override;
+};
+
