Cameras: add BP utility function to access the evaluated camera's orientation

diff --git a/Engine/Plugins/Cameras/GameplayCameras/Source/GameplayCameras/Private/GameFramework/GameplayCameraComponentBase.cpp b/Engine/Plugins/Cameras/GameplayCameras/Source/GameplayCameras/Private/GameFramework/GameplayCameraComponentBase.cpp
--- a/Engine/Plugins/Cameras/GameplayCameras/Source/GameplayCameras/Private/GameFramework/GameplayCameraComponentBase.cpp
+++ b/Engine/Plugins/Cameras/GameplayCameras/Source/GameplayCameras/Private/GameFramework/GameplayCameraComponentBase.cpp
@@ -394,6 +394,18 @@ FBlueprintCameraEvaluationDataRef UGameplayCameraComponentBase::GetConditionalRe
 
 #undef UE_PRIVATE_GAMEPLAY_CAMERA_COMPONENT_VALIDATE_EVALUATION_CONTEXT
 
+FRotator UGameplayCameraComponentBase::GetEvaluatedCameraRotation() const
+{
+	using namespace UE::Cameras;
+
+	if (CameraSystemEvaluator.IsValid())
+	{
+		const FCameraSystemEvaluationResult& Result = CameraSystemEvaluator->GetEvaluatedResult();
+		return Result.CameraPose.GetRotation();
+	}
+	return FRotator();
+}
+
 void UGameplayCameraComponentBase::OnRegister()
 {
 	using namespace UE::Cameras;
diff --git a/Engine/Plugins/Cameras/GameplayCameras/Source/GameplayCameras/Public/GameFramework/GameplayCameraComponentBase.h b/Engine/Plugins/Cameras/GameplayCameras/Source/GameplayCameras/Public/GameFramework/GameplayCameraComponentBase.h
--- a/Engine/Plugins/Cameras/GameplayCameras/Source/GameplayCameras/Public/GameFramework/GameplayCameraComponentBase.h
+++ b/Engine/Plugins/Cameras/GameplayCameras/Source/GameplayCameras/Public/GameFramework/GameplayCameraComponentBase.h
@@ -123,6 +123,10 @@ public:
 	UFUNCTION(BlueprintPure, Category=Camera, meta=(DisplayName="Get Conditional Camera Data"))
 	GAMEPLAYCAMERAS_API FBlueprintCameraEvaluationDataRef GetConditionalResult(ECameraEvaluationDataCondition Condition) const;
 
+	/** Gets the last evaluated orientation of the camera. */
+	UFUNCTION(BlueprintPure, Category=Camera)
+	GAMEPLAYCAMERAS_API FRotator GetEvaluatedCameraRotation() const;
+
 public:
 
 	// UActorComponent interface
