Cameras: resolve camera rig name conflicts when building the property bag for a camera asset

Also add some unit-tests for it, including a builder class for camera assets.

diff --git a/Engine/Plugins/Cameras/GameplayCameras/Source/GameplayCameras/Private/Build/CameraAssetBuilder.cpp b/Engine/Plugins/Cameras/GameplayCameras/Source/GameplayCameras/Private/Build/CameraAssetBuilder.cpp
--- a/Engine/Plugins/Cameras/GameplayCameras/Source/GameplayCameras/Private/Build/CameraAssetBuilder.cpp
+++ b/Engine/Plugins/Cameras/GameplayCameras/Source/GameplayCameras/Private/Build/CameraAssetBuilder.cpp
@@ -93,17 +93,58 @@ void FCameraAssetBuilder::BuildCameraImpl()
 		CameraRigBuilder.BuildCameraRig(CameraRig);
 	}
 
-	// Get the list of all the camera rigs' interface parameters, and cache some information
-	// about them.
+	// Get the list of all the camera rigs' interface parameters.
+	TMap<FName, TArray<const UCameraRigAsset*>> UsedParameterNames;
+	TMap<const UCameraRigAsset*, TArray<FCameraObjectInterfaceParameterDefinition>> DefinitionsByCameraRig;
+	for (const UCameraRigAsset* CameraRig : AllCameraRigs)
+	{
+		if (ensure(!DefinitionsByCameraRig.Contains(CameraRig)))
+		{
+			TArray<FCameraObjectInterfaceParameterDefinition>& DefinitionsForCameraRig = DefinitionsByCameraRig.Add(CameraRig);
+			DefinitionsForCameraRig.Append(CameraRig->GetParameterDefinitions());
+
+			for (const FCameraObjectInterfaceParameterDefinition& Definition : DefinitionsForCameraRig)
+			{
+				UsedParameterNames.FindOrAdd(Definition.ParameterName).Add(CameraRig);
+			}
+		}
+	}
+	// Resolve name conflicts.
+	// We can safely change the parameter definitions and property bag property names because the look-ups
+	// for setting the default values afterwards are using Guids.
+	for (const TPair<FName, TArray<const UCameraRigAsset*>>& Pair : UsedParameterNames)
+	{
+		const TArray<const UCameraRigAsset*>& ConflictCameraRigs(Pair.Value);
+		if (ConflictCameraRigs.Num() > 1)
+		{
+			const FName ConflictName(Pair.Key);
+			for (const UCameraRigAsset* CameraRig : ConflictCameraRigs)
+			{
+				TArray<FCameraObjectInterfaceParameterDefinition>& Definitions = DefinitionsByCameraRig.FindChecked(CameraRig);
+				FCameraObjectInterfaceParameterDefinition* ConflictDefinition = Definitions.FindByPredicate(
+						[ConflictName](FCameraObjectInterfaceParameterDefinition& Item)
+						{
+							return Item.ParameterName == ConflictName;
+						});
+				if (ensure(ConflictDefinition))
+				{
+					const FString NewName = FString::Format(TEXT("{0}_{1}"), { *GetNameSafe(CameraRig), *ConflictName.ToString() });
+					ConflictDefinition->ParameterName = FName(NewName);
+				}
+			}
+		}
+	}
+
+	// Build the final list of parameter definitions.
 	TArray<FCameraObjectInterfaceParameterDefinition> ParameterDefinitions;
 	TArray<TObjectPtr<const UCameraRigAsset>> ParameterOwners;
-
-	for (const UCameraRigAsset* CameraRig : AllCameraRigs)
+	for (const TPair<const UCameraRigAsset*, TArray<FCameraObjectInterfaceParameterDefinition>>& Pair : DefinitionsByCameraRig)
 	{
-		for (const FCameraObjectInterfaceParameterDefinition& Definition : CameraRig->GetParameterDefinitions())
+		const TArray<FCameraObjectInterfaceParameterDefinition>& DefinitionsForCameraRig(Pair.Value);
+		for (const FCameraObjectInterfaceParameterDefinition& Definition : DefinitionsForCameraRig)
 		{
 			ParameterDefinitions.Add(Definition);
-			ParameterOwners.Add(CameraRig);
+			ParameterOwners.Add(Pair.Key);
 		}
 	}
 
@@ -114,15 +155,14 @@ void FCameraAssetBuilder::BuildCameraImpl()
 		CameraAsset->ParameterOwners = ParameterOwners;
 	}
 
-	// Get the list of all the camera rigs' interface parameters, and rebuild our
-	// parameters property bag.
+	// Rebuild the default parameters property bag.
 	TArray<FPropertyBagPropertyDesc> DefaultParameterProperties;
-	for (const UCameraRigAsset* CameraRig : AllCameraRigs)
-	{
-		FCameraObjectInterfaceParameterBuilder::AppendDefaultParameterProperties(CameraRig, DefaultParameterProperties);
-	}
+	FCameraObjectInterfaceParameterBuilder::AppendDefaultParameterProperties(
+			CameraAsset->ParameterDefinitions, DefaultParameterProperties);
+
 	FInstancedPropertyBag DefaultParameters;
 	DefaultParameters.AddProperties(DefaultParameterProperties);
+
 	for (const UCameraRigAsset* CameraRig : AllCameraRigs)
 	{
 		FCameraObjectInterfaceParameterBuilder::SetDefaultParameterValues(CameraRig, DefaultParameters);
diff --git a/Engine/Plugins/Cameras/GameplayCameras/Source/GameplayCameras/Private/Build/CameraObjectInterfaceParameterBuilder.cpp b/Engine/Plugins/Cameras/GameplayCameras/Source/GameplayCameras/Private/Build/CameraObjectInterfaceParameterBuilder.cpp
--- a/Engine/Plugins/Cameras/GameplayCameras/Source/GameplayCameras/Private/Build/CameraObjectInterfaceParameterBuilder.cpp
+++ b/Engine/Plugins/Cameras/GameplayCameras/Source/GameplayCameras/Private/Build/CameraObjectInterfaceParameterBuilder.cpp
@@ -97,7 +97,12 @@ void FCameraObjectInterfaceParameterBuilder::BuildDefaultParameters(const UBaseC
 
 void FCameraObjectInterfaceParameterBuilder::AppendDefaultParameterProperties(const UBaseCameraObject* CameraObject, TArray<FPropertyBagPropertyDesc>& OutProperties)
 {
-	for (const FCameraObjectInterfaceParameterDefinition& Definition : CameraObject->GetParameterDefinitions())
+	AppendDefaultParameterProperties(CameraObject->GetParameterDefinitions(), OutProperties);
+}
+
+void FCameraObjectInterfaceParameterBuilder::AppendDefaultParameterProperties(TConstArrayView<FCameraObjectInterfaceParameterDefinition> ParameterDefinitions, TArray<FPropertyBagPropertyDesc>& OutProperties)
+{
+	for (const FCameraObjectInterfaceParameterDefinition& Definition : ParameterDefinitions)
 	{
 		bool bIsValidProperty = true;
 		EPropertyBagPropertyType PropertyType = EPropertyBagPropertyType::Struct;
diff --git a/Engine/Plugins/Cameras/GameplayCameras/Source/GameplayCameras/Private/Tests/CameraAssetBuilderTests.cpp b/Engine/Plugins/Cameras/GameplayCameras/Source/GameplayCameras/Private/Tests/CameraAssetBuilderTests.cpp
new file mode 100644
--- /dev/null
+++ b/Engine/Plugins/Cameras/GameplayCameras/Source/GameplayCameras/Private/Tests/CameraAssetBuilderTests.cpp
@@ -0,0 +1,158 @@
+// Copyright Epic Games, Inc. All Rights Reserved.
+
+#include "Core/CameraAsset.h"
+#include "Core/CameraObjectInterfaceParameterDefinition.h"
+#include "Core/CameraParameters.h"
+#include "Misc/AssertionMacros.h"
+#include "Misc/AutomationTest.h"
+#include "Nodes/Common/LensParametersCameraNode.h"
+#include "Nodes/Common/OffsetCameraNode.h"
+#include "StructUtils/PropertyBag.h"
+#include "Tests/GameplayCamerasTestBuilder.h"
+#include "Tests/GameplayCamerasTestObjects.h"
+
+#define LOCTEXT_NAMESPACE "CameraAssetBuilderTests"
+
+namespace UE::Cameras::Test
+{
+
+template<typename ParameterType>
+bool CheckDefaultParameter(const FInstancedPropertyBag& PropertyBag, const FName ParameterName, typename TCallTraits<typename ParameterType::ValueType>::ParamType ExpectedValue)
+{
+	TValueOrError<ParameterType*, EPropertyBagResult> ValueOrError = PropertyBag.GetValueStruct<ParameterType>(ParameterName);
+	if (!ValueOrError.HasError() || ValueOrError.HasValue())
+	{
+		const ParameterType* Value = ValueOrError.GetValue();
+		return Value->Value == ExpectedValue;
+	}
+	return false;
+}
+
+}  // namespace UE::Cameras::Test
+
+IMPLEMENT_SIMPLE_AUTOMATION_TEST(FCameraAssetBuilderDefaultParametersTest, "System.Engine.GameplayCameras.CameraAssetBuilder.DefaultParameters", EAutomationTestFlags::EditorContext | EAutomationTestFlags::EngineFilter)
+bool FCameraAssetBuilderDefaultParametersTest::RunTest(const FString& Parameters)
+{
+	using namespace UE::Cameras::Test;
+
+	UFixedTestCameraDirector* CameraDirector = nullptr;
+	UCameraAsset* Camera = FCameraAssetTestBuilder()
+		.MakeDirector<UFixedTestCameraDirector>()
+			.Pin(CameraDirector)
+			.Done()
+		.Get();
+
+	UCameraRigAsset* CameraRigOne = FCameraRigAssetTestBuilder(TEXT("One"))
+		.MakeRootNode<UOffsetCameraNode>().Named(TEXT("Offset"))
+			.SetParameter(&UOffsetCameraNode::TranslationOffset, FVector3d(10.0, 20.0, 0.0))
+			.Done()
+		.AddBlendableParameter(TEXT("TranslationOffset"), ECameraVariableType::Vector3d, TEXT("Offset"), GET_MEMBER_NAME_CHECKED(UOffsetCameraNode, TranslationOffset))
+		.Get();
+	CameraDirector->AddCameraRig(CameraRigOne, TEXT("One"));
+
+	UCameraRigAsset* CameraRigTwo = FCameraRigAssetTestBuilder(TEXT("Two"))
+		.MakeRootNode<ULensParametersCameraNode>().Named(TEXT("Lens"))
+			.SetParameter(&ULensParametersCameraNode::FocalLength, 55.0f)
+			.Done()
+		.AddBlendableParameter(TEXT("FocalLength"), ECameraVariableType::Float, TEXT("Lens"), GET_MEMBER_NAME_CHECKED(ULensParametersCameraNode, FocalLength))
+		.Get();
+	CameraDirector->AddCameraRig(CameraRigTwo, TEXT("Two"));
+
+	Camera->BuildCamera();
+
+	TConstArrayView<FCameraObjectInterfaceParameterDefinition> ParameterDefinitions = Camera->GetParameterDefinitions();
+	UTEST_EQUAL("NumParameters", ParameterDefinitions.Num(), 2);
+	UTEST_EQUAL("Parameter_Name", ParameterDefinitions[0].ParameterName, FName("TranslationOffset"));
+	UTEST_EQUAL("Parameter_Type", ParameterDefinitions[0].ParameterType, ECameraObjectInterfaceParameterType::Blendable);
+	UTEST_EQUAL("Parameter_VariableType", ParameterDefinitions[0].VariableType, ECameraVariableType::Vector3d);
+	UTEST_EQUAL("Parameter_Name", ParameterDefinitions[1].ParameterName, FName("FocalLength"));
+	UTEST_EQUAL("Parameter_Type", ParameterDefinitions[1].ParameterType, ECameraObjectInterfaceParameterType::Blendable);
+	UTEST_EQUAL("Parameter_VariableType", ParameterDefinitions[1].VariableType, ECameraVariableType::Float);
+
+	const FInstancedPropertyBag& DefaultParameters = Camera->GetDefaultParameters();
+	UTEST_TRUE("Parameter_Default", CheckDefaultParameter<FVector3dCameraParameter>(DefaultParameters, TEXT("TranslationOffset"), FVector3d(10.0, 20.0, 0.0)));
+	UTEST_TRUE("Parameter_Default", CheckDefaultParameter<FFloatCameraParameter>(DefaultParameters, TEXT("FocalLength"), 55.f));
+
+	return true;
+}
+
+IMPLEMENT_SIMPLE_AUTOMATION_TEST(FCameraAssetBuilderDefaultParametersCollisionTest, "System.Engine.GameplayCameras.CameraAssetBuilder.DefaultParametersCollision", EAutomationTestFlags::EditorContext | EAutomationTestFlags::EngineFilter)
+bool FCameraAssetBuilderDefaultParametersCollisionTest::RunTest(const FString& Parameters)
+{
+	using namespace UE::Cameras::Test;
+
+	UFixedTestCameraDirector* CameraDirector = nullptr;
+	UCameraAsset* Camera = FCameraAssetTestBuilder()
+		.MakeDirector<UFixedTestCameraDirector>()
+			.Pin(CameraDirector)
+			.Done()
+		.Get();
+
+	UCameraRigAsset* CameraRigOne = FCameraRigAssetTestBuilder(TEXT("One"))
+		.MakeRootNode<UOffsetCameraNode>().Named(TEXT("Offset"))
+			.SetParameter(&UOffsetCameraNode::TranslationOffset, FVector3d(10.0, 20.0, 0.0))
+			.Done()
+		.AddBlendableParameter(TEXT("TranslationOffset"), ECameraVariableType::Vector3d, TEXT("Offset"), GET_MEMBER_NAME_CHECKED(UOffsetCameraNode, TranslationOffset))
+		.AddBlendableParameter(TEXT("RotationOffset"), ECameraVariableType::Rotator3d, TEXT("Offset"), GET_MEMBER_NAME_CHECKED(UOffsetCameraNode, RotationOffset))
+		.Get();
+	CameraDirector->AddCameraRig(CameraRigOne, TEXT("One"));
+
+	UCameraRigAsset* CameraRigTwo = FCameraRigAssetTestBuilder(TEXT("Two"))
+		.MakeRootNode<ULensParametersCameraNode>().Named(TEXT("Lens"))
+			.SetParameter(&ULensParametersCameraNode::FocalLength, 55.0f)
+			.Done()
+		.AddBlendableParameter(TEXT("FocalLength"), ECameraVariableType::Float, TEXT("Lens"), GET_MEMBER_NAME_CHECKED(ULensParametersCameraNode, FocalLength))
+		.AddBlendableParameter(TEXT("Aperture"), ECameraVariableType::Float, TEXT("Lens"), GET_MEMBER_NAME_CHECKED(ULensParametersCameraNode, Aperture))
+		.Get();
+	CameraDirector->AddCameraRig(CameraRigTwo, TEXT("Two"));
+
+	UCameraRigAsset* CameraRigThree = FCameraRigAssetTestBuilder(TEXT("Three"))
+		.MakeArrayRootNode()
+			.AddChild<UOffsetCameraNode>(&UArrayCameraNode::Children).Named(TEXT("Offset2"))
+				.SetParameter(&UOffsetCameraNode::TranslationOffset, FVector3d(0.0, 10.0, 20.0))
+				.Done()
+			.AddChild<ULensParametersCameraNode>(&UArrayCameraNode::Children).Named(TEXT("Lens2"))
+				.SetParameter(&ULensParametersCameraNode::FocalLength, 35.0f)
+				.Done()
+			.Done()
+		.AddBlendableParameter(TEXT("TranslationOffset"), ECameraVariableType::Vector3d, TEXT("Offset2"), GET_MEMBER_NAME_CHECKED(UOffsetCameraNode, TranslationOffset))
+		.AddBlendableParameter(TEXT("FocalLength"), ECameraVariableType::Float, TEXT("Lens2"), GET_MEMBER_NAME_CHECKED(ULensParametersCameraNode, FocalLength))
+		.Get();
+	CameraDirector->AddCameraRig(CameraRigThree, TEXT("Three"));
+
+	Camera->BuildCamera();
+
+	TConstArrayView<FCameraObjectInterfaceParameterDefinition> ParameterDefinitions = Camera->GetParameterDefinitions();
+	UTEST_EQUAL("NumParameters", ParameterDefinitions.Num(), 6);
+	UTEST_EQUAL("Parameter_Name", ParameterDefinitions[0].ParameterName, FName("One_TranslationOffset"));
+	UTEST_EQUAL("Parameter_Type", ParameterDefinitions[0].ParameterType, ECameraObjectInterfaceParameterType::Blendable);
+	UTEST_EQUAL("Parameter_VariableType", ParameterDefinitions[0].VariableType, ECameraVariableType::Vector3d);
+	UTEST_EQUAL("Parameter_Name", ParameterDefinitions[1].ParameterName, FName("RotationOffset"));
+	UTEST_EQUAL("Parameter_Type", ParameterDefinitions[1].ParameterType, ECameraObjectInterfaceParameterType::Blendable);
+	UTEST_EQUAL("Parameter_VariableType", ParameterDefinitions[1].VariableType, ECameraVariableType::Rotator3d);
+	UTEST_EQUAL("Parameter_Name", ParameterDefinitions[2].ParameterName, FName("Two_FocalLength"));
+	UTEST_EQUAL("Parameter_Type", ParameterDefinitions[2].ParameterType, ECameraObjectInterfaceParameterType::Blendable);
+	UTEST_EQUAL("Parameter_VariableType", ParameterDefinitions[2].VariableType, ECameraVariableType::Float);
+	UTEST_EQUAL("Parameter_Name", ParameterDefinitions[3].ParameterName, FName("Aperture"));
+	UTEST_EQUAL("Parameter_Type", ParameterDefinitions[3].ParameterType, ECameraObjectInterfaceParameterType::Blendable);
+	UTEST_EQUAL("Parameter_VariableType", ParameterDefinitions[3].VariableType, ECameraVariableType::Float);
+	UTEST_EQUAL("Parameter_Name", ParameterDefinitions[4].ParameterName, FName("Three_TranslationOffset"));
+	UTEST_EQUAL("Parameter_Type", ParameterDefinitions[4].ParameterType, ECameraObjectInterfaceParameterType::Blendable);
+	UTEST_EQUAL("Parameter_VariableType", ParameterDefinitions[4].VariableType, ECameraVariableType::Vector3d);
+	UTEST_EQUAL("Parameter_Name", ParameterDefinitions[5].ParameterName, FName("Three_FocalLength"));
+	UTEST_EQUAL("Parameter_Type", ParameterDefinitions[5].ParameterType, ECameraObjectInterfaceParameterType::Blendable);
+	UTEST_EQUAL("Parameter_VariableType", ParameterDefinitions[5].VariableType, ECameraVariableType::Float);
+
+	const FInstancedPropertyBag& DefaultParameters = Camera->GetDefaultParameters();
+	UTEST_TRUE("Parameter_Default", CheckDefaultParameter<FVector3dCameraParameter>(DefaultParameters, TEXT("One_TranslationOffset"), FVector3d(10.0, 20.0, 0.0)));
+	UTEST_TRUE("Parameter_Default", CheckDefaultParameter<FRotator3dCameraParameter>(DefaultParameters, TEXT("RotationOffset"), FRotator(0.0, 0.0, 0.0)));
+	UTEST_TRUE("Parameter_Default", CheckDefaultParameter<FFloatCameraParameter>(DefaultParameters, TEXT("Two_FocalLength"), 55.f));
+	UTEST_TRUE("Parameter_Default", CheckDefaultParameter<FFloatCameraParameter>(DefaultParameters, TEXT("Aperture"), 0.f));
+	UTEST_TRUE("Parameter_Default", CheckDefaultParameter<FVector3dCameraParameter>(DefaultParameters, TEXT("Three_TranslationOffset"), FVector3d(0.0, 10.0, 20.0)));
+	UTEST_TRUE("Parameter_Default", CheckDefaultParameter<FFloatCameraParameter>(DefaultParameters, TEXT("Three_FocalLength"), 35.f));
+
+	return true;
+}
+
+#undef LOCTEXT_NAMESPACE
+
diff --git a/Engine/Plugins/Cameras/GameplayCameras/Source/GameplayCameras/Private/Tests/GameplayCamerasTestBuilder.cpp b/Engine/Plugins/Cameras/GameplayCameras/Source/GameplayCameras/Private/Tests/GameplayCamerasTestBuilder.cpp
--- a/Engine/Plugins/Cameras/GameplayCameras/Source/GameplayCameras/Private/Tests/GameplayCamerasTestBuilder.cpp
+++ b/Engine/Plugins/Cameras/GameplayCameras/Source/GameplayCameras/Private/Tests/GameplayCamerasTestBuilder.cpp
@@ -15,6 +15,20 @@ FCameraRigAssetTestBuilder::FCameraRigAssetTestBuilder(TSharedPtr<FNamedObjectRe
 {
 }
 
+FCameraAssetTestBuilder::FCameraAssetTestBuilder(UObject* Owner)
+{
+	if (Owner == nullptr)
+	{
+		Owner = GetTransientPackage();
+	}
+
+	CameraAsset = NewObject<UCameraAsset>(Owner);
+
+	TCameraObjectInitializer<UCameraAsset>::SetObject(CameraAsset);
+
+	NamedObjectRegistry = MakeShared<FNamedObjectRegistry>();
+}
+
 FCameraEvaluationContextTestBuilder::FCameraEvaluationContextTestBuilder(UObject* Owner)
 {
 	if (Owner == nullptr)
diff --git a/Engine/Plugins/Cameras/GameplayCameras/Source/GameplayCameras/Private/Tests/GameplayCamerasTestBuilder.h b/Engine/Plugins/Cameras/GameplayCameras/Source/GameplayCameras/Private/Tests/GameplayCamerasTestBuilder.h
--- a/Engine/Plugins/Cameras/GameplayCameras/Source/GameplayCameras/Private/Tests/GameplayCamerasTestBuilder.h
+++ b/Engine/Plugins/Cameras/GameplayCameras/Source/GameplayCameras/Private/Tests/GameplayCamerasTestBuilder.h
@@ -679,6 +679,45 @@ private:
 	DirectorType* CameraDirector;
 };
 
+/**
+ * Builder class for a camera asset.
+ */
+class FCameraAssetTestBuilder
+	: public TCameraObjectInitializer<UCameraAsset>
+	, public IHasNamedObjectRegistry
+{
+public:
+
+	using ThisType = FCameraAssetTestBuilder;
+
+	/** Create a new instance of this builder class. */
+	FCameraAssetTestBuilder(UObject* Owner = nullptr);
+
+	/** Gets the created camera asset. */
+	UCameraAsset* Get() const { return CameraAsset; }
+
+	/** Builds a new camera director of the given type and returns a builder object for its. */
+	template<typename DirectorType>
+	TCameraDirectorTestBuilder<ThisType, DirectorType> MakeDirector()
+	{
+		TCameraDirectorTestBuilder<ThisType, DirectorType> DirectorBuilder(*this, CameraAsset);
+		CameraAsset->SetCameraDirector(DirectorBuilder.Get());
+		return DirectorBuilder;
+	}
+
+	/** Gets the named object registry. */
+	virtual TSharedPtr<FNamedObjectRegistry> GetNamedObjectRegistry() override
+	{
+		return NamedObjectRegistry;
+	}
+
+private:
+
+	UCameraAsset* CameraAsset;
+
+	TSharedPtr<FNamedObjectRegistry> NamedObjectRegistry;
+};
+
 /**
  * Builder class for a camera evaluation context and its camera asset.
  */
diff --git a/Engine/Plugins/Cameras/GameplayCameras/Source/GameplayCameras/Private/Tests/GameplayCamerasTestObjects.cpp b/Engine/Plugins/Cameras/GameplayCameras/Source/GameplayCameras/Private/Tests/GameplayCamerasTestObjects.cpp
--- a/Engine/Plugins/Cameras/GameplayCameras/Source/GameplayCameras/Private/Tests/GameplayCamerasTestObjects.cpp
+++ b/Engine/Plugins/Cameras/GameplayCameras/Source/GameplayCameras/Private/Tests/GameplayCamerasTestObjects.cpp
@@ -17,6 +17,24 @@ void FUpdateTrackerCameraNodeEvaluator::OnRun(const FCameraNodeEvaluationParams&
 			});
 }
 
+UE_DEFINE_CAMERA_DIRECTOR_EVALUATOR(FFixedTextCameraDirectorEvaluator)
+
+void FFixedTextCameraDirectorEvaluator::SetActiveCameraRig(const FName InCameraRigName)
+{
+	const UFixedTestCameraDirector* FixedTestDirector = GetCameraDirectorAs<UFixedTestCameraDirector>();
+	ActiveIndex = FixedTestDirector->GetCameraRigIndex(InCameraRigName);
+}
+
+void FFixedTextCameraDirectorEvaluator::OnRun(const FCameraDirectorEvaluationParams& Params, FCameraDirectorEvaluationResult& OutResult)
+{
+	const UFixedTestCameraDirector* FixedTestDirector = GetCameraDirectorAs<UFixedTestCameraDirector>();
+	UCameraRigAsset* ChosenCameraRig = FixedTestDirector->GetCameraRig(ActiveIndex);
+	if (ChosenCameraRig)
+	{
+		OutResult.Add(GetEvaluationContext(), ChosenCameraRig);
+	}
+}
+
 }  // namespace UE::Cameras::Tests
 
 FCameraNodeEvaluatorPtr UUpdateTrackerCameraNode::OnBuildEvaluator(FCameraNodeEvaluatorBuilder& Builder) const
@@ -25,3 +43,39 @@ FCameraNodeEvaluatorPtr UUpdateTrackerCameraNode::OnBuildEvaluator(FCameraNodeEv
 	return Builder.BuildEvaluator<FUpdateTrackerCameraNodeEvaluator>();
 }
 
+void UFixedTestCameraDirector::AddCameraRig(UCameraRigAsset* InCameraRig, const FName InCameraRigName)
+{
+	CameraRigs.Add(InCameraRig);
+	CameraRigNames.Add(InCameraRigName);
+}
+
+int32 UFixedTestCameraDirector::GetCameraRigIndex(const FName InCameraRigName) const
+{
+	const int32 FoundIndex = CameraRigNames.Find(InCameraRigName);
+	if (CameraRigs.IsValidIndex(FoundIndex))
+	{
+		return FoundIndex;
+	}
+	return INDEX_NONE;
+}
+
+UCameraRigAsset* UFixedTestCameraDirector::GetCameraRig(int32 Index) const
+{
+	if (CameraRigs.IsValidIndex(Index))
+	{
+		return CameraRigs[Index];
+	}
+	return nullptr;
+}
+
+FCameraDirectorEvaluatorPtr UFixedTestCameraDirector::OnBuildEvaluator(FCameraDirectorEvaluatorBuilder& Builder) const
+{
+	using namespace UE::Cameras::Test;
+	return Builder.BuildEvaluator<FFixedTextCameraDirectorEvaluator>();
+}
+
+void UFixedTestCameraDirector::OnGatherRigUsageInfo(FCameraDirectorRigUsageInfo& UsageInfo) const
+{
+	UsageInfo.CameraRigs.Append(CameraRigs);
+}
+
diff --git a/Engine/Plugins/Cameras/GameplayCameras/Source/GameplayCameras/Private/Tests/GameplayCamerasTestObjects.h b/Engine/Plugins/Cameras/GameplayCameras/Source/GameplayCameras/Private/Tests/GameplayCamerasTestObjects.h
--- a/Engine/Plugins/Cameras/GameplayCameras/Source/GameplayCameras/Private/Tests/GameplayCamerasTestObjects.h
+++ b/Engine/Plugins/Cameras/GameplayCameras/Source/GameplayCameras/Private/Tests/GameplayCamerasTestObjects.h
@@ -2,6 +2,8 @@
 
 #pragma once
 
+#include "Core/CameraDirector.h"
+#include "Core/CameraDirectorEvaluator.h"
 #include "Core/CameraNode.h"
 #include "Core/CameraNodeEvaluator.h"
 #include "Core/CameraParameters.h"
@@ -33,6 +35,24 @@ protected:
 	virtual void OnRun(const FCameraNodeEvaluationParams& Params, FCameraNodeEvaluationResult& OutResult) override;
 };
 
+class FFixedTextCameraDirectorEvaluator : public FCameraDirectorEvaluator
+{
+	UE_DECLARE_CAMERA_DIRECTOR_EVALUATOR(GAMEPLAYCAMERAS_API, FFixedTextCameraDirectorEvaluator)
+
+public:
+
+	void SetActiveCameraRig(const FName InCameraRigName);
+
+protected:
+
+	// FCameraDirectorEvaluator interface.
+	virtual void OnRun(const FCameraDirectorEvaluationParams& Params, FCameraDirectorEvaluationResult& OutResult) override;
+
+private:
+
+	int32 ActiveIndex = INDEX_NONE;
+};
+
 }  // namespace UE::Cameras::Tests
 
 UCLASS(MinimalAPI, Hidden)
@@ -54,3 +74,29 @@ public:
 	FVector3dCameraParameter VectorParameter;
 };
 
+UCLASS(MinimalAPI, Hidden)
+class UFixedTestCameraDirector : public UCameraDirector
+{
+	GENERATED_BODY()
+
+public:
+
+	void AddCameraRig(UCameraRigAsset* InCameraRig, const FName InCameraRigName);
+	int32 GetCameraRigIndex(const FName InCameraRigName) const;
+	UCameraRigAsset* GetCameraRig(int32 Index) const;
+
+protected:
+
+	// UCameraDirector interface.
+	virtual FCameraDirectorEvaluatorPtr OnBuildEvaluator(FCameraDirectorEvaluatorBuilder& Builder) const override;
+	virtual void OnGatherRigUsageInfo(FCameraDirectorRigUsageInfo& UsageInfo) const override;
+
+private:
+
+	UPROPERTY()
+	TArray<TObjectPtr<UCameraRigAsset>> CameraRigs;
+
+	UPROPERTY()
+	TArray<FName> CameraRigNames;
+};
+
diff --git a/Engine/Plugins/Cameras/GameplayCameras/Source/GameplayCameras/Public/Build/CameraObjectInterfaceParameterBuilder.h b/Engine/Plugins/Cameras/GameplayCameras/Source/GameplayCameras/Public/Build/CameraObjectInterfaceParameterBuilder.h
--- a/Engine/Plugins/Cameras/GameplayCameras/Source/GameplayCameras/Public/Build/CameraObjectInterfaceParameterBuilder.h
+++ b/Engine/Plugins/Cameras/GameplayCameras/Source/GameplayCameras/Public/Build/CameraObjectInterfaceParameterBuilder.h
@@ -9,10 +9,12 @@
 
 class UBaseCameraObject;
 class UCameraObjectInterfaceDataParameter;
+struct FCameraObjectInterfaceParameterDefinition;
 
 namespace UE::Cameras
 {
 
+
 /**
  * A helper class for building an FInstancedPropertyBag from a list of camera rig
  * parameter definitions.
@@ -34,6 +36,7 @@ public:
 
 	static UE_API void BuildDefaultParameters(const UBaseCameraObject* CameraObject, FInstancedPropertyBag& OutPropertyBag);
 	static UE_API void AppendDefaultParameterProperties(const UBaseCameraObject* CameraObject, TArray<FPropertyBagPropertyDesc>& OutProperties);
+	static UE_API void AppendDefaultParameterProperties(TConstArrayView<FCameraObjectInterfaceParameterDefinition> ParameterDefinitions, TArray<FPropertyBagPropertyDesc>& OutProperties);
 	static UE_API void SetDefaultParameterValues(const UBaseCameraObject* CameraObject, FInstancedPropertyBag& PropertyBag);
 
 private:
