Cameras: don't show byte/int types for data parameters

diff --git a/Engine/Plugins/Cameras/GameplayCameras/Source/GameplayCamerasEditor/Private/Toolkits/CameraObjectInterfaceParametersToolkit.cpp b/Engine/Plugins/Cameras/GameplayCameras/Source/GameplayCamerasEditor/Private/Toolkits/CameraObjectInterfaceParametersToolkit.cpp
--- a/Engine/Plugins/Cameras/GameplayCameras/Source/GameplayCamerasEditor/Private/Toolkits/CameraObjectInterfaceParametersToolkit.cpp
+++ b/Engine/Plugins/Cameras/GameplayCameras/Source/GameplayCamerasEditor/Private/Toolkits/CameraObjectInterfaceParametersToolkit.cpp
@@ -410,6 +410,7 @@ class FDataParameterPinTypeSelectorFilter : public IPinTypeSelectorFilter
 		FEdGraphPinType PinType = InItem->GetPinType(false);
 		if (PinType.PinCategory == UEdGraphSchema_K2::PC_Boolean ||
 				PinType.PinCategory == UEdGraphSchema_K2::PC_Int ||
+				PinType.PinCategory == UEdGraphSchema_K2::PC_Int64 ||
 				PinType.PinCategory == UEdGraphSchema_K2::PC_Float ||
 				PinType.PinCategory == UEdGraphSchema_K2::PC_Double ||
 				PinType.PinCategory == UEdGraphSchema_K2::PC_Real)
@@ -417,6 +418,11 @@ class FDataParameterPinTypeSelectorFilter : public IPinTypeSelectorFilter
 			return false;
 		}
 
+		if (PinType.PinCategory == UEdGraphSchema_K2::PC_Byte && PinType.PinSubCategoryObject == nullptr)
+		{
+			return false;
+		}
+
 		if (PinType.PinCategory == UEdGraphSchema_K2::PC_Struct)
 		{
 			const UObject* TypeObject = PinType.PinSubCategoryObject.Get();
