Cameras: 5.7 compatibility pass

diff --git a/Engine/Plugins/Cameras/GameplayCameras/Source/GameplayCameras/Private/GameFramework/GameplayCameraComponentBase.cpp b/Engine/Plugins/Cameras/GameplayCameras/Source/GameplayCameras/Private/GameFramework/GameplayCameraComponentBase.cpp
--- a/Engine/Plugins/Cameras/GameplayCameras/Source/GameplayCameras/Private/GameFramework/GameplayCameraComponentBase.cpp
+++ b/Engine/Plugins/Cameras/GameplayCameras/Source/GameplayCameras/Private/GameFramework/GameplayCameraComponentBase.cpp
@@ -719,7 +719,9 @@ void UGameplayCameraComponentBase::UpdateOutputCameraComponent()
 			OutputCameraComponent->FocusSettings.ManualFocusDistance = Result.CameraPose.GetFocusDistance();
 			OutputCameraComponent->FocusSettings.FocusMethod = (Result.CameraPose.GetEnablePhysicalCamera() ? ECameraFocusMethod::Manual : ECameraFocusMethod::DoNotOverride);
 
+#if UE_VERSION_NEWER_THAN_OR_EQUAL(5,7,0)
 			OutputCameraComponent->ExposureMethod = (Result.CameraPose.GetEnablePhysicalCamera() ? ECameraExposureMethod::Enabled : ECameraExposureMethod::DoNotOverride);
+#endif
 
 			OutputCameraComponent->ProjectionMode = Result.CameraPose.GetProjectionMode();
 			OutputCameraComponent->OrthoWidth = Result.CameraPose.GetOrthographicWidth();
diff --git a/Engine/Plugins/Cameras/GameplayCameras/Source/GameplayCamerasEditor/Private/Editors/ObjectTreeGraphSchema.cpp b/Engine/Plugins/Cameras/GameplayCameras/Source/GameplayCamerasEditor/Private/Editors/ObjectTreeGraphSchema.cpp
--- a/Engine/Plugins/Cameras/GameplayCameras/Source/GameplayCamerasEditor/Private/Editors/ObjectTreeGraphSchema.cpp
+++ b/Engine/Plugins/Cameras/GameplayCameras/Source/GameplayCamerasEditor/Private/Editors/ObjectTreeGraphSchema.cpp
@@ -13,7 +13,6 @@
 #include "Exporters/Exporter.h"
 #include "Factories.h"
 #include "IGameplayCamerasEditorModule.h"
-#include "Misc/StringOutputDevice.h"
 #include "ScopedTransaction.h"
 #include "Serialization/ArchiveUObject.h"
 #include "ToolMenu.h"
diff --git a/Engine/Plugins/Cameras/GameplayCameras/Source/GameplayCamerasEditor/Public/Editors/ObjectTreeGraphSchema.h b/Engine/Plugins/Cameras/GameplayCameras/Source/GameplayCamerasEditor/Public/Editors/ObjectTreeGraphSchema.h
--- a/Engine/Plugins/Cameras/GameplayCameras/Source/GameplayCamerasEditor/Public/Editors/ObjectTreeGraphSchema.h
+++ b/Engine/Plugins/Cameras/GameplayCameras/Source/GameplayCamerasEditor/Public/Editors/ObjectTreeGraphSchema.h
@@ -5,8 +5,12 @@
 #include "Compat/EditorCompat.h"
 #include "CoreTypes.h"
 #include "EdGraph/EdGraphSchema.h"
-#include "Misc/StringOutputDevice.h"
 #include "GraphEditor.h"
+#include "Misc/EngineVersionComparison.h"
+
+#if UE_VERSION_NEWER_THAN_OR_EQUAL(5,7,0)
+#include "Misc/StringOutputDevice.h"
+#endif
 
 #include "ObjectTreeGraphSchema.generated.h"
 
