Cameras: fix one frame delay in the input accumulator node due to outdated code

diff --git a/Engine/Plugins/Cameras/GameplayCameras/Source/GameplayCameras/Private/Nodes/Input/InputAccumulator2DCameraNode.cpp b/Engine/Plugins/Cameras/GameplayCameras/Source/GameplayCameras/Private/Nodes/Input/InputAccumulator2DCameraNode.cpp
--- a/Engine/Plugins/Cameras/GameplayCameras/Source/GameplayCameras/Private/Nodes/Input/InputAccumulator2DCameraNode.cpp
+++ b/Engine/Plugins/Cameras/GameplayCameras/Source/GameplayCameras/Private/Nodes/Input/InputAccumulator2DCameraNode.cpp
@@ -22,7 +22,6 @@ protected:
 	// FCameraNodeEvaluator interface.
 	virtual void OnBuild(const FCameraNodeEvaluatorBuildParams& Params) override;
 	virtual FCameraNodeEvaluatorChildrenView OnGetChildren() override;
-	virtual void OnUpdateParameters(const FCameraBlendedParameterUpdateParams& Params, FCameraBlendedParameterUpdateResult& OutResult) override;
 	virtual void OnRun(const FCameraNodeEvaluationParams& Params, FCameraNodeEvaluationResult& OutResult) override;
 
 private:
@@ -50,22 +49,12 @@ FCameraNodeEvaluatorChildrenView FInputAccumulator2DCameraNodeEvaluator::OnGetCh
 	return FCameraNodeEvaluatorChildrenView{ InputSlotEvaluator };
 }
 
-void FInputAccumulator2DCameraNodeEvaluator::OnUpdateParameters(const FCameraBlendedParameterUpdateParams& Params, FCameraBlendedParameterUpdateResult& OutResult)
-{
-	if (InputSlotEvaluator)
-	{
-		InputSlotEvaluator->UpdateParameters(Params, OutResult);
-		DeltaInputValue = InputSlotEvaluator->GetInputValue();
-	}
-
-	Super::OnUpdateParameters(Params, OutResult);
-}
-
 void FInputAccumulator2DCameraNodeEvaluator::OnRun(const FCameraNodeEvaluationParams& Params, FCameraNodeEvaluationResult& OutResult)
 {
 	if (InputSlotEvaluator)
 	{
 		InputSlotEvaluator->Run(Params, OutResult);
+		DeltaInputValue = InputSlotEvaluator->GetInputValue();
 	}
 
 	Super::OnRun(Params, OutResult);
