Cameras: do a proper normalization of input values

diff --git a/Engine/Plugins/Cameras/GameplayCameras/Source/GameplayCameras/Private/Nodes/Input/CameraRigInputSlotTypes.cpp b/Engine/Plugins/Cameras/GameplayCameras/Source/GameplayCameras/Private/Nodes/Input/CameraRigInputSlotTypes.cpp
--- a/Engine/Plugins/Cameras/GameplayCameras/Source/GameplayCameras/Private/Nodes/Input/CameraRigInputSlotTypes.cpp
+++ b/Engine/Plugins/Cameras/GameplayCameras/Source/GameplayCameras/Private/Nodes/Input/CameraRigInputSlotTypes.cpp
@@ -25,11 +25,12 @@ void FCameraParameterClamping::GetEffectiveClamping(double& OutMinValue, double&
 
 double FCameraParameterNormalization::NormalizeValue(double Value) const
 {
-	if (bNormalize && MaxValue > 0)
+	if (bNormalize)
 	{
-		while (Value > MaxValue)
+		Value = FMath::Fmod(Value, MaxValue);
+		if (Value < 0.0)
 		{
-			Value -= MaxValue;
+			Value += MaxValue;
 		}
 	}
 	return Value;
diff --git a/Engine/Plugins/Cameras/GameplayCameras/Source/GameplayCameras/Public/Nodes/Input/CameraRigInputSlotTypes.h b/Engine/Plugins/Cameras/GameplayCameras/Source/GameplayCameras/Public/Nodes/Input/CameraRigInputSlotTypes.h
--- a/Engine/Plugins/Cameras/GameplayCameras/Source/GameplayCameras/Public/Nodes/Input/CameraRigInputSlotTypes.h
+++ b/Engine/Plugins/Cameras/GameplayCameras/Source/GameplayCameras/Public/Nodes/Input/CameraRigInputSlotTypes.h
@@ -63,7 +63,7 @@ struct FCameraParameterNormalization
 	GENERATED_BODY()
 
 	UPROPERTY(EditAnywhere, Category="Common", meta=(ClampMin=0, EditCondition="bNormalize"))
-	double MaxValue = 0;
+	double MaxValue = 360;
 
 	UPROPERTY(EditAnywhere, Category="Common")
 	bool bNormalize = false;
