Cameras: add proper API for letting external modules implement and run operations

- Added registration of new operation type IDs
- Added public API for running camera operations on the camera system

diff --git a/Engine/Plugins/Cameras/GameplayCameras/Source/GameplayCameras/Private/Core/BlendStackCameraNode.cpp b/Engine/Plugins/Cameras/GameplayCameras/Source/GameplayCameras/Private/Core/BlendStackCameraNode.cpp
--- a/Engine/Plugins/Cameras/GameplayCameras/Source/GameplayCameras/Private/Core/BlendStackCameraNode.cpp
+++ b/Engine/Plugins/Cameras/GameplayCameras/Source/GameplayCameras/Private/Core/BlendStackCameraNode.cpp
@@ -6,6 +6,7 @@
 #include "Core/BlendStackRootCameraNode.h"
 #include "Core/CameraAsset.h"
 #include "Core/CameraEvaluationContext.h"
+#include "Core/CameraOperation.h"
 #include "Core/CameraRigAsset.h"
 #include "Core/CameraRigAssetReference.h"
 #include "Core/CameraSystemEvaluator.h"
@@ -230,6 +231,25 @@ void FBlendStackCameraNodeEvaluator::OnInitialize(const FCameraNodeEvaluatorInit
 	Layer = BlendStack->Layer;
 }
 
+void FBlendStackCameraNodeEvaluator::OnExecuteOperation(const FCameraOperationParams& Params, FCameraOperation& Operation)
+{
+	ensure(Params.Evaluator == OwningEvaluator);
+
+	// Execute operations from top to bottom.
+	for (int32 Index = Entries.Num() - 1; Index >= 0; --Index)
+	{
+		FCameraRigEntry& Entry(Entries[Index]);
+		TSharedPtr<const FCameraEvaluationContext> CurContext = Entry.EvaluationContext.Pin();
+
+		if (!Entry.Flags.bIsFrozen)
+		{
+			FCameraOperationParams CurParams(Params);
+			CurParams.EvaluationContext = CurContext;
+			Entry.EvaluatorHierarchy.CallExecuteOperation(CurParams, Operation);
+		}
+	}
+}
+
 void FBlendStackCameraNodeEvaluator::ResolveEntries(TArray<FResolvedEntry>& OutResolvedEntries)
 {
 	constexpr ECameraVariableTableFilter VariableTableFilter = ECameraVariableTableFilter::ChangedOnly;
diff --git a/Engine/Plugins/Cameras/GameplayCameras/Source/GameplayCameras/Private/Core/CameraOperation.cpp b/Engine/Plugins/Cameras/GameplayCameras/Source/GameplayCameras/Private/Core/CameraOperation.cpp
new file mode 100644
--- /dev/null
+++ b/Engine/Plugins/Cameras/GameplayCameras/Source/GameplayCameras/Private/Core/CameraOperation.cpp
@@ -0,0 +1,15 @@
+// Copyright Epic Games, Inc. All Rights Reserved.
+
+#include "Core/CameraOperation.h"
+
+namespace UE::Cameras
+{
+
+FCameraOperationTypeID FCameraOperationTypeID::RegisterNew()
+{
+	static uint32 NextID = (uint32)EBuiltInCameraOperationTypes::MAX + 1;
+	return FCameraOperationTypeID(NextID++);
+}
+
+}  // namespace UE::Cameras
+
diff --git a/Engine/Plugins/Cameras/GameplayCameras/Source/GameplayCameras/Private/Core/CameraSystemEvaluator.cpp b/Engine/Plugins/Cameras/GameplayCameras/Source/GameplayCameras/Private/Core/CameraSystemEvaluator.cpp
--- a/Engine/Plugins/Cameras/GameplayCameras/Source/GameplayCameras/Private/Core/CameraSystemEvaluator.cpp
+++ b/Engine/Plugins/Cameras/GameplayCameras/Source/GameplayCameras/Private/Core/CameraSystemEvaluator.cpp
@@ -508,6 +508,13 @@ void FCameraSystemEvaluator::GetEvaluatedCameraView(FMinimalViewInfo& DesiredVie
 	DesiredView.ApplyOverscan(CameraPose.GetOverscan());
 }
 
+void FCameraSystemEvaluator::ExecuteOperation(FCameraOperation& Operation)
+{
+	FCameraOperationParams Params;
+	Params.Evaluator = this;
+	RootEvaluator->ExecuteOperation(Params, Operation);
+}
+
 #if WITH_EDITOR
 
 void FCameraSystemEvaluator::EditorPreviewUpdate(const FCameraSystemEvaluationParams& Params)
diff --git a/Engine/Plugins/Cameras/GameplayCameras/Source/GameplayCameras/Private/Core/DefaultRootCameraNode.cpp b/Engine/Plugins/Cameras/GameplayCameras/Source/GameplayCameras/Private/Core/DefaultRootCameraNode.cpp
--- a/Engine/Plugins/Cameras/GameplayCameras/Source/GameplayCameras/Private/Core/DefaultRootCameraNode.cpp
+++ b/Engine/Plugins/Cameras/GameplayCameras/Source/GameplayCameras/Private/Core/DefaultRootCameraNode.cpp
@@ -87,6 +87,11 @@ void FDefaultRootCameraNodeEvaluator::OnRun(const FCameraNodeEvaluationParams& P
 	}
 }
 
+void FDefaultRootCameraNodeEvaluator::OnExecuteOperation(const FCameraOperationParams& Params, FCameraOperation& Operation)
+{
+	MainLayer->ExecuteOperation(Params, Operation);
+}
+
 FCameraRigInstanceID FDefaultRootCameraNodeEvaluator::OnActivateCameraRig(const FActivateCameraRigParams& Params)
 {
 	if (Params.Layer == ECameraRigLayer::Main)
diff --git a/Engine/Plugins/Cameras/GameplayCameras/Source/GameplayCameras/Public/Core/BlendStackCameraNode.h b/Engine/Plugins/Cameras/GameplayCameras/Source/GameplayCameras/Public/Core/BlendStackCameraNode.h
--- a/Engine/Plugins/Cameras/GameplayCameras/Source/GameplayCameras/Public/Core/BlendStackCameraNode.h
+++ b/Engine/Plugins/Cameras/GameplayCameras/Source/GameplayCameras/Public/Core/BlendStackCameraNode.h
@@ -127,6 +127,7 @@ protected:
 	// FCameraNodeEvaluator interface
 	virtual FCameraNodeEvaluatorChildrenView OnGetChildren() override;
 	virtual void OnInitialize(const FCameraNodeEvaluatorInitializeParams& Params, FCameraNodeEvaluationResult& OutResult) override;
+	virtual void OnExecuteOperation(const FCameraOperationParams& Params, FCameraOperation& Operation) override;
 	virtual void OnAddReferencedObjects(FReferenceCollector& Collector) override;
 	virtual void OnSerialize(const FCameraNodeEvaluatorSerializeParams& Params, FArchive& Ar) override;
 
diff --git a/Engine/Plugins/Cameras/GameplayCameras/Source/GameplayCameras/Public/Core/CameraOperation.h b/Engine/Plugins/Cameras/GameplayCameras/Source/GameplayCameras/Public/Core/CameraOperation.h
--- a/Engine/Plugins/Cameras/GameplayCameras/Source/GameplayCameras/Public/Core/CameraOperation.h
+++ b/Engine/Plugins/Cameras/GameplayCameras/Source/GameplayCameras/Public/Core/CameraOperation.h
@@ -30,6 +30,9 @@ struct FCameraOperationTypeID
 
 public:
 
+	/** Generate a new operation type ID. */
+	GAMEPLAYCAMERAS_API static FCameraOperationTypeID RegisterNew();
+
 	/** Returns whether this ID is valid. */
 	bool IsValid() const { return Value != INVALID; }
 
@@ -74,7 +77,7 @@ private:
 	FCameraOperationTypeID OperationTypeID;
 };
 
-#define UE_DEFINE_CAMERA_IK_OPERATION(OperationName, OperationType)\
+#define UE_DECLARE_BUILT_IN_CAMERA_OPERATION(OperationName, OperationType)\
 	static FCameraOperationTypeID GetOperationTypeID()\
 	{\
 		return FCameraOperationTypeID(OperationType);\
@@ -88,7 +91,7 @@ private:
  */
 struct FYawPitchCameraOperation : public FCameraOperation
 {
-	UE_DEFINE_CAMERA_IK_OPERATION(FYawPitchCameraOperation, EBuiltInCameraOperationTypes::YawPitch)
+	UE_DECLARE_BUILT_IN_CAMERA_OPERATION(FYawPitchCameraOperation, EBuiltInCameraOperationTypes::YawPitch)
 
 	FConsumableDouble Yaw;
 	FConsumableDouble Pitch;
@@ -100,10 +103,26 @@ struct FYawPitchCameraOperation : public FCameraOperation
  */
 struct FSingleValueCameraOperation : public FCameraOperation
 {
-	UE_DEFINE_CAMERA_IK_OPERATION(FSingleValueCameraOperation, EBuiltInCameraOperationTypes::SingleValue)
+	UE_DECLARE_BUILT_IN_CAMERA_OPERATION(FSingleValueCameraOperation, EBuiltInCameraOperationTypes::SingleValue)
 
 	FConsumableDouble Value;
 };
 
+#undef UE_DECLARE_BUILT_IN_CAMERA_OPERATION
+
 }  // namespace UE::Cameras
 
+#define UE_DECLARE_CAMERA_OPERATION(ApiDeclSpec, OperationName)\
+	public:\
+		ApiDeclSpec static UE::Cameras::FCameraOperationTypeID GetOperationTypeID();\
+		OperationName()\
+			: FCameraOperation(OperationName::GetOperationTypeID())\
+		{}
+
+#define UE_DEFINE_CAMERA_OPERATION(OperationName)\
+	UE::Cameras::FCameraOperationTypeID OperationName::GetOperationTypeID()\
+	{\
+		static UE::Cameras::FCameraOperationTypeID StaticOperationTypeID = UE::Cameras::FCameraOperationTypeID::RegisterNew();\
+		return StaticOperationTypeID;\
+	}\
+
diff --git a/Engine/Plugins/Cameras/GameplayCameras/Source/GameplayCameras/Public/Core/CameraSystemEvaluator.h b/Engine/Plugins/Cameras/GameplayCameras/Source/GameplayCameras/Public/Core/CameraSystemEvaluator.h
--- a/Engine/Plugins/Cameras/GameplayCameras/Source/GameplayCameras/Public/Core/CameraSystemEvaluator.h
+++ b/Engine/Plugins/Cameras/GameplayCameras/Source/GameplayCameras/Public/Core/CameraSystemEvaluator.h
@@ -257,6 +257,9 @@ public:
 	/** Get the last evaluated camera. */
 	GAMEPLAYCAMERAS_API void GetEvaluatedCameraView(FMinimalViewInfo& DesiredView);
 
+	/** Executes an operation, to be dispatched by the root node evaluator. */
+	GAMEPLAYCAMERAS_API void ExecuteOperation(FCameraOperation& Operation);
+
 	/** Collect reference objects for the garbage collector. */
 	GAMEPLAYCAMERAS_API void AddReferencedObjects(FReferenceCollector& Collector);
 
diff --git a/Engine/Plugins/Cameras/GameplayCameras/Source/GameplayCameras/Public/Core/DefaultRootCameraNode.h b/Engine/Plugins/Cameras/GameplayCameras/Source/GameplayCameras/Public/Core/DefaultRootCameraNode.h
--- a/Engine/Plugins/Cameras/GameplayCameras/Source/GameplayCameras/Public/Core/DefaultRootCameraNode.h
+++ b/Engine/Plugins/Cameras/GameplayCameras/Source/GameplayCameras/Public/Core/DefaultRootCameraNode.h
@@ -59,6 +59,7 @@ protected:
 	virtual FCameraNodeEvaluatorChildrenView OnGetChildren() override;
 	virtual void OnBuild(const FCameraNodeEvaluatorBuildParams& Params) override;
 	virtual void OnRun(const FCameraNodeEvaluationParams& Params, FCameraNodeEvaluationResult& OutResult) override;
+	virtual void OnExecuteOperation(const FCameraOperationParams& Params, FCameraOperation& Operation) override;
 
 #if UE_GAMEPLAY_CAMERAS_DEBUG
 	virtual void OnBuildDebugBlocks(const FCameraDebugBlockBuildParams& Params, FCameraDebugBlockBuilder& Builder) override;
