Cameras: add Teardown API for node evaluators, add layer and active camera rig information

diff --git a/Engine/Plugins/Cameras/GameplayCameras/Source/GameplayCameras/Private/Core/BlendStackCameraNode.cpp b/Engine/Plugins/Cameras/GameplayCameras/Source/GameplayCameras/Private/Core/BlendStackCameraNode.cpp
--- a/Engine/Plugins/Cameras/GameplayCameras/Source/GameplayCameras/Private/Core/BlendStackCameraNode.cpp
+++ b/Engine/Plugins/Cameras/GameplayCameras/Source/GameplayCameras/Private/Core/BlendStackCameraNode.cpp
@@ -121,6 +121,7 @@ void FBlendStackCameraNodeEvaluator::InitializeEntry(
 		InitParams.Evaluator = OwningEvaluator;
 		InitParams.EvaluationContext = EvaluationContext;
 		InitParams.LastActiveCameraRigInfo = GetActiveCameraRigEvaluationInfo();
+		InitParams.Layer = Layer;
 		RootEvaluator->Initialize(InitParams, NewEntry.ContextResult);  // Initializing with the context result here.
 	}
 
@@ -153,6 +154,16 @@ int32 FBlendStackCameraNodeEvaluator::IndexOfEntry(const FBlendStackEntryID Entr
 
 void FBlendStackCameraNodeEvaluator::FreezeEntry(FCameraRigEntry& Entry)
 {
+	if (Entry.RootEvaluator)
+	{
+		FCameraNodeEvaluatorTeardownParams TeardownParams;
+		TeardownParams.Evaluator = this->OwningEvaluator;
+		TeardownParams.EvaluationContext = Entry.EvaluationContext.Pin();
+		TeardownParams.Layer = Layer;
+
+		Entry.RootEvaluator->Teardown(TeardownParams);
+	}
+
 	// Deallocate our node evaluators and clear any pointers we kept to them.
 	Entry.EvaluatorStorage.DestroyEvaluatorTree(true);
 	Entry.RootEvaluator = nullptr;
@@ -409,6 +420,16 @@ void FBlendStackCameraNodeEvaluator::PopEntries(int32 FirstIndexToKeep)
 	{
 		FCameraRigEntry& FirstEntry = Entries[0];
 
+		if (FirstEntry.RootEvaluator)
+		{
+			FCameraNodeEvaluatorTeardownParams TeardownParams;
+			TeardownParams.Evaluator = this->OwningEvaluator;
+			TeardownParams.EvaluationContext = FirstEntry.EvaluationContext.Pin();
+			TeardownParams.Layer = Layer;
+
+			FirstEntry.RootEvaluator->Teardown(TeardownParams);
+		}
+
 #if WITH_EDITOR
 		RemoveListenedPackages(LiveEditManager, FirstEntry);
 #endif  // WITH_EDITOR
diff --git a/Engine/Plugins/Cameras/GameplayCameras/Source/GameplayCameras/Private/Core/CameraIKAim.cpp b/Engine/Plugins/Cameras/GameplayCameras/Source/GameplayCameras/Private/Core/CameraIKAim.cpp
--- a/Engine/Plugins/Cameras/GameplayCameras/Source/GameplayCameras/Private/Core/CameraIKAim.cpp
+++ b/Engine/Plugins/Cameras/GameplayCameras/Source/GameplayCameras/Private/Core/CameraIKAim.cpp
@@ -180,6 +180,7 @@ void FCameraIKAim::RunRootCameraNode(const FCameraIKAimParams& Params, const FCa
 		FSingleCameraRigEvaluationParams SingleParams;
 		SingleParams.EvaluationParams.DeltaTime = Params.DeltaTime;
 		SingleParams.EvaluationParams.bIsFirstFrame = Params.bIsFirstFrame;
+		SingleParams.EvaluationParams.bIsActiveCameraRig = Params.bIsActiveCameraRig;
 		SingleParams.EvaluationParams.EvaluationType = ECameraNodeEvaluationType::IK;
 		SingleParams.EvaluationParams.EvaluationContext = CameraRigInfo.EvaluationContext;
 		SingleParams.EvaluationParams.Evaluator = Params.Evaluator;
diff --git a/Engine/Plugins/Cameras/GameplayCameras/Source/GameplayCameras/Private/Core/CameraNodeEvaluator.cpp b/Engine/Plugins/Cameras/GameplayCameras/Source/GameplayCameras/Private/Core/CameraNodeEvaluator.cpp
--- a/Engine/Plugins/Cameras/GameplayCameras/Source/GameplayCameras/Private/Core/CameraNodeEvaluator.cpp
+++ b/Engine/Plugins/Cameras/GameplayCameras/Source/GameplayCameras/Private/Core/CameraNodeEvaluator.cpp
@@ -165,6 +165,19 @@ void FCameraNodeEvaluator::Initialize(const FCameraNodeEvaluatorInitializeParams
 	}
 }
 
+void FCameraNodeEvaluator::Teardown(const FCameraNodeEvaluatorTeardownParams& Params)
+{
+	OnTeardown(Params);
+
+	for (FCameraNodeEvaluator* Child : GetChildren())
+	{
+		if (Child)
+		{
+			Child->Teardown(Params);
+		}
+	}
+}
+
 void FCameraNodeEvaluator::AddReferencedObjects(FReferenceCollector& Collector)
 {
 	if (PrivateCameraNode)
diff --git a/Engine/Plugins/Cameras/GameplayCameras/Source/GameplayCameras/Private/Core/PersistentBlendStackCameraNode.cpp b/Engine/Plugins/Cameras/GameplayCameras/Source/GameplayCameras/Private/Core/PersistentBlendStackCameraNode.cpp
--- a/Engine/Plugins/Cameras/GameplayCameras/Source/GameplayCameras/Private/Core/PersistentBlendStackCameraNode.cpp
+++ b/Engine/Plugins/Cameras/GameplayCameras/Source/GameplayCameras/Private/Core/PersistentBlendStackCameraNode.cpp
@@ -166,6 +166,7 @@ void FPersistentBlendStackCameraNodeEvaluator::RemoveEntry(int32 EntryIndex, con
 				FCameraNodeEvaluatorInitializeParams BlendOutInitParams;
 				BlendOutInitParams.Evaluator = OwningEvaluator;
 				BlendOutInitParams.EvaluationContext = Entry.EvaluationContext.Pin();
+				BlendOutInitParams.Layer = Layer;
 				BlendOutEvaluator->Initialize(BlendOutInitParams, Entry.Result);
 
 				// Reverse this blend so it plays as a blend-out. Also, see if we are going to 
diff --git a/Engine/Plugins/Cameras/GameplayCameras/Source/GameplayCameras/Private/Core/TransientBlendStackCameraNode.cpp b/Engine/Plugins/Cameras/GameplayCameras/Source/GameplayCameras/Private/Core/TransientBlendStackCameraNode.cpp
--- a/Engine/Plugins/Cameras/GameplayCameras/Source/GameplayCameras/Private/Core/TransientBlendStackCameraNode.cpp
+++ b/Engine/Plugins/Cameras/GameplayCameras/Source/GameplayCameras/Private/Core/TransientBlendStackCameraNode.cpp
@@ -141,6 +141,7 @@ FBlendStackEntryID FTransientBlendStackCameraNodeEvaluator::PushMergedEntry(cons
 	FCameraNodeEvaluatorInitializeParams InitParams;
 	InitParams.Evaluator = OwningEvaluator;
 	InitParams.EvaluationContext = TopEntry.EvaluationContext.Pin();
+	InitParams.Layer = Layer;
 
 	TopEntry.RootEvaluator->MergeCameraRig(BuildParams, InitParams, TopEntry.Result, PushParams.CameraRig, Blend);
 
@@ -252,6 +253,7 @@ void FTransientBlendStackCameraNodeEvaluator::InternalPreBlendPrepare(TArrayView
 		FCameraNodeEvaluationParams CurParams(Params);
 		CurParams.EvaluationContext = ResolvedEntry.Context;
 		CurParams.bIsFirstFrame = Entry.Flags.bIsFirstFrame;
+		CurParams.bIsActiveCameraRig = (ResolvedEntry.EntryIndex == Entries.Num() - 1);
 
 		FCameraNodeEvaluationResult& CurResult(Entry.Result);
 
@@ -385,6 +387,7 @@ void FTransientBlendStackCameraNodeEvaluator::InternalUpdate(TArrayView<FResolve
 		FCameraNodeEvaluationParams CurParams(Params);
 		CurParams.EvaluationContext = ResolvedEntry.Context;
 		CurParams.bIsFirstFrame = Entry.Flags.bIsFirstFrame;
+		CurParams.bIsActiveCameraRig = (ResolvedEntry.EntryIndex == Entries.Num() - 1);
 
 		// Start with the input given to us.
 		CurResult.Reset();
diff --git a/Engine/Plugins/Cameras/GameplayCameras/Source/GameplayCameras/Private/Services/CameraShakeService.cpp b/Engine/Plugins/Cameras/GameplayCameras/Source/GameplayCameras/Private/Services/CameraShakeService.cpp
--- a/Engine/Plugins/Cameras/GameplayCameras/Source/GameplayCameras/Private/Services/CameraShakeService.cpp
+++ b/Engine/Plugins/Cameras/GameplayCameras/Source/GameplayCameras/Private/Services/CameraShakeService.cpp
@@ -330,6 +330,7 @@ void FCameraShakeServiceCameraNodeEvaluator::InitializeEntry(
 		FCameraNodeEvaluatorInitializeParams InitParams(&NewEntry.EvaluatorHierarchy);
 		InitParams.Evaluator = OwningEvaluator;
 		InitParams.EvaluationContext = EvaluationContext;
+		InitParams.Layer = ECameraRigLayer::Visual;
 		RootEvaluator->Initialize(InitParams, NewEntry.Result);
 	}
 
diff --git a/Engine/Plugins/Cameras/GameplayCameras/Source/GameplayCameras/Private/Services/OrientationInitializationService.cpp b/Engine/Plugins/Cameras/GameplayCameras/Source/GameplayCameras/Private/Services/OrientationInitializationService.cpp
--- a/Engine/Plugins/Cameras/GameplayCameras/Source/GameplayCameras/Private/Services/OrientationInitializationService.cpp
+++ b/Engine/Plugins/Cameras/GameplayCameras/Source/GameplayCameras/Private/Services/OrientationInitializationService.cpp
@@ -212,6 +212,7 @@ void FOrientationInitializationService::TryPreserveTarget(const FCameraRigEvalua
 
 	FCameraIKAimParams AimParams;
 	AimParams.bIsFirstFrame = true;
+	AimParams.bIsActiveCameraRig = true;
 	AimParams.DeltaTime = 0.f;
 	AimParams.Evaluator = Evaluator;
 	AimParams.TargetLocation = TargetToPreserve;
diff --git a/Engine/Plugins/Cameras/GameplayCameras/Source/GameplayCameras/Public/Core/CameraIKAim.h b/Engine/Plugins/Cameras/GameplayCameras/Source/GameplayCameras/Public/Core/CameraIKAim.h
--- a/Engine/Plugins/Cameras/GameplayCameras/Source/GameplayCameras/Public/Core/CameraIKAim.h
+++ b/Engine/Plugins/Cameras/GameplayCameras/Source/GameplayCameras/Public/Core/CameraIKAim.h
@@ -22,6 +22,8 @@ struct FCameraIKAimParams
 	float DeltaTime = 0.f;
 	/** Whether this is the first update of the camera rig. */
 	bool bIsFirstFrame = false;
+	/** Whether this camera rig is the active one in its layer (e.g. top of the blend stack). */
+	bool bIsActiveCameraRig = false;
 	
 	/** The desired target that the camera rig should be aiming at. */
 	FVector3d TargetLocation = FVector3d::ZeroVector;
diff --git a/Engine/Plugins/Cameras/GameplayCameras/Source/GameplayCameras/Public/Core/CameraNodeEvaluator.h b/Engine/Plugins/Cameras/GameplayCameras/Source/GameplayCameras/Public/Core/CameraNodeEvaluator.h
--- a/Engine/Plugins/Cameras/GameplayCameras/Source/GameplayCameras/Public/Core/CameraNodeEvaluator.h
+++ b/Engine/Plugins/Cameras/GameplayCameras/Source/GameplayCameras/Public/Core/CameraNodeEvaluator.h
@@ -88,6 +88,9 @@ struct FCameraNodeEvaluatorInitializeParams
 	 */
 	FCameraRigEvaluationInfo LastActiveCameraRigInfo;
 
+	/** The layer on which the node evaluator is being initialized. */
+	ECameraRigLayer Layer = ECameraRigLayer::Main;
+
 public:
 
 	FCameraNodeEvaluatorInitializeParams() = default;
@@ -101,6 +104,20 @@ private:
 	friend class FCameraNodeEvaluator;
 };
 
+/**
+ * Structure for tearing down a camera node evaluator.
+ */
+struct FCameraNodeEvaluatorTeardownParams
+{
+	/** The evaluation running this evaluation. */
+	FCameraSystemEvaluator* Evaluator = nullptr;
+	/** The evaluation context (if any) responsible for this branch of the evaluation. */
+	TSharedPtr<const FCameraEvaluationContext> EvaluationContext;
+
+	/** The layer on which the node evaluator was running. */
+	ECameraRigLayer Layer = ECameraRigLayer::Main;
+};
+
 /**
  * Parameter structure for updating the pre-blended parameters of a camera node.
  */
@@ -161,6 +178,8 @@ struct FCameraNodeEvaluationParams
 	ECameraNodeEvaluationType EvaluationType = ECameraNodeEvaluationType::Standard;
 	/** Whether this is the first evaluation of this camera node hierarchy. */
 	bool bIsFirstFrame = false;
+	/** Whether this camera node is running inside the active camera rig in this layer. */
+	bool bIsActiveCameraRig = false;
 
 	bool IsStatelessEvaluation() const
 	{
@@ -282,6 +301,9 @@ public:
 	/** Initialize this evaluator and all its descendants. */
 	void Initialize(const FCameraNodeEvaluatorInitializeParams& Params, FCameraNodeEvaluationResult& OutResult);
 
+	/** Tear down this evaluator and all its descendants. */
+	void Teardown(const FCameraNodeEvaluatorTeardownParams& Params);
+
 	/** Collect referenced UObjects for this node and all its descendants. */
 	void AddReferencedObjects(FReferenceCollector& Collector);
 
@@ -350,6 +372,9 @@ protected:
 	/** Initialize this evaluator. Children and descendants will be automatically initialized too. */
 	virtual void OnInitialize(const FCameraNodeEvaluatorInitializeParams& Params, FCameraNodeEvaluationResult& OutResult) {}
 
+	/** Tear down this evaluator and all its descendants. */
+	virtual void OnTeardown(const FCameraNodeEvaluatorTeardownParams& Params) {}
+
 	/** Collect referenced UObjects for this node. */
 	virtual void OnAddReferencedObjects(FReferenceCollector& Collector) {}
 
diff --git a/Engine/Plugins/Cameras/GameplayCameras/Source/GameplayCameras/Public/Core/RootCameraNodeCameraRigEvent.h b/Engine/Plugins/Cameras/GameplayCameras/Source/GameplayCameras/Public/Core/RootCameraNodeCameraRigEvent.h
--- a/Engine/Plugins/Cameras/GameplayCameras/Source/GameplayCameras/Public/Core/RootCameraNodeCameraRigEvent.h
+++ b/Engine/Plugins/Cameras/GameplayCameras/Source/GameplayCameras/Public/Core/RootCameraNodeCameraRigEvent.h
@@ -26,9 +26,9 @@ enum class ERootCameraNodeCameraRigEventType
 struct FRootCameraNodeCameraRigEvent
 {
 	/** The type of event. */
-	ERootCameraNodeCameraRigEventType EventType;
+	ERootCameraNodeCameraRigEventType EventType = ERootCameraNodeCameraRigEventType::Activated;
 	/** The layer on which the event happened. */
-	ECameraRigLayer EventLayer;
+	ECameraRigLayer EventLayer = ECameraRigLayer::Main;
 	/** The evaluation information of the associated camera rig. */
 	FCameraRigEvaluationInfo CameraRigInfo;
 	/** If a camera rig was activated, the transition used to activate it, if any. */
