Cameras: add "driven control rotation" input node

This input node allows using a "centralized" player input system which drives control rotation on the pawn, while still allowing for cameras to have different orientations, such as when using target preservation with various lateral offsets (and therefore parallax requiring converging cameras with different orientations). The active camera rig uses control rotation as-is, but any blending-out camera rigs use their last control rotation driven by the "delta" control rotation each frame. This seems to be a good compromise between cameras limited by a single orientation and having player input managed by camera rigs individually.

diff --git a/Engine/Plugins/Cameras/GameplayCameras/Source/GameplayCameras/Private/Nodes/Input/DrivenControlRotationCameraNode.cpp b/Engine/Plugins/Cameras/GameplayCameras/Source/GameplayCameras/Private/Nodes/Input/DrivenControlRotationCameraNode.cpp
new file mode 100644
--- /dev/null
+++ b/Engine/Plugins/Cameras/GameplayCameras/Source/GameplayCameras/Private/Nodes/Input/DrivenControlRotationCameraNode.cpp
@@ -0,0 +1,322 @@
+// Copyright Epic Games, Inc. All Rights Reserved.
+
+#include "Nodes/Input/DrivenControlRotationCameraNode.h"
+
+#include "Core/CameraEvaluationContext.h"
+#include "Core/CameraEvaluationService.h"
+#include "Core/CameraOperation.h"
+#include "Core/CameraSystemEvaluator.h"
+#include "Core/RootCameraNode.h"
+#include "Debug/CameraDebugBlock.h"
+#include "Debug/CameraDebugBlockBuilder.h"
+#include "Debug/CameraDebugRenderer.h"
+#include "GameFramework/PlayerController.h"
+
+namespace UE::Cameras
+{
+
+class FDrivenControlRotationHelperService : public FCameraEvaluationService
+{
+	UE_DECLARE_CAMERA_EVALUATION_SERVICE(, FDrivenControlRotationHelperService)
+
+public:
+
+	void AddMonitoredContext(TSharedPtr<const FCameraEvaluationContext> InContext);
+	void RemoveMonitoredContext(TSharedPtr<const FCameraEvaluationContext> InContext);
+
+	FRotator3d GetCachedControlRotation(TSharedPtr<const FCameraEvaluationContext> InContext) const;
+
+protected:
+
+	// FCameraEvaluationService interface.
+	virtual void OnInitialize(const FCameraEvaluationServiceInitializeParams& Params) override;
+	virtual void OnPreUpdate(const FCameraEvaluationServiceUpdateParams& Params, FCameraEvaluationServiceUpdateResult& OutResult) override;
+
+#if UE_GAMEPLAY_CAMERAS_DEBUG
+	virtual void OnBuildDebugBlocks(const FCameraDebugBlockBuildParams& Params, FCameraDebugBlockBuilder& Builder) override;
+#endif  // UE_GAMEPLAY_CAMERAS_DEBUG
+
+private:
+
+	struct FMonitoredContext
+	{
+		TWeakObjectPtr<APlayerController> WeakPlayerController;
+		FRotator3d CachedControlRotation = FRotator3d::ZeroRotator;
+		uint8 MonitoringRequests = 0;
+	};
+
+	using FMonitoredContexts = TMap<TSharedPtr<const FCameraEvaluationContext>, FMonitoredContext>;
+	FMonitoredContexts MonitoredContexts;
+};
+
+UE_DEFINE_CAMERA_EVALUATION_SERVICE(FDrivenControlRotationHelperService)
+
+void FDrivenControlRotationHelperService::OnInitialize(const FCameraEvaluationServiceInitializeParams& Params)
+{
+	SetEvaluationServiceFlags(ECameraEvaluationServiceFlags::NeedsPreUpdate);
+}
+
+void FDrivenControlRotationHelperService::OnPreUpdate(const FCameraEvaluationServiceUpdateParams& Params, FCameraEvaluationServiceUpdateResult& OutResult)
+{
+	for (auto It = MonitoredContexts.CreateIterator(); It; ++It)
+	{
+		FMonitoredContext& MonitoredContext(It->Value);
+		if (APlayerController* PlayerController = MonitoredContext.WeakPlayerController.Get())
+		{
+			MonitoredContext.CachedControlRotation = PlayerController->GetControlRotation();
+		}
+		else
+		{
+			It.RemoveCurrent();
+		}
+	}
+}
+
+void FDrivenControlRotationHelperService::AddMonitoredContext(TSharedPtr<const FCameraEvaluationContext> InContext)
+{
+	if (FMonitoredContext* ExistingEntry = MonitoredContexts.Find(InContext))
+	{
+		++ExistingEntry->MonitoringRequests;
+		return;
+	}
+
+	if (APlayerController* PlayerController = InContext->GetPlayerController())
+	{
+		FMonitoredContext& NewEntry = MonitoredContexts.Add(InContext);
+		NewEntry.WeakPlayerController = PlayerController;
+		NewEntry.CachedControlRotation = PlayerController->GetControlRotation();
+		NewEntry.MonitoringRequests = 1;
+	}
+	else
+	{
+		UE_LOG(LogCameraSystem, Warning, 
+				TEXT("Can't monitor control rotation for camera context owned by '%s', it is running without "
+					 "any association to a player controller."),
+				*GetNameSafe(InContext->GetOwner()));
+	}
+}
+
+void FDrivenControlRotationHelperService::RemoveMonitoredContext(TSharedPtr<const FCameraEvaluationContext> InContext)
+{
+	if (FMonitoredContext* ExistingEntry = MonitoredContexts.Find(InContext))
+	{
+		--ExistingEntry->MonitoringRequests;
+		if (ExistingEntry->MonitoringRequests == 0)
+		{
+			MonitoredContexts.Remove(InContext);
+		}
+	}
+}
+
+FRotator3d FDrivenControlRotationHelperService::GetCachedControlRotation(TSharedPtr<const FCameraEvaluationContext> InContext) const
+{
+	if (const FMonitoredContext* ExistingEntry = MonitoredContexts.Find(InContext))
+	{
+		return ExistingEntry->CachedControlRotation;
+	}
+	return FRotator3d::ZeroRotator;
+}
+
+#if UE_GAMEPLAY_CAMERAS_DEBUG
+
+UE_DECLARE_CAMERA_DEBUG_BLOCK_START(, FDrivenControlRotationHelperDebugBlock)
+	UE_DECLARE_CAMERA_DEBUG_BLOCK_FIELD(TArray<FString>, ContextNames)
+	UE_DECLARE_CAMERA_DEBUG_BLOCK_FIELD(TArray<FRotator3d>, CachedControlRotations)
+	UE_DECLARE_CAMERA_DEBUG_BLOCK_FIELD(TArray<uint8>, MonitoringRequests)
+UE_DECLARE_CAMERA_DEBUG_BLOCK_END()
+
+UE_DEFINE_CAMERA_DEBUG_BLOCK_WITH_FIELDS(FDrivenControlRotationHelperDebugBlock)
+
+void FDrivenControlRotationHelperService::OnBuildDebugBlocks(const FCameraDebugBlockBuildParams& Params, FCameraDebugBlockBuilder& Builder)
+{
+	FDrivenControlRotationHelperDebugBlock& DebugBlock = Builder.AttachDebugBlock<FDrivenControlRotationHelperDebugBlock>();
+	for (const FMonitoredContexts::ElementType& Pair : MonitoredContexts)
+	{
+		TSharedPtr<const FCameraEvaluationContext> Context = Pair.Key;
+		const FMonitoredContext& MonitoredContext = Pair.Value;
+		DebugBlock.ContextNames.Add(GetNameSafe(Context->GetOwner()));
+		DebugBlock.CachedControlRotations.Add(MonitoredContext.CachedControlRotation);
+		DebugBlock.MonitoringRequests.Add(MonitoredContext.MonitoringRequests);
+	}
+}
+
+void FDrivenControlRotationHelperDebugBlock::OnDebugDraw(const FCameraDebugBlockDrawParams& Params, FCameraDebugRenderer& Renderer)
+{
+	const int32 Num = FMath::Min(ContextNames.Num(), FMath::Min(CachedControlRotations.Num(), MonitoringRequests.Num()));
+	Renderer.AddText(TEXT("%d monitored contexts"), Num);
+	Renderer.AddIndent();
+	{
+		for (int32 Index = 0; Index < Num; ++Index)
+		{
+			Renderer.AddText(
+					TEXT("{cam_notice}'%s'{cam_passive} (%d requests){cam_defualt} : %s\n"),
+					*ContextNames[Index], MonitoringRequests[Index], *CachedControlRotations[Index].ToString());
+		}
+	}
+	Renderer.RemoveIndent();
+}
+
+#endif  // UE_GAMEPLAY_CAMERAS_DEBUG
+
+class FDrivenControlRotationCameraNodeEvaluator : public FInput2DCameraNodeEvaluator
+{
+	UE_DECLARE_CAMERA_NODE_EVALUATOR_EX(, FDrivenControlRotationCameraNodeEvaluator, FInput2DCameraNodeEvaluator)
+
+public:
+
+	FDrivenControlRotationCameraNodeEvaluator();
+
+protected:
+
+	// FCameraNodeEvaluator interface.
+	virtual void OnInitialize(const FCameraNodeEvaluatorInitializeParams& Params, FCameraNodeEvaluationResult& OutResult) override;
+	virtual void OnTeardown(const FCameraNodeEvaluatorTeardownParams& Params) override;
+	virtual void OnRun(const FCameraNodeEvaluationParams& Params, FCameraNodeEvaluationResult& OutResult) override;
+	virtual void OnExecuteOperation(const FCameraOperationParams& Params, FCameraOperation& Operation) override;
+	virtual void OnSerialize(const FCameraNodeEvaluatorSerializeParams& Params, FArchive& Ar) override;
+
+#if UE_GAMEPLAY_CAMERAS_DEBUG
+	virtual void OnBuildDebugBlocks(const FCameraDebugBlockBuildParams& Params, FCameraDebugBlockBuilder& Builder);
+#endif  // UE_GAMEPLAY_CAMERAS_DEBUG
+
+private:
+
+	TSharedPtr<FDrivenControlRotationHelperService> HelperService;
+
+	FRotator3d LastControlRotation = FRotator3d::ZeroRotator;
+	bool bLastWasActiveCameraRig = false;
+
+#if UE_GAMEPLAY_CAMERAS_DEBUG
+	FRotator3d LastDeltaControlRotation;
+#endif  // UE_GAMEPLAY_CAMERAS_DEBUG
+};
+
+UE_DEFINE_CAMERA_NODE_EVALUATOR(FDrivenControlRotationCameraNodeEvaluator)
+
+FDrivenControlRotationCameraNodeEvaluator::FDrivenControlRotationCameraNodeEvaluator()
+{
+	SetNodeEvaluatorFlags(ECameraNodeEvaluatorFlags::NeedsSerialize | ECameraNodeEvaluatorFlags::SupportsOperations);
+}
+
+void FDrivenControlRotationCameraNodeEvaluator::OnInitialize(const FCameraNodeEvaluatorInitializeParams& Params, FCameraNodeEvaluationResult& OutResult)
+{
+	bLastWasActiveCameraRig = true;
+
+	HelperService = Params.Evaluator->FindOrRegisterEvaluationService<FDrivenControlRotationHelperService>();
+	HelperService->AddMonitoredContext(Params.EvaluationContext);
+}
+
+void FDrivenControlRotationCameraNodeEvaluator::OnTeardown(const FCameraNodeEvaluatorTeardownParams& Params)
+{
+	if (HelperService)
+	{
+		HelperService->RemoveMonitoredContext(Params.EvaluationContext);
+	}
+}
+
+void FDrivenControlRotationCameraNodeEvaluator::OnRun(const FCameraNodeEvaluationParams& Params, FCameraNodeEvaluationResult& OutResult)
+{
+	APlayerController* PlayerController = Params.EvaluationContext->GetPlayerController();
+	if (PlayerController)
+	{
+
+		const FRotator3d ControlRotation = PlayerController->GetControlRotation();
+
+		if (Params.bIsActiveCameraRig)
+		{
+			// We are running in the active camera rig. Just move with the control rotation.
+			InputValue = FVector2d(ControlRotation.Yaw, ControlRotation.Pitch);
+		}
+		else
+		{
+			// We are not the active camera rig anymore. Only apply the delta between the last frame's
+			// control rotation and this frame's control rotation. One important thing to note here
+			// is that we do NOT get the control rotation from the player controller for this frame.
+			// This is because any newly activated camera rig may have changed it for their own purpose,
+			// such as a new camera rig that needs target preservation and, therefore, ran an IK 
+			// correction that affect the control rotation (see OnExecuteOperation below). So here we
+			// get the current control rotation from our helper service, which cached the control
+			// rotation in OnPreUpdate, BEFORE the camera director runs and new camera rigs are
+			// activated. This lets us get the "new" control rotation for this frame before it gets
+			// messed up by new camera rigs' initializations.
+			const FRotator3d CachedControlRotation = HelperService->GetCachedControlRotation(Params.EvaluationContext);
+			FRotator3d DeltaRotation((CachedControlRotation - LastControlRotation).GetNormalized());
+			InputValue.X += DeltaRotation.Yaw;
+			InputValue.Y += DeltaRotation.Pitch;
+
+#if UE_GAMEPLAY_CAMERAS_DEBUG
+			LastDeltaControlRotation = DeltaRotation;
+#endif  // UE_GAMEPLAY_CAMERAS_DEBUG
+		}
+
+		LastControlRotation = ControlRotation;
+	}
+
+	bLastWasActiveCameraRig = Params.bIsActiveCameraRig;
+}
+
+void FDrivenControlRotationCameraNodeEvaluator::OnExecuteOperation(const FCameraOperationParams& Params, FCameraOperation& Operation)
+{
+	if (FYawPitchCameraOperation* Op = Operation.CastOperation<FYawPitchCameraOperation>())
+	{
+		// If we are running in the active camera rig, we have permission to affect the control rotation.
+		// Let's apply the yaw/pitch operation to it.
+		APlayerController* PlayerController = Params.EvaluationContext->GetPlayerController();
+		if (bLastWasActiveCameraRig && PlayerController)
+		{
+			FRotator3d ControlRotation = PlayerController->GetControlRotation();
+			ControlRotation.Yaw = Op->Yaw.Apply(ControlRotation.Yaw);
+			ControlRotation.Pitch = Op->Pitch.Apply(ControlRotation.Pitch);
+			PlayerController->SetControlRotation(ControlRotation);
+		}
+	}
+}
+
+void FDrivenControlRotationCameraNodeEvaluator::OnSerialize(const FCameraNodeEvaluatorSerializeParams& Params, FArchive& Ar)
+{
+	Ar << LastControlRotation;
+	Ar << bLastWasActiveCameraRig;
+}
+
+#if UE_GAMEPLAY_CAMERAS_DEBUG
+
+UE_DECLARE_CAMERA_DEBUG_BLOCK_START(, FDrivenControlRotationDebugBlock)
+	UE_DECLARE_CAMERA_DEBUG_BLOCK_FIELD(FVector2d, InputValue)
+	UE_DECLARE_CAMERA_DEBUG_BLOCK_FIELD(FRotator3d, LastDeltaControlRotation)
+	UE_DECLARE_CAMERA_DEBUG_BLOCK_FIELD(bool, bLastWasActiveCameraRig)
+UE_DECLARE_CAMERA_DEBUG_BLOCK_END()
+
+UE_DEFINE_CAMERA_DEBUG_BLOCK_WITH_FIELDS(FDrivenControlRotationDebugBlock)
+
+void FDrivenControlRotationCameraNodeEvaluator::OnBuildDebugBlocks(const FCameraDebugBlockBuildParams& Params, FCameraDebugBlockBuilder& Builder)
+{
+	FDrivenControlRotationDebugBlock& DebugBlock = Builder.AttachDebugBlock<FDrivenControlRotationDebugBlock>();
+	DebugBlock.InputValue = InputValue;
+	DebugBlock.LastDeltaControlRotation = LastDeltaControlRotation;
+	DebugBlock.bLastWasActiveCameraRig = bLastWasActiveCameraRig;
+}
+
+void FDrivenControlRotationDebugBlock::OnDebugDraw(const FCameraDebugBlockDrawParams& Params, FCameraDebugRenderer& Renderer)
+{
+	if (bLastWasActiveCameraRig)
+	{
+		Renderer.AddText(TEXT("Active ; P=%.1f Y=%.1f"), InputValue.Y, InputValue.X);
+	}
+	else
+	{
+		Renderer.AddText(
+				TEXT("Driven ; P=%.1f Y=%.1f (Delta: P=%.1f Y=%.1f"), 
+				InputValue.Y, InputValue.X, LastDeltaControlRotation.Pitch, LastDeltaControlRotation.Yaw);
+	}
+}
+
+#endif  // UE_GAMEPLAY_CAMERAS_DEBUG
+
+}  // namespace UE::Cameras
+
+FCameraNodeEvaluatorPtr UDrivenControlRotationCameraNode::OnBuildEvaluator(FCameraNodeEvaluatorBuilder& Builder) const
+{
+	using namespace UE::Cameras;
+	return Builder.BuildEvaluator<FDrivenControlRotationCameraNodeEvaluator>();
+}
+
diff --git a/Engine/Plugins/Cameras/GameplayCameras/Source/GameplayCameras/Public/Core/CameraSystemEvaluator.h b/Engine/Plugins/Cameras/GameplayCameras/Source/GameplayCameras/Public/Core/CameraSystemEvaluator.h
--- a/Engine/Plugins/Cameras/GameplayCameras/Source/GameplayCameras/Public/Core/CameraSystemEvaluator.h
+++ b/Engine/Plugins/Cameras/GameplayCameras/Source/GameplayCameras/Public/Core/CameraSystemEvaluator.h
@@ -237,6 +237,19 @@ public:
 		return nullptr;
 	}
 
+	/** Finds or creates/registers an evaluation service of the given type. */
+	template<typename EvaluationServiceType>
+	TSharedRef<EvaluationServiceType> FindOrRegisterEvaluationService()
+	{
+		TSharedPtr<EvaluationServiceType> EvaluationService = FindEvaluationService<EvaluationServiceType>();
+		if (!EvaluationService)
+		{
+			EvaluationService = MakeShared<EvaluationServiceType>();
+			RegisterEvaluationService(EvaluationService.ToSharedRef());
+		}
+		return EvaluationService.ToSharedRef();
+	}
+
 public:
 
 	/** Run an update of the camera system. */
diff --git a/Engine/Plugins/Cameras/GameplayCameras/Source/GameplayCameras/Public/Nodes/Input/DrivenControlRotationCameraNode.h b/Engine/Plugins/Cameras/GameplayCameras/Source/GameplayCameras/Public/Nodes/Input/DrivenControlRotationCameraNode.h
new file mode 100644
--- /dev/null
+++ b/Engine/Plugins/Cameras/GameplayCameras/Source/GameplayCameras/Public/Nodes/Input/DrivenControlRotationCameraNode.h
@@ -0,0 +1,24 @@
+// Copyright Epic Games, Inc. All Rights Reserved.
+
+#pragma once
+
+#include "Nodes/Input/Input2DCameraNode.h"
+
+#include "DrivenControlRotationCameraNode.generated.h"
+
+/**
+ * An input camera node that provides the yaw and pitch components of the control rotation
+ * when it is in the active camera rig, and then "detaches" from it while blending out,
+ * applying delta-rotations on the last known active control rotation.
+ */
+UCLASS()
+class UDrivenControlRotationCameraNode : public UInput2DCameraNode
+{
+	GENERATED_BODY()
+
+protected:
+	
+	// UCameraNode interface.
+	virtual FCameraNodeEvaluatorPtr OnBuildEvaluator(FCameraNodeEvaluatorBuilder& Builder) const override;
+};
+
