Cameras: make the orientation initialization service available to 3rd parties, and add an API for overriding the target or yaw/pitch to preserve

diff --git a/Engine/Plugins/Cameras/GameplayCameras/Source/GameplayCameras/Private/Services/OrientationInitializationService.cpp b/Engine/Plugins/Cameras/GameplayCameras/Source/GameplayCameras/Private/Services/OrientationInitializationService.cpp
--- a/Engine/Plugins/Cameras/GameplayCameras/Source/GameplayCameras/Private/Services/OrientationInitializationService.cpp
+++ b/Engine/Plugins/Cameras/GameplayCameras/Source/GameplayCameras/Private/Services/OrientationInitializationService.cpp
@@ -26,6 +26,16 @@ static FAutoConsoleVariableRef CVarGameplayCamerasDebugOrientationInitialization
 
 UE_DEFINE_CAMERA_EVALUATION_SERVICE(FOrientationInitializationService)
 
+void FOrientationInitializationService::SetYawPitchPreservationOverride(const FRotator3d& InOrientation)
+{
+	YawPitchPreservationOverride = InOrientation;
+}
+
+void FOrientationInitializationService::SetTargetPreservationOverride(const FVector3d& InTarget)
+{
+	TargetPreservationOverride = InTarget;
+}
+
 void FOrientationInitializationService::OnInitialize(const FCameraEvaluationServiceInitializeParams& Params)
 {
 	SetEvaluationServiceFlags(
@@ -49,6 +59,9 @@ void FOrientationInitializationService::OnPostUpdate(const FCameraEvaluationServ
 	}
 
 	PreviousEvaluationContext = ActiveContext;
+
+	YawPitchPreservationOverride.Reset();
+	TargetPreservationOverride.Reset();
 }
 
 void FOrientationInitializationService::OnRootCameraNodeEvent(const FRootCameraNodeCameraRigEvent& InEvent)
@@ -140,7 +153,11 @@ void FOrientationInitializationService::TryPreserveYawPitch(const FCameraRigEval
 		return;
 	}
 	
-	const FRotator3d LastOrientation = LastResult.CameraPose.GetRotation();
+	FRotator3d LastOrientation = LastResult.CameraPose.GetRotation();
+	if (YawPitchPreservationOverride.IsSet())
+	{
+		LastOrientation = YawPitchPreservationOverride.GetValue();
+	}
 	TryInitializeYawPitch(CameraRigInfo, LastOrientation.Yaw, LastOrientation.Pitch);
 }
 
@@ -196,18 +213,25 @@ void FOrientationInitializationService::TryPreserveTarget(const FCameraRigEvalua
 
 	FVector3d TargetToPreserve = LastResult.CameraPose.GetTarget();
 
-	// Relative target preservation means that if the context turned since last frame, we want to preserve
-	// a target that has "turned" along with it.
-	if (bUseRelativeTarget && bHasPreviousContextTransform &&
-			(PreviousEvaluationContext == nullptr || CameraRigInfo.EvaluationContext == PreviousEvaluationContext))
+	if (TargetPreservationOverride.IsSet())
+	{
+		TargetToPreserve = TargetPreservationOverride.GetValue();
+	}
+	else
 	{
-		const FRotator3d PreviousInverseContextRotation = PreviousContextRotation.GetInverse();
-		const FVector3d LastRelativeTarget = PreviousInverseContextRotation.RotateVector(TargetToPreserve - PreviousContextLocation);
+		// Relative target preservation means that if the context turned since last frame, we want to preserve
+		// a target that has "turned" along with it.
+		if (bUseRelativeTarget && bHasPreviousContextTransform &&
+				(PreviousEvaluationContext == nullptr || CameraRigInfo.EvaluationContext == PreviousEvaluationContext))
+		{
+			const FRotator3d PreviousInverseContextRotation = PreviousContextRotation.GetInverse();
+			const FVector3d LastRelativeTarget = PreviousInverseContextRotation.RotateVector(TargetToPreserve - PreviousContextLocation);
 
-		const FCameraNodeEvaluationResult& InitialResult = CameraRigInfo.EvaluationContext->GetInitialResult();
-		const FVector3d& InitialLocation = InitialResult.CameraPose.GetLocation();
-		const FRotator3d& InitialRotation = InitialResult.CameraPose.GetRotation();
-		TargetToPreserve = InitialRotation.RotateVector(LastRelativeTarget) + InitialLocation;
+			const FCameraNodeEvaluationResult& InitialResult = CameraRigInfo.EvaluationContext->GetInitialResult();
+			const FVector3d& InitialLocation = InitialResult.CameraPose.GetLocation();
+			const FRotator3d& InitialRotation = InitialResult.CameraPose.GetRotation();
+			TargetToPreserve = InitialRotation.RotateVector(LastRelativeTarget) + InitialLocation;
+		}
 	}
 
 	FCameraIKAimParams AimParams;
diff --git a/Engine/Plugins/Cameras/GameplayCameras/Source/GameplayCameras/Private/Services/OrientationInitializationService.h b/Engine/Plugins/Cameras/GameplayCameras/Source/GameplayCameras/Private/Services/OrientationInitializationService.h
deleted file mode 100644
diff --git a/Engine/Plugins/Cameras/GameplayCameras/Source/GameplayCameras/Public/Services/OrientationInitializationService.h b/Engine/Plugins/Cameras/GameplayCameras/Source/GameplayCameras/Public/Services/OrientationInitializationService.h
new file mode 100644
--- /dev/null
+++ b/Engine/Plugins/Cameras/GameplayCameras/Source/GameplayCameras/Public/Services/OrientationInitializationService.h
@@ -0,0 +1,75 @@
+// Copyright Epic Games, Inc. All Rights Reserved.
+
+#pragma once
+
+#include "Core/CameraEvaluationService.h"
+#include "Core/CameraIKAim.h"
+#include "Core/CameraRigTransition.h"
+
+namespace UE::Cameras
+{
+
+class FCameraEvaluationContext;
+struct FCameraRigEvaluationInfo;
+
+/**
+ * A camera evaluation service responsible for initializing the orientation of newly activated camera rigs.
+ */
+class FOrientationInitializationService : public FCameraEvaluationService
+{
+	UE_DECLARE_CAMERA_EVALUATION_SERVICE(GAMEPLAYCAMERAS_API, FOrientationInitializationService)
+
+public:
+
+	/**
+	 * Override the yaw/pitch to preserve on any camera rig activated this frame. This will affect camera rigs
+	 * who have the "Preserve Yaw/Pitch" policy.
+	 */
+	GAMEPLAYCAMERAS_API void SetYawPitchPreservationOverride(const FRotator3d& InOrientation);
+
+	/**
+	 * Override the target location to preserve on any camera rig activated this frame. This will affect camera rigs
+	 * who have the "Preserve Absolute/Relative Target" policies.
+	 */
+	GAMEPLAYCAMERAS_API void SetTargetPreservationOverride(const FVector3d& InTarget);
+
+protected:
+
+	// FCameraEvaluationService interface.
+	virtual void OnInitialize(const FCameraEvaluationServiceInitializeParams& Params) override;
+	virtual void OnPostUpdate(const FCameraEvaluationServiceUpdateParams& Params, FCameraEvaluationServiceUpdateResult& OutResult) override;
+	virtual void OnRootCameraNodeEvent(const FRootCameraNodeCameraRigEvent& InEvent) override;
+
+#if UE_GAMEPLAY_CAMERAS_DEBUG
+	virtual void OnBuildDebugBlocks(const FCameraDebugBlockBuildParams& Params, FCameraDebugBlockBuilder& Builder) override;
+#endif  // UE_GAMEPLAY_CAMERAS_DEBUG
+
+private:
+
+	void TryInitializeContextYawPitch(const FCameraRigEvaluationInfo& CameraRigInfo);
+	void TryPreserveYawPitch(const FCameraRigEvaluationInfo& CameraRigInfo);
+	void TryInitializeYawPitch(const FCameraRigEvaluationInfo& CameraRigInfo, TOptional<double> Yaw, TOptional<double> Pitch);
+	void TryPreserveTarget(const FCameraRigEvaluationInfo& CameraRigInfo, bool bUseRelativeTarget);
+
+private:
+
+	FCameraSystemEvaluator* Evaluator = nullptr;
+
+	TWeakPtr<FCameraEvaluationContext> PreviousEvaluationContext;
+	FVector3d PreviousContextLocation;
+	FRotator3d PreviousContextRotation;
+	bool bHasPreviousContextTransform = false;
+
+	TOptional<FRotator3d> YawPitchPreservationOverride;
+	TOptional<FVector3d> TargetPreservationOverride;
+
+#if UE_GAMEPLAY_CAMERAS_DEBUG
+	FVector3d DebugLastEvaluatedTarget;
+	FCameraIKAimDebugInfo LastAimDebugInfo;
+
+	friend class FOrientationInitializationDebugBlock;
+#endif  // UE_GAMEPLAY_CAMERAS_DEBUG
+};
+
+}  // namespace UE::Cameras
+
