Cameras: revamp how parameter overrides are applied

- Simplify the helper code for applying parameter overrides.
- Add a way to apply parameter overrides based on a cached version of their values, comparing them to the current ones to know which ones have changed. This is because we don't just want to apply them all each frame, in case they are also changed via other means like the BP setters.
- Remove the Sequencer decorators and the "animated" flags on parameters, since they're not used anymore.

diff --git a/Engine/Plugins/Cameras/GameplayCameras/Source/GameplayCameras/Private/Core/BaseCameraObjectReference.cpp b/Engine/Plugins/Cameras/GameplayCameras/Source/GameplayCameras/Private/Core/BaseCameraObjectReference.cpp
--- a/Engine/Plugins/Cameras/GameplayCameras/Source/GameplayCameras/Private/Core/BaseCameraObjectReference.cpp
+++ b/Engine/Plugins/Cameras/GameplayCameras/Source/GameplayCameras/Private/Core/BaseCameraObjectReference.cpp
@@ -50,21 +50,6 @@ void FBaseCameraObjectReference::SetParameterOverridden(const FGuid& PropertyID,
 	MetaData.bIsOverridden = bIsOverridden;
 }
 
-bool FBaseCameraObjectReference::IsParameterAnimated(const FGuid& PropertyID) const
-{
-	if (const FCameraObjectInterfaceParameterMetaData* MetaData = FindMetaData(PropertyID))
-	{
-		return MetaData->bIsAnimated;
-	}
-	return false;
-}
-
-void FBaseCameraObjectReference::SetParameterAnimated(const FGuid& PropertyID, bool bIsAnimated)
-{
-	FCameraObjectInterfaceParameterMetaData& MetaData = FindOrAddMetaData(PropertyID);
-	MetaData.bIsAnimated = bIsAnimated;
-}
-
 bool FBaseCameraObjectReference::NeedsRebuildParameters() const
 {
 	const UBaseCameraObject* CameraObject = GetCameraObject();
diff --git a/Engine/Plugins/Cameras/GameplayCameras/Source/GameplayCameras/Private/Core/BlendStackRootCameraNode.cpp b/Engine/Plugins/Cameras/GameplayCameras/Source/GameplayCameras/Private/Core/BlendStackRootCameraNode.cpp
--- a/Engine/Plugins/Cameras/GameplayCameras/Source/GameplayCameras/Private/Core/BlendStackRootCameraNode.cpp
+++ b/Engine/Plugins/Cameras/GameplayCameras/Source/GameplayCameras/Private/Core/BlendStackRootCameraNode.cpp
@@ -11,7 +11,6 @@
 #include "Debug/CameraDebugBlockBuilder.h"
 #include "Debug/CameraDebugRenderer.h"
 #include "Helpers/CameraObjectInterfaceParameterOverrideHelper.h"
-#include "Helpers/CameraObjectReferenceParameterOverrideEvaluator.h"
 #include "Nodes/Common/CameraRigCameraNode.h"
 
 #include UE_INLINE_GENERATED_CPP_BY_NAME(BlendStackRootCameraNode)
@@ -292,8 +291,7 @@ void FBlendStackRootCameraNodeEvaluator::RunBlendedParameterOverridesStack(const
 		for (int32 Index = BlendedParameterOverrides.PrefabTrail.Num() - 1; Index >= 0; --Index)
 		{
 			const UCameraRigCameraNode* CurPrefabNode = BlendedParameterOverrides.PrefabTrail[Index];
-			FCameraObjectReferenceParameterOverrideEvaluator OverrideEvaluator(CurPrefabNode->CameraRigReference);
-			OverrideEvaluator.ApplyParameterOverrides(CurResult.VariableTable, false);
+			CurPrefabNode->CameraRigReference.ApplyParameterOverrides(CurResult.VariableTable, false);
 		}
 
 		// Finally, update the parameter overrides' blend, and apply it.
diff --git a/Engine/Plugins/Cameras/GameplayCameras/Source/GameplayCameras/Private/Core/CameraAssetReference.cpp b/Engine/Plugins/Cameras/GameplayCameras/Source/GameplayCameras/Private/Core/CameraAssetReference.cpp
--- a/Engine/Plugins/Cameras/GameplayCameras/Source/GameplayCameras/Private/Core/CameraAssetReference.cpp
+++ b/Engine/Plugins/Cameras/GameplayCameras/Source/GameplayCameras/Private/Core/CameraAssetReference.cpp
@@ -4,7 +4,7 @@
 
 #include "Core/CameraAsset.h"
 #include "Core/CameraNodeEvaluator.h"
-#include "Helpers/CameraAssetParameterOverrideEvaluator.h"
+#include "Helpers/CameraObjectInterfaceParameterOverrideHelper.h"
 
 #include UE_INLINE_GENERATED_CPP_BY_NAME(CameraAssetReference)
 
@@ -17,11 +17,25 @@ FCameraAssetReference::FCameraAssetReference(UCameraAsset* InCameraAsset)
 {
 }
 
-void FCameraAssetReference::ApplyParameterOverrides(UE::Cameras::FCameraNodeEvaluationResult& OutResult, bool bDrivenOverridesOnly)
+void FCameraAssetReference::ApplyParameterOverrides(UE::Cameras::FCameraNodeEvaluationResult& OutResult, bool bDrivenOnly) const
 {
 	using namespace UE::Cameras;
-	FCameraAssetParameterOverrideEvaluator OverrideEvaluator(*this);
-	OverrideEvaluator.ApplyParameterOverrides(OutResult.VariableTable, OutResult.ContextDataTable, bDrivenOverridesOnly);
+	if (CameraAsset)
+	{
+		TSet<FGuid> OverridenParameterGuids(GetOverriddenParameterGuids());
+		FCameraObjectInterfaceParameterOverrideHelper Helper(&OutResult.VariableTable, &OutResult.ContextDataTable);
+		Helper.ApplyParameterOverrides(CameraAsset, CameraAsset->GetParameterDefinitions(), Parameters, OverridenParameterGuids, bDrivenOnly);
+	}
+}
+
+void FCameraAssetReference::ApplyParameterOverrides(const FInstancedPropertyBag& CachedParameters, UE::Cameras::FCameraNodeEvaluationResult& OutResult) const
+{
+	using namespace UE::Cameras;
+	if (CameraAsset)
+	{
+		FCameraObjectInterfaceParameterOverrideHelper Helper(&OutResult.VariableTable, &OutResult.ContextDataTable);
+		Helper.ApplyParameterOverrides(CameraAsset, CameraAsset->GetParameterDefinitions(), Parameters, CachedParameters);
+	}
 }
 
 bool FCameraAssetReference::IsParameterOverridden(const FGuid PropertyID) const
@@ -41,23 +55,6 @@ void FCameraAssetReference::SetParameterOverridden(const FGuid PropertyID, bool
 	}
 }
 
-bool FCameraAssetReference::IsParameterAnimated(const FGuid PropertyID) const
-{
-	return ParameterAnimatedGuids.Contains(PropertyID);
-}
-
-void FCameraAssetReference::SetParameterAnimated(const FGuid PropertyID, bool bIsAnimated)
-{
-	if (bIsAnimated)
-	{
-		ParameterAnimatedGuids.AddUnique(PropertyID);
-	}
-	else
-	{
-		ParameterAnimatedGuids.Remove(PropertyID);
-	}
-}
-
 bool FCameraAssetReference::NeedsRebuildParameters() const
 {
 	if ((!CameraAsset && Parameters.IsValid()) || (CameraAsset && !Parameters.IsValid()))
diff --git a/Engine/Plugins/Cameras/GameplayCameras/Source/GameplayCameras/Private/Core/CameraRigAssetReference.cpp b/Engine/Plugins/Cameras/GameplayCameras/Source/GameplayCameras/Private/Core/CameraRigAssetReference.cpp
--- a/Engine/Plugins/Cameras/GameplayCameras/Source/GameplayCameras/Private/Core/CameraRigAssetReference.cpp
+++ b/Engine/Plugins/Cameras/GameplayCameras/Source/GameplayCameras/Private/Core/CameraRigAssetReference.cpp
@@ -5,7 +5,7 @@
 #include "Build/CameraObjectInterfaceParameterBuilder.h"
 #include "Core/CameraRigAsset.h"
 #include "Core/CameraNodeEvaluator.h"
-#include "Helpers/CameraObjectReferenceParameterOverrideEvaluator.h"
+#include "Helpers/CameraObjectInterfaceParameterOverrideHelper.h"
 
 #include UE_INLINE_GENERATED_CPP_BY_NAME(CameraRigAssetReference)
 
@@ -27,11 +27,43 @@ const UBaseCameraObject* FCameraRigAssetReference::GetCameraObject() const
 	return CameraRig;
 }
 
-void FCameraRigAssetReference::ApplyParameterOverrides(UE::Cameras::FCameraNodeEvaluationResult& OutResult, bool bDrivenOverridesOnly)
+void FCameraRigAssetReference::ApplyParameterOverrides(UE::Cameras::FCameraVariableTable& OutVariableTable, bool bDrivenOnly) const
+{
+	ApplyParameterOverridesImpl(&OutVariableTable, nullptr, bDrivenOnly);
+}
+
+void FCameraRigAssetReference::ApplyParameterOverrides(UE::Cameras::FCameraVariableTable& OutVariableTable, UE::Cameras::FCameraContextDataTable& OutContextDataTable, bool bDrivenOnly) const
+{
+	ApplyParameterOverridesImpl(&OutVariableTable, &OutContextDataTable, bDrivenOnly);
+}
+
+void FCameraRigAssetReference::ApplyParameterOverrides(UE::Cameras::FCameraNodeEvaluationResult& OutResult, bool bDrivenOnly) const
+{
+	ApplyParameterOverridesImpl(&OutResult.VariableTable, &OutResult.ContextDataTable, bDrivenOnly);
+}
+
+void FCameraRigAssetReference::ApplyParameterOverrides(const FInstancedPropertyBag& CachedParameters, UE::Cameras::FCameraNodeEvaluationResult& OutResult) const
 {
 	using namespace UE::Cameras;
-	FCameraObjectReferenceParameterOverrideEvaluator OverrideEvaluator(*this);
-	OverrideEvaluator.ApplyParameterOverrides(OutResult.VariableTable, OutResult.ContextDataTable, bDrivenOverridesOnly);
+
+	if (CameraRig)
+	{
+		FCameraObjectInterfaceParameterOverrideHelper Helper(&OutResult.VariableTable, &OutResult.ContextDataTable);
+		Helper.ApplyParameterOverrides(CameraRig, CameraRig->GetParameterDefinitions(), Parameters, CachedParameters);
+	}
+}
+
+void FCameraRigAssetReference::ApplyParameterOverridesImpl(UE::Cameras::FCameraVariableTable* OutVariableTable, UE::Cameras::FCameraContextDataTable* OutContextDataTable, bool bDrivenOnly) const
+{
+	using namespace UE::Cameras;
+
+	if (CameraRig)
+	{
+		TSet<FGuid> OverridenParameterGuids;
+		GetOverriddenParameterGuids(OverridenParameterGuids);
+		FCameraObjectInterfaceParameterOverrideHelper Helper(OutVariableTable, OutContextDataTable);
+		Helper.ApplyParameterOverrides(CameraRig, CameraRig->GetParameterDefinitions(), Parameters, OverridenParameterGuids, bDrivenOnly);
+	}
 }
 
 bool FCameraRigAssetReference::SerializeFromMismatchedTag(struct FPropertyTag const& Tag, FStructuredArchive::FSlot Slot)
diff --git a/Engine/Plugins/Cameras/GameplayCameras/Source/GameplayCameras/Private/Core/CameraRigCombinationRegistry.cpp b/Engine/Plugins/Cameras/GameplayCameras/Source/GameplayCameras/Private/Core/CameraRigCombinationRegistry.cpp
--- a/Engine/Plugins/Cameras/GameplayCameras/Source/GameplayCameras/Private/Core/CameraRigCombinationRegistry.cpp
+++ b/Engine/Plugins/Cameras/GameplayCameras/Source/GameplayCameras/Private/Core/CameraRigCombinationRegistry.cpp
@@ -3,7 +3,7 @@
 #include "Core/CameraRigCombinationRegistry.h"
 
 #include "Core/CameraRigAsset.h"
-#include "Helpers/CameraObjectReferenceParameterOverrideEvaluator.h"
+#include "Helpers/CameraObjectInterfaceParameterOverrideHelper.h"
 
 #include UE_INLINE_GENERATED_CPP_BY_NAME(CameraRigCombinationRegistry)
 
@@ -125,8 +125,7 @@ void FCombinedCameraRigsCameraNodeEvaluator::ApplyParameterOverrides(FCameraVari
 	const UCombinedCameraRigsCameraNode* CombinedRigsNode = GetCameraNodeAs<UCombinedCameraRigsCameraNode>();
 	for (const FCameraRigAssetReference& IndividualCameraRigReference : CombinedRigsNode->CameraRigReferences)
 	{
-		FCameraObjectReferenceParameterOverrideEvaluator OverrideEvaluator(IndividualCameraRigReference);
-		OverrideEvaluator.ApplyParameterOverrides(OutVariableTable, bDrivenOnly);
+		IndividualCameraRigReference.ApplyParameterOverrides(OutVariableTable, bDrivenOnly);
 	}
 }
 
diff --git a/Engine/Plugins/Cameras/GameplayCameras/Source/GameplayCameras/Private/Core/CameraShakeAssetReference.cpp b/Engine/Plugins/Cameras/GameplayCameras/Source/GameplayCameras/Private/Core/CameraShakeAssetReference.cpp
--- a/Engine/Plugins/Cameras/GameplayCameras/Source/GameplayCameras/Private/Core/CameraShakeAssetReference.cpp
+++ b/Engine/Plugins/Cameras/GameplayCameras/Source/GameplayCameras/Private/Core/CameraShakeAssetReference.cpp
@@ -3,7 +3,8 @@
 #include "Core/CameraShakeAssetReference.h"
 
 #include "Core/CameraShakeAsset.h"
-#include "Helpers/CameraObjectReferenceParameterOverrideEvaluator.h"
+#include "Core/CameraVariableTable.h"
+#include "Helpers/CameraObjectInterfaceParameterOverrideHelper.h"
 
 #include UE_INLINE_GENERATED_CPP_BY_NAME(CameraShakeAssetReference)
 
@@ -21,10 +22,31 @@ const UBaseCameraObject* FCameraShakeAssetReference::GetCameraObject() const
 	return CameraShake;
 }
 
-void FCameraShakeAssetReference::ApplyParameterOverrides(UE::Cameras::FCameraNodeEvaluationResult& OutResult, bool bDrivenOverridesOnly)
+void FCameraShakeAssetReference::ApplyParameterOverrides(UE::Cameras::FCameraVariableTable& OutVariableTable, bool bDrivenOnly) const
+{
+	ApplyParameterOverridesImpl(&OutVariableTable, nullptr, bDrivenOnly);
+}
+
+void FCameraShakeAssetReference::ApplyParameterOverrides(UE::Cameras::FCameraVariableTable& OutVariableTable, UE::Cameras::FCameraContextDataTable& OutContextDataTable, bool bDrivenOnly) const
+{
+	ApplyParameterOverridesImpl(&OutVariableTable, &OutContextDataTable, bDrivenOnly);
+}
+
+void FCameraShakeAssetReference::ApplyParameterOverrides(UE::Cameras::FCameraNodeEvaluationResult& OutResult, bool bDrivenOnly) const
+{
+	ApplyParameterOverridesImpl(&OutResult.VariableTable, &OutResult.ContextDataTable, bDrivenOnly);
+}
+
+void FCameraShakeAssetReference::ApplyParameterOverridesImpl(UE::Cameras::FCameraVariableTable* OutVariableTable, UE::Cameras::FCameraContextDataTable* OutContextDataTable, bool bDrivenOnly) const
 {
 	using namespace UE::Cameras;
-	FCameraObjectReferenceParameterOverrideEvaluator OverrideEvaluator(*this);
-	OverrideEvaluator.ApplyParameterOverrides(OutResult.VariableTable, OutResult.ContextDataTable, bDrivenOverridesOnly);
+	
+	if (CameraShake)
+	{
+		TSet<FGuid> OverridenParameterGuids;
+		GetOverriddenParameterGuids(OverridenParameterGuids);
+		FCameraObjectInterfaceParameterOverrideHelper Helper(OutVariableTable, OutContextDataTable);
+		Helper.ApplyParameterOverrides(CameraShake, CameraShake->GetParameterDefinitions(), Parameters, OverridenParameterGuids, bDrivenOnly);
+	}
 }
 
diff --git a/Engine/Plugins/Cameras/GameplayCameras/Source/GameplayCameras/Private/GameFramework/GameplayCameraComponent.cpp b/Engine/Plugins/Cameras/GameplayCameras/Source/GameplayCameras/Private/GameFramework/GameplayCameraComponent.cpp
--- a/Engine/Plugins/Cameras/GameplayCameras/Source/GameplayCameras/Private/GameFramework/GameplayCameraComponent.cpp
+++ b/Engine/Plugins/Cameras/GameplayCameras/Source/GameplayCameras/Private/GameFramework/GameplayCameraComponent.cpp
@@ -74,8 +74,16 @@ void UGameplayCameraComponent::OnUpdateCameraEvaluationContext(bool bForceApplyP
 
 	FCameraNodeEvaluationResult& InitialResult = GetEvaluationContext()->GetInitialResult();
 
-	const bool bApplyDrivenParametersOnly = !bForceApplyParameterOverrides;
-	CameraReference.ApplyParameterOverrides(InitialResult, bApplyDrivenParametersOnly);
+	if (bForceApplyParameterOverrides)
+	{
+		CameraReference.ApplyParameterOverrides(InitialResult, false);
+	}
+	else
+	{
+		CameraReference.ApplyParameterOverrides(CachedParameterOverrides, InitialResult);
+	}
+
+	CachedParameterOverrides = CameraReference.GetParameters();
 }
 
 #if WITH_EDITOR
diff --git a/Engine/Plugins/Cameras/GameplayCameras/Source/GameplayCameras/Private/GameFramework/GameplayCameraRigComponent.cpp b/Engine/Plugins/Cameras/GameplayCameras/Source/GameplayCameras/Private/GameFramework/GameplayCameraRigComponent.cpp
--- a/Engine/Plugins/Cameras/GameplayCameras/Source/GameplayCameras/Private/GameFramework/GameplayCameraRigComponent.cpp
+++ b/Engine/Plugins/Cameras/GameplayCameras/Source/GameplayCameras/Private/GameFramework/GameplayCameraRigComponent.cpp
@@ -74,8 +74,16 @@ void UGameplayCameraRigComponent::OnUpdateCameraEvaluationContext(bool bForceApp
 
 	FCameraNodeEvaluationResult& InitialResult = GetEvaluationContext()->GetInitialResult();
 
-	const bool bApplyDrivenParametersOnly = !bForceApplyParameterOverrides;
-	CameraRigReference.ApplyParameterOverrides(InitialResult, bApplyDrivenParametersOnly);
+	if (bForceApplyParameterOverrides)
+	{
+		CameraRigReference.ApplyParameterOverrides(InitialResult, false);
+	}
+	else
+	{
+		CameraRigReference.ApplyParameterOverrides(CachedParameterOverrides, InitialResult);
+	}
+
+	CachedParameterOverrides = CameraRigReference.GetParameters();
 }
 
 #if WITH_EDITOR
diff --git a/Engine/Plugins/Cameras/GameplayCameras/Source/GameplayCameras/Private/GameplayCamerasModule.cpp b/Engine/Plugins/Cameras/GameplayCameras/Source/GameplayCameras/Private/GameplayCamerasModule.cpp
--- a/Engine/Plugins/Cameras/GameplayCameras/Source/GameplayCameras/Private/GameplayCamerasModule.cpp
+++ b/Engine/Plugins/Cameras/GameplayCameras/Source/GameplayCameras/Private/GameplayCamerasModule.cpp
@@ -4,11 +4,7 @@
 
 #include "Debug/CameraDebugColors.h"
 #include "GameplayCameras.h"
-#include "Logging/MessageLog.h"
-#include "Misc/EngineVersionComparison.h"
 #include "Modules/ModuleManager.h"
-#include "MovieScene/MovieSceneCameraParameterInstantiator.h"
-#include "MovieSceneSection.h"
 #include "Nodes/Framing/CameraFramingZone.h"
 #include "ShowFlags.h"
 #include "UObject/UObjectBase.h"
@@ -37,7 +33,6 @@ public:
 	virtual void StartupModule() override
 	{
 		RegisterBuiltInBlendableStructs();
-		RegisterMovieSceneIntegration();
 		
 #if UE_GAMEPLAY_CAMERAS_DEBUG
 		UE::Cameras::FCameraDebugColors::RegisterBuiltinColorSchemes();
@@ -47,7 +42,6 @@ public:
 	virtual void ShutdownModule() override
 	{
 		UnregisterBuiltInBlendableStructs();
-		UnregisterMovieSceneIntegration();
 	}
 
 public:
@@ -120,25 +114,6 @@ private:
 		}
 	}
 
-	void RegisterMovieSceneIntegration()
-	{
-#if WITH_EDITORONLY_DATA
-#if UE_VERSION_NEWER_THAN_OR_EQUAL(5,6,0)
-		OnMovieSceneSectionAddedToTrackHandle = UMovieSceneSection::GetOnSectionAddedToTrack().AddStatic(
-				&UMovieSceneCameraParameterInstantiator::OnMovieSceneSectionAddedToTrack);
-#endif
-#endif
-	}
-
-	void UnregisterMovieSceneIntegration()
-	{
-#if WITH_EDITORONLY_DATA
-#if UE_VERSION_NEWER_THAN_OR_EQUAL(5,6,0)
-		UMovieSceneSection::GetOnSectionAddedToTrack().Remove(OnMovieSceneSectionAddedToTrackHandle);
-#endif
-#endif
-	}
-
 private:
 
 	TArray<UE::Cameras::FBlendableStructInfo> BlendableStructs;
diff --git a/Engine/Plugins/Cameras/GameplayCameras/Source/GameplayCameras/Private/Helpers/CameraAssetParameterOverrideEvaluator.cpp b/Engine/Plugins/Cameras/GameplayCameras/Source/GameplayCameras/Private/Helpers/CameraAssetParameterOverrideEvaluator.cpp
deleted file mode 100644
diff --git a/Engine/Plugins/Cameras/GameplayCameras/Source/GameplayCameras/Private/Helpers/CameraAssetParameterOverrideEvaluator.h b/Engine/Plugins/Cameras/GameplayCameras/Source/GameplayCameras/Private/Helpers/CameraAssetParameterOverrideEvaluator.h
deleted file mode 100644
diff --git a/Engine/Plugins/Cameras/GameplayCameras/Source/GameplayCameras/Private/Helpers/CameraObjectInterfaceParameterOverrideHelper.cpp b/Engine/Plugins/Cameras/GameplayCameras/Source/GameplayCameras/Private/Helpers/CameraObjectInterfaceParameterOverrideHelper.cpp
--- a/Engine/Plugins/Cameras/GameplayCameras/Source/GameplayCameras/Private/Helpers/CameraObjectInterfaceParameterOverrideHelper.cpp
+++ b/Engine/Plugins/Cameras/GameplayCameras/Source/GameplayCameras/Private/Helpers/CameraObjectInterfaceParameterOverrideHelper.cpp
@@ -21,7 +21,7 @@ void ApplyBlendableParameterOverride(
 		const FCameraObjectInterfaceParameterDefinition& ParameterDefinition,
 		const ParameterType& ParameterValue,
 		FCameraVariableTable& VariableTable,
-		bool bDrivenOverridesOnly)
+		bool bDrivenOnly)
 {
 	using ValueType = typename ParameterType::ValueType;
 
@@ -36,7 +36,7 @@ void ApplyBlendableParameterOverride(
 				OverrideDefinition.VariableID, ParameterValue.Variable->GetDefaultValue());
 		VariableTable.SetValue<ValueType>(ParameterVariableID, OverrideValue);
 	}
-	else if (!bDrivenOverridesOnly)
+	else if (!bDrivenOnly)
 	{
 		// The override is a fixed value. Just set that on the prefab's variable.
 		VariableTable.SetValue<ValueType>(ParameterVariableID, ParameterValue.Value);
@@ -44,12 +44,12 @@ void ApplyBlendableParameterOverride(
 }
 
 void ApplyBlendableParameterOverride(
-		const UBaseCameraObject* CameraObject,
+		const UObject* CameraObject,
 		const FCameraObjectInterfaceParameterDefinition& ParameterDefinition,
 		const FInstancedPropertyBag& PropertyBag,
 		const FPropertyBagPropertyDesc& PropertyBagPropertyDesc,
 		FCameraVariableTable& VariableTable,
-		bool bDrivenOverridesOnly)
+		bool bDrivenOnly)
 {
 	ensure(ParameterDefinition.ParameterType == ECameraObjectInterfaceParameterType::Blendable);
 
@@ -91,7 +91,7 @@ void ApplyBlendableParameterOverride(
 			{\
 				check(ParameterType == F##ValueName##CameraParameter::StaticStruct());\
 				const F##ValueName##CameraParameter& TypedParameterValue = ParameterValue.Get<F##ValueName##CameraParameter>();\
-				ApplyBlendableParameterOverride(ParameterDefinition, TypedParameterValue, VariableTable, bDrivenOverridesOnly);\
+				ApplyBlendableParameterOverride(ParameterDefinition, TypedParameterValue, VariableTable, bDrivenOnly);\
 			}\
 			break;
 UE_CAMERA_VARIABLE_FOR_ALL_TYPES()
@@ -190,7 +190,7 @@ void ApplyDataParameterElementOverride(
 }
 
 void ApplyDataParameterSingleOverride(
-		const UBaseCameraObject* CameraObject,
+		const UObject* CameraObject,
 		const FCameraObjectInterfaceParameterDefinition& ParameterDefinition,
 		const FInstancedPropertyBag& PropertyBag,
 		const FPropertyBagPropertyDesc& PropertyBagPropertyDesc,
@@ -273,7 +273,7 @@ void ApplyDataParameterSingleOverride(
 }
 
 void ApplyDataParameterArrayOverride(
-		const UBaseCameraObject* CameraObject,
+		const UObject* CameraObject,
 		const FCameraObjectInterfaceParameterDefinition& ParameterDefinition,
 		const FInstancedPropertyBag& PropertyBag,
 		const FPropertyBagPropertyDesc& PropertyBagPropertyDesc,
@@ -360,7 +360,7 @@ void ApplyDataParameterArrayOverride(
 }
 
 void ApplyDataParameterOverride(
-		const UBaseCameraObject* CameraObject,
+		const UObject* CameraObject,
 		const FCameraObjectInterfaceParameterDefinition& ParameterDefinition,
 		const FInstancedPropertyBag& PropertyBag,
 		const FPropertyBagPropertyDesc& PropertyBagPropertyDesc,
@@ -399,56 +399,96 @@ FCameraObjectInterfaceParameterOverrideHelper::FCameraObjectInterfaceParameterOv
 }
 
 void FCameraObjectInterfaceParameterOverrideHelper::ApplyParameterOverrides(
-		const UBaseCameraObject* CameraObject,
+		const UObject* CameraObject,
+		TConstArrayView<FCameraObjectInterfaceParameterDefinition> ParameterDefinitions,
 		const FInstancedPropertyBag& ParameterOverrides,
 		const TSet<FGuid>& OverridenParameterGuids,
-		const TSet<FGuid>& AnimatedParameterGuids,
-		FCameraVariableTable* OutVariableTable,
-		FCameraContextDataTable* OutContextDataTable,
-		bool bDrivenOverridesOnly)
+		bool bDrivenOnly)
 {
 	using namespace Internal;
 
-	check(OutVariableTable);
-
-	const UPropertyBag* CameraShakeParametersStruct = ParameterOverrides.GetPropertyBagStruct();
-	if (!CameraObject || !CameraShakeParametersStruct)
+	const UPropertyBag* ParameterOverridesStruct = ParameterOverrides.GetPropertyBagStruct();
+	if (!CameraObject || !ParameterOverridesStruct)
 	{
 		return;
 	}
 
-	TConstArrayView<FCameraObjectInterfaceParameterDefinition> ParameterDefinitions = CameraObject->GetParameterDefinitions();
-
 	for (const FCameraObjectInterfaceParameterDefinition& Definition : ParameterDefinitions)
 	{
-		if (!OutContextDataTable && Definition.ParameterType == ECameraObjectInterfaceParameterType::Data)
+		if (!ContextDataTable && Definition.ParameterType == ECameraObjectInterfaceParameterType::Data)
 		{
 			continue;
 		}
 
-		const bool bIsAnimated = AnimatedParameterGuids.Contains(Definition.ParameterGuid);
-		if (!OverridenParameterGuids.Contains(Definition.ParameterGuid) && !bIsAnimated)
+		if (!OverridenParameterGuids.Contains(Definition.ParameterGuid))
 		{
 			continue;
 		}
 
-		const FPropertyBagPropertyDesc* PropertyDesc = CameraShakeParametersStruct->FindPropertyDescByID(Definition.ParameterGuid);
+		const FPropertyBagPropertyDesc* PropertyDesc = ParameterOverridesStruct->FindPropertyDescByID(Definition.ParameterGuid);
 		if (!ensure(PropertyDesc))
 		{
 			continue;
 		}
 
-		const bool bThisDrivenOnly = bDrivenOverridesOnly && !bIsAnimated;
-		ApplyParameterOverride(CameraObject, Definition, ParameterOverrides, *PropertyDesc, bThisDrivenOnly);
+		ApplyParameterOverride(CameraObject, Definition, ParameterOverrides, *PropertyDesc, bDrivenOnly);
+	}
+}
+
+void FCameraObjectInterfaceParameterOverrideHelper::ApplyParameterOverrides(
+		const UObject* CameraObject,
+		TConstArrayView<FCameraObjectInterfaceParameterDefinition> ParameterDefinitions,
+		const FInstancedPropertyBag& ParameterOverrides,
+		const FInstancedPropertyBag& CachedParameterOverrides)
+{
+	using namespace Internal;
+
+	const UPropertyBag* ParameterOverridesStruct = ParameterOverrides.GetPropertyBagStruct();
+	if (!CameraObject || !ParameterOverridesStruct)
+	{
+		return;
+	}
+
+	const uint8* ParameterOverridesMemory = ParameterOverrides.GetValue().GetMemory();
+	const uint8* CachedParameterOverridesMemory = CachedParameterOverrides.GetValue().GetMemory();
+	if (!ParameterOverridesMemory || !CachedParameterOverridesMemory)
+	{
+		return;
+	}
+
+	if (!ensure(ParameterOverridesStruct == CachedParameterOverrides.GetPropertyBagStruct()))
+	{
+		return;
+	}
+
+	for (const FCameraObjectInterfaceParameterDefinition& Definition : ParameterDefinitions)
+	{
+		if (!ContextDataTable && Definition.ParameterType == ECameraObjectInterfaceParameterType::Data)
+		{
+			continue;
+		}
+
+		const FPropertyBagPropertyDesc* PropertyDesc = ParameterOverridesStruct->FindPropertyDescByID(Definition.ParameterGuid);
+		if (!ensure(PropertyDesc))
+		{
+			continue;
+		}
+
+		const void* RawValue = PropertyDesc->CachedProperty->ContainerPtrToValuePtr<void>(ParameterOverridesMemory);
+		const void* CachedRawValue = PropertyDesc->CachedProperty->ContainerPtrToValuePtr<void>(CachedParameterOverridesMemory);
+		if (!PropertyDesc->CachedProperty->Identical(RawValue, CachedRawValue))
+		{
+			ApplyParameterOverride(CameraObject, Definition, ParameterOverrides, *PropertyDesc, false);
+		}
 	}
 }
 
 void FCameraObjectInterfaceParameterOverrideHelper::ApplyParameterOverride(
-		const UBaseCameraObject* CameraObject,
+		const UObject* CameraObject,
 		const FCameraObjectInterfaceParameterDefinition& ParameterDefinition,
 		const FInstancedPropertyBag& PropertyBag,
 		const FPropertyBagPropertyDesc& PropertyBagPropertyDesc,
-		bool bDrivenOverridesOnly)
+		bool bDrivenOnly)
 {
 	using namespace Internal;
 
@@ -464,15 +504,14 @@ void FCameraObjectInterfaceParameterOverrideHelper::ApplyParameterOverride(
 							CameraObject, 
 							ParameterDefinition,
 							PropertyBag, PropertyBagPropertyDesc, 
-							*VariableTable, bDrivenOverridesOnly);
+							*VariableTable,
+							bDrivenOnly);
 				}
 			}
 			break;
 		case ECameraObjectInterfaceParameterType::Data:
 			{
-				// Data parameters can't be driven by variables so only apply them if we are not
-				// just applying driven parameters.
-				if (ensure(ContextDataTable) && !bDrivenOverridesOnly)
+				if (ensure(ContextDataTable) && !bDrivenOnly)
 				{
 					ApplyDataParameterOverride(
 							CameraObject, 
diff --git a/Engine/Plugins/Cameras/GameplayCameras/Source/GameplayCameras/Private/Helpers/CameraObjectInterfaceParameterOverrideHelper.h b/Engine/Plugins/Cameras/GameplayCameras/Source/GameplayCameras/Private/Helpers/CameraObjectInterfaceParameterOverrideHelper.h
--- a/Engine/Plugins/Cameras/GameplayCameras/Source/GameplayCameras/Private/Helpers/CameraObjectInterfaceParameterOverrideHelper.h
+++ b/Engine/Plugins/Cameras/GameplayCameras/Source/GameplayCameras/Private/Helpers/CameraObjectInterfaceParameterOverrideHelper.h
@@ -5,6 +5,7 @@
 #include "StructUtils/PropertyBag.h"
 
 class UBaseCameraObject;
+class UCameraAsset;
 struct FCameraObjectInterfaceParameterDefinition;
 
 namespace UE::Cameras
@@ -44,21 +45,27 @@ public:
 
 	/** Sets overriden values of interface parameters in the given variable and context data tables. */
 	void ApplyParameterOverrides(
-			const UBaseCameraObject* CameraObject,
+			const UObject* CameraObject,
+			TConstArrayView<FCameraObjectInterfaceParameterDefinition> ParameterDefinitions,
 			const FInstancedPropertyBag& ParameterOverrides,
 			const TSet<FGuid>& OverridenParameterGuids,
-			const TSet<FGuid>& AnimatedParameterGuids,
-			FCameraVariableTable* OutVariableTable,
-			FCameraContextDataTable* OutContextDataTable,
-			bool bDrivenOverridesOnly);
+			bool bDrivenOnly);
+
+	/** Sets overriden values of interface parameters in the given variable and context data tables. */
+	void ApplyParameterOverrides(
+			const UObject* CameraObject,
+			TConstArrayView<FCameraObjectInterfaceParameterDefinition> ParameterDefinitions,
+			const FInstancedPropertyBag& ParameterOverrides,
+			const FInstancedPropertyBag& CachedParameterOverrides);
+
+private:
 
-	/** Sets a given override value for an interface parameter. */
 	void ApplyParameterOverride(
-			const UBaseCameraObject* CameraObject,
+			const UObject* CameraObject,
 			const FCameraObjectInterfaceParameterDefinition& ParameterDefinition,
 			const FInstancedPropertyBag& PropertyBag,
 			const FPropertyBagPropertyDesc& PropertyBagPropertyDesc,
-			bool bDrivenOverridesOnly);
+			bool bDrivenOnly);
 
 private:
 
diff --git a/Engine/Plugins/Cameras/GameplayCameras/Source/GameplayCameras/Private/Helpers/CameraObjectReferenceParameterOverrideEvaluator.h b/Engine/Plugins/Cameras/GameplayCameras/Source/GameplayCameras/Private/Helpers/CameraObjectReferenceParameterOverrideEvaluator.h
deleted file mode 100644
diff --git a/Engine/Plugins/Cameras/GameplayCameras/Source/GameplayCameras/Private/MovieScene/MovieSceneCameraParameterInstantiator.cpp b/Engine/Plugins/Cameras/GameplayCameras/Source/GameplayCameras/Private/MovieScene/MovieSceneCameraParameterInstantiator.cpp
deleted file mode 100644
diff --git a/Engine/Plugins/Cameras/GameplayCameras/Source/GameplayCameras/Private/MovieScene/MovieSceneCameraParameterInstantiator.h b/Engine/Plugins/Cameras/GameplayCameras/Source/GameplayCameras/Private/MovieScene/MovieSceneCameraParameterInstantiator.h
deleted file mode 100644
diff --git a/Engine/Plugins/Cameras/GameplayCameras/Source/GameplayCameras/Private/Nodes/Common/CameraRigCameraNode.cpp b/Engine/Plugins/Cameras/GameplayCameras/Source/GameplayCameras/Private/Nodes/Common/CameraRigCameraNode.cpp
--- a/Engine/Plugins/Cameras/GameplayCameras/Source/GameplayCameras/Private/Nodes/Common/CameraRigCameraNode.cpp
+++ b/Engine/Plugins/Cameras/GameplayCameras/Source/GameplayCameras/Private/Nodes/Common/CameraRigCameraNode.cpp
@@ -7,7 +7,6 @@
 #include "Core/CameraNodeEvaluator.h"
 #include "Core/CameraRigAsset.h"
 #include "Core/CameraRigAssetReference.h"
-#include "Helpers/CameraObjectReferenceParameterOverrideEvaluator.h"
 #include "Helpers/CameraObjectInterfaceParameterOverrideHelper.h"
 #include "Logging/TokenizedMessage.h"
 
@@ -74,17 +73,13 @@ void FCameraRigCameraNodeEvaluator::OnRun(const FCameraNodeEvaluationParams& Par
 void FCameraRigCameraNodeEvaluator::ApplyParameterOverrides(FCameraVariableTable& OutVariableTable, bool bDrivenOnly)
 {
 	const UCameraRigCameraNode* PrefabNode = GetCameraNodeAs<UCameraRigCameraNode>();
-
-	FCameraObjectReferenceParameterOverrideEvaluator OverrideEvaluator(PrefabNode->CameraRigReference);
-	OverrideEvaluator.ApplyParameterOverrides(OutVariableTable, bDrivenOnly);
+	PrefabNode->CameraRigReference.ApplyParameterOverrides(OutVariableTable, bDrivenOnly);
 }
 
 void FCameraRigCameraNodeEvaluator::ApplyParameterOverrides(FCameraVariableTable& OutVariableTable, FCameraContextDataTable& OutContextDataTable, bool bDrivenOnly)
 {
 	const UCameraRigCameraNode* PrefabNode = GetCameraNodeAs<UCameraRigCameraNode>();
-
-	FCameraObjectReferenceParameterOverrideEvaluator OverrideEvaluator(PrefabNode->CameraRigReference);
-	OverrideEvaluator.ApplyParameterOverrides(OutVariableTable, OutContextDataTable, bDrivenOnly);
+	PrefabNode->CameraRigReference.ApplyParameterOverrides(OutVariableTable, OutContextDataTable, bDrivenOnly);
 }
 
 }  // namespace UE::Cameras
diff --git a/Engine/Plugins/Cameras/GameplayCameras/Source/GameplayCameras/Private/Nodes/Shakes/CameraShakeCameraNode.cpp b/Engine/Plugins/Cameras/GameplayCameras/Source/GameplayCameras/Private/Nodes/Shakes/CameraShakeCameraNode.cpp
--- a/Engine/Plugins/Cameras/GameplayCameras/Source/GameplayCameras/Private/Nodes/Shakes/CameraShakeCameraNode.cpp
+++ b/Engine/Plugins/Cameras/GameplayCameras/Source/GameplayCameras/Private/Nodes/Shakes/CameraShakeCameraNode.cpp
@@ -7,7 +7,6 @@
 #include "Core/CameraNodeEvaluator.h"
 #include "Core/CameraSystemEvaluator.h"
 #include "Core/ShakeCameraNode.h"
-#include "Helpers/CameraObjectReferenceParameterOverrideEvaluator.h"
 #include "Services/CameraShakeService.h"
 
 #include UE_INLINE_GENERATED_CPP_BY_NAME(CameraShakeCameraNode)
@@ -124,17 +123,13 @@ void FCameraShakeCameraNodeEvaluator::OnRun(const FCameraNodeEvaluationParams& P
 void FCameraShakeCameraNodeEvaluator::ApplyParameterOverrides(FCameraVariableTable& OutVariableTable, bool bDrivenOnly)
 {
 	const UCameraShakeCameraNode* PrefabNode = GetCameraNodeAs<UCameraShakeCameraNode>();
-
-	FCameraObjectReferenceParameterOverrideEvaluator OverrideEvaluator(PrefabNode->CameraShakeReference);
-	OverrideEvaluator.ApplyParameterOverrides(OutVariableTable, bDrivenOnly);
+	PrefabNode->CameraShakeReference.ApplyParameterOverrides(OutVariableTable, bDrivenOnly);
 }
 
 void FCameraShakeCameraNodeEvaluator::ApplyParameterOverrides(FCameraVariableTable& OutVariableTable, FCameraContextDataTable& OutContextDataTable, bool bDrivenOnly)
 {
 	const UCameraShakeCameraNode* PrefabNode = GetCameraNodeAs<UCameraShakeCameraNode>();
-
-	FCameraObjectReferenceParameterOverrideEvaluator OverrideEvaluator(PrefabNode->CameraShakeReference);
-	OverrideEvaluator.ApplyParameterOverrides(OutVariableTable, OutContextDataTable, bDrivenOnly);
+	PrefabNode->CameraShakeReference.ApplyParameterOverrides(OutVariableTable, OutContextDataTable, bDrivenOnly);
 }
 
 }  // namespace UE::Cameras
diff --git a/Engine/Plugins/Cameras/GameplayCameras/Source/GameplayCameras/Private/Nodes/Utility/BlueprintCameraNode.cpp b/Engine/Plugins/Cameras/GameplayCameras/Source/GameplayCameras/Private/Nodes/Utility/BlueprintCameraNode.cpp
--- a/Engine/Plugins/Cameras/GameplayCameras/Source/GameplayCameras/Private/Nodes/Utility/BlueprintCameraNode.cpp
+++ b/Engine/Plugins/Cameras/GameplayCameras/Source/GameplayCameras/Private/Nodes/Utility/BlueprintCameraNode.cpp
@@ -48,10 +48,6 @@ private:
 
 	void ApplyParameterOverrides(const FCameraVariableTable& VariableTable, const FCameraContextDataTable& ContextDataTable);
 
-	template<typename CameraParameterType>
-	void ApplyParameterOverride(
-			TValueOrError<CameraParameterType*, EPropertyBagResult> ParameterOrError);
-
 private:
 
 	TObjectPtr<UBlueprintCameraNodeEvaluator> EvaluatorBlueprint;
diff --git a/Engine/Plugins/Cameras/GameplayCameras/Source/GameplayCameras/Public/Core/BaseCameraObjectReference.h b/Engine/Plugins/Cameras/GameplayCameras/Source/GameplayCameras/Public/Core/BaseCameraObjectReference.h
--- a/Engine/Plugins/Cameras/GameplayCameras/Source/GameplayCameras/Public/Core/BaseCameraObjectReference.h
+++ b/Engine/Plugins/Cameras/GameplayCameras/Source/GameplayCameras/Public/Core/BaseCameraObjectReference.h
@@ -39,10 +39,6 @@ struct FCameraObjectInterfaceParameterMetaData
 	/** Whether this parameter has an override value. */
 	UPROPERTY()
 	bool bIsOverridden = false;
-
-	/** Whether this parameter is being animated. */
-	UPROPERTY()
-	bool bIsAnimated = false;
 };
 
 USTRUCT(BlueprintType)
@@ -85,17 +81,11 @@ public:
 	GAMEPLAYCAMERAS_API bool IsParameterOverridden(const FGuid& PropertyID) const;
 	GAMEPLAYCAMERAS_API void SetParameterOverridden(const FGuid& PropertyID, bool bIsOverridden);
 
-	GAMEPLAYCAMERAS_API bool IsParameterAnimated(const FGuid& PropertyID) const;
-	GAMEPLAYCAMERAS_API void SetParameterAnimated(const FGuid& PropertyID, bool bIsAnimated);
-
 	GAMEPLAYCAMERAS_API void GetCustomCameraNodeParameters(FCustomCameraNodeParameterInfos& OutParameterInfos);
 	
 	template<typename ContainerType>
 	void GetOverriddenParameterGuids(ContainerType& OutOverriddenIDs) const;
 
-	template<typename ContainerType>
-	void GetAnimatedParameterGuids(ContainerType& OutAnimatedIDs) const;
-
 private:
 
 	const FCameraObjectInterfaceParameterMetaData* FindMetaData(const FGuid& PropertyID) const;
@@ -124,15 +114,3 @@ void FBaseCameraObjectReference::GetOverriddenParameterGuids(ContainerType& OutO
 	}
 }
 
-template<typename ContainerType>
-void FBaseCameraObjectReference::GetAnimatedParameterGuids(ContainerType& OutAnimatedIDs) const
-{
-	for (const FCameraObjectInterfaceParameterMetaData& MetaData : ParameterMetaData)
-	{
-		if (MetaData.bIsAnimated)
-		{
-			OutAnimatedIDs.Add(MetaData.ParameterGuid);
-		}
-	}
-}
-
diff --git a/Engine/Plugins/Cameras/GameplayCameras/Source/GameplayCameras/Public/Core/CameraAssetReference.h b/Engine/Plugins/Cameras/GameplayCameras/Source/GameplayCameras/Public/Core/CameraAssetReference.h
--- a/Engine/Plugins/Cameras/GameplayCameras/Source/GameplayCameras/Public/Core/CameraAssetReference.h
+++ b/Engine/Plugins/Cameras/GameplayCameras/Source/GameplayCameras/Public/Core/CameraAssetReference.h
@@ -77,7 +77,10 @@ public:
 	}
 
 	/** Applies the parameter override values to the given evaluation result. */
-	void ApplyParameterOverrides(UE::Cameras::FCameraNodeEvaluationResult& OutResult, bool bDrivenOverridesOnly);
+	void ApplyParameterOverrides(UE::Cameras::FCameraNodeEvaluationResult& OutResult, bool bDrivenOnly) const;
+
+	/** Applies the parameter override values to the given evaluation result. */
+	void ApplyParameterOverrides(const FInstancedPropertyBag& CachedParameters, UE::Cameras::FCameraNodeEvaluationResult& OutResult) const;
 
 public:
 
@@ -88,9 +91,6 @@ public:
 	GAMEPLAYCAMERAS_API bool IsParameterOverridden(const FGuid PropertyID) const;
 	GAMEPLAYCAMERAS_API void SetParameterOverridden(const FGuid PropertyID, bool bIsOverridden);
 
-	GAMEPLAYCAMERAS_API bool IsParameterAnimated(const FGuid PropertyID) const;
-	GAMEPLAYCAMERAS_API void SetParameterAnimated(const FGuid PropertyID, bool bIsAnimated);
-
 	GAMEPLAYCAMERAS_API bool NeedsRebuildParameters() const;
 	GAMEPLAYCAMERAS_API bool RebuildParametersIfNeeded();
 	GAMEPLAYCAMERAS_API void RebuildParameters();
@@ -109,9 +109,6 @@ private:
 	UPROPERTY(EditAnywhere, Category="")
 	TArray<FGuid> ParameterOverrideGuids;
 
-	UPROPERTY()
-	TArray<FGuid> ParameterAnimatedGuids;
-
 	friend class UE::Cameras::FCameraAssetReferenceDetailsCustomization;
 };
 
diff --git a/Engine/Plugins/Cameras/GameplayCameras/Source/GameplayCameras/Public/Core/CameraRigAssetReference.h b/Engine/Plugins/Cameras/GameplayCameras/Source/GameplayCameras/Public/Core/CameraRigAssetReference.h
--- a/Engine/Plugins/Cameras/GameplayCameras/Source/GameplayCameras/Public/Core/CameraRigAssetReference.h
+++ b/Engine/Plugins/Cameras/GameplayCameras/Source/GameplayCameras/Public/Core/CameraRigAssetReference.h
@@ -3,8 +3,10 @@
 #pragma once
 
 #include "BaseCameraObjectReference.h"
+#include "Core/CameraContextDataTable.h"
 #include "Core/CameraObjectInterfaceParameterDefinition.h"
 #include "Core/CameraParameters.h"
+#include "Core/CameraVariableTable.h"
 #include "UObject/ObjectPtr.h"
 
 #include "CameraRigAssetReference.generated.h"
@@ -277,8 +279,18 @@ public:
 		}
 	}
 
+	/** Applies the parameter override values to the given variable table. */
+	void ApplyParameterOverrides(UE::Cameras::FCameraVariableTable& OutVariableTable, bool bDrivenOnly) const;
+	/** Applies the parameter override values to the given variable and context data tables. */
+	void ApplyParameterOverrides(UE::Cameras::FCameraVariableTable& OutVariableTable, UE::Cameras::FCameraContextDataTable& OutContextDataTable, bool bDrivenOnly) const;
 	/** Applies the parameter override values to the given evaluation result. */
-	void ApplyParameterOverrides(UE::Cameras::FCameraNodeEvaluationResult& OutResult, bool bDrivenOverridesOnly);
+	void ApplyParameterOverrides(UE::Cameras::FCameraNodeEvaluationResult& OutResult, bool bDrivenOnly) const;
+	/** Applies the parameter override values to the given evaluation result. */
+	void ApplyParameterOverrides(const FInstancedPropertyBag& CachedParameters, UE::Cameras::FCameraNodeEvaluationResult& OutResult) const;
+
+private:
+
+	void ApplyParameterOverridesImpl(UE::Cameras::FCameraVariableTable* OutVariableTable, UE::Cameras::FCameraContextDataTable* OutContextDataTable, bool bDrivenOnly) const;
 
 public:
 
diff --git a/Engine/Plugins/Cameras/GameplayCameras/Source/GameplayCameras/Public/Core/CameraShakeAssetReference.h b/Engine/Plugins/Cameras/GameplayCameras/Source/GameplayCameras/Public/Core/CameraShakeAssetReference.h
--- a/Engine/Plugins/Cameras/GameplayCameras/Source/GameplayCameras/Public/Core/CameraShakeAssetReference.h
+++ b/Engine/Plugins/Cameras/GameplayCameras/Source/GameplayCameras/Public/Core/CameraShakeAssetReference.h
@@ -3,6 +3,8 @@
 #pragma once
 
 #include "BaseCameraObjectReference.h"
+#include "Core/CameraContextDataTable.h"
+#include "Core/CameraVariableTable.h"
 #include "UObject/ObjectPtr.h"
 
 #include "CameraShakeAssetReference.generated.h"
@@ -59,8 +61,18 @@ public:
 		}
 	}
 
+	/** Applies the parameter override values to the given variable table. */
+	void ApplyParameterOverrides(UE::Cameras::FCameraVariableTable& OutVariableTable, bool bDrivenOnly) const;
+
+	/** Applies the parameter override values to the given variable and context data tables.. */
+	void ApplyParameterOverrides(UE::Cameras::FCameraVariableTable& OutVariableTable, UE::Cameras::FCameraContextDataTable& OutContextDataTable, bool bDrivenOnly) const;
+
 	/** Applies the parameter override values to the given evaluation result. */
-	void ApplyParameterOverrides(UE::Cameras::FCameraNodeEvaluationResult& OutResult, bool bDrivenOverridesOnly);
+	void ApplyParameterOverrides(UE::Cameras::FCameraNodeEvaluationResult& OutResult, bool bDrivenOnly) const;
+
+private:
+
+	void ApplyParameterOverridesImpl(UE::Cameras::FCameraVariableTable* OutVariableTable, UE::Cameras::FCameraContextDataTable* OutContextDataTable, bool bDrivenOnly) const;
 
 public:
 
diff --git a/Engine/Plugins/Cameras/GameplayCameras/Source/GameplayCameras/Public/GameFramework/GameplayCameraComponent.h b/Engine/Plugins/Cameras/GameplayCameras/Source/GameplayCameras/Public/GameFramework/GameplayCameraComponent.h
--- a/Engine/Plugins/Cameras/GameplayCameras/Source/GameplayCameras/Public/GameFramework/GameplayCameraComponent.h
+++ b/Engine/Plugins/Cameras/GameplayCameras/Source/GameplayCameras/Public/GameFramework/GameplayCameraComponent.h
@@ -55,6 +55,13 @@ private:
 
 private:
 
+	/** Cached parameter overrides, for detecting when some of them change. */
+	UPROPERTY(Transient)
+	FInstancedPropertyBag CachedParameterOverrides;
+
+
+	// Deprecated
+
 	UPROPERTY()
 	TObjectPtr<UCameraAsset> Camera_DEPRECATED;
 };
diff --git a/Engine/Plugins/Cameras/GameplayCameras/Source/GameplayCameras/Public/GameFramework/GameplayCameraRigComponent.h b/Engine/Plugins/Cameras/GameplayCameras/Source/GameplayCameras/Public/GameFramework/GameplayCameraRigComponent.h
--- a/Engine/Plugins/Cameras/GameplayCameras/Source/GameplayCameras/Public/GameFramework/GameplayCameraRigComponent.h
+++ b/Engine/Plugins/Cameras/GameplayCameras/Source/GameplayCameras/Public/GameFramework/GameplayCameraRigComponent.h
@@ -54,9 +54,14 @@ private:
 
 private:
 
-	UPROPERTY()
+	/** Generated camera asset to run our camera rig. */
+	UPROPERTY(Transient)
 	TObjectPtr<UCameraAsset> GeneratedCameraAsset;
 
+	/** Cached parameter overrides, for detecting when some of them change. */
+	UPROPERTY(Transient)
+	FInstancedPropertyBag CachedParameterOverrides;
+
 #if WITH_EDITOR
 	bool bIsBuildingGeneratedCameraAsset = false;
 #endif
diff --git a/Engine/Plugins/Cameras/GameplayCameras/Source/GameplayCamerasEditor/Private/Sequencer/GameplayCameraComponentTrackEditor.cpp b/Engine/Plugins/Cameras/GameplayCameras/Source/GameplayCamerasEditor/Private/Sequencer/GameplayCameraComponentTrackEditor.cpp
--- a/Engine/Plugins/Cameras/GameplayCameras/Source/GameplayCamerasEditor/Private/Sequencer/GameplayCameraComponentTrackEditor.cpp
+++ b/Engine/Plugins/Cameras/GameplayCameras/Source/GameplayCamerasEditor/Private/Sequencer/GameplayCameraComponentTrackEditor.cpp
@@ -20,11 +20,17 @@
 #include "Styles/GameplayCamerasEditorStyle.h"
 #include "Tracks/MovieScenePropertyTrack.h"
 
+#if UE_VERSION_NEWER_THAN_OR_EQUAL(5,6,0)
+#include "Misc/SequencerObjectBindingHelper.h"
+#endif
+
 #define LOCTEXT_NAMESPACE "GameplayCameraComponentTrackEditor"
 
 namespace UE::Cameras::Internal
 {
 
+#if UE_VERSION_OLDER_THAN(5,7,0)
+
 // TODO: duplicated from GetKeyablePropertyPaths in ObjectBindingModel.cpp
 void GetKeyablePropertyPathsImpl(TSharedPtr<ISequencer> Sequencer, const UClass* BaseObjectClass, const UStruct* Struct, const void* StructValuePtr, FPropertyPath PropertyPath, TArray<FPropertyPath>& KeyablePropertyPaths)
 {
@@ -71,6 +77,8 @@ void GetKeyablePropertyPathsImpl(TSharedPtr<ISequencer> Sequencer, const UClass*
 	}
 }
 
+#endif  // pre-5.7.0
+
 void GetKeyablePropertyPaths(TSharedPtr<ISequencer> Sequencer, const UGameplayCameraComponentBase* CameraComponentBase, TArray<FPropertyPath>& KeyablePropertyPaths)
 {
 	// Start us off with the property path of the parameters struct, and then get all keyable property paths from there.
@@ -114,7 +122,11 @@ void GetKeyablePropertyPaths(TSharedPtr<ISequencer> Sequencer, const UGameplayCa
 	PropertyPath.AddProperty(FPropertyInfo(PropertyBagStruct->FindPropertyByName(TEXT("Value"))));
 
 	const UClass* ComponentClass = CameraComponentBase->GetClass();
+#if UE_VERSION_NEWER_THAN_OR_EQUAL(5,7,0)
+	FSequencerObjectBindingHelper::GetKeyablePropertyPaths(ComponentClass, CameraParametersMemory, CameraParametersStruct, PropertyPath, Sequencer.ToSharedRef(), KeyablePropertyPaths);
+#else
 	GetKeyablePropertyPathsImpl(Sequencer, ComponentClass, CameraParametersStruct, CameraParametersMemory, PropertyPath, KeyablePropertyPaths);
+#endif
 }
 
 class FCameraParameterTrackSetupHandler : public UE::MovieScene::TIntrusiveEventHandler<UE::MovieScene::ISequenceDataEventHandler>
