Cameras: don't freeze blend stack entries during stateless evaluations

We can't modify the evaluator tree during stateless evaluations because we save and restore its entire state.

diff --git a/Engine/Plugins/Cameras/GameplayCameras/Source/GameplayCameras/Private/Core/BlendStackCameraNode.cpp b/Engine/Plugins/Cameras/GameplayCameras/Source/GameplayCameras/Private/Core/BlendStackCameraNode.cpp
--- a/Engine/Plugins/Cameras/GameplayCameras/Source/GameplayCameras/Private/Core/BlendStackCameraNode.cpp
+++ b/Engine/Plugins/Cameras/GameplayCameras/Source/GameplayCameras/Private/Core/BlendStackCameraNode.cpp
@@ -261,7 +261,7 @@ void FBlendStackCameraNodeEvaluator::OnExecuteOperation(const FCameraOperationPa
 	}
 }
 
-void FBlendStackCameraNodeEvaluator::ResolveEntries(TArray<FResolvedEntry>& OutResolvedEntries)
+void FBlendStackCameraNodeEvaluator::ResolveEntries(const FCameraNodeEvaluationParams& Params, TArray<FResolvedEntry>& OutResolvedEntries)
 {
 	constexpr ECameraVariableTableFilter VariableTableFilter = ECameraVariableTableFilter::ChangedOnly;
 	constexpr ECameraContextDataTableFilter ContextDataTableFilter = ECameraContextDataTableFilter::ChangedOnly;
@@ -284,20 +284,25 @@ void FBlendStackCameraNodeEvaluator::ResolveEntries(TArray<FResolvedEntry>& OutR
 		if (!Entry.Flags.bIsFrozen)
 		{
 			// Check that we still have a valid context. If not, let's freeze the entry, since
-			// we won't be able to evaluate it anymore.
+			// we won't be able to evaluate it anymore... unless we are in a stateless evaluation,
+			// in which case we're not supposed to change the structure of the camera node 
+			// evaluator tree.
 			if (UNLIKELY(!CurContext.IsValid()))
 			{
-				FreezeEntry(Entry);
+				if (!Params.IsStatelessEvaluation())
+				{
+					FreezeEntry(Entry);
 
 #if UE_GAMEPLAY_CAMERAS_TRACE
-				if (Entry.Flags.bLogWarnings)
-				{
-					UE_LOG(LogCameraSystem, Warning,
+					if (Entry.Flags.bLogWarnings)
+					{
+						UE_LOG(LogCameraSystem, Warning,
 							TEXT("Freezing camera rig '%s' because its evaluation context isn't valid anymore."),
 							*GetNameSafe(Entry.CameraRig));
-					Entry.Flags.bLogWarnings = false;
-				}
+						Entry.Flags.bLogWarnings = false;
+					}
 #endif  // UE_GAMEPLAY_CAMERAS_TRACE
+				}
 
 				continue;
 			}
diff --git a/Engine/Plugins/Cameras/GameplayCameras/Source/GameplayCameras/Private/Core/PersistentBlendStackCameraNode.cpp b/Engine/Plugins/Cameras/GameplayCameras/Source/GameplayCameras/Private/Core/PersistentBlendStackCameraNode.cpp
--- a/Engine/Plugins/Cameras/GameplayCameras/Source/GameplayCameras/Private/Core/PersistentBlendStackCameraNode.cpp
+++ b/Engine/Plugins/Cameras/GameplayCameras/Source/GameplayCameras/Private/Core/PersistentBlendStackCameraNode.cpp
@@ -212,7 +212,7 @@ void FPersistentBlendStackCameraNodeEvaluator::OnRun(const FCameraNodeEvaluation
 
 	// Validate our entries and resolve evaluation context weak pointers.
 	TArray<FResolvedEntry> ResolvedEntries;
-	ResolveEntries(ResolvedEntries);
+	ResolveEntries(Params, ResolvedEntries);
 
 	// Run the stack!
 	InternalUpdate(ResolvedEntries, Params, OutResult);
@@ -230,7 +230,7 @@ void FPersistentBlendStackCameraNodeEvaluator::InternalUpdate(TArrayView<FResolv
 		FCameraRigEntry& Entry(ResolvedEntry.Entry);
 		FCameraRigEntryExtraInfo& EntryExtraInfo(EntryExtraInfos[ResolvedEntry.EntryIndex]);
 
-		if (!Entry.Flags.bIsFrozen)
+		if (!Entry.Flags.bIsFrozen && ResolvedEntry.Context.IsValid())
 		{
 			FCameraNodeEvaluationParams CurParams(Params);
 			CurParams.EvaluationContext = ResolvedEntry.Context;
diff --git a/Engine/Plugins/Cameras/GameplayCameras/Source/GameplayCameras/Private/Core/TransientBlendStackCameraNode.cpp b/Engine/Plugins/Cameras/GameplayCameras/Source/GameplayCameras/Private/Core/TransientBlendStackCameraNode.cpp
--- a/Engine/Plugins/Cameras/GameplayCameras/Source/GameplayCameras/Private/Core/TransientBlendStackCameraNode.cpp
+++ b/Engine/Plugins/Cameras/GameplayCameras/Source/GameplayCameras/Private/Core/TransientBlendStackCameraNode.cpp
@@ -213,7 +213,7 @@ void FTransientBlendStackCameraNodeEvaluator::OnRun(const FCameraNodeEvaluationP
 
 	// Validate our entries and resolve evaluation context weak pointers.
 	TArray<FResolvedEntry> ResolvedEntries;
-	ResolveEntries(ResolvedEntries);
+	ResolveEntries(Params, ResolvedEntries);
 
 	// Gather parameters to pre-blend, and evaluate blend nodes.
 	InternalPreBlendPrepare(ResolvedEntries, Params, OutResult);
@@ -245,7 +245,7 @@ void FTransientBlendStackCameraNodeEvaluator::InternalPreBlendPrepare(TArrayView
 		FCameraRigEntry& Entry(ResolvedEntry.Entry);
 		FCameraRigEntryExtraInfo& EntryExtraInfo(EntryExtraInfos[ResolvedEntry.EntryIndex]);
 
-		if (UNLIKELY(Entry.Flags.bIsFrozen))
+		if (UNLIKELY(Entry.Flags.bIsFrozen || !ResolvedEntry.Context.IsValid()))
 		{
 			continue;
 		}
@@ -326,7 +326,7 @@ void FTransientBlendStackCameraNodeEvaluator::InternalPreBlendExecute(TArrayView
 
 		FCameraNodeEvaluationResult& CurResult(Entry.Result);
 
-		if (!Entry.Flags.bIsFrozen)
+		if (!Entry.Flags.bIsFrozen && ResolvedEntry.Context.IsValid())
 		{
 			FCameraNodeEvaluationParams CurParams(Params);
 			CurParams.EvaluationContext = ResolvedEntry.Context;
@@ -362,7 +362,7 @@ void FTransientBlendStackCameraNodeEvaluator::InternalPreBlendExecute(TArrayView
 	for (FResolvedEntry& ResolvedEntry : ResolvedEntries)
 	{
 		FCameraRigEntry& Entry(ResolvedEntry.Entry);
-		if (!Entry.Flags.bIsFrozen)
+		if (!Entry.Flags.bIsFrozen && ResolvedEntry.Context.IsValid())
 		{
 			FCameraNodeEvaluationResult& CurResult(Entry.Result);
 			CurResult.VariableTable.Override(PreBlendVariableTable, ECameraVariableTableFilter::KnownOnly);
@@ -377,7 +377,7 @@ void FTransientBlendStackCameraNodeEvaluator::InternalUpdate(TArrayView<FResolve
 		FCameraRigEntry& Entry(ResolvedEntry.Entry);
 		FCameraRigEntryExtraInfo& EntryExtraInfo(EntryExtraInfos[ResolvedEntry.EntryIndex]);
 
-		if (UNLIKELY(Entry.Flags.bIsFrozen))
+		if (UNLIKELY(Entry.Flags.bIsFrozen || !ResolvedEntry.Context.IsValid()))
 		{
 			continue;
 		}
@@ -425,7 +425,7 @@ void FTransientBlendStackCameraNodeEvaluator::InternalPostBlendExecute(TArrayVie
 
 		FCameraNodeEvaluationResult& CurResult(Entry.Result);
 
-		if (!Entry.Flags.bIsFrozen)
+		if (!Entry.Flags.bIsFrozen && ResolvedEntry.Context.IsValid())
 		{
 			FCameraNodeEvaluationParams CurParams(Params);
 			CurParams.EvaluationContext = ResolvedEntry.Context;
diff --git a/Engine/Plugins/Cameras/GameplayCameras/Source/GameplayCameras/Public/Core/BlendStackCameraNode.h b/Engine/Plugins/Cameras/GameplayCameras/Source/GameplayCameras/Public/Core/BlendStackCameraNode.h
--- a/Engine/Plugins/Cameras/GameplayCameras/Source/GameplayCameras/Public/Core/BlendStackCameraNode.h
+++ b/Engine/Plugins/Cameras/GameplayCameras/Source/GameplayCameras/Public/Core/BlendStackCameraNode.h
@@ -186,7 +186,7 @@ protected:
 		bool bIsActiveEntry = false;
 	};
 
-	void ResolveEntries(TArray<FResolvedEntry>& OutResolvedEntries);
+	void ResolveEntries(const FCameraNodeEvaluationParams& Params, TArray<FResolvedEntry>& OutResolvedEntries);
 	void OnRunFinished(FCameraNodeEvaluationResult& OutResult);
 
 protected:
