Cameras: fix simple fixed-time blends not saving/restoring their current time

diff --git a/Engine/Plugins/Cameras/GameplayCameras/Source/GameplayCameras/Private/Nodes/Blends/SimpleBlendCameraNode.cpp b/Engine/Plugins/Cameras/GameplayCameras/Source/GameplayCameras/Private/Nodes/Blends/SimpleBlendCameraNode.cpp
--- a/Engine/Plugins/Cameras/GameplayCameras/Source/GameplayCameras/Private/Nodes/Blends/SimpleBlendCameraNode.cpp
+++ b/Engine/Plugins/Cameras/GameplayCameras/Source/GameplayCameras/Private/Nodes/Blends/SimpleBlendCameraNode.cpp
@@ -57,6 +57,8 @@ void FSimpleBlendCameraNodeEvaluator::OnBlendResults(const FCameraNodeBlendParam
 
 void FSimpleBlendCameraNodeEvaluator::OnSerialize(const FCameraNodeEvaluatorSerializeParams& Params, FArchive& Ar)
 {
+	Super::OnSerialize(Params, Ar);
+
 	Ar << BlendFactor;
 	Ar << bIsBlendFinished;
 	Ar << bReverse;
@@ -105,6 +107,13 @@ void FSimpleFixedTimeBlendCameraNodeEvaluator::OnRun(const FCameraNodeEvaluation
 	FSimpleBlendCameraNodeEvaluator::OnRun(Params, OutResult);
 }
 
+void FSimpleFixedTimeBlendCameraNodeEvaluator::OnSerialize(const FCameraNodeEvaluatorSerializeParams& Params, FArchive& Ar)
+{
+	Super::OnSerialize(Params, Ar);
+
+	Ar << CurrentTime;
+}
+
 bool FSimpleFixedTimeBlendCameraNodeEvaluator::OnInitializeFromInterruption(const FCameraNodeBlendInterruptionParams& Params)
 {
 	// If we are interrupting a fixed-time blend, adjust our own time to the complementary time ratio.
diff --git a/Engine/Plugins/Cameras/GameplayCameras/Source/GameplayCameras/Public/Nodes/Blends/SimpleBlendCameraNode.h b/Engine/Plugins/Cameras/GameplayCameras/Source/GameplayCameras/Public/Nodes/Blends/SimpleBlendCameraNode.h
--- a/Engine/Plugins/Cameras/GameplayCameras/Source/GameplayCameras/Public/Nodes/Blends/SimpleBlendCameraNode.h
+++ b/Engine/Plugins/Cameras/GameplayCameras/Source/GameplayCameras/Public/Nodes/Blends/SimpleBlendCameraNode.h
@@ -100,6 +100,7 @@ protected:
 	// FBlendCameraNodeEvaluator interface.
 	GAMEPLAYCAMERAS_API virtual void OnInitialize(const FCameraNodeEvaluatorInitializeParams& Params, FCameraNodeEvaluationResult& OutResult) override;
 	GAMEPLAYCAMERAS_API virtual void OnRun(const FCameraNodeEvaluationParams& Params, FCameraNodeEvaluationResult& OutResult) override;
+	GAMEPLAYCAMERAS_API virtual void OnSerialize(const FCameraNodeEvaluatorSerializeParams& Params, FArchive& Ar) override;
 	GAMEPLAYCAMERAS_API virtual bool OnInitializeFromInterruption(const FCameraNodeBlendInterruptionParams& Params) override;
 
 	float GetTimeFactor() const;
