Cameras: add ResolveCameraRigProxy method to Blueprint camera director

diff --git a/Engine/Plugins/Cameras/GameplayCameras/Source/GameplayCameras/Private/Directors/BlueprintCameraDirector.cpp b/Engine/Plugins/Cameras/GameplayCameras/Source/GameplayCameras/Private/Directors/BlueprintCameraDirector.cpp
--- a/Engine/Plugins/Cameras/GameplayCameras/Source/GameplayCameras/Private/Directors/BlueprintCameraDirector.cpp
+++ b/Engine/Plugins/Cameras/GameplayCameras/Source/GameplayCameras/Private/Directors/BlueprintCameraDirector.cpp
@@ -30,6 +30,10 @@ class FBlueprintCameraDirectorEvaluator : public FCameraDirectorEvaluator
 {
 	UE_DECLARE_CAMERA_DIRECTOR_EVALUATOR(GAMEPLAYCAMERAS_API, FBlueprintCameraDirectorEvaluator)
 
+public:
+	
+	~FBlueprintCameraDirectorEvaluator();
+
 protected:
 
 	virtual void OnInitialize(const FCameraDirectorInitializeParams& Params) override;
@@ -49,6 +53,14 @@ private:
 
 UE_DEFINE_CAMERA_DIRECTOR_EVALUATOR(FBlueprintCameraDirectorEvaluator)
 
+FBlueprintCameraDirectorEvaluator::~FBlueprintCameraDirectorEvaluator()
+{
+	if (EvaluatorBlueprint)
+	{
+		EvaluatorBlueprint->NativeAbandonCameraDirector();
+	}
+}
+
 void FBlueprintCameraDirectorEvaluator::OnInitialize(const FCameraDirectorInitializeParams& Params)
 {
 	const UBlueprintCameraDirector* Blueprint = GetCameraDirectorAs<UBlueprintCameraDirector>();
@@ -67,7 +79,7 @@ void FBlueprintCameraDirectorEvaluator::OnInitialize(const FCameraDirectorInitia
 	{
 		UObject* Outer = Params.OwnerContext->GetOwner();
 		EvaluatorBlueprint = NewObject<UBlueprintCameraDirectorEvaluator>(Outer, Blueprint->CameraDirectorEvaluatorClass);
-		EvaluatorBlueprint->NativeInitializeCameraDirector(Params);
+		EvaluatorBlueprint->NativeInitializeCameraDirector(this, Params);
 	}
 	else
 	{
@@ -252,6 +264,23 @@ void UBlueprintCameraDirectorEvaluator::ActivateCameraRigViaProxy(UCameraRigProx
 	EvaluationResult.Requests.Add(Request);
 }
 
+UCameraRigAsset* UBlueprintCameraDirectorEvaluator::ResolveCameraRigProxy(const UCameraRigProxyAsset* CameraRigProxy) const
+{
+	if (ensure(OwningDirectorEvaluator))
+	{
+		const UBlueprintCameraDirector* BlueprintCameraDirector = OwningDirectorEvaluator->GetCameraDirectorAs<UBlueprintCameraDirector>();
+		if (ensure(BlueprintCameraDirector))
+		{
+			const FCameraRigProxyRedirectTable& ProxyTable = BlueprintCameraDirector->CameraRigProxyRedirectTable;
+
+			FCameraRigProxyResolveParams ResolveParams;
+			ResolveParams.CameraRigProxy = CameraRigProxy;
+			return ProxyTable.ResolveProxy(ResolveParams);
+		}
+	}
+	return nullptr;
+}
+
 AActor* UBlueprintCameraDirectorEvaluator::FindEvaluationContextOwnerActor(TSubclassOf<AActor> ActorClass) const
 {
 	if (EvaluationContext)
@@ -336,13 +365,20 @@ UWorld* UBlueprintCameraDirectorEvaluator::GetWorld() const
 	return nullptr;
 }
 
-void UBlueprintCameraDirectorEvaluator::NativeInitializeCameraDirector(const UE::Cameras::FCameraDirectorInitializeParams& Params)
+void UBlueprintCameraDirectorEvaluator::NativeInitializeCameraDirector(UE::Cameras::FCameraDirectorEvaluator* InOwningDirectorEvaluator, const UE::Cameras::FCameraDirectorInitializeParams& Params)
 {
 	using namespace UE::Cameras;
 
+	ensure(OwningDirectorEvaluator == nullptr && InOwningDirectorEvaluator != nullptr);
+	OwningDirectorEvaluator = InOwningDirectorEvaluator;
 	EvaluationContext = Params.OwnerContext;
 }
 
+void UBlueprintCameraDirectorEvaluator::NativeAbandonCameraDirector()
+{
+	OwningDirectorEvaluator = nullptr;
+}
+
 void UBlueprintCameraDirectorEvaluator::NativeActivateCameraDirector(const UE::Cameras::FCameraDirectorActivateParams& Params)
 {
 	using namespace UE::Cameras;
diff --git a/Engine/Plugins/Cameras/GameplayCameras/Source/GameplayCameras/Public/Directors/BlueprintCameraDirector.h b/Engine/Plugins/Cameras/GameplayCameras/Source/GameplayCameras/Public/Directors/BlueprintCameraDirector.h
--- a/Engine/Plugins/Cameras/GameplayCameras/Source/GameplayCameras/Public/Directors/BlueprintCameraDirector.h
+++ b/Engine/Plugins/Cameras/GameplayCameras/Source/GameplayCameras/Public/Directors/BlueprintCameraDirector.h
@@ -133,7 +133,13 @@ public:
 	 * via the proxy table of the Blueprint camera director.
 	 */
 	UFUNCTION(BlueprintCallable, Category="Camera Director|Evaluation")
-	void ActivateCameraRigViaProxy(UCameraRigProxyAsset* CameraRigProxy, bool bForceNewInstance = false);
+	GAMEPLAYCAMERAS_API void ActivateCameraRigViaProxy(UCameraRigProxyAsset* CameraRigProxy, bool bForceNewInstance = false);
+
+public:
+
+	/** Resolves the camera rig proxy using the camera director's proxy table. */
+	UFUNCTION(BlueprintPure, Category="Camera Director|Evaluation", meta=(HideSelfPin=true))
+	GAMEPLAYCAMERAS_API UCameraRigAsset* ResolveCameraRigProxy(const UCameraRigProxyAsset* CameraRigProxy) const;
 
 public:
 
@@ -172,7 +178,13 @@ public:
 
 public:
 
-	void NativeInitializeCameraDirector(const UE::Cameras::FCameraDirectorInitializeParams& Params);
+	/** Initialize this camera director evaluator. */
+	void NativeInitializeCameraDirector(UE::Cameras::FCameraDirectorEvaluator* InOwningDirectorEvaluator, const UE::Cameras::FCameraDirectorInitializeParams& Params);
+
+	/** Abandon this camera director evalutor. */
+	void NativeAbandonCameraDirector();
+
+public:
 
 	/** Native wrapper for ActivateCameraDirector. */
 	void NativeActivateCameraDirector(const UE::Cameras::FCameraDirectorActivateParams& Params);
@@ -189,6 +201,9 @@ public:
 
 private:
 
+	/** The owning camera director evaluator. */
+	UE::Cameras::FCameraDirectorEvaluator* OwningDirectorEvaluator = nullptr;
+
 	/** The current evaluation result. */
 	FCameraDirectorEvaluationResult EvaluationResult;
 
