Cameras: add some DLL exports

diff --git a/Engine/Plugins/Cameras/GameplayCameras/Source/GameplayCameras/Public/Core/RootCameraNode.h b/Engine/Plugins/Cameras/GameplayCameras/Source/GameplayCameras/Public/Core/RootCameraNode.h
--- a/Engine/Plugins/Cameras/GameplayCameras/Source/GameplayCameras/Public/Core/RootCameraNode.h
+++ b/Engine/Plugins/Cameras/GameplayCameras/Source/GameplayCameras/Public/Core/RootCameraNode.h
@@ -116,37 +116,37 @@ public:
 	 * Activates a camera rig.
 	 * What it means to activate a camera rig may differ depending on the layer it runs on.
 	 */
-	FCameraRigInstanceID ActivateCameraRig(const FActivateCameraRigParams& Params);
+	GAMEPLAYCAMERAS_API FCameraRigInstanceID ActivateCameraRig(const FActivateCameraRigParams& Params);
 
 	/** 
 	 * Deactivates a camera rig. 
 	 * What it means to deactivate a camera rig may differ depending on the layer it runs on.
 	 */
-	void DeactivateCameraRig(const FDeactivateCameraRigParams& Params);
+	GAMEPLAYCAMERAS_API void DeactivateCameraRig(const FDeactivateCameraRigParams& Params);
 
 	/** Deactivates all camera rigs with the given evaluation context. */
-	void DeactivateAllCameraRigs(TSharedPtr<const FCameraEvaluationContext> InContext, bool bImmediately);
+	GAMEPLAYCAMERAS_API void DeactivateAllCameraRigs(TSharedPtr<const FCameraEvaluationContext> InContext, bool bImmediately);
 
 	/** Execute a request to activate or deactivate a camera rig. */
-	void ExecuteCameraDirectorRequest(const FCameraRigActivationDeactivationRequest& Request);
+	GAMEPLAYCAMERAS_API void ExecuteCameraDirectorRequest(const FCameraRigActivationDeactivationRequest& Request);
 
 	/** Gets information about the active camera rig in the main layer. */
-	void GetActiveCameraRigInfo(FCameraRigEvaluationInfo& OutCameraRigInfo) const;
+	GAMEPLAYCAMERAS_API void GetActiveCameraRigInfo(FCameraRigEvaluationInfo& OutCameraRigInfo) const;
 
 	/** Gets information about a specified camera rig. */
-	void GetCameraRigInfo(const FCameraRigInstanceID InstanceID, FCameraRigEvaluationInfo& OutCameraRigInfo) const;
+	GAMEPLAYCAMERAS_API void GetCameraRigInfo(const FCameraRigInstanceID InstanceID, FCameraRigEvaluationInfo& OutCameraRigInfo) const;
 
 	/** Gets whether any camera rig is running in the main layer. */
-	bool HasAnyActiveCameraRig() const;
+	GAMEPLAYCAMERAS_API bool HasAnyActiveCameraRig() const;
 
 	/** Gets whether any camera rig is running with the given context in the main layer. */
-	bool HasAnyRunningCameraRig(TSharedPtr<const FCameraEvaluationContext> InContext) const;
+	GAMEPLAYCAMERAS_API bool HasAnyRunningCameraRig(TSharedPtr<const FCameraEvaluationContext> InContext) const;
 
 	/** Gets the variable table containing the blended camera object interfaces parameters. */
-	const FCameraVariableTable* GetBlendedParameters() const;
+	GAMEPLAYCAMERAS_API const FCameraVariableTable* GetBlendedParameters() const;
 
 	/** Returns the evaluation result without the contribution of the visual layer. */
-	const FCameraNodeEvaluationResult& GetPreVisualLayerResult() const;
+	GAMEPLAYCAMERAS_API const FCameraNodeEvaluationResult& GetPreVisualLayerResult() const;
 
 	/**
 	 * Builds the hierarchy of the system for a given single camera rig.
@@ -154,14 +154,14 @@ public:
 	 * should only have the nodes of the given camera rig (i.e. it shouldn't have nodes of
 	 * other currently active camera rigs).
 	 */
-	void BuildSingleCameraRigHierarchy(const FSingleCameraRigHierarchyBuildParams& Params, FCameraNodeEvaluatorHierarchy& OutHierarchy);
+	GAMEPLAYCAMERAS_API void BuildSingleCameraRigHierarchy(const FSingleCameraRigHierarchyBuildParams& Params, FCameraNodeEvaluatorHierarchy& OutHierarchy);
 
 	/**
 	 * Evaluates a single camera rig.
 	 * This is expected to run all layers as usual, except for the main layer which should
 	 * only run the given camera rig instead.
 	 */
-	void RunSingleCameraRig(const FSingleCameraRigEvaluationParams& Params, FCameraNodeEvaluationResult& OutResult);
+	GAMEPLAYCAMERAS_API void RunSingleCameraRig(const FSingleCameraRigEvaluationParams& Params, FCameraNodeEvaluationResult& OutResult);
 
 	/** Gets the delegate for camera rig events. */
 	FOnRootCameraNodeCameraRigEvent& OnCameraRigEvent() { return OnCameraRigEventDelegate; }
@@ -169,7 +169,7 @@ public:
 protected:
 
 	// FCameraNodeEvaluator interface.
-	virtual void OnInitialize(const FCameraNodeEvaluatorInitializeParams& Params, FCameraNodeEvaluationResult& OutResult) override;
+	GAMEPLAYCAMERAS_API virtual void OnInitialize(const FCameraNodeEvaluatorInitializeParams& Params, FCameraNodeEvaluationResult& OutResult) override;
 
 protected:
 
@@ -202,9 +202,9 @@ protected:
 
 protected:
 
-	void SetPreVisualLayerResult(const FCameraNodeEvaluationResult& InResult);
+	GAMEPLAYCAMERAS_API void SetPreVisualLayerResult(const FCameraNodeEvaluationResult& InResult);
 
-	void BroadcastCameraRigEvent(const FRootCameraNodeCameraRigEvent& InEvent) const;
+	GAMEPLAYCAMERAS_API void BroadcastCameraRigEvent(const FRootCameraNodeCameraRigEvent& InEvent) const;
 
 private:
 
diff --git a/Engine/Plugins/Cameras/GameplayCameras/Source/GameplayCameras/Public/Directors/BlueprintCameraDirector.h b/Engine/Plugins/Cameras/GameplayCameras/Source/GameplayCameras/Public/Directors/BlueprintCameraDirector.h
--- a/Engine/Plugins/Cameras/GameplayCameras/Source/GameplayCameras/Public/Directors/BlueprintCameraDirector.h
+++ b/Engine/Plugins/Cameras/GameplayCameras/Source/GameplayCameras/Public/Directors/BlueprintCameraDirector.h
@@ -88,45 +88,45 @@ PRAGMA_ENABLE_DEPRECATION_WARNINGS
 public:
 
 	UFUNCTION(BlueprintCallable, BlueprintImplementableEvent, Category="Evaluation")
-	FName AddChildEvaluationContext(UObject* ChildEvaluationContextOwner);
+	GAMEPLAYCAMERAS_API FName AddChildEvaluationContext(UObject* ChildEvaluationContextOwner);
 
 	UFUNCTION(BlueprintCallable, BlueprintImplementableEvent, Category="Evaluation")
-	bool RemoveChildEvaluationContext(UObject* ChildEvaluationContextOwner, FName ChildSlotName);
+	GAMEPLAYCAMERAS_API bool RemoveChildEvaluationContext(UObject* ChildEvaluationContextOwner, FName ChildSlotName);
 
 	UFUNCTION(BlueprintCallable, Category="Evaluation")
-	bool RunChildCameraDirector(float DeltaTime, FName ChildSlotName);
+	GAMEPLAYCAMERAS_API bool RunChildCameraDirector(float DeltaTime, FName ChildSlotName);
 
 public:
 
 	/** Activates the given camera rig prefab in the base layer. */
 	UFUNCTION(BlueprintCallable, Category="Camera Director|Activation")
-	void ActivatePersistentBaseCameraRig(UCameraRigAsset* CameraRigPrefab);
+	GAMEPLAYCAMERAS_API void ActivatePersistentBaseCameraRig(UCameraRigAsset* CameraRigPrefab);
 
 	/** Activates the given camera rig prefab in the global layer. */
 	UFUNCTION(BlueprintCallable, Category="Camera Director|Activation")
-	void ActivatePersistentGlobalCameraRig(UCameraRigAsset* CameraRigPrefab);
+	GAMEPLAYCAMERAS_API void ActivatePersistentGlobalCameraRig(UCameraRigAsset* CameraRigPrefab);
 
 	/** Activates the given camera rig prefab in the visual layer. */
 	UFUNCTION(BlueprintCallable, Category="Camera Director|Activation")
-	void ActivatePersistentVisualCameraRig(UCameraRigAsset* CameraRigPrefab);
+	GAMEPLAYCAMERAS_API void ActivatePersistentVisualCameraRig(UCameraRigAsset* CameraRigPrefab);
 
 	/** Deactivates the given camera rig prefab in the base layer. */
 	UFUNCTION(BlueprintCallable, Category="Camera Director|Activation")
-	void DeactivatePersistentBaseCameraRig(UCameraRigAsset* CameraRigPrefab);
+	GAMEPLAYCAMERAS_API void DeactivatePersistentBaseCameraRig(UCameraRigAsset* CameraRigPrefab);
 
 	/** Deactivates the given camera rig prefab in the global layer. */
 	UFUNCTION(BlueprintCallable, Category="Camera Director|Activation")
-	void DeactivatePersistentGlobalCameraRig(UCameraRigAsset* CameraRigPrefab);
+	GAMEPLAYCAMERAS_API void DeactivatePersistentGlobalCameraRig(UCameraRigAsset* CameraRigPrefab);
 
 	/** Deactivates the given camera rig prefab in the visual layer. */
 	UFUNCTION(BlueprintCallable, Category="Camera Director|Activation")
-	void DeactivatePersistentVisualCameraRig(UCameraRigAsset* CameraRigPrefab);
+	GAMEPLAYCAMERAS_API void DeactivatePersistentVisualCameraRig(UCameraRigAsset* CameraRigPrefab);
 
 public:
 
 	/** Specifies a camera rig to be active this frame. */
 	UFUNCTION(BlueprintCallable, Category="Camera Director|Evaluation")
-	void ActivateCameraRig(UCameraRigAsset* CameraRig, bool bForceNewInstance = false);
+	GAMEPLAYCAMERAS_API void ActivateCameraRig(UCameraRigAsset* CameraRig, bool bForceNewInstance = false);
 
 	/**
 	 * Specifies a camera rig to be active this frame, via a proxy which is later resolved
@@ -149,15 +149,15 @@ public:
 	 * UGameplayCameraComponent) or an actor itself.
 	 */
 	UFUNCTION(BlueprintPure, Category="Camera Director|Evaluation", meta=(DeterminesOutputType="ActorClass"))
-	AActor* FindEvaluationContextOwnerActor(TSubclassOf<AActor> ActorClass) const;
+	GAMEPLAYCAMERAS_API AActor* FindEvaluationContextOwnerActor(TSubclassOf<AActor> ActorClass) const;
 
 	/** Gets the shared evaluation context data. */
 	UFUNCTION(BlueprintPure, Category="Camera Director|Evaluation", meta=(DisplayName="Get Shared Camera Data"))
-	FBlueprintCameraEvaluationDataRef GetInitialContextResult() const;
+	GAMEPLAYCAMERAS_API FBlueprintCameraEvaluationDataRef GetInitialContextResult() const;
 
 	/** Gets the evaluation context data for a sub-set of camera rigs. */	
 	UFUNCTION(BlueprintPure, Category="Camera Director|Evaluation", meta=(DisplayName="Get Conditional Camera Data"))
-	FBlueprintCameraEvaluationDataRef GetConditionalContextResult(ECameraEvaluationDataCondition Condition) const;
+	GAMEPLAYCAMERAS_API FBlueprintCameraEvaluationDataRef GetConditionalContextResult(ECameraEvaluationDataCondition Condition) const;
 
 public:
 
