Cameras: move persistent rig activation functions onto the camera components and the camera manager

These Blueprint functions were previously global, affecting whatever camera system host was currently considered "active". This turns out to be unclear. Now the user activates persistent rigs exactly on whatever they intend to.

diff --git a/Engine/Plugins/Cameras/GameplayCameras/Source/GameplayCameras/Private/GameFramework/GameplayCameraComponentBase.cpp b/Engine/Plugins/Cameras/GameplayCameras/Source/GameplayCameras/Private/GameFramework/GameplayCameraComponentBase.cpp
--- a/Engine/Plugins/Cameras/GameplayCameras/Source/GameplayCameras/Private/GameFramework/GameplayCameraComponentBase.cpp
+++ b/Engine/Plugins/Cameras/GameplayCameras/Source/GameplayCameras/Private/GameFramework/GameplayCameraComponentBase.cpp
@@ -56,6 +56,7 @@ void UGameplayCameraComponentBase::BeginDestroy()
 	DestroyCameraSystem();
 
 	EvaluationContext = nullptr;
+	NullContext = nullptr;
 
 	Super::BeginDestroy();
 }
@@ -70,6 +71,10 @@ void UGameplayCameraComponentBase::AddReferencedObjects(UObject* InThis, FRefere
 	{
 		This->EvaluationContext->AddReferencedObjects(Collector);
 	}
+	if (This->NullContext.IsValid())
+	{
+		This->NullContext->AddReferencedObjects(Collector);
+	}
 }
 
 TSharedPtr<const UE::Cameras::FCameraEvaluationContext> UGameplayCameraComponentBase::GetEvaluationContext() const
@@ -194,6 +199,7 @@ void UGameplayCameraComponentBase::DeactivateCameraEvaluationContext(bool bImmed
 		}
 
 		EvaluationContext = nullptr;
+		NullContext = nullptr;
 	}
 }
 
@@ -364,6 +370,16 @@ void UGameplayCameraComponentBase::TryCreateCameraEvaluationContext(APlayerContr
 			UpdateCameraEvaluationContext(true);
 		}
 	}
+
+	if (!NullContext.IsValid())
+	{
+		NullContext = MakeShared<FCameraEvaluationContext>();
+
+		FCameraEvaluationContextInitializeParams InitParams;
+		InitParams.Owner = this;
+		InitParams.PlayerController = PlayerController;
+		NullContext->Initialize(InitParams);
+	}
 }
 
 #define UE_PRIVATE_GAMEPLAY_CAMERA_COMPONENT_VALIDATE_EVALUATION_CONTEXT(ErrorMsg, ErrorResult)\
@@ -372,7 +388,7 @@ void UGameplayCameraComponentBase::TryCreateCameraEvaluationContext(APlayerContr
 	{\
 		FFrame::KismetExecutionMessage(\
 				*FString::Format(\
-					TEXT(#ErrorResult " on Gameplay Camera component '{0}': it isn't active."),\
+					TEXT(#ErrorMsg " on Gameplay Camera component '{0}': it isn't active."),\
 					{ *GetNameSafe(this) }),\
 				ELogVerbosity::Error);\
 		return ErrorResult;\
@@ -394,6 +410,37 @@ FBlueprintCameraEvaluationDataRef UGameplayCameraComponentBase::GetConditionalRe
 
 #undef UE_PRIVATE_GAMEPLAY_CAMERA_COMPONENT_VALIDATE_EVALUATION_CONTEXT
 
+#define UE_PRIVATE_GAMEPLAY_CAMERA_COMPONENT_ACTIVATE_NON_MAIN_RIG(ErrorMsg, RigLayer)\
+	if (NullContext.IsValid())\
+	{\
+		IGameplayCameraSystemHost::ActivateCameraRig(CameraRig, NullContext, ECameraRigLayer::RigLayer);\
+	}\
+	else\
+	{\
+		FFrame::KismetExecutionMessage(\
+				*FString::Format(\
+					TEXT(#ErrorMsg " on Gameplay Camera component '{0}': it isn't active."),\
+					{ *GetNameSafe(this) }),\
+				ELogVerbosity::Error);\
+	}
+
+void UGameplayCameraComponentBase::ActivatePersistentBaseCameraRig(UCameraRigAsset* CameraRig)
+{
+	UE_PRIVATE_GAMEPLAY_CAMERA_COMPONENT_ACTIVATE_NON_MAIN_RIG("Can't activate base camera rig", Base);
+}
+
+void UGameplayCameraComponentBase::ActivatePersistentGlobalCameraRig(UCameraRigAsset* CameraRig)
+{
+	UE_PRIVATE_GAMEPLAY_CAMERA_COMPONENT_ACTIVATE_NON_MAIN_RIG("Can't activate global camera rig", Global);
+}
+
+void UGameplayCameraComponentBase::ActivatePersistentVisualCameraRig(UCameraRigAsset* CameraRig)
+{
+	UE_PRIVATE_GAMEPLAY_CAMERA_COMPONENT_ACTIVATE_NON_MAIN_RIG("Can't activate visual camera rig", Visual);
+}
+
+#undef UE_PRIVATE_GAMEPLAY_CAMERA_COMPONENT_ACTIVATE_NON_MAIN_RIG
+
 FRotator UGameplayCameraComponentBase::GetEvaluatedCameraRotation() const
 {
 	using namespace UE::Cameras;
@@ -545,6 +592,7 @@ void UGameplayCameraComponentBase::CheckPendingDeactivation()
 	{
 		TeardownCameraSystemHost();
 		EvaluationContext = nullptr;
+		NullContext = nullptr;
 
 		// Only call the base class method here: we just want to finish deactivating ourselves
 		// by stopping ticking.
@@ -664,7 +712,7 @@ void UGameplayCameraComponentBase::RecreateEditorWorldCameraEvaluationContext()
 		ensure(AllContexts.Num() == 1 && AllContexts[0] == EvaluationContext);
 	}
 
-	// Teardown and rebuild the evaluation context.
+	// Teardown and rebuild the main evaluation context.
 	if (EvaluationContext)
 	{
 		if (CameraSystemEvaluator)
@@ -770,7 +818,7 @@ void UGameplayCameraComponentBase::AutoManageEditorPreviewEvaluator()
 	}
 	
 	const bool bCanRun = CanRunCameraSystem();
-	if (bCanRun && !(CameraSystemEvaluator && EvaluationContext))
+	if (bCanRun && !(CameraSystemEvaluator && EvaluationContext && NullContext))
 	{
 		// We want to run the camera logic in the editor but we haven't set things up for that.
 		// Let's create the preview evaluator and the evaluation context.
@@ -784,7 +832,7 @@ void UGameplayCameraComponentBase::AutoManageEditorPreviewEvaluator()
 
 		// OutputCameraComponent will be updated on the next tick.
 	}
-	else if (!bCanRun && (CameraSystemEvaluator || EvaluationContext))
+	else if (!bCanRun && (CameraSystemEvaluator || EvaluationContext || NullContext))
 	{
 		// We don't want to run the camera logic in the editor anymore. Let's tear things down.
 		DeactivateCameraEvaluationContext(true);
diff --git a/Engine/Plugins/Cameras/GameplayCameras/Source/GameplayCameras/Private/GameFramework/GameplayCamerasPlayerCameraManager.cpp b/Engine/Plugins/Cameras/GameplayCameras/Source/GameplayCameras/Private/GameFramework/GameplayCamerasPlayerCameraManager.cpp
--- a/Engine/Plugins/Cameras/GameplayCameras/Source/GameplayCameras/Private/GameFramework/GameplayCamerasPlayerCameraManager.cpp
+++ b/Engine/Plugins/Cameras/GameplayCameras/Source/GameplayCameras/Private/GameFramework/GameplayCamerasPlayerCameraManager.cpp
@@ -129,6 +129,34 @@ void AGameplayCamerasPlayerCameraManager::ReleasePlayerController()
 	PCOwner = nullptr;
 }
 
+void AGameplayCamerasPlayerCameraManager::ActivatePersistentBaseCameraRig(UCameraRigAsset* CameraRig)
+{
+	EnsureNullContext();
+	IGameplayCameraSystemHost::ActivateCameraRig(CameraRig, NullContext, ECameraRigLayer::Base);
+}
+
+void AGameplayCamerasPlayerCameraManager::ActivatePersistentGlobalCameraRig(UCameraRigAsset* CameraRig)
+{
+	EnsureNullContext();
+	IGameplayCameraSystemHost::ActivateCameraRig(CameraRig, NullContext, ECameraRigLayer::Global);
+}
+
+void AGameplayCamerasPlayerCameraManager::ActivatePersistentVisualCameraRig(UCameraRigAsset* CameraRig)
+{
+	EnsureNullContext();
+	IGameplayCameraSystemHost::ActivateCameraRig(CameraRig, NullContext, ECameraRigLayer::Visual);
+}
+
+void AGameplayCamerasPlayerCameraManager::EnsureNullContext()
+{
+	using namespace UE::Cameras;
+
+	if (!NullContext)
+	{
+		NullContext = MakeShared<FCameraEvaluationContext>();
+	}
+}
+
 FCameraRigInstanceID AGameplayCamerasPlayerCameraManager::StartGlobalCameraModifierRig(const UCameraRigAsset* CameraRig, int32 OrderKey)
 {
 	using namespace UE::Cameras;
diff --git a/Engine/Plugins/Cameras/GameplayCameras/Source/GameplayCameras/Private/GameFramework/IGameplayCameraSystemHost.cpp b/Engine/Plugins/Cameras/GameplayCameras/Source/GameplayCameras/Private/GameFramework/IGameplayCameraSystemHost.cpp
--- a/Engine/Plugins/Cameras/GameplayCameras/Source/GameplayCameras/Private/GameFramework/IGameplayCameraSystemHost.cpp
+++ b/Engine/Plugins/Cameras/GameplayCameras/Source/GameplayCameras/Private/GameFramework/IGameplayCameraSystemHost.cpp
@@ -5,7 +5,11 @@
 #include "Camera/PlayerCameraManager.h"
 #include "Core/CameraEvaluationContext.h"
 #include "Core/CameraEvaluationContextStack.h"
+#include "Core/CameraRigAsset.h"
+#include "Core/CameraRigInstanceID.h"
 #include "Core/CameraSystemEvaluator.h"
+#include "Core/CameraSystemEvaluator.h"
+#include "Core/RootCameraNode.h"
 #include "Debug/CameraSystemDebugRegistry.h"
 #include "Debug/DebugDrawService.h"
 #include "Engine/Canvas.h"
@@ -56,6 +60,42 @@ TScriptInterface<IGameplayCameraSystemHost> IGameplayCameraSystemHost::GetAsScri
 	return Result;
 }
 
+void IGameplayCameraSystemHost::ActivateCameraRig(UCameraRigAsset* CameraRig, TSharedPtr<UE::Cameras::FCameraEvaluationContext> EvaluationContext, ECameraRigLayer EvaluationLayer)
+{
+	using namespace UE::Cameras;
+
+	if (!CameraRig)
+	{
+		FFrame::KismetExecutionMessage(
+				*FString::Printf(TEXT("Can't activate camera rig on '%s': no camera rig given!"), *GetNameSafe(GetAsObject())),
+				ELogVerbosity::Error);
+		return;
+	}
+	
+	if (!EvaluationContext)
+	{
+		FFrame::KismetExecutionMessage(
+				*FString::Printf(
+					TEXT("Can't activate camera rig '%s' on '%s': invalid evaluation context given!"), 
+					*GetNameSafe(CameraRig), *GetNameSafe(GetAsObject())),
+				ELogVerbosity::Error);
+		return;
+	}
+
+	EnsureCameraSystemInitialized();
+
+	if (ensure(CameraSystemEvaluator))
+	{
+		FActivateCameraRigParams Params;
+		Params.CameraRig = CameraRig;
+		Params.EvaluationContext = EvaluationContext;
+		Params.Layer = EvaluationLayer;
+
+		FRootCameraNodeEvaluator* RootNodeEvaluator = CameraSystemEvaluator->GetRootNodeEvaluator();
+		RootNodeEvaluator->ActivateCameraRig(Params);
+	}
+}
+
 IGameplayCameraSystemHost* IGameplayCameraSystemHost::FindActiveHost(APlayerController* PlayerController)
 {
 	if (PlayerController)
diff --git a/Engine/Plugins/Cameras/GameplayCameras/Source/GameplayCameras/Public/GameFramework/ActivateCameraRigFunctions.h b/Engine/Plugins/Cameras/GameplayCameras/Source/GameplayCameras/Public/GameFramework/ActivateCameraRigFunctions.h
--- a/Engine/Plugins/Cameras/GameplayCameras/Source/GameplayCameras/Public/GameFramework/ActivateCameraRigFunctions.h
+++ b/Engine/Plugins/Cameras/GameplayCameras/Source/GameplayCameras/Public/GameFramework/ActivateCameraRigFunctions.h
@@ -17,6 +17,8 @@ enum class ECameraRigLayer : uint8;
  * These camera rigs run with a global, shared evaluation context that doesn't provide any
  * meaningful initial result. They are activated on the camera system found to be running
  * on the given player controller.
+ *
+ * Deprecated in 5.7.0
  */
 UCLASS()
 class UActivateCameraRigFunctions : public UBlueprintFunctionLibrary
@@ -26,15 +28,15 @@ class UActivateCameraRigFunctions : public UBlueprintFunctionLibrary
 public:
 
 	/** Activates the given camera rig prefab in the base layer. */
-	UFUNCTION(BlueprintCallable, Category="Camera", meta=(WorldContext="WorldContextObject"))
+	UFUNCTION(BlueprintCallable, Category="Camera", meta=(WorldContext="WorldContextObject", DeprecatedFunction, DeprecationMessage="Use the similar functions on the GameplayCamera components or camera manager"))
 	static void ActivatePersistentBaseCameraRig(UObject* WorldContextObject, APlayerController* PlayerController, UCameraRigAsset* CameraRig);
 
 	/** Activates the given camera rig prefab in the global layer. */
-	UFUNCTION(BlueprintCallable, Category="Camera", meta = (WorldContext = "WorldContextObject"))
+	UFUNCTION(BlueprintCallable, Category="Camera", meta = (WorldContext = "WorldContextObject", DeprecatedFunction, DeprecationMessage="Use the similar functions on the GameplayCamera components or camera manager"))
 	static void ActivatePersistentGlobalCameraRig(UObject* WorldContextObject, APlayerController* PlayerController, UCameraRigAsset* CameraRig);
 
 	/** Activates the given camera rig prefab in the visual layer. */
-	UFUNCTION(BlueprintCallable, Category="Camera", meta = (WorldContext = "WorldContextObject"))
+	UFUNCTION(BlueprintCallable, Category="Camera", meta = (WorldContext = "WorldContextObject", DeprecatedFunction, DeprecationMessage="Use the similar functions on the GameplayCamera components or camera manager"))
 	static void ActivatePersistentVisualCameraRig(UObject* WorldContextObject, APlayerController* PlayerController, UCameraRigAsset* CameraRig);
 
 private:
diff --git a/Engine/Plugins/Cameras/GameplayCameras/Source/GameplayCameras/Public/GameFramework/GameplayCameraComponentBase.h b/Engine/Plugins/Cameras/GameplayCameras/Source/GameplayCameras/Public/GameFramework/GameplayCameraComponentBase.h
--- a/Engine/Plugins/Cameras/GameplayCameras/Source/GameplayCameras/Public/GameFramework/GameplayCameraComponentBase.h
+++ b/Engine/Plugins/Cameras/GameplayCameras/Source/GameplayCameras/Public/GameFramework/GameplayCameraComponentBase.h
@@ -127,6 +127,20 @@ public:
 	UFUNCTION(BlueprintPure, Category=Camera)
 	GAMEPLAYCAMERAS_API FRotator GetEvaluatedCameraRotation() const;
 
+public:
+
+	/** Activates the given camera rig prefab in the base layer. */
+	UFUNCTION(BlueprintCallable, Category="Camera")
+	void ActivatePersistentBaseCameraRig(UCameraRigAsset* CameraRig);
+
+	/** Activates the given camera rig prefab in the global layer. */
+	UFUNCTION(BlueprintCallable, Category="Camera")
+	void ActivatePersistentGlobalCameraRig(UCameraRigAsset* CameraRig);
+
+	/** Activates the given camera rig prefab in the visual layer. */
+	UFUNCTION(BlueprintCallable, Category="Camera")
+	void ActivatePersistentVisualCameraRig(UCameraRigAsset* CameraRig);
+
 public:
 
 	// UActorComponent interface
@@ -237,9 +251,14 @@ private:
 private:
 
 	using FGameplayCameraComponentEvaluationContext = UE::Cameras::FGameplayCameraComponentEvaluationContext;
+	using FCameraEvaluationContext = UE::Cameras::FCameraEvaluationContext;
 
+	/** Evaluation context for running main-layer camera rigs. */
 	TSharedPtr<FGameplayCameraComponentEvaluationContext> EvaluationContext;
 
+	/** Evaluation context for running base/global/visual camera rigs. */
+	TSharedPtr<FCameraEvaluationContext> NullContext;
+
 	bool bIsCameraCutNextFrame = false;
 
 	bool bIsDeactivating = false;
diff --git a/Engine/Plugins/Cameras/GameplayCameras/Source/GameplayCameras/Public/GameFramework/GameplayCamerasPlayerCameraManager.h b/Engine/Plugins/Cameras/GameplayCameras/Source/GameplayCameras/Public/GameFramework/GameplayCamerasPlayerCameraManager.h
--- a/Engine/Plugins/Cameras/GameplayCameras/Source/GameplayCameras/Public/GameFramework/GameplayCamerasPlayerCameraManager.h
+++ b/Engine/Plugins/Cameras/GameplayCameras/Source/GameplayCameras/Public/GameFramework/GameplayCamerasPlayerCameraManager.h
@@ -79,6 +79,20 @@ public:
 	UFUNCTION(BlueprintCallable, Category="Camera")
 	void ReleasePlayerController();
 
+public:
+
+	/** Activates the given camera rig prefab in the base layer. */
+	UFUNCTION(BlueprintCallable, Category="Camera")
+	void ActivatePersistentBaseCameraRig(UCameraRigAsset* CameraRig);
+
+	/** Activates the given camera rig prefab in the global layer. */
+	UFUNCTION(BlueprintCallable, Category="Camera")
+	void ActivatePersistentGlobalCameraRig(UCameraRigAsset* CameraRig);
+
+	/** Activates the given camera rig prefab in the visual layer. */
+	UFUNCTION(BlueprintCallable, Category="Camera")
+	void ActivatePersistentVisualCameraRig(UCameraRigAsset* CameraRig);
+
 public:
 
 	UFUNCTION(BlueprintCallable, Category="Camera")
@@ -116,6 +130,8 @@ private:
 
 	void RunViewRotationPreviewUpdate(float DeltaTime, FRotator& OutViewRotation, FRotator& OutDeltaRot);
 
+	void EnsureNullContext();
+
 	void OnContextStackChanged();
 	void CleanUpViewTargetContexts();
 
@@ -136,6 +152,8 @@ private:
 
 	TArray<TSharedRef<UE::Cameras::FCameraEvaluationContext>> ViewTargetContexts;
 
+	TSharedPtr<UE::Cameras::FCameraEvaluationContext> NullContext;
+
 	FMinimalViewInfo LastFrameDesiredView;
 
 	bool bIsSettingNewViewTarget = false;
diff --git a/Engine/Plugins/Cameras/GameplayCameras/Source/GameplayCameras/Public/GameFramework/IGameplayCameraSystemHost.h b/Engine/Plugins/Cameras/GameplayCameras/Source/GameplayCameras/Public/GameFramework/IGameplayCameraSystemHost.h
--- a/Engine/Plugins/Cameras/GameplayCameras/Source/GameplayCameras/Public/GameFramework/IGameplayCameraSystemHost.h
+++ b/Engine/Plugins/Cameras/GameplayCameras/Source/GameplayCameras/Public/GameFramework/IGameplayCameraSystemHost.h
@@ -8,10 +8,13 @@
 
 #include "IGameplayCameraSystemHost.generated.h"
 
+class UCameraRigAsset;
 class UCanvas;
+enum class ECameraRigLayer : uint8;
 
 namespace UE::Cameras
 {
+	class FCameraEvaluationContext;
 	class FCameraSystemEvaluator;
 	struct FCameraSystemEvaluatorCreateParams;
 }
@@ -70,6 +73,9 @@ protected:
 	/** Updates the camera system, if it exists. */
 	GAMEPLAYCAMERAS_API void UpdateCameraSystem(float DeltaTime);
 
+	/** Activates the given camera rig in the given layer. Should not be used with Main layer. */
+	GAMEPLAYCAMERAS_API void ActivateCameraRig(UCameraRigAsset* CameraRig, TSharedPtr<UE::Cameras::FCameraEvaluationContext> EvaluationContext, ECameraRigLayer EvaluationLayer);
+
 #if WITH_EDITOR
 	/** Updates the camera system, if it exists, for an editor world preview. */
 	GAMEPLAYCAMERAS_API void UpdateCameraSystemForEditorPreview(float DeltaTime);
