Cameras: fix static initialization problems that can lead to debug blocks not saving their fields in trace files

diff --git a/Engine/Plugins/Cameras/GameplayCameras/Source/GameplayCameras/Public/Debug/CameraDebugBlock.h b/Engine/Plugins/Cameras/GameplayCameras/Source/GameplayCameras/Public/Debug/CameraDebugBlock.h
--- a/Engine/Plugins/Cameras/GameplayCameras/Source/GameplayCameras/Public/Debug/CameraDebugBlock.h
+++ b/Engine/Plugins/Cameras/GameplayCameras/Source/GameplayCameras/Public/Debug/CameraDebugBlock.h
@@ -158,12 +158,16 @@ private:
 		UE_DECLARE_CAMERA_DEBUG_BLOCK(ApiDeclSpec, ClassName)\
 	private:\
 		using ThisClassName = ClassName;\
-		static FStaticFieldArray StaticFields;\
+		static FStaticFieldArray& GetStaticFields()\
+		{\
+			static FStaticFieldArray StaticFields;\
+			return StaticFields;\
+		}\
 	protected:\
 		virtual void OnSerialize(FArchive& Ar) override\
 		{\
 			Super::OnSerialize(Ar);\
-			for (FCameraDebugBlockField* StaticField : ClassName::StaticFields)\
+			for (FCameraDebugBlockField* StaticField : ClassName::GetStaticFields())\
 			{\
 				StaticField->SerializeField(this, Ar);\
 			}\
@@ -181,7 +185,7 @@ private:
 			static FField StaticField = CreateField<FieldType>(TEXT(#FieldName), Get##FieldName##Offset());\
 			return &StaticField;\
 		}\
-		inline static const int32 FieldName##FieldIndex = RegisterField(Get##FieldName##Field(), StaticFields);
+		inline static const int32 FieldName##FieldIndex = RegisterField(Get##FieldName##Field(), GetStaticFields());
 
 #define UE_DECLARE_CAMERA_DEBUG_BLOCK_END()\
 		virtual void OnDebugDraw(const FCameraDebugBlockDrawParams& Params, FCameraDebugRenderer& Renderer) override;\
@@ -189,7 +193,6 @@ private:
 
 #define UE_DEFINE_CAMERA_DEBUG_BLOCK_WITH_FIELDS(ClassName)\
 	UE_DEFINE_CAMERA_DEBUG_BLOCK(ClassName)\
-	TArray<FCameraDebugBlockField*> ClassName::StaticFields;
 
 #else  // UE_GAMEPLAY_CAMERAS_DEBUG
 
