Cameras: add layer name to blend stack debug info

diff --git a/Engine/Plugins/Cameras/GameplayCameras/Source/GameplayCameras/Private/Core/BlendStackCameraNode.cpp b/Engine/Plugins/Cameras/GameplayCameras/Source/GameplayCameras/Private/Core/BlendStackCameraNode.cpp
--- a/Engine/Plugins/Cameras/GameplayCameras/Source/GameplayCameras/Private/Core/BlendStackCameraNode.cpp
+++ b/Engine/Plugins/Cameras/GameplayCameras/Source/GameplayCameras/Private/Core/BlendStackCameraNode.cpp
@@ -679,12 +679,19 @@ FBlendStackSummaryCameraDebugBlock::FBlendStackSummaryCameraDebugBlock()
 
 FBlendStackSummaryCameraDebugBlock::FBlendStackSummaryCameraDebugBlock(const FBlendStackCameraNodeEvaluator& InEvaluator)
 {
+	const UBlendStackCameraNode* BlendStack = InEvaluator.GetCameraNodeAs<UBlendStackCameraNode>();
+
 	NumEntries = InEvaluator.Entries.Num();
-	BlendStackType = InEvaluator.GetCameraNodeAs<UBlendStackCameraNode>()->BlendStackType;
+	BlendStackType = BlendStack->BlendStackType;
+	BlendStackLayer = BlendStack->Layer;
 }
 
 void FBlendStackSummaryCameraDebugBlock::OnDebugDraw(const FCameraDebugBlockDrawParams& Params, FCameraDebugRenderer& Renderer)
 {
+	if (BlendStackLayer != ECameraRigLayer::None)
+	{
+		Renderer.AddText(TEXT("(layer %s) "), *UEnum::GetValueAsString(BlendStackLayer));
+	}
 	Renderer.AddText(TEXT("%d entries"), NumEntries);
 }
 
@@ -692,6 +699,7 @@ void FBlendStackSummaryCameraDebugBlock::OnSerialize(FArchive& Ar)
 {
 	Ar << NumEntries;
 	Ar << BlendStackType;
+	Ar << BlendStackLayer;
 }
 
 UE_DEFINE_CAMERA_DEBUG_BLOCK(FBlendStackCameraDebugBlock);
diff --git a/Engine/Plugins/Cameras/GameplayCameras/Source/GameplayCameras/Public/Core/BlendStackCameraNode.h b/Engine/Plugins/Cameras/GameplayCameras/Source/GameplayCameras/Public/Core/BlendStackCameraNode.h
--- a/Engine/Plugins/Cameras/GameplayCameras/Source/GameplayCameras/Public/Core/BlendStackCameraNode.h
+++ b/Engine/Plugins/Cameras/GameplayCameras/Source/GameplayCameras/Public/Core/BlendStackCameraNode.h
@@ -278,6 +278,7 @@ private:
 
 	int32 NumEntries;
 	ECameraBlendStackType BlendStackType;
+	ECameraRigLayer BlendStackLayer;
 };
 
 class FBlendStackCameraDebugBlock : public FCameraDebugBlock
