Cameras: add more debug information for blend stacks

diff --git a/Engine/Plugins/Cameras/GameplayCameras/Source/GameplayCameras/Private/Core/BlendStackCameraNode.cpp b/Engine/Plugins/Cameras/GameplayCameras/Source/GameplayCameras/Private/Core/BlendStackCameraNode.cpp
--- a/Engine/Plugins/Cameras/GameplayCameras/Source/GameplayCameras/Private/Core/BlendStackCameraNode.cpp
+++ b/Engine/Plugins/Cameras/GameplayCameras/Source/GameplayCameras/Private/Core/BlendStackCameraNode.cpp
@@ -624,6 +624,15 @@ void FBlendStackCameraNodeEvaluator::OnPostBuildAsset(const FGameplayCameraAsset
 
 #if UE_GAMEPLAY_CAMERAS_DEBUG
 
+UE_DECLARE_CAMERA_DEBUG_BLOCK_START(GAMEPLAYCAMERAS_API, FBlendStackEntryCameraDebugBlock)
+	UE_DECLARE_CAMERA_DEBUG_BLOCK_FIELD(int32, EntryIndex)
+	UE_DECLARE_CAMERA_DEBUG_BLOCK_FIELD(FBlendStackEntryID, EntryID)
+	UE_DECLARE_CAMERA_DEBUG_BLOCK_FIELD(FString, CameraRigName)
+	UE_DECLARE_CAMERA_DEBUG_BLOCK_FIELD(bool, bIsFrozen)
+UE_DECLARE_CAMERA_DEBUG_BLOCK_END()
+
+UE_DEFINE_CAMERA_DEBUG_BLOCK_WITH_FIELDS(FBlendStackEntryCameraDebugBlock)
+
 void FBlendStackCameraNodeEvaluator::OnBuildDebugBlocks(const FCameraDebugBlockBuildParams& Params, FCameraDebugBlockBuilder& Builder)
 {
 	FBlendStackSummaryCameraDebugBlock& DebugBlock = Builder.AttachDebugBlock<FBlendStackSummaryCameraDebugBlock>(*this);
@@ -635,13 +644,19 @@ void FBlendStackCameraNodeEvaluator::OnBuildDebugBlocks(const FCameraDebugBlockB
 
 FBlendStackCameraDebugBlock* FBlendStackCameraNodeEvaluator::BuildDetailedDebugBlock(const FCameraDebugBlockBuildParams& Params, FCameraDebugBlockBuilder& Builder)
 {
-	FBlendStackCameraDebugBlock& StackDebugBlock = Builder.BuildDebugBlock<FBlendStackCameraDebugBlock>(*this);
-	for (const FCameraRigEntry& Entry : Entries)
+	FBlendStackCameraDebugBlock& StackDebugBlock = Builder.BuildDebugBlock<FBlendStackCameraDebugBlock>();
+	for (int32 Index = 0; Index < Entries.Num(); ++Index)
 	{
+		const FCameraRigEntry& Entry(Entries[Index]);
+
 		// Each entry has a wrapper debug block with 2 children blocks:
 		// - block for the blend
 		// - block for the result
-		FCameraDebugBlock& EntryDebugBlock = Builder.BuildDebugBlock<FCameraDebugBlock>();
+		FBlendStackEntryCameraDebugBlock& EntryDebugBlock = Builder.BuildDebugBlock<FBlendStackEntryCameraDebugBlock>();
+		EntryDebugBlock.EntryIndex = Index;
+		EntryDebugBlock.EntryID = Entry.EntryID;
+		EntryDebugBlock.CameraRigName = GetNameSafe(Entry.CameraRig);
+		EntryDebugBlock.bIsFrozen = Entry.Flags.bIsFrozen;
 		StackDebugBlock.AddChild(&EntryDebugBlock);
 		{
 			FCameraNodeEvaluator* BlendEvaluator = Entry.RootEvaluator ? Entry.RootEvaluator->GetBlendEvaluator() : nullptr;
@@ -690,7 +705,8 @@ void FBlendStackSummaryCameraDebugBlock::OnDebugDraw(const FCameraDebugBlockDraw
 {
 	if (BlendStackLayer != ECameraRigLayer::None)
 	{
-		Renderer.AddText(TEXT("(layer %s) "), *UEnum::GetValueAsString(BlendStackLayer));
+		UEnum* LayerEnum = StaticEnum<ECameraRigLayer>();
+		Renderer.AddText(TEXT("(%s) "), *LayerEnum->GetDisplayNameTextByValue((int64)BlendStackLayer).ToString());
 	}
 	Renderer.AddText(TEXT("%d entries"), NumEntries);
 }
@@ -704,51 +720,16 @@ void FBlendStackSummaryCameraDebugBlock::OnSerialize(FArchive& Ar)
 
 UE_DEFINE_CAMERA_DEBUG_BLOCK(FBlendStackCameraDebugBlock);
 
-FBlendStackCameraDebugBlock::FBlendStackCameraDebugBlock()
-{
-}
-
-FBlendStackCameraDebugBlock::FBlendStackCameraDebugBlock(const FBlendStackCameraNodeEvaluator& InEvaluator)
-{
-	for (const FBlendStackCameraNodeEvaluator::FCameraRigEntry& Entry : InEvaluator.Entries)
-	{
-		FEntryDebugInfo EntryDebugInfo;
-		EntryDebugInfo.CameraRigName = Entry.CameraRig ? Entry.CameraRig->GetName() : FString("<None>");
-		Entries.Add(EntryDebugInfo);
-	}
-}
-
-void FBlendStackCameraDebugBlock::OnDebugDraw(const FCameraDebugBlockDrawParams& Params, FCameraDebugRenderer& Renderer)
+void FBlendStackEntryCameraDebugBlock::OnDebugDraw(const FCameraDebugBlockDrawParams& Params, FCameraDebugRenderer& Renderer)
 {
-	TArrayView<FCameraDebugBlock*> ChildrenView(GetChildren());
-
-	for (int32 Index = 0; Index < Entries.Num(); ++Index)
+	Renderer.AddText(
+			TEXT("{cam_passive}[%d ID=%s] {cam_notice}%s{cam_default}"),
+			EntryIndex + 1, *LexToString(EntryID), *CameraRigName);
+	if (bIsFrozen)
 	{
-		const FEntryDebugInfo& Entry(Entries[Index]);
-
-		Renderer.AddText(TEXT("{cam_passive}[%d] {cam_notice}%s{cam_default}\n"), Index + 1, *Entry.CameraRigName);
-
-		if (ChildrenView.IsValidIndex(Index))
-		{
-			Renderer.AddIndent();
-			ChildrenView[Index]->DebugDraw(Params, Renderer);
-			Renderer.RemoveIndent();
-		}
+		Renderer.AddText(" {cam_notice2}FROZEN{cam_default}");
 	}
-
-	// We've already manually renderered our children blocks.
-	Renderer.SkipAllBlocks();
-}
-
-void FBlendStackCameraDebugBlock::OnSerialize(FArchive& Ar)
-{
-	Ar << Entries;
-}
-
-FArchive& operator<< (FArchive& Ar, FBlendStackCameraDebugBlock::FEntryDebugInfo& EntryDebugInfo)
-{
-	Ar << EntryDebugInfo.CameraRigName;
-	return Ar;
+	Renderer.NewLine();
 }
 
 #endif  // UE_GAMEPLAY_CAMERAS_DEBUG
diff --git a/Engine/Plugins/Cameras/GameplayCameras/Source/GameplayCameras/Public/Core/BlendStackCameraNode.h b/Engine/Plugins/Cameras/GameplayCameras/Source/GameplayCameras/Public/Core/BlendStackCameraNode.h
--- a/Engine/Plugins/Cameras/GameplayCameras/Source/GameplayCameras/Public/Core/BlendStackCameraNode.h
+++ b/Engine/Plugins/Cameras/GameplayCameras/Source/GameplayCameras/Public/Core/BlendStackCameraNode.h
@@ -284,27 +284,6 @@ private:
 class FBlendStackCameraDebugBlock : public FCameraDebugBlock
 {
 	UE_DECLARE_CAMERA_DEBUG_BLOCK(GAMEPLAYCAMERAS_API, FBlendStackCameraDebugBlock)
-
-public:
-
-	FBlendStackCameraDebugBlock();
-	FBlendStackCameraDebugBlock(const FBlendStackCameraNodeEvaluator& InEvaluator);
-	
-protected:
-
-	virtual void OnDebugDraw(const FCameraDebugBlockDrawParams& Params, FCameraDebugRenderer& Renderer) override;
-	virtual void OnSerialize(FArchive& Ar) override;
-
-private:
-
-	struct FEntryDebugInfo
-	{
-		FString CameraRigName;
-	};
-
-	TArray<FEntryDebugInfo> Entries;
-
-	friend FArchive& operator<< (FArchive&, FEntryDebugInfo&);
 };
 
 #endif  // UE_GAMEPLAY_CAMERAS_DEBUG
diff --git a/Engine/Plugins/Cameras/GameplayCameras/Source/GameplayCameras/Public/Core/BlendStackEntryID.h b/Engine/Plugins/Cameras/GameplayCameras/Source/GameplayCameras/Public/Core/BlendStackEntryID.h
--- a/Engine/Plugins/Cameras/GameplayCameras/Source/GameplayCameras/Public/Core/BlendStackEntryID.h
+++ b/Engine/Plugins/Cameras/GameplayCameras/Source/GameplayCameras/Public/Core/BlendStackEntryID.h
@@ -2,6 +2,7 @@
 
 #pragma once
 
+#include "Containers/UnrealString.h"
 #include "CoreTypes.h"
 #include "Serialization/Archive.h"
 
@@ -53,6 +54,11 @@ public:
 		return Ar;
 	}
 
+	friend FString LexToString(FBlendStackEntryID In)
+	{
+		return LexToString(In.Value);
+	}
+
 private:
 
 	FBlendStackEntryID(uint32 InValue) 
