Cameras: properly freeze blends when freezing a camera rig in a blend stack

Before, a frozen blend stack entry would contribute to the evaluation as if it had been frozen at 100% blend. This of course doesn't work if the entry had been frozen in the middle of blending.
This change selectively tears down only the node evaluation tree that corresponds to the camera rig, leaving the blend evaluation tree in place. We do call a new "Freeze()" method on the root blend evaluator, however, in case the blend needs to cache anything before getting frozen.
The other part of this change is for making the blend stack correctly run this frozen blend.

diff --git a/Engine/Plugins/Cameras/GameplayCameras/Source/GameplayCameras/Private/Core/BlendCameraNode.cpp b/Engine/Plugins/Cameras/GameplayCameras/Source/GameplayCameras/Private/Core/BlendCameraNode.cpp
--- a/Engine/Plugins/Cameras/GameplayCameras/Source/GameplayCameras/Private/Core/BlendCameraNode.cpp
+++ b/Engine/Plugins/Cameras/GameplayCameras/Source/GameplayCameras/Private/Core/BlendCameraNode.cpp
@@ -29,5 +29,14 @@ bool FBlendCameraNodeEvaluator::SetReversed(bool bInReverse)
 	return OnSetReversed(bInReverse);
 }
 
+void FBlendCameraNodeEvaluator::Freeze()
+{
+	OnFreeze();
+
+	// When frozen, we can't access the camera node anymore, as we may have been frozen
+	// because the data we came from has been unloaded.
+	SetPrivateCameraNode(nullptr);
+}
+
 }  // namespace UE::Cameras
 
diff --git a/Engine/Plugins/Cameras/GameplayCameras/Source/GameplayCameras/Private/Core/BlendStackCameraNode.cpp b/Engine/Plugins/Cameras/GameplayCameras/Source/GameplayCameras/Private/Core/BlendStackCameraNode.cpp
--- a/Engine/Plugins/Cameras/GameplayCameras/Source/GameplayCameras/Private/Core/BlendStackCameraNode.cpp
+++ b/Engine/Plugins/Cameras/GameplayCameras/Source/GameplayCameras/Private/Core/BlendStackCameraNode.cpp
@@ -156,17 +156,30 @@ void FBlendStackCameraNodeEvaluator::FreezeEntry(FCameraRigEntry& Entry)
 {
 	if (Entry.RootEvaluator)
 	{
-		FCameraNodeEvaluatorTeardownParams TeardownParams;
-		TeardownParams.Evaluator = this->OwningEvaluator;
-		TeardownParams.EvaluationContext = Entry.EvaluationContext.Pin();
-		TeardownParams.Layer = Layer;
+		// Freeze the blend.
+		FBlendCameraNodeEvaluator* BlendEvaluator = Entry.RootEvaluator->GetBlendEvaluator();
+		if (BlendEvaluator)
+		{
+			BlendEvaluator->Freeze();
+		}
+
+		// Teardown and destroy the evaluation tree.
+		FCameraNodeEvaluator* RootEvaluator = Entry.RootEvaluator->GetRootEvaluator();
+		if (RootEvaluator)
+		{
+			FCameraNodeEvaluatorTeardownParams TeardownParams;
+			TeardownParams.Evaluator = this->OwningEvaluator;
+			TeardownParams.EvaluationContext = Entry.EvaluationContext.Pin();
+			TeardownParams.Layer = Layer;
+
+			RootEvaluator->Teardown(TeardownParams);
 
-		Entry.RootEvaluator->Teardown(TeardownParams);
+			Entry.EvaluatorStorage.DestroyEvaluatorTree(RootEvaluator, true);
+
+			Entry.RootEvaluator->SetRootEvaluator(nullptr);
+		}
 	}
 
-	// Deallocate our node evaluators and clear any pointers we kept to them.
-	Entry.EvaluatorStorage.DestroyEvaluatorTree(true);
-	Entry.RootEvaluator = nullptr;
 	Entry.EvaluatorHierarchy.Reset();
 
 	Entry.RootNode = nullptr;
diff --git a/Engine/Plugins/Cameras/GameplayCameras/Source/GameplayCameras/Private/Core/BlendStackRootCameraNode.cpp b/Engine/Plugins/Cameras/GameplayCameras/Source/GameplayCameras/Private/Core/BlendStackRootCameraNode.cpp
--- a/Engine/Plugins/Cameras/GameplayCameras/Source/GameplayCameras/Private/Core/BlendStackRootCameraNode.cpp
+++ b/Engine/Plugins/Cameras/GameplayCameras/Source/GameplayCameras/Private/Core/BlendStackRootCameraNode.cpp
@@ -103,6 +103,11 @@ void FBlendStackRootCameraNodeEvaluator::SetBlendEvaluator(FBlendCameraNodeEvalu
 	BlendEvaluator = InBlendEvaluator;
 }
 
+void FBlendStackRootCameraNodeEvaluator::SetRootEvaluator(FCameraNodeEvaluator* InRootEvaluator)
+{
+	RootEvaluator = InRootEvaluator;
+}
+
 ECameraRigMergingEligibility FBlendStackRootCameraNodeEvaluator::CompareCameraRigForMerging(const UCameraRigAsset* CameraRig) const
 {
 	const UCameraRigAsset* NewCameraRigPrefab = FindInnermostCameraRigPrefab(CameraRig);
diff --git a/Engine/Plugins/Cameras/GameplayCameras/Source/GameplayCameras/Private/Core/BlendStackRootCameraNode.h b/Engine/Plugins/Cameras/GameplayCameras/Source/GameplayCameras/Private/Core/BlendStackRootCameraNode.h
--- a/Engine/Plugins/Cameras/GameplayCameras/Source/GameplayCameras/Private/Core/BlendStackRootCameraNode.h
+++ b/Engine/Plugins/Cameras/GameplayCameras/Source/GameplayCameras/Private/Core/BlendStackRootCameraNode.h
@@ -70,6 +70,7 @@ public:
 	FCameraNodeEvaluator* GetRootEvaluator() const { return RootEvaluator; }
 
 	void SetBlendEvaluator(FBlendCameraNodeEvaluator* InBlendEvaluator);
+	void SetRootEvaluator(FCameraNodeEvaluator* InRootEvaluator);
 
 	ECameraRigMergingEligibility CompareCameraRigForMerging(const UCameraRigAsset* CameraRig) const;
 
diff --git a/Engine/Plugins/Cameras/GameplayCameras/Source/GameplayCameras/Private/Core/CameraNodeEvaluatorStorage.cpp b/Engine/Plugins/Cameras/GameplayCameras/Source/GameplayCameras/Private/Core/CameraNodeEvaluatorStorage.cpp
--- a/Engine/Plugins/Cameras/GameplayCameras/Source/GameplayCameras/Private/Core/CameraNodeEvaluatorStorage.cpp
+++ b/Engine/Plugins/Cameras/GameplayCameras/Source/GameplayCameras/Private/Core/CameraNodeEvaluatorStorage.cpp
@@ -2,10 +2,8 @@
 
 #include "Core/CameraNodeEvaluatorStorage.h"
 
-#include "Core/CameraNode.h"
 #include "Core/CameraNodeEvaluator.h"
-#include "Core/CameraRigAsset.h"
-#include "UObject/Package.h"
+#include "Core/CameraNodeEvaluatorBuilder.h"
 
 namespace UE::Cameras
 {
@@ -35,6 +33,39 @@ void FCameraNodeEvaluatorStorage::DestroyEvaluatorTree(bool bFreeAllocations)
 	Super::DestroyObjects(bFreeAllocations);
 }
 
+void FCameraNodeEvaluatorStorage::DestroyEvaluatorTree(FCameraNodeEvaluator* InRootEvaluator, bool bResetMemory)
+{
+	if (!InRootEvaluator)
+	{
+		return;
+	}
+
+	const FCameraObjectTypeRegistry& TypeRegistry = FCameraObjectTypeRegistry::Get();
+
+	TArray<FCameraNodeEvaluator*> Stack;
+	Stack.Add(InRootEvaluator);
+	while (Stack.Num() > 0)
+	{
+		FCameraNodeEvaluator* Evaluator = Stack.Pop(EAllowShrinking::No);
+		FCameraNodeEvaluatorChildrenView Children(Evaluator->GetChildren());
+		for (FCameraNodeEvaluator* Child : ReverseIterate(Children))
+		{
+			if (Child)
+			{
+				Stack.Add(Child);
+			}
+		}
+
+		const FCameraObjectTypeID TypeID = Evaluator->GetTypeID();
+		const FCameraObjectTypeInfo* TypeInfo = TypeRegistry.GetTypeInfo(TypeID);
+
+		if (ensure(TypeInfo))
+		{
+			Super::DestroyObject(Evaluator, bResetMemory ? TypeInfo->Sizeof : 0);
+		}
+	}
+}
+
 void FCameraNodeEvaluatorStorage::GetAllocationInfo(FCameraNodeEvaluatorAllocationInfo& OutAllocationInfo)
 {
 	uint32 TotalUsed;
diff --git a/Engine/Plugins/Cameras/GameplayCameras/Source/GameplayCameras/Private/Core/TransientBlendStackCameraNode.cpp b/Engine/Plugins/Cameras/GameplayCameras/Source/GameplayCameras/Private/Core/TransientBlendStackCameraNode.cpp
--- a/Engine/Plugins/Cameras/GameplayCameras/Source/GameplayCameras/Private/Core/TransientBlendStackCameraNode.cpp
+++ b/Engine/Plugins/Cameras/GameplayCameras/Source/GameplayCameras/Private/Core/TransientBlendStackCameraNode.cpp
@@ -326,33 +326,27 @@ void FTransientBlendStackCameraNodeEvaluator::InternalPreBlendExecute(TArrayView
 
 		FCameraNodeEvaluationResult& CurResult(Entry.Result);
 
-		if (!Entry.Flags.bIsFrozen && ResolvedEntry.Context.IsValid())
-		{
-			FCameraNodeEvaluationParams CurParams(Params);
-			CurParams.EvaluationContext = ResolvedEntry.Context;
-			CurParams.bIsFirstFrame = Entry.Flags.bIsFirstFrame;
+		// Running entries pre-blend the values that were put in the variable table in InternalPreBlendPrepare.
+		// Frozen entries still contribute by blending their last evaluated values.
+		FCameraNodeEvaluationParams CurParams(Params);
+		CurParams.EvaluationContext = ResolvedEntry.Context;
+		CurParams.bIsFirstFrame = Entry.Flags.bIsFirstFrame;
 
-			FCameraNodePreBlendParams PreBlendParams(CurParams, CurResult.CameraPose, CurResult.VariableTable);
-			PreBlendParams.VariableTableFilter = VariableTableFilter;
+		FCameraNodePreBlendParams PreBlendParams(CurParams, CurResult.CameraPose, CurResult.VariableTable);
+		PreBlendParams.VariableTableFilter = VariableTableFilter;
 
-			FCameraNodePreBlendResult PreBlendResult(PreBlendVariableTable);
+		FCameraNodePreBlendResult PreBlendResult(PreBlendVariableTable);
 
-			FBlendCameraNodeEvaluator* EntryBlendEvaluator = Entry.RootEvaluator->GetBlendEvaluator();
-			if (EntryBlendEvaluator)
-			{
-				EntryBlendEvaluator->BlendParameters(PreBlendParams, PreBlendResult);
-				EntryExtraInfo.bIsPreBlendFull = (PreBlendResult.bIsBlendFinished && PreBlendResult.bIsBlendFull);
-			}
-			else
-			{
-				PreBlendVariableTable.Override(CurResult.VariableTable, PreBlendParams.VariableTableFilter);
-				EntryExtraInfo.bIsPreBlendFull = true;
-			}
+		FBlendCameraNodeEvaluator* EntryBlendEvaluator = Entry.RootEvaluator ? Entry.RootEvaluator->GetBlendEvaluator() : nullptr;
+		if (EntryBlendEvaluator)
+		{
+			EntryBlendEvaluator->BlendParameters(PreBlendParams, PreBlendResult);
+			EntryExtraInfo.bIsPreBlendFull = (PreBlendResult.bIsBlendFinished && PreBlendResult.bIsBlendFull);
 		}
 		else
 		{
-			// Frozen entries still contribute to the blend using their last evaluated values.
-			PreBlendVariableTable.Override(CurResult.VariableTable, VariableTableFilter);
+			PreBlendVariableTable.Override(CurResult.VariableTable, PreBlendParams.VariableTableFilter);
+			EntryExtraInfo.bIsPreBlendFull = true;
 		}
 	}
 
@@ -425,29 +419,20 @@ void FTransientBlendStackCameraNodeEvaluator::InternalPostBlendExecute(TArrayVie
 
 		FCameraNodeEvaluationResult& CurResult(Entry.Result);
 
-		if (!Entry.Flags.bIsFrozen && ResolvedEntry.Context.IsValid())
-		{
-			FCameraNodeEvaluationParams CurParams(Params);
-			CurParams.EvaluationContext = ResolvedEntry.Context;
-			CurParams.bIsFirstFrame = Entry.Flags.bIsFirstFrame;
-			FCameraNodeBlendParams BlendParams(CurParams, CurResult);
+		FCameraNodeEvaluationParams CurParams(Params);
+		CurParams.EvaluationContext = ResolvedEntry.Context;
+		CurParams.bIsFirstFrame = Entry.Flags.bIsFirstFrame;
+		FCameraNodeBlendParams BlendParams(CurParams, CurResult);
 
-			FCameraNodeBlendResult BlendResult(OutResult);
+		FCameraNodeBlendResult BlendResult(OutResult);
 
-			FBlendCameraNodeEvaluator* EntryBlendEvaluator = Entry.RootEvaluator->GetBlendEvaluator();
-			if (EntryBlendEvaluator)
-			{
-				EntryBlendEvaluator->BlendResults(BlendParams, BlendResult);
+		FBlendCameraNodeEvaluator* EntryBlendEvaluator = Entry.RootEvaluator ? Entry.RootEvaluator->GetBlendEvaluator() : nullptr;
+		if (EntryBlendEvaluator)
+		{
+			EntryBlendEvaluator->BlendResults(BlendParams, BlendResult);
 
-				if (BlendResult.bIsBlendFull && BlendResult.bIsBlendFinished)
-				{
-					PopEntriesBelow = ResolvedEntry.EntryIndex;
-				}
-			}
-			else
+			if (BlendResult.bIsBlendFull && BlendResult.bIsBlendFinished)
 			{
-				OutResult.OverrideAll(CurResult);
-
 				PopEntriesBelow = ResolvedEntry.EntryIndex;
 			}
 		}
@@ -457,19 +442,14 @@ void FTransientBlendStackCameraNodeEvaluator::InternalPostBlendExecute(TArrayVie
 
 			PopEntriesBelow = ResolvedEntry.EntryIndex;
 		}
-		// else, merged entries only contribute to the pre-blend.
 	}
 
 	// Pop out camera rigs that have been blended out.
-	if (!Params.IsStatelessEvaluation())
+	if (!Params.IsStatelessEvaluation() && PopEntriesBelow != INDEX_NONE)
 	{
-		const UBlendStackCameraNode* BlendStackNode = GetCameraNodeAs<UBlendStackCameraNode>();
-		if (BlendStackNode->BlendStackType == ECameraBlendStackType::IsolatedTransient && PopEntriesBelow != INDEX_NONE)
-		{
-			PopEntries(PopEntriesBelow);
+		PopEntries(PopEntriesBelow);
 
-			EntryExtraInfos.RemoveAt(0, PopEntriesBelow);
-		}
+		EntryExtraInfos.RemoveAt(0, PopEntriesBelow);
 	}
 }
 
diff --git a/Engine/Plugins/Cameras/GameplayCameras/Source/GameplayCameras/Public/Core/BlendCameraNode.h b/Engine/Plugins/Cameras/GameplayCameras/Source/GameplayCameras/Public/Core/BlendCameraNode.h
--- a/Engine/Plugins/Cameras/GameplayCameras/Source/GameplayCameras/Public/Core/BlendCameraNode.h
+++ b/Engine/Plugins/Cameras/GameplayCameras/Source/GameplayCameras/Public/Core/BlendCameraNode.h
@@ -145,6 +145,12 @@ public:
 	 */
 	bool SetReversed(bool bInReverse);
 
+	/**
+	 * Freezes this blend. A frozen blend can't access its evaluation context or underlying camera node
+	 * anymore, and must operate entirely standalone.
+	 */
+	void Freeze();
+
 protected:
 
 	/** Blend the parameters produced by a camera node tree over another set of values. */
@@ -158,6 +164,9 @@ protected:
 
 	/** Reverse the direction of this blend. See comments from SetReversed. */
 	virtual bool OnSetReversed(bool bInReverse) { return false; }
+
+	/** Freezes this blend. See commands from Freeze. */
+	virtual void OnFreeze() {}
 };
 
 }  // namespace UE::Cameras
diff --git a/Engine/Plugins/Cameras/GameplayCameras/Source/GameplayCameras/Public/Core/CameraNodeEvaluatorStorage.h b/Engine/Plugins/Cameras/GameplayCameras/Source/GameplayCameras/Public/Core/CameraNodeEvaluatorStorage.h
--- a/Engine/Plugins/Cameras/GameplayCameras/Source/GameplayCameras/Public/Core/CameraNodeEvaluatorStorage.h
+++ b/Engine/Plugins/Cameras/GameplayCameras/Source/GameplayCameras/Public/Core/CameraNodeEvaluatorStorage.h
@@ -5,7 +5,6 @@
 #include "Core/CameraNodeEvaluator.h"
 #include "Core/CameraNodeEvaluatorFwd.h"
 #include "Core/CameraObjectStorage.h"
-#include "CoreTypes.h"
 #include "UObject/ObjectPtr.h"
 
 class FReferenceCollector;
@@ -44,6 +43,9 @@ public:
 	/** Destroy any allocated evaluators. */
 	void DestroyEvaluatorTree(bool bFreeAllocations = false);
 
+	/** Destroy a sub-tree of allocated evaluators. */
+	void DestroyEvaluatorTree(FCameraNodeEvaluator* InRootEvaluator, bool bResetMemory = true);
+
 	/** Computes the currently allocated totals. */
 	void GetAllocationInfo(FCameraNodeEvaluatorAllocationInfo& OutAllocationInfo);
 
diff --git a/Engine/Plugins/Cameras/GameplayCameras/Source/GameplayCameras/Public/Core/CameraObjectStorage.h b/Engine/Plugins/Cameras/GameplayCameras/Source/GameplayCameras/Public/Core/CameraObjectStorage.h
--- a/Engine/Plugins/Cameras/GameplayCameras/Source/GameplayCameras/Public/Core/CameraObjectStorage.h
+++ b/Engine/Plugins/Cameras/GameplayCameras/Source/GameplayCameras/Public/Core/CameraObjectStorage.h
@@ -70,6 +70,11 @@ protected:
 	 */
 	void RegisterInitializedObject(BaseObjectType* BaseObjectPtr);
 
+	/**
+	 * Destroys a given object.
+	 */
+	void DestroyObject(BaseObjectType* BaseObjectPtr, uint32 MemoryResetSize = 0);
+
 	/**
 	 * Destroys all objects in the storage.
 	 *
@@ -204,6 +209,39 @@ void TCameraObjectStorage<BaseObjectType>::RegisterInitializedObject(BaseObjectT
 	ObjectInfos.Add({ BaseObjectPtr });
 }
 
+template<typename BaseObjectType>
+void TCameraObjectStorage<BaseObjectType>::DestroyObject(BaseObjectType* BaseObjectPtr, uint32 MemoryResetSize)
+{
+	// Check that this is an object we have allocated ourselves.
+	const int32 NumRemoved = ObjectInfos.RemoveAll([BaseObjectPtr](FObjectInfo& ObjectInfo)
+			{
+				return ObjectInfo.Ptr == BaseObjectPtr;
+			});
+	ensureMsgf(NumRemoved == 1, TEXT("Given object pointer isn't in this storage, or was found multiple times."));
+	if (NumRemoved > 0)
+	{
+		// Destroy the object.
+		BaseObjectPtr->~BaseObjectType();
+
+		// Optionally, reset the object's memory, trusting that the caller is giving us the right size.
+		// In that case, first check that the memory to reset is contained within our allocation pages.
+		if (MemoryResetSize > 0)
+		{
+			const FAllocation* FoundAllocation = Allocations.FindByPredicate(
+					[BaseObjectPtr](const FAllocation& Allocation)
+					{
+						return (
+								Allocation.Memory <= (uint8*)BaseObjectPtr && 
+								Allocation.Memory + Allocation.Used > (uint8*)BaseObjectPtr);
+					});
+			if (ensure(FoundAllocation && (uint8*)BaseObjectPtr + MemoryResetSize <= FoundAllocation->Memory + FoundAllocation->Used))
+			{
+				FMemory::Memzero((void*)BaseObjectPtr, MemoryResetSize);
+			}
+		}
+	}
+}
+
 template<typename BaseObjectType>
 void TCameraObjectStorage<BaseObjectType>::DestroyObjects(bool bFreeAllocations)
 {
