Cameras: fix minimum engine version for SequencerObjectBindingHelper class.

diff --git a/Engine/Plugins/Cameras/GameplayCameras/Source/GameplayCamerasEditor/Private/Sequencer/GameplayCameraComponentTrackEditor.cpp b/Engine/Plugins/Cameras/GameplayCameras/Source/GameplayCamerasEditor/Private/Sequencer/GameplayCameraComponentTrackEditor.cpp
--- a/Engine/Plugins/Cameras/GameplayCameras/Source/GameplayCamerasEditor/Private/Sequencer/GameplayCameraComponentTrackEditor.cpp
+++ b/Engine/Plugins/Cameras/GameplayCameras/Source/GameplayCamerasEditor/Private/Sequencer/GameplayCameraComponentTrackEditor.cpp
@@ -20,7 +20,7 @@
 #include "Styles/GameplayCamerasEditorStyle.h"
 #include "Tracks/MovieScenePropertyTrack.h"
 
-#if UE_VERSION_NEWER_THAN_OR_EQUAL(5,6,0)
+#if UE_VERSION_NEWER_THAN_OR_EQUAL(5,7,0)
 #include "Misc/SequencerObjectBindingHelper.h"
 #endif
 
