Cameras: small improvements for Gameplay Camera components

- Auto-create camera rigs and assets when nothing is specified on the component (but log a warning).
- Improve logic for handling deferred deactivation.
- Improve logging.

diff --git a/Engine/Plugins/Cameras/GameplayCameras/Source/GameplayCameras/Private/Core/CameraEvaluationContext.cpp b/Engine/Plugins/Cameras/GameplayCameras/Source/GameplayCameras/Private/Core/CameraEvaluationContext.cpp
--- a/Engine/Plugins/Cameras/GameplayCameras/Source/GameplayCameras/Private/Core/CameraEvaluationContext.cpp
+++ b/Engine/Plugins/Cameras/GameplayCameras/Source/GameplayCameras/Private/Core/CameraEvaluationContext.cpp
@@ -311,6 +311,15 @@ void FCameraEvaluationContext::Deactivate(const FCameraEvaluationContextDeactiva
 	bActivated = false;
 }
 
+TSharedPtr<FCameraSystemEvaluator> FCameraEvaluationContext::GetCameraSystemEvaluator() const
+{
+	if (CameraSystemEvaluator)
+	{
+		return CameraSystemEvaluator->AsShared();
+	}
+	return TSharedPtr<FCameraSystemEvaluator>();
+}
+
 bool FCameraEvaluationContext::AddChildContext(TSharedRef<FCameraEvaluationContext> ChildContext)
 {
 	if (DirectorEvaluator)
diff --git a/Engine/Plugins/Cameras/GameplayCameras/Source/GameplayCameras/Private/GameFramework/GameplayCameraComponent.cpp b/Engine/Plugins/Cameras/GameplayCameras/Source/GameplayCameras/Private/GameFramework/GameplayCameraComponent.cpp
--- a/Engine/Plugins/Cameras/GameplayCameras/Source/GameplayCameras/Private/GameFramework/GameplayCameraComponent.cpp
+++ b/Engine/Plugins/Cameras/GameplayCameras/Source/GameplayCameras/Private/GameFramework/GameplayCameraComponent.cpp
@@ -56,18 +56,6 @@ UCameraAsset* UGameplayCameraComponent::GetCameraAsset()
 	return CameraReference.GetCameraAsset();
 }
 
-bool UGameplayCameraComponent::OnValidateCameraEvaluationContextActivation()
-{
-	const bool bIsValid = CameraReference.IsValid();
-	if (!bIsValid && !IsEditorWorld())
-	{
-		FFrame::KismetExecutionMessage(
-				TEXT("Can't activate gameplay camera component: no camera asset was set!"),
-				ELogVerbosity::Error);
-	}
-	return bIsValid;
-}
-
 void UGameplayCameraComponent::OnUpdateCameraEvaluationContext(bool bForceApplyParameterOverrides)
 {
 	using namespace UE::Cameras;
diff --git a/Engine/Plugins/Cameras/GameplayCameras/Source/GameplayCameras/Private/GameFramework/GameplayCameraComponentBase.cpp b/Engine/Plugins/Cameras/GameplayCameras/Source/GameplayCameras/Private/GameFramework/GameplayCameraComponentBase.cpp
--- a/Engine/Plugins/Cameras/GameplayCameras/Source/GameplayCameras/Private/GameFramework/GameplayCameraComponentBase.cpp
+++ b/Engine/Plugins/Cameras/GameplayCameras/Source/GameplayCameras/Private/GameFramework/GameplayCameraComponentBase.cpp
@@ -7,9 +7,11 @@
 #include "CineCameraComponent.h"
 #include "Core/CameraAsset.h"
 #include "Core/CameraEvaluationContextStack.h"
+#include "Core/CameraRigAsset.h"
 #include "Core/CameraSystemEvaluator.h"
 #include "Core/RootCameraNode.h"
 #include "Debug/CameraDebugRenderer.h"
+#include "Directors/SingleCameraDirector.h"
 #include "Engine/Canvas.h"
 #include "Engine/EngineTypes.h"
 #include "Engine/World.h"
@@ -17,6 +19,7 @@
 #include "IGameplayCamerasLiveEditManager.h"
 #include "IGameplayCamerasModule.h"
 #include "Kismet/GameplayStatics.h"
+#include "Logging/LogVerbosity.h"
 #include "Misc/AssertionMacros.h"
 #include "Misc/EngineVersionComparison.h"
 #include "SceneView.h"
@@ -54,9 +57,7 @@ UGameplayCameraComponentBase::UGameplayCameraComponentBase(const FObjectInitiali
 void UGameplayCameraComponentBase::BeginDestroy()
 {
 	DestroyCameraSystem();
-
-	EvaluationContext = nullptr;
-	NullContext = nullptr;
+	DestroyCameraEvaluationContext();
 
 	Super::BeginDestroy();
 }
@@ -116,16 +117,17 @@ void UGameplayCameraComponentBase::ActivateCameraForPlayerController(
 	}
 	if (!PlayerControllerHost && ActivationMode != EGameplayCameraComponentActivationMode::Push)
 	{
-		UE_LOG(LogCameraSystem, Warning,
-				TEXT("Gameplay camera component '%s' cannot activate with mode '%s' because no camera system "
-					 "was found on the given player controller, or not player controller was specified. "),
-				*GetNameSafe(this),
-				*UEnum::GetValueAsString(ActivationMode));
+		FFrame::KismetExecutionMessage(
+				*FString::Format(
+					TEXT("Gameplay Camera component '{0}.{1}' cannot activate with mode '{2}' because no camera system "
+						 "was found on the given player controller, or no player controller was specified."),
+					{ *GetNameSafe(GetOwner()), *GetNameSafe(this), *UEnum::GetValueAsString(ActivationMode) }),
+				ELogVerbosity::Warning);
 	}
 
 	if (PlayerControllerHost)
 	{
-		TeardownCameraSystemHost();
+		DestroyCameraSystem();
 		ActivateCameraEvaluationContext(PlayerController, PlayerControllerHost, ActivationMode);
 	}
 	else
@@ -162,44 +164,58 @@ void UGameplayCameraComponentBase::DeactivateCameraEvaluationContext(bool bImmed
 {
 	using namespace UE::Cameras;
 
-	// Remove the evaluation context for the director tree.
-	if (EvaluationContext && CameraSystemEvaluator)
+	if (!bImmediately && HasCameraSystem())
+	{
+		FFrame::KismetExecutionMessage(
+				*FString::Format(
+					TEXT("Can't deactivate Gameplay Camera component '{0}.{1}' with bImmediately set to false: "
+						 "this component is running its own camera system and cannot therefore blend out from itself."),
+					{ *GetNameSafe(GetOwner()), *GetNameSafe(this) }),
+				ELogVerbosity::Warning);
+		bImmediately = true;
+	}
+
+	if (EvaluationContext && EvaluationContext->IsActive())
 	{
 		UE_LOG(LogCameraSystem, Log, TEXT("Deactivating gameplay camera '%s'."), *GetNameSafe(this));
 
+		TSharedPtr<FCameraSystemEvaluator> HostEvaluator = EvaluationContext->GetCameraSystemEvaluator();
+		ensure(HostEvaluator);
+
 		if (TSharedPtr<FCameraEvaluationContext> ParentContext = EvaluationContext->GetParentContext())
 		{
 			ParentContext->RemoveChildContext(EvaluationContext.ToSharedRef());
 		}
 		else
 		{
-			CameraSystemEvaluator->RemoveEvaluationContext(EvaluationContext.ToSharedRef());
+			HostEvaluator->RemoveEvaluationContext(EvaluationContext.ToSharedRef());
 		}
 
 		if (bImmediately)
 		{
 			// We are deactivating immediately (i.e. without letting our camera rigs blend out), so make 
 			// sure everything is frozen or disabled before we delete our evaluation context.
-			FRootCameraNodeEvaluator* RootNodeEvaluator = CameraSystemEvaluator->GetRootNodeEvaluator();
+			FRootCameraNodeEvaluator* RootNodeEvaluator = HostEvaluator->GetRootNodeEvaluator();
 			RootNodeEvaluator->DeactivateAllCameraRigs(EvaluationContext, true);
 		}
+		else
+		{
+			// Don't deactivate the component right away: we still need to update our evaluation context 
+			// while any running camera rigs blend out.
+			PendingDeactivateCameraSystemEvaluator = HostEvaluator;
+			bIsPendingDeactivate = true;
+		}
 	}
 
-	if (EvaluationContext && CameraSystemEvaluator && !bImmediately)
-	{
-		// Don't deactivate the component right away: we still need to update our evaluation context 
-		// while any running camera rigs blend out.
-		bIsDeactivating = true;
-	}
-	else
+	if (bImmediately)
 	{
 		if (OutputCameraComponent)
 		{
 			OutputCameraComponent->SetRelativeTransform(FTransform());
 		}
 
-		EvaluationContext = nullptr;
-		NullContext = nullptr;
+		DestroyCameraSystem();
+		DestroyCameraEvaluationContext();
 	}
 }
 
@@ -243,14 +259,6 @@ void UGameplayCameraComponentBase::EnsureCameraSystemHost()
 	}
 }
 
-void UGameplayCameraComponentBase::TeardownCameraSystemHost()
-{
-	if (HasCameraSystem())
-	{
-		DestroyCameraSystem();
-	}
-}
-
 void UGameplayCameraComponentBase::ActivateCameraEvaluationContext(APlayerController* PlayerController, IGameplayCameraSystemHost* Host, EGameplayCameraComponentActivationMode ActivationMode)
 {
 	using namespace UE::Cameras;
@@ -265,24 +273,14 @@ void UGameplayCameraComponentBase::ActivateCameraEvaluationContext(APlayerContro
 		return;
 	}
 
-	if (!OnValidateCameraEvaluationContextActivation())
-	{
-		return;
-	}
-
-	// Make sure the evaluation context has been created. However, this can fail, such as when
-	// we don't have a valid camera asset specified.
-	TryCreateCameraEvaluationContext(PlayerController);
-
-	if (!EvaluationContext.IsValid())
-	{
-		return;
-	}
+	CreateCameraEvaluationContext(PlayerController);
 
 	if (EvaluationContext->IsActive())
 	{
 		FFrame::KismetExecutionMessage(
-				TEXT("Can't activate gameplay camera component: it is already active!"),
+				*FString::Format(
+					TEXT("Can't activate Gameplay Camera component '{0}.{1}': it is already active!"),
+					{ *GetNameSafe(GetOwner()), *GetNameSafe(this) }),
 				ELogVerbosity::Error);
 		return;
 	}
@@ -329,10 +327,11 @@ void UGameplayCameraComponentBase::ActivateCameraEvaluationContext(APlayerContro
 	}
 
 	// Cancel any ongoing deactivation.
-	bIsDeactivating = false;
+	PendingDeactivateCameraSystemEvaluator = nullptr;
+	bIsPendingDeactivate = false;
 }
 
-void UGameplayCameraComponentBase::TryCreateCameraEvaluationContext(APlayerController* PlayerController)
+void UGameplayCameraComponentBase::CreateCameraEvaluationContext(APlayerController* PlayerController)
 {
 	using namespace UE::Cameras;
 
@@ -340,35 +339,52 @@ void UGameplayCameraComponentBase::TryCreateCameraEvaluationContext(APlayerContr
 
 	if (!EvaluationContext.IsValid())
 	{
-		if (UCameraAsset* CameraAsset = GetCameraAsset())
-		{
+		UCameraAsset* CameraAsset = GetCameraAsset();
+
 #if WITH_EDITOR
-			if (bIsFirstActivation)
+		if (CameraAsset && bIsFirstActivation)
+		{
+			UWorld* World = GetWorld();
+			const UGameplayCamerasSettings* Settings = GetDefault<UGameplayCamerasSettings>();
+			if (Settings->bAutoBuildInPIE && World && World->WorldType == EWorldType::PIE)
 			{
-				UWorld* World = GetWorld();
-				const UGameplayCamerasSettings* Settings = GetDefault<UGameplayCamerasSettings>();
-				if (Settings->bAutoBuildInPIE && World && World->WorldType == EWorldType::PIE)
-				{
-					// Auto-build the camera asset on begin play to make sure we've got the latest user edits.
-					FCameraBuildLog BuildLog;
-					FCameraAssetBuilder Builder(BuildLog);
-					Builder.BuildCamera(CameraAsset);
-				}
-
-				bIsFirstActivation = false;
+				// Auto-build the camera asset on begin play to make sure we've got the latest user edits.
+				FCameraBuildLog BuildLog;
+				FCameraAssetBuilder Builder(BuildLog);
+				Builder.BuildCamera(CameraAsset);
 			}
+
+			bIsFirstActivation = false;
+		}
 #endif
+		
+		// If we have no camera asset specified, make a placeholder one and log a warning.
+		if (!CameraAsset)
+		{
+			UE_LOG(LogCameraSystem, Warning, 
+					TEXT("No camera asset specified on Gameplay Camera component '%s.%s', using a placeholder one."),
+					*GetNameSafe(this), *GetNameSafe(GetOwner()));
 
-			EvaluationContext = MakeShared<FGameplayCameraComponentEvaluationContext>();
+			UCameraRigAsset* PlaceholderCameraRig = NewObject<UCameraRigAsset>();
 
-			FCameraEvaluationContextInitializeParams InitParams;
-			InitParams.Owner = this;
-			InitParams.CameraAsset = CameraAsset;
-			InitParams.PlayerController = PlayerController;
-			EvaluationContext->Initialize(InitParams);
+			USingleCameraDirector* PlaceholderCameraDirector = NewObject<USingleCameraDirector>();
+			PlaceholderCameraDirector->CameraRig = PlaceholderCameraRig;
 
-			UpdateCameraEvaluationContext(true);
+			UCameraAsset* PlaceholderCameraAsset = NewObject<UCameraAsset>();
+			PlaceholderCameraAsset->SetCameraDirector(PlaceholderCameraDirector);
+			
+			CameraAsset = PlaceholderCameraAsset;
 		}
+
+		EvaluationContext = MakeShared<FGameplayCameraComponentEvaluationContext>();
+
+		FCameraEvaluationContextInitializeParams InitParams;
+		InitParams.Owner = this;
+		InitParams.CameraAsset = CameraAsset;
+		InitParams.PlayerController = PlayerController;
+		EvaluationContext->Initialize(InitParams);
+
+		UpdateCameraEvaluationContext(true);
 	}
 
 	if (!NullContext.IsValid())
@@ -380,6 +396,14 @@ void UGameplayCameraComponentBase::TryCreateCameraEvaluationContext(APlayerContr
 		InitParams.PlayerController = PlayerController;
 		NullContext->Initialize(InitParams);
 	}
+
+	ensure(EvaluationContext.IsValid() && NullContext.IsValid());
+}
+
+void UGameplayCameraComponentBase::DestroyCameraEvaluationContext()
+{
+	EvaluationContext = nullptr;
+	NullContext = nullptr;
 }
 
 #define UE_PRIVATE_GAMEPLAY_CAMERA_COMPONENT_VALIDATE_EVALUATION_CONTEXT(ErrorMsg, ErrorResult)\
@@ -388,8 +412,8 @@ void UGameplayCameraComponentBase::TryCreateCameraEvaluationContext(APlayerContr
 	{\
 		FFrame::KismetExecutionMessage(\
 				*FString::Format(\
-					TEXT(#ErrorMsg " on Gameplay Camera component '{0}': it isn't active."),\
-					{ *GetNameSafe(this) }),\
+					TEXT(#ErrorMsg " on Gameplay Camera component '{0}.{1}': it isn't active."),\
+					{ *GetNameSafe(GetOwner()), *GetNameSafe(this) }),\
 				ELogVerbosity::Error);\
 		return ErrorResult;\
 	}
@@ -419,8 +443,8 @@ FBlueprintCameraEvaluationDataRef UGameplayCameraComponentBase::GetConditionalRe
 	{\
 		FFrame::KismetExecutionMessage(\
 				*FString::Format(\
-					TEXT(#ErrorMsg " on Gameplay Camera component '{0}': it isn't active."),\
-					{ *GetNameSafe(this) }),\
+					TEXT(#ErrorMsg " on Gameplay Camera component '{0}.{1}': it isn't active."),\
+					{ *GetNameSafe(GetOwner()), *GetNameSafe(this) }),\
 				ELogVerbosity::Error);\
 	}
 
@@ -573,7 +597,7 @@ void UGameplayCameraComponentBase::CheckPendingDeactivation()
 {
 	using namespace UE::Cameras;
 
-	if (!bIsDeactivating)
+	if (!bIsPendingDeactivate)
 	{
 		return;
 	}
@@ -582,23 +606,24 @@ void UGameplayCameraComponentBase::CheckPendingDeactivation()
 	// at which point we can tear down all our evaluation apparatus.
 	bool bDoneDeactivating = true;
 
-	if (CameraSystemEvaluator && EvaluationContext)
+	TSharedPtr<FCameraSystemEvaluator> CameraSystemToCheck = PendingDeactivateCameraSystemEvaluator.Pin();
+	if (EvaluationContext && CameraSystemToCheck)
 	{
-		FRootCameraNodeEvaluator* RootNodeEvaluator = CameraSystemEvaluator->GetRootNodeEvaluator();
+		FRootCameraNodeEvaluator* RootNodeEvaluator = CameraSystemToCheck->GetRootNodeEvaluator();
 		bDoneDeactivating = (RootNodeEvaluator->HasAnyRunningCameraRig(EvaluationContext) == false);
 	}
 
 	if (bDoneDeactivating)
 	{
-		TeardownCameraSystemHost();
-		EvaluationContext = nullptr;
-		NullContext = nullptr;
+		DestroyCameraSystem();
+		DestroyCameraEvaluationContext();
 
 		// Only call the base class method here: we just want to finish deactivating ourselves
 		// by stopping ticking.
 		Super::Deactivate();
 
-		bIsDeactivating = false;
+		PendingDeactivateCameraSystemEvaluator = nullptr;
+		bIsPendingDeactivate = false;
 	}
 }
 
@@ -652,7 +677,7 @@ void UGameplayCameraComponentBase::UpdateControlRotationIfNeeded()
 
 	// Set control rotation if we are the view target.
 	AActor* OwnerActor = GetOwner();
-	if (OwnerActor && PlayerController->GetViewTarget() == OwnerActor)
+	if (CameraSystemEvaluator && OwnerActor && PlayerController->GetViewTarget() == OwnerActor)
 	{
 		const FCameraSystemEvaluationResult& Result = CameraSystemEvaluator->GetPreVisualLayerEvaluatedResult();
 		const FRotator3d& ControlRotation = Result.CameraPose.GetRotation();
@@ -685,9 +710,9 @@ void UGameplayCameraComponentBase::ReinitializeCameraEvaluationContext(
 
 		// Also freeze/remove any of our currently running camera rigs, because they might continue
 		// accessing variables and data that don't exist anymore.
-		if (CameraSystemEvaluator)
+		if (TSharedPtr<FCameraSystemEvaluator> HostEvaluator = EvaluationContext->GetCameraSystemEvaluator())
 		{
-			FRootCameraNodeEvaluator* RootEvaluator = CameraSystemEvaluator->GetRootNodeEvaluator();
+			FRootCameraNodeEvaluator* RootEvaluator = HostEvaluator->GetRootNodeEvaluator();
 			RootEvaluator->DeactivateAllCameraRigs(EvaluationContext, true);
 		}
 	}
@@ -706,6 +731,8 @@ void UGameplayCameraComponentBase::RecreateEditorWorldCameraEvaluationContext()
 	// this is indeed the case.
 	if (EvaluationContext && CameraSystemEvaluator)
 	{
+		ensure(CameraSystemEvaluator->GetRole() == ECameraSystemEvaluatorRole::EditorPreview);
+
 		FCameraEvaluationContextStack& ContextStack = CameraSystemEvaluator->GetEvaluationContextStack();
 		TArray<TSharedPtr<FCameraEvaluationContext>> AllContexts;
 		ContextStack.GetAllContexts(AllContexts);
@@ -723,7 +750,7 @@ void UGameplayCameraComponentBase::RecreateEditorWorldCameraEvaluationContext()
 		}
 		EvaluationContext = nullptr;
 
-		TryCreateCameraEvaluationContext(nullptr);
+		CreateCameraEvaluationContext(nullptr);
 		if (CameraSystemEvaluator)
 		{
 			CameraSystemEvaluator->PushEvaluationContext(EvaluationContext.ToSharedRef());
@@ -836,7 +863,7 @@ void UGameplayCameraComponentBase::AutoManageEditorPreviewEvaluator()
 	{
 		// We don't want to run the camera logic in the editor anymore. Let's tear things down.
 		DeactivateCameraEvaluationContext(true);
-		TeardownCameraSystemHost();
+		DestroyCameraSystem();
 		EvaluationContext = nullptr;
 
 		if (OutputCameraComponent)
@@ -923,7 +950,7 @@ UE_DEFINE_CAMERA_EVALUATION_CONTEXT(FGameplayCameraComponentEvaluationContext)
 
 void FGameplayCameraComponentEvaluationContext::UpdateForEditorPreview()
 {
-	FCameraSystemEvaluator* ActiveEvaluator = GetCameraSystemEvaluator();
+	FCameraSystemEvaluator* ActiveEvaluator = GetPrivateCameraSystemEvaluator();
 	if (ActiveEvaluator && ActiveEvaluator->GetRole() == ECameraSystemEvaluatorRole::EditorPreview)
 	{
 		if (GCurrentLevelEditingViewportClient && GCurrentLevelEditingViewportClient->Viewport)
diff --git a/Engine/Plugins/Cameras/GameplayCameras/Source/GameplayCameras/Private/GameFramework/GameplayCameraRigComponent.cpp b/Engine/Plugins/Cameras/GameplayCameras/Source/GameplayCameras/Private/GameFramework/GameplayCameraRigComponent.cpp
--- a/Engine/Plugins/Cameras/GameplayCameras/Source/GameplayCameras/Private/GameFramework/GameplayCameraRigComponent.cpp
+++ b/Engine/Plugins/Cameras/GameplayCameras/Source/GameplayCameras/Private/GameFramework/GameplayCameraRigComponent.cpp
@@ -5,6 +5,7 @@
 #include "Core/CameraAsset.h"
 #include "Core/CameraRigAsset.h"
 #include "Directors/SingleCameraDirector.h"
+#include "GameFramework/Actor.h"
 #include "GameplayCamerasDelegates.h"
 
 #include UE_INLINE_GENERATED_CPP_BY_NAME(GameplayCameraRigComponent)
@@ -46,6 +47,15 @@ UCameraAsset* UGameplayCameraRigComponent::GetCameraAsset()
 	{
 		USingleCameraDirector* SingleDirector = NewObject<USingleCameraDirector>(this, TEXT("GeneratedCameraDirector"), RF_Transient);
 		SingleDirector->CameraRig = CameraRigReference.GetCameraRig();
+		if (!SingleDirector->CameraRig)
+		{
+			UE_LOG(LogCameraSystem, Warning, 
+					TEXT("No camera rig specified on Gameplay Camera component '%s.%s', using a placeholder one."),
+					*GetNameSafe(GetOwner()), *GetNameSafe(this));
+
+			UCameraRigAsset* PlaceholderCameraRig = NewObject<UCameraRigAsset>();
+			SingleDirector->CameraRig = PlaceholderCameraRig;
+		}
 
 		GeneratedCameraAsset = NewObject<UCameraAsset>(this, TEXT("GeneratedCameraAsset"), RF_Transient);
 		GeneratedCameraAsset->SetCameraDirector(SingleDirector);
@@ -56,18 +66,6 @@ UCameraAsset* UGameplayCameraRigComponent::GetCameraAsset()
 	return GeneratedCameraAsset;
 }
 
-bool UGameplayCameraRigComponent::OnValidateCameraEvaluationContextActivation()
-{
-	const bool bIsValid = CameraRigReference.IsValid();
-	if (!bIsValid && !IsEditorWorld())
-	{
-		FFrame::KismetExecutionMessage(
-				TEXT("Can't activate gameplay camera rig component: no camera rig asset was set!"),
-				ELogVerbosity::Error);
-	}
-	return bIsValid;
-}
-
 void UGameplayCameraRigComponent::OnUpdateCameraEvaluationContext(bool bForceApplyParameterOverrides)
 {
 	using namespace UE::Cameras;
diff --git a/Engine/Plugins/Cameras/GameplayCameras/Source/GameplayCameras/Public/Core/CameraEvaluationContext.h b/Engine/Plugins/Cameras/GameplayCameras/Source/GameplayCameras/Public/Core/CameraEvaluationContext.h
--- a/Engine/Plugins/Cameras/GameplayCameras/Source/GameplayCameras/Public/Core/CameraEvaluationContext.h
+++ b/Engine/Plugins/Cameras/GameplayCameras/Source/GameplayCameras/Public/Core/CameraEvaluationContext.h
@@ -146,6 +146,12 @@ public:
 	/** Whether this context is active. */
 	bool IsActive() const { return bActivated; }
 
+	/**
+	 * Returns the camera system for which this context was activated, or null if this
+	 * context is currently not active.
+	 */
+	TSharedPtr<FCameraSystemEvaluator> GetCameraSystemEvaluator() const;
+
 public:
 
 	/** Adds a child context. */
@@ -173,7 +179,7 @@ protected:
 	virtual void OnActivate(const FCameraEvaluationContextActivateParams& Params) {}
 	virtual void OnDeactivate(const FCameraEvaluationContextDeactivateParams& Params) {}
 
-	FCameraSystemEvaluator* GetCameraSystemEvaluator() const { return CameraSystemEvaluator; }
+	FCameraSystemEvaluator* GetPrivateCameraSystemEvaluator() const { return CameraSystemEvaluator; }
 
 	void AutoCreateDirectorEvaluator();
 
diff --git a/Engine/Plugins/Cameras/GameplayCameras/Source/GameplayCameras/Public/GameFramework/GameplayCameraComponent.h b/Engine/Plugins/Cameras/GameplayCameras/Source/GameplayCameras/Public/GameFramework/GameplayCameraComponent.h
--- a/Engine/Plugins/Cameras/GameplayCameras/Source/GameplayCameras/Public/GameFramework/GameplayCameraComponent.h
+++ b/Engine/Plugins/Cameras/GameplayCameras/Source/GameplayCameras/Public/GameFramework/GameplayCameraComponent.h
@@ -44,7 +44,6 @@ protected:
 
 	// UGameplayCameraComponentBase interface.
 	virtual UCameraAsset* GetCameraAsset() override;
-	virtual bool OnValidateCameraEvaluationContextActivation() override;
 	virtual void OnUpdateCameraEvaluationContext(bool bForceApplyParameterOverrides) override;
 
 private:
diff --git a/Engine/Plugins/Cameras/GameplayCameras/Source/GameplayCameras/Public/GameFramework/GameplayCameraComponentBase.h b/Engine/Plugins/Cameras/GameplayCameras/Source/GameplayCameras/Public/GameFramework/GameplayCameraComponentBase.h
--- a/Engine/Plugins/Cameras/GameplayCameras/Source/GameplayCameras/Public/GameFramework/GameplayCameraComponentBase.h
+++ b/Engine/Plugins/Cameras/GameplayCameras/Source/GameplayCameras/Public/GameFramework/GameplayCameraComponentBase.h
@@ -180,7 +180,6 @@ protected:
 
 	// UGameplayCameraComponentBase interface.
 	virtual UCameraAsset* GetCameraAsset() PURE_VIRTUAL(UGameplayCameraComponentBase::GetCameraAsset, return nullptr;)
-	virtual bool OnValidateCameraEvaluationContextActivation() { return true; }
 	virtual void OnUpdateCameraEvaluationContext(bool bForceApplyParameterOverrides) {}
 
 	void UpdateCameraEvaluationContext(bool bForceApplyParameterOverrides);
@@ -200,10 +199,11 @@ private:
 
 	bool CanRunCameraSystem() const;
 	void EnsureCameraSystemHost();
-	void TeardownCameraSystemHost();
+
+	void CreateCameraEvaluationContext(APlayerController* PlayerController);
+	void DestroyCameraEvaluationContext();
 
 	void ActivateCameraEvaluationContext(APlayerController* PlayerController, IGameplayCameraSystemHost* Host, EGameplayCameraComponentActivationMode ActivationMode);
-	void TryCreateCameraEvaluationContext(APlayerController* PlayerController);
 	void UpdateOutputCameraComponent();
 	void DeactivateCameraEvaluationContext(bool bImmediately);
 	void CheckPendingDeactivation();
@@ -259,16 +259,24 @@ private:
 	/** Evaluation context for running base/global/visual camera rigs. */
 	TSharedPtr<FCameraEvaluationContext> NullContext;
 
+	/** The camera system in which to check for running rigs before deactivating our evaluation context. */
+	TWeakPtr<FCameraSystemEvaluator> PendingDeactivateCameraSystemEvaluator;
+
+	/** Whether to force a camera cut next frame. */
 	bool bIsCameraCutNextFrame = false;
 
-	bool bIsDeactivating = false;
+	/** Whether to check for a pending deactivation. */
+	bool bIsPendingDeactivate = false;
 
 #if WITH_EDITOR
 	
+	/** Whether this component is running in an editor world. */
 	bool bIsEditorWorld = false;
 
+	/** Whether this component has never been activated so far. */
 	bool bIsFirstActivation = true;
 
+	/** The show flag for camera system debug rendering. */
 	int32 CustomShowFlag = INDEX_NONE;
 
 #endif  // WITH_EDITOR
diff --git a/Engine/Plugins/Cameras/GameplayCameras/Source/GameplayCameras/Public/GameFramework/GameplayCameraRigComponent.h b/Engine/Plugins/Cameras/GameplayCameras/Source/GameplayCameras/Public/GameFramework/GameplayCameraRigComponent.h
--- a/Engine/Plugins/Cameras/GameplayCameras/Source/GameplayCameras/Public/GameFramework/GameplayCameraRigComponent.h
+++ b/Engine/Plugins/Cameras/GameplayCameras/Source/GameplayCameras/Public/GameFramework/GameplayCameraRigComponent.h
@@ -43,7 +43,6 @@ protected:
 
 	// UGameplayCameraComponentBase interface.
 	virtual UCameraAsset* GetCameraAsset() override;
-	virtual bool OnValidateCameraEvaluationContextActivation() override;
 	virtual void OnUpdateCameraEvaluationContext(bool bForceApplyParameterOverrides) override;
 
 private:
