Cameras: make blend time a parameter

diff --git a/Engine/Plugins/Cameras/GameplayCameras/Source/GameplayCameras/Private/Nodes/Blends/SimpleBlendCameraNode.cpp b/Engine/Plugins/Cameras/GameplayCameras/Source/GameplayCameras/Private/Nodes/Blends/SimpleBlendCameraNode.cpp
--- a/Engine/Plugins/Cameras/GameplayCameras/Source/GameplayCameras/Private/Nodes/Blends/SimpleBlendCameraNode.cpp
+++ b/Engine/Plugins/Cameras/GameplayCameras/Source/GameplayCameras/Private/Nodes/Blends/SimpleBlendCameraNode.cpp
@@ -92,7 +92,8 @@ void FSimpleFixedTimeBlendCameraNodeEvaluator::OnInitialize(const FCameraNodeEva
 	Super::OnInitialize(Params, OutResult);
 
 	const USimpleFixedTimeBlendCameraNode* BlendNode = GetCameraNodeAs<USimpleFixedTimeBlendCameraNode>();
-	TotalTime = BlendNode->BlendTime;
+	BlendTimeReader.Initialize(BlendNode->BlendTime);
+	TotalTime = BlendTimeReader.Get(OutResult.VariableTable);
 }
 
 void FSimpleFixedTimeBlendCameraNodeEvaluator::OnRun(const FCameraNodeEvaluationParams& Params, FCameraNodeEvaluationResult& OutResult)
diff --git a/Engine/Plugins/Cameras/GameplayCameras/Source/GameplayCameras/Private/Services/CameraShakeService.cpp b/Engine/Plugins/Cameras/GameplayCameras/Source/GameplayCameras/Private/Services/CameraShakeService.cpp
--- a/Engine/Plugins/Cameras/GameplayCameras/Source/GameplayCameras/Private/Services/CameraShakeService.cpp
+++ b/Engine/Plugins/Cameras/GameplayCameras/Source/GameplayCameras/Private/Services/CameraShakeService.cpp
@@ -253,7 +253,7 @@ void FCameraShakeServiceCameraNodeEvaluator::OnRun(const FCameraNodeEvaluationPa
 
 		// Update blends.
 		if (Entry.CameraShake->BlendIn && 
-				Entry.CurrentTime < Entry.CameraShake->BlendIn->BlendTime &&
+				Entry.CurrentTime < Entry.CameraShake->BlendIn->BlendTime.GetValue(OutResult.VariableTable) &&
 				ensure(Entry.BlendInEvaluator))
 		{
 			Entry.BlendInEvaluator->Run(CurParams, CurResult);
@@ -263,7 +263,7 @@ void FCameraShakeServiceCameraNodeEvaluator::OnRun(const FCameraNodeEvaluationPa
 			Entry.BlendInEvaluator->BlendResults(BlendParams, BlendResult);
 		}
 		else if (Entry.CameraShake->BlendOut &&
-				CurTimeLeft >= 0.f && CurTimeLeft < Entry.CameraShake->BlendOut->BlendTime &&
+				CurTimeLeft >= 0.f && CurTimeLeft < Entry.CameraShake->BlendOut->BlendTime.GetValue(OutResult.VariableTable) &&
 				ensure(Entry.BlendOutEvaluator))
 		{
 			Entry.BlendOutEvaluator->Run(CurParams, CurResult);
diff --git a/Engine/Plugins/Cameras/GameplayCameras/Source/GameplayCameras/Public/Nodes/Blends/SimpleBlendCameraNode.h b/Engine/Plugins/Cameras/GameplayCameras/Source/GameplayCameras/Public/Nodes/Blends/SimpleBlendCameraNode.h
--- a/Engine/Plugins/Cameras/GameplayCameras/Source/GameplayCameras/Public/Nodes/Blends/SimpleBlendCameraNode.h
+++ b/Engine/Plugins/Cameras/GameplayCameras/Source/GameplayCameras/Public/Nodes/Blends/SimpleBlendCameraNode.h
@@ -3,7 +3,8 @@
 #pragma once
 
 #include "Core/BlendCameraNode.h"
-
+#include "Core/CameraParameters.h"
+#include "Core/CameraParameterReader.h"
 #include "SimpleBlendCameraNode.generated.h"
 
 /**
@@ -24,14 +25,18 @@ class USimpleFixedTimeBlendCameraNode : public USimpleBlendCameraNode
 	GENERATED_BODY()
 
 public:
+	/** Returns the volume of this element */
+	UE_DEPRECATED(5.7, "Now sets the default time, use SetDefaultBlendTime instead")
+	GAMEPLAYCAMERAS_API void SetBlendTime(float BlendTimeIn) { BlendTime.Value = BlendTimeIn; }
 
-	void SetBlendTime(float BlendTimeIn) { BlendTime = BlendTimeIn; }
+	GAMEPLAYCAMERAS_API void SetDefaultBlendTime(float BlendTimeIn) { BlendTime.Value = BlendTimeIn; }
 
+	GAMEPLAYCAMERAS_API float GetDefaultBlendTime() { return BlendTime.Value; }
 public:
 
 	/** Duration of the blend. */
 	UPROPERTY(EditAnywhere, Category=Blending)
-	float BlendTime = 1.f;
+	FFloatCameraParameter BlendTime = 1.f;
 };
 
 namespace UE::Cameras
@@ -84,6 +89,8 @@ protected:
 	GAMEPLAYCAMERAS_API virtual void OnBuildDebugBlocks(const FCameraDebugBlockBuildParams& Params, FCameraDebugBlockBuilder& Builder) override;
 #endif  // UE_GAMEPLAY_CAMERAS_DEBUG
 
+	TCameraParameterReader<float> BlendTimeReader;
+
 private:
 
 	float BlendFactor = 0.f;
