Cameras: rewind debugger fixes and improvements

- Fix possible crash if PIE stops and the "visualized world" goes away before we can deregister the debug draw delegate.
- Look at the level editor viewport to determine if we are outside or inside of the camera system.

diff --git a/Engine/Plugins/Cameras/GameplayCameras/Source/GameplayCamerasEditor/Private/Trace/CameraSystemRewindDebuggerExtension.cpp b/Engine/Plugins/Cameras/GameplayCameras/Source/GameplayCamerasEditor/Private/Trace/CameraSystemRewindDebuggerExtension.cpp
--- a/Engine/Plugins/Cameras/GameplayCameras/Source/GameplayCamerasEditor/Private/Trace/CameraSystemRewindDebuggerExtension.cpp
+++ b/Engine/Plugins/Cameras/GameplayCameras/Source/GameplayCamerasEditor/Private/Trace/CameraSystemRewindDebuggerExtension.cpp
@@ -2,14 +2,16 @@
 
 #include "Trace/CameraSystemRewindDebuggerExtension.h"
 
-#include "Core/CameraPose.h"
 #include "Debug/CameraDebugRenderer.h"
 #include "Debug/CameraSystemTrace.h"
-#include "Debug/CategoryTitleDebugBlock.h"
 #include "Debug/DebugDrawService.h"
 #include "Debug/RootCameraDebugBlock.h"
 #include "Engine/Canvas.h"
 #include "Engine/World.h"
+#include "LevelEditor.h"
+#include "LevelEditorViewport.h"
+#include "Modules/ModuleManager.h"
+#include "SLevelViewport.h"
 #include "Trace/CameraSystemTraceProvider.h"
 
 #if UE_GAMEPLAY_CAMERAS_TRACE
@@ -70,7 +72,7 @@ void FCameraSystemRewindDebuggerExtension::Update(float DeltaTime, IRewindDebugg
 		if (FoundFrameData && CurrentTraceTime != LastTraceTime)
 		{
 			LastTraceTime = CurrentTraceTime;
-			VisualizedWorld = RewindDebugger->GetWorldToVisualize();
+			WeakVisualizedWorld = RewindDebugger->GetWorldToVisualize();
 
 			DebugBlockStorage.DestroyDebugBlocks();
 
@@ -91,7 +93,7 @@ void FCameraSystemRewindDebuggerExtension::Clear(IRewindDebugger* RewindDebugger
 {
 	EnsureDebugDrawDelegate(false);
 
-	VisualizedWorld = nullptr;
+	WeakVisualizedWorld = nullptr;
 
 	DebugBlockStorage.DestroyDebugBlocks(true);
 	RootDebugBlock = nullptr;
@@ -113,9 +115,23 @@ void FCameraSystemRewindDebuggerExtension::EnsureDebugDrawDelegate(bool bIsRegis
 
 void FCameraSystemRewindDebuggerExtension::DebugDraw(UCanvas* Canvas, APlayerController* PlayController)
 {
+	if (LevelEditorModule == nullptr)
+	{
+		LevelEditorModule = FModuleManager::GetModulePtr<FLevelEditorModule>("LevelEditor");
+	}
+
+	UWorld* VisualizedWorld = WeakVisualizedWorld.Get();
 	if (RootDebugBlock && VisualizedWorld)
 	{
-		FCameraDebugRenderer CameraDebugRenderer(VisualizedWorld, Canvas);
+		bool bIsExternalRendering = false;
+		if (LevelEditorModule)
+		{
+			TSharedPtr<SLevelViewport> LevelViewport = LevelEditorModule->GetFirstActiveLevelViewport();
+			FLevelEditorViewportClient& LevelViewportClient = LevelViewport->GetLevelViewportClient();
+			bIsExternalRendering = !LevelViewportClient.IsAnyActorLocked();
+		}
+
+		FCameraDebugRenderer CameraDebugRenderer(VisualizedWorld, Canvas, bIsExternalRendering);
 		RootDebugBlock->RootDebugDraw(CameraDebugRenderer);
 	}
 }
diff --git a/Engine/Plugins/Cameras/GameplayCameras/Source/GameplayCamerasEditor/Private/Trace/CameraSystemRewindDebuggerExtension.h b/Engine/Plugins/Cameras/GameplayCameras/Source/GameplayCamerasEditor/Private/Trace/CameraSystemRewindDebuggerExtension.h
--- a/Engine/Plugins/Cameras/GameplayCameras/Source/GameplayCamerasEditor/Private/Trace/CameraSystemRewindDebuggerExtension.h
+++ b/Engine/Plugins/Cameras/GameplayCameras/Source/GameplayCamerasEditor/Private/Trace/CameraSystemRewindDebuggerExtension.h
@@ -8,10 +8,12 @@
 #include "IRewindDebugger.h"
 #include "IRewindDebuggerExtension.h"
 #include "Templates/SharedPointerFwd.h"
+#include "UObject/WeakObjectPtrFwd.h"
 
 #if UE_GAMEPLAY_CAMERAS_TRACE
 
 class APlayerController;
+class FLevelEditorModule;
 class UCanvas;
 
 namespace UE::Cameras
@@ -45,7 +47,8 @@ private:
 	FDelegateHandle DebugDrawDelegateHandle;
 	double LastTraceTime = 0.f;
 
-	UWorld* VisualizedWorld = nullptr;
+	FLevelEditorModule* LevelEditorModule = nullptr;
+	TWeakObjectPtr<UWorld> WeakVisualizedWorld;
 
 	FCameraDebugBlockStorage DebugBlockStorage;
 	FRootCameraDebugBlock* RootDebugBlock = nullptr;
