Cameras: fix issues with Driven Control Rotation input node

Make sure that whatever control rotation offsets the node applies, the resulting values are always within the limits defined by the player camera manager.

diff --git a/Engine/Plugins/Cameras/GameplayCameras/Source/GameplayCameras/Private/Nodes/Input/DrivenControlRotationCameraNode.cpp b/Engine/Plugins/Cameras/GameplayCameras/Source/GameplayCameras/Private/Nodes/Input/DrivenControlRotationCameraNode.cpp
--- a/Engine/Plugins/Cameras/GameplayCameras/Source/GameplayCameras/Private/Nodes/Input/DrivenControlRotationCameraNode.cpp
+++ b/Engine/Plugins/Cameras/GameplayCameras/Source/GameplayCameras/Private/Nodes/Input/DrivenControlRotationCameraNode.cpp
@@ -246,6 +246,17 @@ void FDrivenControlRotationCameraNodeEvaluator::OnRun(const FCameraNodeEvaluatio
 			InputValue.X += DeltaRotation.Yaw;
 			InputValue.Y += DeltaRotation.Pitch;
 
+			// If we have a player camera manager, apply its limits to the correct angles.
+			if (APlayerCameraManager* CameraManager = PlayerController->PlayerCameraManager)
+			{
+				FRotator3d LocalControlRotation(InputValue.Y, InputValue.X, 0.0);
+				CameraManager->LimitViewPitch(LocalControlRotation, CameraManager->ViewPitchMin, CameraManager->ViewPitchMax);
+				CameraManager->LimitViewYaw(LocalControlRotation, CameraManager->ViewYawMin, CameraManager->ViewYawMax);
+				CameraManager->LimitViewRoll(LocalControlRotation, CameraManager->ViewRollMin, CameraManager->ViewRollMax);
+				InputValue.X = LocalControlRotation.Yaw;
+				InputValue.Y = LocalControlRotation.Pitch;
+			}
+
 #if UE_GAMEPLAY_CAMERAS_DEBUG
 			LastDeltaControlRotation = DeltaRotation;
 #endif  // UE_GAMEPLAY_CAMERAS_DEBUG
@@ -267,8 +278,19 @@ void FDrivenControlRotationCameraNodeEvaluator::OnExecuteOperation(const FCamera
 		if (bLastWasActiveCameraRig && PlayerController)
 		{
 			FRotator3d ControlRotation = PlayerController->GetControlRotation();
-			ControlRotation.Yaw = Op->Yaw.Apply(ControlRotation.Yaw);
-			ControlRotation.Pitch = Op->Pitch.Apply(ControlRotation.Pitch);
+
+			// Make sure the corrected yaw/pitch angles are in [0..360[ and ]-180..180] respectively.
+			ControlRotation.Yaw = FRotator3d::ClampAxis(Op->Yaw.Apply(ControlRotation.Yaw));
+			ControlRotation.Pitch = FRotator3d::NormalizeAxis(Op->Pitch.Apply(ControlRotation.Pitch));
+
+			// If we have a player camera manager, apply its limits to the correct angles.
+			if (APlayerCameraManager* CameraManager = PlayerController->PlayerCameraManager)
+			{
+				CameraManager->LimitViewPitch(ControlRotation, CameraManager->ViewPitchMin, CameraManager->ViewPitchMax);
+				CameraManager->LimitViewYaw(ControlRotation, CameraManager->ViewYawMin, CameraManager->ViewYawMax);
+				CameraManager->LimitViewRoll(ControlRotation, CameraManager->ViewRollMin, CameraManager->ViewRollMax);
+			}
+
 			PlayerController->SetControlRotation(ControlRotation);
 		}
 	}
