Cameras: make yaw/pitch debug info consistent across a few nodes

diff --git a/Engine/Plugins/Cameras/GameplayCameras/Source/GameplayCameras/Private/Nodes/Common/BoomArmCameraNode.cpp b/Engine/Plugins/Cameras/GameplayCameras/Source/GameplayCameras/Private/Nodes/Common/BoomArmCameraNode.cpp
--- a/Engine/Plugins/Cameras/GameplayCameras/Source/GameplayCameras/Private/Nodes/Common/BoomArmCameraNode.cpp
+++ b/Engine/Plugins/Cameras/GameplayCameras/Source/GameplayCameras/Private/Nodes/Common/BoomArmCameraNode.cpp
@@ -251,7 +251,7 @@ void FBoomArmCameraNodeEvaluator::OnBuildDebugBlocks(const FCameraDebugBlockBuil
 
 void FBoomArmCameraDebugBlock::OnDebugDraw(const FCameraDebugBlockDrawParams& Params, FCameraDebugRenderer& Renderer)
 {
-	Renderer.AddText(TEXT("yaw: %.3f pitch %.3f"), BoomYawPitch.X, BoomYawPitch.Y);
+	Renderer.AddText(TEXT("yaw: %.3f pitch: %.3f"), BoomYawPitch.X, BoomYawPitch.Y);
 
 	if (bHasBoomLengthInterpolator)
 	{
diff --git a/Engine/Plugins/Cameras/GameplayCameras/Source/GameplayCameras/Private/Nodes/Common/SplineOrbitCameraNode.cpp b/Engine/Plugins/Cameras/GameplayCameras/Source/GameplayCameras/Private/Nodes/Common/SplineOrbitCameraNode.cpp
--- a/Engine/Plugins/Cameras/GameplayCameras/Source/GameplayCameras/Private/Nodes/Common/SplineOrbitCameraNode.cpp
+++ b/Engine/Plugins/Cameras/GameplayCameras/Source/GameplayCameras/Private/Nodes/Common/SplineOrbitCameraNode.cpp
@@ -326,7 +326,7 @@ void FSplineOrbitCameraNodeEvaluator::OnBuildDebugBlocks(const FCameraDebugBlock
 
 void FSplineOrbitCameraDebugBlock::OnDebugDraw(const FCameraDebugBlockDrawParams& Params, FCameraDebugRenderer& Renderer)
 {
-	Renderer.AddText(TEXT("yaw: %.3f pitch %.3f"), OrbitYawPitch.X, OrbitYawPitch.Y);
+	Renderer.AddText(TEXT("yaw: %.3f pitch: %.3f"), OrbitYawPitch.X, OrbitYawPitch.Y);
 	Renderer.NewLine();
 	Renderer.AddIndent();
 	{
diff --git a/Engine/Plugins/Cameras/GameplayCameras/Source/GameplayCameras/Private/Nodes/Input/DrivenControlRotationCameraNode.cpp b/Engine/Plugins/Cameras/GameplayCameras/Source/GameplayCameras/Private/Nodes/Input/DrivenControlRotationCameraNode.cpp
--- a/Engine/Plugins/Cameras/GameplayCameras/Source/GameplayCameras/Private/Nodes/Input/DrivenControlRotationCameraNode.cpp
+++ b/Engine/Plugins/Cameras/GameplayCameras/Source/GameplayCameras/Private/Nodes/Input/DrivenControlRotationCameraNode.cpp
@@ -324,13 +324,13 @@ void FDrivenControlRotationDebugBlock::OnDebugDraw(const FCameraDebugBlockDrawPa
 {
 	if (bLastWasActiveCameraRig)
 	{
-		Renderer.AddText(TEXT("Active ; P=%.1f Y=%.1f"), InputValue.Y, InputValue.X);
+		Renderer.AddText(TEXT("[Active] yaw: %.3f pitch: %.3f"), InputValue.X, InputValue.Y);
 	}
 	else
 	{
 		Renderer.AddText(
-				TEXT("Driven ; P=%.1f Y=%.1f (Delta: P=%.1f Y=%.1f"), 
-				InputValue.Y, InputValue.X, LastDeltaControlRotation.Pitch, LastDeltaControlRotation.Yaw);
+				TEXT("[Driven] yaw: %.3f pitch: %.3f (delta yaw: %.3f pitch: %.3f)"), 
+				InputValue.X, InputValue.Y, LastDeltaControlRotation.Yaw, LastDeltaControlRotation.Pitch);
 	}
 }
 
