Cameras: add option for showing context data IDs in debug render

diff --git a/Engine/Plugins/Cameras/GameplayCameras/Source/GameplayCameras/Private/Core/BlendStackCameraNode.cpp b/Engine/Plugins/Cameras/GameplayCameras/Source/GameplayCameras/Private/Core/BlendStackCameraNode.cpp
--- a/Engine/Plugins/Cameras/GameplayCameras/Source/GameplayCameras/Private/Core/BlendStackCameraNode.cpp
+++ b/Engine/Plugins/Cameras/GameplayCameras/Source/GameplayCameras/Private/Core/BlendStackCameraNode.cpp
@@ -17,6 +17,7 @@
 #include "Debug/CameraNodeEvaluationResultDebugBlock.h"
 #include "Debug/CameraPoseDebugBlock.h"
 #include "Debug/CameraPoseLocationTrailDebugBlock.h"
+#include "Debug/ContextDataTableDebugBlock.h"
 #include "Debug/VariableTableDebugBlock.h"
 #include "HAL/IConsoleManager.h"
 #include "Helpers/CameraObjectInterfaceParameterOverrideHelper.h"
@@ -59,6 +60,12 @@ static FAutoConsoleVariableRef CVarGameplayCamerasDebugBlendStackShowVariableIDs
 	GGameplayCamerasDebugBlendStackShowVariableIDs,
 	TEXT(""));
 
+bool GGameplayCamerasDebugBlendStackShowDataIDs = false;
+static FAutoConsoleVariableRef CVarGameplayCamerasDebugBlendStackShowDataIDs(
+	TEXT("GameplayCameras.Debug.BlendStack.ShowDataIDs"),
+	GGameplayCamerasDebugBlendStackShowDataIDs,
+	TEXT(""));
+
 UE_DEFINE_CAMERA_NODE_EVALUATOR(FBlendStackCameraNodeEvaluator)
 
 FBlendStackCameraNodeEvaluator::FBlendStackCameraNodeEvaluator()
@@ -693,6 +700,7 @@ FBlendStackCameraDebugBlock* FBlendStackCameraNodeEvaluator::BuildDetailedDebugB
 				ResultDebugBlock.Initialize(Entry.Result, Builder);
 				ResultDebugBlock.GetCameraPoseDebugBlock()->WithShowUnchangedCVar(TEXT("GameplayCameras.Debug.BlendStack.ShowUnchanged"));
 				ResultDebugBlock.GetVariableTableDebugBlock()->WithShowVariableIDsCVar(TEXT("GameplayCameras.Debug.BlendStack.ShowVariableIDs"));
+				ResultDebugBlock.GetContextDataTableDebugBlock()->WithShowDataIDsCVar(TEXT("GameplayCameras.Debug.BlendStack.ShowDataIDs"));
 			}
 		}
 	}
diff --git a/Engine/Plugins/Cameras/GameplayCameras/Source/GameplayCameras/Private/Debug/ContextDataTableDebugBlock.cpp b/Engine/Plugins/Cameras/GameplayCameras/Source/GameplayCameras/Private/Debug/ContextDataTableDebugBlock.cpp
--- a/Engine/Plugins/Cameras/GameplayCameras/Source/GameplayCameras/Private/Debug/ContextDataTableDebugBlock.cpp
+++ b/Engine/Plugins/Cameras/GameplayCameras/Source/GameplayCameras/Private/Debug/ContextDataTableDebugBlock.cpp
@@ -123,9 +123,26 @@ void FContextDataTableDebugBlock::OnDebugDraw(const FCameraDebugBlockDrawParams&
 {
 	const FCameraDebugColors& Colors = FCameraDebugColors::Get();
 
+#if WITH_EDITORONLY_DATA
+	bool bShowDataIDs = false;
+	if (!ShowDataIDsCVarName.IsEmpty())
+	{
+		IConsoleVariable* ShowDataIDsCVar = IConsoleManager::Get().FindConsoleVariable(*ShowDataIDsCVarName, false);
+		if (ensureMsgf(ShowDataIDsCVar, TEXT("No such console variable: %s"), *ShowDataIDsCVarName))
+		{
+			bShowDataIDs = ShowDataIDsCVar->GetBool();
+		}
+	}
+#endif
+
 	for (const FEntryDebugInfo& Entry : Entries)
 	{
 #if WITH_EDITORONLY_DATA
+		if (bShowDataIDs)
+		{
+			Renderer.AddText(TEXT("{cam_passive}[%d]{cam_default} "), Entry.ID);
+		}
+
 		if (!Entry.Name.IsEmpty())
 		{
 			Renderer.AddText(TEXT("%s [%s] "), *Entry.Name, *Entry.TypeName.ToString());
@@ -164,6 +181,7 @@ void FContextDataTableDebugBlock::OnDebugDraw(const FCameraDebugBlockDrawParams&
 void FContextDataTableDebugBlock::OnSerialize(FArchive& Ar)
 {
 	Ar << Entries;
+	Ar << ShowDataIDsCVarName;
 }
 
 FArchive& operator<< (FArchive& Ar, FContextDataTableDebugBlock::FEntryDebugInfo& EntryDebugInfo)
diff --git a/Engine/Plugins/Cameras/GameplayCameras/Source/GameplayCameras/Private/Debug/RootCameraDebugBlock.cpp b/Engine/Plugins/Cameras/GameplayCameras/Source/GameplayCameras/Private/Debug/RootCameraDebugBlock.cpp
--- a/Engine/Plugins/Cameras/GameplayCameras/Source/GameplayCameras/Private/Debug/RootCameraDebugBlock.cpp
+++ b/Engine/Plugins/Cameras/GameplayCameras/Source/GameplayCameras/Private/Debug/RootCameraDebugBlock.cpp
@@ -8,13 +8,12 @@
 #include "Core/RootCameraNode.h"
 #include "Debug/CameraDebugBlockBuilder.h"
 #include "Debug/CameraDebugCategories.h"
-#include "Debug/CameraDebugColors.h"
 #include "Debug/CameraDebugRenderer.h"
 #include "Debug/CameraDirectorTreeDebugBlock.h"
 #include "Debug/CameraEvaluationServiceDebugBlock.h"
 #include "Debug/CameraNodeEvaluationResultDebugBlock.h"
-#include "Debug/CameraNodeEvaluatorDebugBlock.h"
 #include "Debug/CameraPoseDebugBlock.h"
+#include "Debug/ContextDataTableDebugBlock.h"
 #include "Debug/CategoryTitleDebugBlock.h"
 #include "Debug/PlayerControllersDebugBlock.h"
 #include "Debug/VariableTableDebugBlock.h"
@@ -57,6 +56,12 @@ static FAutoConsoleVariableRef CVarGameplayCamerasDebugPoseStatsShowVariableIDs(
 	GGameplayCamerasDebugPoseStatsShowVariableIDs,
 	TEXT(""));
 
+bool GGameplayCamerasDebugPoseStatsShowDataIDs = false;
+static FAutoConsoleVariableRef CVarGameplayCamerasDebugPoseStatsShowDataIDs(
+	TEXT("GameplayCameras.Debug.PoseStats.ShowDataIDs"),
+	GGameplayCamerasDebugPoseStatsShowDataIDs,
+	TEXT(""));
+
 UE_DEFINE_CAMERA_DEBUG_BLOCK(FRootCameraDebugBlock)
 
 void FRootCameraDebugBlock::BuildDebugBlocks(const FCameraSystemEvaluator& CameraSystem, const FCameraDebugBlockBuildParams& Params, FCameraDebugBlockBuilder& Builder)
@@ -127,6 +132,7 @@ void FRootCameraDebugBlock::BuildDebugBlocks(const FCameraSystemEvaluator& Camer
 			ResultDebugBlock.Initialize(CameraSystem.GetEvaluatedResult(), Builder);
 			ResultDebugBlock.GetCameraPoseDebugBlock()->WithShowUnchangedCVar(TEXT("GameplayCameras.Debug.PoseStats.ShowUnchanged"));
 			ResultDebugBlock.GetVariableTableDebugBlock()->WithShowVariableIDsCVar(TEXT("GameplayCameras.Debug.PoseStats.ShowVariableIDs"));
+			ResultDebugBlock.GetContextDataTableDebugBlock()->WithShowDataIDsCVar(TEXT("GameplayCameras.Debug.PoseStats.ShowDataIDs"));
 		}
 		FPlayerControllersDebugBlock& PlayerControllersDebugBlock = Builder.BuildDebugBlock<FPlayerControllersDebugBlock>();
 		PoseStatsCategory.AddChild(&PlayerControllersDebugBlock);
diff --git a/Engine/Plugins/Cameras/GameplayCameras/Source/GameplayCameras/Public/Debug/ContextDataTableDebugBlock.h b/Engine/Plugins/Cameras/GameplayCameras/Source/GameplayCameras/Public/Debug/ContextDataTableDebugBlock.h
--- a/Engine/Plugins/Cameras/GameplayCameras/Source/GameplayCameras/Public/Debug/ContextDataTableDebugBlock.h
+++ b/Engine/Plugins/Cameras/GameplayCameras/Source/GameplayCameras/Public/Debug/ContextDataTableDebugBlock.h
@@ -26,6 +26,13 @@ public:
 	/** Creates a new context data table debug block. */
 	GAMEPLAYCAMERAS_API FContextDataTableDebugBlock(const FCameraContextDataTable& InContextDataTable);
 
+	/** Specifies the console variable to use to toggle the printing of variable IDs. */
+	FContextDataTableDebugBlock& WithShowDataIDsCVar(const TCHAR* InShowDataIDsCVarName)
+	{
+		ShowDataIDsCVarName = InShowDataIDsCVarName;
+		return *this;
+	}
+
 protected:
 
 	virtual void OnDebugDraw(const FCameraDebugBlockDrawParams& Params, FCameraDebugRenderer& Renderer) override;
@@ -57,6 +64,8 @@ private:
 	};
 	TArray<FEntryDebugInfo> Entries;
 
+	FString ShowDataIDsCVarName;
+
 	friend FArchive& operator<< (FArchive&, FEntryDebugInfo&);
 };
 
diff --git a/Engine/Plugins/Cameras/GameplayCameras/Source/GameplayCamerasEditor/Private/Debugger/SBlendStacksDebugPanel.cpp b/Engine/Plugins/Cameras/GameplayCameras/Source/GameplayCamerasEditor/Private/Debugger/SBlendStacksDebugPanel.cpp
--- a/Engine/Plugins/Cameras/GameplayCameras/Source/GameplayCamerasEditor/Private/Debugger/SBlendStacksDebugPanel.cpp
+++ b/Engine/Plugins/Cameras/GameplayCameras/Source/GameplayCamerasEditor/Private/Debugger/SBlendStacksDebugPanel.cpp
@@ -31,6 +31,13 @@ void SBlendStacksDebugPanel::Construct(const FArguments& InArgs)
 						LOCTEXT("ShowVariableIDs", "Show variable IDs"),
 						TEXT("GameplayCameras.Debug.BlendStack.ShowVariableIDs"))
 			]
+		+ SVerticalBox::Slot()
+			.AutoHeight()
+			[
+				SDebugWidgetUtils::CreateConsoleVariableCheckBox(
+						LOCTEXT("ShowDataIDs", "Show data IDs"),
+						TEXT("GameplayCameras.Debug.BlendStack.ShowDataIDs"))
+			]
 	];
 }
 
diff --git a/Engine/Plugins/Cameras/GameplayCameras/Source/GameplayCamerasEditor/Private/Debugger/SCameraPoseStatsDebugPanel.cpp b/Engine/Plugins/Cameras/GameplayCameras/Source/GameplayCamerasEditor/Private/Debugger/SCameraPoseStatsDebugPanel.cpp
--- a/Engine/Plugins/Cameras/GameplayCameras/Source/GameplayCamerasEditor/Private/Debugger/SCameraPoseStatsDebugPanel.cpp
+++ b/Engine/Plugins/Cameras/GameplayCameras/Source/GameplayCamerasEditor/Private/Debugger/SCameraPoseStatsDebugPanel.cpp
@@ -30,6 +30,13 @@ void SCameraPoseStatsDebugPanel::Construct(const FArguments& InArgs)
 						LOCTEXT("ShowVariableIDs", "Show variable IDs"),
 						TEXT("GameplayCameras.Debug.PoseStats.ShowVariableIDs"))
 			]
+		+ SVerticalBox::Slot()
+			.AutoHeight()
+			[
+				SDebugWidgetUtils::CreateConsoleVariableCheckBox(
+						LOCTEXT("ShowDataIDs", "Show data IDs"),
+						TEXT("GameplayCameras.Debug.PoseStats.ShowDataIDs"))
+			]
 	];
 }
 
