Cameras: big pass on debug info

- Improve the "auto" debug mode: do debug rendering if there's only one active camera system even if it's not the active one.
- Render blend stack entries with colors and labels.
- Add ability to filter blend stack debug rendering by stack name.
- More options to the director tree debug rendering.
- Hide damping nodes' local space rendering (and make it an option) because it often gets in the way.
- Tweak various debug colors and text sizes/positions/etc.
- Fix debug rendering that should only be done in external/debug camera.

diff --git a/Engine/Plugins/Cameras/GameplayCameras/Source/GameplayCameras/Private/Core/BlendStackCameraNode.cpp b/Engine/Plugins/Cameras/GameplayCameras/Source/GameplayCameras/Private/Core/BlendStackCameraNode.cpp
--- a/Engine/Plugins/Cameras/GameplayCameras/Source/GameplayCameras/Private/Core/BlendStackCameraNode.cpp
+++ b/Engine/Plugins/Cameras/GameplayCameras/Source/GameplayCameras/Private/Core/BlendStackCameraNode.cpp
@@ -13,6 +13,7 @@
 #include "Core/PersistentBlendStackCameraNode.h"
 #include "Core/TransientBlendStackCameraNode.h"
 #include "Debug/CameraDebugBlockBuilder.h"
+#include "Debug/CameraDebugColors.h"
 #include "Debug/CameraDebugRenderer.h"
 #include "Debug/CameraNodeEvaluationResultDebugBlock.h"
 #include "Debug/CameraPoseDebugBlock.h"
@@ -646,6 +647,7 @@ void FBlendStackCameraNodeEvaluator::OnPostBuildAsset(const FGameplayCameraAsset
 
 UE_DECLARE_CAMERA_DEBUG_BLOCK_START(GAMEPLAYCAMERAS_API, FBlendStackEntryCameraDebugBlock)
 	UE_DECLARE_CAMERA_DEBUG_BLOCK_FIELD(int32, EntryIndex)
+	UE_DECLARE_CAMERA_DEBUG_BLOCK_FIELD(FLinearColor, EntryColor)
 	UE_DECLARE_CAMERA_DEBUG_BLOCK_FIELD(FBlendStackEntryID, EntryID)
 	UE_DECLARE_CAMERA_DEBUG_BLOCK_FIELD(FString, CameraRigName)
 	UE_DECLARE_CAMERA_DEBUG_BLOCK_FIELD(bool, bIsFrozen)
@@ -662,18 +664,25 @@ void FBlendStackCameraNodeEvaluator::OnBuildDebugBlocks(const FCameraDebugBlockB
 	}
 }
 
-FBlendStackCameraDebugBlock* FBlendStackCameraNodeEvaluator::BuildDetailedDebugBlock(const FCameraDebugBlockBuildParams& Params, FCameraDebugBlockBuilder& Builder)
+FBlendStackCameraDebugBlock* FBlendStackCameraNodeEvaluator::BuildDetailedDebugBlock(
+		const FCameraDebugBlockBuildParams& Params, 
+		const FLinearColor& StartColor, const FLinearColor& EndColor,
+		FCameraDebugBlockBuilder& Builder)
 {
 	FBlendStackCameraDebugBlock& StackDebugBlock = Builder.BuildDebugBlock<FBlendStackCameraDebugBlock>();
 	for (int32 Index = 0; Index < Entries.Num(); ++Index)
 	{
 		const FCameraRigEntry& Entry(Entries[Index]);
 
+		const FLinearColor EntryColor = LerpLinearColorUsingHSV(
+				StartColor, EndColor, Index, Entries.Num());
+
 		// Each entry has a wrapper debug block with 2 children blocks:
 		// - block for the blend
 		// - block for the result
 		FBlendStackEntryCameraDebugBlock& EntryDebugBlock = Builder.BuildDebugBlock<FBlendStackEntryCameraDebugBlock>();
 		EntryDebugBlock.EntryIndex = Index;
+		EntryDebugBlock.EntryColor = EntryColor;
 		EntryDebugBlock.EntryID = Entry.EntryID;
 		EntryDebugBlock.CameraRigName = GetNameSafe(Entry.CameraRig);
 		EntryDebugBlock.bIsFrozen = Entry.Flags.bIsFrozen;
@@ -739,8 +748,6 @@ void FBlendStackSummaryCameraDebugBlock::OnSerialize(FArchive& Ar)
 	Ar << BlendStackLayer;
 }
 
-UE_DEFINE_CAMERA_DEBUG_BLOCK(FBlendStackCameraDebugBlock);
-
 void FBlendStackEntryCameraDebugBlock::OnDebugDraw(const FCameraDebugBlockDrawParams& Params, FCameraDebugRenderer& Renderer)
 {
 	Renderer.AddText(
@@ -751,6 +758,26 @@ void FBlendStackEntryCameraDebugBlock::OnDebugDraw(const FCameraDebugBlockDrawPa
 		Renderer.AddText(" {cam_notice2}FROZEN{cam_default}");
 	}
 	Renderer.NewLine();
+
+	// Render children with some color/label override for the entry's result pose.
+	{
+		FScopedGlobalCameraPoseRenderingParams ScopedParams(*LexToString(EntryIndex + 1), EntryColor);
+
+		TArrayView<FCameraDebugBlock*> ChildrenBlocks(GetChildren());
+		for (FCameraDebugBlock* ChildBlock : ChildrenBlocks)
+		{
+			ChildBlock->DebugDraw(Params, Renderer);
+		}
+
+		Renderer.SkipChildrenBlocks();
+	}
+}
+
+UE_DEFINE_CAMERA_DEBUG_BLOCK(FBlendStackCameraDebugBlock);
+
+void FBlendStackCameraDebugBlock::OnSerialize(FArchive& Ar)
+{
+	Ar << StartEndColors;
 }
 
 #endif  // UE_GAMEPLAY_CAMERAS_DEBUG
diff --git a/Engine/Plugins/Cameras/GameplayCameras/Source/GameplayCameras/Private/Core/CameraIKAim.cpp b/Engine/Plugins/Cameras/GameplayCameras/Source/GameplayCameras/Private/Core/CameraIKAim.cpp
--- a/Engine/Plugins/Cameras/GameplayCameras/Source/GameplayCameras/Private/Core/CameraIKAim.cpp
+++ b/Engine/Plugins/Cameras/GameplayCameras/Source/GameplayCameras/Private/Core/CameraIKAim.cpp
@@ -494,12 +494,13 @@ void FCameraIKAimDebugInfo::DebugDraw(const FCameraDebugBlockDrawParams& Params,
 	}
 	Renderer.AddText(TEXT("{cam_default} in %d iterations\n"), Iterations.Num());
 
-	UFont* TinyFont = GEngine->GetTinyFont();
-	Renderer.DrawSphere(DesiredTarget, 1.f, 8, FLinearColor::Yellow, 1.f);
+	UFont* SmallFont = GEngine->GetSmallFont();
+	Renderer.DrawPoint(DesiredTarget, 2.f, FLinearColor::Yellow, 2.f);
 
 	Renderer.AddIndent();
 	{
-		int32 IterationIndex = 1;
+		int32 IterationIndex = 0;
+		const FVector2d TextOffset(-20, -20);
 		for (const FCameraIKAimIterationDebugInfo& IterationDebugInfo : Iterations)
 		{
 			const bool bDirectionIsNormalized = true;
@@ -508,19 +509,23 @@ void FCameraIKAimDebugInfo::DebugDraw(const FCameraDebugBlockDrawParams& Params,
 					IterationDebugInfo.CameraPoseLocation, 
 					IterationDebugInfo.CameraPoseRotation.RotateVector(TargetDir), bDirectionIsNormalized);
 
+			const FLinearColor IterationColor = LerpLinearColorUsingHSV(
+					FLinearColor::Yellow, FLinearColor::Red, IterationIndex, Iterations.Num());
+
 			Renderer.DrawLine(
 					IterationDebugInfo.CameraPoseLocation,
 					DirectionRay.PointAt(1000.0),
-					FLinearColor::Yellow);
-			Renderer.DrawText(
+					IterationColor);
+			Renderer.DrawTextView(
 					IterationDebugInfo.CameraPoseLocation,
-					FString::Format(TEXT("Iteration {0}"), { IterationIndex }),
-					FLinearColor::Yellow,
-					TinyFont);
+					TextOffset,
+					FString::Format(TEXT("Iteration {0}"), { IterationIndex + 1 }),
+					IterationColor,
+					SmallFont);
 
 			Renderer.AddText(
 					TEXT("%d : error angle %.2fdeg, error distance %.1fcm, "),
-					IterationIndex,
+					IterationIndex + 1,
 					IterationDebugInfo.ErrorAngle, IterationDebugInfo.ErrorDistance);
 
 			++IterationIndex;
diff --git a/Engine/Plugins/Cameras/GameplayCameras/Source/GameplayCameras/Private/Core/CameraRigCombinationRegistry.cpp b/Engine/Plugins/Cameras/GameplayCameras/Source/GameplayCameras/Private/Core/CameraRigCombinationRegistry.cpp
--- a/Engine/Plugins/Cameras/GameplayCameras/Source/GameplayCameras/Private/Core/CameraRigCombinationRegistry.cpp
+++ b/Engine/Plugins/Cameras/GameplayCameras/Source/GameplayCameras/Private/Core/CameraRigCombinationRegistry.cpp
@@ -3,6 +3,9 @@
 #include "Core/CameraRigCombinationRegistry.h"
 
 #include "Core/CameraRigAsset.h"
+#include "Debug/CameraDebugBlock.h"
+#include "Debug/CameraDebugBlockBuilder.h"
+#include "Debug/CameraDebugRenderer.h"
 #include "Helpers/CameraObjectInterfaceParameterOverrideHelper.h"
 
 #include UE_INLINE_GENERATED_CPP_BY_NAME(CameraRigCombinationRegistry)
@@ -68,6 +71,10 @@ protected:
 	virtual void OnBuild(const FCameraNodeEvaluatorBuildParams& Params) override;
 	virtual void OnRun(const FCameraNodeEvaluationParams& Params, FCameraNodeEvaluationResult& OutResult) override;
 
+#if UE_GAMEPLAY_CAMERAS_DEBUG
+	virtual void OnBuildDebugBlocks(const FCameraDebugBlockBuildParams& Params, FCameraDebugBlockBuilder& Builder) override;
+#endif  // UE_GAMEPLAY_CAMERAS_DEBUG
+
 private:
 
 	void ApplyParameterOverrides(FCameraVariableTable& OutVariableTable, bool bDrivenOnly);
@@ -79,6 +86,13 @@ private:
 
 UE_DEFINE_CAMERA_NODE_EVALUATOR(FCombinedCameraRigsCameraNodeEvaluator)
 
+UE_DECLARE_CAMERA_DEBUG_BLOCK_START(GAMEPLAYCAMERAS_API, FCombinedCameraRigCameraDebugBlock)
+	UE_DECLARE_CAMERA_DEBUG_BLOCK_FIELD(int32, CameraRigIndex);
+	UE_DECLARE_CAMERA_DEBUG_BLOCK_FIELD(FString, CameraRigName);
+UE_DECLARE_CAMERA_DEBUG_BLOCK_END()
+
+UE_DEFINE_CAMERA_DEBUG_BLOCK_WITH_FIELDS(FCombinedCameraRigCameraDebugBlock)
+
 FCameraNodeEvaluatorChildrenView FCombinedCameraRigsCameraNodeEvaluator::OnGetChildren()
 {
 	return FCameraNodeEvaluatorChildrenView(CameraRigRootEvaluators);
@@ -129,6 +143,47 @@ void FCombinedCameraRigsCameraNodeEvaluator::ApplyParameterOverrides(FCameraVari
 	}
 }
 
+#if UE_GAMEPLAY_CAMERAS_DEBUG
+
+void FCombinedCameraRigsCameraNodeEvaluator::OnBuildDebugBlocks(const FCameraDebugBlockBuildParams& Params, FCameraDebugBlockBuilder& Builder)
+{
+	// Wrap our children inside debug blocks that show their camera rig name.
+	for (int32 Index = 0; Index < CameraRigRootEvaluators.Num(); ++Index)
+	{
+		FCameraNodeEvaluator* CameraRigRootEvaluator(CameraRigRootEvaluators[Index]);
+		if (!CameraRigRootEvaluator)
+		{
+			continue;
+		}
+
+		FString CameraRigName;
+		if (const UCameraNode* CameraRigRootNode = CameraRigRootEvaluator->GetCameraNode())
+		{
+			if (const UCameraRigAsset* CameraRig = CameraRigRootNode->GetTypedOuter<UCameraRigAsset>())
+			{
+				CameraRigName = CameraRig->GetName();
+			}
+		}
+
+		FCombinedCameraRigCameraDebugBlock& ChildBlock = Builder.StartChildDebugBlock<FCombinedCameraRigCameraDebugBlock>();
+		ChildBlock.CameraRigIndex = Index;
+		ChildBlock.CameraRigName = CameraRigName;
+		{
+			CameraRigRootEvaluator->BuildDebugBlocks(Params, Builder);
+		}
+		Builder.EndChildDebugBlock();
+	}
+
+	Builder.SkipChildren();
+}
+
+void FCombinedCameraRigCameraDebugBlock::OnDebugDraw(const FCameraDebugBlockDrawParams& Params, FCameraDebugRenderer& Renderer)
+{
+	Renderer.AddText(TEXT("{cam_passive}[%d] {cam_notice}%s{cam_default}\n"), CameraRigIndex + 1, *CameraRigName);
+}
+
+#endif  // UE_GAMEPLAY_CAMERAS_DEBUG
+
 }  // namespace UE::Cameras
 
 void UCombinedCameraRigsCameraNode::GetAllCombinationCameraRigs(const UCameraRigAsset* InCameraRig, TArray<const UCameraRigAsset*>& OutCameraRigs)
diff --git a/Engine/Plugins/Cameras/GameplayCameras/Source/GameplayCameras/Private/Core/CameraSystemEvaluator.cpp b/Engine/Plugins/Cameras/GameplayCameras/Source/GameplayCameras/Private/Core/CameraSystemEvaluator.cpp
--- a/Engine/Plugins/Cameras/GameplayCameras/Source/GameplayCameras/Private/Core/CameraSystemEvaluator.cpp
+++ b/Engine/Plugins/Cameras/GameplayCameras/Source/GameplayCameras/Private/Core/CameraSystemEvaluator.cpp
@@ -534,7 +534,8 @@ void FCameraSystemEvaluator::DrawEditorPreview(const FCameraSystemEditorPreviewP
 
 	if (RootDebugBlock)
 	{
-		RootDebugBlock->RootDebugDraw(Renderer);
+		FRootCameraDebugDrawParams RootParams;
+		RootDebugBlock->RootDebugDraw(RootParams, Renderer);
 	}
 	else
 	{
@@ -609,8 +610,11 @@ void FCameraSystemEvaluator::DebugUpdate(const FCameraSystemDebugUpdateParams& P
 	}
 #endif
 
+	FRootCameraDebugDrawParams RootParams;
+	RootParams.bIsCameraManagerOrViewTarget = Params.bIsCameraManagerOrViewTarget;
+	RootParams.bForceDraw = Params.bForceDraw;
 	FCameraDebugRenderer Renderer(OwnerWorld, Params.CanvasObject, Params.bIsDebugCameraEnabled);
-	RootDebugBlock->RootDebugDraw(Renderer, Params.bForceDraw);
+	RootDebugBlock->RootDebugDraw(RootParams, Renderer);
 }
 
 #endif  // UE_GAMEPLAY_CAMERAS_DEBUG
diff --git a/Engine/Plugins/Cameras/GameplayCameras/Source/GameplayCameras/Private/Core/DefaultRootCameraNode.cpp b/Engine/Plugins/Cameras/GameplayCameras/Source/GameplayCameras/Private/Core/DefaultRootCameraNode.cpp
--- a/Engine/Plugins/Cameras/GameplayCameras/Source/GameplayCameras/Private/Core/DefaultRootCameraNode.cpp
+++ b/Engine/Plugins/Cameras/GameplayCameras/Source/GameplayCameras/Private/Core/DefaultRootCameraNode.cpp
@@ -315,10 +315,10 @@ void FDefaultRootCameraNodeEvaluator::OnBuildDebugBlocks(const FCameraDebugBlock
 	// Create the debug block that shows the overall blend stack layers.
 	FBlendStacksCameraDebugBlock& DebugBlock = Builder.BuildDebugBlock<FBlendStacksCameraDebugBlock>();
 	{
-		DebugBlock.AddBlendStack(TEXT("Base Layer"), BaseLayer->BuildDetailedDebugBlock(Params, Builder));
-		DebugBlock.AddBlendStack(TEXT("Main Layer"), MainLayer->BuildDetailedDebugBlock(Params, Builder));
-		DebugBlock.AddBlendStack(TEXT("Global Layer"), GlobalLayer->BuildDetailedDebugBlock(Params, Builder));
-		DebugBlock.AddBlendStack(TEXT("Visual Layer"), VisualLayer->BuildDetailedDebugBlock(Params, Builder));
+		DebugBlock.AddBlendStack(TEXT("Base Layer"), BaseLayer->BuildDetailedDebugBlock(Params, FColorList::Grey, FColorList::LightGrey, Builder));
+		DebugBlock.AddBlendStack(TEXT("Main Layer"), MainLayer->BuildDetailedDebugBlock(Params, FColorList::LightBlue, FColorList::SlateBlue, Builder));
+		DebugBlock.AddBlendStack(TEXT("Global Layer"), GlobalLayer->BuildDetailedDebugBlock(Params, FColorList::Orange, FColorList::OrangeRed, Builder));
+		DebugBlock.AddBlendStack(TEXT("Visual Layer"), VisualLayer->BuildDetailedDebugBlock(Params, FColorList::Pink, FColorList::NeonPink, Builder));
 	}
 
 	Builder.GetRootDebugBlock().AddChild(&DebugBlock);
diff --git a/Engine/Plugins/Cameras/GameplayCameras/Source/GameplayCameras/Private/Debug/BlendStacksCameraDebugBlock.cpp b/Engine/Plugins/Cameras/GameplayCameras/Source/GameplayCameras/Private/Debug/BlendStacksCameraDebugBlock.cpp
--- a/Engine/Plugins/Cameras/GameplayCameras/Source/GameplayCameras/Private/Debug/BlendStacksCameraDebugBlock.cpp
+++ b/Engine/Plugins/Cameras/GameplayCameras/Source/GameplayCameras/Private/Debug/BlendStacksCameraDebugBlock.cpp
@@ -6,14 +6,19 @@
 #include "Debug/CameraDebugCategories.h"
 #include "Debug/CameraDebugColors.h"
 #include "Debug/CameraDebugRenderer.h"
-#include "Debug/CameraNodeEvaluatorDebugBlock.h"
-#include "Math/ColorList.h"
+#include "HAL/IConsoleManager.h"
 
 #if UE_GAMEPLAY_CAMERAS_DEBUG
 
 namespace UE::Cameras
 {
 
+FString GGameplayCamerasDebugBlendStacksFilter;
+static FAutoConsoleVariableRef CVarGameplayCamerasDebugBlendStacksFilter(
+	TEXT("GameplayCameras.Debug.BlendStacks.Filter"),
+	GGameplayCamerasDebugBlendStacksFilter,
+	TEXT(""));
+
 UE_DEFINE_CAMERA_DEBUG_BLOCK(FBlendStacksCameraDebugBlock)
 
 FBlendStacksCameraDebugBlock::FBlendStacksCameraDebugBlock()
@@ -42,21 +47,29 @@ void FBlendStacksCameraDebugBlock::OnDebugDraw(const FCameraDebugBlockDrawParams
 
 	TArrayView<FCameraDebugBlock*> BlendStackBlocks(GetChildren());
 
-	const int32 MaxIndex = FMath::Max(BlendStackBlocks.Num(), BlendStackNames.Num());
+	const int32 MaxNum = FMath::Max(BlendStackBlocks.Num(), BlendStackNames.Num());
 
-	for (int32 Index = 0; Index < MaxIndex; ++Index)
+	for (int32 Index = 0; Index < MaxNum; ++Index)
 	{
-		Renderer.AddText(BlendStackNames.IsValidIndex(Index) ? BlendStackNames[Index] : TEXT("<unnamed blend stack>"));
+		const FString CurrentName = BlendStackNames.IsValidIndex(Index) ? BlendStackNames[Index] : TEXT("<unnamed blend stack>");
+
+		Renderer.AddText(CurrentName);
 		Renderer.NewLine();
 
-		if (BlendStackBlocks.IsValidIndex(Index))
-		{
-			FCameraDebugBlock* BlendStackBlock = BlendStackBlocks[Index];
-			BlendStackBlock->DebugDraw(Params, Renderer);
-		}
-		else
+		const bool bDoDebugDraw = (
+				GGameplayCamerasDebugBlendStacksFilter.IsEmpty() ||
+				CurrentName.Contains(GGameplayCamerasDebugBlendStacksFilter));
+		if (bDoDebugDraw)
 		{
-			Renderer.AddText(TEXT("<missing blend stack>"));
+			if (BlendStackBlocks.IsValidIndex(Index))
+			{
+				FCameraDebugBlock* BlendStackBlock = BlendStackBlocks[Index];
+				BlendStackBlock->DebugDraw(Params, Renderer);
+			}
+			else
+			{
+				Renderer.AddText(TEXT("<missing blend stack>"));
+			}
 		}
 	}
 
diff --git a/Engine/Plugins/Cameras/GameplayCameras/Source/GameplayCameras/Private/Debug/CameraDebugColors.cpp b/Engine/Plugins/Cameras/GameplayCameras/Source/GameplayCameras/Private/Debug/CameraDebugColors.cpp
--- a/Engine/Plugins/Cameras/GameplayCameras/Source/GameplayCameras/Private/Debug/CameraDebugColors.cpp
+++ b/Engine/Plugins/Cameras/GameplayCameras/Source/GameplayCameras/Private/Debug/CameraDebugColors.cpp
@@ -159,6 +159,21 @@ void FCameraDebugColors::RegisterBuiltinColorSchemes()
 	Set(TEXT("SolarizedDark"));
 }
 
+FLinearColor LerpLinearColorUsingHSV(
+	const FLinearColor& Start, const FLinearColor& End,
+	int32 Increment, int32 TotalIncrements)
+{
+	if (TotalIncrements == 1)
+	{
+		return End;
+	}
+	else
+	{
+		const float Alpha = (float)Increment / (float)(TotalIncrements - 1);
+		return FLinearColor::LerpUsingHSV(Start, End, Alpha);
+	}
+}
+
 }  // namespace UE::Cameras
 
 #endif  // UE_GAMEPLAY_CAMERAS_DEBUG
diff --git a/Engine/Plugins/Cameras/GameplayCameras/Source/GameplayCameras/Private/Debug/CameraDebugRenderer.cpp b/Engine/Plugins/Cameras/GameplayCameras/Source/GameplayCameras/Private/Debug/CameraDebugRenderer.cpp
--- a/Engine/Plugins/Cameras/GameplayCameras/Source/GameplayCameras/Private/Debug/CameraDebugRenderer.cpp
+++ b/Engine/Plugins/Cameras/GameplayCameras/Source/GameplayCameras/Private/Debug/CameraDebugRenderer.cpp
@@ -622,24 +622,68 @@ void FCameraDebugRenderer::DrawCoordinateSystem(const FTransform3d& Transform, f
 	DrawCoordinateSystem(Transform.GetLocation(), Transform.GetRotation().Rotator(), AxesLength);
 }
 
-void FCameraDebugRenderer::DrawText(const FVector3d& WorldPosition, const FString& Text, const FLinearColor& TextColor, UFont* TextFont)
+void FCameraDebugRenderer::DrawText(const FVector3d& WorldPosition, const FString& Text, const FLinearColor& TextColor, UFont* TextFont, float TextScale)
 {
-	DrawText(WorldPosition, FVector2d::ZeroVector, Text, TextColor, TextFont);
+	DrawText(WorldPosition, FVector2d::ZeroVector, Text, TextColor, TextFont, TextScale);
 }
 
-void FCameraDebugRenderer::DrawText(const FVector3d& WorldPosition, const FVector2d& ScreenOffset, const FString& Text, const FLinearColor& TextColor, UFont* TextFont)
+void FCameraDebugRenderer::DrawTextView(const FVector3d& WorldPosition, FStringView Text, const FLinearColor& TextColor, UFont* TextFont, float TextScale)
+{
+	DrawTextView(WorldPosition, FVector2d::ZeroVector, Text, TextColor, TextFont, TextScale);
+}
+
+void FCameraDebugRenderer::DrawText(const FVector3d& WorldPosition, const FVector2d& ScreenOffset, const FString& Text, const FLinearColor& TextColor, UFont* TextFont, float TextScale)
 {
 	if (Canvas && SceneView)
 	{
 		UFont* ActualTextFont = TextFont ? TextFont : GEngine->GetSmallFont();
 
-		const FVector3d ScreenPosition = SceneView->Project(WorldPosition);
+		FVector2d PixelOffset(FVector2d::ZeroVector);
+		if (ScreenOffset.X != 0.0 || ScreenOffset.Y != 0.0)
+		{
+			const double ScreenWidth = SceneView->UnscaledViewRect.Width();
+			const double ScreenHeight = SceneView->UnscaledViewRect.Height();
+			PixelOffset = FVector2d(ScreenOffset.X / ScreenWidth, ScreenOffset.Y / ScreenHeight);
+		}
+
+		FVector2d PixelPosition;
+		SceneView->WorldToPixel(WorldPosition, PixelPosition);
+
+		FCanvasTextItem TextItem(
+				FVector2D(PixelPosition.X + PixelOffset.X, PixelPosition.Y + PixelOffset.Y),
+				FText::FromString(Text),
+				ActualTextFont,
+				TextColor);
+		TextItem.BlendMode = SE_BLEND_Translucent;
+		TextItem.Scale = FVector2D(TextScale, TextScale);
+		Canvas->DrawItem(TextItem);	
+	}
+}
+
+void FCameraDebugRenderer::DrawTextView(const FVector3d& WorldPosition, const FVector2d& ScreenOffset, FStringView Text, const FLinearColor& TextColor, UFont* TextFont, float TextScale)
+{
+	if (Canvas && SceneView)
+	{
+		UFont* ActualTextFont = TextFont ? TextFont : GEngine->GetSmallFont();
+
+		FVector2d PixelOffset(FVector2d::ZeroVector);
+		if (ScreenOffset.X != 0.0 || ScreenOffset.Y != 0.0)
+		{
+			const double ScreenWidth = SceneView->UnscaledViewRect.Width();
+			const double ScreenHeight = SceneView->UnscaledViewRect.Height();
+			PixelOffset = FVector2d(ScreenOffset.X / ScreenWidth, ScreenOffset.Y / ScreenHeight);
+		}
+
+		FVector2d PixelPosition;
+		SceneView->WorldToPixel(WorldPosition, PixelPosition);
+
 		FCanvasTextStringViewItem TextItem(
-				FVector2D(ScreenPosition.X + ScreenOffset.Y, ScreenPosition.Y + ScreenOffset.Y),
+				FVector2D(PixelPosition.X + PixelOffset.X, PixelPosition.Y + PixelOffset.Y),
 				Text,
 				ActualTextFont,
 				TextColor);
 		TextItem.BlendMode = SE_BLEND_Translucent;
+		TextItem.Scale = FVector2D(TextScale, TextScale);
 		Canvas->DrawItem(TextItem);	
 	}
 }
diff --git a/Engine/Plugins/Cameras/GameplayCameras/Source/GameplayCameras/Private/Debug/CameraDirectorDebugBlock.cpp b/Engine/Plugins/Cameras/GameplayCameras/Source/GameplayCameras/Private/Debug/CameraDirectorDebugBlock.cpp
new file mode 100644
--- /dev/null
+++ b/Engine/Plugins/Cameras/GameplayCameras/Source/GameplayCameras/Private/Debug/CameraDirectorDebugBlock.cpp
@@ -0,0 +1,143 @@
+// Copyright Epic Games, Inc. All Rights Reserved.
+
+#include "Debug/CameraDirectorDebugBlock.h"
+
+#include "Core/CameraAsset.h"
+#include "Core/CameraDirector.h"
+#include "Core/CameraDirectorEvaluator.h"
+#include "Core/CameraEvaluationContext.h"
+#include "Debug/CameraDebugBlockBuilder.h"
+#include "Debug/CameraDebugRenderer.h"
+#include "Debug/CameraDirectorTreeDebugBlock.h"
+#include "Debug/CameraNodeEvaluationResultDebugBlock.h"
+#include "Debug/CameraPoseDebugBlock.h"
+#include "Debug/ContextDataTableDebugBlock.h"
+#include "Debug/VariableTableDebugBlock.h"
+#include "HAL/IConsoleManager.h"
+
+#if UE_GAMEPLAY_CAMERAS_DEBUG
+
+namespace UE::Cameras
+{
+
+bool GGameplayCamerasDebugContextInitialResultShowUnchanged = false;
+static FAutoConsoleVariableRef CVarGameplayCamerasDebugContextInitialResultShowUnchanged(
+	TEXT("GameplayCameras.Debug.ContextInitialResult.ShowUnchanged"),
+	GGameplayCamerasDebugContextInitialResultShowUnchanged,
+	TEXT(""));
+
+bool GGameplayCamerasDebugContextInitialResultShowVariableIDs = false;
+static FAutoConsoleVariableRef CVarGameplayCamerasDebugContextInitialResultShowVariableIDs(
+	TEXT("GameplayCameras.Debug.ContextInitialResult.ShowVariableIDs"),
+	GGameplayCamerasDebugContextInitialResultShowVariableIDs,
+	TEXT(""));
+
+bool GGameplayCamerasDebugContextInitialResultShowDataIDs = false;
+static FAutoConsoleVariableRef CVarGameplayCamerasDebugContextInitialResultShowDataIDs(
+	TEXT("GameplayCameras.Debug.ContextInitialResult.ShowDataIDs"),
+	GGameplayCamerasDebugContextInitialResultShowDataIDs,
+	TEXT(""));
+
+bool GGameplayCamerasDebugContextInitialResultShowCoordinateSystem = true;
+static FAutoConsoleVariableRef CVarGameplayCamerasDebugContextInitialResultShowCoordinateSystem(
+	TEXT("GameplayCameras.Debug.ContextInitialResult.ShowCoordinateSystem"),
+	GGameplayCamerasDebugContextInitialResultShowCoordinateSystem,
+	TEXT(""));
+
+UE_DEFINE_CAMERA_DEBUG_BLOCK(FCameraDirectorDebugBlock)
+
+void FCameraDirectorDebugBlock::Initialize(TSharedPtr<FCameraEvaluationContext> Context, FCameraDebugBlockBuilder& Builder)
+{
+	if (Context)
+	{
+		const FCameraObjectTypeRegistry& TypeRegistry = FCameraObjectTypeRegistry::Get();
+		const FName ContextTypeName = TypeRegistry.GetTypeNameSafe(Context->GetTypeID());
+
+		const UObject* ContextOwner = Context->GetOwner();
+		const FCameraDirectorEvaluator* DirectorEvaluator = Context->GetDirectorEvaluator();
+		const UCameraDirector* CameraDirector = DirectorEvaluator->GetCameraDirector();
+
+		ContextClassName = ContextTypeName;
+		OwnerName = *GetPathNameSafe(ContextOwner);
+		OwnerClassName = ContextOwner ? ContextOwner->GetClass()->GetFName() : NAME_None;
+		CameraAssetName = GetNameSafe(Context->GetCameraAsset());
+		CameraDirectorClassName = GetFNameSafe(CameraDirector);
+		InitialContextTransform = Context->GetInitialResult().CameraPose.GetTransform();
+		bIsValid = true;
+
+		const FCameraNodeEvaluationResult& InitialResult = Context->GetInitialResult();
+		FCameraNodeEvaluationResultDebugBlock& ResultBlock = Builder.BuildDebugBlock<FCameraNodeEvaluationResultDebugBlock>();
+		ResultBlock.Initialize(InitialResult, Builder);
+		AddChild(&ResultBlock);
+		{
+			ResultBlock.GetCameraPoseDebugBlock()
+				->WithShowUnchangedCVar(TEXT("GameplayCameras.Debug.ContextInitialResult.ShowUnchanged"));
+			ResultBlock.GetVariableTableDebugBlock()
+				->WithShowVariableIDsCVar(TEXT("GameplayCameras.Debug.ContextInitialResult.ShowVariableIDs"));
+			ResultBlock.GetContextDataTableDebugBlock()
+				->WithShowDataIDsCVar(TEXT("GameplayCameras.Debug.ContextInitialResult.ShowDataIDs"));
+		}
+
+		TArrayView<const TSharedPtr<FCameraEvaluationContext>> ChildrenContexts = Context->GetChildrenContexts();
+		if (ChildrenContexts.Num() > 0)
+		{
+			FCameraDirectorTreeDebugBlock& ChildBlock = Builder.StartChildDebugBlock<FCameraDirectorTreeDebugBlock>();
+			{
+				ChildBlock.Initialize(ChildrenContexts, Builder);
+			}
+			Builder.EndChildDebugBlock();
+		}
+	}
+	else
+	{
+		bIsValid = false;
+	}
+}
+
+void FCameraDirectorDebugBlock::OnDebugDraw(const FCameraDebugBlockDrawParams& Params, FCameraDebugRenderer& Renderer)
+{
+	// Basic debug info.
+	if (bIsValid)
+	{
+		Renderer.AddText(TEXT("{cam_passive}[%s]{cam_default}"), *CameraDirectorClassName.ToString());
+		Renderer.AddIndent();
+		{
+			Renderer.AddText(TEXT("Context {cam_passive}[%s]{cam_default}\n"), *ContextClassName.ToString());
+
+			Renderer.AddText(TEXT("Owned by {cam_passive}[%s]{cam_default}\n"), *OwnerClassName.ToString());
+			Renderer.AddIndent();
+			{
+				Renderer.AddText(*OwnerName);
+			}
+			Renderer.RemoveIndent();
+
+			Renderer.AddText(TEXT("{cam_passive}From camera asset {cam_notice}%s{cam_default}\n"), *CameraAssetName);
+		}
+		Renderer.RemoveIndent();
+
+		if (Renderer.IsExternalRendering() && GGameplayCamerasDebugContextInitialResultShowCoordinateSystem)
+		{
+			Renderer.DrawCoordinateSystem(InitialContextTransform);
+		}
+	}
+	else
+	{
+		Renderer.AddText(TEXT("{cam_error}Invalid context!{cam_default}\n"));
+	}
+}
+
+void FCameraDirectorDebugBlock::OnSerialize(FArchive& Ar)
+{
+	Ar << ContextClassName;
+	Ar << OwnerClassName;
+	Ar << OwnerName;
+	Ar << CameraAssetName;
+	Ar << CameraDirectorClassName;
+	Ar << InitialContextTransform;
+	Ar << bIsValid;
+}
+	
+}  // namespace UE::Cameras
+
+#endif  // UE_GAMEPLAY_CAMERAS_DEBUG
+
diff --git a/Engine/Plugins/Cameras/GameplayCameras/Source/GameplayCameras/Private/Debug/CameraDirectorTreeDebugBlock.cpp b/Engine/Plugins/Cameras/GameplayCameras/Source/GameplayCameras/Private/Debug/CameraDirectorTreeDebugBlock.cpp
--- a/Engine/Plugins/Cameras/GameplayCameras/Source/GameplayCameras/Private/Debug/CameraDirectorTreeDebugBlock.cpp
+++ b/Engine/Plugins/Cameras/GameplayCameras/Source/GameplayCameras/Private/Debug/CameraDirectorTreeDebugBlock.cpp
@@ -2,34 +2,22 @@
 
 #include "Debug/CameraDirectorTreeDebugBlock.h"
 
-#include "Core/CameraAsset.h"
-#include "Core/CameraDirector.h"
 #include "Core/CameraDirectorEvaluator.h"
 #include "Core/CameraEvaluationContext.h"
 #include "Core/CameraEvaluationContextStack.h"
 #include "Debug/CameraDebugBlockBuilder.h"
 #include "Debug/CameraDebugColors.h"
 #include "Debug/CameraDebugRenderer.h"
+#include "Debug/CameraDirectorDebugBlock.h"
 #include "Debug/CameraPoseDebugBlock.h"
-#include "HAL/IConsoleManager.h"
 
 #if UE_GAMEPLAY_CAMERAS_DEBUG
 
 namespace UE::Cameras
 {
 
-bool GGameplayCamerasDebugContextInitialResultShowUnchanged = false;
-static FAutoConsoleVariableRef CVarGameplayCamerasDebugContextInitialResultShowUnchanged(
-	TEXT("GameplayCameras.Debug.ContextInitialResult.ShowUnchanged"),
-	GGameplayCamerasDebugContextInitialResultShowUnchanged,
-	TEXT(""));
-
 UE_DEFINE_CAMERA_DEBUG_BLOCK(FCameraDirectorTreeDebugBlock)
 
-FCameraDirectorTreeDebugBlock::FCameraDirectorTreeDebugBlock()
-{
-}
-
 void FCameraDirectorTreeDebugBlock::Initialize(const FCameraEvaluationContextStack& ContextStack, FCameraDebugBlockBuilder& Builder)
 {
 	const int32 NumContexts = ContextStack.NumContexts();
@@ -38,9 +26,11 @@ void FCameraDirectorTreeDebugBlock::Initialize(const FCameraEvaluationContextSta
 		const FCameraEvaluationContextStack::FContextEntry& Entry(ContextStack.Entries[Index]);
 		TSharedPtr<FCameraEvaluationContext> Context = Entry.WeakContext.Pin();
 
-		FDirectorDebugInfo EntryDebugInfo;
-		InitializeEntry(Context, EntryDebugInfo, Builder);
-		CameraDirectors.Add(EntryDebugInfo);
+		FCameraDirectorDebugBlock& Child = Builder.StartChildDebugBlock<FCameraDirectorDebugBlock>();
+		{
+			Child.Initialize(Context, Builder);
+		}
+		Builder.EndChildDebugBlock();
 	}
 }
 
@@ -48,54 +38,11 @@ void FCameraDirectorTreeDebugBlock::Initialize(TArrayView<const TSharedPtr<FCame
 {
 	for (TSharedPtr<FCameraEvaluationContext> Context : Contexts)
 	{
-		FDirectorDebugInfo EntryDebugInfo;
-		InitializeEntry(Context, EntryDebugInfo, Builder);
-		CameraDirectors.Add(EntryDebugInfo);
-	}
-}
-
-void FCameraDirectorTreeDebugBlock::InitializeEntry(TSharedPtr<FCameraEvaluationContext> Context, FDirectorDebugInfo& EntryDebugInfo, FCameraDebugBlockBuilder& Builder)
-{
-	if (Context)
-	{
-		const FCameraObjectTypeRegistry& TypeRegistry = FCameraObjectTypeRegistry::Get();
-		const FName ContextTypeName = TypeRegistry.GetTypeNameSafe(Context->GetTypeID());
-
-		const UObject* ContextOwner = Context->GetOwner();
-		const FCameraDirectorEvaluator* DirectorEvaluator = Context->GetDirectorEvaluator();
-		const UCameraDirector* CameraDirector = DirectorEvaluator->GetCameraDirector();
-
-		TArrayView<const TSharedPtr<FCameraEvaluationContext>> ChildrenContexts = Context->GetChildrenContexts();
-
-		EntryDebugInfo.ContextClassName = ContextTypeName;
-		EntryDebugInfo.OwnerName = *GetPathNameSafe(ContextOwner);
-		EntryDebugInfo.OwnerClassName = ContextOwner ? ContextOwner->GetClass()->GetFName() : NAME_None;
-		EntryDebugInfo.CameraAssetName = GetNameSafe(Context->GetCameraAsset());
-		EntryDebugInfo.CameraDirectorClassName = GetFNameSafe(CameraDirector);
-		EntryDebugInfo.NumChildrenContexts = ChildrenContexts.Num();
-		EntryDebugInfo.InitialContextTransform = Context->GetInitialResult().CameraPose.GetTransform();
-		EntryDebugInfo.bIsValid = true;
-
-		const FCameraNodeEvaluationResult& InitialResult = Context->GetInitialResult();
-		AddChild(&Builder.BuildDebugBlock<FCameraPoseDebugBlock>(InitialResult.CameraPose)
-				.WithShowUnchangedCVar(TEXT("GameplayCameras.Debug.ContextInitialResult.ShowUnchanged")));
-
-		if (ChildrenContexts.Num() > 0)
+		FCameraDirectorDebugBlock& Child = Builder.StartChildDebugBlock<FCameraDirectorDebugBlock>();
 		{
-			FCameraDirectorTreeDebugBlock& ChildBlock = Builder.StartChildDebugBlock<FCameraDirectorTreeDebugBlock>();
-			{
-				ChildBlock.Initialize(ChildrenContexts, Builder);
-				ChildBlock.bIsTreeRoot = false;
-			}
-			Builder.EndChildDebugBlock();
+			Child.Initialize(Context, Builder);
 		}
-	}
-	else
-	{
-		EntryDebugInfo.bIsValid = false;
-
-		// Dummy debug block.
-		AddChild(&Builder.BuildDebugBlock<FCameraDebugBlock>());
+		Builder.EndChildDebugBlock();
 	}
 }
 
@@ -103,121 +50,59 @@ void FCameraDirectorTreeDebugBlock::OnDebugDraw(const FCameraDebugBlockDrawParam
 {
 	const FCameraDebugColors& Colors = FCameraDebugColors::Get();
 
-	// Our children view looks like this:
-	//
-	//		[0] Initial result pose block
-	//		[1] Child director block
-	//		[2] Child director block
-	//		[3] Initial result pose block
-	//		[4] Initial result pose block
-	//
-	TArrayView<FCameraDebugBlock*> ChildrenView(GetChildren());
-
-	// For the root debug info, separate the active director from the inactive ones.
+	// Separate the active director from the inactive ones.
 	// The inactive ones are at the beginning (bottom) of the stack.
-	if (bIsTreeRoot)
-	{
-		Renderer.SetTextColor(Colors.Notice);
-		Renderer.AddText("Inactive Directors\n");
-		Renderer.SetTextColor(Colors.Default);
-		Renderer.AddIndent();
-	}
+	Renderer.SetTextColor(Colors.Notice);
+	Renderer.AddText("Inactive Directors\n");
+	Renderer.SetTextColor(Colors.Default);
+	Renderer.AddIndent();
+
+	TArrayView<FCameraDebugBlock*> ChildrenView(GetChildren());
 
-	int32 ChildContextIndex = 0;
+	const int32 NumInactiveDirectors = ChildrenView.Num() - 1;
 
 	for (int32 Index = 0; Index < ChildrenView.Num(); ++Index)
 	{
-		// If we reached the top of the stack, display the active director separately.
-		if (bIsTreeRoot && ChildContextIndex == CameraDirectors.Num() - 1)
+		FLinearColor CameraPoseLineColor(FColorList::SlateBlue);
+
+		// If we reached the top of the stack, display the active director separately,
+		// and draw the initial camera pose in SlateBlue color. Otherwise, draw the
+		// the initial camera pose somewhere between dark and light grey.
+		if (Index == ChildrenView.Num() - 1)
 		{
 			Renderer.RemoveIndent();
 
 			Renderer.SetTextColor(Colors.Notice);
 			Renderer.AddText("Active Director\n");
 			Renderer.SetTextColor(Colors.Default);
-			Renderer.AddIndent();
-		}
 
-		Renderer.AddText(TEXT("{cam_passive}[%d]{cam_default} "), ChildContextIndex + 1);
-
-		const FDirectorDebugInfo& EntryDebugInfo(CameraDirectors[ChildContextIndex]);
-		if (EntryDebugInfo.bIsValid)
-		{
-			Renderer.AddText(TEXT("{cam_passive}[%s]{cam_default}"), 
-					*EntryDebugInfo.CameraDirectorClassName.ToString());
 			Renderer.AddIndent();
-			{
-				Renderer.AddText(TEXT("Context {cam_passive}[%s]{cam_default}\n"), *EntryDebugInfo.ContextClassName.ToString());
-
-				Renderer.AddText(TEXT("Owned by {cam_passive}[%s]{cam_default}\n"), *EntryDebugInfo.OwnerClassName.ToString());
-				Renderer.AddIndent();
-				{
-					Renderer.AddText(*EntryDebugInfo.OwnerName);
-				}
-				Renderer.RemoveIndent();
-
-				Renderer.AddText(TEXT("{cam_passive}From camera asset {cam_notice}%s{cam_default}\n"), 
-						*EntryDebugInfo.CameraAssetName);
-			}
-			Renderer.RemoveIndent();
-
-			Renderer.DrawCoordinateSystem(EntryDebugInfo.InitialContextTransform);
 		}
 		else
 		{
-			Renderer.AddText(TEXT("{cam_error}Invalid context!{cam_default}\n"));
+			CameraPoseLineColor = LerpLinearColorUsingHSV(
+					FColorList::DimGrey, FColorList::LightGrey, Index, NumInactiveDirectors);
 		}
 
-		// Initial result pose block.
-		Renderer.AddIndent();
-		ChildrenView[Index]->DebugDraw(Params, Renderer);
-		Renderer.RemoveIndent();
+		FString PoseLabel(LexToString(Index));
+		FScopedGlobalCameraPoseRenderingParams ScopedPoseParams(PoseLabel, CameraPoseLineColor);
 
-		// Children blocks.
-		if (EntryDebugInfo.NumChildrenContexts > 0)
-		{
-			Renderer.AddIndent();
-			for (int32 ChildIndex = 0; ChildIndex < EntryDebugInfo.NumChildrenContexts; ++ChildIndex)
-			{
-				ChildrenView[Index + 1 + ChildIndex]->DebugDraw(Params, Renderer);
-			}
-			Renderer.RemoveIndent();
+		Renderer.AddText(TEXT("{cam_passive}[%d]{cam_default} "), Index + 1);
 
-			Index += EntryDebugInfo.NumChildrenContexts;
+		if (FCameraDirectorDebugBlock* EntryDebugBlock = ChildrenView[Index]->CastThisChecked<FCameraDirectorDebugBlock>())
+		{
+			EntryDebugBlock->DebugDraw(Params, Renderer);
 		}
 
-		++ChildContextIndex;
-
 		Renderer.NewLine();
 	}
 
-	if (bIsTreeRoot)
-	{
-		Renderer.RemoveIndent();
-	}
+	Renderer.RemoveIndent();
 	Renderer.SetTextColor(Colors.Default);
 
 	Renderer.SkipAllBlocks();
 }
 
-void FCameraDirectorTreeDebugBlock::OnSerialize(FArchive& Ar)
-{
-	 Ar << CameraDirectors;
-	 Ar << bIsTreeRoot;
-}
-
-FArchive& operator<< (FArchive& Ar, FCameraDirectorTreeDebugBlock::FDirectorDebugInfo& DirectorDebugInfo)
-{
-	Ar << DirectorDebugInfo.ContextClassName;
-	Ar << DirectorDebugInfo.OwnerClassName;
-	Ar << DirectorDebugInfo.OwnerName;
-	Ar << DirectorDebugInfo.CameraAssetName;
-	Ar << DirectorDebugInfo.CameraDirectorClassName;
-	Ar << DirectorDebugInfo.InitialContextTransform;
-	Ar << DirectorDebugInfo.bIsValid;
-	return Ar;
-}
-
 }  // namespace UE::Cameras
 
 #endif  // UE_GAMEPLAY_CAMERAS_DEBUG
diff --git a/Engine/Plugins/Cameras/GameplayCameras/Source/GameplayCameras/Private/Debug/CameraPoseDebugBlock.cpp b/Engine/Plugins/Cameras/GameplayCameras/Source/GameplayCameras/Private/Debug/CameraPoseDebugBlock.cpp
--- a/Engine/Plugins/Cameras/GameplayCameras/Source/GameplayCameras/Private/Debug/CameraPoseDebugBlock.cpp
+++ b/Engine/Plugins/Cameras/GameplayCameras/Source/GameplayCameras/Private/Debug/CameraPoseDebugBlock.cpp
@@ -6,6 +6,7 @@
 #include "Debug/CameraDebugColors.h"
 #include "Debug/CameraDebugRenderer.h"
 #include "Debug/DebugTextRenderer.h"
+#include "Engine/Engine.h"
 #include "HAL/IConsoleManager.h"
 #include "Math/ColorList.h"
 
@@ -14,6 +15,12 @@
 namespace UE::Cameras
 {
 
+float GGameplayCamerasDebugCameraPoseLabelWorldOffset = 20.f;
+static FAutoConsoleVariableRef CVarGameplayCamerasDebugCameraPoseLabelWorldOffset(
+	TEXT("GameplayCameras.Debug.CameraPose.LabelWorldOffset"),
+	GGameplayCamerasDebugCameraPoseLabelWorldOffset,
+	TEXT(""));
+
 namespace Private
 {
 
@@ -28,14 +35,40 @@ void DebugDrawCameraPoseField(FCameraDebugRenderer& Renderer, const TCHAR* Field
 
 UE_DEFINE_CAMERA_DEBUG_BLOCK(FCameraPoseDebugBlock)
 
+FLinearColor FCameraPoseDebugBlock::GlobalCameraPoseLineColor(FColorList::SlateBlue);
+FString FCameraPoseDebugBlock::GlobalCameraPoseLabel;
+
+FScopedGlobalCameraPoseRenderingParams::FScopedGlobalCameraPoseRenderingParams(const FString& Label)
+	: FScopedGlobalCameraPoseRenderingParams(Label, FCameraPoseDebugBlock::GlobalCameraPoseLineColor)
+{
+}
+
+FScopedGlobalCameraPoseRenderingParams::FScopedGlobalCameraPoseRenderingParams(const FLinearColor& LineColor)
+	: FScopedGlobalCameraPoseRenderingParams(FCameraPoseDebugBlock::GlobalCameraPoseLabel, LineColor)
+{
+}
+
+FScopedGlobalCameraPoseRenderingParams::FScopedGlobalCameraPoseRenderingParams(const FString& Label, const FLinearColor& LineColor)
+{
+	PreviousLabel = FCameraPoseDebugBlock::GlobalCameraPoseLabel;
+	PreviousLineColor = FCameraPoseDebugBlock::GlobalCameraPoseLineColor;
+
+	FCameraPoseDebugBlock::GlobalCameraPoseLabel = Label;
+	FCameraPoseDebugBlock::GlobalCameraPoseLineColor = LineColor;
+}
+
+FScopedGlobalCameraPoseRenderingParams::~FScopedGlobalCameraPoseRenderingParams()
+{
+	FCameraPoseDebugBlock::GlobalCameraPoseLabel = PreviousLabel;
+	FCameraPoseDebugBlock::GlobalCameraPoseLineColor = PreviousLineColor;
+}
+
 FCameraPoseDebugBlock::FCameraPoseDebugBlock()
-	: CameraPoseLineColor(FColorList::SlateBlue)
 {
 }
 
 FCameraPoseDebugBlock::FCameraPoseDebugBlock(const FCameraPose& InCameraPose)
 	: CameraPose(InCameraPose)
-	, CameraPoseLineColor(FColorList::SlateBlue)
 {
 }
 
@@ -68,14 +101,25 @@ void FCameraPoseDebugBlock::OnDebugDraw(const FCameraDebugBlockDrawParams& Param
 		UE_CAMERA_POSE_FOR_ALL_PROPERTIES()
 #undef UE_CAMERA_POSE_FOR_PROPERTY
 
-			Renderer.SetTextColor(Colors.Default);
+		Renderer.SetTextColor(Colors.Default);
 		Renderer.AddText(TEXT("Effective FOV  : %f\n"), CameraPose.GetEffectiveFieldOfView());
 		Renderer.AddText(TEXT("Effective Aspect Ratio  : %f\n"), CameraPose.GetSensorAspectRatio());
 	}
 
 	if (bDrawInExternalRendering && Renderer.IsExternalRendering())
 	{
-		Renderer.DrawCameraPose(CameraPose, CameraPoseLineColor, CameraPoseSize);
+		Renderer.DrawCameraPose(CameraPose, GlobalCameraPoseLineColor);
+
+		if (!GlobalCameraPoseLabel.IsEmpty())
+		{
+			const FVector3d TextWorldOffset(0, 0, GGameplayCamerasDebugCameraPoseLabelWorldOffset);
+			UFont* LargeFont = GEngine->GetLargeFont();
+			Renderer.DrawText(
+					CameraPose.GetLocation() + TextWorldOffset,
+					GlobalCameraPoseLabel,
+					GlobalCameraPoseLineColor,
+					LargeFont);
+		}
 	}
 }
 
@@ -83,6 +127,8 @@ void FCameraPoseDebugBlock::OnSerialize(FArchive& Ar)
 {
 	FCameraPose::SerializeWithFlags(Ar, CameraPose);
 	Ar << ShowUnchangedCVarName;
+	Ar << bDrawText;
+	Ar << bDrawInExternalRendering;
 }
 
 }  // namespace UE::Cameras
diff --git a/Engine/Plugins/Cameras/GameplayCameras/Source/GameplayCameras/Private/Debug/CameraPoseLocationTrailDebugBlock.cpp b/Engine/Plugins/Cameras/GameplayCameras/Source/GameplayCameras/Private/Debug/CameraPoseLocationTrailDebugBlock.cpp
--- a/Engine/Plugins/Cameras/GameplayCameras/Source/GameplayCameras/Private/Debug/CameraPoseLocationTrailDebugBlock.cpp
+++ b/Engine/Plugins/Cameras/GameplayCameras/Source/GameplayCameras/Private/Debug/CameraPoseLocationTrailDebugBlock.cpp
@@ -39,7 +39,7 @@ void FCameraPoseLocationTrailDebugBlock::DrawCameraPoseLocationTrail(FCameraDebu
 
 		if (!PrevPoint.IsZero() && FVector3d::Distance(PrevPoint, NextPoint) > UE_SMALL_NUMBER)
 		{
-			Renderer.DrawLine(PrevPoint, NextPoint, TrailColor, 1.f);
+			Renderer.DrawLine(PrevPoint, NextPoint, TrailColor, 0.3f);
 		}
 	}
 }
diff --git a/Engine/Plugins/Cameras/GameplayCameras/Source/GameplayCameras/Private/Debug/CameraSystemDebugRegistry.cpp b/Engine/Plugins/Cameras/GameplayCameras/Source/GameplayCameras/Private/Debug/CameraSystemDebugRegistry.cpp
--- a/Engine/Plugins/Cameras/GameplayCameras/Source/GameplayCameras/Private/Debug/CameraSystemDebugRegistry.cpp
+++ b/Engine/Plugins/Cameras/GameplayCameras/Source/GameplayCameras/Private/Debug/CameraSystemDebugRegistry.cpp
@@ -50,6 +50,11 @@ void FCameraSystemDebugRegistry::GetRegisteredCameraSystemEvaluators(FRegistered
 	}
 }
 
+int32 FCameraSystemDebugRegistry::NumRegisteredCameraSystemEvaluators() const
+{
+	return Entries.Num();
+}
+
 }  // namespace UE::Cameras
 
 #endif  // UE_GAMEPLAY_CAMERAS_DEBUG
diff --git a/Engine/Plugins/Cameras/GameplayCameras/Source/GameplayCameras/Private/Debug/RootCameraDebugBlock.cpp b/Engine/Plugins/Cameras/GameplayCameras/Source/GameplayCameras/Private/Debug/RootCameraDebugBlock.cpp
--- a/Engine/Plugins/Cameras/GameplayCameras/Source/GameplayCameras/Private/Debug/RootCameraDebugBlock.cpp
+++ b/Engine/Plugins/Cameras/GameplayCameras/Source/GameplayCameras/Private/Debug/RootCameraDebugBlock.cpp
@@ -147,33 +147,26 @@ void FRootCameraDebugBlock::BuildDebugBlocks(const FCameraSystemEvaluator& Camer
 	AddChild(&Builder.BuildDebugBlock<FViewfinderDebugBlock>());
 }
 
-void FRootCameraDebugBlock::RootDebugDraw(FCameraDebugRenderer& Renderer, bool bForceDraw)
+void FRootCameraDebugBlock::RootDebugDraw(const FRootCameraDebugDrawParams& Params, FCameraDebugRenderer& Renderer)
 {
-	if (!GGameplayCamerasDebugEnable)
-	{
-		return;
-	}
-
-	FCameraSystemDebugID WantedDebugID(GGameplayCamerasDebugSystemID);
-	const bool bDoDebugDraw = (bForceDraw || WantedDebugID.IsAny() || WantedDebugID == DebugID);
-	if (!bDoDebugDraw)
+	if (!ShouldDebugDraw(DebugID, Params.bIsCameraManagerOrViewTarget))
 	{
 		return;
 	}
 
 	// Figure out what debug categories are active.
-	FCameraDebugBlockDrawParams Params;
+	FCameraDebugBlockDrawParams DrawParams;
 
 	TArray<FStringView, TInlineAllocator<4>> ActiveCategories;
 	UE::String::ParseTokens(GGameplayCamerasDebugCategories, ',', ActiveCategories);
 	for (FStringView CategoryView : ActiveCategories)
 	{
-		Params.ActiveCategories.Add(FString(CategoryView));
+		DrawParams.ActiveCategories.Add(FString(CategoryView));
 	}
 
 	// Do the drawing!
 	Renderer.BeginDrawing();
-	FCameraDebugBlock::DebugDraw(Params, Renderer);
+	FCameraDebugBlock::DebugDraw(DrawParams, Renderer);
 	Renderer.EndDrawing();
 }
 
@@ -182,6 +175,26 @@ void FRootCameraDebugBlock::OnSerialize(FArchive& Ar)
 	Ar << DebugID;
 }
 
+bool FRootCameraDebugBlock::ShouldDebugDraw(FCameraSystemDebugID InDebugID, bool bIsActive)
+{
+	if (!GGameplayCamerasDebugEnable)
+	{
+		return false;
+	}
+
+	const FCameraSystemDebugID WantedDebugID(GGameplayCamerasDebugSystemID);
+	const int32 NumCameraSystems = FCameraSystemDebugRegistry::Get().NumRegisteredCameraSystemEvaluators();
+
+	// If the wanted debug ID is ours, we draw.
+	// If the wanted debug ID is "any", we draw.
+	// If the wanted debug ID is "auto", we draw if we are the view target or camera manager or
+	// we are in charge of the "active" camera, whatever that means.
+	const bool bDoDebugDraw = (WantedDebugID == InDebugID) 
+		|| WantedDebugID.IsAny()
+		|| (WantedDebugID.IsAuto() && (bIsActive || NumCameraSystems == 1));
+	return bDoDebugDraw;
+}
+
 }  // namespace UE::Cameras
 
 #endif  // UE_GAMEPLAY_CAMERAS_DEBUG
diff --git a/Engine/Plugins/Cameras/GameplayCameras/Source/GameplayCameras/Private/GameFramework/IGameplayCameraSystemHost.cpp b/Engine/Plugins/Cameras/GameplayCameras/Source/GameplayCameras/Private/GameFramework/IGameplayCameraSystemHost.cpp
--- a/Engine/Plugins/Cameras/GameplayCameras/Source/GameplayCameras/Private/GameFramework/IGameplayCameraSystemHost.cpp
+++ b/Engine/Plugins/Cameras/GameplayCameras/Source/GameplayCameras/Private/GameFramework/IGameplayCameraSystemHost.cpp
@@ -221,15 +221,10 @@ void IGameplayCameraSystemHost::DebugDraw(UCanvas* Canvas, APlayerController* Pl
 				(!bThisIsCameraManager && !bThisIsViewTarget) ||
 				!ActualPlayerController || !ActualPlayerController->Player);
 
-		// Force draw this host's camera system if the wanted debug ID is "auto" and we are the
-		// view target or camera manager.
-		FCameraSystemDebugID WantedDebugID(GGameplayCamerasDebugSystemID);
-		bool bForceDraw = WantedDebugID.IsAuto() && (bThisIsCameraManager || bThisIsViewTarget);
-
 		FCameraSystemDebugUpdateParams DebugUpdateParams;
 		DebugUpdateParams.CanvasObject = Canvas;
 		DebugUpdateParams.bIsDebugCameraEnabled = bIsDebugCameraEnabled;
-		DebugUpdateParams.bForceDraw = bForceDraw;
+		DebugUpdateParams.bIsCameraManagerOrViewTarget = (bThisIsCameraManager || bThisIsViewTarget);
 		CameraSystemEvaluator->DebugUpdate(DebugUpdateParams);
 	}
 }
diff --git a/Engine/Plugins/Cameras/GameplayCameras/Source/GameplayCameras/Private/Nodes/Collision/CollisionPushCameraNode.cpp b/Engine/Plugins/Cameras/GameplayCameras/Source/GameplayCameras/Private/Nodes/Collision/CollisionPushCameraNode.cpp
--- a/Engine/Plugins/Cameras/GameplayCameras/Source/GameplayCameras/Private/Nodes/Collision/CollisionPushCameraNode.cpp
+++ b/Engine/Plugins/Cameras/GameplayCameras/Source/GameplayCameras/Private/Nodes/Collision/CollisionPushCameraNode.cpp
@@ -563,9 +563,11 @@ void FCollisionPushCameraDebugBlock::OnDebugDraw(const FCameraDebugBlockDrawPara
 		Renderer.SetTextColor(Colors.Default);
 	}
 
-	if (bGotSafePosition)
+	if (Renderer.IsExternalRendering() && bGotSafePosition)
 	{
-		Renderer.DrawText(SafePosition, TEXT("Safe Position"), FLinearColor::Gray, GEngine->GetTinyFont());
+		const FVector2d TextOffset(-20, -20);
+		Renderer.DrawPoint(SafePosition, 2.f, FLinearColor::Gray, 2.f);
+		Renderer.DrawTextView(SafePosition, TextOffset, TEXT("Safe Position"), FLinearColor::Gray, GEngine->GetTinyFont());
 	}
 }
 
diff --git a/Engine/Plugins/Cameras/GameplayCameras/Source/GameplayCameras/Private/Nodes/Common/DampenPositionCameraNode.cpp b/Engine/Plugins/Cameras/GameplayCameras/Source/GameplayCameras/Private/Nodes/Common/DampenPositionCameraNode.cpp
--- a/Engine/Plugins/Cameras/GameplayCameras/Source/GameplayCameras/Private/Nodes/Common/DampenPositionCameraNode.cpp
+++ b/Engine/Plugins/Cameras/GameplayCameras/Source/GameplayCameras/Private/Nodes/Common/DampenPositionCameraNode.cpp
@@ -9,6 +9,7 @@
 #include "Debug/CameraDebugBlockBuilder.h"
 #include "Debug/CameraDebugRenderer.h"
 #include "GameplayCameras.h"
+#include "HAL/IConsoleManager.h"
 #include "Math/CriticalDamper.h"
 #include "Templates/Tuple.h"
 
@@ -17,6 +18,12 @@
 namespace UE::Cameras
 {
 
+bool GGameplayCamerasDebugDampingShowLocalSpace = false;
+static FAutoConsoleVariableRef CVarGameplayCamerasDebugDampingShowLocalSpace(
+	TEXT("GameplayCameras.Debug.Damping.ShowLocalSpace"),
+	GGameplayCamerasDebugDampingShowLocalSpace,
+	TEXT(""));
+
 class FDampenPositionCameraNodeEvaluator : public FCameraNodeEvaluator
 {
 	UE_DECLARE_CAMERA_NODE_EVALUATOR(GAMEPLAYCAMERAS_API, FDampenPositionCameraNodeEvaluator)
@@ -261,21 +268,27 @@ void FDampenPositionCameraDebugBlock::OnDebugDraw(const FCameraDebugBlockDrawPar
 			LateralX0, LateralDampingFactor,
 			VerticalX0, VerticalDampingFactor);
 
-	const double DampingAxesLength = 100.f;
-	Renderer.DrawLine(
-			UndampedPosition, 
-			UndampedPosition + DampingRotation.RotateVector(FVector3d::ForwardVector * DampingAxesLength),
-			FLinearColor::Red);
-	Renderer.DrawLine(
-			UndampedPosition, 
-			UndampedPosition + DampingRotation.RotateVector(FVector3d::RightVector * DampingAxesLength),
-			FLinearColor::Green);
-	Renderer.DrawLine(
-			UndampedPosition, 
-			UndampedPosition + DampingRotation.RotateVector(FVector3d::UpVector * DampingAxesLength),
-			FLinearColor::Blue);
-
-	Renderer.DrawLine(UndampedPosition, DampedPosition, FLinearColor::Yellow);
+	if (Renderer.IsExternalRendering())
+	{
+		if (GGameplayCamerasDebugDampingShowLocalSpace)
+		{
+			const double DampingAxesLength = 100.f;
+			Renderer.DrawLine(
+					UndampedPosition, 
+					UndampedPosition + DampingRotation.RotateVector(FVector3d::ForwardVector * DampingAxesLength),
+					FLinearColor::Red);
+			Renderer.DrawLine(
+					UndampedPosition, 
+					UndampedPosition + DampingRotation.RotateVector(FVector3d::RightVector * DampingAxesLength),
+					FLinearColor::Green);
+			Renderer.DrawLine(
+					UndampedPosition, 
+					UndampedPosition + DampingRotation.RotateVector(FVector3d::UpVector * DampingAxesLength),
+					FLinearColor::Blue);
+		}
+
+		Renderer.DrawLine(UndampedPosition, DampedPosition, FLinearColor::Yellow);
+	}
 }
 
 #endif  // UE_GAMEPLAY_CAMERAS_DEBUG
diff --git a/Engine/Plugins/Cameras/GameplayCameras/Source/GameplayCameras/Private/Services/OrientationInitializationService.cpp b/Engine/Plugins/Cameras/GameplayCameras/Source/GameplayCameras/Private/Services/OrientationInitializationService.cpp
--- a/Engine/Plugins/Cameras/GameplayCameras/Source/GameplayCameras/Private/Services/OrientationInitializationService.cpp
+++ b/Engine/Plugins/Cameras/GameplayCameras/Source/GameplayCameras/Private/Services/OrientationInitializationService.cpp
@@ -292,12 +292,15 @@ void FOrientationInitializationDebugBlock::OnDebugDraw(const FCameraDebugBlockDr
 	{
 		AimDebugInfo.DebugDraw(Params, Renderer);
 
-		Renderer.DrawSphere(LastEvaluatedTarget, 1.f, 8, FLinearColor::Green, 1.f);
+		const FVector2d TextOffset(20, -10);
+		Renderer.DrawPoint(LastEvaluatedTarget, 10.f, FLinearColor::Green, 0.f);
+		Renderer.DrawTextView(LastEvaluatedTarget, TextOffset, TEXT("Preserved Target"), FLinearColor::Green, nullptr);
 	}
 }
 
 void FOrientationInitializationDebugBlock::OnSerialize(FArchive& Ar)
 {
+	Ar << LastEvaluatedTarget;
 	Ar << AimDebugInfo;
 }
 
diff --git a/Engine/Plugins/Cameras/GameplayCameras/Source/GameplayCameras/Public/Core/BlendStackCameraNode.h b/Engine/Plugins/Cameras/GameplayCameras/Source/GameplayCameras/Public/Core/BlendStackCameraNode.h
--- a/Engine/Plugins/Cameras/GameplayCameras/Source/GameplayCameras/Public/Core/BlendStackCameraNode.h
+++ b/Engine/Plugins/Cameras/GameplayCameras/Source/GameplayCameras/Public/Core/BlendStackCameraNode.h
@@ -114,7 +114,10 @@ public:
 	bool HasAnyRunningCameraRig(TSharedPtr<const FCameraEvaluationContext> InContext) const;
 
 #if UE_GAMEPLAY_CAMERAS_DEBUG
-	FBlendStackCameraDebugBlock* BuildDetailedDebugBlock(const FCameraDebugBlockBuildParams& Params, FCameraDebugBlockBuilder& Builder);
+	FBlendStackCameraDebugBlock* BuildDetailedDebugBlock(
+			const FCameraDebugBlockBuildParams& Params,
+			const FLinearColor& StartColor, const FLinearColor& EndColor,
+			FCameraDebugBlockBuilder& Builder);
 #endif  // UE_GAMEPLAY_CAMERAS_DEBUG
 
 public:
@@ -284,6 +287,15 @@ private:
 class FBlendStackCameraDebugBlock : public FCameraDebugBlock
 {
 	UE_DECLARE_CAMERA_DEBUG_BLOCK(GAMEPLAYCAMERAS_API, FBlendStackCameraDebugBlock)
+
+protected:
+
+	virtual void OnSerialize(FArchive& Ar) override;
+
+private:
+
+	using FBlendStackStartEndColors = TTuple<FLinearColor, FLinearColor>;
+	FBlendStackStartEndColors StartEndColors;
 };
 
 #endif  // UE_GAMEPLAY_CAMERAS_DEBUG
diff --git a/Engine/Plugins/Cameras/GameplayCameras/Source/GameplayCameras/Public/Core/CameraSystemEvaluator.h b/Engine/Plugins/Cameras/GameplayCameras/Source/GameplayCameras/Public/Core/CameraSystemEvaluator.h
--- a/Engine/Plugins/Cameras/GameplayCameras/Source/GameplayCameras/Public/Core/CameraSystemEvaluator.h
+++ b/Engine/Plugins/Cameras/GameplayCameras/Source/GameplayCameras/Public/Core/CameraSystemEvaluator.h
@@ -134,6 +134,9 @@ struct FCameraSystemDebugUpdateParams
 	/** Whether the debug camera is enabled, giving an "outside" view of camera system. */
 	bool bIsDebugCameraEnabled = false;
 
+	/** Whether this camera system is run by the active camera manager or view target. */
+	bool bIsCameraManagerOrViewTarget = false;
+
 	/** Whether to force drawing debug info for this camera system. */
 	bool bForceDraw = false;
 };
diff --git a/Engine/Plugins/Cameras/GameplayCameras/Source/GameplayCameras/Public/Debug/CameraDebugColors.h b/Engine/Plugins/Cameras/GameplayCameras/Source/GameplayCameras/Public/Debug/CameraDebugColors.h
--- a/Engine/Plugins/Cameras/GameplayCameras/Source/GameplayCameras/Public/Debug/CameraDebugColors.h
+++ b/Engine/Plugins/Cameras/GameplayCameras/Source/GameplayCameras/Public/Debug/CameraDebugColors.h
@@ -59,7 +59,6 @@ public:
 	/** Gets all registered color scheme names. */
 	GAMEPLAYCAMERAS_API static void GetColorSchemeNames(TArray<FString>& OutColorSchemeNames);
 
-
 public:
 
 	// Internal API.
@@ -78,6 +77,10 @@ private:
 	static TMap<FString, FCameraDebugColors> ColorSchemes;
 };
 
+GAMEPLAYCAMERAS_API FLinearColor LerpLinearColorUsingHSV(
+	const FLinearColor& Start, const FLinearColor& End,
+	int32 Increment, int32 TotalIncrements);
+
 }  // namespace UE::Cameras
 
 #endif  // UE_GAMEPLAY_CAMERAS_DEBUG
diff --git a/Engine/Plugins/Cameras/GameplayCameras/Source/GameplayCameras/Public/Debug/CameraDebugRenderer.h b/Engine/Plugins/Cameras/GameplayCameras/Source/GameplayCameras/Public/Debug/CameraDebugRenderer.h
--- a/Engine/Plugins/Cameras/GameplayCameras/Source/GameplayCameras/Public/Debug/CameraDebugRenderer.h
+++ b/Engine/Plugins/Cameras/GameplayCameras/Source/GameplayCameras/Public/Debug/CameraDebugRenderer.h
@@ -128,9 +128,13 @@ public:
 	/** Draws a 3D coordinate system. */
 	GAMEPLAYCAMERAS_API void DrawCoordinateSystem(const FTransform3d& Transform, float AxesLength = 0.f);
 	/** Draws text at a projected 3D position. */
-	GAMEPLAYCAMERAS_API void DrawText(const FVector3d& WorldPosition, const FString& Text, const FLinearColor& TextColor, UFont* TextFont = nullptr);
+	GAMEPLAYCAMERAS_API void DrawText(const FVector3d& WorldPosition, const FString& Text, const FLinearColor& TextColor, UFont* TextFont = nullptr, float TextScale = 1.f);
+	/** Draws text at a projected 3D position. */
+	GAMEPLAYCAMERAS_API void DrawTextView(const FVector3d& WorldPosition, FStringView Text, const FLinearColor& TextColor, UFont* TextFont = nullptr, float TextScale = 1.f);
+	/** Draws text at a projected 3D position, with an added screen-space offset. */
+	GAMEPLAYCAMERAS_API void DrawText(const FVector3d& WorldPosition, const FVector2d& ScreenOffset, const FString& Text, const FLinearColor& TextColor, UFont* TextFont = nullptr, float TextScale = 1.f);
 	/** Draws text at a projected 3D position, with an added screen-space offset. */
-	GAMEPLAYCAMERAS_API void DrawText(const FVector3d& WorldPosition, const FVector2d& ScreenOffset, const FString& Text, const FLinearColor& TextColor, UFont* TextFont = nullptr);
+	GAMEPLAYCAMERAS_API void DrawTextView(const FVector3d& WorldPosition, const FVector2d& ScreenOffset, FStringView Text, const FLinearColor& TextColor, UFont* TextFont = nullptr, float TextScale = 1.f);
 
 public:
 
diff --git a/Engine/Plugins/Cameras/GameplayCameras/Source/GameplayCameras/Public/Debug/CameraDirectorDebugBlock.h b/Engine/Plugins/Cameras/GameplayCameras/Source/GameplayCameras/Public/Debug/CameraDirectorDebugBlock.h
new file mode 100644
--- /dev/null
+++ b/Engine/Plugins/Cameras/GameplayCameras/Source/GameplayCameras/Public/Debug/CameraDirectorDebugBlock.h
@@ -0,0 +1,47 @@
+// Copyright Epic Games, Inc. All Rights Reserved.
+
+#pragma once
+
+#include "Debug/CameraDebugBlock.h"
+#include "Math/Transform.h"
+
+#if UE_GAMEPLAY_CAMERAS_DEBUG
+
+namespace UE::Cameras
+{
+
+class FCameraEvaluationContext;
+struct FCameraDebugBlockBuilder;
+struct FCameraEvaluationContextStack;
+
+/**
+ * A debug block for a camera director entry in the evaluation stack.
+ */
+class FCameraDirectorDebugBlock : public FCameraDebugBlock
+{
+	UE_DECLARE_CAMERA_DEBUG_BLOCK(GAMEPLAYCAMERAS_API, FCameraDirectorDebugBlock)
+
+public:
+
+	void Initialize(TSharedPtr<FCameraEvaluationContext> Context, FCameraDebugBlockBuilder& Builder);
+
+protected:
+
+	virtual void OnDebugDraw(const FCameraDebugBlockDrawParams& Params, FCameraDebugRenderer& Renderer) override;
+	virtual void OnSerialize(FArchive& Ar) override;
+
+private:
+
+	FName ContextClassName;
+	FName OwnerClassName;
+	FString OwnerName;
+	FString CameraAssetName;
+	FName CameraDirectorClassName;
+	FTransform3d InitialContextTransform = FTransform3d::Identity;
+	bool bIsValid = false;
+};
+
+}  // namespace UE::Cameras
+
+#endif  // UE_GAMEPLAY_CAMERAS_DEBUG
+
diff --git a/Engine/Plugins/Cameras/GameplayCameras/Source/GameplayCameras/Public/Debug/CameraDirectorTreeDebugBlock.h b/Engine/Plugins/Cameras/GameplayCameras/Source/GameplayCameras/Public/Debug/CameraDirectorTreeDebugBlock.h
--- a/Engine/Plugins/Cameras/GameplayCameras/Source/GameplayCameras/Public/Debug/CameraDirectorTreeDebugBlock.h
+++ b/Engine/Plugins/Cameras/GameplayCameras/Source/GameplayCameras/Public/Debug/CameraDirectorTreeDebugBlock.h
@@ -3,7 +3,6 @@
 #pragma once
 
 #include "Debug/CameraDebugBlock.h"
-#include "Math/Transform.h"
 
 #if UE_GAMEPLAY_CAMERAS_DEBUG
 
@@ -23,40 +22,12 @@ class FCameraDirectorTreeDebugBlock : public FCameraDebugBlock
 
 public:
 
-	FCameraDirectorTreeDebugBlock();
-
 	void Initialize(const FCameraEvaluationContextStack& ContextStack, FCameraDebugBlockBuilder& Builder);
 	void Initialize(TArrayView<const TSharedPtr<FCameraEvaluationContext>> Contexts, FCameraDebugBlockBuilder& Builder);
 
 protected:
 
 	virtual void OnDebugDraw(const FCameraDebugBlockDrawParams& Params, FCameraDebugRenderer& Renderer) override;
-	virtual void OnSerialize(FArchive& Ar) override;
-
-private:
-
-	struct FDirectorDebugInfo;
-
-	void InitializeEntry(TSharedPtr<FCameraEvaluationContext> Context, FDirectorDebugInfo& EntryDebugInfo, FCameraDebugBlockBuilder& Builder);
-
-private:
-
-	struct FDirectorDebugInfo
-	{
-		FName ContextClassName;
-		FName OwnerClassName;
-		FString OwnerName;
-		FString CameraAssetName;
-		FName CameraDirectorClassName;
-		int32 NumChildrenContexts = 0;
-		FTransform3d InitialContextTransform = FTransform3d::Identity;
-		bool bIsValid = false;
-	};
-
-	TArray<FDirectorDebugInfo> CameraDirectors;
-	bool bIsTreeRoot = true;
-
-	friend FArchive& operator<< (FArchive&, FDirectorDebugInfo&);
 };
 
 }  // namespace UE::Cameras
diff --git a/Engine/Plugins/Cameras/GameplayCameras/Source/GameplayCameras/Public/Debug/CameraPoseDebugBlock.h b/Engine/Plugins/Cameras/GameplayCameras/Source/GameplayCameras/Public/Debug/CameraPoseDebugBlock.h
--- a/Engine/Plugins/Cameras/GameplayCameras/Source/GameplayCameras/Public/Debug/CameraPoseDebugBlock.h
+++ b/Engine/Plugins/Cameras/GameplayCameras/Source/GameplayCameras/Public/Debug/CameraPoseDebugBlock.h
@@ -10,6 +10,22 @@
 namespace UE::Cameras
 {
 
+/**
+ * Sets the global options for external rendering of a camera pose.
+ */
+struct FScopedGlobalCameraPoseRenderingParams
+{
+	GAMEPLAYCAMERAS_API FScopedGlobalCameraPoseRenderingParams(const FString& Label);
+	GAMEPLAYCAMERAS_API FScopedGlobalCameraPoseRenderingParams(const FLinearColor& LineColor);
+	GAMEPLAYCAMERAS_API FScopedGlobalCameraPoseRenderingParams(const FString& Label, const FLinearColor& LineColor);
+	GAMEPLAYCAMERAS_API ~FScopedGlobalCameraPoseRenderingParams();
+
+private:
+
+	FString PreviousLabel;
+	FLinearColor PreviousLineColor;
+};
+
 /**
  * A debug block that displays information about a camera pose.
  */
@@ -38,20 +54,6 @@ public:
 		return *this;
 	}
 
-	/** Sets the external rendering color. */
-	FCameraPoseDebugBlock& SetExternalRenderingLineColor(const FLinearColor& LineColor)
-	{
-		CameraPoseLineColor = LineColor;
-		return *this;
-	}
-
-	/** Sets the external rendering size. */
-	FCameraPoseDebugBlock& SetExternalRenderingSize(float CameraSize)
-	{
-		CameraPoseSize = CameraSize;
-		return *this;
-	}
-
 	/** 
 	 * Specifies the console variable to use to toggle between only showing camera pose
 	 * properties that were written to, or showing all camera pose properties.
@@ -69,12 +71,15 @@ protected:
 
 private:
 
+	static FLinearColor GlobalCameraPoseLineColor;
+	static FString GlobalCameraPoseLabel;
+
 	FCameraPose CameraPose;
 	FString ShowUnchangedCVarName;
-	FLinearColor CameraPoseLineColor;
-	float CameraPoseSize = -1.f;
 	bool bDrawText = true;
 	bool bDrawInExternalRendering = true;
+
+	friend struct FScopedGlobalCameraPoseRenderingParams;
 };
 
 }  // namespace UE::Cameras
diff --git a/Engine/Plugins/Cameras/GameplayCameras/Source/GameplayCameras/Public/Debug/CameraSystemDebugRegistry.h b/Engine/Plugins/Cameras/GameplayCameras/Source/GameplayCameras/Public/Debug/CameraSystemDebugRegistry.h
--- a/Engine/Plugins/Cameras/GameplayCameras/Source/GameplayCameras/Public/Debug/CameraSystemDebugRegistry.h
+++ b/Engine/Plugins/Cameras/GameplayCameras/Source/GameplayCameras/Public/Debug/CameraSystemDebugRegistry.h
@@ -88,6 +88,8 @@ public:
 	using FRegisteredCameraSystems = TArray<TSharedPtr<FCameraSystemEvaluator>>;
 	GAMEPLAYCAMERAS_API void GetRegisteredCameraSystemEvaluators(FRegisteredCameraSystems& OutEvaluators) const;
 
+	GAMEPLAYCAMERAS_API int32 NumRegisteredCameraSystemEvaluators() const;
+
 private:
 
 	struct FEntry
diff --git a/Engine/Plugins/Cameras/GameplayCameras/Source/GameplayCameras/Public/Debug/RootCameraDebugBlock.h b/Engine/Plugins/Cameras/GameplayCameras/Source/GameplayCameras/Public/Debug/RootCameraDebugBlock.h
--- a/Engine/Plugins/Cameras/GameplayCameras/Source/GameplayCameras/Public/Debug/RootCameraDebugBlock.h
+++ b/Engine/Plugins/Cameras/GameplayCameras/Source/GameplayCameras/Public/Debug/RootCameraDebugBlock.h
@@ -19,6 +19,21 @@ GAMEPLAYCAMERAS_API extern bool GGameplayCamerasDebugEnable;
 GAMEPLAYCAMERAS_API extern int32 GGameplayCamerasDebugSystemID;
 GAMEPLAYCAMERAS_API extern FString GGameplayCamerasDebugCategories;
 
+/**
+ * Parameters for drawing a hierarchy of debug blocks.
+ */
+struct FRootCameraDebugDrawParams
+{
+	/**
+	 * Whether the debug blocks belong to a camera system running as the camera manager
+	 * of active view target.
+	 */
+	bool bIsCameraManagerOrViewTarget = false;
+
+	/** Whether to force draw the debug blocks. */
+	bool bForceDraw = false;
+};
+
 /**
  * The root debug block for the camera system.
  */
@@ -32,11 +47,16 @@ public:
 	GAMEPLAYCAMERAS_API void BuildDebugBlocks(const FCameraSystemEvaluator& CameraSystem, const FCameraDebugBlockBuildParams& Params, FCameraDebugBlockBuilder& Builder);
 
 	/** Initiate the debug drawing. */
-	GAMEPLAYCAMERAS_API void RootDebugDraw(FCameraDebugRenderer& Renderer, bool bForceDraw = false);
+	GAMEPLAYCAMERAS_API void RootDebugDraw(const FRootCameraDebugDrawParams& Params, FCameraDebugRenderer& Renderer);
 
 	/** Gets the debug ID of the camera system that generated this debug info. */
 	const FCameraSystemDebugID& GetDebugID() const { return DebugID; }
 
+public:
+
+	/** Checks whether the debug data with the given ID would be drawn on screen. */
+	GAMEPLAYCAMERAS_API static bool ShouldDebugDraw(FCameraSystemDebugID InDebugID, bool bIsActive);
+
 protected:
 
 	// FCameraDebugBlock interface.
diff --git a/Engine/Plugins/Cameras/GameplayCameras/Source/GameplayCamerasEditor/Private/Debugger/SBlendStacksDebugPanel.cpp b/Engine/Plugins/Cameras/GameplayCameras/Source/GameplayCamerasEditor/Private/Debugger/SBlendStacksDebugPanel.cpp
--- a/Engine/Plugins/Cameras/GameplayCameras/Source/GameplayCamerasEditor/Private/Debugger/SBlendStacksDebugPanel.cpp
+++ b/Engine/Plugins/Cameras/GameplayCameras/Source/GameplayCamerasEditor/Private/Debugger/SBlendStacksDebugPanel.cpp
@@ -3,9 +3,8 @@
 #include "Debugger/SBlendStacksDebugPanel.h"
 
 #include "Debugger/SDebugWidgetUtils.h"
-#include "Widgets/Input/SCheckBox.h"
-#include "Widgets/Layout/SBox.h"
 #include "Widgets/SBoxPanel.h"
+#include "Widgets/Text/STextBlock.h"
 
 #define LOCTEXT_NAMESPACE "SBlendStacksDebugPanel"
 
@@ -17,6 +16,19 @@ void SBlendStacksDebugPanel::Construct(const FArguments& InArgs)
 	ChildSlot
 	[
 		SNew(SVerticalBox)
+		+ SVerticalBox::Slot()
+			.AutoHeight()
+			[
+				SNew(STextBlock)
+					.Margin(4.f)
+					.Text(LOCTEXT("FilterBlendStackNames", "Filter blend stack names:"))
+			]
+		+ SVerticalBox::Slot()
+			.AutoHeight()
+			[
+				SDebugWidgetUtils::CreateConsoleVariableTextBox(
+						TEXT("GameplayCameras.Debug.BlendStacks.Filter"))
+			]
 		+ SVerticalBox::Slot()
 			.AutoHeight()
 			[
diff --git a/Engine/Plugins/Cameras/GameplayCameras/Source/GameplayCamerasEditor/Private/Debugger/SCameraDirectorTreeDebugPanel.cpp b/Engine/Plugins/Cameras/GameplayCameras/Source/GameplayCamerasEditor/Private/Debugger/SCameraDirectorTreeDebugPanel.cpp
--- a/Engine/Plugins/Cameras/GameplayCameras/Source/GameplayCamerasEditor/Private/Debugger/SCameraDirectorTreeDebugPanel.cpp
+++ b/Engine/Plugins/Cameras/GameplayCameras/Source/GameplayCamerasEditor/Private/Debugger/SCameraDirectorTreeDebugPanel.cpp
@@ -3,7 +3,6 @@
 #include "Debugger/SCameraDirectorTreeDebugPanel.h"
 
 #include "Debugger/SDebugWidgetUtils.h"
-#include "Widgets/Input/SCheckBox.h"
 #include "Widgets/SBoxPanel.h"
 
 #define LOCTEXT_NAMESPACE "SCameraDirectorTreeDebugPanel"
@@ -23,6 +22,27 @@ void SCameraDirectorTreeDebugPanel::Construct(const FArguments& InArgs)
 						LOCTEXT("ShowUnchanged", "Show unchanged properties"),
 						TEXT("GameplayCameras.Debug.ContextInitialResult.ShowUnchanged"))
 			]
+		+ SVerticalBox::Slot()
+			.AutoHeight()
+			[
+				SDebugWidgetUtils::CreateConsoleVariableCheckBox(
+						LOCTEXT("ShowVariableIDs", "Show variable IDs"),
+						TEXT("GameplayCameras.Debug.ContextInitialResult.ShowVariableIDs"))
+			]
+		+ SVerticalBox::Slot()
+			.AutoHeight()
+			[
+				SDebugWidgetUtils::CreateConsoleVariableCheckBox(
+						LOCTEXT("ShowDataIDs", "Show data IDs"),
+						TEXT("GameplayCameras.Debug.ContextInitialResult.ShowDataIDs"))
+			]
+		+ SVerticalBox::Slot()
+			.AutoHeight()
+			[
+				SDebugWidgetUtils::CreateConsoleVariableCheckBox(
+						LOCTEXT("ShowCoordinateSystem", "Show context coordinate system"),
+						TEXT("GameplayCameras.Debug.ContextInitialResult.ShowCoordinateSystem"))
+			]
 	];
 }
 
diff --git a/Engine/Plugins/Cameras/GameplayCameras/Source/GameplayCamerasEditor/Private/Debugger/SCameraNodeTreeDebugPanel.cpp b/Engine/Plugins/Cameras/GameplayCameras/Source/GameplayCamerasEditor/Private/Debugger/SCameraNodeTreeDebugPanel.cpp
--- a/Engine/Plugins/Cameras/GameplayCameras/Source/GameplayCamerasEditor/Private/Debugger/SCameraNodeTreeDebugPanel.cpp
+++ b/Engine/Plugins/Cameras/GameplayCameras/Source/GameplayCamerasEditor/Private/Debugger/SCameraNodeTreeDebugPanel.cpp
@@ -4,7 +4,6 @@
 
 #include "Debugger/SDebugWidgetUtils.h"
 #include "Widgets/SBoxPanel.h"
-#include "Widgets/Input/SEditableTextBox.h"
 #include "Widgets/Text/STextBlock.h"
 
 #define LOCTEXT_NAMESPACE "SCameraNodeTreeDebugPanel"
@@ -30,6 +29,13 @@ void SCameraNodeTreeDebugPanel::Construct(const FArguments& InArgs)
 				SDebugWidgetUtils::CreateConsoleVariableTextBox(
 						TEXT("GameplayCameras.Debug.NodeTree.Filter"))
 			]
+		+ SVerticalBox::Slot()
+			.AutoHeight()
+			[
+				SDebugWidgetUtils::CreateConsoleVariableCheckBox(
+						LOCTEXT("ShowDampingLocalSpace", "Show damping local spaces"),
+						TEXT("GameplayCameras.Debug.Damping.ShowLocalSpace"))
+			]
 	];
 }
 
diff --git a/Engine/Plugins/Cameras/GameplayCameras/Source/GameplayCamerasEditor/Private/Debugger/SCameraPoseStatsDebugPanel.cpp b/Engine/Plugins/Cameras/GameplayCameras/Source/GameplayCamerasEditor/Private/Debugger/SCameraPoseStatsDebugPanel.cpp
--- a/Engine/Plugins/Cameras/GameplayCameras/Source/GameplayCamerasEditor/Private/Debugger/SCameraPoseStatsDebugPanel.cpp
+++ b/Engine/Plugins/Cameras/GameplayCameras/Source/GameplayCamerasEditor/Private/Debugger/SCameraPoseStatsDebugPanel.cpp
@@ -4,7 +4,6 @@
 
 #include "Debugger/SDebugWidgetUtils.h"
 #include "Widgets/SBoxPanel.h"
-#include "Widgets/Input/SCheckBox.h"
 
 #define LOCTEXT_NAMESPACE "SCameraPoseStatsDebugPanel"
 
diff --git a/Engine/Plugins/Cameras/GameplayCameras/Source/GameplayCamerasEditor/Private/Trace/CameraSystemRewindDebuggerExtension.cpp b/Engine/Plugins/Cameras/GameplayCameras/Source/GameplayCamerasEditor/Private/Trace/CameraSystemRewindDebuggerExtension.cpp
--- a/Engine/Plugins/Cameras/GameplayCameras/Source/GameplayCamerasEditor/Private/Trace/CameraSystemRewindDebuggerExtension.cpp
+++ b/Engine/Plugins/Cameras/GameplayCameras/Source/GameplayCamerasEditor/Private/Trace/CameraSystemRewindDebuggerExtension.cpp
@@ -55,14 +55,14 @@ void FCameraSystemRewindDebuggerExtension::Update(float DeltaTime, IRewindDebugg
 	if (FrameProvider.GetFrameFromTime(ETraceFrameType::TraceFrameType_Game, CurrentTraceTime, Frame))
 	{
 		const FCameraSystemTraceFrameData* FoundFrameData = nullptr;
-		const FCameraSystemDebugID WantedDebugID(GGameplayCamerasDebugSystemID);
+		const bool bIsActiveCamera = false; // TODO
 
 		const FCameraSystemTraceTimeline* CameraSystemTimeline = CameraSystemProvider->GetTimeline();
 		CameraSystemTimeline->EnumerateEvents(
 				Frame.StartTime, Frame.EndTime, 
-				[&FoundFrameData, WantedDebugID](double InStartTime, double InEndTime, uint32 InDepth, const FCameraSystemTraceFrameData& FrameData)  
+				[&FoundFrameData, bIsActiveCamera](double InStartTime, double InEndTime, uint32 InDepth, const FCameraSystemTraceFrameData& FrameData)  
 				{
-					if (WantedDebugID.IsAny() || WantedDebugID.GetValue() == FrameData.CameraSystemID)
+					if (FRootCameraDebugBlock::ShouldDebugDraw(FrameData.CameraSystemID, bIsActiveCamera))
 					{
 						FoundFrameData = &FrameData;
 					}
@@ -131,8 +131,9 @@ void FCameraSystemRewindDebuggerExtension::DebugDraw(UCanvas* Canvas, APlayerCon
 			bIsExternalRendering = !LevelViewportClient.IsAnyActorLocked();
 		}
 
+		FRootCameraDebugDrawParams Params;
 		FCameraDebugRenderer CameraDebugRenderer(VisualizedWorld, Canvas, bIsExternalRendering);
-		RootDebugBlock->RootDebugDraw(CameraDebugRenderer);
+		RootDebugBlock->RootDebugDraw(Params, CameraDebugRenderer);
 	}
 }
 
