Cameras: fix single-rig updates missing some of the steps the blend stack does

diff --git a/Engine/Plugins/Cameras/GameplayCameras/Source/GameplayCameras/Private/Core/DefaultRootCameraNode.cpp b/Engine/Plugins/Cameras/GameplayCameras/Source/GameplayCameras/Private/Core/DefaultRootCameraNode.cpp
--- a/Engine/Plugins/Cameras/GameplayCameras/Source/GameplayCameras/Private/Core/DefaultRootCameraNode.cpp
+++ b/Engine/Plugins/Cameras/GameplayCameras/Source/GameplayCameras/Private/Core/DefaultRootCameraNode.cpp
@@ -5,12 +5,15 @@
 #include "Core/BlendStackCameraNode.h"
 #include "Core/CameraEvaluationContext.h"
 #include "Core/CameraNodeEvaluatorHierarchy.h"
+#include "Core/CameraSystemEvaluator.h"
 #include "Core/PersistentBlendStackCameraNode.h"
 #include "Core/RootCameraNodeCameraRigEvent.h"
 #include "Core/TransientBlendStackCameraNode.h"
 #include "Debug/BlendStacksCameraDebugBlock.h"
 #include "Debug/CameraDebugBlockBuilder.h"
 #include "Debug/RootCameraDebugBlock.h"
+#include "Math/ColorList.h"
+#include "Services/CameraParameterSetterService.h"
 
 #include UE_INLINE_GENERATED_CPP_BY_NAME(DefaultRootCameraNode)
 
@@ -74,6 +77,13 @@ FCameraNodeEvaluatorChildrenView FDefaultRootCameraNodeEvaluator::OnGetChildren(
 	return FCameraNodeEvaluatorChildrenView({ BaseLayer, MainLayer, GlobalLayer, VisualLayer });
 }
 
+void FDefaultRootCameraNodeEvaluator::OnInitialize(const FCameraNodeEvaluatorInitializeParams& Params, FCameraNodeEvaluationResult& OutResult)
+{
+	Super::OnInitialize(Params, OutResult);
+
+	ParameterSetterService = Params.Evaluator->FindEvaluationService<FCameraParameterSetterService>();
+}
+
 void FDefaultRootCameraNodeEvaluator::OnRun(const FCameraNodeEvaluationParams& Params, FCameraNodeEvaluationResult& OutResult)
 {
 	BaseLayer->Run(Params, OutResult);
@@ -232,6 +242,9 @@ void FDefaultRootCameraNodeEvaluator::OnBuildSingleCameraRigHierarchy(const FSin
 
 void FDefaultRootCameraNodeEvaluator::OnRunSingleCameraRig(const FSingleCameraRigEvaluationParams& Params, FCameraNodeEvaluationResult& OutResult)
 {
+	ensure(Params.EvaluationParams.EvaluationContext == Params.CameraRigInfo.EvaluationContext);
+	TSharedPtr<const FCameraEvaluationContext> EvaluationContext = Params.EvaluationParams.EvaluationContext;
+
 	BaseLayer->Run(Params.EvaluationParams, OutResult);
 
 	FCameraNodeEvaluator* RootEvaluator = Params.CameraRigInfo.RootEvaluator;
@@ -239,6 +252,25 @@ void FDefaultRootCameraNodeEvaluator::OnRunSingleCameraRig(const FSingleCameraRi
 	// Emulate what the main blend stack does.
 
 	{
+		const FCameraNodeEvaluationResult& InitialResult(EvaluationContext->GetInitialResult());
+		OutResult.VariableTable.OverrideAll(InitialResult.VariableTable, true);
+		OutResult.ContextDataTable.OverrideAll(InitialResult.ContextDataTable);
+
+		if (Params.EvaluationParams.bIsActiveCameraRig)
+		{
+			const FCameraNodeEvaluationResult* ActiveOnlyResult = EvaluationContext->GetConditionalResult(ECameraEvaluationDataCondition::ActiveCameraRig);
+			if (ActiveOnlyResult)
+			{
+				OutResult.VariableTable.OverrideAll(ActiveOnlyResult->VariableTable, true);
+				OutResult.ContextDataTable.OverrideAll(ActiveOnlyResult->ContextDataTable);
+			}
+		}
+
+		if (ParameterSetterService)
+		{
+			ParameterSetterService->ApplyCameraVariableSetters(OutResult.VariableTable);
+		}
+
 		const FCameraNodeEvaluationResult* CameraRigResult = Params.CameraRigInfo.LastResult;
 		FCameraBlendedParameterUpdateParams InputParams(Params.EvaluationParams, CameraRigResult->CameraPose);
 		FCameraBlendedParameterUpdateResult InputResult(OutResult.VariableTable);
@@ -250,11 +282,15 @@ void FDefaultRootCameraNodeEvaluator::OnRunSingleCameraRig(const FSingleCameraRi
 	// No parameter blending: we are running this camera rig in isolation.
 
 	{
-		const FCameraNodeEvaluationResult& InitialResult = Params.CameraRigInfo.EvaluationContext->GetInitialResult();
+		const FCameraNodeEvaluationResult& InitialResult(EvaluationContext->GetInitialResult());
 		OutResult.CameraPose.OverrideChanged(InitialResult.CameraPose);
-		OutResult.VariableTable.OverrideAll(InitialResult.VariableTable);
+		OutResult.bIsCameraCut = (OutResult.bIsCameraCut || InitialResult.bIsCameraCut);
+		OutResult.bIsValid = true;
 
-		RootEvaluator->Run(Params.EvaluationParams, OutResult);
+		if (RootEvaluator)
+		{
+			RootEvaluator->Run(Params.EvaluationParams, OutResult);
+		}
 	}
 
 	GlobalLayer->Run(Params.EvaluationParams, OutResult);
diff --git a/Engine/Plugins/Cameras/GameplayCameras/Source/GameplayCameras/Public/Core/DefaultRootCameraNode.h b/Engine/Plugins/Cameras/GameplayCameras/Source/GameplayCameras/Public/Core/DefaultRootCameraNode.h
--- a/Engine/Plugins/Cameras/GameplayCameras/Source/GameplayCameras/Public/Core/DefaultRootCameraNode.h
+++ b/Engine/Plugins/Cameras/GameplayCameras/Source/GameplayCameras/Public/Core/DefaultRootCameraNode.h
@@ -4,6 +4,7 @@
 
 #include "Core/RootCameraNode.h"
 #include "Core/BlendStackCameraRigEvent.h"
+#include "Templates/SharedPointerFwd.h"
 
 #include "DefaultRootCameraNode.generated.h"
 
@@ -43,6 +44,7 @@ public:
 namespace UE::Cameras
 {
 
+class FCameraParameterSetterService;
 class FPersistentBlendStackCameraNodeEvaluator;
 class FTransientBlendStackCameraNodeEvaluator;
 
@@ -51,13 +53,14 @@ class FTransientBlendStackCameraNodeEvaluator;
  */
 class FDefaultRootCameraNodeEvaluator : public FRootCameraNodeEvaluator
 {
-	UE_DECLARE_CAMERA_NODE_EVALUATOR(GAMEPLAYCAMERAS_API, FDefaultRootCameraNodeEvaluator)
+	UE_DECLARE_CAMERA_NODE_EVALUATOR_EX(GAMEPLAYCAMERAS_API, FDefaultRootCameraNodeEvaluator, FRootCameraNodeEvaluator)
 
 protected:
 
 	// FCameraNodeEvaluator interface.
 	virtual FCameraNodeEvaluatorChildrenView OnGetChildren() override;
 	virtual void OnBuild(const FCameraNodeEvaluatorBuildParams& Params) override;
+	virtual void OnInitialize(const FCameraNodeEvaluatorInitializeParams& Params, FCameraNodeEvaluationResult& OutResult) override;
 	virtual void OnRun(const FCameraNodeEvaluationParams& Params, FCameraNodeEvaluationResult& OutResult) override;
 	virtual void OnExecuteOperation(const FCameraOperationParams& Params, FCameraOperation& Operation) override;
 
@@ -89,6 +92,8 @@ private:
 	FTransientBlendStackCameraNodeEvaluator* MainLayer;
 	FPersistentBlendStackCameraNodeEvaluator* GlobalLayer;
 	FPersistentBlendStackCameraNodeEvaluator* VisualLayer;
+
+	TSharedPtr<FCameraParameterSetterService> ParameterSetterService;
 };
 
 }  // namespace UE::Cameras
