Cameras: don't update dampers and interpolators when the delta-time is zero

diff --git a/Engine/Plugins/Cameras/GameplayCameras/Source/GameplayCameras/Private/Math/CriticalDamper.cpp b/Engine/Plugins/Cameras/GameplayCameras/Source/GameplayCameras/Private/Math/CriticalDamper.cpp
--- a/Engine/Plugins/Cameras/GameplayCameras/Source/GameplayCameras/Private/Math/CriticalDamper.cpp
+++ b/Engine/Plugins/Cameras/GameplayCameras/Source/GameplayCameras/Private/Math/CriticalDamper.cpp
@@ -40,10 +40,13 @@ float FCriticalDamper::Update(float DeltaTime)
 
 float FCriticalDamper::Update(float X, float DeltaTime)
 {
-	// Last frame we were at X0. This frame we are at X, so something forcibly moved
-	// us by the difference.
-	const float ForcedMovement = (X - X0);
-	InternalUpdate(ForcedMovement, DeltaTime);
+	if (DeltaTime != 0.f)
+	{
+		// Last frame we were at X0. This frame we are at X, so something forcibly moved
+		// us by the difference.
+		const float ForcedMovement = (X - X0);
+		InternalUpdate(ForcedMovement, DeltaTime);
+	}
 	return X0;
 }
 
@@ -57,6 +60,7 @@ float FCriticalDamper::Update(float PreviousDamped, float NextTarget, float Delt
 
 void FCriticalDamper::InternalUpdate(float ForcedMovement, float DeltaTime)
 {
+	ensure(DeltaTime != 0.f);
 	if (W0 > 0.f)
 	{
 		// We need to move the base of the spring by ForcedMovement, and then run the
diff --git a/Engine/Plugins/Cameras/GameplayCameras/Source/GameplayCameras/Public/Core/CameraValueInterpolator.h b/Engine/Plugins/Cameras/GameplayCameras/Source/GameplayCameras/Public/Core/CameraValueInterpolator.h
--- a/Engine/Plugins/Cameras/GameplayCameras/Source/GameplayCameras/Public/Core/CameraValueInterpolator.h
+++ b/Engine/Plugins/Cameras/GameplayCameras/Source/GameplayCameras/Public/Core/CameraValueInterpolator.h
@@ -103,7 +103,10 @@ public:
 	/** Evaluates the interpolator, advancing the current value towards the target. */
 	ValueType Run(const FCameraValueInterpolationParams& Params, FCameraValueInterpolationResult& OutResult)
 	{
-		OnRun(Params, OutResult);
+		if (Params.DeltaTime != 0.f)
+		{
+			OnRun(Params, OutResult);
+		}
 		return CurrentValue;
 	}
 
