Cameras: move view-target blend to a different file

diff --git a/Engine/Plugins/Cameras/GameplayCameras/Source/GameplayCameras/Private/GameFramework/GameplayCamerasPlayerCameraManager.cpp b/Engine/Plugins/Cameras/GameplayCameras/Source/GameplayCameras/Private/GameFramework/GameplayCamerasPlayerCameraManager.cpp
--- a/Engine/Plugins/Cameras/GameplayCameras/Source/GameplayCameras/Private/GameFramework/GameplayCamerasPlayerCameraManager.cpp
+++ b/Engine/Plugins/Cameras/GameplayCameras/Source/GameplayCameras/Private/GameFramework/GameplayCamerasPlayerCameraManager.cpp
@@ -16,6 +16,7 @@
 #include "GameFramework/ActorCameraEvaluationContext.h"
 #include "GameFramework/GameplayCameraComponentBase.h"
 #include "GameFramework/PlayerController.h"
+#include "GameFramework/ViewTargetTransitionParamsBlendNode.h"
 #include "GameplayCamerasSettings.h"
 #include "Services/CameraModifierService.h"
 #include "UObject/Package.h"
@@ -23,56 +24,6 @@
 
 #include UE_INLINE_GENERATED_CPP_BY_NAME(GameplayCamerasPlayerCameraManager)
 
-namespace UE::Cameras
-{
-
-/**
- * A blend node evaluator for UViewTargetTransitionParamsBlendCameraNode, which emulates the basic engine
- * blend curves for view targets.
- */
-class FViewTargetTransitionParamsBlendCameraNodeEvaluator : public FSimpleBlendCameraNodeEvaluator
-{
-	UE_DECLARE_BLEND_CAMERA_NODE_EVALUATOR_EX(GAMEPLAYCAMERAS_API, FViewTargetTransitionParamsBlendCameraNodeEvaluator, FSimpleBlendCameraNodeEvaluator)
-
-protected:
-
-	// FSimpleBlendCameraNodeEvaluator interface.
-	virtual void OnRun(const FCameraNodeEvaluationParams& Params, FCameraNodeEvaluationResult& OutResult) override;
-	virtual void OnComputeBlendFactor(const FCameraNodeEvaluationParams& Params, FSimpleBlendCameraNodeEvaluationResult& OutResult) override;
-
-private:
-
-	float CurrentTime = 0.f;
-};
-
-void FViewTargetTransitionParamsBlendCameraNodeEvaluator::OnRun(const FCameraNodeEvaluationParams& Params, FCameraNodeEvaluationResult& OutResult)
-{
-	const UViewTargetTransitionParamsBlendCameraNode* TransitionParamsNode = GetCameraNodeAs<UViewTargetTransitionParamsBlendCameraNode>();
-	CurrentTime += Params.DeltaTime;
-	if (CurrentTime >= TransitionParamsNode->TransitionParams.BlendTime)
-	{
-		CurrentTime = TransitionParamsNode->TransitionParams.BlendTime;
-		SetBlendFinished();
-	}
-
-	FSimpleBlendCameraNodeEvaluator::OnRun(Params, OutResult);
-}
-
-void FViewTargetTransitionParamsBlendCameraNodeEvaluator::OnComputeBlendFactor(const FCameraNodeEvaluationParams& Params, FSimpleBlendCameraNodeEvaluationResult& OutResult)
-{
-	const UViewTargetTransitionParamsBlendCameraNode* TransitionParamsNode = GetCameraNodeAs<UViewTargetTransitionParamsBlendCameraNode>();
-	float TimeFactor = 1.f;
-	if (TransitionParamsNode->TransitionParams.BlendTime > 0.f)
-	{
-		TimeFactor = CurrentTime / TransitionParamsNode->TransitionParams.BlendTime;
-	}
-	OutResult.BlendFactor = TransitionParamsNode->TransitionParams.GetBlendAlpha(TimeFactor);
-}
-
-UE_DEFINE_BLEND_CAMERA_NODE_EVALUATOR(FViewTargetTransitionParamsBlendCameraNodeEvaluator)
-
-}  // namespace UE::Cameras
-
 AGameplayCamerasPlayerCameraManager::AGameplayCamerasPlayerCameraManager(const FObjectInitializer& ObjectInitializer)
 	: Super(ObjectInitializer)
 {
@@ -428,9 +379,3 @@ void AGameplayCamerasPlayerCameraManager::DisplayDebug(UCanvas* Canvas, const FD
 	Super::DisplayDebug(Canvas, DebugDisplay, YL, YPos);
 }
 
-FCameraNodeEvaluatorPtr UViewTargetTransitionParamsBlendCameraNode::OnBuildEvaluator(FCameraNodeEvaluatorBuilder& Builder) const
-{
-	using namespace UE::Cameras;
-	return Builder.BuildEvaluator<FViewTargetTransitionParamsBlendCameraNodeEvaluator>();
-}
-
diff --git a/Engine/Plugins/Cameras/GameplayCameras/Source/GameplayCameras/Private/GameFramework/ViewTargetTransitionParamsBlendNode.cpp b/Engine/Plugins/Cameras/GameplayCameras/Source/GameplayCameras/Private/GameFramework/ViewTargetTransitionParamsBlendNode.cpp
new file mode 100644
--- /dev/null
+++ b/Engine/Plugins/Cameras/GameplayCameras/Source/GameplayCameras/Private/GameFramework/ViewTargetTransitionParamsBlendNode.cpp
@@ -0,0 +1,70 @@
+// Copyright Epic Games, Inc. All Rights Reserved.
+
+#include "GameFramework/ViewTargetTransitionParamsBlendNode.h"
+
+#include UE_INLINE_GENERATED_CPP_BY_NAME(ViewTargetTransitionParamsBlendNode)
+
+namespace UE::Cameras
+{
+
+/**
+ * A blend node evaluator for UViewTargetTransitionParamsBlendCameraNode, which emulates the basic engine
+ * blend curves for view targets.
+ */
+class FViewTargetTransitionParamsBlendCameraNodeEvaluator : public FSimpleBlendCameraNodeEvaluator
+{
+	UE_DECLARE_BLEND_CAMERA_NODE_EVALUATOR_EX(GAMEPLAYCAMERAS_API, FViewTargetTransitionParamsBlendCameraNodeEvaluator, FSimpleBlendCameraNodeEvaluator)
+
+protected:
+
+	// FSimpleBlendCameraNodeEvaluator interface.
+	virtual void OnRun(const FCameraNodeEvaluationParams& Params, FCameraNodeEvaluationResult& OutResult) override;
+	virtual void OnComputeBlendFactor(const FCameraNodeEvaluationParams& Params, FSimpleBlendCameraNodeEvaluationResult& OutResult) override;
+	virtual void OnSerialize(const FCameraNodeEvaluatorSerializeParams& Params, FArchive& Ar) override;
+
+private:
+
+	float CurrentTime = 0.f;
+};
+
+void FViewTargetTransitionParamsBlendCameraNodeEvaluator::OnRun(const FCameraNodeEvaluationParams& Params, FCameraNodeEvaluationResult& OutResult)
+{
+	const UViewTargetTransitionParamsBlendCameraNode* TransitionParamsNode = GetCameraNodeAs<UViewTargetTransitionParamsBlendCameraNode>();
+	CurrentTime += Params.DeltaTime;
+	if (CurrentTime >= TransitionParamsNode->TransitionParams.BlendTime)
+	{
+		CurrentTime = TransitionParamsNode->TransitionParams.BlendTime;
+		SetBlendFinished();
+	}
+
+	FSimpleBlendCameraNodeEvaluator::OnRun(Params, OutResult);
+}
+
+void FViewTargetTransitionParamsBlendCameraNodeEvaluator::OnComputeBlendFactor(const FCameraNodeEvaluationParams& Params, FSimpleBlendCameraNodeEvaluationResult& OutResult)
+{
+	const UViewTargetTransitionParamsBlendCameraNode* TransitionParamsNode = GetCameraNodeAs<UViewTargetTransitionParamsBlendCameraNode>();
+	float TimeFactor = 1.f;
+	if (TransitionParamsNode->TransitionParams.BlendTime > 0.f)
+	{
+		TimeFactor = CurrentTime / TransitionParamsNode->TransitionParams.BlendTime;
+	}
+	OutResult.BlendFactor = TransitionParamsNode->TransitionParams.GetBlendAlpha(TimeFactor);
+}
+
+void FViewTargetTransitionParamsBlendCameraNodeEvaluator::OnSerialize(const FCameraNodeEvaluatorSerializeParams& Params, FArchive& Ar)
+{
+	Super::OnSerialize(Params, Ar);
+
+	Ar << CurrentTime;
+}
+
+UE_DEFINE_BLEND_CAMERA_NODE_EVALUATOR(FViewTargetTransitionParamsBlendCameraNodeEvaluator)
+
+}  // namespace UE::Cameras
+
+FCameraNodeEvaluatorPtr UViewTargetTransitionParamsBlendCameraNode::OnBuildEvaluator(FCameraNodeEvaluatorBuilder& Builder) const
+{
+	using namespace UE::Cameras;
+	return Builder.BuildEvaluator<FViewTargetTransitionParamsBlendCameraNodeEvaluator>();
+}
+
diff --git a/Engine/Plugins/Cameras/GameplayCameras/Source/GameplayCameras/Private/GameFramework/ViewTargetTransitionParamsBlendNode.h b/Engine/Plugins/Cameras/GameplayCameras/Source/GameplayCameras/Private/GameFramework/ViewTargetTransitionParamsBlendNode.h
new file mode 100644
--- /dev/null
+++ b/Engine/Plugins/Cameras/GameplayCameras/Source/GameplayCameras/Private/GameFramework/ViewTargetTransitionParamsBlendNode.h
@@ -0,0 +1,29 @@
+// Copyright Epic Games, Inc. All Rights Reserved.
+
+#pragma once
+
+#include "Camera/PlayerCameraManager.h"
+#include "Nodes/Blends/SimpleBlendCameraNode.h"
+
+#include "ViewTargetTransitionParamsBlendNode.generated.h"
+
+/**
+ * A blend node that implements the blend algorithms of the FViewTargetTransitionParams.
+ */
+UCLASS(MinimalAPI, Hidden)
+class UViewTargetTransitionParamsBlendCameraNode : public USimpleBlendCameraNode
+{
+	GENERATED_BODY()
+
+public:
+
+	/** The transition params to use. */
+	UPROPERTY()
+	FViewTargetTransitionParams TransitionParams;
+
+protected:
+
+	// UCameraNode interface.
+	virtual FCameraNodeEvaluatorPtr OnBuildEvaluator(FCameraNodeEvaluatorBuilder& Builder) const override;
+};
+
diff --git a/Engine/Plugins/Cameras/GameplayCameras/Source/GameplayCameras/Public/GameFramework/GameplayCamerasPlayerCameraManager.h b/Engine/Plugins/Cameras/GameplayCameras/Source/GameplayCameras/Public/GameFramework/GameplayCamerasPlayerCameraManager.h
--- a/Engine/Plugins/Cameras/GameplayCameras/Source/GameplayCameras/Public/GameFramework/GameplayCamerasPlayerCameraManager.h
+++ b/Engine/Plugins/Cameras/GameplayCameras/Source/GameplayCameras/Public/GameFramework/GameplayCamerasPlayerCameraManager.h
@@ -5,7 +5,6 @@
 #include "Camera/PlayerCameraManager.h"
 #include "Core/CameraEvaluationContext.h"
 #include "GameFramework/IGameplayCameraSystemHost.h"
-#include "Nodes/Blends/SimpleBlendCameraNode.h"
 
 #include "GameplayCamerasPlayerCameraManager.generated.h"
 
@@ -159,23 +158,3 @@ private:
 	bool bIsSettingNewViewTarget = false;
 };
 
-/**
- * A blend node that implements the blend algorithms of the FViewTargetTransitionParams.
- */
-UCLASS(MinimalAPI, Hidden)
-class UViewTargetTransitionParamsBlendCameraNode : public USimpleBlendCameraNode
-{
-	GENERATED_BODY()
-
-public:
-
-	/** The transition params to use. */
-	UPROPERTY()
-	FViewTargetTransitionParams TransitionParams;
-
-protected:
-
-	// UCameraNode interface.
-	virtual FCameraNodeEvaluatorPtr OnBuildEvaluator(FCameraNodeEvaluatorBuilder& Builder) const override;
-};
-
