GameplayCameras: Expose FCameraEvaluationContext OnActivate and OnDeactivate

Add virtual OnAddReferencedObjects to FCameraEvaluationContext so that derived camera evaluation contexts can inform the garbage collector about owned UObjects

Swarm 44891568

diff --git a/Engine/Plugins/Cameras/GameplayCameras/Source/GameplayCameras/Private/Core/CameraEvaluationContext.cpp b/Engine/Plugins/Cameras/GameplayCameras/Source/GameplayCameras/Private/Core/CameraEvaluationContext.cpp
--- a/Engine/Plugins/Cameras/GameplayCameras/Source/GameplayCameras/Private/Core/CameraEvaluationContext.cpp
+++ b/Engine/Plugins/Cameras/GameplayCameras/Source/GameplayCameras/Private/Core/CameraEvaluationContext.cpp
@@ -131,6 +131,8 @@ void FCameraEvaluationContext::AddReferencedObjects(FReferenceCollector& Collect
 {
 	Collector.AddReferencedObject(CameraAsset);
 
+	OnAddReferencedObjects(Collector);
+
 	if (DirectorEvaluator)
 	{
 		DirectorEvaluator->AddReferencedObjects(Collector);
diff --git a/Engine/Plugins/Cameras/GameplayCameras/Source/GameplayCameras/Public/Core/CameraEvaluationContext.h b/Engine/Plugins/Cameras/GameplayCameras/Source/GameplayCameras/Public/Core/CameraEvaluationContext.h
--- a/Engine/Plugins/Cameras/GameplayCameras/Source/GameplayCameras/Public/Core/CameraEvaluationContext.h
+++ b/Engine/Plugins/Cameras/GameplayCameras/Source/GameplayCameras/Public/Core/CameraEvaluationContext.h
@@ -176,8 +176,9 @@ public:
 
 protected:
 
-	virtual void OnActivate(const FCameraEvaluationContextActivateParams& Params) {}
-	virtual void OnDeactivate(const FCameraEvaluationContextDeactivateParams& Params) {}
+	GAMEPLAYCAMERAS_API virtual void OnActivate(const FCameraEvaluationContextActivateParams& Params) {}
+	GAMEPLAYCAMERAS_API virtual void OnDeactivate(const FCameraEvaluationContextDeactivateParams& Params) {}
+	GAMEPLAYCAMERAS_API virtual void OnAddReferencedObjects(FReferenceCollector& Collector) {}
 
 	FCameraSystemEvaluator* GetPrivateCameraSystemEvaluator() const { return CameraSystemEvaluator; }
 
