GameplayCameras:  Add some helper methods to FCameraDirectorEvaluationResult to allow adding/removing camera rigs on different layers

Swarm 44891607

diff --git a/Engine/Plugins/Cameras/GameplayCameras/Source/GameplayCameras/Public/Core/CameraDirectorEvaluator.h b/Engine/Plugins/Cameras/GameplayCameras/Source/GameplayCameras/Public/Core/CameraDirectorEvaluator.h
--- a/Engine/Plugins/Cameras/GameplayCameras/Source/GameplayCameras/Public/Core/CameraDirectorEvaluator.h
+++ b/Engine/Plugins/Cameras/GameplayCameras/Source/GameplayCameras/Public/Core/CameraDirectorEvaluator.h
@@ -132,6 +132,29 @@ struct FCameraDirectorEvaluationResult
 		Requests.Add({ InEvaluationContext, InCameraRigProxy });
 	}
 
+	/** Adds an activation request for the specified layer */
+	void Add(
+			TSharedPtr<const FCameraEvaluationContext> InEvaluationContext,
+			TObjectPtr<const UCameraRigAsset> InCameraRig,
+			ECameraRigLayer Layer, 
+			int32 OrderKey)
+	{
+		FCameraRigActivationDeactivationRequest& Request = Requests.Add_GetRef({ InEvaluationContext, InCameraRig });
+		Request.Layer = Layer;
+		Request.OrderKey = OrderKey;
+	}
+
+	/** Adds a deactivation request for the specified layer */
+	void Remove(
+				TSharedPtr<const FCameraEvaluationContext> InEvaluationContext, 
+				TObjectPtr<const UCameraRigAsset> InCameraRig,
+				ECameraRigLayer Layer)
+	{
+		FCameraRigActivationDeactivationRequest& Request = Requests.Add_GetRef({ InEvaluationContext, InCameraRig });
+		Request.Layer = Layer;
+		Request.RequestType = ECameraRigActivationDeactivationRequestType::Deactivate;
+	}
+
 	/** Reset this result. */
 	void Reset()
 	{
