Cameras: fix unnecessary double-building of the camera rig component's referenced asset

diff --git a/Engine/Plugins/Cameras/GameplayCameras/Source/GameplayCameras/Private/GameFramework/GameplayCameraRigComponent.cpp b/Engine/Plugins/Cameras/GameplayCameras/Source/GameplayCameras/Private/GameFramework/GameplayCameraRigComponent.cpp
--- a/Engine/Plugins/Cameras/GameplayCameras/Source/GameplayCameras/Private/GameFramework/GameplayCameraRigComponent.cpp
+++ b/Engine/Plugins/Cameras/GameplayCameras/Source/GameplayCameras/Private/GameFramework/GameplayCameraRigComponent.cpp
@@ -60,12 +60,32 @@ UCameraAsset* UGameplayCameraRigComponent::GetCameraAsset()
 		GeneratedCameraAsset = NewObject<UCameraAsset>(this, TEXT("GeneratedCameraAsset"), RF_Transient);
 		GeneratedCameraAsset->SetCameraDirector(SingleDirector);
 
-		GeneratedCameraAsset->BuildCamera();
+		ensure(!bIsBuildingGeneratedCameraAsset);
+		{
+			TGuardValue<bool> ReentrancyGuard(bIsBuildingGeneratedCameraAsset, true);
+			GeneratedCameraAsset->BuildCamera();
+		}
+		OnCameraRigAssetBuiltImpl();
 	}
 
 	return GeneratedCameraAsset;
 }
 
+void UGameplayCameraRigComponent::OnCameraRigAssetBuiltImpl()
+{
+	CameraRigReference.RebuildParametersIfNeeded();
+	if (HasCameraEvaluationContext())
+	{
+#if WITH_EDITOR
+		const UCameraRigAsset* CameraRigAsset = CameraRigReference.GetCameraRig();
+		const FCameraObjectAllocationInfo& AllocationInfo = CameraRigAsset->AllocationInfo;
+		ReinitializeCameraEvaluationContext(AllocationInfo.VariableTableInfo, AllocationInfo.ContextDataTableInfo);
+#endif  // WITH_EDITOR
+
+		UpdateCameraEvaluationContext(true);
+	}
+}
+
 void UGameplayCameraRigComponent::OnUpdateCameraEvaluationContext(bool bForceApplyParameterOverrides)
 {
 	using namespace UE::Cameras;
@@ -102,13 +122,7 @@ void UGameplayCameraRigComponent::OnCameraRigAssetBuilt(const UCameraRigAsset* I
 		TGuardValue<bool> ReentrancyGuard(bIsBuildingGeneratedCameraAsset, true);
 		GeneratedCameraAsset->BuildCamera();
 	}
-	CameraRigReference.RebuildParametersIfNeeded();
-	if (HasCameraEvaluationContext())
-	{
-		const FCameraObjectAllocationInfo& AllocationInfo = InCameraRigAsset->AllocationInfo;
-		ReinitializeCameraEvaluationContext(AllocationInfo.VariableTableInfo, AllocationInfo.ContextDataTableInfo);
-		UpdateCameraEvaluationContext(true);
-	}
+	OnCameraRigAssetBuiltImpl();
 }
 
 void UGameplayCameraRigComponent::PostEditChangeProperty( struct FPropertyChangedEvent& PropertyChangedEvent)
diff --git a/Engine/Plugins/Cameras/GameplayCameras/Source/GameplayCameras/Public/GameFramework/GameplayCameraRigComponent.h b/Engine/Plugins/Cameras/GameplayCameras/Source/GameplayCameras/Public/GameFramework/GameplayCameraRigComponent.h
--- a/Engine/Plugins/Cameras/GameplayCameras/Source/GameplayCameras/Public/GameFramework/GameplayCameraRigComponent.h
+++ b/Engine/Plugins/Cameras/GameplayCameras/Source/GameplayCameras/Public/GameFramework/GameplayCameraRigComponent.h
@@ -47,6 +47,8 @@ protected:
 
 private:
 
+	void OnCameraRigAssetBuiltImpl();
+
 #if WITH_EDITOR
 	void OnCameraRigAssetBuilt(const UCameraRigAsset* InCameraRigAsset);
 #endif
@@ -61,8 +63,6 @@ private:
 	UPROPERTY(Transient)
 	FInstancedPropertyBag CachedParameterOverrides;
 
-#if WITH_EDITOR
 	bool bIsBuildingGeneratedCameraAsset = false;
-#endif
 };
 
