Cameras: move overriden parameter IDs into a new struct variant of FInstancedPropertyBag

Move some code related to override parameter values inside a dedicated structure for property bags with overrides. This simplifies the code elsewhere a little.

Make use of Sequencer's new ability to animate properties directly inside a property bag.

Remove the unused `ParameterOwners` from the camera asset.

diff --git a/Engine/Plugins/Cameras/GameplayCameras/Source/GameplayCameras/Private/Build/CameraAssetBuilder.cpp b/Engine/Plugins/Cameras/GameplayCameras/Source/GameplayCameras/Private/Build/CameraAssetBuilder.cpp
--- a/Engine/Plugins/Cameras/GameplayCameras/Source/GameplayCameras/Private/Build/CameraAssetBuilder.cpp
+++ b/Engine/Plugins/Cameras/GameplayCameras/Source/GameplayCameras/Private/Build/CameraAssetBuilder.cpp
@@ -137,22 +137,19 @@ void FCameraAssetBuilder::BuildCameraImpl()
 
 	// Build the final list of parameter definitions.
 	TArray<FCameraObjectInterfaceParameterDefinition> ParameterDefinitions;
-	TArray<TObjectPtr<const UCameraRigAsset>> ParameterOwners;
 	for (const TPair<const UCameraRigAsset*, TArray<FCameraObjectInterfaceParameterDefinition>>& Pair : DefinitionsByCameraRig)
 	{
 		const TArray<FCameraObjectInterfaceParameterDefinition>& DefinitionsForCameraRig(Pair.Value);
 		for (const FCameraObjectInterfaceParameterDefinition& Definition : DefinitionsForCameraRig)
 		{
 			ParameterDefinitions.Add(Definition);
-			ParameterOwners.Add(Pair.Key);
 		}
 	}
 
-	if (ParameterDefinitions != CameraAsset->ParameterDefinitions || ParameterOwners != CameraAsset->ParameterOwners)
+	if (ParameterDefinitions != CameraAsset->ParameterDefinitions)
 	{
 		CameraAsset->Modify();
 		CameraAsset->ParameterDefinitions = ParameterDefinitions;
-		CameraAsset->ParameterOwners = ParameterOwners;
 	}
 
 	// Rebuild the default parameters property bag.
diff --git a/Engine/Plugins/Cameras/GameplayCameras/Source/GameplayCameras/Private/Build/CameraObjectInterfaceParameterBuilder.cpp b/Engine/Plugins/Cameras/GameplayCameras/Source/GameplayCameras/Private/Build/CameraObjectInterfaceParameterBuilder.cpp
--- a/Engine/Plugins/Cameras/GameplayCameras/Source/GameplayCameras/Private/Build/CameraObjectInterfaceParameterBuilder.cpp
+++ b/Engine/Plugins/Cameras/GameplayCameras/Source/GameplayCameras/Private/Build/CameraObjectInterfaceParameterBuilder.cpp
@@ -5,8 +5,8 @@
 #include "Core/BaseCameraObject.h"
 #include "Core/CameraNode.h"
 #include "Core/CameraObjectInterfaceParameterDefinition.h"
-#include "Core/CameraParameters.h"
-#include "Core/CameraVariableReferences.h"
+#include "Core/CameraParameters.h"  // IWYU pragma: keep
+#include "Core/CameraVariableReferences.h"  // IWYU pragma: keep
 #include "Core/ICustomCameraNodeParameterProvider.h"
 #include "Misc/EngineVersionComparison.h"
 #include "StructUtils/PropertyBag.h"
diff --git a/Engine/Plugins/Cameras/GameplayCameras/Source/GameplayCameras/Private/Core/BaseCameraObjectReference.cpp b/Engine/Plugins/Cameras/GameplayCameras/Source/GameplayCameras/Private/Core/BaseCameraObjectReference.cpp
--- a/Engine/Plugins/Cameras/GameplayCameras/Source/GameplayCameras/Private/Core/BaseCameraObjectReference.cpp
+++ b/Engine/Plugins/Cameras/GameplayCameras/Source/GameplayCameras/Private/Core/BaseCameraObjectReference.cpp
@@ -3,8 +3,7 @@
 #include "Core/BaseCameraObjectReference.h"
 
 #include "Core/BaseCameraObject.h"
-#include "Core/CameraNodeEvaluator.h"
-#include "Core/CameraParameters.h"
+#include "Core/CameraParameters.h"  // IWYU pragma: keep
 #include "Core/ICustomCameraNodeParameterProvider.h"
 
 #include UE_INLINE_GENERATED_CPP_BY_NAME(BaseCameraObjectReference)
@@ -37,17 +36,12 @@ FCameraObjectInterfaceParameterMetaData& FBaseCameraObjectReference::FindOrAddMe
 
 bool FBaseCameraObjectReference::IsParameterOverridden(const FGuid& PropertyID) const
 {
-	if (const FCameraObjectInterfaceParameterMetaData* MetaData = FindMetaData(PropertyID))
-	{
-		return MetaData->bIsOverridden;
-	}
-	return false;
+	return Parameters.IsPropertyOverriden(PropertyID);
 }
 
 void FBaseCameraObjectReference::SetParameterOverridden(const FGuid& PropertyID, bool bIsOverridden)
 {
-	FCameraObjectInterfaceParameterMetaData& MetaData = FindOrAddMetaData(PropertyID);
-	MetaData.bIsOverridden = bIsOverridden;
+	Parameters.SetPropertyOverriden(PropertyID, bIsOverridden);
 }
 
 bool FBaseCameraObjectReference::NeedsRebuildParameters() const
@@ -88,9 +82,7 @@ void FBaseCameraObjectReference::RebuildParameters()
 
 	if (CameraObject)
 	{
-		TArray<FGuid> ParameterOverrideGuids;
-		GetOverriddenParameterGuids(ParameterOverrideGuids);
-		Parameters.MigrateToNewBagInstanceWithOverrides(CameraObject->GetDefaultParameters(), ParameterOverrideGuids);
+		Parameters.MigrateToNewBagInstanceWithOverrides(CameraObject->GetDefaultParameters());
 
 		if (const UPropertyBag* ParametersType = Parameters.GetPropertyBagStruct())
 		{
@@ -270,3 +262,17 @@ void FBaseCameraObjectReference::GetCustomCameraNodeParameters(FCustomCameraNode
 	}
 }
 
+void FBaseCameraObjectReference::PostSerialize(const FArchive& Ar)
+{
+	PRAGMA_DISABLE_DEPRECATION_WARNINGS
+	for (FCameraObjectInterfaceParameterMetaData& MetaData : ParameterMetaData)
+	{
+		if (MetaData.bIsOverridden_DEPRECATED)
+		{
+			Parameters.SetPropertyOverriden(MetaData.ParameterGuid, true);
+			MetaData.bIsOverridden_DEPRECATED = false;
+		}
+	}
+	PRAGMA_ENABLE_DEPRECATION_WARNINGS
+}
+
diff --git a/Engine/Plugins/Cameras/GameplayCameras/Source/GameplayCameras/Private/Core/CameraAssetReference.cpp b/Engine/Plugins/Cameras/GameplayCameras/Source/GameplayCameras/Private/Core/CameraAssetReference.cpp
--- a/Engine/Plugins/Cameras/GameplayCameras/Source/GameplayCameras/Private/Core/CameraAssetReference.cpp
+++ b/Engine/Plugins/Cameras/GameplayCameras/Source/GameplayCameras/Private/Core/CameraAssetReference.cpp
@@ -22,9 +22,8 @@ void FCameraAssetReference::ApplyParameterOverrides(UE::Cameras::FCameraNodeEval
 	using namespace UE::Cameras;
 	if (CameraAsset)
 	{
-		TSet<FGuid> OverridenParameterGuids(GetOverriddenParameterGuids());
 		FCameraObjectInterfaceParameterOverrideHelper Helper(&OutResult.VariableTable, &OutResult.ContextDataTable);
-		Helper.ApplyParameterOverrides(CameraAsset, CameraAsset->GetParameterDefinitions(), Parameters, OverridenParameterGuids, bDrivenOnly);
+		Helper.ApplyParameterOverrides(CameraAsset, CameraAsset->GetParameterDefinitions(), Parameters, bDrivenOnly);
 	}
 }
 
@@ -40,19 +39,12 @@ void FCameraAssetReference::ApplyParameterOverrides(const FInstancedPropertyBag&
 
 bool FCameraAssetReference::IsParameterOverridden(const FGuid PropertyID) const
 {
-	return ParameterOverrideGuids.Contains(PropertyID);
+	return Parameters.IsPropertyOverriden(PropertyID);
 }
 
 void FCameraAssetReference::SetParameterOverridden(const FGuid PropertyID, bool bIsOverridden)
 {
-	if (bIsOverridden)
-	{
-		ParameterOverrideGuids.AddUnique(PropertyID);
-	}
-	else
-	{
-		ParameterOverrideGuids.Remove(PropertyID);
-	}
+	Parameters.SetPropertyOverriden(PropertyID, bIsOverridden);
 }
 
 bool FCameraAssetReference::NeedsRebuildParameters() const
@@ -89,19 +81,7 @@ void FCameraAssetReference::RebuildParameters()
 {
 	if (CameraAsset)
 	{
-		Parameters.MigrateToNewBagInstanceWithOverrides(CameraAsset->GetDefaultParameters(), ParameterOverrideGuids);
-		
-		// Remove overrides for parameters that don't exist anymore.
-		if (const UPropertyBag* ParametersType = Parameters.GetPropertyBagStruct())
-		{
-			for (TArray<FGuid>::TIterator It = ParameterOverrideGuids.CreateIterator(); It; ++It)
-			{
-				if (!ParametersType->FindPropertyDescByID(*It))
-				{
-					It.RemoveCurrentSwap();
-				}
-			}
-		}
+		Parameters.MigrateToNewBagInstanceWithOverrides(CameraAsset->GetDefaultParameters());
 	}
 	else
 	{
@@ -121,3 +101,14 @@ bool FCameraAssetReference::SerializeFromMismatchedTag(FPropertyTag const& Tag,
 	return false;
 }
 
+void FCameraAssetReference::PostSerialize(const FArchive& Ar)
+{
+	PRAGMA_DISABLE_DEPRECATION_WARNINGS
+	for (const FGuid& Guid : ParameterOverrideGuids_DEPRECATED)
+	{
+		Parameters.SetPropertyOverriden(Guid, true);
+	}
+	ParameterOverrideGuids_DEPRECATED.Reset();
+	PRAGMA_ENABLE_DEPRECATION_WARNINGS
+}
+
diff --git a/Engine/Plugins/Cameras/GameplayCameras/Source/GameplayCameras/Private/Core/CameraRigAssetReference.cpp b/Engine/Plugins/Cameras/GameplayCameras/Source/GameplayCameras/Private/Core/CameraRigAssetReference.cpp
--- a/Engine/Plugins/Cameras/GameplayCameras/Source/GameplayCameras/Private/Core/CameraRigAssetReference.cpp
+++ b/Engine/Plugins/Cameras/GameplayCameras/Source/GameplayCameras/Private/Core/CameraRigAssetReference.cpp
@@ -59,10 +59,8 @@ void FCameraRigAssetReference::ApplyParameterOverridesImpl(UE::Cameras::FCameraV
 
 	if (CameraRig)
 	{
-		TSet<FGuid> OverridenParameterGuids;
-		GetOverriddenParameterGuids(OverridenParameterGuids);
 		FCameraObjectInterfaceParameterOverrideHelper Helper(OutVariableTable, OutContextDataTable);
-		Helper.ApplyParameterOverrides(CameraRig, CameraRig->GetParameterDefinitions(), Parameters, OverridenParameterGuids, bDrivenOnly);
+		Helper.ApplyParameterOverrides(CameraRig, CameraRig->GetParameterDefinitions(), Parameters, bDrivenOnly);
 	}
 }
 
@@ -98,7 +96,7 @@ void FCameraRigAssetReference::PostSerialize(const FArchive& Ar)
 		LegacyParameterProperties.Add(LegacyParameterProperty);\
 		FCameraObjectInterfaceParameterMetaData MetaData;\
 		MetaData.ParameterGuid = ParameterOverride.InterfaceParameterGuid;\
-		MetaData.bIsOverridden = true;\
+		MetaData.bIsOverridden_DEPRECATED = true;\
 		LegacyParameterMetaData.Add(MetaData);\
 		bHasAnyLegacyOverride = true;\
 	}
@@ -107,7 +105,7 @@ void FCameraRigAssetReference::PostSerialize(const FArchive& Ar)
 
 	if (bHasAnyLegacyOverride)
 	{
-		Parameters = FInstancedPropertyBag();
+		Parameters = FInstancedOverridablePropertyBag();
 		Parameters.AddProperties(LegacyParameterProperties);
 
 		ParameterMetaData = LegacyParameterMetaData;
@@ -132,11 +130,13 @@ void FCameraRigAssetReference::PostSerialize(const FArchive& Ar)
 		{
 			FCameraObjectInterfaceParameterMetaData MetaData;
 			MetaData.ParameterGuid = Guid;
-			MetaData.bIsOverridden = true;
+			MetaData.bIsOverridden_DEPRECATED = true;
 			ParameterMetaData.Add(MetaData);
 		}
 
 		ParameterOverrideGuids_DEPRECATED.Reset();
 	}
+
+	FBaseCameraObjectReference::PostSerialize(Ar);
 }
 
diff --git a/Engine/Plugins/Cameras/GameplayCameras/Source/GameplayCameras/Private/Core/CameraShakeAssetReference.cpp b/Engine/Plugins/Cameras/GameplayCameras/Source/GameplayCameras/Private/Core/CameraShakeAssetReference.cpp
--- a/Engine/Plugins/Cameras/GameplayCameras/Source/GameplayCameras/Private/Core/CameraShakeAssetReference.cpp
+++ b/Engine/Plugins/Cameras/GameplayCameras/Source/GameplayCameras/Private/Core/CameraShakeAssetReference.cpp
@@ -43,10 +43,8 @@ void FCameraShakeAssetReference::ApplyParameterOverridesImpl(UE::Cameras::FCamer
 	
 	if (CameraShake)
 	{
-		TSet<FGuid> OverridenParameterGuids;
-		GetOverriddenParameterGuids(OverridenParameterGuids);
 		FCameraObjectInterfaceParameterOverrideHelper Helper(OutVariableTable, OutContextDataTable);
-		Helper.ApplyParameterOverrides(CameraShake, CameraShake->GetParameterDefinitions(), Parameters, OverridenParameterGuids, bDrivenOnly);
+		Helper.ApplyParameterOverrides(CameraShake, CameraShake->GetParameterDefinitions(), Parameters, bDrivenOnly);
 	}
 }
 
diff --git a/Engine/Plugins/Cameras/GameplayCameras/Source/GameplayCameras/Private/GameFramework/CameraRigParameterInterop.cpp b/Engine/Plugins/Cameras/GameplayCameras/Source/GameplayCameras/Private/GameFramework/CameraRigParameterInterop.cpp
--- a/Engine/Plugins/Cameras/GameplayCameras/Source/GameplayCameras/Private/GameFramework/CameraRigParameterInterop.cpp
+++ b/Engine/Plugins/Cameras/GameplayCameras/Source/GameplayCameras/Private/GameFramework/CameraRigParameterInterop.cpp
@@ -82,9 +82,6 @@ DEFINE_FUNCTION(UCameraRigParameterInterop::execGetCameraParameter)
 		}
 		else
 		{
-			const FInstancedPropertyBag& DefaultParameters = CameraRig->GetDefaultParameters();
-			const UPropertyBag* PropertyBag = DefaultParameters.GetPropertyBagStruct();
-
 			if (ParameterDefinition->ParameterType == ECameraObjectInterfaceParameterType::Blendable)
 			{
 				const FCameraVariableTable& VariableTable = CameraData.GetResult()->VariableTable;
@@ -163,9 +160,6 @@ DEFINE_FUNCTION(UCameraRigParameterInterop::execSetCameraParameter)
 		}
 		else
 		{
-			const FInstancedPropertyBag& DefaultParameters = CameraRig->GetDefaultParameters();
-			const UPropertyBag* PropertyBag = DefaultParameters.GetPropertyBagStruct();
-
 			if (ParameterDefinition->ParameterType == ECameraObjectInterfaceParameterType::Blendable)
 			{
 				FCameraVariableTable& VariableTable = CameraData.GetResult()->VariableTable;
diff --git a/Engine/Plugins/Cameras/GameplayCameras/Source/GameplayCameras/Private/GameFramework/GameplayCameraParameterSetterComponent.cpp b/Engine/Plugins/Cameras/GameplayCameras/Source/GameplayCameras/Private/GameFramework/GameplayCameraParameterSetterComponent.cpp
--- a/Engine/Plugins/Cameras/GameplayCameras/Source/GameplayCameras/Private/GameFramework/GameplayCameraParameterSetterComponent.cpp
+++ b/Engine/Plugins/Cameras/GameplayCameras/Source/GameplayCameras/Private/GameFramework/GameplayCameraParameterSetterComponent.cpp
@@ -96,7 +96,7 @@ void UGameplayCameraParameterSetterComponent::StartParameterSetters()
 	}
 
 	const UCameraRigAsset* CameraRig = CameraRigReference.GetCameraRig();
-	const FInstancedPropertyBag& ParameterValues = CameraRigReference.GetParameters();
+	const FInstancedOverridablePropertyBag& ParameterValues = CameraRigReference.GetParameters();
 	const uint8* ParameterValuesPtr = ParameterValues.GetValue().GetMemory();
 
 	for (const FCameraObjectInterfaceParameterDefinition& ParameterDefinition : CameraRig->GetParameterDefinitions())
@@ -106,7 +106,7 @@ void UGameplayCameraParameterSetterComponent::StartParameterSetters()
 			continue;
 		}
 
-		if (!CameraRigReference.IsParameterOverridden(ParameterDefinition.ParameterGuid))
+		if (!ParameterValues.IsPropertyOverriden(ParameterDefinition.ParameterGuid))
 		{
 			continue;
 		}
diff --git a/Engine/Plugins/Cameras/GameplayCameras/Source/GameplayCameras/Private/Helpers/CameraObjectInterfaceParameterOverrideHelper.cpp b/Engine/Plugins/Cameras/GameplayCameras/Source/GameplayCameras/Private/Helpers/CameraObjectInterfaceParameterOverrideHelper.cpp
--- a/Engine/Plugins/Cameras/GameplayCameras/Source/GameplayCameras/Private/Helpers/CameraObjectInterfaceParameterOverrideHelper.cpp
+++ b/Engine/Plugins/Cameras/GameplayCameras/Source/GameplayCameras/Private/Helpers/CameraObjectInterfaceParameterOverrideHelper.cpp
@@ -8,7 +8,7 @@
 #include "Core/CameraParameters.h"
 #include "Core/CameraVariableTable.h"
 #include "Misc/EngineVersionComparison.h"
-#include "StructUtils/PropertyBag.h"
+#include "StructUtils/OverridablePropertyBag.h"
 
 namespace UE::Cameras
 {
@@ -46,7 +46,7 @@ void ApplyBlendableParameterOverride(
 void ApplyBlendableParameterOverride(
 		const UObject* CameraObject,
 		const FCameraObjectInterfaceParameterDefinition& ParameterDefinition,
-		const FInstancedPropertyBag& PropertyBag,
+		const FInstancedOverridablePropertyBag& PropertyBag,
 		const FPropertyBagPropertyDesc& PropertyBagPropertyDesc,
 		FCameraVariableTable& VariableTable,
 		bool bDrivenOnly)
@@ -192,7 +192,7 @@ void ApplyDataParameterElementOverride(
 void ApplyDataParameterSingleOverride(
 		const UObject* CameraObject,
 		const FCameraObjectInterfaceParameterDefinition& ParameterDefinition,
-		const FInstancedPropertyBag& PropertyBag,
+		const FInstancedOverridablePropertyBag& PropertyBag,
 		const FPropertyBagPropertyDesc& PropertyBagPropertyDesc,
 		FCameraContextDataTable& ContextDataTable)
 {
@@ -275,7 +275,7 @@ void ApplyDataParameterSingleOverride(
 void ApplyDataParameterArrayOverride(
 		const UObject* CameraObject,
 		const FCameraObjectInterfaceParameterDefinition& ParameterDefinition,
-		const FInstancedPropertyBag& PropertyBag,
+		const FInstancedOverridablePropertyBag& PropertyBag,
 		const FPropertyBagPropertyDesc& PropertyBagPropertyDesc,
 		FCameraContextDataTable& ContextDataTable)
 {
@@ -362,7 +362,7 @@ void ApplyDataParameterArrayOverride(
 void ApplyDataParameterOverride(
 		const UObject* CameraObject,
 		const FCameraObjectInterfaceParameterDefinition& ParameterDefinition,
-		const FInstancedPropertyBag& PropertyBag,
+		const FInstancedOverridablePropertyBag& PropertyBag,
 		const FPropertyBagPropertyDesc& PropertyBagPropertyDesc,
 		FCameraContextDataTable& ContextDataTable)
 {
@@ -401,8 +401,7 @@ FCameraObjectInterfaceParameterOverrideHelper::FCameraObjectInterfaceParameterOv
 void FCameraObjectInterfaceParameterOverrideHelper::ApplyParameterOverrides(
 		const UObject* CameraObject,
 		TConstArrayView<FCameraObjectInterfaceParameterDefinition> ParameterDefinitions,
-		const FInstancedPropertyBag& ParameterOverrides,
-		const TSet<FGuid>& OverridenParameterGuids,
+		const FInstancedOverridablePropertyBag& ParameterOverrides,
 		bool bDrivenOnly)
 {
 	using namespace Internal;
@@ -420,7 +419,7 @@ void FCameraObjectInterfaceParameterOverrideHelper::ApplyParameterOverrides(
 			continue;
 		}
 
-		if (!OverridenParameterGuids.Contains(Definition.ParameterGuid))
+		if (!ParameterOverrides.IsPropertyOverriden(Definition.ParameterGuid))
 		{
 			continue;
 		}
@@ -438,7 +437,7 @@ void FCameraObjectInterfaceParameterOverrideHelper::ApplyParameterOverrides(
 void FCameraObjectInterfaceParameterOverrideHelper::ApplyParameterOverrides(
 		const UObject* CameraObject,
 		TConstArrayView<FCameraObjectInterfaceParameterDefinition> ParameterDefinitions,
-		const FInstancedPropertyBag& ParameterOverrides,
+		const FInstancedOverridablePropertyBag& ParameterOverrides,
 		const FInstancedPropertyBag& CachedParameterOverrides)
 {
 	using namespace Internal;
@@ -467,6 +466,11 @@ void FCameraObjectInterfaceParameterOverrideHelper::ApplyParameterOverrides(
 		{
 			continue;
 		}
+		
+		if (!ParameterOverrides.IsPropertyOverriden(Definition.ParameterGuid))
+		{
+			continue;
+		}
 
 		const FPropertyBagPropertyDesc* PropertyDesc = ParameterOverridesStruct->FindPropertyDescByID(Definition.ParameterGuid);
 		if (!ensure(PropertyDesc))
@@ -486,7 +490,7 @@ void FCameraObjectInterfaceParameterOverrideHelper::ApplyParameterOverrides(
 void FCameraObjectInterfaceParameterOverrideHelper::ApplyParameterOverride(
 		const UObject* CameraObject,
 		const FCameraObjectInterfaceParameterDefinition& ParameterDefinition,
-		const FInstancedPropertyBag& PropertyBag,
+		const FInstancedOverridablePropertyBag& PropertyBag,
 		const FPropertyBagPropertyDesc& PropertyBagPropertyDesc,
 		bool bDrivenOnly)
 {
diff --git a/Engine/Plugins/Cameras/GameplayCameras/Source/GameplayCameras/Private/Helpers/CameraObjectInterfaceParameterOverrideHelper.h b/Engine/Plugins/Cameras/GameplayCameras/Source/GameplayCameras/Private/Helpers/CameraObjectInterfaceParameterOverrideHelper.h
--- a/Engine/Plugins/Cameras/GameplayCameras/Source/GameplayCameras/Private/Helpers/CameraObjectInterfaceParameterOverrideHelper.h
+++ b/Engine/Plugins/Cameras/GameplayCameras/Source/GameplayCameras/Private/Helpers/CameraObjectInterfaceParameterOverrideHelper.h
@@ -2,11 +2,16 @@
 
 #pragma once
 
-#include "StructUtils/PropertyBag.h"
+#include "Containers/ArrayView.h"
+#include "Misc/Guid.h"
 
 class UBaseCameraObject;
 class UCameraAsset;
+class UObject;
 struct FCameraObjectInterfaceParameterDefinition;
+struct FInstancedOverridablePropertyBag;
+struct FInstancedPropertyBag;
+struct FPropertyBagPropertyDesc;
 
 namespace UE::Cameras
 {
@@ -47,15 +52,14 @@ public:
 	void ApplyParameterOverrides(
 			const UObject* CameraObject,
 			TConstArrayView<FCameraObjectInterfaceParameterDefinition> ParameterDefinitions,
-			const FInstancedPropertyBag& ParameterOverrides,
-			const TSet<FGuid>& OverridenParameterGuids,
+			const FInstancedOverridablePropertyBag& ParameterOverrides,
 			bool bDrivenOnly);
 
 	/** Sets overriden values of interface parameters in the given variable and context data tables. */
 	void ApplyParameterOverrides(
 			const UObject* CameraObject,
 			TConstArrayView<FCameraObjectInterfaceParameterDefinition> ParameterDefinitions,
-			const FInstancedPropertyBag& ParameterOverrides,
+			const FInstancedOverridablePropertyBag& ParameterOverrides,
 			const FInstancedPropertyBag& CachedParameterOverrides);
 
 private:
@@ -63,7 +67,7 @@ private:
 	void ApplyParameterOverride(
 			const UObject* CameraObject,
 			const FCameraObjectInterfaceParameterDefinition& ParameterDefinition,
-			const FInstancedPropertyBag& PropertyBag,
+			const FInstancedOverridablePropertyBag& PropertyBag,
 			const FPropertyBagPropertyDesc& PropertyBagPropertyDesc,
 			bool bDrivenOnly);
 
diff --git a/Engine/Plugins/Cameras/GameplayCameras/Source/GameplayCameras/Private/StructUtils/OverridablePropertyBag.cpp b/Engine/Plugins/Cameras/GameplayCameras/Source/GameplayCameras/Private/StructUtils/OverridablePropertyBag.cpp
new file mode 100644
--- /dev/null
+++ b/Engine/Plugins/Cameras/GameplayCameras/Source/GameplayCameras/Private/StructUtils/OverridablePropertyBag.cpp
@@ -0,0 +1,53 @@
+// Copyright Epic Games, Inc. All Rights Reserved.
+
+#include "StructUtils/OverridablePropertyBag.h"
+
+#include UE_INLINE_GENERATED_CPP_BY_NAME(OverridablePropertyBag)
+
+bool FInstancedOverridablePropertyBag::IsPropertyOverriden(const FGuid& InPropertyID) const
+{
+	return OverridenPropertyIDs.Contains(InPropertyID);
+}
+
+void FInstancedOverridablePropertyBag::SetPropertyOverriden(const FGuid& InPropertyID, bool bIsOverriden)
+{
+	if (bIsOverriden)
+	{
+		OverridenPropertyIDs.AddUnique(InPropertyID);
+	}
+	else
+	{
+		OverridenPropertyIDs.Remove(InPropertyID);
+	}
+}
+
+void FInstancedOverridablePropertyBag::MigrateToNewBagInstanceWithOverrides(const FInstancedPropertyBag& NewBagInstance)
+{
+	FInstancedPropertyBag::MigrateToNewBagInstanceWithOverrides(NewBagInstance, OverridenPropertyIDs);
+
+	// Remove overrides for propeties that don't exist anymore.
+	if (const UPropertyBag* ParametersType = GetPropertyBagStruct())
+	{
+		for (TArray<FGuid>::TIterator It = OverridenPropertyIDs.CreateIterator(); It; ++It)
+		{
+			if (!ParametersType->FindPropertyDescByID(*It))
+			{
+				It.RemoveCurrentSwap();
+			}
+		}
+	}
+}
+
+bool FInstancedOverridablePropertyBag::SerializeFromMismatchedTag(FPropertyTag const& Tag, FStructuredArchive::FSlot Slot)
+{
+	static const FName NAME_InstancedPropertyBag = FInstancedPropertyBag::StaticStruct()->GetFName();
+	if (Tag.GetType().IsStruct(NAME_InstancedPropertyBag))
+	{
+		// Read the structured data as an FInstancedPropertyBag. Our list of overriden property IDs
+		// will stay empty.
+		Serialize(Slot.GetUnderlyingArchive());
+		return true;
+	}
+	return false;
+}
+
diff --git a/Engine/Plugins/Cameras/GameplayCameras/Source/GameplayCameras/Public/Build/CameraObjectInterfaceParameterBuilder.h b/Engine/Plugins/Cameras/GameplayCameras/Source/GameplayCameras/Public/Build/CameraObjectInterfaceParameterBuilder.h
--- a/Engine/Plugins/Cameras/GameplayCameras/Source/GameplayCameras/Public/Build/CameraObjectInterfaceParameterBuilder.h
+++ b/Engine/Plugins/Cameras/GameplayCameras/Source/GameplayCameras/Public/Build/CameraObjectInterfaceParameterBuilder.h
@@ -3,13 +3,14 @@
 #pragma once
 
 #include "Containers/Array.h"
-#include "StructUtils/PropertyBag.h"
 
 #define UE_API GAMEPLAYCAMERAS_API
 
 class UBaseCameraObject;
 class UCameraObjectInterfaceDataParameter;
 struct FCameraObjectInterfaceParameterDefinition;
+struct FInstancedPropertyBag;
+struct FPropertyBagPropertyDesc;
 
 namespace UE::Cameras
 {
diff --git a/Engine/Plugins/Cameras/GameplayCameras/Source/GameplayCameras/Public/Core/BaseCameraObjectReference.h b/Engine/Plugins/Cameras/GameplayCameras/Source/GameplayCameras/Public/Core/BaseCameraObjectReference.h
--- a/Engine/Plugins/Cameras/GameplayCameras/Source/GameplayCameras/Public/Core/BaseCameraObjectReference.h
+++ b/Engine/Plugins/Cameras/GameplayCameras/Source/GameplayCameras/Public/Core/BaseCameraObjectReference.h
@@ -4,7 +4,7 @@
 
 #include "Core/CameraContextDataTableFwd.h"
 #include "Core/CameraVariableTableFwd.h"
-#include "StructUtils/PropertyBag.h"
+#include "StructUtils/OverridablePropertyBag.h"
 
 #include "BaseCameraObjectReference.generated.h"
 
@@ -36,9 +36,11 @@ struct FCameraObjectInterfaceParameterMetaData
 	UPROPERTY()
 	FCameraContextDataID OverrideDataID;
 
-	/** Whether this parameter has an override value. */
+	
+	// Deprecated
+
 	UPROPERTY()
-	bool bIsOverridden = false;
+	bool bIsOverridden_DEPRECATED = false;
 };
 
 USTRUCT(BlueprintType)
@@ -51,17 +53,23 @@ public:
 	virtual ~FBaseCameraObjectReference() {}
 
 	/** Gets the parameters for this camera rig, some of which containing overrides. */
-	const FInstancedPropertyBag& GetParameters() const
+	const FInstancedOverridablePropertyBag& GetParameters() const
 	{
 		return Parameters;
 	}
 
 	/** Gets the parameters for this camera rig, some of which containing overrides. */
-	FInstancedPropertyBag& GetParameters()
+	FInstancedOverridablePropertyBag& GetParameters()
 	{
 		return Parameters;
 	}
 
+	/** Gets the IDs of the parameters with override values. */
+	TConstArrayView<FGuid> GetOverriddenParameterGuids() const
+	{
+		return Parameters.GetOverridenPropertyIDs();
+	}
+
 public:
 
 	/** Gets the camera object that this reference points to. */
@@ -78,11 +86,13 @@ public:
 
 	// Internal API.
 
-	GAMEPLAYCAMERAS_API bool IsParameterOverridden(const FGuid& PropertyID) const;
-	GAMEPLAYCAMERAS_API void SetParameterOverridden(const FGuid& PropertyID, bool bIsOverridden);
+	void PostSerialize(const FArchive& Ar);
 
 	GAMEPLAYCAMERAS_API void GetCustomCameraNodeParameters(FCustomCameraNodeParameterInfos& OutParameterInfos);
 	
+	GAMEPLAYCAMERAS_API bool IsParameterOverridden(const FGuid& PropertyID) const;
+	GAMEPLAYCAMERAS_API void SetParameterOverridden(const FGuid& PropertyID, bool bIsOverridden);
+
 	template<typename ContainerType>
 	void GetOverriddenParameterGuids(ContainerType& OutOverriddenIDs) const;
 
@@ -94,8 +104,8 @@ private:
 protected:
 
 	/** The camera rig's parameters. */
-	UPROPERTY(EditAnywhere, Category="", meta=(FixedLayout))
-	FInstancedPropertyBag Parameters;
+	UPROPERTY(EditAnywhere, Category="", meta=(FixedLayout, InterpBagProperties))
+	FInstancedOverridablePropertyBag Parameters;
 
 	/** Metadata for the parameters. */
 	UPROPERTY()
@@ -105,12 +115,15 @@ protected:
 template<typename ContainerType>
 void FBaseCameraObjectReference::GetOverriddenParameterGuids(ContainerType& OutOverriddenIDs) const
 {
-	for (const FCameraObjectInterfaceParameterMetaData& MetaData : ParameterMetaData)
-	{
-		if (MetaData.bIsOverridden)
-		{
-			OutOverriddenIDs.Add(MetaData.ParameterGuid);
-		}
-	}
+	Parameters.GetOverridenPropertyIDs(OutOverriddenIDs);
 }
 
+template<>
+struct TStructOpsTypeTraits<FBaseCameraObjectReference> : public TStructOpsTypeTraitsBase2<FBaseCameraObjectReference>
+{
+	enum
+	{
+		WithPostSerialize = true
+	};
+};
+
diff --git a/Engine/Plugins/Cameras/GameplayCameras/Source/GameplayCameras/Public/Core/CameraAsset.h b/Engine/Plugins/Cameras/GameplayCameras/Source/GameplayCameras/Public/Core/CameraAsset.h
--- a/Engine/Plugins/Cameras/GameplayCameras/Source/GameplayCameras/Public/Core/CameraAsset.h
+++ b/Engine/Plugins/Cameras/GameplayCameras/Source/GameplayCameras/Public/Core/CameraAsset.h
@@ -191,10 +191,6 @@ private:
 	UPROPERTY()
 	TArray<FCameraObjectInterfaceParameterDefinition> ParameterDefinitions;
 
-	/** Owner rigs for each parameter definition. */
-	UPROPERTY()
-	TArray<TObjectPtr<const UCameraRigAsset>> ParameterOwners;
-
 	/** Allocation info for the camera asset. */
 	UPROPERTY()
 	FCameraAssetAllocationInfo AllocationInfo;
diff --git a/Engine/Plugins/Cameras/GameplayCameras/Source/GameplayCameras/Public/Core/CameraAssetReference.h b/Engine/Plugins/Cameras/GameplayCameras/Source/GameplayCameras/Public/Core/CameraAssetReference.h
--- a/Engine/Plugins/Cameras/GameplayCameras/Source/GameplayCameras/Public/Core/CameraAssetReference.h
+++ b/Engine/Plugins/Cameras/GameplayCameras/Source/GameplayCameras/Public/Core/CameraAssetReference.h
@@ -2,7 +2,7 @@
 
 #pragma once
 
-#include "StructUtils/PropertyBag.h"
+#include "StructUtils/OverridablePropertyBag.h"
 #include "UObject/ObjectPtr.h"
 
 #include "CameraAssetReference.generated.h"
@@ -59,13 +59,13 @@ public:
 	}
 
 	/** Gets the parameters for this camera, some of which containing overrides. */
-	const FInstancedPropertyBag& GetParameters() const
+	const FInstancedOverridablePropertyBag& GetParameters() const
 	{
 		return Parameters;
 	}
 
 	/** Gets the parameters for this camera, some of which containing overrides. */
-	FInstancedPropertyBag& GetParameters()
+	FInstancedOverridablePropertyBag& GetParameters()
 	{
 		return Parameters;
 	}
@@ -73,28 +73,36 @@ public:
 	/** Gets the IDs of the parameters with override values. */
 	TConstArrayView<FGuid> GetOverriddenParameterGuids() const
 	{
-		return ParameterOverrideGuids;
+		return Parameters.GetOverridenPropertyIDs();
 	}
 
+public:
+
 	/** Applies the parameter override values to the given evaluation result. */
-	void ApplyParameterOverrides(UE::Cameras::FCameraNodeEvaluationResult& OutResult, bool bDrivenOnly) const;
+	GAMEPLAYCAMERAS_API void ApplyParameterOverrides(UE::Cameras::FCameraNodeEvaluationResult& OutResult, bool bDrivenOnly) const;
 
 	/** Applies the parameter override values to the given evaluation result. */
-	void ApplyParameterOverrides(const FInstancedPropertyBag& CachedParameters, UE::Cameras::FCameraNodeEvaluationResult& OutResult) const;
+	GAMEPLAYCAMERAS_API void ApplyParameterOverrides(const FInstancedPropertyBag& CachedParameters, UE::Cameras::FCameraNodeEvaluationResult& OutResult) const;
+
+public:
+
+	/** Returns whether the override parameters structure needs to be rebuilt. */
+	GAMEPLAYCAMERAS_API bool NeedsRebuildParameters() const;
+	/** Rebuilds the override parameters structure, if needed. */
+	GAMEPLAYCAMERAS_API bool RebuildParametersIfNeeded();
+	/** Rebuilds the override parameters structure. */
+	GAMEPLAYCAMERAS_API void RebuildParameters();
 
 public:
 
 	// Internal API.
 
 	bool SerializeFromMismatchedTag(FPropertyTag const& Tag, FStructuredArchive::FSlot Slot);
+	void PostSerialize(const FArchive& Ar);
 
 	GAMEPLAYCAMERAS_API bool IsParameterOverridden(const FGuid PropertyID) const;
 	GAMEPLAYCAMERAS_API void SetParameterOverridden(const FGuid PropertyID, bool bIsOverridden);
 
-	GAMEPLAYCAMERAS_API bool NeedsRebuildParameters() const;
-	GAMEPLAYCAMERAS_API bool RebuildParametersIfNeeded();
-	GAMEPLAYCAMERAS_API void RebuildParameters();
-
 private:
 
 	/** The referenced camera asset. */
@@ -102,12 +110,14 @@ private:
 	TObjectPtr<UCameraAsset> CameraAsset;
 
 	/** The camera asset's parameters. */
-	UPROPERTY(EditAnywhere, Category="", meta=(FixedLayout=true, InterpBagProperties=true, DisplayName="Camera Parameters"))
-	FInstancedPropertyBag Parameters;
+	UPROPERTY(EditAnywhere, Category="", meta=(FixedLayout, InterpBagProperties))
+	FInstancedOverridablePropertyBag Parameters;
+
+
+	// Deprecated
 
-	/** The list of camera rig parameters with override values. */
-	UPROPERTY(EditAnywhere, Category="")
-	TArray<FGuid> ParameterOverrideGuids;
+	UPROPERTY()
+	TArray<FGuid> ParameterOverrideGuids_DEPRECATED;
 
 	friend class UE::Cameras::FCameraAssetReferenceDetailsCustomization;
 };
@@ -117,7 +127,8 @@ struct TStructOpsTypeTraits<FCameraAssetReference> : public TStructOpsTypeTraits
 {
 	enum
 	{
-		WithStructuredSerializeFromMismatchedTag = true
+		WithStructuredSerializeFromMismatchedTag = true,
+		WithPostSerialize = true
 	};
 };
 
diff --git a/Engine/Plugins/Cameras/GameplayCameras/Source/GameplayCameras/Public/Core/CameraObjectInterfaceParameterDefinition.h b/Engine/Plugins/Cameras/GameplayCameras/Source/GameplayCameras/Public/Core/CameraObjectInterfaceParameterDefinition.h
--- a/Engine/Plugins/Cameras/GameplayCameras/Source/GameplayCameras/Public/Core/CameraObjectInterfaceParameterDefinition.h
+++ b/Engine/Plugins/Cameras/GameplayCameras/Source/GameplayCameras/Public/Core/CameraObjectInterfaceParameterDefinition.h
@@ -4,14 +4,9 @@
 
 #include "Core/CameraContextDataTableFwd.h"
 #include "Core/CameraVariableTableFwd.h"
-#include "Misc/TVariant.h"
 
 #include "CameraObjectInterfaceParameterDefinition.generated.h"
 
-class UBaseCameraObject;
-struct FInstancedPropertyBag;
-struct FPropertyBagPropertyDesc;
-
 /**
  * The type of a camera rig parameter.
  */
diff --git a/Engine/Plugins/Cameras/GameplayCameras/Source/GameplayCameras/Public/StructUtils/OverridablePropertyBag.h b/Engine/Plugins/Cameras/GameplayCameras/Source/GameplayCameras/Public/StructUtils/OverridablePropertyBag.h
new file mode 100644
--- /dev/null
+++ b/Engine/Plugins/Cameras/GameplayCameras/Source/GameplayCameras/Public/StructUtils/OverridablePropertyBag.h
@@ -0,0 +1,83 @@
+// Copyright Epic Games, Inc. All Rights Reserved.
+
+#pragma once
+
+#include "Containers/Array.h"
+#include "Misc/Guid.h"
+#include "StructUtils/PropertyBag.h"
+
+#include "OverridablePropertyBag.generated.h"
+
+/**
+* Extended version of FInstancedPropertyBag to support "overriden properties" with an extra array of
+* overriden parameter IDs.
+*/
+USTRUCT()
+struct FInstancedOverridablePropertyBag : public FInstancedPropertyBag
+{
+	GENERATED_BODY()
+
+public:
+
+	/**
+	 * Gets the IDs of overriden properties, i.e. properties that have override values.
+	 */
+	template<typename ContainerType>
+	void GetOverridenPropertyIDs(ContainerType& OutPropertyIDs) const;
+
+	/**
+	 * Gets the IDs of overriden properties, i.e. properties that have override values.
+	 */
+	TConstArrayView<FGuid> GetOverridenPropertyIDs() const { return OverridenPropertyIDs; }
+
+	/**
+	 * Returns whether the given property is overriden.
+	 */
+	GAMEPLAYCAMERAS_API bool IsPropertyOverriden(const FGuid& InPropertyID) const;
+
+	/**
+	 * Sets whether the given property is overriden by setting the matching override property's value.
+	 */
+	GAMEPLAYCAMERAS_API void SetPropertyOverriden(const FGuid& InPropertyID, bool bIsOverriden);
+
+public:
+
+	/**
+	 * Migrate this property bag to the layout of the given one. Overriden values in this property bag
+	 * will be preserved. Other values will adopt the values of the given new bag.
+	 */
+	GAMEPLAYCAMERAS_API void MigrateToNewBagInstanceWithOverrides(const FInstancedPropertyBag& NewBagInstance);
+
+public:
+
+	// Internal API.
+
+	bool SerializeFromMismatchedTag(FPropertyTag const& Tag, FStructuredArchive::FSlot Slot);
+
+private:
+
+	// List of properties that have override values.
+	UPROPERTY()
+	TArray<FGuid> OverridenPropertyIDs;
+};
+
+template<typename ContainerType>
+void FInstancedOverridablePropertyBag::GetOverridenPropertyIDs(ContainerType& OutPropertyIDs) const
+{
+	OutPropertyIDs.Reserve(OutPropertyIDs.Num() + OverridenPropertyIDs.Num());
+
+	for (const FGuid& PropertyID : OverridenPropertyIDs)
+	{
+		OutPropertyIDs.Add(PropertyID);
+	}
+}
+
+template<>
+struct TStructOpsTypeTraits<FInstancedOverridablePropertyBag> : public TStructOpsTypeTraitsBase2<FInstancedOverridablePropertyBag>
+{
+	enum
+	{
+		WithStructuredSerializeFromMismatchedTag = true
+	};
+};
+
diff --git a/Engine/Plugins/Cameras/GameplayCameras/Source/GameplayCamerasEditor/Private/GameplayCamerasEditorModule.cpp b/Engine/Plugins/Cameras/GameplayCameras/Source/GameplayCamerasEditor/Private/GameplayCamerasEditorModule.cpp
--- a/Engine/Plugins/Cameras/GameplayCameras/Source/GameplayCamerasEditor/Private/GameplayCamerasEditorModule.cpp
+++ b/Engine/Plugins/Cameras/GameplayCameras/Source/GameplayCamerasEditor/Private/GameplayCamerasEditorModule.cpp
@@ -50,6 +50,7 @@
 #include "Kismet2/KismetEditorUtilities.h"
 #include "Modules/ModuleManager.h"
 #include "ObjectTools.h"
+#include "PropertyBagDetails.h"
 #include "PropertyEditorModule.h"
 #include "Sequencer/CameraFramingZoneTrackEditor.h"
 #include "Sequencer/GameplayCameraComponentTrackEditor.h"
@@ -468,6 +469,11 @@ private:
 		FCameraVariableReferenceDetailsCustomization::Register(PropertyEditorModule);
 		FRichCurveDetailsCustomization::Register(PropertyEditorModule);
 
+		PropertyEditorModule.RegisterCustomPropertyTypeLayout(
+				"InstancedOverridablePropertyBag",
+				FOnGetPropertyTypeCustomizationInstance::CreateStatic(
+					&FPropertyBagDetails::MakeInstance));
+
 		PropertyEditorModule.RegisterCustomPropertyTypeLayout(
 				"CameraAssetReference",
 				FOnGetPropertyTypeCustomizationInstance::CreateStatic(
@@ -499,7 +505,9 @@ private:
 			FCameraVariableReferenceDetailsCustomization::Unregister(*PropertyEditorModule);
 			FRichCurveDetailsCustomization::Unregister(*PropertyEditorModule);
 
-			PropertyEditorModule->UnregisterCustomPropertyTypeLayout("CameraRigAsset");
+			PropertyEditorModule->UnregisterCustomPropertyTypeLayout("InstancedOverridablePropertyBag");
+
+			PropertyEditorModule->UnregisterCustomPropertyTypeLayout("CameraAssetReference");
 			PropertyEditorModule->UnregisterCustomPropertyTypeLayout("CameraRigAssetReference");
 			PropertyEditorModule->UnregisterCustomPropertyTypeLayout("CameraShakeAssetReference");
 
diff --git a/Engine/Plugins/Cameras/GameplayCameras/Source/GameplayCamerasEditor/Private/Sequencer/GameplayCameraComponentTrackEditor.cpp b/Engine/Plugins/Cameras/GameplayCameras/Source/GameplayCamerasEditor/Private/Sequencer/GameplayCameraComponentTrackEditor.cpp
--- a/Engine/Plugins/Cameras/GameplayCameras/Source/GameplayCamerasEditor/Private/Sequencer/GameplayCameraComponentTrackEditor.cpp
+++ b/Engine/Plugins/Cameras/GameplayCameras/Source/GameplayCamerasEditor/Private/Sequencer/GameplayCameraComponentTrackEditor.cpp
@@ -5,8 +5,7 @@
 #if UE_VERSION_NEWER_THAN_OR_EQUAL(5,6,0)
 
 #include "Core/CameraAsset.h"
-#include "Core/CameraObjectInterfaceParameterDefinition.h"
-#include "Core/CameraParameters.h"
+#include "Core/CameraParameters.h"  // IWYU pragma: keep
 #include "EventHandlers/MovieSceneDataEventContainer.h"
 #include "Framework/MultiBox/MultiBoxBuilder.h"
 #include "GameFramework/Actor.h"
@@ -22,7 +21,7 @@
 
 #if UE_VERSION_NEWER_THAN_OR_EQUAL(5,7,0)
 #include "Misc/SequencerObjectBindingHelper.h"
-#endif
+#endif  // >=5.7.0
 
 #define LOCTEXT_NAMESPACE "GameplayCameraComponentTrackEditor"
 
@@ -118,15 +117,19 @@ void GetKeyablePropertyPaths(TSharedPtr<ISequencer> Sequencer, const UGameplayCa
 		return;
 	}
 
+#if UE_VERSION_OLDER_THAN(5,7,0)
+	// Before 5.7.0, Sequencer can animate instanced structs directly but not property bags, so we need to add the "Value" field to
+	// dive inside the property bag's instanced struct.
 	const UStruct* PropertyBagStruct = FInstancedPropertyBag::StaticStruct();
 	PropertyPath.AddProperty(FPropertyInfo(PropertyBagStruct->FindPropertyByName(TEXT("Value"))));
+#endif  // pre-5.7.0
 
 	const UClass* ComponentClass = CameraComponentBase->GetClass();
 #if UE_VERSION_NEWER_THAN_OR_EQUAL(5,7,0)
 	FSequencerObjectBindingHelper::GetKeyablePropertyPaths(ComponentClass, CameraParametersMemory, CameraParametersStruct, PropertyPath, Sequencer.ToSharedRef(), KeyablePropertyPaths);
 #else
 	GetKeyablePropertyPathsImpl(Sequencer, ComponentClass, CameraParametersStruct, CameraParametersMemory, PropertyPath, KeyablePropertyPaths);
-#endif
+#endif  // 5.7 and above
 }
 
 class FCameraParameterTrackSetupHandler : public UE::MovieScene::TIntrusiveEventHandler<UE::MovieScene::ISequenceDataEventHandler>
@@ -219,8 +222,13 @@ void FGameplayCameraComponentTrackEditor::OnExtendObjectBindingTrackMenu(FMenuBu
 
 		MenuBuilder.BeginSection(TEXT("CameraParameters"), LOCTEXT("AddCameraParametersMenuSection", "Camera Parameters"));
 
+#if UE_VERSION_NEWER_THAN_OR_EQUAL(5,7,0)
+		// Start at level 2 to skip adding menus for CameraReference > Parameters
+		BuildAddParameterTrackMenuItems(ObjectBindings[0], MenuBuilder, KeyablePropertyPaths, 2);
+#else
 		// Start at level 3 to skip adding menus for CameraReference > Parameters > Value
 		BuildAddParameterTrackMenuItems(ObjectBindings[0], MenuBuilder, KeyablePropertyPaths, 3);
+#endif  // >= 5.7.0
 
 		MenuBuilder.EndSection();
 	}
