Cameras: fix correct base class for struct type ops

diff --git a/Engine/Plugins/Cameras/GameplayCameras/Source/GameplayCameras/Public/Core/CameraRigAssetReference.h b/Engine/Plugins/Cameras/GameplayCameras/Source/GameplayCameras/Public/Core/CameraRigAssetReference.h
--- a/Engine/Plugins/Cameras/GameplayCameras/Source/GameplayCameras/Public/Core/CameraRigAssetReference.h
+++ b/Engine/Plugins/Cameras/GameplayCameras/Source/GameplayCameras/Public/Core/CameraRigAssetReference.h
@@ -325,7 +325,7 @@ private:
 };
 
 template<>
-struct TStructOpsTypeTraits<FCameraRigAssetReference> : public TStructOpsTypeTraitsBase2<FCameraRigAssetReference>
+struct TStructOpsTypeTraits<FCameraRigAssetReference> : public TStructOpsTypeTraitsBase2<FBaseCameraObjectReference>
 {
 	enum
 	{
