Cameras: use the new sequencer notify feature to only update parameter overrides when they're animated

This fixes not having to check for changed values every frame.

diff --git a/Engine/Plugins/Cameras/GameplayCameras/Source/GameplayCameras/Private/GameFramework/GameplayCameraComponent.cpp b/Engine/Plugins/Cameras/GameplayCameras/Source/GameplayCameras/Private/GameFramework/GameplayCameraComponent.cpp
--- a/Engine/Plugins/Cameras/GameplayCameras/Source/GameplayCameras/Private/GameFramework/GameplayCameraComponent.cpp
+++ b/Engine/Plugins/Cameras/GameplayCameras/Source/GameplayCameras/Private/GameFramework/GameplayCameraComponent.cpp
@@ -4,6 +4,7 @@
 
 #include "Core/CameraAsset.h"
 #include "GameplayCamerasDelegates.h"
+#include "Misc/EngineVersionComparison.h"
 
 #include UE_INLINE_GENERATED_CPP_BY_NAME(GameplayCameraComponent)
 
@@ -66,12 +67,32 @@ void UGameplayCameraComponent::OnUpdateCameraEvaluationContext(bool bForceApplyP
 	{
 		CameraReference.ApplyParameterOverrides(InitialResult, false);
 	}
+#if UE_VERSION_OLDER_THAN(5,7,0)
+	// Before 5.7.0, we don't have a notify callback from Sequencer to know that it's animating parameters,
+	// so we need to always re-apply values.
 	else
 	{
 		CameraReference.ApplyParameterOverrides(CachedParameterOverrides, InitialResult);
 	}
 
 	CachedParameterOverrides = CameraReference.GetParameters();
+#endif  // pre-5.7.0
+}
+
+void UGameplayCameraComponent::NotifyChangeCameraReference()
+{
+	using namespace UE::Cameras;
+
+	if (HasCameraEvaluationContext())
+	{
+		// Sequencer animated some of our parameters... look for those whose value changed, compared to our
+		// cached parameter bag, and re-apply them to the evaluation context.
+		// TODO: This isn't a very efficient process, but it's unclear how to reconcile the reference's parameters struct
+		//		 with the evaluation context's variable/context-data tables.
+		FCameraNodeEvaluationResult& InitialResult = GetEvaluationContext()->GetInitialResult();
+		CameraReference.ApplyParameterOverrides(CachedParameterOverrides, InitialResult);
+		CachedParameterOverrides = CameraReference.GetParameters();
+	}
 }
 
 #if WITH_EDITOR
diff --git a/Engine/Plugins/Cameras/GameplayCameras/Source/GameplayCameras/Private/GameFramework/GameplayCameraRigComponent.cpp b/Engine/Plugins/Cameras/GameplayCameras/Source/GameplayCameras/Private/GameFramework/GameplayCameraRigComponent.cpp
--- a/Engine/Plugins/Cameras/GameplayCameras/Source/GameplayCameras/Private/GameFramework/GameplayCameraRigComponent.cpp
+++ b/Engine/Plugins/Cameras/GameplayCameras/Source/GameplayCameras/Private/GameFramework/GameplayCameraRigComponent.cpp
@@ -5,8 +5,9 @@
 #include "Core/CameraAsset.h"
 #include "Core/CameraRigAsset.h"
 #include "Directors/SingleCameraDirector.h"
-#include "GameFramework/Actor.h"
+#include "GameFramework/Actor.h"  // IWYU pragma: keep
 #include "GameplayCamerasDelegates.h"
+#include "Misc/EngineVersionComparison.h"
 
 #include UE_INLINE_GENERATED_CPP_BY_NAME(GameplayCameraRigComponent)
 
@@ -96,12 +97,32 @@ void UGameplayCameraRigComponent::OnUpdateCameraEvaluationContext(bool bForceApp
 	{
 		CameraRigReference.ApplyParameterOverrides(InitialResult, false);
 	}
+#if UE_VERSION_OLDER_THAN(5,7,0)
+	// Before 5.7.0, we don't have a notify callback from Sequencer to know that it's animating parameters,
+	// so we need to always re-apply values.
 	else
 	{
 		CameraRigReference.ApplyParameterOverrides(CachedParameterOverrides, InitialResult);
 	}
 
 	CachedParameterOverrides = CameraRigReference.GetParameters();
+#endif  // pre-5.7.0
+}
+
+void UGameplayCameraRigComponent::NotifyChangeCameraRigReference()
+{
+	using namespace UE::Cameras;
+
+	if (HasCameraEvaluationContext())
+	{
+		// Sequencer animated some of our parameters... look for those whose value changed, compared to our
+		// cached parameter bag, and re-apply them to the evaluation context.
+		// TODO: This isn't a very efficient process, but it's unclear how to reconcile the reference's parameters struct
+		//		 with the evaluation context's variable/context-data tables.
+		FCameraNodeEvaluationResult& InitialResult = GetEvaluationContext()->GetInitialResult();
+		CameraRigReference.ApplyParameterOverrides(CachedParameterOverrides, InitialResult);
+		CachedParameterOverrides = CameraRigReference.GetParameters();
+	}
 }
 
 #if WITH_EDITOR
diff --git a/Engine/Plugins/Cameras/GameplayCameras/Source/GameplayCameras/Public/GameFramework/GameplayCameraComponent.h b/Engine/Plugins/Cameras/GameplayCameras/Source/GameplayCameras/Public/GameFramework/GameplayCameraComponent.h
--- a/Engine/Plugins/Cameras/GameplayCameras/Source/GameplayCameras/Public/GameFramework/GameplayCameraComponent.h
+++ b/Engine/Plugins/Cameras/GameplayCameras/Source/GameplayCameras/Public/GameFramework/GameplayCameraComponent.h
@@ -37,11 +37,14 @@ public:
 public:
 
 	/** The camera asset to run. */
-	UPROPERTY(BlueprintReadWrite, EditAnywhere, Category=Camera, meta=(SequencerHideProperty=true))
+	UPROPERTY(BlueprintReadWrite, EditAnywhere, Category=Camera, meta=(InterpNotify="NotifyChangeCameraReference", SequencerHideProperty=true))
 	FCameraAssetReference CameraReference;
 
 protected:
 
+	UFUNCTION()
+	void NotifyChangeCameraReference();
+
 	// UGameplayCameraComponentBase interface.
 	virtual UCameraAsset* GetCameraAsset() override;
 	virtual void OnUpdateCameraEvaluationContext(bool bForceApplyParameterOverrides) override;
diff --git a/Engine/Plugins/Cameras/GameplayCameras/Source/GameplayCameras/Public/GameFramework/GameplayCameraRigComponent.h b/Engine/Plugins/Cameras/GameplayCameras/Source/GameplayCameras/Public/GameFramework/GameplayCameraRigComponent.h
--- a/Engine/Plugins/Cameras/GameplayCameras/Source/GameplayCameras/Public/GameFramework/GameplayCameraRigComponent.h
+++ b/Engine/Plugins/Cameras/GameplayCameras/Source/GameplayCameras/Public/GameFramework/GameplayCameraRigComponent.h
@@ -36,11 +36,14 @@ public:
 public:
 
 	/** The camera rig asset to run. */
-	UPROPERTY(BlueprintReadWrite, EditAnywhere, Category=Camera, meta=(SequencerHideProperty=true))
+	UPROPERTY(BlueprintReadWrite, EditAnywhere, Category=Camera, meta=(InterpNotify="NotifyChangeCameraRigReference", SequencerHideProperty=true))
 	FCameraRigAssetReference CameraRigReference;
 
 protected:
 
+	UFUNCTION()
+	void NotifyChangeCameraRigReference();
+
 	// UGameplayCameraComponentBase interface.
 	virtual UCameraAsset* GetCameraAsset() override;
 	virtual void OnUpdateCameraEvaluationContext(bool bForceApplyParameterOverrides) override;
