Cameras: remove unused custom build step for asset builders, make camera rig component build its generated camera asset without rebuilding the rig when it's not necessary

diff --git a/Engine/Plugins/Cameras/GameplayCameras/Source/GameplayCameras/Private/Build/CameraAssetBuilder.cpp b/Engine/Plugins/Cameras/GameplayCameras/Source/GameplayCameras/Private/Build/CameraAssetBuilder.cpp
--- a/Engine/Plugins/Cameras/GameplayCameras/Source/GameplayCameras/Private/Build/CameraAssetBuilder.cpp
+++ b/Engine/Plugins/Cameras/GameplayCameras/Source/GameplayCameras/Private/Build/CameraAssetBuilder.cpp
@@ -22,12 +22,7 @@ FCameraAssetBuilder::FCameraAssetBuilder(FCameraBuildLog& InBuildLog)
 {
 }
 
-void FCameraAssetBuilder::BuildCamera(UCameraAsset* InCameraAsset)
-{
-	BuildCamera(InCameraAsset, FCustomBuildStep::CreateLambda([](UCameraAsset*, FCameraBuildLog&) {}));
-}
-	
-void FCameraAssetBuilder::BuildCamera(UCameraAsset* InCameraAsset, FCustomBuildStep InCustomBuildStep)
+void FCameraAssetBuilder::BuildCamera(UCameraAsset* InCameraAsset, bool bBuildReferencedAssets)
 {
 	if (!ensure(InCameraAsset))
 	{
@@ -37,9 +32,7 @@ void FCameraAssetBuilder::BuildCamera(UCameraAsset* InCameraAsset, FCustomBuildS
 	CameraAsset = InCameraAsset;
 	BuildLog.SetLoggingPrefix(InCameraAsset->GetPathName() + TEXT(": "));
 	{
-		BuildCameraImpl();
-
-		InCustomBuildStep.ExecuteIfBound(CameraAsset, BuildLog);
+		BuildCameraImpl(bBuildReferencedAssets);
 	}
 	BuildLog.SetLoggingPrefix(FString());
 	UpdateBuildStatus();
@@ -47,7 +40,7 @@ void FCameraAssetBuilder::BuildCamera(UCameraAsset* InCameraAsset, FCustomBuildS
 	FGameplayCamerasDelegates::OnCameraAssetBuilt().Broadcast(CameraAsset);
 }
 
-void FCameraAssetBuilder::BuildCameraImpl()
+void FCameraAssetBuilder::BuildCameraImpl(bool bBuildReferencedAssets)
 {
 	TSet<UCameraRigAsset*> AllCameraRigs;
 	TSet<UCameraRigProxyAsset*> AllCameraRigProxies;
@@ -86,11 +79,14 @@ void FCameraAssetBuilder::BuildCameraImpl()
 		}
 	}
 
-	// Build each of the camera rigs.
-	for (UCameraRigAsset* CameraRig : AllCameraRigs)
+	if (bBuildReferencedAssets)
 	{
-		FCameraRigAssetBuilder CameraRigBuilder(BuildLog);
-		CameraRigBuilder.BuildCameraRig(CameraRig);
+		// Build each of the camera rigs.
+		for (UCameraRigAsset* CameraRig : AllCameraRigs)
+		{
+			FCameraRigAssetBuilder CameraRigBuilder(BuildLog);
+			CameraRigBuilder.BuildCameraRig(CameraRig);
+		}
 	}
 
 	// Get the list of all the camera rigs' interface parameters.
diff --git a/Engine/Plugins/Cameras/GameplayCameras/Source/GameplayCameras/Private/Build/CameraRigAssetBuilder.cpp b/Engine/Plugins/Cameras/GameplayCameras/Source/GameplayCameras/Private/Build/CameraRigAssetBuilder.cpp
--- a/Engine/Plugins/Cameras/GameplayCameras/Source/GameplayCameras/Private/Build/CameraRigAssetBuilder.cpp
+++ b/Engine/Plugins/Cameras/GameplayCameras/Source/GameplayCameras/Private/Build/CameraRigAssetBuilder.cpp
@@ -19,11 +19,6 @@ FCameraRigAssetBuilder::FCameraRigAssetBuilder(FCameraBuildLog& InBuildLog)
 }
 
 void FCameraRigAssetBuilder::BuildCameraRig(UCameraRigAsset* InCameraRig)
-{
-	BuildCameraRig(InCameraRig, FCustomBuildStep::CreateLambda([](UCameraRigAsset*, FCameraBuildLog&) {}));
-}
-
-void FCameraRigAssetBuilder::BuildCameraRig(UCameraRigAsset* InCameraRig, FCustomBuildStep InCustomBuildStep)
 {
 	if (!ensure(InCameraRig))
 	{
@@ -36,8 +31,6 @@ void FCameraRigAssetBuilder::BuildCameraRig(UCameraRigAsset* InCameraRig, FCusto
 	{
 		BuildCameraRigImpl();
 
-		InCustomBuildStep.ExecuteIfBound(CameraRig, BuildLog);
-
 		CameraRig->EventHandlers.Notify(&ICameraRigAssetEventHandler::OnCameraRigBuilt, CameraRig);
 	}
 	BuildLog.SetLoggingPrefix(FString());
diff --git a/Engine/Plugins/Cameras/GameplayCameras/Source/GameplayCameras/Private/GameFramework/GameplayCameraRigComponent.cpp b/Engine/Plugins/Cameras/GameplayCameras/Source/GameplayCameras/Private/GameFramework/GameplayCameraRigComponent.cpp
--- a/Engine/Plugins/Cameras/GameplayCameras/Source/GameplayCameras/Private/GameFramework/GameplayCameraRigComponent.cpp
+++ b/Engine/Plugins/Cameras/GameplayCameras/Source/GameplayCameras/Private/GameFramework/GameplayCameraRigComponent.cpp
@@ -2,6 +2,8 @@
 
 #include "GameFramework/GameplayCameraRigComponent.h"
 
+#include "Build/CameraAssetBuilder.h"
+#include "Build/CameraBuildLog.h"
 #include "Core/CameraAsset.h"
 #include "Core/CameraRigAsset.h"
 #include "Directors/SingleCameraDirector.h"
@@ -62,18 +64,37 @@ UCameraAsset* UGameplayCameraRigComponent::GetCameraAsset()
 		GeneratedCameraAsset->SetCameraDirector(SingleDirector);
 
 		ensure(!bIsBuildingGeneratedCameraAsset);
-		{
-			TGuardValue<bool> ReentrancyGuard(bIsBuildingGeneratedCameraAsset, true);
-			GeneratedCameraAsset->BuildCamera();
-		}
+		BuildGeneratedCamera();
 		OnCameraRigAssetBuiltImpl();
 	}
 
 	return GeneratedCameraAsset;
 }
 
+void UGameplayCameraRigComponent::BuildGeneratedCamera()
+{
+	using namespace UE::Cameras;
+
+	if (GeneratedCameraAsset)
+	{
+		TGuardValue<bool> ReentrancyGuard(bIsBuildingGeneratedCameraAsset, true);
+
+		// Build only the camera asset, not the referenced camera rigs... there's only one camera rig
+		// in this camera (ours) and it was just built, that's why we're in this callback in the
+		// first place.
+		const bool bBuildReferencedAssets = false;
+
+		FCameraBuildLog BuildLog;
+		BuildLog.SetForwardMessagesToLogging(true);
+		FCameraAssetBuilder Builder(BuildLog);
+		Builder.BuildCamera(GeneratedCameraAsset, bBuildReferencedAssets);
+	}
+}
+
 void UGameplayCameraRigComponent::OnCameraRigAssetBuiltImpl()
 {
+	using namespace UE::Cameras;
+
 	CameraRigReference.RebuildParametersIfNeeded();
 	if (HasCameraEvaluationContext())
 	{
@@ -138,11 +159,8 @@ void UGameplayCameraRigComponent::OnCameraRigAssetBuilt(const UCameraRigAsset* I
 
 	// If our camera rig asset was just built, it may have some new parameters. We need to rebuild
 	// our variable table and context data table, and re-apply overrides.
-	if (GeneratedCameraAsset)
-	{
-		TGuardValue<bool> ReentrancyGuard(bIsBuildingGeneratedCameraAsset, true);
-		GeneratedCameraAsset->BuildCamera();
-	}
+	BuildGeneratedCamera();
+	
 	OnCameraRigAssetBuiltImpl();
 }
 
diff --git a/Engine/Plugins/Cameras/GameplayCameras/Source/GameplayCameras/Public/Build/CameraAssetBuilder.h b/Engine/Plugins/Cameras/GameplayCameras/Source/GameplayCameras/Public/Build/CameraAssetBuilder.h
--- a/Engine/Plugins/Cameras/GameplayCameras/Source/GameplayCameras/Public/Build/CameraAssetBuilder.h
+++ b/Engine/Plugins/Cameras/GameplayCameras/Source/GameplayCameras/Public/Build/CameraAssetBuilder.h
@@ -5,7 +5,6 @@
 #include "Build/CameraBuildLog.h"
 #include "CoreTypes.h"
 #include "GameplayCameras.h"
-#include "Templates/Tuple.h"
 
 #define UE_API GAMEPLAYCAMERAS_API
 
@@ -21,20 +20,15 @@ class FCameraAssetBuilder
 {
 public:
 
-	DECLARE_DELEGATE_TwoParams(FCustomBuildStep, UCameraAsset*, FCameraBuildLog&);
-
 	/** Creates a new camera builder. */
 	UE_API FCameraAssetBuilder(FCameraBuildLog& InBuildLog);
 
 	/** Builds the given camera. */
-	UE_API void BuildCamera(UCameraAsset* InCameraAsset);
-
-	/** Builds the given camera. */
-	UE_API void BuildCamera(UCameraAsset* InCameraAsset, FCustomBuildStep InCustomBuildStep);
+	UE_API void BuildCamera(UCameraAsset* InCameraAsset, bool bBuildReferencedAssets = true);
 
 private:
 
-	void BuildCameraImpl();
+	void BuildCameraImpl(bool bBuildReferencedAssets);
 
 	void UpdateBuildStatus();
 
diff --git a/Engine/Plugins/Cameras/GameplayCameras/Source/GameplayCameras/Public/Build/CameraRigAssetBuilder.h b/Engine/Plugins/Cameras/GameplayCameras/Source/GameplayCameras/Public/Build/CameraRigAssetBuilder.h
--- a/Engine/Plugins/Cameras/GameplayCameras/Source/GameplayCameras/Public/Build/CameraRigAssetBuilder.h
+++ b/Engine/Plugins/Cameras/GameplayCameras/Source/GameplayCameras/Public/Build/CameraRigAssetBuilder.h
@@ -25,17 +25,12 @@ class FCameraRigAssetBuilder
 {
 public:
 
-	DECLARE_DELEGATE_TwoParams(FCustomBuildStep, UCameraRigAsset*, FCameraBuildLog&);
-
 	/** Creates a new camera rig builder. */
 	UE_API FCameraRigAssetBuilder(FCameraBuildLog& InBuildLog);
 
 	/** Builds the given camera rig. */
 	UE_API void BuildCameraRig(UCameraRigAsset* InCameraRig);
 
-	/** Builds the given camera rig. */
-	UE_API void BuildCameraRig(UCameraRigAsset* InCameraRig, FCustomBuildStep InCustomBuildStep);
-
 private:
 
 	void BuildCameraRigImpl();
diff --git a/Engine/Plugins/Cameras/GameplayCameras/Source/GameplayCameras/Public/GameFramework/GameplayCameraRigComponent.h b/Engine/Plugins/Cameras/GameplayCameras/Source/GameplayCameras/Public/GameFramework/GameplayCameraRigComponent.h
--- a/Engine/Plugins/Cameras/GameplayCameras/Source/GameplayCameras/Public/GameFramework/GameplayCameraRigComponent.h
+++ b/Engine/Plugins/Cameras/GameplayCameras/Source/GameplayCameras/Public/GameFramework/GameplayCameraRigComponent.h
@@ -50,6 +50,7 @@ protected:
 
 private:
 
+	void BuildGeneratedCamera();
 	void OnCameraRigAssetBuiltImpl();
 
 #if WITH_EDITOR
