Cameras: backwards compatibility fix

diff --git a/Engine/Plugins/Cameras/GameplayCameras/Source/GameplayCameras/Private/Debug/PostProcessSettingsDebugBlock.cpp b/Engine/Plugins/Cameras/GameplayCameras/Source/GameplayCameras/Private/Debug/PostProcessSettingsDebugBlock.cpp
--- a/Engine/Plugins/Cameras/GameplayCameras/Source/GameplayCameras/Private/Debug/PostProcessSettingsDebugBlock.cpp
+++ b/Engine/Plugins/Cameras/GameplayCameras/Source/GameplayCameras/Private/Debug/PostProcessSettingsDebugBlock.cpp
@@ -92,7 +92,9 @@ void FPostProcessSettingsDebugBlock::OnDebugDraw(const FCameraDebugBlockDrawPara
 		UE_DRAW_PP(SceneFringeIntensity);
 		UE_DRAW_PP(ChromaticAberrationStartOffset);
 		UE_DRAW_PP(BloomIntensity);
+#if UE_VERSION_NEWER_THAN_OR_EQUAL(5,7,0)
 		UE_DRAW_PP(BloomGaussianIntensity);
+#endif
 		UE_DRAW_PP(BloomThreshold);
 		UE_DRAW_PP(Bloom1Tint);
 		UE_DRAW_PP(BloomSizeScale);
@@ -109,7 +111,9 @@ void FPostProcessSettingsDebugBlock::OnDebugDraw(const FCameraDebugBlockDrawPara
 		UE_DRAW_PP(Bloom6Size);
 		UE_DRAW_PP(BloomDirtMaskIntensity);
 		UE_DRAW_PP(BloomDirtMaskTint);
+#if UE_VERSION_NEWER_THAN_OR_EQUAL(5,7,0)
 		UE_DRAW_PP(BloomConvolutionIntensity);
+#endif
 		UE_DRAW_PP(BloomConvolutionScatterDispersion);
 		UE_DRAW_PP(BloomConvolutionSize);
 		UE_DRAW_PP(BloomConvolutionCenterUV);
diff --git a/Engine/Plugins/Cameras/GameplayCameras/Source/GameplayCameras/Private/MovieScene/MovieSceneGameplayCamerasComponentTypes.cpp b/Engine/Plugins/Cameras/GameplayCameras/Source/GameplayCameras/Private/MovieScene/MovieSceneGameplayCamerasComponentTypes.cpp
--- a/Engine/Plugins/Cameras/GameplayCameras/Source/GameplayCameras/Private/MovieScene/MovieSceneGameplayCamerasComponentTypes.cpp
+++ b/Engine/Plugins/Cameras/GameplayCameras/Source/GameplayCameras/Private/MovieScene/MovieSceneGameplayCamerasComponentTypes.cpp
@@ -10,6 +10,7 @@
 #include "MovieSceneTracksComponentTypes.h"
 #include "Systems/MovieScenePiecewiseDoubleBlenderSystem.h"
 
+#if UE_VERSION_NEWER_THAN_OR_EQUAL(5,7,0)
 
 namespace UE::MovieScene
 {
@@ -22,6 +23,8 @@ namespace UE::MovieScene
 	}
 }
 
+#endif  // UE >= 5.7.0
+
 namespace UE::Cameras
 {
 
diff --git a/Engine/Plugins/Cameras/GameplayCameras/Source/GameplayCameras/Private/MovieScene/MovieSceneGameplayCamerasComponentTypes.h b/Engine/Plugins/Cameras/GameplayCameras/Source/GameplayCameras/Private/MovieScene/MovieSceneGameplayCamerasComponentTypes.h
--- a/Engine/Plugins/Cameras/GameplayCameras/Source/GameplayCameras/Private/MovieScene/MovieSceneGameplayCamerasComponentTypes.h
+++ b/Engine/Plugins/Cameras/GameplayCameras/Source/GameplayCameras/Private/MovieScene/MovieSceneGameplayCamerasComponentTypes.h
@@ -2,18 +2,27 @@
 
 #pragma once
 
-#include "Nodes/Framing/CameraFramingZone.h"
 #include "EntitySystem/MovieSceneEntityIDs.h"
 #include "EntitySystem/MovieScenePropertySystemTypes.h"
 #include "EntitySystem/MovieScenePropertyTraits.h"
+#include "Misc/EngineVersionComparison.h"
+#include "Nodes/Framing/CameraFramingZone.h"
+
+#if UE_VERSION_NEWER_THAN_OR_EQUAL(5,7,0)
 #include "EntitySystem/MovieScenePropertyMetaData.h"
+#else
+#include "EntitySystem/MovieScenePropertyMetaDataTraits.h"
+#endif  // UE >= 5.7.0
 
+#if UE_VERSION_NEWER_THAN_OR_EQUAL(5,7,0)
 
 namespace UE::MovieScene
 {
 	void UnpackChannelsFromOperational(FCameraFramingZone InZone, const FProperty& Property, FUnpackedChannelValues& OutUnpackedValues);
 }
 
+#endif  // UE >= 5.7.0
+
 namespace UE::Cameras
 {
 
