Cameras: add evaluation context to transition conditions' parameter structure

diff --git a/Engine/Plugins/Cameras/GameplayCameras/Source/GameplayCameras/Private/Core/PersistentBlendStackCameraNode.cpp b/Engine/Plugins/Cameras/GameplayCameras/Source/GameplayCameras/Private/Core/PersistentBlendStackCameraNode.cpp
--- a/Engine/Plugins/Cameras/GameplayCameras/Source/GameplayCameras/Private/Core/PersistentBlendStackCameraNode.cpp
+++ b/Engine/Plugins/Cameras/GameplayCameras/Source/GameplayCameras/Private/Core/PersistentBlendStackCameraNode.cpp
@@ -362,8 +362,8 @@ const UCameraRigTransition* FPersistentBlendStackCameraNodeEvaluator::FindEnterT
 	// Find a transition that works for blending the given camera rig in.
 	return FCameraRigTransitionFinder::FindTransition(
 			Params.CameraRig->EnterTransitions,
-			nullptr, nullptr, false,
-			Params.CameraRig, nullptr);
+			nullptr, nullptr, nullptr, false,
+			Params.EvaluationContext, Params.CameraRig, nullptr);
 }
 
 const UCameraRigTransition* FPersistentBlendStackCameraNodeEvaluator::FindExitTransition(const FCameraRigEntry& Entry, const UCameraRigTransition* TransitionOverride) const
@@ -377,8 +377,8 @@ const UCameraRigTransition* FPersistentBlendStackCameraNodeEvaluator::FindExitTr
 	// Find a transition that works for blending the given camera rig out.
 	return FCameraRigTransitionFinder::FindTransition(
 			Entry.CameraRig->ExitTransitions,
-			Entry.CameraRig, nullptr, Entry.Flags.bIsFrozen,
-			nullptr, nullptr);
+			Entry.EvaluationContext.Pin(), Entry.CameraRig, nullptr, Entry.Flags.bIsFrozen,
+			nullptr, nullptr, nullptr);
 }
 
 #if WITH_EDITOR
diff --git a/Engine/Plugins/Cameras/GameplayCameras/Source/GameplayCameras/Private/Core/TransientBlendStackCameraNode.cpp b/Engine/Plugins/Cameras/GameplayCameras/Source/GameplayCameras/Private/Core/TransientBlendStackCameraNode.cpp
--- a/Engine/Plugins/Cameras/GameplayCameras/Source/GameplayCameras/Private/Core/TransientBlendStackCameraNode.cpp
+++ b/Engine/Plugins/Cameras/GameplayCameras/Source/GameplayCameras/Private/Core/TransientBlendStackCameraNode.cpp
@@ -521,8 +521,8 @@ const UCameraRigTransition* FTransientBlendStackCameraNodeEvaluator::FindTransit
 				{
 					TransitionToUse = FCameraRigTransitionFinder::FindTransition(
 							FromCameraRig->ExitTransitions,
-							FromCameraRig, FromCameraAsset, bFromFrozen,
-							ToCameraRig, ToCameraAsset);
+							FromContext, FromCameraRig, FromCameraAsset, bFromFrozen,
+							ToContext, ToCameraRig, ToCameraAsset);
 					if (TransitionToUse)
 					{
 						return TransitionToUse;
@@ -539,8 +539,8 @@ const UCameraRigTransition* FTransientBlendStackCameraNodeEvaluator::FindTransit
 				{
 					TransitionToUse = FCameraRigTransitionFinder::FindTransition(
 							FromCameraAsset->GetExitTransitions(),
-							FromCameraRig, FromCameraAsset, bFromFrozen,
-							ToCameraRig, ToCameraAsset);
+							FromContext, FromCameraRig, FromCameraAsset, bFromFrozen,
+							ToContext, ToCameraRig, ToCameraAsset);
 					if (TransitionToUse)
 					{
 						return TransitionToUse;
@@ -559,8 +559,8 @@ const UCameraRigTransition* FTransientBlendStackCameraNodeEvaluator::FindTransit
 				{
 					TransitionToUse = FCameraRigTransitionFinder::FindTransition(
 							ToCameraRig->EnterTransitions,
-							FromCameraRig, FromCameraAsset, bFromFrozen,
-							ToCameraRig, ToCameraAsset);
+							FromContext, FromCameraRig, FromCameraAsset, bFromFrozen,
+							ToContext, ToCameraRig, ToCameraAsset);
 					if (TransitionToUse)
 					{
 						return TransitionToUse;
@@ -577,8 +577,8 @@ const UCameraRigTransition* FTransientBlendStackCameraNodeEvaluator::FindTransit
 				{
 					TransitionToUse = FCameraRigTransitionFinder::FindTransition(
 							ToCameraAsset->GetEnterTransitions(),
-							FromCameraRig, FromCameraAsset, bFromFrozen,
-							ToCameraRig, ToCameraAsset);
+							FromContext, FromCameraRig, FromCameraAsset, bFromFrozen,
+							ToContext, ToCameraRig, ToCameraAsset);
 					if (TransitionToUse)
 					{
 						return TransitionToUse;
diff --git a/Engine/Plugins/Cameras/GameplayCameras/Source/GameplayCameras/Private/Helpers/CameraRigTransitionFinder.cpp b/Engine/Plugins/Cameras/GameplayCameras/Source/GameplayCameras/Private/Helpers/CameraRigTransitionFinder.cpp
--- a/Engine/Plugins/Cameras/GameplayCameras/Source/GameplayCameras/Private/Helpers/CameraRigTransitionFinder.cpp
+++ b/Engine/Plugins/Cameras/GameplayCameras/Source/GameplayCameras/Private/Helpers/CameraRigTransitionFinder.cpp
@@ -9,14 +9,18 @@
 namespace UE::Cameras
 {
 
-const UCameraRigTransition* FCameraRigTransitionFinder::FindTransition(
+	const UCameraRigTransition* FCameraRigTransitionFinder::FindTransition(
 		TArrayView<const TObjectPtr<UCameraRigTransition>> Transitions, 
+		TSharedPtr<const FCameraEvaluationContext> FromEvaluationContext,
 		const UCameraRigAsset* FromCameraRig, const UCameraAsset* FromCameraAsset, bool bFromFrozen,
+		TSharedPtr<const FCameraEvaluationContext> ToEvaluationContext,
 		const UCameraRigAsset* ToCameraRig, const UCameraAsset* ToCameraAsset)
 {
 	FCameraRigTransitionConditionMatchParams MatchParams;
+	MatchParams.FromEvaluationContext = FromEvaluationContext;
 	MatchParams.FromCameraRig = FromCameraRig;
 	MatchParams.FromCameraAsset = FromCameraAsset;
+	MatchParams.ToEvaluationContext = ToEvaluationContext;
 	MatchParams.ToCameraRig = ToCameraRig;
 	MatchParams.ToCameraAsset = ToCameraAsset;
 
diff --git a/Engine/Plugins/Cameras/GameplayCameras/Source/GameplayCameras/Private/Helpers/CameraRigTransitionFinder.h b/Engine/Plugins/Cameras/GameplayCameras/Source/GameplayCameras/Private/Helpers/CameraRigTransitionFinder.h
--- a/Engine/Plugins/Cameras/GameplayCameras/Source/GameplayCameras/Private/Helpers/CameraRigTransitionFinder.h
+++ b/Engine/Plugins/Cameras/GameplayCameras/Source/GameplayCameras/Private/Helpers/CameraRigTransitionFinder.h
@@ -3,6 +3,7 @@
 #pragma once
 
 #include "UObject/ObjectPtr.h"
+#include "Templates/SharedPointerFwd.h"
 
 class UCameraAsset;
 class UCameraRigAsset;
@@ -11,16 +12,19 @@ class UCameraRigTransition;
 namespace UE::Cameras
 {
 
+class FCameraEvaluationContext;
+
 /**
  * Helper class for finding a camera rig transition that matches a given situation.
  */
 class FCameraRigTransitionFinder
 {
 public:
-
 	static const UCameraRigTransition* FindTransition(
 			TArrayView<const TObjectPtr<UCameraRigTransition>> Transitions, 
+			TSharedPtr<const FCameraEvaluationContext> FromEvaluationContext,
 			const UCameraRigAsset* FromCameraRig, const UCameraAsset* FromCameraAsset, bool bFromFrozen,
+			TSharedPtr<const FCameraEvaluationContext> ToEvaluationContext,
 			const UCameraRigAsset* ToCameraRig, const UCameraAsset* ToCameraAsset);
 };
 
diff --git a/Engine/Plugins/Cameras/GameplayCameras/Source/GameplayCameras/Public/Core/CameraRigTransition.h b/Engine/Plugins/Cameras/GameplayCameras/Source/GameplayCameras/Public/Core/CameraRigTransition.h
--- a/Engine/Plugins/Cameras/GameplayCameras/Source/GameplayCameras/Public/Core/CameraRigTransition.h
+++ b/Engine/Plugins/Cameras/GameplayCameras/Source/GameplayCameras/Public/Core/CameraRigTransition.h
@@ -5,6 +5,7 @@
 #include "Containers/Array.h"
 #include "Core/ObjectTreeGraphObject.h"
 #include "CoreTypes.h"
+#include "Templates/SharedPointerFwd.h"
 #include "UObject/ObjectPtr.h"
 
 #include "CameraRigTransition.generated.h"
@@ -18,6 +19,7 @@ class UCameraRigAsset;
 namespace UE::Cameras
 {
 	struct FCameraObjectBuildContext;
+	class FCameraEvaluationContext;
 }
 
 /**
@@ -25,11 +27,15 @@ namespace UE::Cameras
  */
 struct FCameraRigTransitionConditionMatchParams
 {
+	/** The evaluation context of the previous camera rig. */
+	TSharedPtr<const UE::Cameras::FCameraEvaluationContext> FromEvaluationContext;
 	/** The previous camera rig. */
 	const UCameraRigAsset* FromCameraRig = nullptr;
 	/** The previous camera asset. */
 	const UCameraAsset* FromCameraAsset = nullptr;
 
+	/** The evaluation context of the next camera rig. */
+	TSharedPtr<const UE::Cameras::FCameraEvaluationContext> ToEvaluationContext;
 	/** The next camera rig. */
 	const UCameraRigAsset* ToCameraRig = nullptr;
 	/** The next camera asset. */
