Cameras: add missing Blueprint getter nodes for camera rig parameters

diff --git a/Engine/Plugins/Cameras/GameplayCameras/Source/GameplayCameras/Public/GameFramework/CameraRigParameterInterop.h b/Engine/Plugins/Cameras/GameplayCameras/Source/GameplayCameras/Public/GameFramework/CameraRigParameterInterop.h
--- a/Engine/Plugins/Cameras/GameplayCameras/Source/GameplayCameras/Public/GameFramework/CameraRigParameterInterop.h
+++ b/Engine/Plugins/Cameras/GameplayCameras/Source/GameplayCameras/Public/GameFramework/CameraRigParameterInterop.h
@@ -30,7 +30,7 @@ public:
 
 public:
 
-	UFUNCTION(BlueprintCallable, CustomThunk, Category="Camera", meta=(BlueprintInternalUseOnly="true", CustomStructureParam="ReturnValue"))
+	UFUNCTION(BlueprintPure, CustomThunk, Category="Camera", meta=(BlueprintInternalUseOnly="true", CustomStructureParam="ReturnValue"))
 	static void GetCameraParameter(UPARAM(Ref) const FBlueprintCameraEvaluationDataRef& CameraData, UCameraRigAsset* CameraRig, FName ParameterName, int32& ReturnValue);
 
 	UFUNCTION(BlueprintCallable, CustomThunk, Category="Camera", meta=(BlueprintInternalUseOnly="true", CustomStructureParam="NewValue"))
diff --git a/Engine/Plugins/Cameras/GameplayCameras/Source/GameplayCamerasUncookedOnly/Private/BlueprintGraph/K2Node_CameraRigBase.cpp b/Engine/Plugins/Cameras/GameplayCameras/Source/GameplayCamerasUncookedOnly/Private/BlueprintGraph/K2Node_CameraRigBase.cpp
--- a/Engine/Plugins/Cameras/GameplayCameras/Source/GameplayCamerasUncookedOnly/Private/BlueprintGraph/K2Node_CameraRigBase.cpp
+++ b/Engine/Plugins/Cameras/GameplayCameras/Source/GameplayCamerasUncookedOnly/Private/BlueprintGraph/K2Node_CameraRigBase.cpp
@@ -46,9 +46,12 @@ void UK2Node_CameraRigBase::AllocateDefaultPins()
 {
 	using namespace UE::Cameras;
 
-	// Add execution pins.
-	CreatePin(EGPD_Input, UEdGraphSchema_K2::PC_Exec, UEdGraphSchema_K2::PN_Execute);
-	CreatePin(EGPD_Output, UEdGraphSchema_K2::PC_Exec, UEdGraphSchema_K2::PN_Then);
+	if (!IsNodePure())
+	{
+		// Add execution pins.
+		CreatePin(EGPD_Input, UEdGraphSchema_K2::PC_Exec, UEdGraphSchema_K2::PN_Execute);
+		CreatePin(EGPD_Output, UEdGraphSchema_K2::PC_Exec, UEdGraphSchema_K2::PN_Then);
+	}
 
 	// Add evalation result pin.
 	CreatePin(EGPD_Input, UEdGraphSchema_K2::PC_Struct, FBlueprintCameraEvaluationDataRef::StaticStruct(), CameraNodeEvaluationResultPinName);
diff --git a/Engine/Plugins/Cameras/GameplayCameras/Source/GameplayCamerasUncookedOnly/Private/BlueprintGraph/K2Node_GetCameraRigParameter.cpp b/Engine/Plugins/Cameras/GameplayCameras/Source/GameplayCamerasUncookedOnly/Private/BlueprintGraph/K2Node_GetCameraRigParameter.cpp
new file mode 100644
--- /dev/null
+++ b/Engine/Plugins/Cameras/GameplayCameras/Source/GameplayCamerasUncookedOnly/Private/BlueprintGraph/K2Node_GetCameraRigParameter.cpp
@@ -0,0 +1,128 @@
+// Copyright Epic Games, Inc. All Rights Reserved.
+
+#include "BlueprintGraph/K2Node_GetCameraRigParameter.h"
+
+#include "BlueprintActionDatabaseRegistrar.h"
+#include "BlueprintNodeSpawner.h"
+#include "Core/CameraRigAsset.h"
+#include "EdGraphSchema_K2.h"
+#include "GameFramework/CameraRigParameterInterop.h"
+#include "K2Node_CallFunction.h"
+#include "Kismet2/BlueprintEditorUtils.h"
+#include "KismetCompiler.h"
+
+#include UE_INLINE_GENERATED_CPP_BY_NAME(K2Node_GetCameraRigParameter)
+
+#define LOCTEXT_NAMESPACE "K2Node_GetCameraRigParameter"
+
+void UK2Node_GetCameraRigParameter::AllocateDefaultPins()
+{
+	Super::AllocateDefaultPins();
+
+	// Add parameter value pin.
+	FEdGraphPinType PinType = GetParameterPinType();
+	CreatePin(EGPD_Output, PinType, FName(CameraParameterName));
+}
+
+FText UK2Node_GetCameraRigParameter::GetNodeTitle(ENodeTitleType::Type TitleType) const
+{
+	return FText::Format(LOCTEXT("BaseNodeTitle", "GET on {0}"), FText::FromString(GetNameSafe(CameraRig)));
+}
+
+FText UK2Node_GetCameraRigParameter::GetTooltipText() const
+{
+	return FText::Format(
+			LOCTEXT("NodeTooltip", "Gets the value of camera rig {0}'s parameter {1} on the given evaluation data."),
+			FText::FromString(GetNameSafe(CameraRig)),
+			FText::FromString(CameraParameterName));
+}
+
+void UK2Node_GetCameraRigParameter::GetMenuActions(FBlueprintActionDatabaseRegistrar& ActionRegistrar, const FAssetData& CameraRigAssetData) const
+{
+	const FText BaseCategoryString = GetMenuCategory();
+
+	for (TPair<FName, FAssetTagValueRef> It : CameraRigAssetData.TagsAndValues)
+	{
+		const FName ParameterName = It.Key;
+
+		UBlueprintNodeSpawner* NodeSpawner = UBlueprintNodeSpawner::Create(GetClass());
+		NodeSpawner->DefaultMenuSignature.Category = FText::Join(
+				FText::FromString(TEXT("|")), BaseCategoryString, FText::FromName(CameraRigAssetData.AssetName));
+		NodeSpawner->DefaultMenuSignature.MenuName = FText::Format(
+				LOCTEXT("GetCameraRigParameterActionMenuName", "Get {0}"),
+				FText::FromName(ParameterName));
+		NodeSpawner->CustomizeNodeDelegate = UBlueprintNodeSpawner::FCustomizeNodeDelegate::CreateLambda(
+				[ParameterName, CameraRigAssetData](UEdGraphNode* NewNode, bool bIsTemplateNode)
+				{
+					UK2Node_GetCameraRigParameter* NewSetter = CastChecked<UK2Node_GetCameraRigParameter>(NewNode);
+					NewSetter->Initialize(CameraRigAssetData, ParameterName.ToString());
+				});
+
+		ActionRegistrar.AddBlueprintAction(CameraRigAssetData, NodeSpawner);
+	}
+}
+
+void UK2Node_GetCameraRigParameter::ExpandNode(FKismetCompilerContext& CompilerContext, UEdGraph* SourceGraph)
+{
+	Super::ExpandNode(CompilerContext, SourceGraph);
+
+	if (!ValidateCameraRigBeforeExpandNode(CompilerContext))
+	{
+		BreakAllNodeLinks();
+		return;
+	}
+
+	if (CameraParameterName.IsEmpty())
+	{
+		CompilerContext.MessageLog.Error(*LOCTEXT("ErrorMissingCameraParameterName", "GetCameraRigParameter node @@ doesn't have a valid camera parameter name set.").ToString(), this);
+		BreakAllNodeLinks();
+		return;
+	}
+
+	// Additional setup is needed for some types.
+	UEnum* ExtraEnumTypePinObject = nullptr;
+	bool bMakeExtraInstanceStructSetup = false;
+	if (DataCameraParameterType == ECameraContextDataType::Enum)
+	{
+		const UEnum* EnumType = CastChecked<const UEnum>(DataCameraParameterTypeObject);
+		ExtraEnumTypePinObject = const_cast<UEnum*>(EnumType);
+	}
+	else if (DataCameraParameterType == ECameraContextDataType::Struct)
+	{
+		bMakeExtraInstanceStructSetup = true;
+	}
+
+	UEdGraphPin* const CameraNodeEvaluationResultPin = GetCameraNodeEvaluationResultPin();
+	UEdGraphPin* const CameraParameterValuePin = FindPinChecked(CameraParameterName);
+
+	// Make the GetXxxParameter function call node.
+	UK2Node_CallFunction* CallGetParameter = CompilerContext.SpawnIntermediateNode<UK2Node_CallFunction>(this, SourceGraph);
+	CallGetParameter->FunctionReference.SetExternalMember(TEXT("GetCameraParameter"), UCameraRigParameterInterop::StaticClass());
+	CallGetParameter->AllocateDefaultPins();
+	CompilerContext.MessageLog.NotifyIntermediateObjectCreation(CallGetParameter, SourceGraph);
+
+	// Connect the camera evaluation result argument.
+	UEdGraphPin* CallGetParameterResultPin = CallGetParameter->FindPinChecked(TEXT("CameraData"));
+	CompilerContext.CopyPinLinksToIntermediate(*CameraNodeEvaluationResultPin, *CallGetParameterResultPin);
+
+	// Set the camera rig argument.
+	UEdGraphPin* CallGetParameterCameraRigPin = CallGetParameter->FindPinChecked(TEXT("CameraRig"));
+	CallGetParameterCameraRigPin->DefaultObject = CameraRig;
+
+	// Set the parameter name argument.
+	UEdGraphPin* CallGetParameterNamePin = CallGetParameter->FindPinChecked(TEXT("ParameterName"));
+	CallGetParameterNamePin->DefaultValue = CameraParameterName;
+
+	// Connect the output parameter value.
+	UEdGraphPin* CallGetParameterValuePin = CallGetParameter->FindPinChecked(TEXT("ReturnValue"));
+	CallGetParameterValuePin->PinType = CameraParameterValuePin->PinType;
+	if (CameraParameterValuePin->LinkedTo.Num() > 0)
+	{
+		CompilerContext.MovePinLinksToIntermediate(*CameraParameterValuePin, *CallGetParameterValuePin);
+	}
+
+	BreakAllNodeLinks();
+}
+
+#undef LOCTEXT_NAMESPACE
+
diff --git a/Engine/Plugins/Cameras/GameplayCameras/Source/GameplayCamerasUncookedOnly/Private/BlueprintGraph/K2Node_GetCameraRigParameter.h b/Engine/Plugins/Cameras/GameplayCameras/Source/GameplayCamerasUncookedOnly/Private/BlueprintGraph/K2Node_GetCameraRigParameter.h
new file mode 100644
--- /dev/null
+++ b/Engine/Plugins/Cameras/GameplayCameras/Source/GameplayCamerasUncookedOnly/Private/BlueprintGraph/K2Node_GetCameraRigParameter.h
@@ -0,0 +1,36 @@
+// Copyright Epic Games, Inc. All Rights Reserved.
+
+#pragma once
+
+#include "K2Node_SingleCameraRigParameterBase.h"
+
+#include "K2Node_GetCameraRigParameter.generated.h"
+
+class FKismetCompilerContext;
+class UCameraRigAsset;
+
+/**
+ * Blueprint node that, given a camera rig, lets the user get the value of one single
+ * exposed parameter.
+ */
+UCLASS(MinimalAPI)
+class UK2Node_GetCameraRigParameter : public UK2Node_SingleCameraRigParameterBase
+{
+	GENERATED_BODY()
+
+public:
+
+	// UEdGraphNode interface.
+	virtual void AllocateDefaultPins() override;
+	virtual FText GetNodeTitle(ENodeTitleType::Type TitleType) const override;
+	virtual FText GetTooltipText() const override;
+	//virtual bool ShouldDrawCompact() const override { return true; }
+	virtual bool IsNodePure() const override { return true; }
+	virtual void ExpandNode(FKismetCompilerContext& CompilerContext, UEdGraph* SourceGraph) override;
+
+protected:
+
+	// UK2Node_CameraRigBase interface.
+	void GetMenuActions(FBlueprintActionDatabaseRegistrar& ActionRegistrar, const FAssetData& CameraRigAssetData) const;
+};
+
diff --git a/Engine/Plugins/Cameras/GameplayCameras/Source/GameplayCamerasUncookedOnly/Private/BlueprintGraph/K2Node_GetCameraRigParameters.cpp b/Engine/Plugins/Cameras/GameplayCameras/Source/GameplayCamerasUncookedOnly/Private/BlueprintGraph/K2Node_GetCameraRigParameters.cpp
new file mode 100644
--- /dev/null
+++ b/Engine/Plugins/Cameras/GameplayCameras/Source/GameplayCamerasUncookedOnly/Private/BlueprintGraph/K2Node_GetCameraRigParameters.cpp
@@ -0,0 +1,167 @@
+// Copyright Epic Games, Inc. All Rights Reserved.
+
+#include "BlueprintGraph/K2Node_GetCameraRigParameters.h"
+
+#include "BlueprintActionDatabaseRegistrar.h"
+#include "BlueprintNodeSpawner.h"
+#include "Core/CameraRigAsset.h"
+#include "EdGraphSchema_K2.h"
+#include "GameFramework/CameraRigParameterInterop.h"
+#include "K2Node_CallFunction.h"
+#include "Kismet2/BlueprintEditorUtils.h"
+#include "KismetCompiler.h"
+
+#include UE_INLINE_GENERATED_CPP_BY_NAME(K2Node_GetCameraRigParameters)
+
+#define LOCTEXT_NAMESPACE "K2Node_GetCameraRigParameters"
+
+void UK2Node_GetCameraRigParameters::AllocateDefaultPins()
+{
+	Super::AllocateDefaultPins();
+
+	CreateParameterPins(EGPD_Output);
+}
+
+FText UK2Node_GetCameraRigParameters::GetNodeTitle(ENodeTitleType::Type TitleType) const
+{
+	return FText::Format(
+			LOCTEXT("BaseNodeTitle", "GET on {0}"),
+			FText::FromString(GetNameSafe(CameraRig)));
+}
+
+FText UK2Node_GetCameraRigParameters::GetTooltipText() const
+{
+	return FText::Format(
+			LOCTEXT("NodeTooltip", "Gets the values of all camera rig parameters on {0} on the given evaluation data."),
+			FText::FromString(GetNameSafe(CameraRig)));
+}
+
+void UK2Node_GetCameraRigParameters::GetMenuActions(FBlueprintActionDatabaseRegistrar& ActionRegistrar, const FAssetData& CameraRigAssetData) const
+{
+	const FText BaseCategoryString = GetMenuCategory();
+
+	UBlueprintNodeSpawner* NodeSpawner = UBlueprintNodeSpawner::Create(GetClass());
+	NodeSpawner->DefaultMenuSignature.Category = BaseCategoryString;
+	NodeSpawner->DefaultMenuSignature.MenuName = FText::Format(
+			LOCTEXT("GetCameraRigParameterActionMenuName", "Get all parameters on {0}"),
+			FText::FromName(CameraRigAssetData.AssetName));
+	NodeSpawner->CustomizeNodeDelegate = UBlueprintNodeSpawner::FCustomizeNodeDelegate::CreateLambda(
+			[CameraRigAssetData](UEdGraphNode* NewNode, bool bIsTemplateNode)
+			{
+				UK2Node_GetCameraRigParameters* NewGetter = CastChecked<UK2Node_GetCameraRigParameters>(NewNode);
+				NewGetter->Initialize(CameraRigAssetData);
+			});
+
+	ActionRegistrar.AddBlueprintAction(CameraRigAssetData, NodeSpawner);
+}
+
+void UK2Node_GetCameraRigParameters::ExpandNode(FKismetCompilerContext& CompilerContext, UEdGraph* SourceGraph)
+{
+	Super::ExpandNode(CompilerContext, SourceGraph);
+
+	if (!ValidateCameraRigBeforeExpandNode(CompilerContext))
+	{
+		BreakAllNodeLinks();
+		return;
+	}
+
+	UEdGraphPin* const CameraNodeEvaluationResultPin = GetCameraNodeEvaluationResultPin();
+
+	// For each blendable and data parameter, we figure out the type of GetXxxParameter function to call on the UCameraRigParameterInterop
+	// function library. We then make a K2Node_CallFunction node for it, and connect all its inputs, including connecting the parameter
+	// value to whatever our node's corresponding parameter value pin was connected to. We basically transform our GetCameraRigParameters
+	// node into an array of single parameter getter nodes.
+
+	TArray<UEdGraphPin*> BlendableParameterPins;
+	FindBlendableParameterPins(BlendableParameterPins);
+	for (UEdGraphPin* RigParameterPin : BlendableParameterPins)
+	{
+		const UCameraObjectInterfaceBlendableParameter* BlendableParameter = CameraRig->Interface.FindBlendableParameterByName(RigParameterPin->GetName());
+		if (!BlendableParameter)
+		{
+			CompilerContext.MessageLog.Error(*LOCTEXT("ErrorMissingParameter", "GetCameraRigParameters node @@ is trying to set parameter @@ but camera rig @@ has no such parameter.").ToString(), this, *RigParameterPin->GetName(), CameraRig);
+			continue;
+		}
+
+		if (!BlendableParameter->PrivateVariableID)
+		{
+			CompilerContext.MessageLog.Error(*LOCTEXT("ErrorMissingParameterVariable", "GetCameraRigParameters node @@ needs camera rig @@ to be built.").ToString(), this, CameraRig);
+			continue;
+		}
+
+		// Make the GetXxxParameter function call node.
+		UK2Node_CallFunction* CallGetParameter = CompilerContext.SpawnIntermediateNode<UK2Node_CallFunction>(this, SourceGraph);
+		CallGetParameter->FunctionReference.SetExternalMember(TEXT("GetCameraParameter"), UCameraRigParameterInterop::StaticClass());
+		CallGetParameter->AllocateDefaultPins();
+		CompilerContext.MessageLog.NotifyIntermediateObjectCreation(CallGetParameter, SourceGraph);
+
+		// Connect the camera evaluation result argument.
+		UEdGraphPin* CallGetParameterResultPin = CallGetParameter->FindPinChecked(TEXT("CameraData"));
+		CompilerContext.CopyPinLinksToIntermediate(*CameraNodeEvaluationResultPin, *CallGetParameterResultPin);
+
+		// Set the camera rig argument.
+		UEdGraphPin* CallGetParameterCameraRigPin = CallGetParameter->FindPinChecked(TEXT("CameraRig"));
+		CallGetParameterCameraRigPin->DefaultObject = CameraRig;
+
+		// Set the parameter name argument.
+		UEdGraphPin* CallGetParameterNamePin = CallGetParameter->FindPinChecked(TEXT("ParameterName"));
+		CallGetParameterNamePin->DefaultValue = BlendableParameter->InterfaceParameterName;
+
+		// Connect the output parameter value.
+		UEdGraphPin* CallGetParameterValuePin = CallGetParameter->FindPinChecked(TEXT("ReturnValue"));
+		CallGetParameterValuePin->PinType = RigParameterPin->PinType;
+		if (RigParameterPin->LinkedTo.Num() > 0)
+		{
+			CompilerContext.MovePinLinksToIntermediate(*RigParameterPin, *CallGetParameterValuePin);
+		}
+	}
+
+	TArray<UEdGraphPin*> DataParameterPins;
+	FindDataParameterPins(DataParameterPins);
+	for (UEdGraphPin* RigParameterPin : DataParameterPins)
+	{
+		const UCameraObjectInterfaceDataParameter* DataParameter = CameraRig->Interface.FindDataParameterByName(RigParameterPin->GetName());
+		if (!DataParameter)
+		{
+			CompilerContext.MessageLog.Error(*LOCTEXT("ErrorMissingParameter", "GetCameraRigParameters node @@ is trying to set parameter @@ but camera rig @@ has no such parameter.").ToString(), this, *RigParameterPin->GetName(), CameraRig);
+			continue;
+		}
+
+		if (!DataParameter->PrivateDataID.IsValid())
+		{
+			CompilerContext.MessageLog.Error(*LOCTEXT("ErrorMissingParameterVariable", "GetCameraRigParameters node @@ needs camera rig @@ to be built.").ToString(), this, CameraRig);
+			continue;
+		}
+
+		// Make the SetXxxData function call node.
+		UK2Node_CallFunction* CallGetParameter = CompilerContext.SpawnIntermediateNode<UK2Node_CallFunction>(this, SourceGraph);
+		CallGetParameter->FunctionReference.SetExternalMember(TEXT("SetCameraParameter"), UCameraRigParameterInterop::StaticClass());
+		CallGetParameter->AllocateDefaultPins();
+		CompilerContext.MessageLog.NotifyIntermediateObjectCreation(CallGetParameter, SourceGraph);
+
+		// Connect the camera evaluation result argument.
+		UEdGraphPin* CallGetParameterResultPin = CallGetParameter->FindPinChecked(TEXT("CameraData"));
+		CompilerContext.CopyPinLinksToIntermediate(*CameraNodeEvaluationResultPin, *CallGetParameterResultPin);
+
+		// Set the camera rig argument.
+		UEdGraphPin* CallGetParameterCameraRigPin = CallGetParameter->FindPinChecked(TEXT("CameraRig"));
+		CallGetParameterCameraRigPin->DefaultObject = CameraRig;
+
+		// Set the parameter name argument.
+		UEdGraphPin* CallGetParameterNamePin = CallGetParameter->FindPinChecked(TEXT("ParameterName"));
+		CallGetParameterNamePin->DefaultValue = DataParameter->InterfaceParameterName;
+
+		// Set or connect the parameter value argument.
+		UEdGraphPin* CallGetParameterValuePin = CallGetParameter->FindPinChecked(TEXT("NewValue"));
+		CallGetParameterValuePin->PinType = RigParameterPin->PinType;
+		if (RigParameterPin->LinkedTo.Num() > 0)
+		{
+			CompilerContext.MovePinLinksToIntermediate(*RigParameterPin, *CallGetParameterValuePin);
+		}
+	}
+
+	BreakAllNodeLinks();
+}
+
+#undef LOCTEXT_NAMESPACE
+
diff --git a/Engine/Plugins/Cameras/GameplayCameras/Source/GameplayCamerasUncookedOnly/Private/BlueprintGraph/K2Node_GetCameraRigParameters.h b/Engine/Plugins/Cameras/GameplayCameras/Source/GameplayCamerasUncookedOnly/Private/BlueprintGraph/K2Node_GetCameraRigParameters.h
new file mode 100644
--- /dev/null
+++ b/Engine/Plugins/Cameras/GameplayCameras/Source/GameplayCamerasUncookedOnly/Private/BlueprintGraph/K2Node_GetCameraRigParameters.h
@@ -0,0 +1,35 @@
+// Copyright Epic Games, Inc. All Rights Reserved.
+
+#pragma once
+
+#include "K2Node_MultiCameraRigParametersBase.h"
+
+#include "K2Node_GetCameraRigParameters.generated.h"
+
+class UCameraRigAsset;
+
+/**
+ * Blueprint node that, given a camera rig, lets the user get the values of all exposed parameters
+ * on that camera rig.
+ */
+UCLASS(MinimalAPI)
+class UK2Node_GetCameraRigParameters : public UK2Node_MultiCameraRigParametersBase
+{
+	GENERATED_BODY()
+
+public:
+
+	// UEdGraphNode interface.
+	virtual void AllocateDefaultPins() override;
+	virtual FText GetNodeTitle(ENodeTitleType::Type TitleType) const override;
+	virtual FText GetTooltipText() const override;
+	//virtual bool ShouldDrawCompact() const override { return true; }
+	virtual bool IsNodePure() const override { return true; }
+	virtual void ExpandNode(class FKismetCompilerContext& CompilerContext, UEdGraph* SourceGraph) override;
+
+protected:
+
+	// UK2Node_CameraRigBase interface.
+	virtual void GetMenuActions(FBlueprintActionDatabaseRegistrar& ActionRegistrar, const FAssetData& CameraRigAssetData) const override;
+};
+
diff --git a/Engine/Plugins/Cameras/GameplayCameras/Source/GameplayCamerasUncookedOnly/Private/BlueprintGraph/K2Node_MultiCameraRigParametersBase.cpp b/Engine/Plugins/Cameras/GameplayCameras/Source/GameplayCamerasUncookedOnly/Private/BlueprintGraph/K2Node_MultiCameraRigParametersBase.cpp
new file mode 100644
--- /dev/null
+++ b/Engine/Plugins/Cameras/GameplayCameras/Source/GameplayCamerasUncookedOnly/Private/BlueprintGraph/K2Node_MultiCameraRigParametersBase.cpp
@@ -0,0 +1,144 @@
+// Copyright Epic Games, Inc. All Rights Reserved.
+
+#include "BlueprintGraph/K2Node_MultiCameraRigParametersBase.h"
+
+#include "BlueprintActionDatabaseRegistrar.h"
+#include "BlueprintNodeSpawner.h"
+#include "Core/CameraRigAsset.h"
+#include "EdGraphSchema_K2.h"
+#include "GameFramework/CameraRigParameterInterop.h"
+#include "K2Node_CallFunction.h"
+#include "Kismet/BlueprintInstancedStructLibrary.h"
+#include "Kismet2/BlueprintEditorUtils.h"
+#include "KismetCompiler.h"
+
+#include UE_INLINE_GENERATED_CPP_BY_NAME(K2Node_MultiCameraRigParametersBase)
+
+#define LOCTEXT_NAMESPACE "K2Node_MultiCameraRigParametersBase"
+
+UK2Node_MultiCameraRigParametersBase::UK2Node_MultiCameraRigParametersBase(const FObjectInitializer& ObjectInit)
+	: Super(ObjectInit)
+{
+}
+
+void UK2Node_MultiCameraRigParametersBase::Initialize(const FAssetData& UnloadedCameraRig)
+{
+	CameraRig = Cast<UCameraRigAsset>(UnloadedCameraRig.GetAsset());
+}
+
+void UK2Node_MultiCameraRigParametersBase::ReallocatePinsDuringReconstruction(TArray<UEdGraphPin*>& OldPins)
+{
+	// The camera rig might not be loaded yet when we are rebuilt on startup, so make sure
+	// this is all good otherwise we won't be able to get all our pins back.
+	if (CameraRig)
+	{
+		PreloadObject(CameraRig);
+		for (UCameraObjectInterfaceBlendableParameter* BlendableParameter : CameraRig->Interface.BlendableParameters)
+		{
+			PreloadObject(BlendableParameter);
+			if (BlendableParameter)
+			{
+				PreloadObject(const_cast<UScriptStruct*>(BlendableParameter->BlendableStructType.Get()));
+			}
+		}
+		for (UCameraObjectInterfaceDataParameter* DataParameter : CameraRig->Interface.DataParameters)
+		{
+			PreloadObject(DataParameter);
+			if (DataParameter)
+			{
+				PreloadObject(const_cast<UObject*>(DataParameter->DataTypeObject.Get()));
+			}
+		}
+	}
+
+	// Now do as usual: create all default pins (including the parameter pins) and 
+	// restore split pins.
+	AllocateDefaultPins();
+	RestoreSplitPins(OldPins);
+}
+
+void UK2Node_MultiCameraRigParametersBase::CreateParameterPins(EEdGraphPinDirection PinDirection)
+{
+	BlendableParameterPinNames.Reset();
+	DataParameterPinNames.Reset();
+
+	if (!CameraRig)
+	{
+		return;
+	}
+
+	for (const UCameraObjectInterfaceBlendableParameter* BlendableParameter : CameraRig->Interface.BlendableParameters)
+	{
+		if (!ensure(BlendableParameter))
+		{
+			continue;
+		}
+
+		if (!BlendableParameter->PrivateVariableID)
+		{
+			// Camera rig isn't fully built.
+			continue;
+		}
+		
+		FEdGraphPinType PinType = MakeBlendableParameterPinType(BlendableParameter);
+		if (PinType.PinCategory.IsNone())
+		{
+			// Unsupported type for Blueprints.
+			continue;
+		}
+
+		UEdGraphPin* NewPin = CreatePin(PinDirection, PinType, FName(BlendableParameter->InterfaceParameterName));
+		BlendableParameterPinNames.Add(NewPin->PinName);
+	}
+
+	for (const UCameraObjectInterfaceDataParameter* DataParameter : CameraRig->Interface.DataParameters)
+	{
+		if (!ensure(DataParameter))
+		{
+			continue;
+		}
+		
+		if (!DataParameter->PrivateDataID.IsValid())
+		{
+			// Camera rig isn't fully built.
+			continue;
+		}
+
+		FEdGraphPinType PinType = MakeDataParameterPinType(DataParameter);
+		if (PinType.PinCategory.IsNone())
+		{
+			// Unsupported type for Blueprints.
+			continue;
+		}
+
+		UEdGraphPin* NewPin = CreatePin(PinDirection, PinType, FName(DataParameter->InterfaceParameterName));
+		DataParameterPinNames.Add(NewPin->PinName);
+	}
+}
+
+void UK2Node_MultiCameraRigParametersBase::FindBlendableParameterPins(TArray<UEdGraphPin*>& OutPins) const
+{
+	for (const FName& PinName : BlendableParameterPinNames)
+	{
+		UEdGraphPin* Pin = FindPin(PinName);
+		if (ensure(Pin))
+		{
+			OutPins.Add(Pin);
+		}
+	}
+}
+
+void UK2Node_MultiCameraRigParametersBase::FindDataParameterPins(TArray<UEdGraphPin*>& OutPins) const
+{
+	for (const FName& PinName : DataParameterPinNames)
+	{
+		UEdGraphPin* Pin = FindPin(PinName);
+		if (ensure(Pin))
+		{
+			OutPins.Add(Pin);
+		}
+	}
+}
+
+#undef LOCTEXT_NAMESPACE
+
diff --git a/Engine/Plugins/Cameras/GameplayCameras/Source/GameplayCamerasUncookedOnly/Private/BlueprintGraph/K2Node_MultiCameraRigParametersBase.h b/Engine/Plugins/Cameras/GameplayCameras/Source/GameplayCamerasUncookedOnly/Private/BlueprintGraph/K2Node_MultiCameraRigParametersBase.h
new file mode 100644
--- /dev/null
+++ b/Engine/Plugins/Cameras/GameplayCameras/Source/GameplayCamerasUncookedOnly/Private/BlueprintGraph/K2Node_MultiCameraRigParametersBase.h
@@ -0,0 +1,44 @@
+// Copyright Epic Games, Inc. All Rights Reserved.
+
+#pragma once
+
+#include "K2Node_CameraRigBase.h"
+
+#include "K2Node_MultiCameraRigParametersBase.generated.h"
+
+class UCameraRigAsset;
+
+/**
+ * A base class for a Blueprint node that gets or sets multiple camera rig parameters.
+ */
+UCLASS(MinimalAPI)
+class UK2Node_MultiCameraRigParametersBase : public UK2Node_CameraRigBase
+{
+	GENERATED_BODY()
+
+public:
+
+	UK2Node_MultiCameraRigParametersBase(const FObjectInitializer& ObjectInit);
+
+	void Initialize(const FAssetData& UnloadedCameraRig);
+
+public:
+
+	// UK2Node interface.
+	virtual void ReallocatePinsDuringReconstruction(TArray<UEdGraphPin*>& OldPins) override;
+
+protected:
+
+	void CreateParameterPins(EEdGraphPinDirection PinDirection);
+	void FindBlendableParameterPins(TArray<UEdGraphPin*>& OutPins) const;
+	void FindDataParameterPins(TArray<UEdGraphPin*>& OutPins) const;
+
+protected:
+
+	UPROPERTY()
+	TArray<FName> BlendableParameterPinNames;
+
+	UPROPERTY()
+	TArray<FName> DataParameterPinNames;
+};
+
diff --git a/Engine/Plugins/Cameras/GameplayCameras/Source/GameplayCamerasUncookedOnly/Private/BlueprintGraph/K2Node_SetCameraRigParameter.cpp b/Engine/Plugins/Cameras/GameplayCameras/Source/GameplayCamerasUncookedOnly/Private/BlueprintGraph/K2Node_SetCameraRigParameter.cpp
--- a/Engine/Plugins/Cameras/GameplayCameras/Source/GameplayCamerasUncookedOnly/Private/BlueprintGraph/K2Node_SetCameraRigParameter.cpp
+++ b/Engine/Plugins/Cameras/GameplayCameras/Source/GameplayCamerasUncookedOnly/Private/BlueprintGraph/K2Node_SetCameraRigParameter.cpp
@@ -3,83 +3,24 @@
 #include "BlueprintGraph/K2Node_SetCameraRigParameter.h"
 
 #include "BlueprintActionDatabaseRegistrar.h"
-#include "BlueprintGraph/K2Node_CameraRigBase.h"
 #include "BlueprintNodeSpawner.h"
 #include "Core/CameraRigAsset.h"
 #include "EdGraphSchema_K2.h"
-#include "EditorCategoryUtils.h"
 #include "GameFramework/CameraRigParameterInterop.h"
 #include "K2Node_CallFunction.h"
-#include "Kismet/BlueprintInstancedStructLibrary.h"
 #include "Kismet2/BlueprintEditorUtils.h"
 #include "KismetCompiler.h"
-#include "Misc/EngineVersionComparison.h"
-#include "Styling/AppStyle.h"
-#include "Styling/SlateIconFinder.h"
 
 #include UE_INLINE_GENERATED_CPP_BY_NAME(K2Node_SetCameraRigParameter)
 
 #define LOCTEXT_NAMESPACE "K2Node_SetCameraRigParameter"
 
-UK2Node_SetCameraRigParameter::UK2Node_SetCameraRigParameter(const FObjectInitializer& ObjectInit)
-	: Super(ObjectInit)
-{
-}
-
-void UK2Node_SetCameraRigParameter::Initialize(const FAssetData& UnloadedCameraRig, const FString& InCameraParameterName)
-{
-	UCameraRigAsset* LoadedCameraRig = Cast<UCameraRigAsset>(UnloadedCameraRig.GetAsset());
-	if (ensure(LoadedCameraRig))
-	{
-		if (const UCameraObjectInterfaceBlendableParameter* BlendableParameter = LoadedCameraRig->Interface.FindBlendableParameterByName(InCameraParameterName))
-		{
-			Initialize(LoadedCameraRig, InCameraParameterName, BlendableParameter->ParameterType, BlendableParameter->BlendableStructType);
-		}
-		else if (const UCameraObjectInterfaceDataParameter* DataParameter = LoadedCameraRig->Interface.FindDataParameterByName(InCameraParameterName))
-		{
-			Initialize(LoadedCameraRig, InCameraParameterName, DataParameter->DataType, DataParameter->DataContainerType, DataParameter->DataTypeObject);
-		}
-		// else, no parameter of that name found...
-	}
-}
-
-void UK2Node_SetCameraRigParameter::Initialize(UCameraRigAsset* InCameraRig, const FString& InCameraParameterName, ECameraVariableType InCameraVariableType, const UScriptStruct* InBlendableStructType)
-{
-	CameraRig = InCameraRig;
-	CameraParameterName = InCameraParameterName;
-	CameraParameterType = EK2Node_CameraParameterType::Blendable;
-	BlendableCameraParameterType = InCameraVariableType;
-	BlendableStructType = InBlendableStructType;
-}
-
-void UK2Node_SetCameraRigParameter::Initialize(UCameraRigAsset* InCameraRig, const FString& InCameraParameterName, ECameraContextDataType InCameraContextDataType, ECameraContextDataContainerType InCameraContextDataContainerType, const UObject* InCameraContextDataTypeObject)
-{
-	CameraRig = InCameraRig;
-	CameraParameterName = InCameraParameterName;
-	CameraParameterType = EK2Node_CameraParameterType::Data;
-	DataCameraParameterType = InCameraContextDataType;
-	DataCameraParameterContainerType = InCameraContextDataContainerType;
-	DataCameraParameterTypeObject = const_cast<UObject*>(InCameraContextDataTypeObject);
-}
-
 void UK2Node_SetCameraRigParameter::AllocateDefaultPins()
 {
 	Super::AllocateDefaultPins();
 
 	// Add parameter value pin.
-	FEdGraphPinType PinType;
-	switch (CameraParameterType)
-	{
-		case EK2Node_CameraParameterType::Blendable:
-			PinType = UK2Node_CameraRigBase::MakeBlendableParameterPinType(BlendableCameraParameterType, BlendableStructType);
-			break;
-		case EK2Node_CameraParameterType::Data:
-			PinType = UK2Node_CameraRigBase::MakeDataParameterPinType(DataCameraParameterType, DataCameraParameterContainerType, DataCameraParameterTypeObject);
-			break;
-		default:
-			PinType.PinCategory = UEdGraphSchema_K2::PC_Wildcard;
-			break;
-	}
+	FEdGraphPinType PinType = GetParameterPinType();
 	CreatePin(EGPD_Input, PinType, FName(CameraParameterName));
 }
 
@@ -96,25 +37,6 @@ FText UK2Node_SetCameraRigParameter::GetTooltipText() const
 			FText::FromString(CameraParameterName));
 }
 
-FSlateIcon UK2Node_SetCameraRigParameter::GetIconAndTint(FLinearColor& OutColor) const
-{
-	UEdGraphPin* ParameterValuePin = FindPin(CameraParameterName);
-	if (ParameterValuePin)
-	{
-		const UEdGraphSchema_K2* K2Schema = GetDefault<UEdGraphSchema_K2>();
-		OutColor = K2Schema->GetPinTypeColor(ParameterValuePin->PinType);
-
-#if UE_VERSION_NEWER_THAN_OR_EQUAL(5,6,0)
-		if (UStruct* Struct = Cast<UStruct>(ParameterValuePin->PinType.PinSubCategoryObject.Get()))
-		{
-			return FSlateIconFinder::FindIconForClass(Struct);
-		}
-#endif
-	}
-
-	return FSlateIcon(FAppStyle::GetAppStyleSetName(), "Kismet.AllClasses.VariableIcon");
-}
-
 void UK2Node_SetCameraRigParameter::GetMenuActions(FBlueprintActionDatabaseRegistrar& ActionRegistrar, const FAssetData& CameraRigAssetData) const
 {
 	const FText BaseCategoryString = GetMenuCategory();
diff --git a/Engine/Plugins/Cameras/GameplayCameras/Source/GameplayCamerasUncookedOnly/Private/BlueprintGraph/K2Node_SetCameraRigParameter.h b/Engine/Plugins/Cameras/GameplayCameras/Source/GameplayCamerasUncookedOnly/Private/BlueprintGraph/K2Node_SetCameraRigParameter.h
--- a/Engine/Plugins/Cameras/GameplayCameras/Source/GameplayCamerasUncookedOnly/Private/BlueprintGraph/K2Node_SetCameraRigParameter.h
+++ b/Engine/Plugins/Cameras/GameplayCameras/Source/GameplayCamerasUncookedOnly/Private/BlueprintGraph/K2Node_SetCameraRigParameter.h
@@ -2,73 +2,33 @@
 
 #pragma once
 
-#include "K2Node_CameraRigBase.h"
+#include "K2Node_SingleCameraRigParameterBase.h"
 
 #include "K2Node_SetCameraRigParameter.generated.h"
 
 class FKismetCompilerContext;
 class UCameraRigAsset;
 
-UENUM()
-enum class EK2Node_CameraParameterType : uint8
-{
-	Unknown,
-	Blendable,
-	Data
-};
-
 /**
  * Blueprint node that, given a camera rig, lets the user set the value of one single
  * exposed parameter.
  */
 UCLASS(MinimalAPI)
-class UK2Node_SetCameraRigParameter : public UK2Node_CameraRigBase
+class UK2Node_SetCameraRigParameter : public UK2Node_SingleCameraRigParameterBase
 {
 	GENERATED_BODY()
 
-public:
-
-	UK2Node_SetCameraRigParameter(const FObjectInitializer& ObjectInit);
-
-	void Initialize(const FAssetData& UnloadedCameraRig, const FString& InCameraParameterName);
-	void Initialize(UCameraRigAsset* InCameraRig, const FString& InCameraParameterName, ECameraVariableType InCameraVariableType, const UScriptStruct* InBlendableStructType);
-	void Initialize(UCameraRigAsset* InCameraRig, const FString& InCameraParameterName, ECameraContextDataType InCameraContextDataType, ECameraContextDataContainerType InCameraContextDataContainerType, const UObject* InCameraContextDataTypeObject);
-
 public:
 
 	// UEdGraphNode interface.
 	virtual void AllocateDefaultPins() override;
 	virtual FText GetNodeTitle(ENodeTitleType::Type TitleType) const override;
 	virtual FText GetTooltipText() const override;
-	virtual FSlateIcon GetIconAndTint(FLinearColor& OutColor) const override;
 	virtual void ExpandNode(FKismetCompilerContext& CompilerContext, UEdGraph* SourceGraph) override;
 
 protected:
 
 	// UK2Node_CameraRigBase interface.
 	void GetMenuActions(FBlueprintActionDatabaseRegistrar& ActionRegistrar, const FAssetData& CameraRigAssetData) const;
-
-private:
-
-	UPROPERTY()
-	FString CameraParameterName;
-
-	UPROPERTY()
-	EK2Node_CameraParameterType CameraParameterType = EK2Node_CameraParameterType::Unknown;
-
-	UPROPERTY()
-	ECameraVariableType BlendableCameraParameterType;
-
-	UPROPERTY()
-	TObjectPtr<const UScriptStruct> BlendableStructType;
-
-	UPROPERTY()
-	ECameraContextDataType DataCameraParameterType;
-
-	UPROPERTY()
-	ECameraContextDataContainerType DataCameraParameterContainerType;
-
-	UPROPERTY()
-	TObjectPtr<UObject> DataCameraParameterTypeObject;
 };
 
diff --git a/Engine/Plugins/Cameras/GameplayCameras/Source/GameplayCamerasUncookedOnly/Private/BlueprintGraph/K2Node_SetCameraRigParameters.cpp b/Engine/Plugins/Cameras/GameplayCameras/Source/GameplayCamerasUncookedOnly/Private/BlueprintGraph/K2Node_SetCameraRigParameters.cpp
--- a/Engine/Plugins/Cameras/GameplayCameras/Source/GameplayCamerasUncookedOnly/Private/BlueprintGraph/K2Node_SetCameraRigParameters.cpp
+++ b/Engine/Plugins/Cameras/GameplayCameras/Source/GameplayCamerasUncookedOnly/Private/BlueprintGraph/K2Node_SetCameraRigParameters.cpp
@@ -16,21 +16,11 @@
 
 #define LOCTEXT_NAMESPACE "K2Node_SetCameraRigParameters"
 
-UK2Node_SetCameraRigParameters::UK2Node_SetCameraRigParameters(const FObjectInitializer& ObjectInit)
-	: Super(ObjectInit)
-{
-}
-
-void UK2Node_SetCameraRigParameters::Initialize(const FAssetData& UnloadedCameraRig)
-{
-	CameraRig = Cast<UCameraRigAsset>(UnloadedCameraRig.GetAsset());
-}
-
 void UK2Node_SetCameraRigParameters::AllocateDefaultPins()
 {
 	Super::AllocateDefaultPins();
 
-	CreateParameterPins();
+	CreateParameterPins(EGPD_Input);
 }
 
 void UK2Node_SetCameraRigParameters::ReallocatePinsDuringReconstruction(TArray<UEdGraphPin*>& OldPins) 
@@ -49,31 +39,6 @@ void UK2Node_SetCameraRigParameters::ReallocatePinsDuringReconstruction(TArray<U
 		(*OldCameraRigPin)->DefaultValue.Reset();
 	}
 
-	// The camera rig might not be loaded yet when we are rebuilt on startup.
-	if (CameraRig)
-	{
-		PreloadObject(CameraRig);
-		for (UCameraObjectInterfaceBlendableParameter* BlendableParameter : CameraRig->Interface.BlendableParameters)
-		{
-			PreloadObject(BlendableParameter);
-			if (BlendableParameter)
-			{
-				PreloadObject(const_cast<UScriptStruct*>(BlendableParameter->BlendableStructType.Get()));
-			}
-		}
-		for (UCameraObjectInterfaceDataParameter* DataParameter : CameraRig->Interface.DataParameters)
-		{
-			PreloadObject(DataParameter);
-			if (DataParameter)
-			{
-				PreloadObject(const_cast<UObject*>(DataParameter->DataTypeObject.Get()));
-			}
-		}
-	}
-
-	// Create all default pins, including the parameter pins.
-	AllocateDefaultPins();
-
 	// Upgrade old result pin name to new pin name.
 	const FName OldResultPinName = TEXT("CameraEvaluationResult");
 	UEdGraphPin** OldResultPin = OldPins.FindByPredicate([&OldResultPinName](UEdGraphPin* OldPin)
@@ -87,7 +52,7 @@ void UK2Node_SetCameraRigParameters::ReallocatePinsDuringReconstruction(TArray<U
 		NewResultPin->MovePersistentDataFromOldPin(**OldResultPin);
 	}
 
-	RestoreSplitPins(OldPins);
+	Super::ReallocatePinsDuringReconstruction(OldPins);
 }
 
 FText UK2Node_SetCameraRigParameters::GetNodeTitle(ENodeTitleType::Type TitleType) const
@@ -299,88 +264,5 @@ void UK2Node_SetCameraRigParameters::ExpandNode(FKismetCompilerContext& Compiler
 	BreakAllNodeLinks();
 }
 
-void UK2Node_SetCameraRigParameters::CreateParameterPins()
-{
-	BlendableParameterPinNames.Reset();
-	DataParameterPinNames.Reset();
-
-	if (!CameraRig)
-	{
-		return;
-	}
-
-	for (const UCameraObjectInterfaceBlendableParameter* BlendableParameter : CameraRig->Interface.BlendableParameters)
-	{
-		if (!ensure(BlendableParameter))
-		{
-			continue;
-		}
-
-		if (!BlendableParameter->PrivateVariableID)
-		{
-			// Camera rig isn't fully built.
-			continue;
-		}
-		
-		FEdGraphPinType PinType = MakeBlendableParameterPinType(BlendableParameter);
-		if (PinType.PinCategory.IsNone())
-		{
-			// Unsupported type for Blueprints.
-			continue;
-		}
-
-		UEdGraphPin* NewPin = CreatePin(EGPD_Input, PinType, FName(BlendableParameter->InterfaceParameterName));
-		BlendableParameterPinNames.Add(NewPin->PinName);
-	}
-
-	for (const UCameraObjectInterfaceDataParameter* DataParameter : CameraRig->Interface.DataParameters)
-	{
-		if (!ensure(DataParameter))
-		{
-			continue;
-		}
-		
-		if (!DataParameter->PrivateDataID.IsValid())
-		{
-			// Camera rig isn't fully built.
-			continue;
-		}
-
-		FEdGraphPinType PinType = MakeDataParameterPinType(DataParameter);
-		if (PinType.PinCategory.IsNone())
-		{
-			// Unsupported type for Blueprints.
-			continue;
-		}
-
-		UEdGraphPin* NewPin = CreatePin(EGPD_Input, PinType, FName(DataParameter->InterfaceParameterName));
-		DataParameterPinNames.Add(NewPin->PinName);
-	}
-}
-
-void UK2Node_SetCameraRigParameters::FindBlendableParameterPins(TArray<UEdGraphPin*>& OutPins) const
-{
-	for (const FName& PinName : BlendableParameterPinNames)
-	{
-		UEdGraphPin* Pin = FindPin(PinName);
-		if (ensure(Pin))
-		{
-			OutPins.Add(Pin);
-		}
-	}
-}
-
-void UK2Node_SetCameraRigParameters::FindDataParameterPins(TArray<UEdGraphPin*>& OutPins) const
-{
-	for (const FName& PinName : DataParameterPinNames)
-	{
-		UEdGraphPin* Pin = FindPin(PinName);
-		if (ensure(Pin))
-		{
-			OutPins.Add(Pin);
-		}
-	}
-}
-
 #undef LOCTEXT_NAMESPACE
 
diff --git a/Engine/Plugins/Cameras/GameplayCameras/Source/GameplayCamerasUncookedOnly/Private/BlueprintGraph/K2Node_SetCameraRigParameters.h b/Engine/Plugins/Cameras/GameplayCameras/Source/GameplayCamerasUncookedOnly/Private/BlueprintGraph/K2Node_SetCameraRigParameters.h
--- a/Engine/Plugins/Cameras/GameplayCameras/Source/GameplayCamerasUncookedOnly/Private/BlueprintGraph/K2Node_SetCameraRigParameters.h
+++ b/Engine/Plugins/Cameras/GameplayCameras/Source/GameplayCamerasUncookedOnly/Private/BlueprintGraph/K2Node_SetCameraRigParameters.h
@@ -2,7 +2,7 @@
 
 #pragma once
 
-#include "K2Node_CameraRigBase.h"
+#include "K2Node_MultiCameraRigParametersBase.h"
 
 #include "K2Node_SetCameraRigParameters.generated.h"
 
@@ -14,16 +14,10 @@ class UCameraRigAsset;
  * values.
  */
 UCLASS(MinimalAPI)
-class UK2Node_SetCameraRigParameters : public UK2Node_CameraRigBase
+class UK2Node_SetCameraRigParameters : public UK2Node_MultiCameraRigParametersBase
 {
 	GENERATED_BODY()
 
-public:
-
-	UK2Node_SetCameraRigParameters(const FObjectInitializer& ObjectInit);
-
-	void Initialize(const FAssetData& UnloadedCameraRig);
-
 public:
 
 	// UEdGraphNode interface.
@@ -39,19 +33,5 @@ protected:
 
 	// UK2Node_CameraRigBase interface.
 	virtual void GetMenuActions(FBlueprintActionDatabaseRegistrar& ActionRegistrar, const FAssetData& CameraRigAssetData) const override;
-
-private:
-
-	void CreateParameterPins();
-	void FindBlendableParameterPins(TArray<UEdGraphPin*>& OutPins) const;
-	void FindDataParameterPins(TArray<UEdGraphPin*>& OutPins) const;
-
-private:
-
-	UPROPERTY()
-	TArray<FName> BlendableParameterPinNames;
-
-	UPROPERTY()
-	TArray<FName> DataParameterPinNames;
 };
 
diff --git a/Engine/Plugins/Cameras/GameplayCameras/Source/GameplayCamerasUncookedOnly/Private/BlueprintGraph/K2Node_SingleCameraRigParameterBase.cpp b/Engine/Plugins/Cameras/GameplayCameras/Source/GameplayCamerasUncookedOnly/Private/BlueprintGraph/K2Node_SingleCameraRigParameterBase.cpp
new file mode 100644
--- /dev/null
+++ b/Engine/Plugins/Cameras/GameplayCameras/Source/GameplayCamerasUncookedOnly/Private/BlueprintGraph/K2Node_SingleCameraRigParameterBase.cpp
@@ -0,0 +1,101 @@
+// Copyright Epic Games, Inc. All Rights Reserved.
+
+#include "BlueprintGraph/K2Node_SingleCameraRigParameterBase.h"
+
+#include "BlueprintActionDatabaseRegistrar.h"
+#include "BlueprintGraph/K2Node_CameraRigBase.h"
+#include "Core/CameraRigAsset.h"
+#include "EdGraphSchema_K2.h"
+#include "EditorCategoryUtils.h"
+#include "GameFramework/CameraRigParameterInterop.h"
+#include "K2Node_CallFunction.h"
+#include "Kismet/BlueprintInstancedStructLibrary.h"
+#include "KismetCompiler.h"
+#include "Misc/EngineVersionComparison.h"
+#include "Styling/AppStyle.h"
+#include "Styling/SlateIconFinder.h"
+
+#include UE_INLINE_GENERATED_CPP_BY_NAME(K2Node_SingleCameraRigParameterBase)
+
+#define LOCTEXT_NAMESPACE "K2Node_SingleCameraRigParameterBase"
+
+UK2Node_SingleCameraRigParameterBase::UK2Node_SingleCameraRigParameterBase(const FObjectInitializer& ObjectInit)
+	: Super(ObjectInit)
+{
+}
+
+void UK2Node_SingleCameraRigParameterBase::Initialize(const FAssetData& UnloadedCameraRig, const FString& InCameraParameterName)
+{
+	UCameraRigAsset* LoadedCameraRig = Cast<UCameraRigAsset>(UnloadedCameraRig.GetAsset());
+	if (ensure(LoadedCameraRig))
+	{
+		if (const UCameraObjectInterfaceBlendableParameter* BlendableParameter = LoadedCameraRig->Interface.FindBlendableParameterByName(InCameraParameterName))
+		{
+			Initialize(LoadedCameraRig, InCameraParameterName, BlendableParameter->ParameterType, BlendableParameter->BlendableStructType);
+		}
+		else if (const UCameraObjectInterfaceDataParameter* DataParameter = LoadedCameraRig->Interface.FindDataParameterByName(InCameraParameterName))
+		{
+			Initialize(LoadedCameraRig, InCameraParameterName, DataParameter->DataType, DataParameter->DataContainerType, DataParameter->DataTypeObject);
+		}
+		// else, no parameter of that name found...
+	}
+}
+
+void UK2Node_SingleCameraRigParameterBase::Initialize(UCameraRigAsset* InCameraRig, const FString& InCameraParameterName, ECameraVariableType InCameraVariableType, const UScriptStruct* InBlendableStructType)
+{
+	CameraRig = InCameraRig;
+	CameraParameterName = InCameraParameterName;
+	CameraParameterType = EK2Node_CameraParameterType::Blendable;
+	BlendableCameraParameterType = InCameraVariableType;
+	BlendableStructType = InBlendableStructType;
+}
+
+void UK2Node_SingleCameraRigParameterBase::Initialize(UCameraRigAsset* InCameraRig, const FString& InCameraParameterName, ECameraContextDataType InCameraContextDataType, ECameraContextDataContainerType InCameraContextDataContainerType, const UObject* InCameraContextDataTypeObject)
+{
+	CameraRig = InCameraRig;
+	CameraParameterName = InCameraParameterName;
+	CameraParameterType = EK2Node_CameraParameterType::Data;
+	DataCameraParameterType = InCameraContextDataType;
+	DataCameraParameterContainerType = InCameraContextDataContainerType;
+	DataCameraParameterTypeObject = const_cast<UObject*>(InCameraContextDataTypeObject);
+}
+
+FSlateIcon UK2Node_SingleCameraRigParameterBase::GetIconAndTint(FLinearColor& OutColor) const
+{
+	UEdGraphPin* ParameterValuePin = FindPin(CameraParameterName);
+	if (ParameterValuePin)
+	{
+		const UEdGraphSchema_K2* K2Schema = GetDefault<UEdGraphSchema_K2>();
+		OutColor = K2Schema->GetPinTypeColor(ParameterValuePin->PinType);
+
+#if UE_VERSION_NEWER_THAN_OR_EQUAL(5,6,0)
+		if (UStruct* Struct = Cast<UStruct>(ParameterValuePin->PinType.PinSubCategoryObject.Get()))
+		{
+			return FSlateIconFinder::FindIconForClass(Struct);
+		}
+#endif
+	}
+
+	return FSlateIcon(FAppStyle::GetAppStyleSetName(), "Kismet.AllClasses.VariableIcon");
+}
+
+FEdGraphPinType UK2Node_SingleCameraRigParameterBase::GetParameterPinType() const
+{
+	FEdGraphPinType PinType;
+	switch (CameraParameterType)
+	{
+		case EK2Node_CameraParameterType::Blendable:
+			PinType = UK2Node_CameraRigBase::MakeBlendableParameterPinType(BlendableCameraParameterType, BlendableStructType);
+			break;
+		case EK2Node_CameraParameterType::Data:
+			PinType = UK2Node_CameraRigBase::MakeDataParameterPinType(DataCameraParameterType, DataCameraParameterContainerType, DataCameraParameterTypeObject);
+			break;
+		default:
+			PinType.PinCategory = UEdGraphSchema_K2::PC_Wildcard;
+			break;
+	}
+	return PinType;
+}
+
+#undef LOCTEXT_NAMESPACE
+
diff --git a/Engine/Plugins/Cameras/GameplayCameras/Source/GameplayCamerasUncookedOnly/Private/BlueprintGraph/K2Node_SingleCameraRigParameterBase.h b/Engine/Plugins/Cameras/GameplayCameras/Source/GameplayCamerasUncookedOnly/Private/BlueprintGraph/K2Node_SingleCameraRigParameterBase.h
new file mode 100644
--- /dev/null
+++ b/Engine/Plugins/Cameras/GameplayCameras/Source/GameplayCamerasUncookedOnly/Private/BlueprintGraph/K2Node_SingleCameraRigParameterBase.h
@@ -0,0 +1,68 @@
+// Copyright Epic Games, Inc. All Rights Reserved.
+
+#pragma once
+
+#include "K2Node_CameraRigBase.h"
+
+#include "K2Node_SingleCameraRigParameterBase.generated.h"
+
+class FKismetCompilerContext;
+class UCameraRigAsset;
+
+UENUM()
+enum class EK2Node_CameraParameterType : uint8
+{
+	Unknown,
+	Blendable,
+	Data
+};
+
+/**
+ * A base class for a Blueprint node that gets or sets a single camera rig parameter.
+ */
+UCLASS(MinimalAPI)
+class UK2Node_SingleCameraRigParameterBase : public UK2Node_CameraRigBase
+{
+	GENERATED_BODY()
+
+public:
+
+	UK2Node_SingleCameraRigParameterBase(const FObjectInitializer& ObjectInit);
+
+	void Initialize(const FAssetData& UnloadedCameraRig, const FString& InCameraParameterName);
+	void Initialize(UCameraRigAsset* InCameraRig, const FString& InCameraParameterName, ECameraVariableType InCameraVariableType, const UScriptStruct* InBlendableStructType);
+	void Initialize(UCameraRigAsset* InCameraRig, const FString& InCameraParameterName, ECameraContextDataType InCameraContextDataType, ECameraContextDataContainerType InCameraContextDataContainerType, const UObject* InCameraContextDataTypeObject);
+
+public:
+
+	// UEdGraphNode interface.
+	virtual FSlateIcon GetIconAndTint(FLinearColor& OutColor) const override;
+
+protected:
+
+	FEdGraphPinType GetParameterPinType() const;
+
+protected:
+
+	UPROPERTY()
+	FString CameraParameterName;
+
+	UPROPERTY()
+	EK2Node_CameraParameterType CameraParameterType = EK2Node_CameraParameterType::Unknown;
+
+	UPROPERTY()
+	ECameraVariableType BlendableCameraParameterType;
+
+	UPROPERTY()
+	TObjectPtr<const UScriptStruct> BlendableStructType;
+
+	UPROPERTY()
+	ECameraContextDataType DataCameraParameterType;
+
+	UPROPERTY()
+	ECameraContextDataContainerType DataCameraParameterContainerType;
+
+	UPROPERTY()
+	TObjectPtr<UObject> DataCameraParameterTypeObject;
+};
+
