Cameras: fix possible crash in GPC camera manager when the last evaluation context is removed from the stack

diff --git a/Engine/Plugins/Cameras/GameplayCameras/Source/GameplayCameras/Private/GameFramework/GameplayCamerasPlayerCameraManager.cpp b/Engine/Plugins/Cameras/GameplayCameras/Source/GameplayCameras/Private/GameFramework/GameplayCamerasPlayerCameraManager.cpp
--- a/Engine/Plugins/Cameras/GameplayCameras/Source/GameplayCameras/Private/GameFramework/GameplayCamerasPlayerCameraManager.cpp
+++ b/Engine/Plugins/Cameras/GameplayCameras/Source/GameplayCameras/Private/GameFramework/GameplayCamerasPlayerCameraManager.cpp
@@ -329,18 +329,17 @@ void AGameplayCamerasPlayerCameraManager::OnContextStackChanged()
 	{
 		TGuardValue<bool> ReentrancyGuard(bIsSettingNewViewTarget, true);
 
+		AActor* NewViewTarget = nullptr;
 		FCameraEvaluationContextStack& ContextStack = CameraSystemEvaluator->GetEvaluationContextStack();
-		TSharedPtr<FCameraEvaluationContext> ActiveContext = ContextStack.GetActiveContext();
-		UObject* ActiveContextOwner = ActiveContext->GetOwner();
-		if (ActiveContextOwner)
+		if (TSharedPtr<FCameraEvaluationContext> ActiveContext = ContextStack.GetActiveContext())
 		{
-			ViewTarget.SetNewTarget(ActiveContextOwner->GetTypedOuter<AActor>());
-		}
-		else
-		{
-			ViewTarget.SetNewTarget(nullptr);
+			UObject* ActiveContextOwner = ActiveContext->GetOwner();
+			if (ActiveContextOwner)
+			{
+				NewViewTarget = ActiveContextOwner->GetTypedOuter<AActor>();
+			}
 		}
-
+		ViewTarget.SetNewTarget(NewViewTarget);
 		ViewTarget.CheckViewTarget(PCOwner);
 		BlendParams = FViewTargetTransitionParams();
 	}
