Cameras: fix main camera assets not being usable as Blueprint variables

diff --git a/Engine/Plugins/Cameras/GameplayCameras/Source/GameplayCameras/Public/Core/CameraAsset.h b/Engine/Plugins/Cameras/GameplayCameras/Source/GameplayCameras/Public/Core/CameraAsset.h
--- a/Engine/Plugins/Cameras/GameplayCameras/Source/GameplayCameras/Public/Core/CameraAsset.h
+++ b/Engine/Plugins/Cameras/GameplayCameras/Source/GameplayCameras/Public/Core/CameraAsset.h
@@ -65,7 +65,7 @@ struct FCameraAssetAllocationInfo
 /**
  * A complete camera asset.
  */
-UCLASS(MinimalAPI)
+UCLASS(MinimalAPI, BlueprintType)
 class UCameraAsset 
 	: public UObject
 	, public IHasCameraBuildStatus
diff --git a/Engine/Plugins/Cameras/GameplayCameras/Source/GameplayCameras/Public/Core/CameraRigAsset.h b/Engine/Plugins/Cameras/GameplayCameras/Source/GameplayCameras/Public/Core/CameraRigAsset.h
--- a/Engine/Plugins/Cameras/GameplayCameras/Source/GameplayCameras/Public/Core/CameraRigAsset.h
+++ b/Engine/Plugins/Cameras/GameplayCameras/Source/GameplayCameras/Public/Core/CameraRigAsset.h
@@ -53,7 +53,7 @@ using FCameraRigPackages = TArray<const UPackage*, TInlineAllocator<4>>;
  * A camera rig asset, which runs a hierarchy of camera nodes to drive 
  * the behavior of a camera.
  */
-UCLASS(MinimalAPI)
+UCLASS(MinimalAPI, BlueprintType)
 class UCameraRigAsset
 	: public UBaseCameraObject
 	, public IGameplayTagAssetInterface
diff --git a/Engine/Plugins/Cameras/GameplayCameras/Source/GameplayCameras/Public/Core/CameraRigProxyAsset.h b/Engine/Plugins/Cameras/GameplayCameras/Source/GameplayCameras/Public/Core/CameraRigProxyAsset.h
--- a/Engine/Plugins/Cameras/GameplayCameras/Source/GameplayCameras/Public/Core/CameraRigProxyAsset.h
+++ b/Engine/Plugins/Cameras/GameplayCameras/Source/GameplayCameras/Public/Core/CameraRigProxyAsset.h
@@ -14,7 +14,7 @@
  * but without hard-referencing a particular camera rig. This way, that camera director can
  * be used by multiple cameras with their own camera rigs.
  */
-UCLASS(MinimalAPI)
+UCLASS(MinimalAPI, BlueprintType)
 class UCameraRigProxyAsset : public UObject
 {
 	GENERATED_BODY()
diff --git a/Engine/Plugins/Cameras/GameplayCameras/Source/GameplayCameras/Public/Core/CameraShakeAsset.h b/Engine/Plugins/Cameras/GameplayCameras/Source/GameplayCameras/Public/Core/CameraShakeAsset.h
--- a/Engine/Plugins/Cameras/GameplayCameras/Source/GameplayCameras/Public/Core/CameraShakeAsset.h
+++ b/Engine/Plugins/Cameras/GameplayCameras/Source/GameplayCameras/Public/Core/CameraShakeAsset.h
@@ -19,7 +19,7 @@ namespace UE::Cameras
 	class FCameraShakeAssetBuilder;
 }
 
-UCLASS(MinimalAPI)
+UCLASS(MinimalAPI, BlueprintType)
 class UCameraShakeAsset
 	: public UBaseCameraObject
 	, public IHasCameraBuildStatus
