Cameras: fix incorrect compilation for data parameters getter, fix upgrade code for old setter node, remove unused/obsolete code

diff --git a/Engine/Plugins/Cameras/GameplayCameras/Source/GameplayCamerasUncookedOnly/Private/BlueprintGraph/K2Node_GetCameraRigParameter.cpp b/Engine/Plugins/Cameras/GameplayCameras/Source/GameplayCamerasUncookedOnly/Private/BlueprintGraph/K2Node_GetCameraRigParameter.cpp
--- a/Engine/Plugins/Cameras/GameplayCameras/Source/GameplayCamerasUncookedOnly/Private/BlueprintGraph/K2Node_GetCameraRigParameter.cpp
+++ b/Engine/Plugins/Cameras/GameplayCameras/Source/GameplayCamerasUncookedOnly/Private/BlueprintGraph/K2Node_GetCameraRigParameter.cpp
@@ -79,19 +79,6 @@ void UK2Node_GetCameraRigParameter::ExpandNode(FKismetCompilerContext& CompilerC
 		return;
 	}
 
-	// Additional setup is needed for some types.
-	UEnum* ExtraEnumTypePinObject = nullptr;
-	bool bMakeExtraInstanceStructSetup = false;
-	if (DataCameraParameterType == ECameraContextDataType::Enum)
-	{
-		const UEnum* EnumType = CastChecked<const UEnum>(DataCameraParameterTypeObject);
-		ExtraEnumTypePinObject = const_cast<UEnum*>(EnumType);
-	}
-	else if (DataCameraParameterType == ECameraContextDataType::Struct)
-	{
-		bMakeExtraInstanceStructSetup = true;
-	}
-
 	UEdGraphPin* const CameraNodeEvaluationResultPin = GetCameraNodeEvaluationResultPin();
 	UEdGraphPin* const CameraParameterValuePin = FindPinChecked(CameraParameterName);
 
diff --git a/Engine/Plugins/Cameras/GameplayCameras/Source/GameplayCamerasUncookedOnly/Private/BlueprintGraph/K2Node_GetCameraRigParameters.cpp b/Engine/Plugins/Cameras/GameplayCameras/Source/GameplayCamerasUncookedOnly/Private/BlueprintGraph/K2Node_GetCameraRigParameters.cpp
--- a/Engine/Plugins/Cameras/GameplayCameras/Source/GameplayCamerasUncookedOnly/Private/BlueprintGraph/K2Node_GetCameraRigParameters.cpp
+++ b/Engine/Plugins/Cameras/GameplayCameras/Source/GameplayCamerasUncookedOnly/Private/BlueprintGraph/K2Node_GetCameraRigParameters.cpp
@@ -152,7 +152,7 @@ void UK2Node_GetCameraRigParameters::ExpandNode(FKismetCompilerContext& Compiler
 		CallGetParameterNamePin->DefaultValue = DataParameter->InterfaceParameterName;
 
 		// Set or connect the parameter value argument.
-		UEdGraphPin* CallGetParameterValuePin = CallGetParameter->FindPinChecked(TEXT("NewValue"));
+		UEdGraphPin* CallGetParameterValuePin = CallGetParameter->FindPinChecked(TEXT("ReturnValue"));
 		CallGetParameterValuePin->PinType = RigParameterPin->PinType;
 		if (RigParameterPin->LinkedTo.Num() > 0)
 		{
diff --git a/Engine/Plugins/Cameras/GameplayCameras/Source/GameplayCamerasUncookedOnly/Private/BlueprintGraph/K2Node_MultiCameraRigParametersBase.cpp b/Engine/Plugins/Cameras/GameplayCameras/Source/GameplayCamerasUncookedOnly/Private/BlueprintGraph/K2Node_MultiCameraRigParametersBase.cpp
--- a/Engine/Plugins/Cameras/GameplayCameras/Source/GameplayCamerasUncookedOnly/Private/BlueprintGraph/K2Node_MultiCameraRigParametersBase.cpp
+++ b/Engine/Plugins/Cameras/GameplayCameras/Source/GameplayCamerasUncookedOnly/Private/BlueprintGraph/K2Node_MultiCameraRigParametersBase.cpp
@@ -30,6 +30,16 @@ void UK2Node_MultiCameraRigParametersBase::ReallocatePinsDuringReconstruction(TA
 {
 	// The camera rig might not be loaded yet when we are rebuilt on startup, so make sure
 	// this is all good otherwise we won't be able to get all our pins back.
+	EnsureCameraRigAssetLoaded();
+
+	// Now do as usual: create all default pins (including the parameter pins) and 
+	// restore split pins.
+	AllocateDefaultPins();
+	RestoreSplitPins(OldPins);
+}
+
+void UK2Node_MultiCameraRigParametersBase::EnsureCameraRigAssetLoaded()
+{
 	if (CameraRig)
 	{
 		PreloadObject(CameraRig);
@@ -50,11 +60,6 @@ void UK2Node_MultiCameraRigParametersBase::ReallocatePinsDuringReconstruction(TA
 			}
 		}
 	}
-
-	// Now do as usual: create all default pins (including the parameter pins) and 
-	// restore split pins.
-	AllocateDefaultPins();
-	RestoreSplitPins(OldPins);
 }
 
 void UK2Node_MultiCameraRigParametersBase::CreateParameterPins(EEdGraphPinDirection PinDirection)
diff --git a/Engine/Plugins/Cameras/GameplayCameras/Source/GameplayCamerasUncookedOnly/Private/BlueprintGraph/K2Node_MultiCameraRigParametersBase.h b/Engine/Plugins/Cameras/GameplayCameras/Source/GameplayCamerasUncookedOnly/Private/BlueprintGraph/K2Node_MultiCameraRigParametersBase.h
--- a/Engine/Plugins/Cameras/GameplayCameras/Source/GameplayCamerasUncookedOnly/Private/BlueprintGraph/K2Node_MultiCameraRigParametersBase.h
+++ b/Engine/Plugins/Cameras/GameplayCameras/Source/GameplayCamerasUncookedOnly/Private/BlueprintGraph/K2Node_MultiCameraRigParametersBase.h
@@ -29,6 +29,7 @@ public:
 
 protected:
 
+	void EnsureCameraRigAssetLoaded();
 	void CreateParameterPins(EEdGraphPinDirection PinDirection);
 	void FindBlendableParameterPins(TArray<UEdGraphPin*>& OutPins) const;
 	void FindDataParameterPins(TArray<UEdGraphPin*>& OutPins) const;
diff --git a/Engine/Plugins/Cameras/GameplayCameras/Source/GameplayCamerasUncookedOnly/Private/BlueprintGraph/K2Node_SetCameraRigParameter.cpp b/Engine/Plugins/Cameras/GameplayCameras/Source/GameplayCamerasUncookedOnly/Private/BlueprintGraph/K2Node_SetCameraRigParameter.cpp
--- a/Engine/Plugins/Cameras/GameplayCameras/Source/GameplayCamerasUncookedOnly/Private/BlueprintGraph/K2Node_SetCameraRigParameter.cpp
+++ b/Engine/Plugins/Cameras/GameplayCameras/Source/GameplayCamerasUncookedOnly/Private/BlueprintGraph/K2Node_SetCameraRigParameter.cpp
@@ -79,19 +79,6 @@ void UK2Node_SetCameraRigParameter::ExpandNode(FKismetCompilerContext& CompilerC
 		return;
 	}
 
-	// Additional setup is needed for some types.
-	UEnum* ExtraEnumTypePinObject = nullptr;
-	bool bMakeExtraInstanceStructSetup = false;
-	if (DataCameraParameterType == ECameraContextDataType::Enum)
-	{
-		const UEnum* EnumType = CastChecked<const UEnum>(DataCameraParameterTypeObject);
-		ExtraEnumTypePinObject = const_cast<UEnum*>(EnumType);
-	}
-	else if (DataCameraParameterType == ECameraContextDataType::Struct)
-	{
-		bMakeExtraInstanceStructSetup = true;
-	}
-
 	UEdGraphPin* const CameraNodeEvaluationResultPin = GetCameraNodeEvaluationResultPin();
 	UEdGraphPin* const CameraParameterValuePin = FindPinChecked(CameraParameterName);
 
@@ -137,13 +124,6 @@ void UK2Node_SetCameraRigParameter::ExpandNode(FKismetCompilerContext& CompilerC
 		CompilerContext.MessageLog.Error(*MissingParameterConnectionMsg.ToString(), this);
 	}
 
-	// Set extra type pin for enums.
-	if (ExtraEnumTypePinObject)
-	{
-		UEdGraphPin* CallSetParameterEnumTypePin = CallSetParameter->FindPinChecked(TEXT("EnumType"));
-		CallSetParameterEnumTypePin->DefaultObject = ExtraEnumTypePinObject;
-	}
-
 	// Setup the execution flow.
 	UEdGraphPin* ThisExecPin = GetExecPin();
 	CompilerContext.MovePinLinksToIntermediate(*ThisExecPin, *FirstExecPin);
diff --git a/Engine/Plugins/Cameras/GameplayCameras/Source/GameplayCamerasUncookedOnly/Private/BlueprintGraph/K2Node_SetCameraRigParameters.cpp b/Engine/Plugins/Cameras/GameplayCameras/Source/GameplayCamerasUncookedOnly/Private/BlueprintGraph/K2Node_SetCameraRigParameters.cpp
--- a/Engine/Plugins/Cameras/GameplayCameras/Source/GameplayCamerasUncookedOnly/Private/BlueprintGraph/K2Node_SetCameraRigParameters.cpp
+++ b/Engine/Plugins/Cameras/GameplayCameras/Source/GameplayCamerasUncookedOnly/Private/BlueprintGraph/K2Node_SetCameraRigParameters.cpp
@@ -39,6 +39,11 @@ void UK2Node_SetCameraRigParameters::ReallocatePinsDuringReconstruction(TArray<U
 		(*OldCameraRigPin)->DefaultValue.Reset();
 	}
 
+	// Do like the base class does, but fixup old-style pins before restoring split pins.
+	EnsureCameraRigAssetLoaded();
+
+	AllocateDefaultPins();
+
 	// Upgrade old result pin name to new pin name.
 	const FName OldResultPinName = TEXT("CameraEvaluationResult");
 	UEdGraphPin** OldResultPin = OldPins.FindByPredicate([&OldResultPinName](UEdGraphPin* OldPin)
@@ -52,7 +57,7 @@ void UK2Node_SetCameraRigParameters::ReallocatePinsDuringReconstruction(TArray<U
 		NewResultPin->MovePersistentDataFromOldPin(**OldResultPin);
 	}
 
-	Super::ReallocatePinsDuringReconstruction(OldPins);
+	RestoreSplitPins(OldPins);
 }
 
 FText UK2Node_SetCameraRigParameters::GetNodeTitle(ENodeTitleType::Type TitleType) const
