Cameras: use correct getter node for data parameters

#DisabledValidations Preflight

diff --git a/Engine/Plugins/Cameras/GameplayCameras/Source/GameplayCamerasUncookedOnly/Private/BlueprintGraph/K2Node_GetCameraRigParameters.cpp b/Engine/Plugins/Cameras/GameplayCameras/Source/GameplayCamerasUncookedOnly/Private/BlueprintGraph/K2Node_GetCameraRigParameters.cpp
--- a/Engine/Plugins/Cameras/GameplayCameras/Source/GameplayCamerasUncookedOnly/Private/BlueprintGraph/K2Node_GetCameraRigParameters.cpp
+++ b/Engine/Plugins/Cameras/GameplayCameras/Source/GameplayCamerasUncookedOnly/Private/BlueprintGraph/K2Node_GetCameraRigParameters.cpp
@@ -135,7 +135,7 @@ void UK2Node_GetCameraRigParameters::ExpandNode(FKismetCompilerContext& Compiler
 
 		// Make the SetXxxData function call node.
 		UK2Node_CallFunction* CallGetParameter = CompilerContext.SpawnIntermediateNode<UK2Node_CallFunction>(this, SourceGraph);
-		CallGetParameter->FunctionReference.SetExternalMember(TEXT("SetCameraParameter"), UCameraRigParameterInterop::StaticClass());
+		CallGetParameter->FunctionReference.SetExternalMember(TEXT("GetCameraParameter"), UCameraRigParameterInterop::StaticClass());
 		CallGetParameter->AllocateDefaultPins();
 		CompilerContext.MessageLog.NotifyIntermediateObjectCreation(CallGetParameter, SourceGraph);
 
