Cameras: fix handling of enum parameters in Blueprint setter nodes

diff --git a/Engine/Plugins/Cameras/GameplayCameras/Source/GameplayCamerasUncookedOnly/Private/BlueprintGraph/K2Node_CameraRigBase.cpp b/Engine/Plugins/Cameras/GameplayCameras/Source/GameplayCamerasUncookedOnly/Private/BlueprintGraph/K2Node_CameraRigBase.cpp
--- a/Engine/Plugins/Cameras/GameplayCameras/Source/GameplayCamerasUncookedOnly/Private/BlueprintGraph/K2Node_CameraRigBase.cpp
+++ b/Engine/Plugins/Cameras/GameplayCameras/Source/GameplayCamerasUncookedOnly/Private/BlueprintGraph/K2Node_CameraRigBase.cpp
@@ -216,7 +216,7 @@ FEdGraphPinType UK2Node_CameraRigBase::MakeDataParameterPinType(ECameraContextDa
 			PinCategory = UEdGraphSchema_K2::PC_String;
 			break;
 		case ECameraContextDataType::Enum:
-			PinCategory = UEdGraphSchema_K2::PC_Enum;
+			PinCategory = UEdGraphSchema_K2::PC_Byte;
 			break;
 		case ECameraContextDataType::Struct:
 			PinCategory = UEdGraphSchema_K2::PC_Struct;
@@ -261,7 +261,7 @@ UK2Node_CallFunction* UK2Node_CameraRigBase::CreateMakeLiteralNode(FKismetCompil
 	return MakeLiteralNode;
 }
 
-UK2Node_CallFunction* UK2Node_CameraRigBase::MakeLiteralValueForPin(FKismetCompilerContext& CompilerContext, UEdGraph* SourceGraph, UK2Node* SourceNode, UEdGraphPin* SourceValuePin)
+UK2Node* UK2Node_CameraRigBase::MakeLiteralValueForPin(FKismetCompilerContext& CompilerContext, UEdGraph* SourceGraph, UK2Node* SourceNode, UEdGraphPin* SourceValuePin)
 {
 #define UE_NEW_MAKE_LITERAL_NODE(FunctionLibraryClassName, MakeLiteralFunctionName)\
 	CreateMakeLiteralNode(\
@@ -304,10 +304,27 @@ UK2Node_CallFunction* UK2Node_CameraRigBase::MakeLiteralValueForPin(FKismetCompi
 		return UE_NEW_MAKE_LITERAL_NODE(UKismetSystemLibrary, MakeLiteralDouble);
 	}
 
-	if (SourceValuePin->PinType.PinCategory == UEdGraphSchema_K2::PC_Byte &&
-		SourceValuePin->PinType.PinSubCategoryObject == nullptr)  // Not an enum
+	if (SourceValuePin->PinType.PinCategory == UEdGraphSchema_K2::PC_Byte)
 	{
-		return UE_NEW_MAKE_LITERAL_NODE(UKismetSystemLibrary, MakeLiteralByte);
+		if (UEnum* EnumType = Cast<UEnum>(SourceValuePin->PinType.PinSubCategoryObject))
+		{
+			UK2Node_EnumLiteral* EnumLiteralNode = CompilerContext.SpawnIntermediateNode<UK2Node_EnumLiteral>(SourceNode, SourceGraph);
+			EnumLiteralNode->Enum = EnumType;
+			EnumLiteralNode->AllocateDefaultPins();
+			CompilerContext.MessageLog.NotifyIntermediateObjectCreation(EnumLiteralNode, SourceGraph);
+
+			UEdGraphPin* EnumLiteralValuePin = EnumLiteralNode->FindPinChecked(UK2Node_EnumLiteral::GetEnumInputPinName());
+			EnumLiteralValuePin->DefaultValue = SourceValuePin->DefaultValue;
+			EnumLiteralValuePin->DefaultTextValue = SourceValuePin->DefaultTextValue;
+			EnumLiteralValuePin->AutogeneratedDefaultValue = SourceValuePin->AutogeneratedDefaultValue;
+			EnumLiteralValuePin->DefaultObject = SourceValuePin->DefaultObject;
+
+			return EnumLiteralNode;
+		}
+		else
+		{
+			return UE_NEW_MAKE_LITERAL_NODE(UKismetSystemLibrary, MakeLiteralByte);
+		}
 	}
 
 	if (SourceValuePin->PinType.PinCategory == UEdGraphSchema_K2::PC_Struct)
diff --git a/Engine/Plugins/Cameras/GameplayCameras/Source/GameplayCamerasUncookedOnly/Private/BlueprintGraph/K2Node_CameraRigBase.h b/Engine/Plugins/Cameras/GameplayCameras/Source/GameplayCamerasUncookedOnly/Private/BlueprintGraph/K2Node_CameraRigBase.h
--- a/Engine/Plugins/Cameras/GameplayCameras/Source/GameplayCamerasUncookedOnly/Private/BlueprintGraph/K2Node_CameraRigBase.h
+++ b/Engine/Plugins/Cameras/GameplayCameras/Source/GameplayCamerasUncookedOnly/Private/BlueprintGraph/K2Node_CameraRigBase.h
@@ -68,7 +68,7 @@ protected:
 
 	static UK2Node_CallFunction* CreateMakeLiteralNode(FKismetCompilerContext& CompilerContext, UEdGraph* SourceGraph, UK2Node* SourceNode, UClass* FunctionLibraryClass, const TCHAR* FunctionName, UEdGraphPin* SourceValuePin);
 
-	static UK2Node_CallFunction* MakeLiteralValueForPin(FKismetCompilerContext& CompilerContext, UEdGraph* SourceGraph, UK2Node* SourceNode, UEdGraphPin* InValuePin);
+	static UK2Node* MakeLiteralValueForPin(FKismetCompilerContext& CompilerContext, UEdGraph* SourceGraph, UK2Node* SourceNode, UEdGraphPin* InValuePin);
 
 protected:
 
diff --git a/Engine/Plugins/Cameras/GameplayCameras/Source/GameplayCamerasUncookedOnly/Private/BlueprintGraph/K2Node_SetCameraRigParameter.cpp b/Engine/Plugins/Cameras/GameplayCameras/Source/GameplayCamerasUncookedOnly/Private/BlueprintGraph/K2Node_SetCameraRigParameter.cpp
--- a/Engine/Plugins/Cameras/GameplayCameras/Source/GameplayCamerasUncookedOnly/Private/BlueprintGraph/K2Node_SetCameraRigParameter.cpp
+++ b/Engine/Plugins/Cameras/GameplayCameras/Source/GameplayCamerasUncookedOnly/Private/BlueprintGraph/K2Node_SetCameraRigParameter.cpp
@@ -109,9 +109,10 @@ void UK2Node_SetCameraRigParameter::ExpandNode(FKismetCompilerContext& CompilerC
 	{
 		CompilerContext.MovePinLinksToIntermediate(*CameraParameterValuePin, *CallSetParameterValuePin);
 	}
-	else if (UK2Node_CallFunction* MakeLiteral = MakeLiteralValueForPin(CompilerContext, SourceGraph, this, CameraParameterValuePin))
+	else if (UK2Node* MakeLiteral = MakeLiteralValueForPin(CompilerContext, SourceGraph, this, CameraParameterValuePin))
 	{
-		MakeLiteral->GetReturnValuePin()->MakeLinkTo(CallSetParameterValuePin);
+		UEdGraphPin* ReturnValuePin = MakeLiteral->FindPinChecked(UEdGraphSchema_K2::PN_ReturnValue);
+		ReturnValuePin->MakeLinkTo(CallSetParameterValuePin);
 	}
 	else
 	{
diff --git a/Engine/Plugins/Cameras/GameplayCameras/Source/GameplayCamerasUncookedOnly/Private/BlueprintGraph/K2Node_SetCameraRigParameters.cpp b/Engine/Plugins/Cameras/GameplayCameras/Source/GameplayCamerasUncookedOnly/Private/BlueprintGraph/K2Node_SetCameraRigParameters.cpp
--- a/Engine/Plugins/Cameras/GameplayCameras/Source/GameplayCamerasUncookedOnly/Private/BlueprintGraph/K2Node_SetCameraRigParameters.cpp
+++ b/Engine/Plugins/Cameras/GameplayCameras/Source/GameplayCamerasUncookedOnly/Private/BlueprintGraph/K2Node_SetCameraRigParameters.cpp
@@ -156,9 +156,10 @@ void UK2Node_SetCameraRigParameters::ExpandNode(FKismetCompilerContext& Compiler
 		{
 			CompilerContext.MovePinLinksToIntermediate(*RigParameterPin, *CallSetParameterValuePin);
 		}
-		else if (UK2Node_CallFunction* MakeLiteral = MakeLiteralValueForPin(CompilerContext, SourceGraph, this, RigParameterPin))
+		else if (UK2Node* MakeLiteral = MakeLiteralValueForPin(CompilerContext, SourceGraph, this, RigParameterPin))
 		{
-			MakeLiteral->GetReturnValuePin()->MakeLinkTo(CallSetParameterValuePin);
+			UEdGraphPin* ReturnValuePin = MakeLiteral->FindPinChecked(UEdGraphSchema_K2::PN_ReturnValue);
+			ReturnValuePin->MakeLinkTo(CallSetParameterValuePin);
 		}
 		else
 		{
@@ -231,9 +232,10 @@ void UK2Node_SetCameraRigParameters::ExpandNode(FKismetCompilerContext& Compiler
 		{
 			CompilerContext.MovePinLinksToIntermediate(*RigParameterPin, *CallSetParameterValuePin);
 		}
-		else if (UK2Node_CallFunction* MakeLiteral = MakeLiteralValueForPin(CompilerContext, SourceGraph, this, RigParameterPin))
+		else if (UK2Node* MakeLiteral = MakeLiteralValueForPin(CompilerContext, SourceGraph, this, RigParameterPin))
 		{
-			MakeLiteral->GetReturnValuePin()->MakeLinkTo(CallSetParameterValuePin);
+			UEdGraphPin* ReturnValuePin = MakeLiteral->FindPinChecked(UEdGraphSchema_K2::PN_ReturnValue);
+			ReturnValuePin->MakeLinkTo(CallSetParameterValuePin);
 		}
 		else
 		{
