Cameras: show the name of the camera rig on the rig prefab node

diff --git a/Engine/Plugins/Cameras/GameplayCameras/Source/GameplayCameras/Private/Nodes/Common/CameraRigCameraNode.cpp b/Engine/Plugins/Cameras/GameplayCameras/Source/GameplayCameras/Private/Nodes/Common/CameraRigCameraNode.cpp
--- a/Engine/Plugins/Cameras/GameplayCameras/Source/GameplayCameras/Private/Nodes/Common/CameraRigCameraNode.cpp
+++ b/Engine/Plugins/Cameras/GameplayCameras/Source/GameplayCameras/Private/Nodes/Common/CameraRigCameraNode.cpp
@@ -165,7 +165,19 @@ void UCameraRigCameraNode::PostEditChangeProperty(struct FPropertyChangedEvent&
 	}
 }
 
-#endif
+EObjectTreeGraphObjectSupportFlags UCameraRigCameraNode::GetSupportFlags(FName InGraphName) const
+{
+	return (Super::GetSupportFlags(InGraphName) | EObjectTreeGraphObjectSupportFlags::CustomTitle);
+}
+
+void UCameraRigCameraNode::GetGraphNodeName(FName InGraphName, FText& OutName) const
+{
+	const UCameraRigAsset* CameraRig = CameraRigReference.GetCameraRig();
+	const FText CameraRigName = CameraRig ? FText::FromString(CameraRig->GetName()) : LOCTEXT("None", "None");
+	OutName = FText::Format(LOCTEXT("GraphNodeNameFormat", "Camera Rig ({0})"), CameraRigName);
+}
+
+#endif  // WITH_EDITOR
 
 #undef LOCTEXT_NAMESPACE
 
diff --git a/Engine/Plugins/Cameras/GameplayCameras/Source/GameplayCameras/Public/Nodes/Common/CameraRigCameraNode.h b/Engine/Plugins/Cameras/GameplayCameras/Source/GameplayCameras/Public/Nodes/Common/CameraRigCameraNode.h
--- a/Engine/Plugins/Cameras/GameplayCameras/Source/GameplayCameras/Public/Nodes/Common/CameraRigCameraNode.h
+++ b/Engine/Plugins/Cameras/GameplayCameras/Source/GameplayCameras/Public/Nodes/Common/CameraRigCameraNode.h
@@ -34,6 +34,12 @@ protected:
 	// ICustomCameraNodeParameterProvider interface.
 	virtual void GetCustomCameraNodeParameters(FCustomCameraNodeParameterInfos& OutParameterInfos) override;
 
+	// IObjectTreeGraphObject interface.
+#if WITH_EDITOR
+	virtual EObjectTreeGraphObjectSupportFlags GetSupportFlags(FName InGraphName) const override;
+	virtual void GetGraphNodeName(FName InGraphName, FText& OutName) const override;
+#endif  // WITH_EDITOR
+
 	// UObject interface.
 #if WITH_EDITOR
 	virtual void PostEditChangeProperty( struct FPropertyChangedEvent& PropertyChangedEvent) override;
