Cameras: fix editor preview to not add null camera rigs to the evaluation

diff --git a/Engine/Plugins/Cameras/GameplayCameras/Source/GameplayCameras/Private/Core/CameraEvaluationContext.cpp b/Engine/Plugins/Cameras/GameplayCameras/Source/GameplayCameras/Private/Core/CameraEvaluationContext.cpp
--- a/Engine/Plugins/Cameras/GameplayCameras/Source/GameplayCameras/Private/Core/CameraEvaluationContext.cpp
+++ b/Engine/Plugins/Cameras/GameplayCameras/Source/GameplayCameras/Private/Core/CameraEvaluationContext.cpp
@@ -41,7 +41,7 @@ protected:
 			PreviewIndex = (CameraRigs.IsEmpty() ? INDEX_NONE : 0);
 		}
 
-		if (PreviewIndex != INDEX_NONE)
+		if (PreviewIndex != INDEX_NONE && CameraRigs[PreviewIndex] != nullptr)
 		{
 			OutResult.Add(GetEvaluationContext(), CameraRigs[PreviewIndex]);
 		}
