Cameras: don't assert on unbound data parameters, just like blendable parameters

diff --git a/Engine/Plugins/Cameras/GameplayCameras/Source/GameplayCameras/Private/Build/CameraObjectInterfaceParameterBuilder.cpp b/Engine/Plugins/Cameras/GameplayCameras/Source/GameplayCameras/Private/Build/CameraObjectInterfaceParameterBuilder.cpp
--- a/Engine/Plugins/Cameras/GameplayCameras/Source/GameplayCameras/Private/Build/CameraObjectInterfaceParameterBuilder.cpp
+++ b/Engine/Plugins/Cameras/GameplayCameras/Source/GameplayCameras/Private/Build/CameraObjectInterfaceParameterBuilder.cpp
@@ -319,12 +319,16 @@ void FCameraObjectInterfaceParameterBuilder::SetDefaultParameterValues(const UBa
 
 	for (const UCameraObjectInterfaceDataParameter* DataParameter : CameraObject->Interface.DataParameters)
 	{
-		if (!ensure(DataParameter && DataParameter->Target))
+		if (!ensure(DataParameter))
 		{
 			continue;
 		}
 
 		UCameraNode* CameraNode = DataParameter->Target;
+		if (!CameraNode)
+		{
+			continue;
+		}
 
 		const uint8* RawSourceValuePtr = nullptr;
 
