Cameras: validate that exposed parameters have the correct type for what they're connected to

diff --git a/Engine/Plugins/Cameras/GameplayCameras/Source/GameplayCameras/Private/Build/CameraObjectInterfaceBuilder.cpp b/Engine/Plugins/Cameras/GameplayCameras/Source/GameplayCameras/Private/Build/CameraObjectInterfaceBuilder.cpp
--- a/Engine/Plugins/Cameras/GameplayCameras/Source/GameplayCameras/Private/Build/CameraObjectInterfaceBuilder.cpp
+++ b/Engine/Plugins/Cameras/GameplayCameras/Source/GameplayCameras/Private/Build/CameraObjectInterfaceBuilder.cpp
@@ -19,6 +19,42 @@ namespace UE::Cameras
 namespace Internal
 {
 
+template<typename CameraParameterOrVariableReference>
+bool MatchesVariableType(ECameraVariableType VariableType)
+{
+	return false;
+}
+
+template<typename CameraParameterOrVariableReference>
+FText GetVariableTypeAsText()
+{
+	return FText();
+}
+
+#define UE_CAMERA_VARIABLE_FOR_TYPE(ValueType, ValueName)\
+	template<>\
+	bool MatchesVariableType<F##ValueName##CameraVariableReference>(ECameraVariableType VariableType)\
+	{\
+		return VariableType == ECameraVariableType::ValueName;\
+	}\
+	template<>\
+	bool MatchesVariableType<F##ValueName##CameraParameter>(ECameraVariableType VariableType)\
+	{\
+		return VariableType == ECameraVariableType::ValueName;\
+	}\
+	template<>\
+	FText GetVariableTypeAsText<F##ValueName##CameraVariableReference>()\
+	{\
+		return UEnum::GetDisplayValueAsText(ECameraVariableType::ValueName);\
+	}\
+	template<>\
+	FText GetVariableTypeAsText<F##ValueName##CameraParameter>()\
+	{\
+		return UEnum::GetDisplayValueAsText(ECameraVariableType::ValueName);\
+	}
+UE_CAMERA_VARIABLE_FOR_ALL_TYPES()
+#undef UE_CAMERA_VARIABLE_FOR_TYPE
+
 struct FInterfaceParameterBindingBuilder
 {
 	UBaseCameraObject* CameraObject;
@@ -100,7 +136,7 @@ struct FInterfaceParameterBindingBuilder
 		ensure(BlendableParameter->PrivateVariableID.IsValid());
 		ensure(BlendableParameter->TargetPropertyName == CustomParameter.ParameterName);
 
-		const bool bIsValid = CheckIfParameterCanBeOverridden(BlendableParameter, CustomParameter.OverrideVariableID);
+		const bool bIsValid = CheckIfParameterCanBeOverridden(BlendableParameter, CustomParameter);
 		if (!bIsValid)
 		{
 			return;
@@ -145,7 +181,7 @@ struct FInterfaceParameterBindingBuilder
 	{
 		ensure(DataParameter->PrivateDataID);
 
-		const bool bIsValid = CheckIfParameterCanBeOverridden(DataParameter, OverrideDataID);
+		const bool bIsValid = CheckIfParameterCanBeOverridden(DataParameter, TargetProperty, OverrideDataID);
 		if (!bIsValid)
 		{
 			return;
@@ -168,7 +204,7 @@ struct FInterfaceParameterBindingBuilder
 		ensure(DataParameter->PrivateDataID);
 		ensure(DataParameter->TargetPropertyName == CustomParameter.ParameterName);
 
-		const bool bIsValid = CheckIfParameterCanBeOverridden(DataParameter, CustomParameter.OverrideDataID);
+		const bool bIsValid = CheckIfParameterCanBeOverridden(DataParameter, CustomParameter);
 		if (!bIsValid)
 		{
 			return;
@@ -211,6 +247,20 @@ private:
 			const UCameraObjectInterfaceBlendableParameter* BlendableParameter, 
 			CameraParameterOrVariableReferenceType* CameraParameterOrVariableReference)
 	{
+		if (!MatchesVariableType<CameraParameterOrVariableReferenceType>(BlendableParameter->ParameterType))
+		{
+			ReportError(BlendableParameter->Target,
+					FText::Format(
+						LOCTEXT(
+							"BlendableParameterTypeMismatch",
+							"Camera node parameter '{0}.{1}' of type {2} is not compatible with camera rig parameter '{3}' of type {4}"),
+						FText::FromName(BlendableParameter->Target->GetFName()), 
+						FText::FromName(BlendableParameter->TargetPropertyName),
+						GetVariableTypeAsText<CameraParameterOrVariableReferenceType>(),
+						FText::FromString(BlendableParameter->InterfaceParameterName),
+						UEnum::GetDisplayValueAsText(BlendableParameter->ParameterType)));
+			return false;
+		}
 		if (CameraParameterOrVariableReference->Variable != nullptr)
 		{
 			ReportError(BlendableParameter->Target,
@@ -230,6 +280,20 @@ private:
 			const UCameraObjectInterfaceBlendableParameter* BlendableParameter, 
 			const FCustomCameraNodeParameterInfos::FBlendableParameterInfo& CustomParameter)
 	{
+		if (BlendableParameter->ParameterType != CustomParameter.ParameterType)
+		{
+			ReportError(BlendableParameter->Target,
+					FText::Format(
+						LOCTEXT(
+							"BlendableParameterTypeMismatch",
+							"Camera node parameter '{0}.{1}' of type {2} is not compatible with camera rig parameter '{3}' of type {4}"),
+						FText::FromName(BlendableParameter->Target->GetFName()), 
+						FText::FromName(CustomParameter.ParameterName),
+						UEnum::GetDisplayValueAsText(CustomParameter.ParameterType),
+						FText::FromString(BlendableParameter->InterfaceParameterName),
+						UEnum::GetDisplayValueAsText(BlendableParameter->ParameterType)));
+			return false;
+		}
 		if (CustomParameter.OverrideVariable != nullptr)
 		{
 			ReportError(BlendableParameter->Target,
@@ -274,6 +338,88 @@ private:
 		return true;
 	}
 
+	bool CheckIfParameterCanBeOverridden(
+			const UCameraObjectInterfaceDataParameter* DataParameter,
+			const FCustomCameraNodeParameterInfos::FDataParameterInfo& CustomParameter)
+	{
+		if (DataParameter->DataType != CustomParameter.ParameterType)
+		{
+			ReportError(DataParameter->Target,
+					FText::Format(
+						LOCTEXT(
+							"DataParameterTypeMismatch",
+							"Camera node parameter '{0}.{1}' of type {2} is not compatible with camera rig parameter '{3}' of type {4}"),
+						FText::FromName(DataParameter->Target->GetFName()), 
+						FText::FromName(CustomParameter.ParameterName),
+						UEnum::GetDisplayValueAsText(CustomParameter.ParameterType),
+						FText::FromString(DataParameter->InterfaceParameterName),
+						UEnum::GetDisplayValueAsText(DataParameter->DataType)));
+			return false;
+		}
+
+		return CheckIfParameterCanBeOverridden(DataParameter, CustomParameter.OverrideDataID);
+	}
+
+	bool CheckIfParameterCanBeOverridden(
+			const UCameraObjectInterfaceDataParameter* DataParameter,
+			const FProperty* TargetProperty,
+			FCameraContextDataID* DataID)
+	{
+		bool bDataTypeMatches = false;
+		switch (DataParameter->DataType)
+		{
+			case ECameraContextDataType::Name:
+				bDataTypeMatches = (TargetProperty->IsA<FNameProperty>());
+				break;
+			case ECameraContextDataType::String:
+				bDataTypeMatches = (TargetProperty->IsA<FStrProperty>());
+				break;
+			case ECameraContextDataType::Enum:
+				if (const FEnumProperty* EnumProperty = CastField<FEnumProperty>(TargetProperty))
+				{
+					bDataTypeMatches = (EnumProperty->GetEnum() == DataParameter->DataTypeObject);
+				}
+				break;
+			case ECameraContextDataType::Struct:
+				if (const FStructProperty* StructProperty = CastField<FStructProperty>(TargetProperty))
+				{
+					bDataTypeMatches = (StructProperty->Struct == DataParameter->DataTypeObject);
+				}
+				break;
+			case ECameraContextDataType::Class:
+				if (const FClassProperty* ClassProperty = CastField<FClassProperty>(TargetProperty))
+				{
+					bDataTypeMatches = (ClassProperty->PropertyClass == DataParameter->DataTypeObject);
+				}
+				break;
+			case ECameraContextDataType::Object:
+				if (const FObjectProperty* ObjectProperty = CastField<FObjectProperty>(TargetProperty))
+				{
+					bDataTypeMatches = (ObjectProperty->PropertyClass == DataParameter->DataTypeObject);
+				}
+				break;
+			default:
+				ensure(false);
+				return false;
+		}
+		if (!bDataTypeMatches)
+		{
+			ReportError(DataParameter->Target,
+					FText::Format(
+						LOCTEXT(
+							"DataParameterTypeMismatch",
+							"Camera node parameter '{0}.{1}' of type {2} is not compatible with camera rig parameter '{3}' of type {4}"),
+						FText::FromName(DataParameter->Target->GetFName()), 
+						FText::FromName(TargetProperty->GetFName()),
+						TargetProperty->GetClass()->GetDisplayNameText(),
+						FText::FromString(DataParameter->InterfaceParameterName),
+						UEnum::GetDisplayValueAsText(DataParameter->DataType)));
+			return false;
+		}
+
+		return CheckIfParameterCanBeOverridden(DataParameter, DataID);
+	}
+
 	bool CheckIfParameterCanBeOverridden(
 			const UCameraObjectInterfaceDataParameter* DataParameter,
 			FCameraContextDataID* DataID)
@@ -601,6 +747,7 @@ void FCameraObjectInterfaceBuilder::BuildInterfaceParameterBindings()
 		}
 		if (BlendableParameter->InterfaceParameterName.IsEmpty())
 		{
+			// Error already logged in BuildInterfaceParameters().
 			continue;
 		}
 
