Cameras: fix some serialization issues

The type traits for some structures were incorrectly setup to inherit from the traits of the base type. This led to a couple of internal assets saved incorrectly.
This change fixes the problem and adds a workaround to prevent "data corrupted" errors in the editor for these few cases.

diff --git a/Engine/Plugins/Cameras/GameplayCameras/Source/GameplayCameras/Private/Core/BaseCameraObjectReference.cpp b/Engine/Plugins/Cameras/GameplayCameras/Source/GameplayCameras/Private/Core/BaseCameraObjectReference.cpp
--- a/Engine/Plugins/Cameras/GameplayCameras/Source/GameplayCameras/Private/Core/BaseCameraObjectReference.cpp
+++ b/Engine/Plugins/Cameras/GameplayCameras/Source/GameplayCameras/Private/Core/BaseCameraObjectReference.cpp
@@ -264,15 +264,18 @@ void FBaseCameraObjectReference::GetCustomCameraNodeParameters(FCustomCameraNode
 
 void FBaseCameraObjectReference::PostSerialize(const FArchive& Ar)
 {
-	PRAGMA_DISABLE_DEPRECATION_WARNINGS
-	for (FCameraObjectInterfaceParameterMetaData& MetaData : ParameterMetaData)
+	if (Ar.IsLoading())
 	{
-		if (MetaData.bIsOverridden_DEPRECATED)
+		PRAGMA_DISABLE_DEPRECATION_WARNINGS
+		for (FCameraObjectInterfaceParameterMetaData& MetaData : ParameterMetaData)
 		{
-			Parameters.SetPropertyOverriden(MetaData.ParameterGuid, true);
-			MetaData.bIsOverridden_DEPRECATED = false;
+			if (MetaData.bIsOverridden_DEPRECATED)
+			{
+				Parameters.SetPropertyOverriden(MetaData.ParameterGuid, true);
+				MetaData.bIsOverridden_DEPRECATED = false;
+			}
 		}
+		PRAGMA_ENABLE_DEPRECATION_WARNINGS
 	}
-	PRAGMA_ENABLE_DEPRECATION_WARNINGS
 }
 
diff --git a/Engine/Plugins/Cameras/GameplayCameras/Source/GameplayCameras/Private/Core/CameraAssetReference.cpp b/Engine/Plugins/Cameras/GameplayCameras/Source/GameplayCameras/Private/Core/CameraAssetReference.cpp
--- a/Engine/Plugins/Cameras/GameplayCameras/Source/GameplayCameras/Private/Core/CameraAssetReference.cpp
+++ b/Engine/Plugins/Cameras/GameplayCameras/Source/GameplayCameras/Private/Core/CameraAssetReference.cpp
@@ -103,12 +103,15 @@ bool FCameraAssetReference::SerializeFromMismatchedTag(FPropertyTag const& Tag,
 
 void FCameraAssetReference::PostSerialize(const FArchive& Ar)
 {
-	PRAGMA_DISABLE_DEPRECATION_WARNINGS
-	for (const FGuid& Guid : ParameterOverrideGuids_DEPRECATED)
+	if (Ar.IsLoading())
 	{
-		Parameters.SetPropertyOverriden(Guid, true);
+		PRAGMA_DISABLE_DEPRECATION_WARNINGS
+		for (const FGuid& Guid : ParameterOverrideGuids_DEPRECATED)
+		{
+			Parameters.SetPropertyOverriden(Guid, true);
+		}
+		ParameterOverrideGuids_DEPRECATED.Reset();
+		PRAGMA_ENABLE_DEPRECATION_WARNINGS
 	}
-	ParameterOverrideGuids_DEPRECATED.Reset();
-	PRAGMA_ENABLE_DEPRECATION_WARNINGS
 }
 
diff --git a/Engine/Plugins/Cameras/GameplayCameras/Source/GameplayCameras/Private/Core/CameraRigAssetReference.cpp b/Engine/Plugins/Cameras/GameplayCameras/Source/GameplayCameras/Private/Core/CameraRigAssetReference.cpp
--- a/Engine/Plugins/Cameras/GameplayCameras/Source/GameplayCameras/Private/Core/CameraRigAssetReference.cpp
+++ b/Engine/Plugins/Cameras/GameplayCameras/Source/GameplayCameras/Private/Core/CameraRigAssetReference.cpp
@@ -78,63 +78,66 @@ bool FCameraRigAssetReference::SerializeFromMismatchedTag(struct FPropertyTag co
 
 void FCameraRigAssetReference::PostSerialize(const FArchive& Ar)
 {
-	PRAGMA_DISABLE_DEPRECATION_WARNINGS
-
-	// Make a property bag with the legacy overrides, and then set the values in it.
-	bool bHasAnyLegacyOverride = false;
-	TArray<FPropertyBagPropertyDesc> LegacyParameterProperties;
-	TArray<FCameraObjectInterfaceParameterMetaData> LegacyParameterMetaData;
-
-#define UE_CAMERA_VARIABLE_FOR_TYPE(ValueType, ValueName)\
-	for (F##ValueName##CameraRigParameterOverride& ParameterOverride : ParameterOverrides_DEPRECATED.ValueName##Overrides)\
-	{\
-		FName PropertyName(ParameterOverride.InterfaceParameterName);\
-		EPropertyBagPropertyType PropertyType = EPropertyBagPropertyType::Struct;\
-		const UObject* PropertyTypeObject = F##ValueName##CameraParameter::StaticStruct();\
-		FPropertyBagPropertyDesc LegacyParameterProperty(PropertyName, PropertyType, PropertyTypeObject);\
-		LegacyParameterProperty.ID = ParameterOverride.InterfaceParameterGuid;\
-		LegacyParameterProperties.Add(LegacyParameterProperty);\
-		FCameraObjectInterfaceParameterMetaData MetaData;\
-		MetaData.ParameterGuid = ParameterOverride.InterfaceParameterGuid;\
-		MetaData.bIsOverridden_DEPRECATED = true;\
-		LegacyParameterMetaData.Add(MetaData);\
-		bHasAnyLegacyOverride = true;\
-	}
-	UE_CAMERA_VARIABLE_FOR_ALL_TYPES()
-#undef UE_CAMERA_VARIABLE_FOR_TYPE
-
-	if (bHasAnyLegacyOverride)
+	if (Ar.IsLoading())
 	{
-		Parameters = FInstancedOverridablePropertyBag();
-		Parameters.AddProperties(LegacyParameterProperties);
+		PRAGMA_DISABLE_DEPRECATION_WARNINGS
 
-		ParameterMetaData = LegacyParameterMetaData;
+		// Make a property bag with the legacy overrides, and then set the values in it.
+		bool bHasAnyLegacyOverride = false;
+		TArray<FPropertyBagPropertyDesc> LegacyParameterProperties;
+		TArray<FCameraObjectInterfaceParameterMetaData> LegacyParameterMetaData;
 
 #define UE_CAMERA_VARIABLE_FOR_TYPE(ValueType, ValueName)\
 		for (F##ValueName##CameraRigParameterOverride& ParameterOverride : ParameterOverrides_DEPRECATED.ValueName##Overrides)\
 		{\
 			FName PropertyName(ParameterOverride.InterfaceParameterName);\
-			Parameters.SetValueStruct<F##ValueName##CameraParameter>(PropertyName, ParameterOverride.Value);\
+			EPropertyBagPropertyType PropertyType = EPropertyBagPropertyType::Struct;\
+			const UObject* PropertyTypeObject = F##ValueName##CameraParameter::StaticStruct();\
+			FPropertyBagPropertyDesc LegacyParameterProperty(PropertyName, PropertyType, PropertyTypeObject);\
+			LegacyParameterProperty.ID = ParameterOverride.InterfaceParameterGuid;\
+			LegacyParameterProperties.Add(LegacyParameterProperty);\
+			FCameraObjectInterfaceParameterMetaData MetaData;\
+			MetaData.ParameterGuid = ParameterOverride.InterfaceParameterGuid;\
+			MetaData.bIsOverridden_DEPRECATED = true;\
+			LegacyParameterMetaData.Add(MetaData);\
+			bHasAnyLegacyOverride = true;\
 		}
 		UE_CAMERA_VARIABLE_FOR_ALL_TYPES()
 #undef UE_CAMERA_VARIABLE_FOR_TYPE
 
-		ParameterOverrides_DEPRECATED = FCameraRigParameterOverrides();
-	}
+		if (bHasAnyLegacyOverride)
+		{
+			Parameters = FInstancedOverridablePropertyBag();
+			Parameters.AddProperties(LegacyParameterProperties);
 
-	PRAGMA_ENABLE_DEPRECATION_WARNINGS
+			ParameterMetaData = LegacyParameterMetaData;
 
-	if (!ParameterOverrideGuids_DEPRECATED.IsEmpty())
-	{
-		for (const FGuid& Guid : ParameterOverrideGuids_DEPRECATED)
-		{
-			FCameraObjectInterfaceParameterMetaData MetaData;
-			MetaData.ParameterGuid = Guid;
-			MetaData.bIsOverridden_DEPRECATED = true;
-			ParameterMetaData.Add(MetaData);
+#define UE_CAMERA_VARIABLE_FOR_TYPE(ValueType, ValueName)\
+			for (F##ValueName##CameraRigParameterOverride& ParameterOverride : ParameterOverrides_DEPRECATED.ValueName##Overrides)\
+			{\
+				FName PropertyName(ParameterOverride.InterfaceParameterName);\
+				Parameters.SetValueStruct<F##ValueName##CameraParameter>(PropertyName, ParameterOverride.Value);\
+			}
+			UE_CAMERA_VARIABLE_FOR_ALL_TYPES()
+#undef UE_CAMERA_VARIABLE_FOR_TYPE
+
+			ParameterOverrides_DEPRECATED = FCameraRigParameterOverrides();
 		}
 
-		ParameterOverrideGuids_DEPRECATED.Reset();
+		PRAGMA_ENABLE_DEPRECATION_WARNINGS
+
+		if (!ParameterOverrideGuids_DEPRECATED.IsEmpty())
+		{
+			for (const FGuid& Guid : ParameterOverrideGuids_DEPRECATED)
+			{
+				FCameraObjectInterfaceParameterMetaData MetaData;
+				MetaData.ParameterGuid = Guid;
+				MetaData.bIsOverridden_DEPRECATED = true;
+				ParameterMetaData.Add(MetaData);
+			}
+
+			ParameterOverrideGuids_DEPRECATED.Reset();
+		}
 	}
 
 	FBaseCameraObjectReference::PostSerialize(Ar);
diff --git a/Engine/Plugins/Cameras/GameplayCameras/Source/GameplayCameras/Private/StructUtils/OverridablePropertyBag.cpp b/Engine/Plugins/Cameras/GameplayCameras/Source/GameplayCameras/Private/StructUtils/OverridablePropertyBag.cpp
--- a/Engine/Plugins/Cameras/GameplayCameras/Source/GameplayCameras/Private/StructUtils/OverridablePropertyBag.cpp
+++ b/Engine/Plugins/Cameras/GameplayCameras/Source/GameplayCameras/Private/StructUtils/OverridablePropertyBag.cpp
@@ -4,6 +4,25 @@
 
 #include UE_INLINE_GENERATED_CPP_BY_NAME(OverridablePropertyBag)
 
+namespace FOverridablePropertyBagCustomVersion
+{
+	enum Type
+	{
+		BeforeCustomVersionWasAdded = 0,
+		FixSerializer = 1,
+
+		VersionPlusOne,
+		LatestVersion = VersionPlusOne - 1
+	};
+
+	const FGuid GUID(0x5426C227, 0x4B3145B2, 0x9B9BED1F, 0x327FB126);
+}
+
+FCustomVersionRegistration GOverridablePropertyBagCustomVersion(
+		FOverridablePropertyBagCustomVersion::GUID,
+		FOverridablePropertyBagCustomVersion::LatestVersion,
+		TEXT("OverridablePropertyBagCustomVersion"));
+
 bool FInstancedOverridablePropertyBag::IsPropertyOverriden(const FGuid& InPropertyID) const
 {
 	return OverridenPropertyIDs.Contains(InPropertyID);
@@ -45,9 +64,29 @@ bool FInstancedOverridablePropertyBag::SerializeFromMismatchedTag(FPropertyTag c
 	{
 		// Read the structured data as an FInstancedPropertyBag. Our list of overriden property IDs
 		// will stay empty.
-		Serialize(Slot.GetUnderlyingArchive());
+		FInstancedPropertyBag::Serialize(Slot.GetUnderlyingArchive());
 		return true;
 	}
 	return false;
 }
 
+bool FInstancedOverridablePropertyBag::Serialize(FArchive& Ar)
+{
+	Ar.UsingCustomVersion(FOverridablePropertyBagCustomVersion::GUID);
+
+	if (Ar.IsLoading() && Ar.CustomVer(FOverridablePropertyBagCustomVersion::GUID) < FOverridablePropertyBagCustomVersion::FixSerializer)
+	{
+		// There was a short time during which developer data was saved with only default tagged
+		// serialization. This causes data corruption so we avoid that with this use-case here.
+		static const FInstancedOverridablePropertyBag ThisDefaults;
+
+		const UScriptStruct* ThisStruct = FInstancedOverridablePropertyBag::StaticStruct();
+		ThisStruct->SerializeTaggedProperties(Ar, (uint8*)this, ThisStruct, (const uint8*)&ThisDefaults);
+		return true;
+	}
+
+	FInstancedPropertyBag::Serialize(Ar);
+	Ar << OverridenPropertyIDs;
+	return true;
+}
+
diff --git a/Engine/Plugins/Cameras/GameplayCameras/Source/GameplayCameras/Public/Core/CameraRigAssetReference.h b/Engine/Plugins/Cameras/GameplayCameras/Source/GameplayCameras/Public/Core/CameraRigAssetReference.h
--- a/Engine/Plugins/Cameras/GameplayCameras/Source/GameplayCameras/Public/Core/CameraRigAssetReference.h
+++ b/Engine/Plugins/Cameras/GameplayCameras/Source/GameplayCameras/Public/Core/CameraRigAssetReference.h
@@ -325,7 +325,7 @@ private:
 };
 
 template<>
-struct TStructOpsTypeTraits<FCameraRigAssetReference> : public TStructOpsTypeTraitsBase2<FBaseCameraObjectReference>
+struct TStructOpsTypeTraits<FCameraRigAssetReference> : public TStructOpsTypeTraits<FBaseCameraObjectReference>
 {
 	enum
 	{
diff --git a/Engine/Plugins/Cameras/GameplayCameras/Source/GameplayCameras/Public/Core/CameraShakeAssetReference.h b/Engine/Plugins/Cameras/GameplayCameras/Source/GameplayCameras/Public/Core/CameraShakeAssetReference.h
--- a/Engine/Plugins/Cameras/GameplayCameras/Source/GameplayCameras/Public/Core/CameraShakeAssetReference.h
+++ b/Engine/Plugins/Cameras/GameplayCameras/Source/GameplayCameras/Public/Core/CameraShakeAssetReference.h
@@ -88,3 +88,8 @@ private:
 	friend class UE::Cameras::FCameraShakeAssetReferenceDetailsCustomization;
 };
 
+template<>
+struct TStructOpsTypeTraits<FCameraShakeAssetReference> : public TStructOpsTypeTraits<FBaseCameraObjectReference>
+{
+};
+
diff --git a/Engine/Plugins/Cameras/GameplayCameras/Source/GameplayCameras/Public/StructUtils/OverridablePropertyBag.h b/Engine/Plugins/Cameras/GameplayCameras/Source/GameplayCameras/Public/StructUtils/OverridablePropertyBag.h
--- a/Engine/Plugins/Cameras/GameplayCameras/Source/GameplayCameras/Public/StructUtils/OverridablePropertyBag.h
+++ b/Engine/Plugins/Cameras/GameplayCameras/Source/GameplayCameras/Public/StructUtils/OverridablePropertyBag.h
@@ -53,6 +53,7 @@ public:
 	// Internal API.
 
 	bool SerializeFromMismatchedTag(FPropertyTag const& Tag, FStructuredArchive::FSlot Slot);
+	bool Serialize(FArchive& Ar);
 
 private:
 
@@ -73,11 +74,12 @@ void FInstancedOverridablePropertyBag::GetOverridenPropertyIDs(ContainerType& Ou
 }
 
 template<>
-struct TStructOpsTypeTraits<FInstancedOverridablePropertyBag> : public TStructOpsTypeTraitsBase2<FInstancedOverridablePropertyBag>
+struct TStructOpsTypeTraits<FInstancedOverridablePropertyBag> : public TStructOpsTypeTraits<FInstancedPropertyBag>
 {
 	enum
 	{
-		WithStructuredSerializeFromMismatchedTag = true
+		WithStructuredSerializeFromMismatchedTag = true,
+		WithSerializer = true
 	};
 };
 
