Cameras: fixes for activating persistent rigs during director initialization, and having these rigs receive custom parameter overrides

- Add a result structure to the Activate/Deactivate methods of the camera director API so that they can activate/deactivate persistent camera rigs on setup/teardown.
- Handle these setup/teardown request, and refactor a bit of code about proxy resolution to reduce duplication.
- Make the additive/persistent blend-stack never override its entries' initial camera pose with the entry's evaluation context.
- The previous point makes it possible to remove the need for a "null context" to run non-main camera rigs (it was previously needed as a context with no camera pose overrides, but this introduced the problem that setting parameter overrides on the initial result would not affect persistent camera rigs)

diff --git a/Engine/Plugins/Cameras/GameplayCameras/Source/GameplayCameras/Private/Core/CameraDirectorEvaluator.cpp b/Engine/Plugins/Cameras/GameplayCameras/Source/GameplayCameras/Private/Core/CameraDirectorEvaluator.cpp
--- a/Engine/Plugins/Cameras/GameplayCameras/Source/GameplayCameras/Private/Core/CameraDirectorEvaluator.cpp
+++ b/Engine/Plugins/Cameras/GameplayCameras/Source/GameplayCameras/Private/Core/CameraDirectorEvaluator.cpp
@@ -10,6 +10,24 @@
 namespace UE::Cameras
 {
 
+bool FCameraRigActivationDeactivationRequest::ResolveCameraRigProxyIfNeeded(const FCameraDirectorEvaluator* InDirectorEvaluator)
+{
+	return ResolveCameraRigProxyIfNeeded(InDirectorEvaluator->GetCameraDirector());
+}
+
+bool FCameraRigActivationDeactivationRequest::ResolveCameraRigProxyIfNeeded(const UCameraDirector* InDirector)
+{
+	if (CameraRig == nullptr && CameraRigProxy != nullptr && ensure(InDirector))
+	{
+		const FCameraRigProxyRedirectTable& ProxyTable = InDirector->CameraRigProxyRedirectTable;
+
+		FCameraRigProxyResolveParams ResolveParams;
+		ResolveParams.CameraRigProxy = CameraRigProxy;
+		CameraRig = ProxyTable.ResolveProxy(ResolveParams);
+	}
+	return CameraRig != nullptr;
+}
+
 void FCameraDirectorEvaluatorStorage::DestroyEvaluator()
 {
 	Evaluator.Reset();
@@ -33,16 +51,16 @@ void FCameraDirectorEvaluator::Initialize(const FCameraDirectorInitializeParams&
 	OnInitialize(Params);
 }
 
-void FCameraDirectorEvaluator::Activate(const FCameraDirectorActivateParams& Params)
+void FCameraDirectorEvaluator::Activate(const FCameraDirectorActivateParams& Params, FCameraDirectorEvaluationResult& OutResult)
 {
 	Evaluator = Params.Evaluator;
 
-	OnActivate(Params);
+	OnActivate(Params, OutResult);
 }
 
-void FCameraDirectorEvaluator::Deactivate(const FCameraDirectorDeactivateParams& Params)
+void FCameraDirectorEvaluator::Deactivate(const FCameraDirectorDeactivateParams& Params, FCameraDirectorEvaluationResult& OutResult)
 {
-	OnDeactivate(Params);
+	OnDeactivate(Params, OutResult);
 
 	Evaluator = nullptr;
 }
diff --git a/Engine/Plugins/Cameras/GameplayCameras/Source/GameplayCameras/Private/Core/CameraEvaluationContext.cpp b/Engine/Plugins/Cameras/GameplayCameras/Source/GameplayCameras/Private/Core/CameraEvaluationContext.cpp
--- a/Engine/Plugins/Cameras/GameplayCameras/Source/GameplayCameras/Private/Core/CameraEvaluationContext.cpp
+++ b/Engine/Plugins/Cameras/GameplayCameras/Source/GameplayCameras/Private/Core/CameraEvaluationContext.cpp
@@ -7,6 +7,7 @@
 #include "Core/CameraDirector.h"
 #include "Core/CameraDirectorEvaluator.h"
 #include "Core/CameraSystemEvaluator.h"
+#include "Core/RootCameraNode.h"
 #include "GameFramework/PlayerController.h"
 
 #include UE_INLINE_GENERATED_CPP_BY_NAME(CameraEvaluationContext)
@@ -280,15 +281,19 @@ void FCameraEvaluationContext::Activate(const FCameraEvaluationContextActivatePa
 
 	AutoCreateDirectorEvaluator();
 
+	CameraSystemEvaluator = Params.Evaluator;
+	bActivated = true;
+
 	if (DirectorEvaluator)
 	{
 		FCameraDirectorActivateParams DirectorParams;
 		DirectorParams.Evaluator = Params.Evaluator;
-		DirectorEvaluator->Activate(DirectorParams);
-	}
+		FCameraDirectorEvaluationResult DirectorResult;
+		DirectorEvaluator->Activate(DirectorParams, DirectorResult);
 
-	CameraSystemEvaluator = Params.Evaluator;
-	bActivated = true;
+		// Execute initial requests from the director.
+		ExecuteSetupAndTeardownRequests(DirectorResult);
+	}
 }
 
 void FCameraEvaluationContext::Deactivate(const FCameraEvaluationContextDeactivateParams& Params)
@@ -301,7 +306,11 @@ void FCameraEvaluationContext::Deactivate(const FCameraEvaluationContextDeactiva
 	if (DirectorEvaluator)
 	{
 		FCameraDirectorDeactivateParams DirectorParams;
-		DirectorEvaluator->Deactivate(DirectorParams);
+		FCameraDirectorEvaluationResult DirectorResult;
+		DirectorEvaluator->Deactivate(DirectorParams, DirectorResult);
+
+		// Execute last requests from the director.
+		ExecuteSetupAndTeardownRequests(DirectorResult);
 	}
 
 	// Don't destroy the camera director evaluator, it could still be useful. We only destroy it
@@ -313,6 +322,38 @@ void FCameraEvaluationContext::Deactivate(const FCameraEvaluationContextDeactiva
 	bActivated = false;
 }
 
+void FCameraEvaluationContext::ExecuteSetupAndTeardownRequests(FCameraDirectorEvaluationResult& DirectorResult)
+{
+	if (!ensure(CameraSystemEvaluator))
+	{
+		return;
+	}
+
+	FRootCameraNodeEvaluator* RootEvaluator = CameraSystemEvaluator->GetRootNodeEvaluator();
+
+	for (FCameraRigActivationDeactivationRequest& Request : DirectorResult.Requests)
+	{
+		if (Request.Layer == ECameraRigLayer::Main)
+		{
+			// Executing main-layer requests would require being able to handle dynamically combined camera rigs,
+			// which is possible but requries refactoring that away from FCameraSystemEvaluator, so leave that
+			// for later, especially since it's arguably wrong to try to activate main rigs in Activate/Deactivate.
+			UObject* RequestedCameraObject = Request.CameraRig ? (UObject*)Request.CameraRig : (UObject*)Request.CameraRigProxy;
+			UE_LOG(LogCameraSystem, Error, 
+					TEXT("Main layer camera rigs can only be activated/deactivated during the director's normal update. "
+						"Ignoring request to activate/deactivate '{0}'."),
+					*GetNameSafe(RequestedCameraObject));
+			continue;
+		}
+
+		Request.ResolveCameraRigProxyIfNeeded(DirectorEvaluator);
+		if (Request.IsValid())
+		{
+			RootEvaluator->ExecuteCameraDirectorRequest(Request);
+		}
+	}
+}
+
 TSharedPtr<FCameraSystemEvaluator> FCameraEvaluationContext::GetCameraSystemEvaluator() const
 {
 	if (CameraSystemEvaluator)
diff --git a/Engine/Plugins/Cameras/GameplayCameras/Source/GameplayCameras/Private/Core/CameraSystemEvaluator.cpp b/Engine/Plugins/Cameras/GameplayCameras/Source/GameplayCameras/Private/Core/CameraSystemEvaluator.cpp
--- a/Engine/Plugins/Cameras/GameplayCameras/Source/GameplayCameras/Private/Core/CameraSystemEvaluator.cpp
+++ b/Engine/Plugins/Cameras/GameplayCameras/Source/GameplayCameras/Private/Core/CameraSystemEvaluator.cpp
@@ -304,17 +304,8 @@ void FCameraSystemEvaluator::UpdateCameraDirector(float DeltaTime, FCameraDirect
 
 	for (FCameraRigActivationDeactivationRequest& Request : DirectorResult.Requests)
 	{
-		if (!ensure(Request.EvaluationContext))
-		{
-			continue;
-		}
-
-		// Resolve camera rig proxies if needed.
-		if (Request.CameraRig == nullptr && Request.CameraRigProxy)
-		{
-			Request.CameraRig = CameraDirectorEvaluator->FindCameraRigByProxy(Request.CameraRigProxy);
-		}
-		if (!ensure(Request.CameraRig))
+		Request.ResolveCameraRigProxyIfNeeded(CameraDirectorEvaluator);
+		if (!Request.IsValid())
 		{
 			continue;
 		}
diff --git a/Engine/Plugins/Cameras/GameplayCameras/Source/GameplayCameras/Private/Core/PersistentBlendStackCameraNode.cpp b/Engine/Plugins/Cameras/GameplayCameras/Source/GameplayCameras/Private/Core/PersistentBlendStackCameraNode.cpp
--- a/Engine/Plugins/Cameras/GameplayCameras/Source/GameplayCameras/Private/Core/PersistentBlendStackCameraNode.cpp
+++ b/Engine/Plugins/Cameras/GameplayCameras/Source/GameplayCameras/Private/Core/PersistentBlendStackCameraNode.cpp
@@ -249,8 +249,10 @@ void FPersistentBlendStackCameraNodeEvaluator::InternalUpdate(TArrayView<FResolv
 				CurResult.PostProcessSettings.OverrideAll(OutResult.PostProcessSettings);
 
 				// Override it with whatever the evaluation context has set on its result.
+				// NOTE: don't override the camera pose, because additive/persistent blend stacks are designed to be
+				//       used for "modifier rigs". As such they shouldn't reset the camera pose to wherever their 
+				//       context is located. They should only modify it.
 				const FCameraNodeEvaluationResult& ContextResult(Entry.ContextResult);
-				CurResult.CameraPose.OverrideChanged(ContextResult.CameraPose);
 				CurResult.VariableTable.OverrideAll(ContextResult.VariableTable, true);
 				CurResult.ContextDataTable.OverrideAll(ContextResult.ContextDataTable);
 
diff --git a/Engine/Plugins/Cameras/GameplayCameras/Source/GameplayCameras/Private/Directors/BlueprintCameraDirector.cpp b/Engine/Plugins/Cameras/GameplayCameras/Source/GameplayCameras/Private/Directors/BlueprintCameraDirector.cpp
--- a/Engine/Plugins/Cameras/GameplayCameras/Source/GameplayCameras/Private/Directors/BlueprintCameraDirector.cpp
+++ b/Engine/Plugins/Cameras/GameplayCameras/Source/GameplayCameras/Private/Directors/BlueprintCameraDirector.cpp
@@ -37,8 +37,8 @@ public:
 protected:
 
 	virtual void OnInitialize(const FCameraDirectorInitializeParams& Params) override;
-	virtual void OnActivate(const FCameraDirectorActivateParams& Params) override;
-	virtual void OnDeactivate(const FCameraDirectorDeactivateParams& Params) override;
+	virtual void OnActivate(const FCameraDirectorActivateParams& Params, FCameraDirectorEvaluationResult& OutResult) override;
+	virtual void OnDeactivate(const FCameraDirectorDeactivateParams& Params, FCameraDirectorEvaluationResult& OutResult) override;
 	virtual void OnRun(const FCameraDirectorEvaluationParams& Params, FCameraDirectorEvaluationResult& OutResult) override;
 	virtual void OnAddChildEvaluationContext(const FChildContextManulationParams& Params, FChildContextManulationResult& Result) override;
 	virtual void OnRemoveChildEvaluationContext(const FChildContextManulationParams& Params, FChildContextManulationResult& Result) override;
@@ -87,13 +87,16 @@ void FBlueprintCameraDirectorEvaluator::OnInitialize(const FCameraDirectorInitia
 	}
 }
 
-void FBlueprintCameraDirectorEvaluator::OnActivate(const FCameraDirectorActivateParams& Params)
+void FBlueprintCameraDirectorEvaluator::OnActivate(const FCameraDirectorActivateParams& Params, FCameraDirectorEvaluationResult& OutResult)
 {
 	OwningEvaluator = Params.Evaluator;
 
 	if (EvaluatorBlueprint)
 	{
 		EvaluatorBlueprint->NativeActivateCameraDirector(Params);
+
+		const FCameraDirectorEvaluationResult& BlueprintResult = EvaluatorBlueprint->GetEvaluationResult();
+		OutResult = BlueprintResult;
 	}
 	else
 	{
@@ -101,7 +104,7 @@ void FBlueprintCameraDirectorEvaluator::OnActivate(const FCameraDirectorActivate
 	}
 }
 
-void FBlueprintCameraDirectorEvaluator::OnDeactivate(const FCameraDirectorDeactivateParams& Params)
+void FBlueprintCameraDirectorEvaluator::OnDeactivate(const FCameraDirectorDeactivateParams& Params, FCameraDirectorEvaluationResult& OutResult)
 {
 	// We need to check a few more things here in case we're being deactivated while the owner object is getting GC'ed.
 	UObject* ContextOwner = GetEvaluationContext()->GetOwner();
@@ -109,6 +112,9 @@ void FBlueprintCameraDirectorEvaluator::OnDeactivate(const FCameraDirectorDeacti
 	if (bIsValid)
 	{
 		EvaluatorBlueprint->NativeDeactivateCameraDirector(Params);
+
+		const FCameraDirectorEvaluationResult& BlueprintResult = EvaluatorBlueprint->GetEvaluationResult();
+		OutResult = BlueprintResult;
 	}
 
 	OwningEvaluator = nullptr;
diff --git a/Engine/Plugins/Cameras/GameplayCameras/Source/GameplayCameras/Private/Directors/StateTreeCameraDirector.cpp b/Engine/Plugins/Cameras/GameplayCameras/Source/GameplayCameras/Private/Directors/StateTreeCameraDirector.cpp
--- a/Engine/Plugins/Cameras/GameplayCameras/Source/GameplayCameras/Private/Directors/StateTreeCameraDirector.cpp
+++ b/Engine/Plugins/Cameras/GameplayCameras/Source/GameplayCameras/Private/Directors/StateTreeCameraDirector.cpp
@@ -31,8 +31,8 @@ public:
 protected:
 
 	// FCameraDirectorEvaluator interface.
-	virtual void OnActivate(const FCameraDirectorActivateParams& Params) override;
-	virtual void OnDeactivate(const FCameraDirectorDeactivateParams& Params) override;
+	virtual void OnActivate(const FCameraDirectorActivateParams& Params, FCameraDirectorEvaluationResult& OutResult) override;
+	virtual void OnDeactivate(const FCameraDirectorDeactivateParams& Params, FCameraDirectorEvaluationResult& OutResult) override;
 	virtual void OnRun(const FCameraDirectorEvaluationParams& Params, FCameraDirectorEvaluationResult& OutResult) override;
 	virtual void OnAddReferencedObjects(FReferenceCollector& Collector) override;
 
@@ -49,7 +49,7 @@ private:
 
 UE_DEFINE_CAMERA_DIRECTOR_EVALUATOR(FStateTreeCameraDirectorEvaluator)
 
-void FStateTreeCameraDirectorEvaluator::OnActivate(const FCameraDirectorActivateParams& Params)
+void FStateTreeCameraDirectorEvaluator::OnActivate(const FCameraDirectorActivateParams& Params, FCameraDirectorEvaluationResult& OutResult)
 {
 	const UStateTreeCameraDirector* StateTreeDirector = GetCameraDirectorAs<UStateTreeCameraDirector>();
 	const FStateTreeReference& StateTreeReference = StateTreeDirector->StateTreeReference;
@@ -96,7 +96,7 @@ void FStateTreeCameraDirectorEvaluator::OnActivate(const FCameraDirectorActivate
 	StateTreeContext.Start(&StateTreeReference.GetParameters());
 }
 
-void FStateTreeCameraDirectorEvaluator::OnDeactivate(const FCameraDirectorDeactivateParams& Params)
+void FStateTreeCameraDirectorEvaluator::OnDeactivate(const FCameraDirectorDeactivateParams& Params, FCameraDirectorEvaluationResult& OutResult)
 {
 	const UStateTreeCameraDirector* StateTreeDirector = GetCameraDirectorAs<UStateTreeCameraDirector>();
 	const FStateTreeReference& StateTreeReference = StateTreeDirector->StateTreeReference;
diff --git a/Engine/Plugins/Cameras/GameplayCameras/Source/GameplayCameras/Private/GameFramework/GameplayCameraComponentBase.cpp b/Engine/Plugins/Cameras/GameplayCameras/Source/GameplayCameras/Private/GameFramework/GameplayCameraComponentBase.cpp
--- a/Engine/Plugins/Cameras/GameplayCameras/Source/GameplayCameras/Private/GameFramework/GameplayCameraComponentBase.cpp
+++ b/Engine/Plugins/Cameras/GameplayCameras/Source/GameplayCameras/Private/GameFramework/GameplayCameraComponentBase.cpp
@@ -72,10 +72,6 @@ void UGameplayCameraComponentBase::AddReferencedObjects(UObject* InThis, FRefere
 	{
 		This->EvaluationContext->AddReferencedObjects(Collector);
 	}
-	if (This->NullContext.IsValid())
-	{
-		This->NullContext->AddReferencedObjects(Collector);
-	}
 }
 
 TSharedPtr<const UE::Cameras::FCameraEvaluationContext> UGameplayCameraComponentBase::GetEvaluationContext() const
@@ -387,23 +383,12 @@ void UGameplayCameraComponentBase::CreateCameraEvaluationContext(APlayerControll
 		UpdateCameraEvaluationContext(true);
 	}
 
-	if (!NullContext.IsValid())
-	{
-		NullContext = MakeShared<FCameraEvaluationContext>();
-
-		FCameraEvaluationContextInitializeParams InitParams;
-		InitParams.Owner = this;
-		InitParams.PlayerController = PlayerController;
-		NullContext->Initialize(InitParams);
-	}
-
-	ensure(EvaluationContext.IsValid() && NullContext.IsValid());
+	ensure(EvaluationContext.IsValid());
 }
 
 void UGameplayCameraComponentBase::DestroyCameraEvaluationContext()
 {
 	EvaluationContext = nullptr;
-	NullContext = nullptr;
 }
 
 #define UE_PRIVATE_GAMEPLAY_CAMERA_COMPONENT_VALIDATE_EVALUATION_CONTEXT(ErrorMsg, ErrorResult)\
@@ -435,9 +420,9 @@ FBlueprintCameraEvaluationDataRef UGameplayCameraComponentBase::GetConditionalRe
 #undef UE_PRIVATE_GAMEPLAY_CAMERA_COMPONENT_VALIDATE_EVALUATION_CONTEXT
 
 #define UE_PRIVATE_GAMEPLAY_CAMERA_COMPONENT_ACTIVATE_NON_MAIN_RIG(ErrorMsg, RigLayer)\
-	if (NullContext.IsValid())\
+	if (EvaluationContext.IsValid())\
 	{\
-		IGameplayCameraSystemHost::ActivateCameraRig(CameraRig, NullContext, ECameraRigLayer::RigLayer);\
+		IGameplayCameraSystemHost::ActivateCameraRig(CameraRig, EvaluationContext, ECameraRigLayer::RigLayer);\
 	}\
 	else\
 	{\
@@ -845,7 +830,7 @@ void UGameplayCameraComponentBase::AutoManageEditorPreviewEvaluator()
 	}
 	
 	const bool bCanRun = CanRunCameraSystem();
-	if (bCanRun && !(CameraSystemEvaluator && EvaluationContext && NullContext))
+	if (bCanRun && !(CameraSystemEvaluator && EvaluationContext))
 	{
 		// We want to run the camera logic in the editor but we haven't set things up for that.
 		// Let's create the preview evaluator and the evaluation context.
@@ -859,7 +844,7 @@ void UGameplayCameraComponentBase::AutoManageEditorPreviewEvaluator()
 
 		// OutputCameraComponent will be updated on the next tick.
 	}
-	else if (!bCanRun && (CameraSystemEvaluator || EvaluationContext || NullContext))
+	else if (!bCanRun && (CameraSystemEvaluator || EvaluationContext))
 	{
 		// We don't want to run the camera logic in the editor anymore. Let's tear things down.
 		DeactivateCameraEvaluationContext(true);
diff --git a/Engine/Plugins/Cameras/GameplayCameras/Source/GameplayCameras/Public/Core/CameraDirectorEvaluator.h b/Engine/Plugins/Cameras/GameplayCameras/Source/GameplayCameras/Public/Core/CameraDirectorEvaluator.h
--- a/Engine/Plugins/Cameras/GameplayCameras/Source/GameplayCameras/Public/Core/CameraDirectorEvaluator.h
+++ b/Engine/Plugins/Cameras/GameplayCameras/Source/GameplayCameras/Public/Core/CameraDirectorEvaluator.h
@@ -104,6 +104,16 @@ public:
 		: EvaluationContext(InContext)
 		, CameraRigProxy(InCameraRigProxy)
 	{}
+
+	bool IsValid() const
+	{
+		return EvaluationContext != nullptr && CameraRig != nullptr;
+	}
+
+public:
+
+	bool ResolveCameraRigProxyIfNeeded(const FCameraDirectorEvaluator* InDirectorEvaluator);
+	bool ResolveCameraRigProxyIfNeeded(const UCameraDirector* InDirector);
 };
 
 /**
@@ -218,10 +228,10 @@ public:
 	GAMEPLAYCAMERAS_API void Initialize(const FCameraDirectorInitializeParams& Params);
 
 	/** Activates the camera director evaluator. */
-	GAMEPLAYCAMERAS_API void Activate(const FCameraDirectorActivateParams& Params);
+	GAMEPLAYCAMERAS_API void Activate(const FCameraDirectorActivateParams& Params, FCameraDirectorEvaluationResult& OutResult);
 
 	/** Deactivates the camera director evaluator. */
-	GAMEPLAYCAMERAS_API void Deactivate(const FCameraDirectorDeactivateParams& Params);
+	GAMEPLAYCAMERAS_API void Deactivate(const FCameraDirectorDeactivateParams& Params, FCameraDirectorEvaluationResult& OutResult);
 
 	/** Gets the camera director. */
 	const UCameraDirector* GetCameraDirector() const { return PrivateCameraDirector; }
@@ -300,10 +310,10 @@ protected:
 	virtual void OnInitialize(const FCameraDirectorInitializeParams& Params) {}
 
 	/** Activates the camera director evaluator. May be called multiple times, in pair with OnDeactivate. */
-	virtual void OnActivate(const FCameraDirectorActivateParams& Params) {}
+	virtual void OnActivate(const FCameraDirectorActivateParams& Params, FCameraDirectorEvaluationResult& OutResult) {}
 
 	/** Deactivates the camera director evaluator. May be called multiple times, in pair with OnActivate. */
-	virtual void OnDeactivate(const FCameraDirectorDeactivateParams& Params) {}
+	virtual void OnDeactivate(const FCameraDirectorDeactivateParams& Params, FCameraDirectorEvaluationResult& OutResult) {}
 
 	/** Runs the camera director to determine what camera rig(s) should be active this frame. */
 	virtual void OnRun(const FCameraDirectorEvaluationParams& Params, FCameraDirectorEvaluationResult& OutResult) {}
diff --git a/Engine/Plugins/Cameras/GameplayCameras/Source/GameplayCameras/Public/Core/CameraEvaluationContext.h b/Engine/Plugins/Cameras/GameplayCameras/Source/GameplayCameras/Public/Core/CameraEvaluationContext.h
--- a/Engine/Plugins/Cameras/GameplayCameras/Source/GameplayCameras/Public/Core/CameraEvaluationContext.h
+++ b/Engine/Plugins/Cameras/GameplayCameras/Source/GameplayCameras/Public/Core/CameraEvaluationContext.h
@@ -25,6 +25,7 @@ namespace UE::Cameras
 {
 
 class FCameraSystemEvaluator;
+struct FCameraDirectorEvaluationResult;
 
 /**
  * Parameter struct for initializing an evaluation context.
@@ -188,6 +189,10 @@ protected:
 	void AutoCreateEditorPreviewDirectorEvaluator(const FCameraEvaluationContextActivateParams& Params);
 #endif  // WITH_EDITOR
 
+private:
+
+	void ExecuteSetupAndTeardownRequests(FCameraDirectorEvaluationResult& DirectorResult);
+
 protected:
 
 	/** The owner of this context, if any. */
@@ -228,7 +233,10 @@ private:
 	using FConditionalResults = TMap<ECameraEvaluationDataCondition, FCameraNodeEvaluationResult>;
 	FConditionalResults ConditionalResults;
 
+	/** Whether this context was initialized. */
 	bool bInitialized = false;
+
+	/** Whether this context is active. */
 	bool bActivated = false;
 };
 
diff --git a/Engine/Plugins/Cameras/GameplayCameras/Source/GameplayCameras/Public/GameFramework/GameplayCameraComponentBase.h b/Engine/Plugins/Cameras/GameplayCameras/Source/GameplayCameras/Public/GameFramework/GameplayCameraComponentBase.h
--- a/Engine/Plugins/Cameras/GameplayCameras/Source/GameplayCameras/Public/GameFramework/GameplayCameraComponentBase.h
+++ b/Engine/Plugins/Cameras/GameplayCameras/Source/GameplayCameras/Public/GameFramework/GameplayCameraComponentBase.h
@@ -253,12 +253,9 @@ private:
 	using FGameplayCameraComponentEvaluationContext = UE::Cameras::FGameplayCameraComponentEvaluationContext;
 	using FCameraEvaluationContext = UE::Cameras::FCameraEvaluationContext;
 
-	/** Evaluation context for running main-layer camera rigs. */
+	/** Evaluation context. */
 	TSharedPtr<FGameplayCameraComponentEvaluationContext> EvaluationContext;
 
-	/** Evaluation context for running base/global/visual camera rigs. */
-	TSharedPtr<FCameraEvaluationContext> NullContext;
-
 	/** The camera system in which to check for running rigs before deactivating our evaluation context. */
 	TWeakPtr<FCameraSystemEvaluator> PendingDeactivateCameraSystemEvaluator;
 
