Cameras: follow-up fix for validating rig parameter types

This change adds validation of container types (i.e. emitting an error when connecting an array of floats to a scalar float property)

diff --git a/Engine/Plugins/Cameras/GameplayCameras/Source/GameplayCameras/Private/Build/CameraObjectInterfaceBuilder.cpp b/Engine/Plugins/Cameras/GameplayCameras/Source/GameplayCameras/Private/Build/CameraObjectInterfaceBuilder.cpp
--- a/Engine/Plugins/Cameras/GameplayCameras/Source/GameplayCameras/Private/Build/CameraObjectInterfaceBuilder.cpp
+++ b/Engine/Plugins/Cameras/GameplayCameras/Source/GameplayCameras/Private/Build/CameraObjectInterfaceBuilder.cpp
@@ -356,6 +356,20 @@ private:
 						UEnum::GetDisplayValueAsText(DataParameter->DataType)));
 			return false;
 		}
+		if (DataParameter->DataContainerType != CustomParameter.ParameterContainerType)
+		{
+			ReportError(DataParameter->Target,
+					FText::Format(
+						LOCTEXT(
+							"DataParameterContainerTypeMismatch",
+							"Camera node parameter '{0}.{1}' has a different container type than camera rig parameter '{2}' ({3} vs {4})"),
+						FText::FromName(DataParameter->Target->GetFName()), 
+						FText::FromName(CustomParameter.ParameterName),
+						FText::FromString(DataParameter->InterfaceParameterName),
+						UEnum::GetDisplayValueAsText(CustomParameter.ParameterContainerType),
+						UEnum::GetDisplayValueAsText(DataParameter->DataContainerType)));
+			return false;
+		}
 
 		return CheckIfParameterCanBeOverridden(DataParameter, CustomParameter.OverrideDataID);
 	}
@@ -365,6 +379,38 @@ private:
 			const FProperty* TargetProperty,
 			FCameraContextDataID* DataID)
 	{
+		bool bDataContainerTypeMatches = false;
+		switch (DataParameter->DataContainerType)
+		{
+			case ECameraContextDataContainerType::Array:
+				if (const FArrayProperty* ArrayProperty = CastField<FArrayProperty>(TargetProperty))
+				{
+					TargetProperty = ArrayProperty->Inner;
+					bDataContainerTypeMatches = true;
+				}
+				break;
+			default:
+				bDataContainerTypeMatches = 
+					!TargetProperty->IsA<FArrayProperty>() && 
+					!TargetProperty->IsA<FMapProperty>() && 
+					!TargetProperty->IsA<FSetProperty>();
+				break;
+		}
+		if (!bDataContainerTypeMatches)
+		{
+			ReportError(DataParameter->Target,
+					FText::Format(
+						LOCTEXT(
+							"DataParameterContainerTypeMismatch",
+							"Camera node parameter '{0}.{1}' has a different container type than camera rig parameter '{2}' ({3} vs {4})"),
+						FText::FromName(DataParameter->Target->GetFName()), 
+						FText::FromName(TargetProperty->GetFName()),
+						FText::FromString(DataParameter->InterfaceParameterName),
+						TargetProperty->GetClass()->GetDisplayNameText(),
+						UEnum::GetDisplayValueAsText(DataParameter->DataContainerType)));
+			return false;
+		}
+
 		bool bDataTypeMatches = false;
 		switch (DataParameter->DataType)
 		{
@@ -714,7 +760,7 @@ void FCameraObjectInterfaceBuilder::BuildInterfaceParameters()
 
 void FCameraObjectInterfaceBuilder::BuildInterfaceParameterBindings()
 {
-	// Now we connect the interface paramters to whatever node property they are supposed to drive.
+	// Now we connect the interface parameters to whatever node property they are supposed to drive.
 	// Each time we need to check for either a custom property (via ICustomCameraNodeParameterProvider),
 	// or a UObject property found with reflection.
 
@@ -968,13 +1014,16 @@ bool FCameraObjectInterfaceBuilder::SetupDataContextPropertyOverride(const UCame
 	FStructProperty* TargetDataIDProperty = CastField<FStructProperty>(TargetClass->FindPropertyByName(TargetDataIDPropertyName));
 	if (!TargetDataIDProperty || TargetDataIDProperty->Struct != FCameraContextDataID::StaticStruct())
 	{
-		UE_LOG(LogCameraSystem, Error,
-				TEXT("Interface parameter '{0}' is driving data context property '{1}' on '{2}' "
-					 "but no FCameraContextDataID property '{3}' was found to store the override ID."),
-				*DataParameter->InterfaceParameterName,
-				*DataParameter->TargetPropertyName.ToString(),
-				*Target->GetName(),
-				*TargetDataIDPropertyName.ToString());
+		BuildLog.AddMessage(EMessageSeverity::Error,
+				DataParameter,
+				FText::Format(LOCTEXT(
+						"MissingDataID", 
+						"Interface parameter '{0}' is driving data context property '{1}' on '{2}' "
+						"but no FCameraContextDataID property '{3}' was found to store the override ID."),
+					FText::FromString(DataParameter->InterfaceParameterName),
+					FText::FromName(DataParameter->TargetPropertyName),
+					FText::FromName(Target->GetFName()),
+					FText::FromName(TargetDataIDPropertyName)));
 		return false;
 	}
 
