Cameras: fix logging formatting

diff --git a/Engine/Plugins/Cameras/GameplayCameras/Source/GameplayCameras/Private/Core/CameraEvaluationContext.cpp b/Engine/Plugins/Cameras/GameplayCameras/Source/GameplayCameras/Private/Core/CameraEvaluationContext.cpp
--- a/Engine/Plugins/Cameras/GameplayCameras/Source/GameplayCameras/Private/Core/CameraEvaluationContext.cpp
+++ b/Engine/Plugins/Cameras/GameplayCameras/Source/GameplayCameras/Private/Core/CameraEvaluationContext.cpp
@@ -341,7 +341,7 @@ void FCameraEvaluationContext::ExecuteSetupAndTeardownRequests(FCameraDirectorEv
 			UObject* RequestedCameraObject = Request.CameraRig ? (UObject*)Request.CameraRig : (UObject*)Request.CameraRigProxy;
 			UE_LOG(LogCameraSystem, Error, 
 					TEXT("Main layer camera rigs can only be activated/deactivated during the director's normal update. "
-						"Ignoring request to activate/deactivate '{0}'."),
+	 					 "Ignoring request to activate/deactivate '%s'."),
 					*GetNameSafe(RequestedCameraObject));
 			continue;
 		}
