Cameras: don't use manual focus distance on the output camera component if the focus distance isn't set

diff --git a/Engine/Plugins/Cameras/GameplayCameras/Source/GameplayCameras/Private/GameFramework/GameplayCameraComponentBase.cpp b/Engine/Plugins/Cameras/GameplayCameras/Source/GameplayCameras/Private/GameFramework/GameplayCameraComponentBase.cpp
--- a/Engine/Plugins/Cameras/GameplayCameras/Source/GameplayCameras/Private/GameFramework/GameplayCameraComponentBase.cpp
+++ b/Engine/Plugins/Cameras/GameplayCameras/Source/GameplayCameras/Private/GameFramework/GameplayCameraComponentBase.cpp
@@ -777,7 +777,10 @@ void UGameplayCameraComponentBase::UpdateOutputCameraComponent()
 			OutputCameraComponent->AspectRatioAxisConstraint = Result.CameraPose.GetAspectRatioAxisConstraint();
 
 			OutputCameraComponent->FocusSettings.ManualFocusDistance = Result.CameraPose.GetFocusDistance();
-			OutputCameraComponent->FocusSettings.FocusMethod = (Result.CameraPose.GetEnablePhysicalCamera() ? ECameraFocusMethod::Manual : ECameraFocusMethod::DoNotOverride);
+			OutputCameraComponent->FocusSettings.FocusMethod = (
+					(Result.CameraPose.GetEnablePhysicalCamera() && Result.CameraPose.GetFocusDistance() > 0.f) ? 
+					ECameraFocusMethod::Manual : 
+					ECameraFocusMethod::DoNotOverride);
 
 #if UE_VERSION_NEWER_THAN_OR_EQUAL(5,7,0)
 			OutputCameraComponent->ExposureMethod = (Result.CameraPose.GetEnablePhysicalCamera() ? ECameraExposureMethod::Enabled : ECameraExposureMethod::DoNotOverride);
