Cameras: backwards compatibility to 5.6

diff --git a/Engine/Plugins/Cameras/GameplayCameras/Source/GameplayCameras/Private/StructUtils/OverridablePropertyBag.cpp b/Engine/Plugins/Cameras/GameplayCameras/Source/GameplayCameras/Private/StructUtils/OverridablePropertyBag.cpp
--- a/Engine/Plugins/Cameras/GameplayCameras/Source/GameplayCameras/Private/StructUtils/OverridablePropertyBag.cpp
+++ b/Engine/Plugins/Cameras/GameplayCameras/Source/GameplayCameras/Private/StructUtils/OverridablePropertyBag.cpp
@@ -2,6 +2,8 @@
 
 #include "StructUtils/OverridablePropertyBag.h"
 
+#include "Misc/EngineVersionComparison.h"
+
 #include UE_INLINE_GENERATED_CPP_BY_NAME(OverridablePropertyBag)
 
 namespace FOverridablePropertyBagCustomVersion
@@ -80,8 +82,13 @@ bool FInstancedOverridablePropertyBag::Serialize(FArchive& Ar)
 		// serialization. This causes data corruption so we avoid that with this use-case here.
 		static const FInstancedOverridablePropertyBag ThisDefaults;
 
-		const UScriptStruct* ThisStruct = FInstancedOverridablePropertyBag::StaticStruct();
+		UScriptStruct* ThisStruct = FInstancedOverridablePropertyBag::StaticStruct();
+#if UE_VERSION_NEWER_THAN_OR_EQUAL(5,7,0)
 		ThisStruct->SerializeTaggedProperties(Ar, (uint8*)this, ThisStruct, (const uint8*)&ThisDefaults);
+#else
+		ThisStruct->SerializeTaggedProperties(Ar, (uint8*)this, ThisStruct, (uint8*)&ThisDefaults);
+#endif  // UE >= 5.7.0
+
 		return true;
 	}
 
