Cameras: fix variable and context data tables sometimes returning uninitialized memory, and replace assert with one-time warning

diff --git a/Engine/Plugins/Cameras/GameplayCameras/Source/GameplayCameras/Private/Core/CameraContextDataTable.cpp b/Engine/Plugins/Cameras/GameplayCameras/Source/GameplayCameras/Private/Core/CameraContextDataTable.cpp
--- a/Engine/Plugins/Cameras/GameplayCameras/Source/GameplayCameras/Private/Core/CameraContextDataTable.cpp
+++ b/Engine/Plugins/Cameras/GameplayCameras/Source/GameplayCameras/Private/Core/CameraContextDataTable.cpp
@@ -619,9 +619,10 @@ FCameraContextDataTable::FEntry* FCameraContextDataTable::FindEntry(FCameraConte
 const uint8* FCameraContextDataTable::GetData(
 		FCameraContextDataID DataID,
 		ECameraContextDataType ExpectedDataType,
-		const UObject* ExpectedDataTypeObject) const
+		const UObject* ExpectedDataTypeObject,
+		bool bOnlyIfWritten) const
 {
-	const uint8* Data = TryGetData(DataID, ExpectedDataType, ExpectedDataTypeObject);
+	const uint8* Data = TryGetData(DataID, ExpectedDataType, ExpectedDataTypeObject, bOnlyIfWritten);
 	ensureMsgf(
 			Data, 
 			TEXT("Can't get camera context data (ID '%d') because it doesn't exist in the table, or isn't of the expected data type."), 
@@ -632,14 +633,16 @@ const uint8* FCameraContextDataTable::GetData(
 const uint8* FCameraContextDataTable::TryGetData(
 		FCameraContextDataID DataID,
 		ECameraContextDataType ExpectedDataType,
-		const UObject* ExpectedDataTypeObject) const
+		const UObject* ExpectedDataTypeObject,
+		bool bOnlyIfWritten) const
 {
 	const FEntry* Entry = FindEntry(DataID);
 	if (Entry)
 	{
 		if (Entry->Type == ExpectedDataType && 
 				Entry->ContainerType == ECameraContextDataContainerType::None && 
-				Entry->TypeObject == ExpectedDataTypeObject)
+				Entry->TypeObject == ExpectedDataTypeObject &&
+				(!bOnlyIfWritten || EnumHasAnyFlags(Entry->Flags, EEntryFlags::Written)))
 		{
 			return Memory + Entry->Offset;
 		}
@@ -651,14 +654,16 @@ const uint8* FCameraContextDataTable::TryGetData(
 const FCameraContextDataTable::FEntryScriptArray* FCameraContextDataTable::TryGetArrayData(
 		FCameraContextDataID DataID,
 		ECameraContextDataType ExpectedDataType,
-		const UObject* ExpectedDataTypeObject) const
+		const UObject* ExpectedDataTypeObject,
+		bool bOnlyIfWritten) const
 {
 	const FEntry* Entry = FindEntry(DataID);
 	if (Entry)
 	{
 		if (Entry->Type == ExpectedDataType && 
 				Entry->ContainerType == ECameraContextDataContainerType::Array && 
-				Entry->TypeObject == ExpectedDataTypeObject)
+				Entry->TypeObject == ExpectedDataTypeObject &&
+				(!bOnlyIfWritten || EnumHasAnyFlags(Entry->Flags, EEntryFlags::Written)))
 		{
 			return (FEntryScriptArray*)(Memory + Entry->Offset);
 		}
@@ -670,12 +675,15 @@ const FCameraContextDataTable::FEntryScriptArray* FCameraContextDataTable::TryGe
 const uint8* FCameraContextDataTable::TryGetRawDataPtr(
 		FCameraContextDataID DataID,
 		ECameraContextDataType ExpectedDataType,
-		const UObject* ExpectedDataTypeObject) const
+		const UObject* ExpectedDataTypeObject,
+		bool bOnlyIfWritten) const
 {
 	const FEntry* Entry = FindEntry(DataID);
 	if (Entry)
 	{
-		if (Entry->Type == ExpectedDataType && Entry->TypeObject == ExpectedDataTypeObject)
+		if (Entry->Type == ExpectedDataType && 
+				Entry->TypeObject == ExpectedDataTypeObject &&
+				(!bOnlyIfWritten || EnumHasAnyFlags(Entry->Flags, EEntryFlags::Written)))
 		{
 			return Memory + Entry->Offset;
 		}
diff --git a/Engine/Plugins/Cameras/GameplayCameras/Source/GameplayCameras/Private/Core/CameraVariableTable.cpp b/Engine/Plugins/Cameras/GameplayCameras/Source/GameplayCameras/Private/Core/CameraVariableTable.cpp
--- a/Engine/Plugins/Cameras/GameplayCameras/Source/GameplayCameras/Private/Core/CameraVariableTable.cpp
+++ b/Engine/Plugins/Cameras/GameplayCameras/Source/GameplayCameras/Private/Core/CameraVariableTable.cpp
@@ -284,9 +284,10 @@ bool FCameraVariableTable::ContainsValue(FCameraVariableID VariableID) const
 const uint8* FCameraVariableTable::GetValue(
 		FCameraVariableID VariableID,
 		ECameraVariableType ExpectedVariableType,
-		const UScriptStruct* ExpectedBlendableStructType) const
+		const UScriptStruct* ExpectedBlendableStructType,
+		bool bOnlyIfWritten) const
 {
-	const uint8* Value = TryGetValue(VariableID, ExpectedVariableType, ExpectedBlendableStructType);
+	const uint8* Value = TryGetValue(VariableID, ExpectedVariableType, ExpectedBlendableStructType, bOnlyIfWritten);
 	ensureMsgf(Value, TEXT("Can't get camera variable (ID '%d') because it doesn't exist in the table."), VariableID.GetValue());
 	return Value;
 }
@@ -294,13 +295,17 @@ const uint8* FCameraVariableTable::GetValue(
 const uint8* FCameraVariableTable::TryGetValue(
 		FCameraVariableID VariableID,
 		ECameraVariableType ExpectedVariableType,
-		const UScriptStruct* ExpectedBlendableStructType) const
+		const UScriptStruct* ExpectedBlendableStructType,
+		bool bOnlyIfWritten) const
 {
 	const FEntry* Entry = FindEntry(VariableID);
 	if (Entry)
 	{
 		ensure(Entry->Type == ExpectedVariableType && Entry->StructType == ExpectedBlendableStructType);
-		return Memory + Entry->Offset;
+		if (!bOnlyIfWritten || EnumHasAnyFlags(Entry->Flags, EEntryFlags::Written))
+		{
+			return Memory + Entry->Offset;
+		}
 	}
 
 	return nullptr;
@@ -309,13 +314,17 @@ const uint8* FCameraVariableTable::TryGetValue(
 uint8* FCameraVariableTable::TryGetMutableValue(
 		FCameraVariableID VariableID,
 		ECameraVariableType ExpectedVariableType,
-		const UScriptStruct* ExpectedBlendableStructType)
+		const UScriptStruct* ExpectedBlendableStructType,
+		bool bOnlyIfWritten)
 {
 	FEntry* Entry = FindEntry(VariableID);
 	if (Entry)
 	{
 		ensure(Entry->Type == ExpectedVariableType && Entry->StructType == ExpectedBlendableStructType);
-		return Memory + Entry->Offset;
+		if (!bOnlyIfWritten || EnumHasAnyFlags(Entry->Flags, EEntryFlags::Written))
+		{
+			return Memory + Entry->Offset;
+		}
 	}
 
 	return nullptr;
@@ -531,9 +540,13 @@ void FCameraVariableTable::InternalOverride(const FCameraVariableTable& OtherTab
 				// We already have the other table's variable in our table. Let's check
 				// that the types match, and then copy the memory.
 #if WITH_EDITORONLY_DATA
-				checkf(ThisEntry->DebugName == OtherEntry.DebugName,
-						TEXT("Camera variable name collision! Expected variable '%d' to be named '%s', but other table has '%s'!"),
-						ThisEntry->ID.GetValue(), *ThisEntry->DebugName, *OtherEntry.DebugName);
+				if (ThisEntry->DebugName != OtherEntry.DebugName && !WarnedEntries.Contains(ThisEntry->ID))
+				{
+					UE_LOG(LogCameraSystem, Warning,
+							TEXT("Camera variable name collision! Expected variable '%d' to be named '%s', but other table has '%s'!"),
+							ThisEntry->ID.GetValue(), *ThisEntry->DebugName, *OtherEntry.DebugName);
+					WarnedEntries.Add(ThisEntry->ID);
+				}
 #endif
 
 #if WITH_EDITORONLY_DATA
diff --git a/Engine/Plugins/Cameras/GameplayCameras/Source/GameplayCameras/Public/Core/CameraContextDataTable.h b/Engine/Plugins/Cameras/GameplayCameras/Source/GameplayCameras/Public/Core/CameraContextDataTable.h
--- a/Engine/Plugins/Cameras/GameplayCameras/Source/GameplayCameras/Public/Core/CameraContextDataTable.h
+++ b/Engine/Plugins/Cameras/GameplayCameras/Source/GameplayCameras/Public/Core/CameraContextDataTable.h
@@ -152,22 +152,26 @@ public:
 	const uint8* GetData(
 			FCameraContextDataID DataID,
 			ECameraContextDataType ExpectedDataType,
-			const UObject* ExpectedDataTypeObject) const;
+			const UObject* ExpectedDataTypeObject,
+			bool bOnlyIfWritten = true) const;
 
 	const uint8* TryGetData(
 			FCameraContextDataID DataID,
 			ECameraContextDataType ExpectedDataType,
-			const UObject* ExpectedDataTypeObject) const;
+			const UObject* ExpectedDataTypeObject,
+			bool bOnlyIfWritten = true) const;
 
 	const FEntryScriptArray* TryGetArrayData(
 			FCameraContextDataID DataID,
 			ECameraContextDataType ExpectedDataType,
-			const UObject* ExpectedDataTypeObject) const;
+			const UObject* ExpectedDataTypeObject,
+			bool bOnlyIfWritten = true) const;
 
 	const uint8* TryGetRawDataPtr(
 			FCameraContextDataID DataID,
 			ECameraContextDataType ExpectedDataType,
-			const UObject* ExpectedDataTypeObject) const;
+			const UObject* ExpectedDataTypeObject,
+			bool bOnlyIfWritten = true) const;
 
 	void SetData(
 			FCameraContextDataID DataID,
@@ -422,7 +426,8 @@ const StorageType* FCameraContextDataTable::GetDataImpl(FCameraContextDataID InI
 	if (Entry && 
 			Entry->Type == DataType && 
 			Entry->ContainerType == ECameraContextDataContainerType::None && 
-			Entry->TypeObject == DataTypeObject)
+			Entry->TypeObject == DataTypeObject &&
+			EnumHasAnyFlags(Entry->Flags, EEntryFlags::Written))
 	{
 		const uint8* RawData = Memory + Entry->Offset;
 		return reinterpret_cast<const StorageType*>(RawData);
@@ -454,7 +459,8 @@ TConstArrayView<StorageType> FCameraContextDataTable::GetArrayDataImpl(FCameraCo
 	if (Entry && 
 			Entry->Type == DataType && 
 			Entry->ContainerType == ECameraContextDataContainerType::Array && 
-			Entry->TypeObject == DataTypeObject)
+			Entry->TypeObject == DataTypeObject &&
+			EnumHasAnyFlags(Entry->Flags, EEntryFlags::Written))
 	{
 		FEntryScriptArray* Array = (FEntryScriptArray*)(Memory + Entry->Offset);
 		const StorageType* ArrayData = reinterpret_cast<const StorageType*>(Array->GetData());
diff --git a/Engine/Plugins/Cameras/GameplayCameras/Source/GameplayCameras/Public/Core/CameraVariableTable.h b/Engine/Plugins/Cameras/GameplayCameras/Source/GameplayCameras/Public/Core/CameraVariableTable.h
--- a/Engine/Plugins/Cameras/GameplayCameras/Source/GameplayCameras/Public/Core/CameraVariableTable.h
+++ b/Engine/Plugins/Cameras/GameplayCameras/Source/GameplayCameras/Public/Core/CameraVariableTable.h
@@ -155,17 +155,20 @@ public:
 	UE_API const uint8* GetValue(
 			FCameraVariableID VariableID,
 			ECameraVariableType ExpectedVariableType,
-			const UScriptStruct* ExpectedBlendableStructType) const;
+			const UScriptStruct* ExpectedBlendableStructType,
+			bool bOnlyIfWritten = true) const;
 
 	UE_API const uint8* TryGetValue(
 			FCameraVariableID VariableID,
 			ECameraVariableType ExpectedVariableType,
-			const UScriptStruct* ExpectedBlendableStructType) const;
+			const UScriptStruct* ExpectedBlendableStructType,
+			bool bOnlyIfWritten = true) const;
 
 	UE_API uint8* TryGetMutableValue(
 			FCameraVariableID VariableID,
 			ECameraVariableType ExpectedVariableType,
-			const UScriptStruct* ExpectedBlendableStructType);
+			const UScriptStruct* ExpectedBlendableStructType,
+			bool bOnlyIfWritten = true);
 
 	UE_API void SetValue(
 			FCameraVariableID VariableID, 
@@ -249,6 +252,10 @@ private:
 	uint32 Capacity = 0;
 	uint32 Used = 0;
 
+#if WITH_EDITORONLY_DATA
+	TSet<FCameraVariableID> WarnedEntries;
+#endif
+
 	static UE_API TArray<FBlendableStructInfo> CachedBlendableStructs;
 	static UE_API bool bCachedBlendableStructs;
 
