Cameras: fix BP parameter getter nodes returning zero values instead of default rig values when the parameter isn't in the given tables

diff --git a/Engine/Plugins/Cameras/GameplayCameras/Source/GameplayCameras/Private/GameFramework/CameraRigParameterInterop.cpp b/Engine/Plugins/Cameras/GameplayCameras/Source/GameplayCameras/Private/GameFramework/CameraRigParameterInterop.cpp
--- a/Engine/Plugins/Cameras/GameplayCameras/Source/GameplayCameras/Private/GameFramework/CameraRigParameterInterop.cpp
+++ b/Engine/Plugins/Cameras/GameplayCameras/Source/GameplayCameras/Private/GameFramework/CameraRigParameterInterop.cpp
@@ -82,6 +82,8 @@ DEFINE_FUNCTION(UCameraRigParameterInterop::execGetCameraParameter)
 		}
 		else
 		{
+			bool bGotValue = false;
+
 			if (ParameterDefinition->ParameterType == ECameraObjectInterfaceParameterType::Blendable)
 			{
 				const FCameraVariableTable& VariableTable = CameraData.GetResult()->VariableTable;
@@ -89,6 +91,7 @@ DEFINE_FUNCTION(UCameraRigParameterInterop::execGetCameraParameter)
 				if (RawValue)
 				{
 					TargetProperty->CopyCompleteValue(TargetPtr, RawValue);
+					bGotValue = true;
 				}
 			}
 			else if (ParameterDefinition->ParameterType == ECameraObjectInterfaceParameterType::Data)
@@ -98,8 +101,31 @@ DEFINE_FUNCTION(UCameraRigParameterInterop::execGetCameraParameter)
 				if (RawValue)
 				{
 					TargetProperty->CopyCompleteValue(TargetPtr, RawValue);
+					bGotValue = true;
 				}
 			}
+
+			if (!bGotValue)
+			{
+				const FInstancedPropertyBag& DefaultParameters = CameraRig->GetDefaultParameters();
+				const UPropertyBag* DefaultParametersStruct = DefaultParameters.GetPropertyBagStruct();
+				const FPropertyBagPropertyDesc* PropertyDesc = DefaultParametersStruct->FindPropertyDescByID(ParameterDefinition->ParameterGuid);
+				if (PropertyDesc && PropertyDesc->CachedProperty)
+				{
+					const void* RawDefaultValue = PropertyDesc->CachedProperty->ContainerPtrToValuePtr<void>((const void*)DefaultParameters.GetValue().GetMemory());
+					TargetProperty->CopyCompleteValue(TargetPtr, RawDefaultValue);
+					bGotValue = true;
+				}
+			}
+
+			if (!bGotValue)
+			{
+				FBlueprintExceptionInfo ExceptionInfo(
+						EBlueprintExceptionType::NonFatalError,
+						FText::Format(LOCTEXT("NoDefaultValueFound", "Parameter {0} does not exist on camera rig {1}"), FText::FromName(ParameterName), FText::FromName(CameraRig->GetFName()))
+						);
+				FBlueprintCoreDelegates::ThrowScriptException(P_THIS, Stack, ExceptionInfo);
+			}
 		}
 
 		P_NATIVE_END
