Cameras: fix error messages for unexpected situations lerping/overriding variable tables

diff --git a/Engine/Plugins/Cameras/GameplayCameras/Source/GameplayCameras/Private/Core/CameraVariableTable.cpp b/Engine/Plugins/Cameras/GameplayCameras/Source/GameplayCameras/Private/Core/CameraVariableTable.cpp
--- a/Engine/Plugins/Cameras/GameplayCameras/Source/GameplayCameras/Private/Core/CameraVariableTable.cpp
+++ b/Engine/Plugins/Cameras/GameplayCameras/Source/GameplayCameras/Private/Core/CameraVariableTable.cpp
@@ -540,13 +540,9 @@ void FCameraVariableTable::InternalOverride(const FCameraVariableTable& OtherTab
 				// We already have the other table's variable in our table. Let's check
 				// that the types match, and then copy the memory.
 #if WITH_EDITORONLY_DATA
-				if (ThisEntry->DebugName != OtherEntry.DebugName && !WarnedEntries.Contains(ThisEntry->ID))
-				{
-					UE_LOG(LogCameraSystem, Warning,
-							TEXT("Camera variable name collision! Expected variable '%d' to be named '%s', but other table has '%s'!"),
-							ThisEntry->ID.GetValue(), *ThisEntry->DebugName, *OtherEntry.DebugName);
-					WarnedEntries.Add(ThisEntry->ID);
-				}
+				ensureMsgf(ThisEntry->DebugName == OtherEntry.DebugName,
+						TEXT("Camera variable name collision! Expected variable '%d' to be named '%s', but other table has '%s'!"),
+						ThisEntry->ID.GetValue(), *ThisEntry->DebugName, *OtherEntry.DebugName);
 #endif
 
 #if WITH_EDITORONLY_DATA
@@ -554,10 +550,14 @@ void FCameraVariableTable::InternalOverride(const FCameraVariableTable& OtherTab
 #else
 				const FString& DebugName = GUnavailableVariableDebugName;
 #endif
-				checkf(ThisEntry->Type == OtherEntry.Type && ThisEntry->StructType == OtherEntry.StructType, 
+				if (!ensureMsgf(ThisEntry->Type == OtherEntry.Type && ThisEntry->StructType == OtherEntry.StructType, 
 						TEXT("Camera variable name collision! Expected '%d' (%s) to be of type '%s' but other table has type '%s'!"),
 						ThisEntry->ID.GetValue(), *DebugName,
-						*UEnum::GetValueAsString(ThisEntry->Type), *UEnum::GetValueAsString(OtherEntry.Type));
+						*UEnum::GetValueAsString(ThisEntry->Type), *UEnum::GetValueAsString(OtherEntry.Type)))
+				{
+					// Unexpected situation: two variables of different types sharing the same ID.
+					continue;
+				}
 			}
 			else if (!bKnownOnly)
 			{
@@ -654,10 +654,14 @@ void FCameraVariableTable::InternalLerp(const FCameraVariableTable& ToTable, ECa
 #else
 				const FString& DebugName = GUnavailableVariableDebugName;
 #endif
-				ensureMsgf(FromEntry->Type == ToEntry.Type, 
+				if (!ensureMsgf(FromEntry->Type == ToEntry.Type, 
 						TEXT("Camera variable name collision! Expected '%d' (%s) to be of type '%s' but other table has type '%s'!"),
 						FromEntry->ID.GetValue(), *DebugName,
-						*UEnum::GetValueAsString(FromEntry->Type), *UEnum::GetValueAsString(ToEntry.Type));
+						*UEnum::GetValueAsString(FromEntry->Type), *UEnum::GetValueAsString(ToEntry.Type)))
+				{
+					// Unexpected situation: two variables of different types sharing the same ID.
+					continue;
+				}
 
 				const EEntryFlags FromFlags = FromEntry->Flags;
 				ensureMsgf(EnumHasAllFlags(FromFlags, EEntryFlags::Written),
diff --git a/Engine/Plugins/Cameras/GameplayCameras/Source/GameplayCameras/Public/Core/CameraVariableTable.h b/Engine/Plugins/Cameras/GameplayCameras/Source/GameplayCameras/Public/Core/CameraVariableTable.h
--- a/Engine/Plugins/Cameras/GameplayCameras/Source/GameplayCameras/Public/Core/CameraVariableTable.h
+++ b/Engine/Plugins/Cameras/GameplayCameras/Source/GameplayCameras/Public/Core/CameraVariableTable.h
@@ -252,10 +252,6 @@ private:
 	uint32 Capacity = 0;
 	uint32 Used = 0;
 
-#if WITH_EDITORONLY_DATA
-	TSet<FCameraVariableID> WarnedEntries;
-#endif
-
 	static UE_API TArray<FBlendableStructInfo> CachedBlendableStructs;
 	static UE_API bool bCachedBlendableStructs;
 
