Cameras: fix SET parameters nodes not setting anything if the parameter wasn't already known in the variable/data tables.

diff --git a/Engine/Plugins/Cameras/GameplayCameras/Source/GameplayCameras/Private/GameFramework/CameraRigParameterInterop.cpp b/Engine/Plugins/Cameras/GameplayCameras/Source/GameplayCameras/Private/GameFramework/CameraRigParameterInterop.cpp
--- a/Engine/Plugins/Cameras/GameplayCameras/Source/GameplayCameras/Private/GameFramework/CameraRigParameterInterop.cpp
+++ b/Engine/Plugins/Cameras/GameplayCameras/Source/GameplayCameras/Private/GameFramework/CameraRigParameterInterop.cpp
@@ -189,12 +189,39 @@ DEFINE_FUNCTION(UCameraRigParameterInterop::execSetCameraParameter)
 			if (ParameterDefinition->ParameterType == ECameraObjectInterfaceParameterType::Blendable)
 			{
 				FCameraVariableTable& VariableTable = CameraData.GetResult()->VariableTable;
-				VariableTable.TrySetValue(ParameterDefinition->VariableID, ParameterDefinition->VariableType, ParameterDefinition->BlendableStructType, SourcePtr);
+				const bool bDidSet = VariableTable.TrySetValue(ParameterDefinition->VariableID, ParameterDefinition->VariableType, ParameterDefinition->BlendableStructType, SourcePtr);
+				if (!bDidSet)
+				{
+					UCameraObjectInterfaceBlendableParameter* BlendableParameter = CameraRig->Interface.FindBlendableParameterByGuid(ParameterDefinition->ParameterGuid);
+					if (ensureMsgf(
+								BlendableParameter, 
+								TEXT("Can't find parameter '%s' on camera rig '%s', was it changed since it was last built?"), 
+								*ParameterDefinition->ParameterName.ToString(), *GetNameSafe(CameraRig)))
+					{
+						FCameraVariableDefinition VariableDefinition = BlendableParameter->GetVariableDefinition();
+						VariableDefinition.bIsPrivate = false;  // If the parameter isn't in the table, it's probably meant for another rig, so let's make sure it propagates.
+						VariableTable.AddVariable(VariableDefinition);
+						VariableTable.SetValue(ParameterDefinition->VariableID, ParameterDefinition->VariableType, ParameterDefinition->BlendableStructType, SourcePtr);
+					}
+				}
 			}
 			else if (ParameterDefinition->ParameterType == ECameraObjectInterfaceParameterType::Data)
 			{
 				FCameraContextDataTable& ContextDataTable = CameraData.GetResult()->ContextDataTable;
 				uint8* RawValue = ContextDataTable.TryGetMutableRawDataPtr(ParameterDefinition->DataID, ParameterDefinition->DataType, ParameterDefinition->DataTypeObject);
+				if (!RawValue)
+				{
+					UCameraObjectInterfaceDataParameter* DataParameter = CameraRig->Interface.FindDataParameterByGuid(ParameterDefinition->ParameterGuid);
+					if (ensureMsgf(
+								DataParameter, 
+								TEXT("Can't find parameter '%s' on camera rig '%s', was it changed since it was last built?"), 
+								*ParameterDefinition->ParameterName.ToString(), *GetNameSafe(CameraRig)))
+					{
+						FCameraContextDataDefinition DataDefinition = DataParameter->GetDataDefinition();
+						ContextDataTable.AddData(DataDefinition);
+						RawValue = ContextDataTable.TryGetMutableRawDataPtr(ParameterDefinition->DataID, ParameterDefinition->DataType, ParameterDefinition->DataTypeObject);
+					}
+				}
 				if (RawValue)
 				{
 					SourceProperty->CopyCompleteValue(RawValue, SourcePtr);
