Cameras: fix shake prefab node not passing shake parameter values over to the visual layer in deferred mode

diff --git a/Engine/Plugins/Cameras/GameplayCameras/Source/GameplayCameras/Private/Nodes/Shakes/CameraShakeCameraNode.cpp b/Engine/Plugins/Cameras/GameplayCameras/Source/GameplayCameras/Private/Nodes/Shakes/CameraShakeCameraNode.cpp
--- a/Engine/Plugins/Cameras/GameplayCameras/Source/GameplayCameras/Private/Nodes/Shakes/CameraShakeCameraNode.cpp
+++ b/Engine/Plugins/Cameras/GameplayCameras/Source/GameplayCameras/Private/Nodes/Shakes/CameraShakeCameraNode.cpp
@@ -161,7 +161,17 @@ void UCameraShakeCameraNode::OnBuild(FCameraObjectBuildContext& BuildContext)
 
 	// Whatever allocations our inner camera shake needs for its evaluators and
 	// their camera variables, we add that to our camera shake's allocation info.
-	BuildContext.AllocationInfo.Append(CameraShake->AllocationInfo);
+	// If we're going to be running the shake in a deferred way, however, we need
+	// to make the variables public so that they propagate to it.
+	FCameraObjectAllocationInfo CameraShakeAllocationInfo(CameraShake->AllocationInfo);
+	if (EvaluationMode == ECameraShakeEvaluationMode::VisualLayer)
+	{
+		for (FCameraVariableDefinition& VariableDefinition : CameraShakeAllocationInfo.VariableTableInfo.VariableDefinitions)
+		{
+			VariableDefinition.bIsPrivate = false;
+		}
+	}
+	BuildContext.AllocationInfo.Append(CameraShakeAllocationInfo);
 }
 
 void UCameraShakeCameraNode::GetCustomCameraNodeParameters(FCustomCameraNodeParameterInfos& OutParameterInfos)
