Cameras: QoL improvements to the Blueprint camera node

- Make the FindEvaluationContextOwnerActor method default to the Actor class
- Add access to the IsActiveCameraRig flag
- Populate new Blueprint evaluators with a few helpful ghost nodes and comments

diff --git a/Engine/Plugins/Cameras/GameplayCameras/Source/GameplayCameras/Private/Nodes/Utility/BlueprintCameraNode.cpp b/Engine/Plugins/Cameras/GameplayCameras/Source/GameplayCameras/Private/Nodes/Utility/BlueprintCameraNode.cpp
--- a/Engine/Plugins/Cameras/GameplayCameras/Source/GameplayCameras/Private/Nodes/Utility/BlueprintCameraNode.cpp
+++ b/Engine/Plugins/Cameras/GameplayCameras/Source/GameplayCameras/Private/Nodes/Utility/BlueprintCameraNode.cpp
@@ -198,6 +198,7 @@ void UBlueprintCameraNodeEvaluator::NativeRunCameraNode(const FCameraNodeEvaluat
 	SetupExecution(Params.EvaluationContext, OutResult);
 	{
 		bIsFirstFrame = Params.bIsFirstFrame;
+		bIsActiveCameraRig = Params.bIsActiveCameraRig;
 
 		TickCameraNode(Params.DeltaTime);
 	}
@@ -238,7 +239,7 @@ AActor* UBlueprintCameraNodeEvaluator::FindEvaluationContextOwnerActor(TSubclass
 			OwnerActor = ContextOwnerAsActor;
 		}
 
-		if (OwnerActor && OwnerActor->IsA(ActorClass))
+		if (OwnerActor && (!ActorClass || OwnerActor->IsA(ActorClass)))
 		{
 			return OwnerActor;
 		}
diff --git a/Engine/Plugins/Cameras/GameplayCameras/Source/GameplayCameras/Public/Nodes/Utility/BlueprintCameraNode.h b/Engine/Plugins/Cameras/GameplayCameras/Source/GameplayCameras/Public/Nodes/Utility/BlueprintCameraNode.h
--- a/Engine/Plugins/Cameras/GameplayCameras/Source/GameplayCameras/Public/Nodes/Utility/BlueprintCameraNode.h
+++ b/Engine/Plugins/Cameras/GameplayCameras/Source/GameplayCameras/Public/Nodes/Utility/BlueprintCameraNode.h
@@ -46,7 +46,7 @@ public:
 	 * UGameplayCameraComponent) or an actor itself.
 	 */
 	UFUNCTION(BlueprintPure, Category="Evaluation", meta=(DeterminesOutputType="ActorClass"))
-	AActor* FindEvaluationContextOwnerActor(TSubclassOf<AActor> ActorClass) const;
+	AActor* FindEvaluationContextOwnerActor(TSubclassOf<AActor> ActorClass = nullptr) const;
 
 	/** A utility function to get the current camera pose from this node's camera data. */
 	UFUNCTION(BlueprintPure, Category="Evaluation")
@@ -96,6 +96,10 @@ protected:
 	UPROPERTY(BlueprintReadOnly, Category="Evaluation")
 	bool bIsFirstFrame = false;
 
+	/** Whether this camera node is running inside the active camera rig in this layer. */
+	UPROPERTY(BlueprintReadOnly, Category="Evaluation")
+	bool bIsActiveCameraRig = false;
+
 	/** The owner object of this camera node's evaluation context. */
 	UPROPERTY(BlueprintReadOnly, Category="Evaluation")
 	TObjectPtr<UObject> EvaluationContextOwner;
diff --git a/Engine/Plugins/Cameras/GameplayCameras/Source/GameplayCamerasEditor/Private/GameplayCamerasEditorModule.cpp b/Engine/Plugins/Cameras/GameplayCameras/Source/GameplayCamerasEditor/Private/GameplayCamerasEditorModule.cpp
--- a/Engine/Plugins/Cameras/GameplayCameras/Source/GameplayCamerasEditor/Private/GameplayCamerasEditorModule.cpp
+++ b/Engine/Plugins/Cameras/GameplayCameras/Source/GameplayCamerasEditor/Private/GameplayCamerasEditorModule.cpp
@@ -49,6 +49,7 @@
 #include "K2Node_Event.h"
 #include "Kismet2/KismetEditorUtilities.h"
 #include "Modules/ModuleManager.h"
+#include "Nodes/Utility/BlueprintCameraNode.h"
 #include "ObjectTools.h"
 #include "PropertyBagDetails.h"
 #include "PropertyEditorModule.h"
@@ -531,19 +532,33 @@ private:
 				UBlueprintCameraDirectorEvaluator::StaticClass(),
 				FKismetEditorUtilities::FOnBlueprintCreated::CreateRaw(
 					this, &FGameplayCamerasEditorModule::OnNewBlueprintCameraDirectorEvaluatorCreated));
+
+		FKismetEditorUtilities::RegisterAutoGeneratedDefaultEvent(
+				this,
+				UBlueprintCameraNodeEvaluator::StaticClass(),
+				GET_FUNCTION_NAME_CHECKED(UBlueprintCameraNodeEvaluator, InitializeCameraNode));
+		FKismetEditorUtilities::RegisterAutoGeneratedDefaultEvent(
+				this,
+				UBlueprintCameraNodeEvaluator::StaticClass(),
+				GET_FUNCTION_NAME_CHECKED(UBlueprintCameraNodeEvaluator, TickCameraNode));
+		FKismetEditorUtilities::RegisterOnBlueprintCreatedCallback(
+				this, 
+				UBlueprintCameraNodeEvaluator::StaticClass(),
+				FKismetEditorUtilities::FOnBlueprintCreated::CreateRaw(
+					this, &FGameplayCamerasEditorModule::OnNewBlueprintCameraNodeEvaluatorCreated));
 	}
 
-	void OnNewBlueprintCameraDirectorEvaluatorCreated(UBlueprint* InBlueprint)
+	static UK2Node_Event* FindEventNodeInEventGraph(UBlueprint* InBlueprint, const FName InEventName) 
 	{
 		if (InBlueprint->BlueprintType != BPTYPE_Normal)
 		{
-			return;
+			return nullptr;
 		}
 		TObjectPtr<UEdGraph>* FoundItem = InBlueprint->UbergraphPages.FindByPredicate(
 					[](UEdGraph* Item) { return Item->GetFName() == TEXT("EventGraph"); });
 		if (!FoundItem)
 		{
-			return;
+			return nullptr;
 		}
 
 		TObjectPtr<UEdGraph> EventGraph(*FoundItem);
@@ -552,29 +567,50 @@ private:
 		EventGraph->GetNodesOfClass(EventNodes);
 		if (EventNodes.IsEmpty())
 		{
-			return;
+			return nullptr;
 		}
 
-		const FName EventName = GET_FUNCTION_NAME_CHECKED(UBlueprintCameraDirectorEvaluator, RunCameraDirector);
 		UK2Node_Event** FoundEventNode = EventNodes.FindByPredicate(
-				[&EventName](UK2Node_Event* Item) 
+				[&InEventName](UK2Node_Event* Item) 
 				{
-					return Item->EventReference.GetMemberName() == EventName;
+					return Item->EventReference.GetMemberName() == InEventName;
 				});
-		if (!FoundEventNode)
+		if (FoundEventNode)
 		{
-			return;
+			return *FoundEventNode;
 		}
+		return nullptr;
+	}
 
-		UK2Node_Event* RunEventNode(*FoundEventNode);
-	
-		const FText RunEventNodeCommentText = LOCTEXT(
-				"BlueprintCameraDirector_RunEventComment",
-				"Implement your camera director logic starting from here.\n"
-				"This node is currently disabled, but start dragging off pins to enable it.\n"
-				"Call ActivateCameraRig at least once to declare which camera rig(s) should be active this frame.");
-		RunEventNode->NodeComment = RunEventNodeCommentText.ToString();
-		RunEventNode->bCommentBubbleVisible = true;
+	void OnNewBlueprintCameraDirectorEvaluatorCreated(UBlueprint* InBlueprint)
+	{
+		const FName EventName = GET_FUNCTION_NAME_CHECKED(UBlueprintCameraDirectorEvaluator, RunCameraDirector);
+		if (UK2Node_Event* RunEventNode = FindEventNodeInEventGraph(InBlueprint, EventName))
+		{
+			const FText RunEventNodeCommentText = LOCTEXT(
+					"BlueprintCameraDirector_RunEventComment",
+					"Implement your camera director logic starting from here.\n"
+					"This node is currently disabled, but start dragging off pins to enable it.\n"
+					"Call ActivateCameraRig at least once to declare which camera rig(s) should be active this frame.");
+			RunEventNode->NodeComment = RunEventNodeCommentText.ToString();
+			RunEventNode->bCommentBubbleVisible = true;
+		}
+	}
+
+	void OnNewBlueprintCameraNodeEvaluatorCreated(UBlueprint* InBlueprint)
+	{
+		const FName EventName = GET_FUNCTION_NAME_CHECKED(UBlueprintCameraNodeEvaluator, TickCameraNode);
+		if (UK2Node_Event* RunEventNode = FindEventNodeInEventGraph(InBlueprint, EventName))
+		{
+			const FText RunEventNodeCommentText = LOCTEXT(
+					"BlueprintCameraNode_TickEventComment",
+					"Implement your camera node logic starting from here.\n"
+					"This node is currently disabled, but start dragging off pins to enable it.\n"
+					"Use the CameraPose and CameraData variables to affect the camera and its parameters. Other useful variables "
+					"and functions are available in the 'Evaluation' category.");
+			RunEventNode->NodeComment = RunEventNodeCommentText.ToString();
+			RunEventNode->bCommentBubbleVisible = true;
+		}
 	}
 
 	void UnregisterEdGraphUtilities()
