Cameras: allow the user to set variables to be pre-blended

diff --git a/Engine/Plugins/Cameras/GameplayCameras/Config/BaseGameplayCameras.ini b/Engine/Plugins/Cameras/GameplayCameras/Config/BaseGameplayCameras.ini
--- a/Engine/Plugins/Cameras/GameplayCameras/Config/BaseGameplayCameras.ini
+++ b/Engine/Plugins/Cameras/GameplayCameras/Config/BaseGameplayCameras.ini
@@ -29,3 +29,6 @@
 
 +StructRedirects=(OldName="/Script/GameplayCameras.CameraRigProxyTableEntry",NewName="CameraRigProxyRedirectTableEntry")
 
+; UE 5.7 redirects.
++PropertyRedirects=(OldName="/Script/GameplayCameras.CameraVariableAsset.bIsInput",NewName="CameraVariableAsset.bIsPreBlended")
+
diff --git a/Engine/Plugins/Cameras/GameplayCameras/Source/GameplayCameras/Private/Core/CameraVariableAssets.cpp b/Engine/Plugins/Cameras/GameplayCameras/Source/GameplayCameras/Private/Core/CameraVariableAssets.cpp
--- a/Engine/Plugins/Cameras/GameplayCameras/Source/GameplayCameras/Private/Core/CameraVariableAssets.cpp
+++ b/Engine/Plugins/Cameras/GameplayCameras/Source/GameplayCameras/Private/Core/CameraVariableAssets.cpp
@@ -21,7 +21,7 @@ FCameraVariableDefinition UCameraVariableAsset::GetVariableDefinition() const
 	VariableDefinition.VariableID = GetVariableID();
 	VariableDefinition.VariableType = GetVariableType();
 	VariableDefinition.bIsPrivate = bIsPrivate;
-	VariableDefinition.bIsInput = bIsInput;
+	VariableDefinition.bIsInput = bIsPreBlended;
 	VariableDefinition.bAutoReset = bAutoReset;
 #if WITH_EDITORONLY_DATA
 	VariableDefinition.VariableName = GetDisplayName();
diff --git a/Engine/Plugins/Cameras/GameplayCameras/Source/GameplayCameras/Public/Core/CameraVariableAssets.h b/Engine/Plugins/Cameras/GameplayCameras/Source/GameplayCameras/Public/Core/CameraVariableAssets.h
--- a/Engine/Plugins/Cameras/GameplayCameras/Source/GameplayCameras/Public/Core/CameraVariableAssets.h
+++ b/Engine/Plugins/Cameras/GameplayCameras/Source/GameplayCameras/Public/Core/CameraVariableAssets.h
@@ -61,14 +61,20 @@ public:
 	UPROPERTY(EditAnywhere, Category=Camera)
 	bool bAutoReset = false;
 
+	/**
+	 * Whether this variable should be pre-blended.
+	 *
+	 * Pre-blending means that if two blending camera rigs share this variable, 
+	 * each of their values will be blended in a first evaluation pass, and then
+	 * both camera rigs will evaluate with the same blended value.
+	 */
+	UPROPERTY(EditAnywhere, Category=Camera)
+	bool bIsPreBlended = false;
+
 	/** Whether this variable is private and shouldn't be propagated with evaluation results. */
 	UPROPERTY()
 	bool bIsPrivate = false;
 
-	/** Whether this variable is an input variable that gets blended before node evaluators are run. */
-	UPROPERTY()
-	bool bIsInput = false;
-
 private:
 
 	UPROPERTY()
