Cameras: fix Blueprint breakpoints not working on GET/SET parameter nodes

diff --git a/Engine/Plugins/Cameras/GameplayCameras/Source/GameplayCamerasUncookedOnly/Private/BlueprintGraph/K2Node_CameraRigBase.cpp b/Engine/Plugins/Cameras/GameplayCameras/Source/GameplayCamerasUncookedOnly/Private/BlueprintGraph/K2Node_CameraRigBase.cpp
--- a/Engine/Plugins/Cameras/GameplayCameras/Source/GameplayCamerasUncookedOnly/Private/BlueprintGraph/K2Node_CameraRigBase.cpp
+++ b/Engine/Plugins/Cameras/GameplayCameras/Source/GameplayCamerasUncookedOnly/Private/BlueprintGraph/K2Node_CameraRigBase.cpp
@@ -250,7 +250,7 @@ UK2Node_CallFunction* UK2Node_CameraRigBase::CreateMakeLiteralNode(FKismetCompil
 	UK2Node_CallFunction* MakeLiteralNode = CompilerContext.SpawnIntermediateNode<UK2Node_CallFunction>(SourceNode, SourceGraph);
 	MakeLiteralNode->FunctionReference.SetExternalMember(FunctionName, FunctionLibraryClass);
 	MakeLiteralNode->AllocateDefaultPins();
-	CompilerContext.MessageLog.NotifyIntermediateObjectCreation(MakeLiteralNode, SourceGraph);
+	CompilerContext.MessageLog.NotifyIntermediateObjectCreation(MakeLiteralNode, SourceNode);
 
 	UEdGraphPin* LiteralValuePin = MakeLiteralNode->FindPinChecked(TEXT("Value"));
 	LiteralValuePin->DefaultValue = SourceValuePin->DefaultValue;
@@ -311,7 +311,7 @@ UK2Node* UK2Node_CameraRigBase::MakeLiteralValueForPin(FKismetCompilerContext& C
 			UK2Node_EnumLiteral* EnumLiteralNode = CompilerContext.SpawnIntermediateNode<UK2Node_EnumLiteral>(SourceNode, SourceGraph);
 			EnumLiteralNode->Enum = EnumType;
 			EnumLiteralNode->AllocateDefaultPins();
-			CompilerContext.MessageLog.NotifyIntermediateObjectCreation(EnumLiteralNode, SourceGraph);
+			CompilerContext.MessageLog.NotifyIntermediateObjectCreation(EnumLiteralNode, SourceNode);
 
 			UEdGraphPin* EnumLiteralValuePin = EnumLiteralNode->FindPinChecked(UK2Node_EnumLiteral::GetEnumInputPinName());
 			EnumLiteralValuePin->DefaultValue = SourceValuePin->DefaultValue;
diff --git a/Engine/Plugins/Cameras/GameplayCameras/Source/GameplayCamerasUncookedOnly/Private/BlueprintGraph/K2Node_GetCameraRigParameter.cpp b/Engine/Plugins/Cameras/GameplayCameras/Source/GameplayCamerasUncookedOnly/Private/BlueprintGraph/K2Node_GetCameraRigParameter.cpp
--- a/Engine/Plugins/Cameras/GameplayCameras/Source/GameplayCamerasUncookedOnly/Private/BlueprintGraph/K2Node_GetCameraRigParameter.cpp
+++ b/Engine/Plugins/Cameras/GameplayCameras/Source/GameplayCamerasUncookedOnly/Private/BlueprintGraph/K2Node_GetCameraRigParameter.cpp
@@ -86,7 +86,7 @@ void UK2Node_GetCameraRigParameter::ExpandNode(FKismetCompilerContext& CompilerC
 	UK2Node_CallFunction* CallGetParameter = CompilerContext.SpawnIntermediateNode<UK2Node_CallFunction>(this, SourceGraph);
 	CallGetParameter->FunctionReference.SetExternalMember(TEXT("GetCameraParameter"), UCameraRigParameterInterop::StaticClass());
 	CallGetParameter->AllocateDefaultPins();
-	CompilerContext.MessageLog.NotifyIntermediateObjectCreation(CallGetParameter, SourceGraph);
+	CompilerContext.MessageLog.NotifyIntermediateObjectCreation(CallGetParameter, this);
 
 	// Connect the camera evaluation result argument.
 	UEdGraphPin* CallGetParameterResultPin = CallGetParameter->FindPinChecked(TEXT("CameraData"));
diff --git a/Engine/Plugins/Cameras/GameplayCameras/Source/GameplayCamerasUncookedOnly/Private/BlueprintGraph/K2Node_GetCameraRigParameters.cpp b/Engine/Plugins/Cameras/GameplayCameras/Source/GameplayCamerasUncookedOnly/Private/BlueprintGraph/K2Node_GetCameraRigParameters.cpp
--- a/Engine/Plugins/Cameras/GameplayCameras/Source/GameplayCamerasUncookedOnly/Private/BlueprintGraph/K2Node_GetCameraRigParameters.cpp
+++ b/Engine/Plugins/Cameras/GameplayCameras/Source/GameplayCamerasUncookedOnly/Private/BlueprintGraph/K2Node_GetCameraRigParameters.cpp
@@ -93,7 +93,7 @@ void UK2Node_GetCameraRigParameters::ExpandNode(FKismetCompilerContext& Compiler
 		UK2Node_CallFunction* CallGetParameter = CompilerContext.SpawnIntermediateNode<UK2Node_CallFunction>(this, SourceGraph);
 		CallGetParameter->FunctionReference.SetExternalMember(TEXT("GetCameraParameter"), UCameraRigParameterInterop::StaticClass());
 		CallGetParameter->AllocateDefaultPins();
-		CompilerContext.MessageLog.NotifyIntermediateObjectCreation(CallGetParameter, SourceGraph);
+		CompilerContext.MessageLog.NotifyIntermediateObjectCreation(CallGetParameter, this);
 
 		// Connect the camera evaluation result argument.
 		UEdGraphPin* CallGetParameterResultPin = CallGetParameter->FindPinChecked(TEXT("CameraData"));
@@ -137,7 +137,7 @@ void UK2Node_GetCameraRigParameters::ExpandNode(FKismetCompilerContext& Compiler
 		UK2Node_CallFunction* CallGetParameter = CompilerContext.SpawnIntermediateNode<UK2Node_CallFunction>(this, SourceGraph);
 		CallGetParameter->FunctionReference.SetExternalMember(TEXT("GetCameraParameter"), UCameraRigParameterInterop::StaticClass());
 		CallGetParameter->AllocateDefaultPins();
-		CompilerContext.MessageLog.NotifyIntermediateObjectCreation(CallGetParameter, SourceGraph);
+		CompilerContext.MessageLog.NotifyIntermediateObjectCreation(CallGetParameter, this);
 
 		// Connect the camera evaluation result argument.
 		UEdGraphPin* CallGetParameterResultPin = CallGetParameter->FindPinChecked(TEXT("CameraData"));
diff --git a/Engine/Plugins/Cameras/GameplayCameras/Source/GameplayCamerasUncookedOnly/Private/BlueprintGraph/K2Node_SetCameraRigParameter.cpp b/Engine/Plugins/Cameras/GameplayCameras/Source/GameplayCamerasUncookedOnly/Private/BlueprintGraph/K2Node_SetCameraRigParameter.cpp
--- a/Engine/Plugins/Cameras/GameplayCameras/Source/GameplayCamerasUncookedOnly/Private/BlueprintGraph/K2Node_SetCameraRigParameter.cpp
+++ b/Engine/Plugins/Cameras/GameplayCameras/Source/GameplayCamerasUncookedOnly/Private/BlueprintGraph/K2Node_SetCameraRigParameter.cpp
@@ -86,7 +86,7 @@ void UK2Node_SetCameraRigParameter::ExpandNode(FKismetCompilerContext& CompilerC
 	UK2Node_CallFunction* CallSetParameter = CompilerContext.SpawnIntermediateNode<UK2Node_CallFunction>(this, SourceGraph);
 	CallSetParameter->FunctionReference.SetExternalMember(TEXT("SetCameraParameter"), UCameraRigParameterInterop::StaticClass());
 	CallSetParameter->AllocateDefaultPins();
-	CompilerContext.MessageLog.NotifyIntermediateObjectCreation(CallSetParameter, SourceGraph);
+	CompilerContext.MessageLog.NotifyIntermediateObjectCreation(CallSetParameter, this);
 	UEdGraphPin* FirstExecPin = CallSetParameter->GetExecPin();
 
 	// Connect the camera evaluation result argument.
diff --git a/Engine/Plugins/Cameras/GameplayCameras/Source/GameplayCamerasUncookedOnly/Private/BlueprintGraph/K2Node_SetCameraRigParameters.cpp b/Engine/Plugins/Cameras/GameplayCameras/Source/GameplayCamerasUncookedOnly/Private/BlueprintGraph/K2Node_SetCameraRigParameters.cpp
--- a/Engine/Plugins/Cameras/GameplayCameras/Source/GameplayCamerasUncookedOnly/Private/BlueprintGraph/K2Node_SetCameraRigParameters.cpp
+++ b/Engine/Plugins/Cameras/GameplayCameras/Source/GameplayCamerasUncookedOnly/Private/BlueprintGraph/K2Node_SetCameraRigParameters.cpp
@@ -134,7 +134,7 @@ void UK2Node_SetCameraRigParameters::ExpandNode(FKismetCompilerContext& Compiler
 		UK2Node_CallFunction* CallSetParameter = CompilerContext.SpawnIntermediateNode<UK2Node_CallFunction>(this, SourceGraph);
 		CallSetParameter->FunctionReference.SetExternalMember(TEXT("SetCameraParameter"), UCameraRigParameterInterop::StaticClass());
 		CallSetParameter->AllocateDefaultPins();
-		CompilerContext.MessageLog.NotifyIntermediateObjectCreation(CallSetParameter, SourceGraph);
+		CompilerContext.MessageLog.NotifyIntermediateObjectCreation(CallSetParameter, this);
 
 		// Connect the camera evaluation result argument.
 		UEdGraphPin* CallSetParameterResultPin = CallSetParameter->FindPinChecked(TEXT("CameraData"));
@@ -208,7 +208,7 @@ void UK2Node_SetCameraRigParameters::ExpandNode(FKismetCompilerContext& Compiler
 		UK2Node_CallFunction* CallSetParameter = CompilerContext.SpawnIntermediateNode<UK2Node_CallFunction>(this, SourceGraph);
 		CallSetParameter->FunctionReference.SetExternalMember(TEXT("SetCameraParameter"), UCameraRigParameterInterop::StaticClass());
 		CallSetParameter->AllocateDefaultPins();
-		CompilerContext.MessageLog.NotifyIntermediateObjectCreation(CallSetParameter, SourceGraph);
+		CompilerContext.MessageLog.NotifyIntermediateObjectCreation(CallSetParameter, this);
 
 		UEdGraphPin* CurExecPin = CallSetParameter->GetExecPin();
 
