Cameras: fix possible crash when running with unbuilt or outdated rigs

diff --git a/Engine/Plugins/Cameras/GameplayCameras/Source/GameplayCameras/Private/Nodes/Utility/BlueprintCameraNode.cpp b/Engine/Plugins/Cameras/GameplayCameras/Source/GameplayCameras/Private/Nodes/Utility/BlueprintCameraNode.cpp
--- a/Engine/Plugins/Cameras/GameplayCameras/Source/GameplayCameras/Private/Nodes/Utility/BlueprintCameraNode.cpp
+++ b/Engine/Plugins/Cameras/GameplayCameras/Source/GameplayCameras/Private/Nodes/Utility/BlueprintCameraNode.cpp
@@ -124,11 +124,16 @@ void FBlueprintCameraNodeEvaluator::ApplyParameterOverrides(const FCameraVariabl
 		{
 			if (BlendableParameter.OverrideVariableID)
 			{
-				const uint8* ValuePtr = VariableTable.GetValue(
+				const uint8* ValuePtr = VariableTable.TryGetValue(
 						BlendableParameter.OverrideVariableID,
 						BlendableParameter.ParameterType,
 						BlendableParameter.BlendableStructType);
-				Property->SetValue_InContainer(EvaluatorBlueprint, ValuePtr);
+				if (ensureMsgf(ValuePtr, 
+							TEXT("Can't find variable table entry for parameter '%s'. Did you build camera rig '%s'?"),
+							*BlendableParameter.ParameterName.ToString(), *GetNameSafe(BlueprintNode->GetOutermost())))
+				{
+					Property->SetValue_InContainer(EvaluatorBlueprint, ValuePtr);
+				}
 			}
 		}
 	}
@@ -141,11 +146,16 @@ void FBlueprintCameraNodeEvaluator::ApplyParameterOverrides(const FCameraVariabl
 		{
 			if (DataParameter.OverrideDataID)
 			{
-				const uint8* DataPtr = ContextDataTable.GetData(
+				const uint8* DataPtr = ContextDataTable.TryGetData(
 						DataParameter.OverrideDataID,
 						DataParameter.ParameterType,
 						DataParameter.ParameterTypeObject);
-				Property->SetValue_InContainer(EvaluatorBlueprint, DataPtr);
+				if (ensureMsgf(DataPtr, 
+							TEXT("Can't find variable table entry for parameter '%s'. Did you build camera rig '%s'?"),
+							*DataParameter.ParameterName.ToString(), *GetNameSafe(BlueprintNode->GetOutermost())))
+				{
+					Property->SetValue_InContainer(EvaluatorBlueprint, DataPtr);
+				}
 			}
 		}
 	}
