Cameras: fix crash when using framing nodes with default pawn target but there's no pawn

diff --git a/Engine/Plugins/Cameras/GameplayCameras/Source/GameplayCameras/Private/Nodes/Framing/BaseFramingCameraNode.cpp b/Engine/Plugins/Cameras/GameplayCameras/Source/GameplayCameras/Private/Nodes/Framing/BaseFramingCameraNode.cpp
--- a/Engine/Plugins/Cameras/GameplayCameras/Source/GameplayCameras/Private/Nodes/Framing/BaseFramingCameraNode.cpp
+++ b/Engine/Plugins/Cameras/GameplayCameras/Source/GameplayCameras/Private/Nodes/Framing/BaseFramingCameraNode.cpp
@@ -148,16 +148,18 @@ bool FBaseFramingCameraNodeEvaluator::AcquireTargetInfo(TSharedPtr<const FCamera
 	}
 	else if (APlayerController* PlayerController = EvaluationContext->GetPlayerController())
 	{
-		APawn* Pawn = PlayerController->GetPawn();
-		FCameraActorComputedTargetInfo OutInfo;
-		OutInfo.Transform = FTransform3d(Pawn->GetActorLocation());
-		OutInfo.LocalBounds = FBoxSphereBounds3d(EForceInit::ForceInit);
-		if (USceneComponent* RootComponent = Pawn->GetRootComponent())
+		if (APawn* Pawn = PlayerController->GetPawn())
 		{
-			OutInfo.LocalBounds = RootComponent->Bounds;
+			FCameraActorComputedTargetInfo OutInfo;
+			OutInfo.Transform = FTransform3d(Pawn->GetActorLocation());
+			OutInfo.LocalBounds = FBoxSphereBounds3d(EForceInit::ForceInit);
+			if (USceneComponent* RootComponent = Pawn->GetRootComponent())
+			{
+				OutInfo.LocalBounds = RootComponent->Bounds;
+			}
+			OutInfos.Add(OutInfo);
+			return true;
 		}
-		OutInfos.Add(OutInfo);
-		return true;
 	}
 
 	return false;
