Cameras: make some recent engine changes backwards compatible

diff --git a/Engine/Plugins/Cameras/GameplayCameras/Source/GameplayCameras/Private/Core/CameraAsset.cpp b/Engine/Plugins/Cameras/GameplayCameras/Source/GameplayCameras/Private/Core/CameraAsset.cpp
--- a/Engine/Plugins/Cameras/GameplayCameras/Source/GameplayCameras/Private/Core/CameraAsset.cpp
+++ b/Engine/Plugins/Cameras/GameplayCameras/Source/GameplayCameras/Private/Core/CameraAsset.cpp
@@ -7,6 +7,7 @@
 #include "Core/CameraDirector.h"
 #include "Core/CameraRigAsset.h"
 #include "Core/CameraRigProxyAsset.h"
+#include "Misc/EngineVersionComparison.h"
 #include "UObject/AssetRegistryTagsContext.h"
 #include "UObject/ObjectRedirector.h"
 #include "UObject/ObjectSaveContext.h"
@@ -175,7 +176,11 @@ void UCameraAsset::PreSave(FObjectPreSaveContext ObjectSaveContext)
 #if WITH_EDITOR
 
 	const bool bIsUserObject = !HasAnyFlags(RF_ClassDefaultObject | RF_ArchetypeObject);
+#if UE_VERSION_NEWER_THAN_OR_EQUAL(5,7,0)
 	const bool bIsEditorAutoSave = ObjectSaveContext.IsFromAutoSave();
+#else
+	const bool bIsEditorAutoSave = ((ObjectSaveContext.GetSaveFlags() & SAVE_FromAutosave) != 0);
+#endif  // UE >= 5.7.0
 	if (bIsUserObject && !bIsEditorAutoSave)
 	{
 		// Build when saving/cooking.
diff --git a/Engine/Plugins/Cameras/GameplayCameras/Source/GameplayCameras/Private/Core/CameraRigAsset.cpp b/Engine/Plugins/Cameras/GameplayCameras/Source/GameplayCameras/Private/Core/CameraRigAsset.cpp
--- a/Engine/Plugins/Cameras/GameplayCameras/Source/GameplayCameras/Private/Core/CameraRigAsset.cpp
+++ b/Engine/Plugins/Cameras/GameplayCameras/Source/GameplayCameras/Private/Core/CameraRigAsset.cpp
@@ -7,6 +7,7 @@
 #include "Core/CameraNode.h"
 #include "Core/CameraVariableAssets.h"
 #include "Core/IAssetReferenceCameraNode.h"
+#include "Misc/EngineVersionComparison.h"
 #include "UObject/AssetRegistryTagsContext.h"
 #include "UObject/ObjectSaveContext.h"
 
@@ -123,7 +124,11 @@ void UCameraRigAsset::PreSave(FObjectPreSaveContext ObjectSaveContext)
 #if WITH_EDITOR
 
 	const bool bIsUserObject = !HasAnyFlags(RF_ClassDefaultObject | RF_ArchetypeObject);
+#if UE_VERSION_NEWER_THAN_OR_EQUAL(5,7,0)
 	const bool bIsEditorAutoSave = ObjectSaveContext.IsFromAutoSave();
+#else
+	const bool bIsEditorAutoSave = ((ObjectSaveContext.GetSaveFlags() & SAVE_FromAutosave) != 0);
+#endif  // UE >= 5.7.0
 	if (bIsUserObject && !bIsEditorAutoSave)
 	{
 		// Build when saving/cooking.
diff --git a/Engine/Plugins/Cameras/GameplayCameras/Source/GameplayCameras/Private/Core/CameraShakeAsset.cpp b/Engine/Plugins/Cameras/GameplayCameras/Source/GameplayCameras/Private/Core/CameraShakeAsset.cpp
--- a/Engine/Plugins/Cameras/GameplayCameras/Source/GameplayCameras/Private/Core/CameraShakeAsset.cpp
+++ b/Engine/Plugins/Cameras/GameplayCameras/Source/GameplayCameras/Private/Core/CameraShakeAsset.cpp
@@ -5,6 +5,7 @@
 #include "Build/CameraBuildLog.h"
 #include "Build/CameraShakeAssetBuilder.h"
 #include "Core/ShakeCameraNode.h"
+#include "Misc/EngineVersionComparison.h"
 #include "UObject/AssetRegistryTagsContext.h"
 #include "UObject/ObjectSaveContext.h"
 
@@ -63,7 +64,11 @@ void UCameraShakeAsset::PreSave(FObjectPreSaveContext ObjectSaveContext)
 #if WITH_EDITOR
 
 	const bool bIsUserObject = !HasAnyFlags(RF_ClassDefaultObject | RF_ArchetypeObject);
+#if UE_VERSION_NEWER_THAN_OR_EQUAL(5,7,0)
 	const bool bIsEditorAutoSave = ObjectSaveContext.IsFromAutoSave();
+#else
+	const bool bIsEditorAutoSave = ((ObjectSaveContext.GetSaveFlags() & SAVE_FromAutosave) != 0);
+#endif  // UE >= 5.7.0
 	if (bIsUserObject && !bIsEditorAutoSave)
 	{
 		// Build when saving/cooking.
