Cameras: export most public methods on the player camera manager

diff --git a/Engine/Plugins/Cameras/GameplayCameras/Source/GameplayCameras/Public/GameFramework/GameplayCamerasPlayerCameraManager.h b/Engine/Plugins/Cameras/GameplayCameras/Source/GameplayCameras/Public/GameFramework/GameplayCamerasPlayerCameraManager.h
--- a/Engine/Plugins/Cameras/GameplayCameras/Source/GameplayCameras/Public/GameFramework/GameplayCamerasPlayerCameraManager.h
+++ b/Engine/Plugins/Cameras/GameplayCameras/Source/GameplayCameras/Public/GameFramework/GameplayCamerasPlayerCameraManager.h
@@ -8,6 +8,8 @@
 
 #include "GameplayCamerasPlayerCameraManager.generated.h"
 
+#define UE_API GAMEPLAYCAMERAS_API
+
 struct FCameraRigInstanceID;
 
 namespace UE::Cameras
@@ -66,7 +68,7 @@ class AGameplayCamerasPlayerCameraManager
 
 public:
 
-	AGameplayCamerasPlayerCameraManager(const FObjectInitializer& ObjectInitializer);
+	UE_API AGameplayCamerasPlayerCameraManager(const FObjectInitializer& ObjectInitializer);
 
 public:
 
@@ -106,16 +108,16 @@ public:
 public:
 
 	// APlayerCameraManager interface.
-	virtual void InitializeFor(APlayerController* PlayerController) override;
-	virtual void SetViewTarget(AActor* NewViewTarget, FViewTargetTransitionParams TransitionParams = FViewTargetTransitionParams()) override;
-	virtual void ProcessViewRotation(float DeltaTime, FRotator& OutViewRotation, FRotator& OutDeltaRot) override;
+	UE_API virtual void InitializeFor(APlayerController* PlayerController) override;
+	UE_API virtual void SetViewTarget(AActor* NewViewTarget, FViewTargetTransitionParams TransitionParams = FViewTargetTransitionParams()) override;
+	UE_API virtual void ProcessViewRotation(float DeltaTime, FRotator& OutViewRotation, FRotator& OutDeltaRot) override;
 
 	// AActor interface.
-	virtual void DisplayDebug(UCanvas* Canvas, const FDebugDisplayInfo& DebugDisplay, float& YL, float& YPos) override;
+	UE_API virtual void DisplayDebug(UCanvas* Canvas, const FDebugDisplayInfo& DebugDisplay, float& YL, float& YPos) override;
 
 	// UObject interface.
-	virtual void BeginDestroy() override;
-	static void AddReferencedObjects(UObject* InThis, FReferenceCollector& Collector);
+	UE_API virtual void BeginDestroy() override;
+	UE_API static void AddReferencedObjects(UObject* InThis, FReferenceCollector& Collector);
 
 	// IGameplayCameraSystemHost interface.
 	virtual UObject* GetAsObject() override { return this; }
@@ -123,7 +125,7 @@ public:
 protected:
 
 	// APlayerCameraManager interface.
-	virtual void DoUpdateCamera(float DeltaTime) override;
+	UE_API virtual void DoUpdateCamera(float DeltaTime) override;
 
 private:
 
@@ -158,3 +160,5 @@ private:
 	bool bIsSettingNewViewTarget = false;
 };
 
+#undef UE_API
+
