Cameras: fixes for the player camera manager

- Don't up-call DoUpdateCamera because it ends up double-updating camera actors, since they'are also updated inside the camera system
- Run the vanilla (Blueprint-based) camera modifiers ourselves after the camera system
- Set the transform of the camera manager itself to where the final viewpoint is

diff --git a/Engine/Plugins/Cameras/GameplayCameras/Source/GameplayCameras/Private/GameFramework/GameplayCamerasPlayerCameraManager.cpp b/Engine/Plugins/Cameras/GameplayCameras/Source/GameplayCameras/Private/GameFramework/GameplayCamerasPlayerCameraManager.cpp
--- a/Engine/Plugins/Cameras/GameplayCameras/Source/GameplayCameras/Private/GameFramework/GameplayCamerasPlayerCameraManager.cpp
+++ b/Engine/Plugins/Cameras/GameplayCameras/Source/GameplayCameras/Private/GameFramework/GameplayCamerasPlayerCameraManager.cpp
@@ -294,12 +294,54 @@ void AGameplayCamerasPlayerCameraManager::DoUpdateCamera(float DeltaTime)
 {
 	using namespace UE::Cameras;
 
-	Super::DoUpdateCamera(DeltaTime);
+	// Don't up-call Super::DoUpdateCamera because:
+	//
+	// 1) it runs the view targets' CalcCamera or GetCameraView, which we already do inside our camera system,
+	//    and we don't want them to double-update.
+	// 2) it does a bunch of stuff regarding the pending view target that we don't care about.
+	// 3) it applies the camera modifiers twice (once for the view target, once for the pending view target)
+	//    which means that things like shakes can look like they're twice as "intense" during blends, and also
+	//    means that modifiers get applied at the wrong level.
+	//
+	// So as a result, until we refactor the base camera manager class, we have to re-do some of the logic we
+	// want to keep, such as color scale interpolation and screen/audio fading.
+	//
 
-	if (CameraSystemEvaluator.IsValid())
+	if (bEnableColorScaleInterp)
 	{
-		FillCameraCache(LastFrameDesiredView);
+		float BlendPct = FMath::Clamp((GetWorld()->TimeSeconds - ColorScaleInterpStartTime) / ColorScaleInterpDuration, 0.f, 1.0f);
+		ColorScale = FMath::Lerp(OriginalColorScale, DesiredColorScale, BlendPct);
+		if (BlendPct == 1.0f)
+		{
+			bEnableColorScaleInterp = false;
+		}
+	}
 
+	if (bEnableFading)
+	{
+		if (bAutoAnimateFade)
+		{
+			FadeTimeRemaining = FMath::Max(FadeTimeRemaining - DeltaTime, 0.0f);
+			if (FadeTime > 0.0f)
+			{
+				FadeAmount = FadeAlpha.X + ((1.f - FadeTimeRemaining / FadeTime) * (FadeAlpha.Y - FadeAlpha.X));
+			}
+
+			if ((bHoldFadeWhenFinished == false) && (FadeTimeRemaining <= 0.f))
+			{
+				// done
+				StopCameraFade();
+			}
+		}
+
+		if (bFadeAudio)
+		{
+			ApplyAudioFade();
+		}
+	}
+
+	if (CameraSystemEvaluator.IsValid())
+	{
 		FCameraSystemEvaluationParams UpdateParams;
 		UpdateParams.DeltaTime = DeltaTime;
 		CameraSystemEvaluator->Update(UpdateParams);
@@ -307,11 +349,13 @@ void AGameplayCamerasPlayerCameraManager::DoUpdateCamera(float DeltaTime)
 		FMinimalViewInfo DesiredView;
 		CameraSystemEvaluator->GetEvaluatedCameraView(DesiredView);
 
-		FillCameraCache(DesiredView);
+		ApplyCameraModifiers(DeltaTime, DesiredView);
 
-		LastFrameDesiredView = DesiredView;
+		FillCameraCache(DesiredView);
 
 		CleanUpViewTargetContexts();
+
+		SetActorLocationAndRotation(DesiredView.Location, DesiredView.Rotation, false);
 	}
 }
 
diff --git a/Engine/Plugins/Cameras/GameplayCameras/Source/GameplayCameras/Public/GameFramework/GameplayCamerasPlayerCameraManager.h b/Engine/Plugins/Cameras/GameplayCameras/Source/GameplayCameras/Public/GameFramework/GameplayCamerasPlayerCameraManager.h
--- a/Engine/Plugins/Cameras/GameplayCameras/Source/GameplayCameras/Public/GameFramework/GameplayCamerasPlayerCameraManager.h
+++ b/Engine/Plugins/Cameras/GameplayCameras/Source/GameplayCameras/Public/GameFramework/GameplayCamerasPlayerCameraManager.h
@@ -155,8 +155,6 @@ private:
 
 	TSharedPtr<UE::Cameras::FCameraEvaluationContext> NullContext;
 
-	FMinimalViewInfo LastFrameDesiredView;
-
 	bool bIsSettingNewViewTarget = false;
 };
 
